Chaos
- Joint tolerance
- Kinematic-Dynamic fast path (WIP. Currently for common constraints combinations only)

diff --git a/Engine/Source/Runtime/Engine/Private/PhysicsEngine/ImmediatePhysics/ImmediatePhysicsChaos/ImmediatePhysicsSimulation_Chaos.cpp b/Engine/Source/Runtime/Engine/Private/PhysicsEngine/ImmediatePhysics/ImmediatePhysicsChaos/ImmediatePhysicsSimulation_Chaos.cpp
--- a/Engine/Source/Runtime/Engine/Private/PhysicsEngine/ImmediatePhysics/ImmediatePhysicsChaos/ImmediatePhysicsSimulation_Chaos.cpp
+++ b/Engine/Source/Runtime/Engine/Private/PhysicsEngine/ImmediatePhysics/ImmediatePhysicsChaos/ImmediatePhysicsSimulation_Chaos.cpp
@@ -56,6 +56,8 @@ FAutoConsoleVariableRef CVarChaosImmPhysCollisionShapePadding(TEXT("p.Chaos.ImmP
 int32 ChaosImmediate_Joint_PairIterations = -1;
 int32 ChaosImmediate_Joint_PushOutPairIterations = -1;
 float ChaosImmediate_Joint_SwingTwistAngleTolerance = 1.0e-6f;
+float ChaosImmediate_Joint_PositionTolerance = 0.05f;
+float ChaosImmediate_Joint_AngleTolerance = 0.01f;
 int32 ChaosImmediate_Joint_EnableTwistLimits = 1;
 int32 ChaosImmediate_Joint_EnableSwingLimits = 1;
 int32 ChaosImmediate_Joint_EnableDrives = 1;
@@ -77,6 +79,8 @@ float ChaosImmediate_Joint_AngularPositionCorrection = 1.0f;
 FAutoConsoleVariableRef CVarChaosImmPhysJointPairIterations(TEXT("p.Chaos.ImmPhys.Joint.PairIterations"), ChaosImmediate_Joint_PairIterations, TEXT("Override joint pair iterations (if >= 0)"));
 FAutoConsoleVariableRef CVarChaosImmPhysJointPushOutPairIterations(TEXT("p.Chaos.ImmPhys.Joint.PushOutPairIterations"), ChaosImmediate_Joint_PushOutPairIterations, TEXT("Override joint push-out pair iterations (if >= 0)"));
 FAutoConsoleVariableRef CVarChaosImmPhysJointSwingTwistAngleTolerance(TEXT("p.Chaos.ImmPhys.Joint.SwingTwistAngleTolerance"), ChaosImmediate_Joint_SwingTwistAngleTolerance, TEXT("SwingTwistAngleTolerance."));
+FAutoConsoleVariableRef CVarChaosImmPhysJointPositionTolerance(TEXT("p.Chaos.ImmPhys.Joint.PositionTolerance"), ChaosImmediate_Joint_PositionTolerance, TEXT("PositionTolerance."));
+FAutoConsoleVariableRef CVarChaosImmPhysJointAngleTolerance(TEXT("p.Chaos.ImmPhys.Joint.AngleTolerance"), ChaosImmediate_Joint_AngleTolerance, TEXT("PositionTolerance."));
 FAutoConsoleVariableRef CVarChaosImmPhysJointEnableTwistLimits(TEXT("p.Chaos.ImmPhys.Joint.EnableTwistLimits"), ChaosImmediate_Joint_EnableTwistLimits, TEXT("EnableTwistLimits."));
 FAutoConsoleVariableRef CVarChaosImmPhysJointEnableSwingLimits(TEXT("p.Chaos.ImmPhys.Joint.EnableSwingLimits"), ChaosImmediate_Joint_EnableSwingLimits, TEXT("EnableSwingLimits."));
 FAutoConsoleVariableRef CVarChaosImmPhysJointEnableDrives(TEXT("p.Chaos.ImmPhys.Joint.EnableDrives"), ChaosImmediate_Joint_EnableDrives, TEXT("EnableDrives."));
@@ -555,6 +559,8 @@ namespace ImmediatePhysics_Chaos
 
 			FPBDJointSolverSettings JointsSettings = Joints.GetSettings();
 			JointsSettings.SwingTwistAngleTolerance = ChaosImmediate_Joint_SwingTwistAngleTolerance;
+			JointsSettings.PositionTolerance = ChaosImmediate_Joint_PositionTolerance;
+			JointsSettings.AngleTolerance = ChaosImmediate_Joint_AngleTolerance;
 			JointsSettings.MinParentMassRatio = ChaosImmediate_Joint_MinParentMassRatio;
 			JointsSettings.MaxInertiaRatio = ChaosImmediate_Joint_MaxInertiaRatio;
 			JointsSettings.AngularConstraintPositionCorrection = ChaosImmediate_Joint_AngularPositionCorrection;
diff --git a/Engine/Source/Runtime/Experimental/Chaos/Private/Chaos/Joint/PBDJointSolverGaussSeidel.cpp b/Engine/Source/Runtime/Experimental/Chaos/Private/Chaos/Joint/PBDJointSolverGaussSeidel.cpp
--- a/Engine/Source/Runtime/Experimental/Chaos/Private/Chaos/Joint/PBDJointSolverGaussSeidel.cpp
+++ b/Engine/Source/Runtime/Experimental/Chaos/Private/Chaos/Joint/PBDJointSolverGaussSeidel.cpp
@@ -318,6 +318,8 @@ namespace Chaos
 		SwingDriveLambda = (FReal)0;
 
 		AngularPositionCorrection = FPBDJointUtilities::GetAngularPositionCorrection(SolverSettings, JointSettings);
+		PositionTolerance = SolverSettings.PositionTolerance;
+		AngleTolerance = SolverSettings.AngleTolerance;
 	}
 
 
@@ -347,45 +349,58 @@ namespace Chaos
 	}
 
 
-	void FJointSolverGaussSeidel::ApplyConstraints(
+	FJointSolverResult FJointSolverGaussSeidel::ApplyConstraints(
 		const FReal Dt,
 		const FPBDJointSolverSettings& SolverSettings,
 		const FPBDJointSettings& JointSettings)
 	{
-		ApplyPositionConstraints(Dt, SolverSettings, JointSettings);
+		FJointSolverResult NetResult;
+
+		NetResult += ApplyPositionConstraints(Dt, SolverSettings, JointSettings);
 	
-		ApplyRotationConstraints(Dt, SolverSettings, JointSettings);
+		NetResult += ApplyRotationConstraints(Dt, SolverSettings, JointSettings);
+
+		return NetResult;
 	}
 
 
-	void FJointSolverGaussSeidel::ApplyDrives(
+	FJointSolverResult FJointSolverGaussSeidel::ApplyDrives(
 		const FReal Dt,
 		const FPBDJointSolverSettings& SolverSettings,
 		const FPBDJointSettings& JointSettings)
 	{
-		ApplyPositionDrives(Dt, SolverSettings, JointSettings);
+		FJointSolverResult NetResult;
+
+		NetResult += ApplyPositionDrives(Dt, SolverSettings, JointSettings);
 
-		ApplyRotationDrives(Dt, SolverSettings, JointSettings);
+		NetResult += ApplyRotationDrives(Dt, SolverSettings, JointSettings);
+
+		return NetResult;
 	}
 
 
-	void FJointSolverGaussSeidel::ApplyProjections(
+	FJointSolverResult FJointSolverGaussSeidel::ApplyProjections(
 		const FReal Dt,
 		const FPBDJointSolverSettings& SolverSettings,
 		const FPBDJointSettings& JointSettings)
 	{
-		//ApplyRotationProjection(Dt, SolverSettings, JointSettings);
+		FJointSolverResult NetResult;
+
+		//NetResult += ApplyRotationProjection(Dt, SolverSettings, JointSettings);
 
-		ApplyPositionProjection(Dt, SolverSettings, JointSettings);
+		NetResult += ApplyPositionProjection(Dt, SolverSettings, JointSettings);
+
+		return NetResult;
 	}
 
 
-	void FJointSolverGaussSeidel::ApplyRotationConstraints(
+	FJointSolverResult FJointSolverGaussSeidel::ApplyRotationConstraints(
 		const FReal Dt,
 		const FPBDJointSolverSettings& SolverSettings,
 		const FPBDJointSettings& JointSettings)
 	{
 		// @todo(ccaulfield): eliminate branches in solver loop by pre-building list of functions to call?
+		FJointSolverResult NetResult;
 
 		EJointMotionType TwistMotion = JointSettings.AngularMotionTypes[(int32)EJointAngularConstraintIndex::Twist];
 		EJointMotionType Swing1Motion = JointSettings.AngularMotionTypes[(int32)EJointAngularConstraintIndex::Swing1];
@@ -398,11 +413,11 @@ namespace Chaos
 		{
 			if ((TwistMotion == EJointMotionType::Limited) && bTwistSoft)
 			{
-				ApplyTwistConstraintSoft(Dt, SolverSettings, JointSettings);
+				NetResult += ApplyTwistConstraintSoft(Dt, SolverSettings, JointSettings);
 			}
 			else if (TwistMotion != EJointMotionType::Free)
 			{
-				ApplyTwistConstraint(Dt, SolverSettings, JointSettings);
+				NetResult += ApplyTwistConstraint(Dt, SolverSettings, JointSettings);
 			}
 		}
 
@@ -413,42 +428,46 @@ namespace Chaos
 			{
 				if (bSwingSoft)
 				{
-					ApplyConeConstraintSoft(Dt, SolverSettings, JointSettings);
+					NetResult += ApplyConeConstraintSoft(Dt, SolverSettings, JointSettings);
 				}
 				else
 				{
-					ApplyConeConstraint(Dt, SolverSettings, JointSettings);
+					NetResult += ApplyConeConstraint(Dt, SolverSettings, JointSettings);
 				}
 			}
 			else
 			{
 				if ((Swing1Motion == EJointMotionType::Limited) && bSwingSoft)
 				{
-					ApplySwingConstraintSoft(Dt, SolverSettings, JointSettings, EJointAngularConstraintIndex::Swing1, EJointAngularAxisIndex::Swing1);
+					NetResult += ApplySwingConstraintSoft(Dt, SolverSettings, JointSettings, EJointAngularConstraintIndex::Swing1, EJointAngularAxisIndex::Swing1);
 				}
 				else if (Swing1Motion != EJointMotionType::Free)
 				{
-					ApplySwingConstraint(Dt, SolverSettings, JointSettings, EJointAngularConstraintIndex::Swing1, EJointAngularAxisIndex::Swing1);
+					NetResult += ApplySwingConstraint(Dt, SolverSettings, JointSettings, EJointAngularConstraintIndex::Swing1, EJointAngularAxisIndex::Swing1);
 				}
 				
 				if ((Swing2Motion == EJointMotionType::Limited) && bSwingSoft)
 				{
-					ApplySwingConstraintSoft(Dt, SolverSettings, JointSettings, EJointAngularConstraintIndex::Swing2, EJointAngularAxisIndex::Swing2);
+					NetResult += ApplySwingConstraintSoft(Dt, SolverSettings, JointSettings, EJointAngularConstraintIndex::Swing2, EJointAngularAxisIndex::Swing2);
 				}
 				else if (Swing2Motion != EJointMotionType::Free)
 				{
-					ApplySwingConstraint(Dt, SolverSettings, JointSettings, EJointAngularConstraintIndex::Swing2, EJointAngularAxisIndex::Swing2);
+					NetResult += ApplySwingConstraint(Dt, SolverSettings, JointSettings, EJointAngularConstraintIndex::Swing2, EJointAngularAxisIndex::Swing2);
 				}
 			}
 		}
+
+		return NetResult;
 	}
 
 
-	void FJointSolverGaussSeidel::ApplyRotationDrives(
+	FJointSolverResult FJointSolverGaussSeidel::ApplyRotationDrives(
 		const FReal Dt,
 		const FPBDJointSolverSettings& SolverSettings,
 		const FPBDJointSettings& JointSettings)
 	{
+		FJointSolverResult NetResult;
+
 		EJointMotionType TwistMotion = JointSettings.AngularMotionTypes[(int32)EJointAngularConstraintIndex::Twist];
 		EJointMotionType Swing1Motion = JointSettings.AngularMotionTypes[(int32)EJointAngularConstraintIndex::Swing1];
 		EJointMotionType Swing2Motion = JointSettings.AngularMotionTypes[(int32)EJointAngularConstraintIndex::Swing2];
@@ -462,38 +481,42 @@ namespace Chaos
 			// No SLerp drive if we have a locked rotation (it will be grayed out in the editor in this case, but could still have been set before the rotation was locked)
 			if ((JointSettings.bAngularSLerpPositionDriveEnabled || JointSettings.bAngularSLerpVelocityDriveEnabled) && !bTwistLocked && !bSwing1Locked && !bSwing2Locked)
 			{
-				ApplySLerpDrive(Dt, SolverSettings, JointSettings);
+				NetResult += ApplySLerpDrive(Dt, SolverSettings, JointSettings);
 			}
 			else
 			{
 				if ((JointSettings.bAngularTwistPositionDriveEnabled || JointSettings.bAngularTwistVelocityDriveEnabled) && !bTwistLocked)
 				{
-					ApplyTwistDrive(Dt, SolverSettings, JointSettings);
+					NetResult += ApplyTwistDrive(Dt, SolverSettings, JointSettings);
 				}
 
 				const bool bSwingDriveEnabled = (JointSettings.bAngularSwingPositionDriveEnabled || JointSettings.bAngularSwingVelocityDriveEnabled);
 				if (bSwingDriveEnabled && !bSwing1Locked && !bSwing2Locked)
 				{
-					ApplyConeDrive(Dt, SolverSettings, JointSettings);
+					NetResult += ApplyConeDrive(Dt, SolverSettings, JointSettings);
 				}
 				else if (bSwingDriveEnabled && !bSwing1Locked)
 				{
-					ApplySwingDrive(Dt, SolverSettings, JointSettings, EJointAngularConstraintIndex::Swing1, EJointAngularAxisIndex::Swing1);
+					NetResult += ApplySwingDrive(Dt, SolverSettings, JointSettings, EJointAngularConstraintIndex::Swing1, EJointAngularAxisIndex::Swing1);
 				}
 				else if (bSwingDriveEnabled && !bSwing2Locked)
 				{
-					ApplySwingDrive(Dt, SolverSettings, JointSettings, EJointAngularConstraintIndex::Swing2, EJointAngularAxisIndex::Swing2);
+					NetResult += ApplySwingDrive(Dt, SolverSettings, JointSettings, EJointAngularConstraintIndex::Swing2, EJointAngularAxisIndex::Swing2);
 				}
 			}
 		}
+
+		return NetResult;
 	}
 
 
-	void FJointSolverGaussSeidel::ApplyRotationProjection(
+	FJointSolverResult FJointSolverGaussSeidel::ApplyRotationProjection(
 		const FReal Dt,
 		const FPBDJointSolverSettings& SolverSettings,
 		const FPBDJointSettings& JointSettings)
 	{
+		FJointSolverResult NetResult;
+
 		EJointMotionType TwistMotion = JointSettings.AngularMotionTypes[(int32)EJointAngularConstraintIndex::Twist];
 		EJointMotionType Swing1Motion = JointSettings.AngularMotionTypes[(int32)EJointAngularConstraintIndex::Swing1];
 		EJointMotionType Swing2Motion = JointSettings.AngularMotionTypes[(int32)EJointAngularConstraintIndex::Swing2];
@@ -504,7 +527,7 @@ namespace Chaos
 		{
 			if ((TwistMotion == EJointMotionType::Locked) || ((TwistMotion == EJointMotionType::Limited) && !bTwistSoft))
 			{
-				ApplyTwistProjection(Dt, SolverSettings, JointSettings);
+				NetResult += ApplyTwistProjection(Dt, SolverSettings, JointSettings);
 			}
 		}
 
@@ -512,29 +535,32 @@ namespace Chaos
 		{
 			if ((Swing1Motion == EJointMotionType::Limited) && (Swing2Motion == EJointMotionType::Limited) && !bSwingSoft)
 			{
-				ApplyConeProjection(Dt, SolverSettings, JointSettings);
+				NetResult += ApplyConeProjection(Dt, SolverSettings, JointSettings);
 			}
 			else
 			{
 				if ((Swing1Motion == EJointMotionType::Locked) || ((Swing1Motion == EJointMotionType::Limited) && !bSwingSoft))
 				{
-					ApplySwingProjection(Dt, SolverSettings, JointSettings, EJointAngularConstraintIndex::Swing1, EJointAngularAxisIndex::Swing1);
+					NetResult += ApplySwingProjection(Dt, SolverSettings, JointSettings, EJointAngularConstraintIndex::Swing1, EJointAngularAxisIndex::Swing1);
 				}
 				if ((Swing2Motion == EJointMotionType::Locked) || ((Swing2Motion == EJointMotionType::Limited) && !bSwingSoft))
 				{
-					ApplySwingProjection(Dt, SolverSettings, JointSettings, EJointAngularConstraintIndex::Swing2, EJointAngularAxisIndex::Swing2);
+					NetResult += ApplySwingProjection(Dt, SolverSettings, JointSettings, EJointAngularConstraintIndex::Swing2, EJointAngularAxisIndex::Swing2);
 				}
 			}
 		}
+
+		return NetResult;
 	}
 
 
-	void FJointSolverGaussSeidel::ApplyPositionConstraints(
+	FJointSolverResult FJointSolverGaussSeidel::ApplyPositionConstraints(
 		const FReal Dt,
 		const FPBDJointSolverSettings& SolverSettings,
 		const FPBDJointSettings& JointSettings)
 	{
 		// @todo(ccaulfield): the branch logic is all constant for a joint - move it to initialization and turn it into a table or something (also for Hard/Soft login in apply functions)
+		FJointSolverResult NetResult;
 
 		const TVector<EJointMotionType, 3>& LinearMotion = JointSettings.LinearMotionTypes;
 		const TVector<bool, 3> bLinearLocked =
@@ -553,62 +579,75 @@ namespace Chaos
 		if (bLinearLocked[0] && bLinearLocked[1] && bLinearLocked[2])
 		{
 			// Hard point constraint (most common case)
-			ApplyPointPositionConstraint(Dt, SolverSettings, JointSettings);
+			if (InvMs[0] == 0)
+			{
+				NetResult += ApplyPointPositionConstraintKD(0, 1, Dt, SolverSettings, JointSettings);
+			}
+			else if (InvMs[1] == 0)
+			{
+				NetResult += ApplyPointPositionConstraintKD(1, 0, Dt, SolverSettings, JointSettings);
+			}
+			else
+			{
+				NetResult += ApplyPointPositionConstraintDD(Dt, SolverSettings, JointSettings);
+			}
 		}
 		else if (bLinearLimted[0] && bLinearLimted[1] && bLinearLimted[2])
 		{
 			// Spherical constraint
-			ApplySphericalPositionConstraint(Dt, SolverSettings, JointSettings);
+			NetResult += ApplySphericalPositionConstraint(Dt, SolverSettings, JointSettings);
 		}
 		else if (bLinearLocked[1] && bLinearLocked[2] && !bLinearLocked[0])
 		{
 			// Line constraint along X axis
-			ApplyCylindricalPositionConstraint(Dt, 0, LinearMotion[0], EJointMotionType::Locked, SolverSettings, JointSettings);
+			NetResult += ApplyCylindricalPositionConstraint(Dt, 0, LinearMotion[0], EJointMotionType::Locked, SolverSettings, JointSettings);
 		}
 		else if (bLinearLocked[0] && bLinearLocked[2] && !bLinearLocked[1])
 		{
 			// Line constraint along Y axis
-			ApplyCylindricalPositionConstraint(Dt, 1, LinearMotion[1], EJointMotionType::Locked, SolverSettings, JointSettings);
+			NetResult += ApplyCylindricalPositionConstraint(Dt, 1, LinearMotion[1], EJointMotionType::Locked, SolverSettings, JointSettings);
 		}
 		else if (bLinearLocked[0] && bLinearLocked[1] && !bLinearLocked[2])
 		{
 			// Line constraint along Z axis
-			ApplyCylindricalPositionConstraint(Dt, 2, LinearMotion[2], EJointMotionType::Locked, SolverSettings, JointSettings);
+			NetResult += ApplyCylindricalPositionConstraint(Dt, 2, LinearMotion[2], EJointMotionType::Locked, SolverSettings, JointSettings);
 		}
 		else if (bLinearLimted[1] && bLinearLimted[2] && !bLinearLimted[0])
 		{
 			// Cylindrical constraint along X axis
-			ApplyCylindricalPositionConstraint(Dt, 0, LinearMotion[0], EJointMotionType::Limited, SolverSettings, JointSettings);
+			NetResult += ApplyCylindricalPositionConstraint(Dt, 0, LinearMotion[0], EJointMotionType::Limited, SolverSettings, JointSettings);
 		}
 		else if (bLinearLimted[0] && bLinearLimted[2] && !bLinearLimted[1])
 		{
 			// Cylindrical constraint along Y axis
-			ApplyCylindricalPositionConstraint(Dt, 1, LinearMotion[1], EJointMotionType::Limited, SolverSettings, JointSettings);
+			NetResult += ApplyCylindricalPositionConstraint(Dt, 1, LinearMotion[1], EJointMotionType::Limited, SolverSettings, JointSettings);
 		}
 		else if (bLinearLimted[0] && bLinearLimted[1] && !bLinearLimted[2])
 		{
 			// Cylindrical constraint along Z axis
-			ApplyCylindricalPositionConstraint(Dt, 2, LinearMotion[2], EJointMotionType::Limited, SolverSettings, JointSettings);
+			NetResult += ApplyCylindricalPositionConstraint(Dt, 2, LinearMotion[2], EJointMotionType::Limited, SolverSettings, JointSettings);
 		}
 		else if (bLinearLocked[0] || bLinearLimted[0])
 		{
 			// Planar constraint along X axis
-			ApplyPlanarPositionConstraint(Dt, 0, LinearMotion[0], SolverSettings, JointSettings);
+			NetResult += ApplyPlanarPositionConstraint(Dt, 0, LinearMotion[0], SolverSettings, JointSettings);
 		}
 		else if (bLinearLocked[1] || bLinearLimted[1])
 		{
 			// Planar constraint along Y axis
-			ApplyPlanarPositionConstraint(Dt, 1, LinearMotion[1], SolverSettings, JointSettings);
+			NetResult += ApplyPlanarPositionConstraint(Dt, 1, LinearMotion[1], SolverSettings, JointSettings);
 		}
 		else if (bLinearLocked[2] || bLinearLimted[2])
 		{
 			// Planar constraint along Z axis
-			ApplyPlanarPositionConstraint(Dt, 2, LinearMotion[2], SolverSettings, JointSettings);
+			NetResult += ApplyPlanarPositionConstraint(Dt, 2, LinearMotion[2], SolverSettings, JointSettings);
 		}
+
+		return NetResult;
 	}
 
 
-	void FJointSolverGaussSeidel::ApplyPositionDrives(
+	FJointSolverResult FJointSolverGaussSeidel::ApplyPositionDrives(
 		const FReal Dt,
 		const FPBDJointSolverSettings& SolverSettings,
 		const FPBDJointSettings& JointSettings)
@@ -617,6 +656,8 @@ namespace Chaos
 		// E.g., if X/Y have position drives and Z is a velocity drive, it should be a circular position-drive and axial velocity-drive rather than a spherical one.
 		// E.g., if X/Y/Z have position drives and only 1 or 2 axes have velocity drives, we need to apply 2 drives. Etc etc.
 		// Basically we need to split the axes by those that have the same drive settings.
+		FJointSolverResult NetResult;
+
 		if (SolverSettings.bEnableDrives)
 		{
 			TVector<bool, 3> bDriven =
@@ -628,33 +669,35 @@ namespace Chaos
 
 			if (bDriven[0] && bDriven[1] && bDriven[2])
 			{
-				ApplySphericalPositionDrive(Dt, SolverSettings, JointSettings);
+				NetResult += ApplySphericalPositionDrive(Dt, SolverSettings, JointSettings);
 			}
 			else if (bDriven[1] && bDriven[2])
 			{
-				ApplyCircularPositionDrive(Dt, 0, SolverSettings, JointSettings);
+				NetResult += ApplyCircularPositionDrive(Dt, 0, SolverSettings, JointSettings);
 			}
 			else if (bDriven[0] && bDriven[2])
 			{
-				ApplyCircularPositionDrive(Dt, 1, SolverSettings, JointSettings);
+				NetResult += ApplyCircularPositionDrive(Dt, 1, SolverSettings, JointSettings);
 			}
 			else if (bDriven[0] && bDriven[1])
 			{
-				ApplyCircularPositionDrive(Dt, 2, SolverSettings, JointSettings);
+				NetResult += ApplyCircularPositionDrive(Dt, 2, SolverSettings, JointSettings);
 			}
 			else if (bDriven[0])
 			{
-				ApplyAxialPositionDrive(Dt, 0, SolverSettings, JointSettings);
+				NetResult += ApplyAxialPositionDrive(Dt, 0, SolverSettings, JointSettings);
 			}
 			else if (bDriven[1])
 			{
-				ApplyAxialPositionDrive(Dt, 1, SolverSettings, JointSettings);
+				NetResult += ApplyAxialPositionDrive(Dt, 1, SolverSettings, JointSettings);
 			}
 			else if (bDriven[2])
 			{
-				ApplyAxialPositionDrive(Dt, 2, SolverSettings, JointSettings);
+				NetResult += ApplyAxialPositionDrive(Dt, 2, SolverSettings, JointSettings);
 			}
 		}
+
+		return NetResult;
 	}
 
 
@@ -706,6 +749,17 @@ namespace Chaos
 		Qs[1].EnforceShortestArcWith(Qs[0]);
 	}
 
+	void FJointSolverGaussSeidel::ApplyDelta(
+		const int32 BodyIndex,
+		const FReal Stiffness,
+		const FVec3& DP,
+		const FVec3& DR)
+	{
+		Ps[BodyIndex] += Stiffness * DP;
+		const FRotation3 DQ = (FRotation3::FromElements(Stiffness * DR, 0) * Qs[BodyIndex]) * (FReal)0.5;
+		Qs[BodyIndex] = (Qs[BodyIndex] + DQ).GetNormalized();
+	}
+
 
 	void FJointSolverGaussSeidel::ApplyVelocityDelta(
 		const int32 BodyIndex,
@@ -864,7 +918,9 @@ namespace Chaos
 
 
 	// See "XPBD: Position-Based Simulation of Compliant Constrained Dynamics"
-	void FJointSolverGaussSeidel::ApplyRotationConstraintSoft(
+	void FJointSolverGaussSeidel::ApplyRotationConstraintSoftKD(
+		const int32 KIndex,
+		const int32 DIndex,
 		const FReal Dt,
 		const FReal Stiffness,
 		const FReal Damping,
@@ -873,10 +929,93 @@ namespace Chaos
 		const FReal Angle,
 		FReal& Lambda)
 	{
+		check(InvMs[DIndex] > 0);
+
 		if (bChaos_Joint_ISPC_Enabled)
 		{
 #if INTEL_ISPC
-			ispc::ApplyRotationConstraintSoft((ispc::FJointSolverGaussSeidel*)this, Dt, Stiffness, Damping, bAccelerationMode, (ispc::FVector&) Axis, Angle, Lambda);
+			ispc::ApplyRotationConstraintSoftKD((ispc::FJointSolverGaussSeidel*)this, KIndex, DIndex, Dt, Stiffness, Damping, bAccelerationMode, (ispc::FVector&) Axis, Angle, Lambda);
+#endif
+		}
+		else
+		{
+			// World-space inverse mass
+			const FMatrix33 InvI1 = Utilities::ComputeWorldSpaceInertia(Qs[DIndex], InvILs[DIndex]);
+			const FVec3 IA1 = Utilities::Multiply(InvI1, Axis);
+
+			// Joint-space inverse mass
+			FReal II1 = FVec3::DotProduct(Axis, IA1);
+
+			// If we are correcting the position, we need to adjust the constraint effective mass using parallel axis theorem
+			// @todo(ccaulfield): the IMs here are constant per constraint...pre-build and cache it
+			if (AngularPositionCorrection > 0)
+			{
+				const FVec3 LinearAxis1 = FVec3::CrossProduct(Xs[DIndex] - Ps[DIndex], Axis);
+				const FReal LinearAxisLen1Sq = LinearAxis1.SizeSquared();
+				if (LinearAxisLen1Sq > KINDA_SMALL_NUMBER)
+				{
+					const FReal IM1 = InvMs[DIndex] / LinearAxisLen1Sq;
+					II1 = II1 * IM1 / (II1 + IM1);
+				}
+			}
+			const FReal II = II1;
+
+			// Damping angular velocity
+			FReal AngVelDt = 0;
+			if (Damping > KINDA_SMALL_NUMBER)
+			{
+				const FVec3 W1 = FRotation3::CalculateAngularVelocity(PrevQs[DIndex], Qs[DIndex], 1.0f);
+				AngVelDt = -FVec3::DotProduct(Axis, W1);
+			}
+
+			const FReal SpringMassScale = (bAccelerationMode) ? 1.0f / II : 1.0f;
+			const FReal S = SpringMassScale * Stiffness * Dt * Dt;
+			const FReal D = SpringMassScale * Damping * Dt;
+			const FReal Multiplier = (FReal)1 / ((S + D) * II + (FReal)1);
+			const FReal DLambda = Multiplier * (S * Angle - D * AngVelDt - Lambda);
+
+			//const FVec3 DR1 = IA1 * -DLambda;
+			const FVec3 DR1 = Axis * -(DLambda * II1);
+
+			Lambda += DLambda;
+			ApplyRotationDelta(DIndex, 1.0f, DR1);
+
+			// Correct the positional error that was introduced by the rotation correction. This is
+			// correct when the position dofs are all locked, but not otherwise (fixable).
+			// This significantly improves angular stiffness at lower iterations, although the same effect
+			// is achieved by increasing iterations.
+			// @todo(ccaulfield): this position correction needs to have components in direction of inactive position constraints removed
+			if (AngularPositionCorrection > 0)
+			{
+				const FVec3 PrevX0 = Xs[KIndex];
+				const FVec3 PrevX1 = Xs[DIndex];
+				UpdateDerivedState(DIndex);
+
+				const FVec3 DP1 = (PrevX1 - PrevX0) - (Xs[DIndex] - Xs[KIndex]);
+				ApplyPositionDelta(DIndex, AngularPositionCorrection, DP1);
+			}
+
+			UpdateDerivedState(DIndex);
+		}
+	}
+
+	// See "XPBD: Position-Based Simulation of Compliant Constrained Dynamics"
+	void FJointSolverGaussSeidel::ApplyRotationConstraintSoftDD(
+		const FReal Dt,
+		const FReal Stiffness,
+		const FReal Damping,
+		const bool bAccelerationMode,
+		const FVec3& Axis,
+		const FReal Angle,
+		FReal& Lambda)
+	{
+		check(InvMs[0] > 0);
+		check(InvMs[1] > 0);
+
+		if (bChaos_Joint_ISPC_Enabled)
+		{
+#if INTEL_ISPC
+			ispc::ApplyRotationConstraintSoftDD((ispc::FJointSolverGaussSeidel*)this, Dt, Stiffness, Damping, bAccelerationMode, (ispc::FVector&) Axis, Angle, Lambda);
 #endif
 		}
 		else
@@ -895,25 +1034,19 @@ namespace Chaos
 			// @todo(ccaulfield): the IMs here are constant per constraint...pre-build and cache it
 			if (AngularPositionCorrection > 0)
 			{
-				if (InvMs[0] > 0)
+				const FVec3 LinearAxis0 = FVec3::CrossProduct(Xs[0] - Ps[0], Axis);
+				const FReal LinearAxisLen0Sq = LinearAxis0.SizeSquared();
+				if (LinearAxisLen0Sq > KINDA_SMALL_NUMBER)
 				{
-					const FVec3 LinearAxis0 = FVec3::CrossProduct(Xs[0] - Ps[0], Axis);
-					const FReal LinearAxisLen0Sq = LinearAxis0.SizeSquared();
-					if (LinearAxisLen0Sq > KINDA_SMALL_NUMBER)
-					{
-						const FReal IM0 = InvMs[0] / LinearAxisLen0Sq;
-						II0 = II0 * IM0 / (II0 + IM0);
-					}
+					const FReal IM0 = InvMs[0] / LinearAxisLen0Sq;
+					II0 = II0 * IM0 / (II0 + IM0);
 				}
-				if (InvMs[1] > 0)
+				const FVec3 LinearAxis1 = FVec3::CrossProduct(Xs[1] - Ps[1], Axis);
+				const FReal LinearAxisLen1Sq = LinearAxis1.SizeSquared();
+				if (LinearAxisLen1Sq > KINDA_SMALL_NUMBER)
 				{
-					const FVec3 LinearAxis1 = FVec3::CrossProduct(Xs[1] - Ps[1], Axis);
-					const FReal LinearAxisLen1Sq = LinearAxis1.SizeSquared();
-					if (LinearAxisLen1Sq > KINDA_SMALL_NUMBER)
-					{
-						const FReal IM1 = InvMs[1] / LinearAxisLen1Sq;
-						II1 = II1 * IM1 / (II1 + IM1);
-					}
+					const FReal IM1 = InvMs[1] / LinearAxisLen1Sq;
+					II1 = II1 * IM1 / (II1 + IM1);
 				}
 			}
 			const FReal II = (II0 + II1);
@@ -962,6 +1095,28 @@ namespace Chaos
 		}
 	}
 
+	void FJointSolverGaussSeidel::ApplyRotationConstraintSoft(
+		const FReal Dt,
+		const FReal Stiffness,
+		const FReal Damping,
+		const bool bAccelerationMode,
+		const FVec3& Axis,
+		const FReal Angle,
+		FReal& Lambda)
+	{
+		if (InvMs[0] == 0)
+		{
+			ApplyRotationConstraintSoftKD(0, 1, Dt, Stiffness, Damping, bAccelerationMode, Axis, Angle, Lambda);
+		}
+		else if (InvMs[1] == 0)
+		{
+			ApplyRotationConstraintSoftKD(1, 0, Dt, Stiffness, Damping, bAccelerationMode, Axis, Angle, Lambda);
+		}
+		else
+		{
+			ApplyRotationConstraintSoftDD(Dt, Stiffness, Damping, bAccelerationMode, Axis, Angle, Lambda);
+		}
+	}
 
 	//
 	//
@@ -970,7 +1125,7 @@ namespace Chaos
 	//
 
 
-	void FJointSolverGaussSeidel::ApplyTwistConstraint(
+	FJointSolverResult FJointSolverGaussSeidel::ApplyTwistConstraint(
 		const FReal Dt,
 		const FPBDJointSolverSettings& SolverSettings,
 		const FPBDJointSettings& JointSettings)
@@ -992,15 +1147,17 @@ namespace Chaos
 		}
 
 		// Apply twist correction
-		if (DTwistAngle != 0)
+		if (FMath::Abs(DTwistAngle) > AngleTolerance)
 		{
 			FReal TwistStiffness = FPBDJointUtilities::GetTwistStiffness(SolverSettings, JointSettings);
 			ApplyRotationConstraint(TwistStiffness, TwistAxis, DTwistAngle);
+			return FJointSolverResult::MakeActive();
 		}
+		return FJointSolverResult::MakeSolved();
 	}
 
 
-	void FJointSolverGaussSeidel::ApplyTwistConstraintSoft(
+	FJointSolverResult FJointSolverGaussSeidel::ApplyTwistConstraintSoft(
 		const FReal Dt,
 		const FPBDJointSolverSettings& SolverSettings,
 		const FPBDJointSettings& JointSettings)
@@ -1022,17 +1179,19 @@ namespace Chaos
 		}
 
 		// Apply twist correction
-		if (DTwistAngle != 0)
+		if (FMath::Abs(DTwistAngle) > AngleTolerance)
 		{
 			const FReal TwistStiffness = FPBDJointUtilities::GetSoftTwistStiffness(SolverSettings, JointSettings);
 			const FReal TwistDamping = FPBDJointUtilities::GetSoftTwistDamping(SolverSettings, JointSettings);
 			const bool bAccelerationMode = FPBDJointUtilities::GetAngularSoftAccelerationMode(SolverSettings, JointSettings);
 			ApplyRotationConstraintSoft(Dt, TwistStiffness, TwistDamping, bAccelerationMode, TwistAxis, DTwistAngle, TwistSoftLambda);
+			return FJointSolverResult::MakeActive();
 		}
+		return FJointSolverResult::MakeSolved();
 	}
 
 
-	void FJointSolverGaussSeidel::ApplyTwistDrive(
+	FJointSolverResult FJointSolverGaussSeidel::ApplyTwistDrive(
 		const FReal Dt,
 		const FPBDJointSolverSettings& SolverSettings,
 		const FPBDJointSettings& JointSettings)
@@ -1045,23 +1204,28 @@ namespace Chaos
 		const FReal DTwistAngle = TwistAngle - TwistAngleTarget;
 
 		// Apply twist correction
-		const FReal AngularDriveStiffness = FPBDJointUtilities::GetAngularTwistDriveStiffness(SolverSettings, JointSettings);
-		const FReal AngularDriveDamping = FPBDJointUtilities::GetAngularTwistDriveDamping(SolverSettings, JointSettings);
-		const bool bAccelerationMode = FPBDJointUtilities::GetDriveAccelerationMode(SolverSettings, JointSettings);
-		ApplyRotationConstraintSoft(Dt, AngularDriveStiffness, AngularDriveDamping, bAccelerationMode, TwistAxis, DTwistAngle, TwistDriveLambda);
+		if (FMath::Abs(DTwistAngle) > AngleTolerance)
+		{
+			const FReal AngularDriveStiffness = FPBDJointUtilities::GetAngularTwistDriveStiffness(SolverSettings, JointSettings);
+			const FReal AngularDriveDamping = FPBDJointUtilities::GetAngularTwistDriveDamping(SolverSettings, JointSettings);
+			const bool bAccelerationMode = FPBDJointUtilities::GetDriveAccelerationMode(SolverSettings, JointSettings);
+			ApplyRotationConstraintSoft(Dt, AngularDriveStiffness, AngularDriveDamping, bAccelerationMode, TwistAxis, DTwistAngle, TwistDriveLambda);
+		}
+		return FJointSolverResult::MakeSolved();
 	}
 
 
-	void FJointSolverGaussSeidel::ApplyTwistProjection(
+	FJointSolverResult FJointSolverGaussSeidel::ApplyTwistProjection(
 		const FReal Dt,
 		const FPBDJointSolverSettings& SolverSettings,
 		const FPBDJointSettings& JointSettings)
 	{
 		// @todo(ccaulfield): implement twist projection
+		return FJointSolverResult::MakeSolved();
 	}
 
 
-	void FJointSolverGaussSeidel::ApplyConeConstraint(
+	FJointSolverResult FJointSolverGaussSeidel::ApplyConeConstraint(
 		const FReal Dt,
 		const FPBDJointSolverSettings& SolverSettings,
 		const FPBDJointSettings& JointSettings)
@@ -1076,15 +1240,17 @@ namespace Chaos
 		FReal DSwingAngle = GetConeAngleError(JointSettings, SwingAxisLocal, SwingAngle);
 
 		// Apply swing correction to each body
-		if (DSwingAngle != 0)
+		if (FMath::Abs(DSwingAngle) > AngleTolerance)
 		{
 			FReal SwingStiffness = FPBDJointUtilities::GetSwingStiffness(SolverSettings, JointSettings);
 			ApplyRotationConstraint(SwingStiffness, SwingAxis, DSwingAngle);
+			return FJointSolverResult::MakeActive();
 		}
+		return FJointSolverResult::MakeSolved();
 	}
 
 
-	void FJointSolverGaussSeidel::ApplyConeConstraintSoft(
+	FJointSolverResult FJointSolverGaussSeidel::ApplyConeConstraintSoft(
 		const FReal Dt,
 		const FPBDJointSolverSettings& SolverSettings,
 		const FPBDJointSettings& JointSettings)
@@ -1099,17 +1265,19 @@ namespace Chaos
 		const FReal DSwingAngle = GetConeAngleError(JointSettings, SwingAxisLocal, SwingAngle);
 
 		// Apply swing correction to each body
-		if (DSwingAngle != 0)
+		if (FMath::Abs(DSwingAngle) > AngleTolerance)
 		{
 			const FReal SoftSwingStiffness = FPBDJointUtilities::GetSoftSwingStiffness(SolverSettings, JointSettings);
 			const FReal SoftSwingDamping = FPBDJointUtilities::GetSoftSwingDamping(SolverSettings, JointSettings);
 			const bool bAccelerationMode = FPBDJointUtilities::GetAngularSoftAccelerationMode(SolverSettings, JointSettings);
 			ApplyRotationConstraintSoft(Dt, SoftSwingStiffness, SoftSwingDamping, bAccelerationMode, SwingAxis, DSwingAngle, SwingSoftLambda);
+			return FJointSolverResult::MakeActive();
 		}
+		return FJointSolverResult::MakeSolved();
 	}
 
 
-	void FJointSolverGaussSeidel::ApplyConeDrive(
+	FJointSolverResult FJointSolverGaussSeidel::ApplyConeDrive(
 		const FReal Dt,
 		const FPBDJointSolverSettings& SolverSettings,
 		const FPBDJointSettings& JointSettings)
@@ -1128,14 +1296,18 @@ namespace Chaos
 		const FReal DSwingAngle = SwingAngle - SwingAngleTarget;
 
 		// Apply drive forces to each body
-		const FReal AngularDriveStiffness = FPBDJointUtilities::GetAngularSwingDriveStiffness(SolverSettings, JointSettings);
-		const FReal AngularDriveDamping = FPBDJointUtilities::GetAngularSwingDriveDamping(SolverSettings, JointSettings);
-		const bool bAccelerationMode = FPBDJointUtilities::GetDriveAccelerationMode(SolverSettings, JointSettings);
-		ApplyRotationConstraintSoft(Dt, AngularDriveStiffness, AngularDriveDamping, bAccelerationMode, SwingAxis, DSwingAngle, SwingDriveLambda);
+		if (FMath::Abs(DSwingAngle) > AngleTolerance)
+		{
+			const FReal AngularDriveStiffness = FPBDJointUtilities::GetAngularSwingDriveStiffness(SolverSettings, JointSettings);
+			const FReal AngularDriveDamping = FPBDJointUtilities::GetAngularSwingDriveDamping(SolverSettings, JointSettings);
+			const bool bAccelerationMode = FPBDJointUtilities::GetDriveAccelerationMode(SolverSettings, JointSettings);
+			ApplyRotationConstraintSoft(Dt, AngularDriveStiffness, AngularDriveDamping, bAccelerationMode, SwingAxis, DSwingAngle, SwingDriveLambda);
+		}
+		return FJointSolverResult::MakeSolved();
 	}
 
 
-	void FJointSolverGaussSeidel::ApplyConeProjection(
+	FJointSolverResult FJointSolverGaussSeidel::ApplyConeProjection(
 		const FReal Dt,
 		const FPBDJointSolverSettings& SolverSettings,
 		const FPBDJointSettings& JointSettings)
@@ -1150,7 +1322,7 @@ namespace Chaos
 		const FReal DSwingAngle = GetConeAngleError(JointSettings, SwingAxisLocal, SwingAngle);
 
 		// Apply rotation to body 1, set relative angular velocity to zero, if it would increase the error
-		if (DSwingAngle != 0)
+		if (FMath::Abs(DSwingAngle) > AngleTolerance)
 		{
 			const FReal AngularProjection = FPBDJointUtilities::GetAngularProjection(SolverSettings, JointSettings);
 
@@ -1184,11 +1356,13 @@ namespace Chaos
 			ApplyPositionDelta(1, AngularProjection, DP1);
 			ApplyVelocityDelta(1, AngularProjection, DV1, DW1);
 			UpdateDerivedState(1);
+			return FJointSolverResult::MakeActive();
 		}
+		return FJointSolverResult::MakeSolved();
 	}
 
 
-	void FJointSolverGaussSeidel::ApplySwingConstraint(
+	FJointSolverResult FJointSolverGaussSeidel::ApplySwingConstraint(
 		const FReal Dt,
 		const FPBDJointSolverSettings& SolverSettings,
 		const FPBDJointSettings& JointSettings,
@@ -1212,15 +1386,17 @@ namespace Chaos
 		}
 
 		// Apply swing correction
-		if (DSwingAngle != 0)
+		if (FMath::Abs(DSwingAngle) > AngleTolerance)
 		{
 			const FReal SwingStiffness = FPBDJointUtilities::GetSwingStiffness(SolverSettings, JointSettings);
 			ApplyRotationConstraint(SwingStiffness, SwingAxis, DSwingAngle);
+			return FJointSolverResult::MakeActive();
 		}
+		return FJointSolverResult::MakeSolved();
 	}
 
 
-	void FJointSolverGaussSeidel::ApplySwingConstraintSoft(
+	FJointSolverResult FJointSolverGaussSeidel::ApplySwingConstraintSoft(
 		const FReal Dt,
 		const FPBDJointSolverSettings& SolverSettings,
 		const FPBDJointSettings& JointSettings,
@@ -1245,17 +1421,19 @@ namespace Chaos
 		}
 
 		// Apply swing correction
-		if (DSwingAngle != 0)
+		if (FMath::Abs(DSwingAngle) > AngleTolerance)
 		{
 			const FReal SoftSwingStiffness = FPBDJointUtilities::GetSoftSwingStiffness(SolverSettings, JointSettings);
 			const FReal SoftSwingDamping = FPBDJointUtilities::GetSoftSwingDamping(SolverSettings, JointSettings);
 			const bool bAccelerationMode = FPBDJointUtilities::GetAngularSoftAccelerationMode(SolverSettings, JointSettings);
 			ApplyRotationConstraintSoft(Dt, SoftSwingStiffness, SoftSwingDamping, bAccelerationMode, SwingAxis, DSwingAngle, SwingSoftLambda);
+			return FJointSolverResult::MakeActive();
 		}
+		return FJointSolverResult::MakeSolved();
 	}
 
 
-	void FJointSolverGaussSeidel::ApplySwingDrive(
+	FJointSolverResult FJointSolverGaussSeidel::ApplySwingDrive(
 		const FReal Dt,
 		const FPBDJointSolverSettings& SolverSettings,
 		const FPBDJointSettings& JointSettings,
@@ -1270,14 +1448,18 @@ namespace Chaos
 		const FReal DSwingAngle = SwingAngle - SwingAngleTarget;
 
 		// Apply drive forces to each body
-		const FReal AngularDriveStiffness = FPBDJointUtilities::GetAngularSwingDriveStiffness(SolverSettings, JointSettings);
-		const FReal AngularDriveDamping = FPBDJointUtilities::GetAngularSwingDriveDamping(SolverSettings, JointSettings);
-		const bool bAccelerationMode = FPBDJointUtilities::GetDriveAccelerationMode(SolverSettings, JointSettings);
-		ApplyRotationConstraintSoft(Dt, AngularDriveStiffness, AngularDriveDamping, bAccelerationMode, SwingAxis, DSwingAngle, SwingDriveLambda);
+		if (FMath::Abs(DSwingAngle) > AngleTolerance)
+		{
+			const FReal AngularDriveStiffness = FPBDJointUtilities::GetAngularSwingDriveStiffness(SolverSettings, JointSettings);
+			const FReal AngularDriveDamping = FPBDJointUtilities::GetAngularSwingDriveDamping(SolverSettings, JointSettings);
+			const bool bAccelerationMode = FPBDJointUtilities::GetDriveAccelerationMode(SolverSettings, JointSettings);
+			ApplyRotationConstraintSoft(Dt, AngularDriveStiffness, AngularDriveDamping, bAccelerationMode, SwingAxis, DSwingAngle, SwingDriveLambda);
+		}
+		return FJointSolverResult::MakeSolved();
 	}
 
 
-	void FJointSolverGaussSeidel::ApplySwingProjection(
+	FJointSolverResult FJointSolverGaussSeidel::ApplySwingProjection(
 		const FReal Dt,
 		const FPBDJointSolverSettings& SolverSettings,
 		const FPBDJointSettings& JointSettings,
@@ -1285,10 +1467,11 @@ namespace Chaos
 		const EJointAngularAxisIndex SwingAxisIndex)
 	{
 		// @todo(ccaulfield): implement swing projection
+		return FJointSolverResult::MakeSolved();
 	}
 
 
-	void FJointSolverGaussSeidel::ApplySLerpDrive(
+	FJointSolverResult FJointSolverGaussSeidel::ApplySLerpDrive(
 		const FReal Dt,
 		const FPBDJointSolverSettings& SolverSettings,
 		const FPBDJointSettings& JointSettings)
@@ -1306,46 +1489,84 @@ namespace Chaos
 				SLerpAngle = SLerpAngle - (FReal)2 * PI;
 			}
 
-			const FReal AngularDriveStiffness = FPBDJointUtilities::GetAngularSLerpDriveStiffness(SolverSettings, JointSettings);
-			const FReal AngularDriveDamping = FPBDJointUtilities::GetAngularSLerpDriveDamping(SolverSettings, JointSettings);
-			const bool bAccelerationMode = FPBDJointUtilities::GetDriveAccelerationMode(SolverSettings, JointSettings);
-			ApplyRotationConstraintSoft(Dt, AngularDriveStiffness, AngularDriveDamping, bAccelerationMode, SLerpAxis, -SLerpAngle, SwingDriveLambda);
+			if (FMath::Abs(SLerpAngle) > AngleTolerance)
+			{
+				const FReal AngularDriveStiffness = FPBDJointUtilities::GetAngularSLerpDriveStiffness(SolverSettings, JointSettings);
+				const FReal AngularDriveDamping = FPBDJointUtilities::GetAngularSLerpDriveDamping(SolverSettings, JointSettings);
+				const bool bAccelerationMode = FPBDJointUtilities::GetDriveAccelerationMode(SolverSettings, JointSettings);
+				ApplyRotationConstraintSoft(Dt, AngularDriveStiffness, AngularDriveDamping, bAccelerationMode, SLerpAxis, -SLerpAngle, SwingDriveLambda);
+			}
 		}
+		return FJointSolverResult::MakeSolved();
 	}
 
+	// Kinematic-Dynamic bodies
+	FJointSolverResult FJointSolverGaussSeidel::ApplyPointPositionConstraintKD(
+		const int32 KIndex,
+		const int32 DIndex,
+		const FReal Dt,
+		const FPBDJointSolverSettings& SolverSettings,
+		const FPBDJointSettings& JointSettings)
+	{
+		check(InvMs[DIndex] > 0);
 
-	void FJointSolverGaussSeidel::ApplyPointPositionConstraint(
+		FReal LinearStiffness = FPBDJointUtilities::GetLinearStiffness(SolverSettings, JointSettings);
+		const FVec3 CX = Xs[DIndex] - Xs[KIndex];
+		if (CX.SizeSquared() > PositionTolerance * PositionTolerance)
+		{
+			if (bChaos_Joint_ISPC_Enabled)
+			{
+#if INTEL_ISPC
+				ispc::ApplyPointPositionConstraintKD((ispc::FJointSolverGaussSeidel*)this, KIndex, DIndex, (ispc::FVector&)CX, LinearStiffness);
+#endif
+			}
+			else
+			{
+				// Calculate constraint correction
+				FMatrix33 InvI1 = Utilities::ComputeWorldSpaceInertia(Qs[DIndex], InvILs[DIndex]);
+				FMatrix33 M = Utilities::ComputeJointFactorMatrix(Xs[DIndex] - Ps[DIndex], InvI1, InvMs[DIndex]);
+				FMatrix33 MI = M.Inverse();
+				const FVec3 DX = Utilities::Multiply(MI, CX);
+
+				// Apply constraint correction
+				const FVec3 DP1 = -InvMs[DIndex] * DX;
+				const FVec3 DR1 = Utilities::Multiply(InvI1, FVec3::CrossProduct(Xs[DIndex] - Ps[DIndex], -DX));
+
+				ApplyDelta(DIndex, LinearStiffness, DP1, DR1);
+				UpdateDerivedState(DIndex);
+			}
+			return FJointSolverResult::MakeActive();
+		}
+		return FJointSolverResult::MakeSolved();
+	}
+
+	// Dynamic-Dynamic bodies
+	FJointSolverResult FJointSolverGaussSeidel::ApplyPointPositionConstraintDD(
 		const FReal Dt,
 		const FPBDJointSolverSettings& SolverSettings,
 		const FPBDJointSettings& JointSettings)
 	{
+		check(InvMs[0] > 0);
+		check(InvMs[1] > 0);
+
+		FReal LinearStiffness = FPBDJointUtilities::GetLinearStiffness(SolverSettings, JointSettings);
 		const FVec3 CX = Xs[1] - Xs[0];
-		if (CX != FVec3(0))
+		if (CX.SizeSquared() > PositionTolerance * PositionTolerance)
 		{
 			if (bChaos_Joint_ISPC_Enabled)
 			{
 #if INTEL_ISPC
-				FReal LinearStiffness = FPBDJointUtilities::GetLinearStiffness(SolverSettings, JointSettings);
-				ispc::ApplyPointPositionConstraint((ispc::FJointSolverGaussSeidel*)this, (ispc::FVector&)CX, LinearStiffness);
+				ispc::ApplyPointPositionConstraintDD((ispc::FJointSolverGaussSeidel*)this, (ispc::FVector&)CX, LinearStiffness);
 #endif
 			}
 			else
 			{
-				const FMatrix33 InvI0 = Utilities::ComputeWorldSpaceInertia(Qs[0], InvILs[0]);
-				const FMatrix33 InvI1 = Utilities::ComputeWorldSpaceInertia(Qs[1], InvILs[1]);
-
 				// Calculate constraint correction
-				FMatrix33 M0 = FMatrix33(0, 0, 0);
-				FMatrix33 M1 = FMatrix33(0, 0, 0);
-				if (InvMs[0] > 0)
-				{
-					M0 = Utilities::ComputeJointFactorMatrix(Xs[0] - Ps[0], InvI0, InvMs[0]);
-				}
-				if (InvMs[1] > 0)
-				{
-					M1 = Utilities::ComputeJointFactorMatrix(Xs[1] - Ps[1], InvI1, InvMs[1]);
-				}
-				const FMatrix33 MI = (M0 + M1).Inverse();
+				FMatrix33 InvI0 = Utilities::ComputeWorldSpaceInertia(Qs[0], InvILs[0]);
+				FMatrix33 InvI1 = Utilities::ComputeWorldSpaceInertia(Qs[1], InvILs[1]);
+				FMatrix33 M0 = Utilities::ComputeJointFactorMatrix(Xs[0] - Ps[0], InvI0, InvMs[0]);
+				FMatrix33 M1 = Utilities::ComputeJointFactorMatrix(Xs[1] - Ps[1], InvI1, InvMs[1]);
+				FMatrix33 MI = (M0 + M1).Inverse();
 				const FVec3 DX = Utilities::Multiply(MI, CX);
 
 				// Apply constraint correction
@@ -1354,16 +1575,17 @@ namespace Chaos
 				const FVec3 DR0 = Utilities::Multiply(InvI0, FVec3::CrossProduct(Xs[0] - Ps[0], DX));
 				const FVec3 DR1 = Utilities::Multiply(InvI1, FVec3::CrossProduct(Xs[1] - Ps[1], -DX));
 
-				FReal LinearStiffness = FPBDJointUtilities::GetLinearStiffness(SolverSettings, JointSettings);
 				ApplyPositionDelta(LinearStiffness, DP0, DP1);
 				ApplyRotationDelta(LinearStiffness, DR0, DR1);
 				UpdateDerivedState();
 			}
+			return FJointSolverResult::MakeActive();
 		}
+		return FJointSolverResult::MakeSolved();
 	}
 
 
-	void FJointSolverGaussSeidel::ApplySphericalPositionConstraint(
+	FJointSolverResult FJointSolverGaussSeidel::ApplySphericalPositionConstraint(
 		const FReal Dt,
 		const FPBDJointSolverSettings& SolverSettings,
 		const FPBDJointSettings& JointSettings)
@@ -1373,7 +1595,7 @@ namespace Chaos
 		GetSphericalAxisDelta(Xs[0], Xs[1], Axis, Delta);
 
 		const FReal Error = FMath::Max((FReal)0, Delta - JointSettings.LinearLimit);
-		if (Error != 0)
+		if (FMath::Abs(Error) > PositionTolerance)
 		{
 			if (!JointSettings.bSoftLinearLimitsEnabled)
 			{
@@ -1387,11 +1609,13 @@ namespace Chaos
 				const bool bAccelerationMode = FPBDJointUtilities::GetLinearSoftAccelerationMode(SolverSettings, JointSettings);
 				ApplyPositionConstraintSoft(Dt, Stiffness, Damping, bAccelerationMode, Axis, Error, LinearSoftLambda);
 			}
+			return FJointSolverResult::MakeActive();
 		}
+		return FJointSolverResult::MakeSolved();
 	}
 
 
-	void FJointSolverGaussSeidel::ApplySphericalPositionDrive(
+	FJointSolverResult FJointSolverGaussSeidel::ApplySphericalPositionDrive(
 		const FReal Dt,
 		const FPBDJointSolverSettings& SolverSettings,
 		const FPBDJointSettings& JointSettings)
@@ -1400,17 +1624,19 @@ namespace Chaos
 		FVec3 Axis;
 		FReal Delta;
 		GetSphericalAxisDelta(XTarget, Xs[1], Axis, Delta);
-		if (Delta != 0)
+		if (FMath::Abs(Delta) > PositionTolerance)
 		{
 			const FReal Stiffness = FPBDJointUtilities::GetLinearDriveStiffness(SolverSettings, JointSettings);
 			const FReal Damping = FPBDJointUtilities::GetLinearDriveDamping(SolverSettings, JointSettings);
 			const bool bAccelerationMode = FPBDJointUtilities::GetDriveAccelerationMode(SolverSettings, JointSettings);
 			ApplyPositionConstraintSoft(Dt, Stiffness, Damping, bAccelerationMode, Axis, Delta, LinearDriveLambda);
+			return FJointSolverResult::MakeActive();
 		}
+		return FJointSolverResult::MakeSolved();
 	}
 
 
-	void FJointSolverGaussSeidel::ApplyCylindricalPositionConstraint(
+	FJointSolverResult FJointSolverGaussSeidel::ApplyCylindricalPositionConstraint(
 		const FReal Dt,
 		const int32 AxisIndex,
 		const EJointMotionType AxialMotion,
@@ -1424,28 +1650,32 @@ namespace Chaos
 		FReal AxialDelta, RadialDelta;
 		GetCylindricalAxesDeltas(Rs[0], Xs[0], Xs[1], AxisIndex, Axis, AxialDelta, RadialAxis, RadialDelta);
 		
+		FJointSolverResult NetResult;
+
 		if ((AxialMotion == EJointMotionType::Limited) && JointSettings.bSoftLinearLimitsEnabled)
 		{
 			// Soft Axial constraint
 			const FReal AxialLimit = JointSettings.LinearLimit;
-			if (FMath::Abs(AxialDelta) > AxialLimit)
+			if (FMath::Abs(AxialDelta) > AxialLimit + PositionTolerance)
 			{
 				const FReal AxialError = (AxialDelta > 0) ? AxialDelta - AxialLimit : AxialDelta + AxialLimit;
 				const FReal Stiffness = FPBDJointUtilities::GetSoftLinearStiffness(SolverSettings, JointSettings);
 				const FReal Damping = FPBDJointUtilities::GetSoftLinearDamping(SolverSettings, JointSettings);
 				const bool bAccelerationMode = FPBDJointUtilities::GetLinearSoftAccelerationMode(SolverSettings, JointSettings);
 				ApplyPositionConstraintSoft(Dt, Stiffness, Damping, bAccelerationMode, Axis, AxialError, LinearSoftLambda);
+				NetResult += FJointSolverResult::MakeActive();
 			}
 		}
 		else if (AxialMotion != EJointMotionType::Free)
 		{
 			// Hard Axial constraint
 			const FReal AxialLimit = (AxialMotion == EJointMotionType::Locked) ? 0 : JointSettings.LinearLimit;
-			if (FMath::Abs(AxialDelta) > AxialLimit)
+			if (FMath::Abs(AxialDelta) > AxialLimit + PositionTolerance)
 			{
 				const FReal AxialError = (AxialDelta > 0) ? AxialDelta - AxialLimit : AxialDelta + AxialLimit;
 				const FReal Stiffness = FPBDJointUtilities::GetLinearStiffness(SolverSettings, JointSettings);
 				ApplyPositionConstraint(Stiffness, Axis, AxialError);
+				NetResult += FJointSolverResult::MakeActive();
 			}
 		}
 
@@ -1453,30 +1683,34 @@ namespace Chaos
 		{
 			// Soft Radial constraint
 			const FReal RadialLimit = JointSettings.LinearLimit;
-			if (RadialDelta > RadialLimit)
+			if (RadialDelta > RadialLimit + PositionTolerance)
 			{
 				const FReal RadialError = FMath::Max((FReal)0, RadialDelta - RadialLimit);
 				const FReal Stiffness = FPBDJointUtilities::GetSoftLinearStiffness(SolverSettings, JointSettings);
 				const FReal Damping = FPBDJointUtilities::GetSoftLinearDamping(SolverSettings, JointSettings);
 				const bool bAccelerationMode = FPBDJointUtilities::GetLinearSoftAccelerationMode(SolverSettings, JointSettings);
 				ApplyPositionConstraintSoft(Dt, Stiffness, Damping, bAccelerationMode, RadialAxis, RadialError, LinearSoftLambda);
+				NetResult += FJointSolverResult::MakeActive();
 			}
 		}
 		else if (RadialMotion != EJointMotionType::Free)
 		{
 			// Hard Radial constraint
 			const FReal RadialLimit = (RadialMotion == EJointMotionType::Locked) ? 0 : JointSettings.LinearLimit;
-			if (RadialDelta > RadialLimit)
+			if (RadialDelta > RadialLimit + PositionTolerance)
 			{
 				const FReal RadialError = FMath::Max((FReal)0, RadialDelta - RadialLimit);
 				const FReal Stiffness = FPBDJointUtilities::GetLinearStiffness(SolverSettings, JointSettings);
 				ApplyPositionConstraint(Stiffness, RadialAxis, RadialError);
+				NetResult += FJointSolverResult::MakeActive();
 			}
 		}
+
+		return NetResult;
 	}
 
 
-	void FJointSolverGaussSeidel::ApplyCircularPositionDrive(
+	FJointSolverResult FJointSolverGaussSeidel::ApplyCircularPositionDrive(
 		const FReal Dt,
 		const int32 AxisIndex,
 		const FPBDJointSolverSettings& SolverSettings,
@@ -1486,17 +1720,18 @@ namespace Chaos
 		FVec3 Axis, RadialAxis;
 		FReal AxialDelta, RadialDelta;
 		GetCylindricalAxesDeltas(Rs[0], XTarget, Xs[1], AxisIndex, Axis, AxialDelta, RadialAxis, RadialDelta);
-		if (RadialDelta != 0)
+		if (RadialDelta > PositionTolerance)
 		{
 			const FReal Stiffness = FPBDJointUtilities::GetLinearDriveStiffness(SolverSettings, JointSettings);
 			const FReal Damping = FPBDJointUtilities::GetLinearDriveDamping(SolverSettings, JointSettings);
 			const bool bAccelerationMode = FPBDJointUtilities::GetDriveAccelerationMode(SolverSettings, JointSettings);
 			ApplyPositionConstraintSoft(Dt, Stiffness, Damping, bAccelerationMode, RadialAxis, RadialDelta, LinearDriveLambda);
 		}
+		return FJointSolverResult::MakeSolved();
 	}
 
 
-	void FJointSolverGaussSeidel::ApplyPlanarPositionConstraint(
+	FJointSolverResult FJointSolverGaussSeidel::ApplyPlanarPositionConstraint(
 		const FReal Dt,
 		const int32 AxisIndex,
 		const EJointMotionType AxialMotion,
@@ -1508,7 +1743,7 @@ namespace Chaos
 		GetPlanarAxisDelta(Rs[0], Xs[0], Xs[1], AxisIndex, Axis, Delta);
 
 		const FReal Limit = (AxialMotion == EJointMotionType::Locked) ? 0 : JointSettings.LinearLimit;
-		if (FMath::Abs(Delta) > Limit)
+		if (FMath::Abs(Delta) > Limit + PositionTolerance)
 		{
 			const FReal Error = (Delta > 0) ? Delta - Limit : Delta + Limit;
 			if ((AxialMotion == EJointMotionType::Limited) && JointSettings.bSoftLinearLimitsEnabled)
@@ -1523,11 +1758,13 @@ namespace Chaos
 				const FReal Stiffness = FPBDJointUtilities::GetLinearStiffness(SolverSettings, JointSettings);
 				ApplyPositionConstraint(Stiffness, Axis, Error);
 			}
+			return FJointSolverResult::MakeActive();
 		}
+		return FJointSolverResult::MakeSolved();
 	}
 
 
-	void FJointSolverGaussSeidel::ApplyAxialPositionDrive(
+	FJointSolverResult FJointSolverGaussSeidel::ApplyAxialPositionDrive(
 		const FReal Dt,
 		const int32 AxisIndex,
 		const FPBDJointSolverSettings& SolverSettings,
@@ -1537,17 +1774,18 @@ namespace Chaos
 		FVec3 Axis;
 		FReal Delta;
 		GetPlanarAxisDelta(Rs[0], XTarget, Xs[1], AxisIndex, Axis, Delta);
-		if (Delta != 0)
+		if (FMath::Abs(Delta) > PositionTolerance)
 		{
 			const FReal Stiffness = FPBDJointUtilities::GetLinearDriveStiffness(SolverSettings, JointSettings);
 			const FReal Damping = FPBDJointUtilities::GetLinearDriveDamping(SolverSettings, JointSettings);
 			const bool bAccelerationMode = FPBDJointUtilities::GetDriveAccelerationMode(SolverSettings, JointSettings);
 			ApplyPositionConstraintSoft(Dt, Stiffness, Damping, bAccelerationMode, Axis, Delta, LinearDriveLambda);
 		}
+		return FJointSolverResult::MakeSolved();
 	}
 
 
-	void FJointSolverGaussSeidel::ApplyPositionProjection(
+	FJointSolverResult FJointSolverGaussSeidel::ApplyPositionProjection(
 		const FReal Dt,
 		const FPBDJointSolverSettings& SolverSettings,
 		const FPBDJointSettings& JointSettings)
@@ -1558,7 +1796,7 @@ namespace Chaos
 			// Apply a position correction with the parent body set to infinite mass, then correct the velocity.
 			FVec3 CX = FPBDJointUtilities::GetLimitedPositionError(JointSettings, Rs[0], Xs[1] - Xs[0]);
 			const FReal CXLen = CX.Size();
-			if (CXLen > KINDA_SMALL_NUMBER)
+			if (CXLen > KINDA_SMALL_NUMBER + PositionTolerance)
 			{
 				const FReal ParentMassScale = FMath::Max(0.0f, 1.0f - LinearProjection);
 				if (bChaos_Joint_ISPC_Enabled)
@@ -1604,7 +1842,9 @@ namespace Chaos
 					ApplyVelocityDelta(Stiffness, DV0, DW0, DV1, DW1);
 					UpdateDerivedState();
 				}
+				return FJointSolverResult::MakeActive();
 			}
 		}
+		return FJointSolverResult::MakeSolved();
 	}
 }
diff --git a/Engine/Source/Runtime/Experimental/Chaos/Private/Chaos/Joint/PBDJointSolverGaussSeidel.ispc b/Engine/Source/Runtime/Experimental/Chaos/Private/Chaos/Joint/PBDJointSolverGaussSeidel.ispc
--- a/Engine/Source/Runtime/Experimental/Chaos/Private/Chaos/Joint/PBDJointSolverGaussSeidel.ispc
+++ b/Engine/Source/Runtime/Experimental/Chaos/Private/Chaos/Joint/PBDJointSolverGaussSeidel.ispc
@@ -45,6 +45,9 @@ struct FJointSolverGaussSeidel
 
 	// Post-angular constraint position fixup (to reduce iterations required for stiff angular constraints)
 	float AngularPositionCorrection;
+
+	float PositionTolerance;					// Distance error below which we consider a constraint or drive solved
+	float AngleTolerance;						// Angle error below which we consider a constraint or drive solved
 };
 
 static inline uniform FMatrix33 ComputeWorldSpaceInertia(const uniform FVector4 &Q, const uniform FVector &I)
@@ -102,6 +105,13 @@ static inline uniform FVector CalculateVelocity(const uniform FVector& P0, const
 	return (P1 - P0) / Dt;
 }
 
+
+static inline void UpdateDerivedState(uniform FJointSolverGaussSeidel * uniform M, const uniform int Index)
+{
+	M->Xs[Index] = M->Ps[Index] + VectorQuaternionRotateVector(M->Qs[Index], M->XLs[Index].Translation);
+	M->Rs[Index] = VectorQuaternionMultiply2(M->Qs[Index], M->XLs[Index].Rotation);
+}
+
 static inline void UpdateDerivedState(uniform FJointSolverGaussSeidel * uniform M)
 {
 	M->Xs[0] = M->Ps[0] + VectorQuaternionRotateVector(M->Qs[0], M->XLs[0].Translation);
@@ -110,6 +120,15 @@ static inline void UpdateDerivedState(uniform FJointSolverGaussSeidel * uniform
 	M->Rs[1] = VectorQuaternionMultiply2(M->Qs[1], M->XLs[1].Rotation);
 }
 
+static inline void ApplyPositionDelta(
+	uniform FJointSolverGaussSeidel * uniform M,
+	const uniform int Index,
+	const uniform float Stiffness,
+	const uniform FVector& DP)
+{
+	M->Ps[Index] = M->Ps[Index] + (Stiffness * DP);
+}
+
 static inline void ApplyPositionDelta(
 	uniform FJointSolverGaussSeidel * uniform M,
 	const uniform float Stiffness,
@@ -120,6 +139,22 @@ static inline void ApplyPositionDelta(
 	M->Ps[1] = M->Ps[1] + (Stiffness * DP1);
 }
 
+static inline void ApplyRotationDelta(
+	uniform FJointSolverGaussSeidel * uniform M,
+	const uniform int Index,
+	const uniform float Stiffness,
+	const uniform FVector& DR)
+{
+	const uniform FVector4 Qs = M->Qs[Index];
+	const uniform FVector4 DRStiffness = SetVector4(DR * Stiffness, 0);
+
+	const uniform FVector4 DQ = VectorQuaternionMultiply2(DRStiffness, Qs) * 0.5f;
+
+	const uniform FVector4 QsNorm = VectorNormalizeQuaternion(Qs + DQ);
+
+	M->Qs[1] = QsNorm;
+}
+
 static inline void ApplyRotationDelta(
 	uniform FJointSolverGaussSeidel * uniform M,
 	const uniform float Stiffness,
@@ -143,6 +178,17 @@ static inline void ApplyRotationDelta(
 	M->Qs[1] = Qs1ShortestArcWith;
 }
 
+static inline void ApplyDelta(
+	uniform FJointSolverGaussSeidel * uniform M,
+	const uniform int Index,
+	const uniform float Stiffness,
+	const uniform FVector& DP,
+	const uniform FVector& DR)
+{
+	ApplyPositionDelta(M, Index, Stiffness, DP);
+	ApplyRotationDelta(M, Index, Stiffness, DR);
+}
+
 static inline void ApplyVelocityDelta(
 	uniform FJointSolverGaussSeidel * uniform M,
 	const uniform float Stiffness,
@@ -211,7 +257,78 @@ export void ApplyPositionConstraintSoft(
 }
 
 // See "XPBD: Position-Based Simulation of Compliant Constrained Dynamics"
-export void ApplyRotationConstraintSoft(
+export void ApplyRotationConstraintSoftKD(
+	uniform FJointSolverGaussSeidel * uniform M,
+	const uniform int KIndex,
+	const uniform int DIndex,
+	const uniform float Dt,
+	const uniform float Stiffness,
+	const uniform float Damping,
+	const uniform bool bAccelerationMode,
+	const uniform FVector& Axis,
+	const uniform float Angle,
+	uniform float& Lambda)
+{
+	// World-space inverse mass
+	const uniform FMatrix33 InvI1 = ComputeWorldSpaceInertia(M->Qs[DIndex], M->InvILs[DIndex]);
+	const uniform FVector IA1 = Multiply(InvI1, Axis);
+
+	// Joint-space inverse mass
+	uniform float II1 = VectorDot(Axis, IA1);
+
+	// If we are correcting the position, we need to adjust the constraint effective mass using parallel axis theorem
+	// @todo(ccaulfield): the IMs here are constant per constraint...pre-build and cache it
+	if (M->AngularPositionCorrection > 0)
+	{
+		const uniform FVector LinearAxis1 = VectorCross(M->Xs[DIndex] - M->Ps[DIndex], Axis);
+		const uniform float LinearAxisLen1Sq = VectorSizeSquared(LinearAxis1);
+		if (LinearAxisLen1Sq > KINDA_SMALL_NUMBER)
+		{
+			const uniform float IM1 = M->InvMs[DIndex] / LinearAxisLen1Sq;
+			II1 = II1 * IM1 / (II1 + IM1);
+		}
+	}
+	const uniform float II = II1;
+
+	// Damping angular velocity
+	uniform float AngVelDt = 0;
+	if (Damping > KINDA_SMALL_NUMBER)
+	{
+		const uniform FVector W1 = CalculateAngularVelocity(M->PrevQs[DIndex], M->Qs[DIndex], 1.0f);
+		AngVelDt = -VectorDot(Axis, W1);
+	}
+
+	const uniform float SpringMassScale = (bAccelerationMode) ? 1.0f / II : 1.0f;
+	const uniform float S = SpringMassScale * Stiffness * Dt * Dt;
+	const uniform float D = SpringMassScale * Damping * Dt;
+	const uniform float Multiplier = 1.0f / ((S + D) * II + 1.0f);
+	const uniform float DLambda = Multiplier * (S * Angle - D * AngVelDt - Lambda);
+
+	//const FVec3 DR1 = IA1 * -DLambda;
+	const uniform FVector DR1 = Axis * -(DLambda * II1);
+
+	Lambda += DLambda;
+	ApplyRotationDelta(M, DIndex, 1.0f, DR1);
+
+	// Correct the positional error that was introduced by the rotation correction. This is
+	// correct when the position dofs are all locked, but not otherwise (fixable).
+	// This significantly improves angular stiffness at lower iterations, although the same effect
+	// is achieved by increasing iterations.
+	// @todo(ccaulfield): this position correction needs to have components in direction of inactive position constraints removed
+	if (M->AngularPositionCorrection > 0)
+	{
+		const uniform FVector PrevX0 = M->Xs[KIndex];
+		const uniform FVector PrevX1 = M->Xs[DIndex];
+		UpdateDerivedState(M);
+
+		const uniform FVector DP1 = (PrevX1 - PrevX0) - (M->Xs[DIndex] - M->Xs[KIndex]);
+		ApplyPositionDelta(M, DIndex, M->AngularPositionCorrection, DP1);
+	}
+
+	UpdateDerivedState(M, DIndex);
+}
+
+export void ApplyRotationConstraintSoftDD(
 	uniform FJointSolverGaussSeidel * uniform M,
 	const uniform float Dt,
 	const uniform float Stiffness,
@@ -235,25 +352,19 @@ export void ApplyRotationConstraintSoft(
 	// @todo(ccaulfield): the IMs here are constant per constraint...pre-build and cache it
 	if (M->AngularPositionCorrection > 0)
 	{
-		if (M->InvMs[0] > 0)
+		const uniform FVector LinearAxis0 = VectorCross(M->Xs[0] - M->Ps[0], Axis);
+		const uniform float LinearAxisLen0Sq = VectorSizeSquared(LinearAxis0);
+		if (LinearAxisLen0Sq > KINDA_SMALL_NUMBER)
 		{
-			const uniform FVector LinearAxis0 = VectorCross(M->Xs[0] - M->Ps[0], Axis);
-			const uniform float LinearAxisLen0Sq = VectorSizeSquared(LinearAxis0);
-			if (LinearAxisLen0Sq > KINDA_SMALL_NUMBER)
-			{
-				const uniform float IM0 = M->InvMs[0] / LinearAxisLen0Sq;
-				II0 = II0 * IM0 / (II0 + IM0);
-			}
+			const uniform float IM0 = M->InvMs[0] / LinearAxisLen0Sq;
+			II0 = II0 * IM0 / (II0 + IM0);
 		}
-		if (M->InvMs[1] > 0)
+		const uniform FVector LinearAxis1 = VectorCross(M->Xs[1] - M->Ps[1], Axis);
+		const uniform float LinearAxisLen1Sq = VectorSizeSquared(LinearAxis1);
+		if (LinearAxisLen1Sq > KINDA_SMALL_NUMBER)
 		{
-			const uniform FVector LinearAxis1 = VectorCross(M->Xs[1] - M->Ps[1], Axis);
-			const uniform float LinearAxisLen1Sq = VectorSizeSquared(LinearAxis1);
-			if (LinearAxisLen1Sq > KINDA_SMALL_NUMBER)
-			{
-				const uniform float IM1 = M->InvMs[1] / LinearAxisLen1Sq;
-				II1 = II1 * IM1 / (II1 + IM1);
-			}
+			const uniform float IM1 = M->InvMs[1] / LinearAxisLen1Sq;
+			II1 = II1 * IM1 / (II1 + IM1);
 		}
 	}
 	const uniform float II = (II0 + II1);
@@ -301,25 +412,38 @@ export void ApplyRotationConstraintSoft(
 	UpdateDerivedState(M);
 }
 
-export void ApplyPointPositionConstraint(
+export void ApplyPointPositionConstraintKD(
 	uniform FJointSolverGaussSeidel * uniform M,
+	const uniform int KIndex,
+	const uniform int DIndex,
 	const uniform FVector &CX,
 	const uniform float LinearStiffness)
 {
-	const uniform FMatrix33 InvI0 = ComputeWorldSpaceInertia(M->Qs[0], M->InvILs[0]);
-	const uniform FMatrix33 InvI1 = ComputeWorldSpaceInertia(M->Qs[1], M->InvILs[1]);
 
 	// Calculate constraint correction
-	uniform FMatrix33 M0 = Matrix33Zero;
-	uniform FMatrix33 M1 = Matrix33Zero;
-	if (M->InvMs[0] > 0)
-	{
-		M0 = ComputeJointFactorMatrix(M->Xs[0] - M->Ps[0], InvI0, M->InvMs[0]);
-	}
-	if (M->InvMs[1] > 0)
-	{
-		M1 = ComputeJointFactorMatrix(M->Xs[1] - M->Ps[1], InvI1, M->InvMs[1]);
-	}
+	const uniform FMatrix33 InvI1 = ComputeWorldSpaceInertia(M->Qs[DIndex], M->InvILs[DIndex]);
+	const uniform FMatrix33 M1 = ComputeJointFactorMatrix(M->Xs[DIndex] - M->Ps[DIndex], InvI1, M->InvMs[DIndex]);
+	const uniform FMatrix33 MI = MatrixInverse(M1);
+	const uniform FVector DX = Multiply(MI, CX);
+
+	// Apply constraint correction
+	const uniform FVector DP1 = DX * -M->InvMs[DIndex];
+	const uniform FVector DR1 = Multiply(InvI1, VectorCross(M->Xs[DIndex] - M->Ps[DIndex], VectorNegate(DX)));
+
+	ApplyDelta(M, DIndex, LinearStiffness, DP1, DR1);
+	UpdateDerivedState(M, DIndex);
+}
+
+export void ApplyPointPositionConstraintDD(
+	uniform FJointSolverGaussSeidel * uniform M,
+	const uniform FVector &CX,
+	const uniform float LinearStiffness)
+{
+	// Calculate constraint correction
+	const uniform FMatrix33 InvI0 = ComputeWorldSpaceInertia(M->Qs[0], M->InvILs[0]);
+	const uniform FMatrix33 InvI1 = ComputeWorldSpaceInertia(M->Qs[1], M->InvILs[1]);
+	uniform FMatrix33 M0 = ComputeJointFactorMatrix(M->Xs[0] - M->Ps[0], InvI0, M->InvMs[0]);
+	uniform FMatrix33 M1 = ComputeJointFactorMatrix(M->Xs[1] - M->Ps[1], InvI1, M->InvMs[1]);
 	const uniform FMatrix33 MI = MatrixInverse(AddAB(M0, M1));
 	const uniform FVector DX = Multiply(MI, CX);
 
diff --git a/Engine/Source/Runtime/Experimental/Chaos/Private/Chaos/PBDCollisionConstraintsContact.cpp b/Engine/Source/Runtime/Experimental/Chaos/Private/Chaos/PBDCollisionConstraintsContact.cpp
--- a/Engine/Source/Runtime/Experimental/Chaos/Private/Chaos/PBDCollisionConstraintsContact.cpp
+++ b/Engine/Source/Runtime/Experimental/Chaos/Private/Chaos/PBDCollisionConstraintsContact.cpp
@@ -13,6 +13,9 @@ namespace Chaos
 {
 	namespace Collisions
 	{
+		int32 Chaos_Collision_EnergyClampEnabled = 1;
+		FAutoConsoleVariableRef CVarChaosImmPhysStepTime(TEXT("p.Chaos.Collision.EnergyClampEnabled"), Chaos_Collision_EnergyClampEnabled, TEXT("Whether to use energy clamping in collision apply step"));
+
 		extern void UpdateManifold(FRigidBodyMultiPointContactConstraint& Constraint, const FReal CullDistance)
 		{
 			const FRigidTransform3 Transform0 = GetTransform(Constraint.Particle[0]);
@@ -168,7 +171,10 @@ namespace Chaos
 					Impulse = ImpulseNumerator / ImpulseDenominator;
 				}
 
-				Impulse = GetEnergyClampedImpulse(Particle0->CastToRigidParticle(), Particle1->CastToRigidParticle(), Impulse, VectorToPoint1, VectorToPoint2, Body1Velocity, Body2Velocity);
+				if (Chaos_Collision_EnergyClampEnabled != 0)
+				{
+					Impulse = GetEnergyClampedImpulse(Particle0->CastToRigidParticle(), Particle1->CastToRigidParticle(), Impulse, VectorToPoint1, VectorToPoint2, Body1Velocity, Body2Velocity);
+				}
 				AccumulatedImpulse += Impulse;
 
 				if (bIsRigidDynamic0)
@@ -335,7 +341,10 @@ namespace Chaos
 				}
 
 				FVec3 VelocityFixImpulse = ImpulseNumerator / ImpulseDenominator;
-				VelocityFixImpulse = GetEnergyClampedImpulse(Particle0->CastToRigidParticle(), Particle1->CastToRigidParticle(), VelocityFixImpulse, VectorToPoint1, VectorToPoint2, Body1Velocity, Body2Velocity);
+				if (Chaos_Collision_EnergyClampEnabled)
+				{
+					VelocityFixImpulse = GetEnergyClampedImpulse(Particle0->CastToRigidParticle(), Particle1->CastToRigidParticle(), VelocityFixImpulse, VectorToPoint1, VectorToPoint2, Body1Velocity, Body2Velocity);
+				}
 				AccumulatedImpulse += VelocityFixImpulse;	//question: should we track this?
 				if (!IsTemporarilyStatic0 && bIsRigidDynamic0)
 				{
diff --git a/Engine/Source/Runtime/Experimental/Chaos/Private/Chaos/PBDJointConstraints.cpp b/Engine/Source/Runtime/Experimental/Chaos/Private/Chaos/PBDJointConstraints.cpp
--- a/Engine/Source/Runtime/Experimental/Chaos/Private/Chaos/PBDJointConstraints.cpp
+++ b/Engine/Source/Runtime/Experimental/Chaos/Private/Chaos/PBDJointConstraints.cpp
@@ -14,6 +14,9 @@
 
 //#pragma optimize("", off)
 
+bool bChaos_Joint_EarlyOut_Enabled = true;
+FAutoConsoleVariableRef CVarChaosJointEarlyOutEnabled(TEXT("p.Chaos.Joint.EarlyOut"), bChaos_Joint_EarlyOut_Enabled, TEXT("Whether to iterating when joints report being solved"));
+
 namespace Chaos
 {
 	DECLARE_CYCLE_STAT(TEXT("Joints::Sort"), STAT_Joints_Sort, STATGROUP_ChaosJoint);
@@ -154,6 +157,8 @@ namespace Chaos
 		: ApplyPairIterations(1)
 		, ApplyPushOutPairIterations(1)
 		, SwingTwistAngleTolerance(1.0e-6f)
+		, PositionTolerance(0)
+		, AngleTolerance(0)
 		, MinParentMassRatio(0)
 		, MaxInertiaRatio(0)
 		, AngularConstraintPositionCorrection(1.0f)
@@ -494,14 +499,17 @@ namespace Chaos
 			PreApplyCallback(Dt, Handles);
 		}
 
+		FJointSolverResult NetResult;
 		if (Settings.ApplyPairIterations > 0)
 		{
 			for (int32 ConstraintIndex = 0; ConstraintIndex < NumConstraints(); ++ConstraintIndex)
 			{
-				SolvePosition_GaussSiedel(Dt, ConstraintIndex, Settings.ApplyPairIterations, It, NumIts);
+				NetResult += SolvePosition_GaussSiedel(Dt, ConstraintIndex, Settings.ApplyPairIterations, It, NumIts);
 			}
 		}
 
+		//UE_LOG(LogChaosJoint, Warning, TEXT("Apply %d / %d Active %d / %d"), It, NumIts, NetResult.GetNumActive(), NetResult.GetNumActive() + NetResult.GetNumSolved());
+
 		if (PostApplyCallback != nullptr)
 		{
 			PostApplyCallback(Dt, Handles);
@@ -515,14 +523,17 @@ namespace Chaos
 		// @todo(ccaulfield): track whether we are sufficiently solved
 		bool bNeedsAnotherIteration = true;
 
+		FJointSolverResult NetResult;
 		if (Settings.ApplyPushOutPairIterations > 0)
 		{
 			for (int32 ConstraintIndex = 0; ConstraintIndex < NumConstraints(); ++ConstraintIndex)
 			{
-				ProjectPosition_GaussSiedel(Dt, ConstraintIndex, Settings.ApplyPushOutPairIterations, It, NumIts);
+				NetResult += ProjectPosition_GaussSiedel(Dt, ConstraintIndex, Settings.ApplyPushOutPairIterations, It, NumIts);
 			}
 		}
 
+		//UE_LOG(LogChaosJoint, Warning, TEXT("Push  %d / %d Active %d / %d"), It, NumIts, NetResult.GetNumActive(), NetResult.GetNumActive() + NetResult.GetNumSolved());
+
 		if (PostProjectCallback != nullptr)
 		{
 			PostProjectCallback(Dt, Handles);
@@ -655,7 +666,7 @@ namespace Chaos
 
 	// This position solver iterates over each of the inner constraints (position, twist, swing) and solves them independently.
 	// This will converge slowly in some cases, particularly where resolving angular constraints violates position constraints and vice versa.
-	void FPBDJointConstraints::SolvePosition_GaussSiedel(const FReal Dt, const int32 ConstraintIndex, const int32 NumPairIts, const int32 It, const int32 NumIts)
+	FJointSolverResult FPBDJointConstraints::SolvePosition_GaussSiedel(const FReal Dt, const int32 ConstraintIndex, const int32 NumPairIts, const int32 It, const int32 NumIts)
 	{
 		const TVector<TGeometryParticleHandle<FReal, 3>*, 2>& Constraint = ConstraintParticles[ConstraintIndex];
 		UE_LOG(LogChaosJoint, VeryVerbose, TEXT("Solve Joint Constraint %d %s %s (dt = %f; it = %d / %d)"), ConstraintIndex, *Constraint[0]->ToString(), *Constraint[1]->ToString(), Dt, It, NumIts);
@@ -679,17 +690,25 @@ namespace Chaos
 			Particle1->V(),
 			Particle1->W());
 
+		FJointSolverResult NetResult;
 		for (int32 PairIt = 0; PairIt < NumPairIts; ++PairIt)
 		{
-			Solver.ApplyConstraints(Dt, Settings, JointSettings);
-			Solver.ApplyDrives(Dt, Settings, JointSettings);
+			NetResult += Solver.ApplyConstraints(Dt, Settings, JointSettings);
+			NetResult +=  Solver.ApplyDrives(Dt, Settings, JointSettings);
+
+			if ((NetResult.GetNumActive() == 0) && bChaos_Joint_EarlyOut_Enabled)
+			{
+				break;
+			}
 		}
 
 		UpdateParticleState(Particle0->CastToRigidParticle(), Dt, Solver.GetP(0), Solver.GetQ(0));
 		UpdateParticleState(Particle1->CastToRigidParticle(), Dt, Solver.GetP(1), Solver.GetQ(1));
+
+		return NetResult;
 	}
 
-	void FPBDJointConstraints::ProjectPosition_GaussSiedel(const FReal Dt, const int32 ConstraintIndex, const int32 NumPairIts, const int32 It, const int32 NumIts)
+	FJointSolverResult FPBDJointConstraints::ProjectPosition_GaussSiedel(const FReal Dt, const int32 ConstraintIndex, const int32 NumPairIts, const int32 It, const int32 NumIts)
 	{
 		const TVector<TGeometryParticleHandle<FReal, 3>*, 2>& Constraint = ConstraintParticles[ConstraintIndex];
 		UE_LOG(LogChaosJoint, VeryVerbose, TEXT("Solve Joint Constraint %d %s %s (dt = %f; it = %d / %d)"), ConstraintIndex, *Constraint[0]->ToString(), *Constraint[1]->ToString(), Dt, It, NumIts);
@@ -713,13 +732,21 @@ namespace Chaos
 			Particle1->V(),
 			Particle1->W());
 
+		FJointSolverResult NetResult;
 		for (int32 PairIt = 0; PairIt < NumPairIts; ++PairIt)
 		{
-			Solver.ApplyProjections(Dt, Settings, JointSettings);
+			NetResult = Solver.ApplyProjections(Dt, Settings, JointSettings);
+			
+			if ((NetResult.GetNumActive() == 0) && bChaos_Joint_EarlyOut_Enabled)
+			{
+				break;
+			}
 		}
 
 		UpdateParticleState(Particle0->CastToRigidParticle(), Dt, Solver.GetP(0), Solver.GetQ(0), Solver.GetV(0), Solver.GetW(0));
 		UpdateParticleState(Particle1->CastToRigidParticle(), Dt, Solver.GetP(1), Solver.GetQ(1), Solver.GetV(1), Solver.GetW(1));
+
+		return NetResult;
 	}
 }
 
diff --git a/Engine/Source/Runtime/Experimental/Chaos/Public/Chaos/Joint/PBDJointSolverGaussSeidel.h b/Engine/Source/Runtime/Experimental/Chaos/Public/Chaos/Joint/PBDJointSolverGaussSeidel.h
--- a/Engine/Source/Runtime/Experimental/Chaos/Public/Chaos/Joint/PBDJointSolverGaussSeidel.h
+++ b/Engine/Source/Runtime/Experimental/Chaos/Public/Chaos/Joint/PBDJointSolverGaussSeidel.h
@@ -11,6 +11,44 @@
 
 namespace Chaos
 {
+	struct FJointSolverResult
+	{
+	public:
+		static FJointSolverResult MakeActive() { return { 1, 0 }; }
+		static FJointSolverResult MakeSolved() { return { 0, 1 }; }
+
+		FJointSolverResult()
+			: NumActive(0)
+			, NumSolved(0)
+		{
+		}
+		
+		int32 GetNumActive() const { return NumActive; }
+		int32 GetNumSolved() const { return NumSolved; }
+		
+		FJointSolverResult& operator+=(const FJointSolverResult& R)
+		{
+			NumActive += R.NumActive;
+			NumSolved += R.NumSolved;
+			return *this;
+		}
+
+		friend FJointSolverResult operator+(const FJointSolverResult& L, const FJointSolverResult& R)
+		{
+			return { L.NumActive + R.NumActive, L.NumSolved + R.NumSolved };
+		}
+
+	private:
+		FJointSolverResult(int32 InActive, int32 InSolved)
+			: NumActive(InActive)
+			, NumSolved(InSolved)
+		{
+		}
+
+		int32 NumActive;
+		int32 NumSolved;
+	};
+
 	/**
 	 * Calculate new positions and rotations for a pair of bodies connected by a joint.
 	 *
@@ -76,17 +114,17 @@ namespace Chaos
 			const FVec3& V1,
 			const FVec3& W1);
 
-		void ApplyConstraints(
+		FJointSolverResult ApplyConstraints(
 			const FReal Dt,
 			const FPBDJointSolverSettings& SolverSettings,
 			const FPBDJointSettings& JointSettings);
 
-		void ApplyDrives(
+		FJointSolverResult ApplyDrives(
 			const FReal Dt,
 			const FPBDJointSolverSettings& SolverSettings,
 			const FPBDJointSettings& JointSettings);
 
-		void ApplyProjections(
+		FJointSolverResult ApplyProjections(
 			const FReal Dt,
 			const FPBDJointSolverSettings& SolverSettings,
 			const FPBDJointSettings& JointSettings);
@@ -96,27 +134,27 @@ namespace Chaos
 		void UpdateDerivedState(const int32 BodyIndex);
 		void UpdateDerivedState();
 
-		void ApplyPositionConstraints(
+		FJointSolverResult ApplyPositionConstraints(
 			const FReal Dt,
 			const FPBDJointSolverSettings& SolverSettings,
 			const FPBDJointSettings& JointSettings);
 
-		void ApplyRotationConstraints(
+		FJointSolverResult ApplyRotationConstraints(
 			const FReal Dt,
 			const FPBDJointSolverSettings& SolverSettings,
 			const FPBDJointSettings& JointSettings);
 
-		void ApplyPositionDrives(
+		FJointSolverResult ApplyPositionDrives(
 			const FReal Dt,
 			const FPBDJointSolverSettings& SolverSettings,
 			const FPBDJointSettings& JointSettings);
 
-		void ApplyRotationDrives(
+		FJointSolverResult ApplyRotationDrives(
 			const FReal Dt,
 			const FPBDJointSolverSettings& SolverSettings,
 			const FPBDJointSettings& JointSettings);
 
-		void ApplyRotationProjection(
+		FJointSolverResult ApplyRotationProjection(
 			const FReal Dt,
 			const FPBDJointSolverSettings& SolverSettings,
 			const FPBDJointSettings& JointSettings);
@@ -142,6 +180,12 @@ namespace Chaos
 			const FVec3& DR0,
 			const FVec3& DR1);
 
+		void ApplyDelta(
+			const int32 BodyIndex,
+			const FReal Stiffness,
+			const FVec3& DP,
+			const FVec3& DR1);
+
 		void ApplyVelocityDelta(
 			const int32 BodyIndex,
 			const FReal Stiffness,
@@ -183,95 +227,122 @@ namespace Chaos
 			const FReal Angle,
 			FReal& Lambda);
 
-		void ApplyTwistConstraint(
+		void ApplyRotationConstraintSoftKD(
+			const int32 KIndex,
+			const int32 DIndex,
+			const FReal Dt,
+			const FReal Stiffness,
+			const FReal Damping,
+			const bool bAccelerationMode,
+			const FVec3& Axis,
+			const FReal Angle,
+			FReal& Lambda);
+
+		void ApplyRotationConstraintSoftDD(
+			const FReal Dt,
+			const FReal Stiffness,
+			const FReal Damping,
+			const bool bAccelerationMode,
+			const FVec3& Axis,
+			const FReal Angle,
+			FReal& Lambda);
+
+		FJointSolverResult ApplyTwistConstraint(
 			const FReal Dt,
 			const FPBDJointSolverSettings& SolverSettings,
 			const FPBDJointSettings& JointSettings);
 
-		void ApplyTwistConstraintSoft(
+		FJointSolverResult ApplyTwistConstraintSoft(
 			const FReal Dt,
 			const FPBDJointSolverSettings& SolverSettings,
 			const FPBDJointSettings& JointSettings);
 
-		void ApplyTwistDrive(
+		FJointSolverResult ApplyTwistDrive(
 			const FReal Dt,
 			const FPBDJointSolverSettings& SolverSettings,
 			const FPBDJointSettings& JointSettings);
 
-		void ApplyTwistProjection(
+		FJointSolverResult ApplyTwistProjection(
 			const FReal Dt,
 			const FPBDJointSolverSettings& SolverSettings,
 			const FPBDJointSettings& JointSettings);
 
-		void ApplyConeConstraint(
+		FJointSolverResult ApplyConeConstraint(
 			const FReal Dt,
 			const FPBDJointSolverSettings& SolverSettings,
 			const FPBDJointSettings& JointSettings);
 
-		void ApplyConeConstraintSoft(
+		FJointSolverResult ApplyConeConstraintSoft(
 			const FReal Dt,
 			const FPBDJointSolverSettings& SolverSettings,
 			const FPBDJointSettings& JointSettings);
 
-		void ApplyConeDrive(
+		FJointSolverResult ApplyConeDrive(
 			const FReal Dt,
 			const FPBDJointSolverSettings& SolverSettings,
 			const FPBDJointSettings& JointSettings);
 
-		void ApplyConeProjection(
+		FJointSolverResult ApplyConeProjection(
 			const FReal Dt,
 			const FPBDJointSolverSettings& SolverSettings,
 			const FPBDJointSettings& JointSettings);
 
-		void ApplySwingConstraint(
+		FJointSolverResult ApplySwingConstraint(
 			const FReal Dt,
 			const FPBDJointSolverSettings& SolverSettings,
 			const FPBDJointSettings& JointSettings,
 			const EJointAngularConstraintIndex SwingConstraintIndex,
 			const EJointAngularAxisIndex SwingAxisIndex);
 
-		void ApplySwingConstraintSoft(
+		FJointSolverResult ApplySwingConstraintSoft(
 			const FReal Dt,
 			const FPBDJointSolverSettings& SolverSettings,
 			const FPBDJointSettings& JointSettings,
 			const EJointAngularConstraintIndex SwingConstraintIndex,
 			const EJointAngularAxisIndex SwingAxisIndex);
 
-		void ApplySwingDrive(
+		FJointSolverResult ApplySwingDrive(
 			const FReal Dt,
 			const FPBDJointSolverSettings& SolverSettings,
 			const FPBDJointSettings& JointSettings,
 			const EJointAngularConstraintIndex SwingConstraintIndex,
 			const EJointAngularAxisIndex SwingAxisIndex);
 
-		void ApplySwingProjection(
+		FJointSolverResult ApplySwingProjection(
 			const FReal Dt,
 			const FPBDJointSolverSettings& SolverSettings,
 			const FPBDJointSettings& JointSettings,
 			const EJointAngularConstraintIndex SwingConstraintIndex,
 			const EJointAngularAxisIndex SwingAxisIndex);
 
-		void ApplySLerpDrive(
+		FJointSolverResult ApplySLerpDrive(
 			const FReal Dt,
 			const FPBDJointSolverSettings& SolverSettings,
 			const FPBDJointSettings& JointSettings);
 
-		void ApplyPointPositionConstraint(
+		FJointSolverResult ApplyPointPositionConstraintKD(
+			const int32 KIndex,
+			const int32 DIndex,
 			const FReal Dt,
 			const FPBDJointSolverSettings& SolverSettings,
 			const FPBDJointSettings& JointSettings);
 
-		void ApplySphericalPositionConstraint(
+		FJointSolverResult ApplyPointPositionConstraintDD(
 			const FReal Dt,
 			const FPBDJointSolverSettings& SolverSettings,
 			const FPBDJointSettings& JointSettings);
 
-		void ApplySphericalPositionDrive(
+		FJointSolverResult ApplySphericalPositionConstraint(
 			const FReal Dt,
 			const FPBDJointSolverSettings& SolverSettings,
 			const FPBDJointSettings& JointSettings);
 
-		void ApplyCylindricalPositionConstraint(
+		FJointSolverResult ApplySphericalPositionDrive(
+			const FReal Dt,
+			const FPBDJointSolverSettings& SolverSettings,
+			const FPBDJointSettings& JointSettings);
+
+		FJointSolverResult ApplyCylindricalPositionConstraint(
 			const FReal Dt,
 			const int32 AxisIndex,
 			const EJointMotionType AxialMotion,
@@ -279,26 +350,26 @@ namespace Chaos
 			const FPBDJointSolverSettings& SolverSettings,
 			const FPBDJointSettings& JointSettings);
 
-		void ApplyCircularPositionDrive(
+		FJointSolverResult ApplyCircularPositionDrive(
 			const FReal Dt,
 			const int32 AxisIndex,
 			const FPBDJointSolverSettings& SolverSettings,
 			const FPBDJointSettings& JointSettings);
 
-		void ApplyPlanarPositionConstraint(
+		FJointSolverResult ApplyPlanarPositionConstraint(
 			const FReal Dt,
 			const int32 AxisIndex,
 			const EJointMotionType AxialMotion,
 			const FPBDJointSolverSettings& SolverSettings,
 			const FPBDJointSettings& JointSettings);
 
-		void ApplyAxialPositionDrive(
+		FJointSolverResult ApplyAxialPositionDrive(
 			const FReal Dt,
 			const int32 AxisIndex,
 			const FPBDJointSolverSettings& SolverSettings,
 			const FPBDJointSettings& JointSettings);
 
-		void ApplyPositionProjection(
+		FJointSolverResult ApplyPositionProjection(
 			const FReal Dt,
 			const FPBDJointSolverSettings& SolverSettings,
 			const FPBDJointSettings& JointSettings);
@@ -332,6 +403,9 @@ namespace Chaos
 		
 		// Post-angular constraint position fixup (to reduce iterations required for stiff angular constraints)
 		FReal AngularPositionCorrection;
+
+		FReal PositionTolerance;					// Distance error below which we consider a constraint or drive solved
+		FReal AngleTolerance;						// Angle error below which we consider a constraint or drive solved
 	};
 
 }
\ No newline at end of file
diff --git a/Engine/Source/Runtime/Experimental/Chaos/Public/Chaos/PBDJointConstraintTypes.h b/Engine/Source/Runtime/Experimental/Chaos/Public/Chaos/PBDJointConstraintTypes.h
--- a/Engine/Source/Runtime/Experimental/Chaos/Public/Chaos/PBDJointConstraintTypes.h
+++ b/Engine/Source/Runtime/Experimental/Chaos/Public/Chaos/PBDJointConstraintTypes.h
@@ -126,6 +126,8 @@ namespace Chaos
 
 		// Tolerances
 		FReal SwingTwistAngleTolerance;
+		FReal PositionTolerance;
+		FReal AngleTolerance;
 
 		// Stability control
 		FReal MinParentMassRatio;
diff --git a/Engine/Source/Runtime/Experimental/Chaos/Public/Chaos/PBDJointConstraints.h b/Engine/Source/Runtime/Experimental/Chaos/Public/Chaos/PBDJointConstraints.h
--- a/Engine/Source/Runtime/Experimental/Chaos/Public/Chaos/PBDJointConstraints.h
+++ b/Engine/Source/Runtime/Experimental/Chaos/Public/Chaos/PBDJointConstraints.h
@@ -169,9 +169,8 @@ namespace Chaos
 		void UpdateParticleState(TPBDRigidParticleHandle<FReal, 3>* Rigid, const FReal Dt, const FVec3& P, const FRotation3& Q, const FVec3& V, const FVec3& W);
 		void SortConstraints();
 
-		void SolvePosition_Cholesky(const FReal Dt, const int32 ConstraintIndex, const int32 NumPairIts, const int32 It, const int32 NumIts);
-		void SolvePosition_GaussSiedel(const FReal Dt, const int32 ConstraintIndex, const int32 NumPairIts, const int32 It, const int32 NumIts);
-		void ProjectPosition_GaussSiedel(const FReal Dt, const int32 ConstraintIndex, const int32 NumPairIts, const int32 It, const int32 NumIts);
+		FJointSolverResult SolvePosition_GaussSiedel(const FReal Dt, const int32 ConstraintIndex, const int32 NumPairIts, const int32 It, const int32 NumIts);
+		FJointSolverResult ProjectPosition_GaussSiedel(const FReal Dt, const int32 ConstraintIndex, const int32 NumPairIts, const int32 It, const int32 NumIts);
 
 		FPBDJointSolverSettings Settings;
 
