Chaos : LLM

- LLM wrappers for NewData allocations.

diff --git a/Engine/Source/Runtime/Experimental/ChaosSolvers/Private/PBDRigidsSolver.cpp b/Engine/Source/Runtime/Experimental/ChaosSolvers/Private/PBDRigidsSolver.cpp
--- a/Engine/Source/Runtime/Experimental/ChaosSolvers/Private/PBDRigidsSolver.cpp
+++ b/Engine/Source/Runtime/Experimental/ChaosSolvers/Private/PBDRigidsSolver.cpp
@@ -169,6 +169,8 @@ namespace Chaos
 
 	void FPBDRigidsSolver::RegisterObject(TGeometryParticle<float, 3>* GTParticle)
 	{
+		LLM_SCOPE(ELLMTag::Chaos);
+
 		UE_LOG(LogPBDRigidsSolverSolver, Verbose, TEXT("FPBDRigidsSolver::RegisterObject()"));
 
 		// Make sure this particle doesn't already have a proxy
@@ -469,6 +471,8 @@ namespace Chaos
 	template<typename ProxyType>
 	void PushPhysicsStateExec(FPBDRigidsSolver * Solver, ProxyType* Proxy, Chaos::IDispatcher* Dispatcher)
 	{
+		LLM_SCOPE(ELLMTag::Chaos);
+
 		if (Chaos::FParticleData* ProxyData = Proxy->NewData())
 		{
 			if (auto* RigidHandle = static_cast<Chaos::TGeometryParticleHandle<float, 3>*>(Proxy->GetHandle()))
