Chaos
- Integrate optimization
- Remove collision constraint handles for RBAN

diff --git a/Engine/Source/Runtime/Engine/Private/PhysicsEngine/ImmediatePhysics/ImmediatePhysicsChaos/ImmediatePhysicsSimulation_Chaos.cpp b/Engine/Source/Runtime/Engine/Private/PhysicsEngine/ImmediatePhysics/ImmediatePhysicsChaos/ImmediatePhysicsSimulation_Chaos.cpp
--- a/Engine/Source/Runtime/Engine/Private/PhysicsEngine/ImmediatePhysics/ImmediatePhysicsChaos/ImmediatePhysicsSimulation_Chaos.cpp
+++ b/Engine/Source/Runtime/Engine/Private/PhysicsEngine/ImmediatePhysics/ImmediatePhysicsChaos/ImmediatePhysicsSimulation_Chaos.cpp
@@ -189,6 +189,8 @@ namespace ImmediatePhysics_Chaos
 		Particles.GetParticleHandles().AddArray(&ParticleMaterials);
 		Particles.GetParticleHandles().AddArray(&PerParticleMaterials);
 
+		Collisions.DisableHandles();
+
 		Evolution.AddConstraintRule(&CollisionsRule);
 		Evolution.AddConstraintRule(&JointsRule);
 
@@ -216,7 +218,7 @@ namespace ImmediatePhysics_Chaos
 			{
 				if (ChaosImmediate_DebugDrawCollisions == 4)
 				{
-					DebugDraw::DrawCollisions(SimulationSpaceTransform, InConstraintHandles, 0.3f);
+					DebugDraw::DrawCollisions(SimulationSpaceTransform, Collisions, 0.3f);
 				}
 				DebugDrawDynamicParticles(4, 4, FColor(128, 0, 0));
 			});
@@ -225,7 +227,7 @@ namespace ImmediatePhysics_Chaos
 			{
 				if (ChaosImmediate_DebugDrawCollisions == 4)
 				{
-					DebugDraw::DrawCollisions(SimulationSpaceTransform, InConstraintHandles, 0.6f);
+					DebugDraw::DrawCollisions(SimulationSpaceTransform, Collisions, 0.6f);
 				}
 			});
 		Joints.SetPreApplyCallback(
diff --git a/Engine/Source/Runtime/Experimental/Chaos/Private/Chaos/ChaosDebugDraw.cpp b/Engine/Source/Runtime/Experimental/Chaos/Private/Chaos/ChaosDebugDraw.cpp
--- a/Engine/Source/Runtime/Experimental/Chaos/Private/Chaos/ChaosDebugDraw.cpp
+++ b/Engine/Source/Runtime/Experimental/Chaos/Private/Chaos/ChaosDebugDraw.cpp
@@ -160,9 +160,8 @@ namespace Chaos
 			}
 		}
 
-		void DrawCollisionImpl(const FRigidTransform3& SpaceTransform, const TPBDCollisionConstraintHandle<float, 3>* ConstraintHandle, float ColorScale)
+		void DrawCollisionImpl(const FRigidTransform3& SpaceTransform, const FCollisionConstraintBase& Contact, float ColorScale)
 		{
-			const FCollisionConstraintBase& Contact = ConstraintHandle->GetContact();
 			const FReal ActiveColorScale = (Contact.GetPhi() > 0)? ColorScale * 0.1f : ColorScale;
 
 			FVec3 Location = SpaceTransform.TransformPosition(Contact.GetLocation());
@@ -170,14 +169,14 @@ namespace Chaos
 
 			if (ContactWidth > 0)
 			{
-				bool bIsManifold = (ConstraintHandle->GetType() == FCollisionConstraintBase::FType::MultiPoint);
+				bool bIsManifold = (Contact.GetType() == FCollisionConstraintBase::FType::MultiPoint);
 				FColor C0 = bIsManifold? (ColorScale * FColor(0, 0, 200)).ToFColor(false) : (ColorScale * FColor(200, 0, 0)).ToFColor(false);
 				FMatrix Axes = FRotationMatrix::MakeFromX(Normal);
 				FDebugDrawQueue::GetInstance().DrawDebugCircle(Location, DrawScale * ContactWidth, 12, C0, false, KINDA_SMALL_NUMBER, DrawPriority, LineThickness, Axes.GetUnitAxis(EAxis::Y), Axes.GetUnitAxis(EAxis::Z), false);
 
 				if (bIsManifold)
 				{
-					const FRigidBodyMultiPointContactConstraint& MultiPointConstraint = *ConstraintHandle->GetContact().As<FRigidBodyMultiPointContactConstraint>();
+					const FRigidBodyMultiPointContactConstraint& MultiPointConstraint = *Contact.As<FRigidBodyMultiPointContactConstraint>();
 					TConstGenericParticleHandle<FReal, 3> PointsParticle = MultiPointConstraint.PointsParticleHandle();
 					FRigidTransform3 PointsTransform = FParticleUtilities::GetActorWorldTransform(PointsParticle) * SpaceTransform;
 					for (int32 SampleIndex = 1; SampleIndex < MultiPointConstraint.NumManifoldPoints(); ++SampleIndex)
@@ -203,12 +202,17 @@ namespace Chaos
 			{
 				FColor C3 = (ColorScale * FColor(128, 128, 128)).ToFColor(false);
 				FMatrix Axes = FRotationMatrix::MakeFromX(Normal);
-				FVec3 P0 = SpaceTransform.TransformPosition(ConstraintHandle->GetConstrainedParticles()[0]->X());
-				FVec3 P1 = SpaceTransform.TransformPosition(ConstraintHandle->GetConstrainedParticles()[1]->X());
+				FVec3 P0 = SpaceTransform.TransformPosition(Contact.Particle[0]->X());
+				FVec3 P1 = SpaceTransform.TransformPosition(Contact.Particle[1]->X());
 				FDebugDrawQueue::GetInstance().DrawDebugLine(Location, P0, C3, false, KINDA_SMALL_NUMBER, DrawPriority, LineThickness * 0.5f);
 				FDebugDrawQueue::GetInstance().DrawDebugLine(Location, P1, C3, false, KINDA_SMALL_NUMBER, DrawPriority, LineThickness * 0.5f);
 			}
 		}
+		
+		void DrawCollisionImpl(const FRigidTransform3& SpaceTransform, const TPBDCollisionConstraintHandle<float, 3>* ConstraintHandle, float ColorScale)
+		{
+			DrawCollisionImpl(SpaceTransform, ConstraintHandle->GetContact(), ColorScale);
+		}
 
 		void DrawJointConstraintImpl(const FRigidTransform3& SpaceTransform, const FVec3& InPa, const FVec3& InCa, const FVec3& InXa, const FMatrix33& Ra, const FVec3& InPb, const FVec3& InCb, const FVec3& InXb, const FMatrix33& Rb, int32 Level, int32 Index, FReal ColorScale, uint32 FeatureMask)
 		{
@@ -456,9 +460,9 @@ namespace Chaos
 #if CHAOS_DEBUG_DRAW
 			if (FDebugDrawQueue::IsDebugDrawingEnabled())
 			{
-				for (const Chaos::TPBDCollisionConstraintHandle<float, 3> * ConstraintHandle : Collisions.GetConstConstraintHandles())
+				for (int32 ConstraintIndex = 0; ConstraintIndex < Collisions.NumConstraints(); ++ConstraintIndex)
 				{
-					DrawCollisionImpl(SpaceTransform, ConstraintHandle, ColorScale);
+					DrawCollisionImpl(SpaceTransform, Collisions.GetConstraint(ConstraintIndex), ColorScale);
 				}
 			}
 #endif
diff --git a/Engine/Source/Runtime/Experimental/Chaos/Private/Chaos/Evolution/PBDMinEvolution.cpp b/Engine/Source/Runtime/Experimental/Chaos/Private/Chaos/Evolution/PBDMinEvolution.cpp
--- a/Engine/Source/Runtime/Experimental/Chaos/Private/Chaos/Evolution/PBDMinEvolution.cpp
+++ b/Engine/Source/Runtime/Experimental/Chaos/Private/Chaos/Evolution/PBDMinEvolution.cpp
@@ -19,16 +19,16 @@
 
 namespace Chaos
 {
-	DECLARE_CYCLE_STAT(TEXT("FPBDMinEvolution::AdvanceOneTimeStep"), STAT_MinEvolution_AdvanceOneTimeStep, STATGROUP_Chaos);
-	DECLARE_CYCLE_STAT(TEXT("FPBDMinEvolution::Integrate"), STAT_MinEvolution_Integrate, STATGROUP_Chaos);
-	DECLARE_CYCLE_STAT(TEXT("FPBDMinEvolution::KinematicTargets"), STAT_MinEvolution_KinematicTargets, STATGROUP_Chaos);
-	DECLARE_CYCLE_STAT(TEXT("FPBDMinEvolution::PrepareConstraints"), STAT_MinEvolution_PrepareConstraints, STATGROUP_Chaos);
-	DECLARE_CYCLE_STAT(TEXT("FPBDMinEvolution::UnprepareConstraints"), STAT_MinEvolution_UnprepareConstraints, STATGROUP_Chaos);
-	DECLARE_CYCLE_STAT(TEXT("FPBDMinEvolution::ApplyConstraints"), STAT_MinEvolution_ApplyConstraints, STATGROUP_Chaos);
-	DECLARE_CYCLE_STAT(TEXT("FPBDMinEvolution::UpdateVelocities"), STAT_MinEvolution_UpdateVelocites, STATGROUP_Chaos);
-	DECLARE_CYCLE_STAT(TEXT("FPBDMinEvolution::ApplyPushOut"), STAT_MinEvolution_ApplyPushOut, STATGROUP_Chaos);
-	DECLARE_CYCLE_STAT(TEXT("FPBDMinEvolution::DetectCollisions"), STAT_MinEvolution_DetectCollisions, STATGROUP_Chaos);
-	DECLARE_CYCLE_STAT(TEXT("FPBDMinEvolution::UpdatePositions"), STAT_MinEvolution_UpdatePositions, STATGROUP_Chaos);
+	DECLARE_CYCLE_STAT(TEXT("MinEvolution::AdvanceOneTimeStep"), STAT_MinEvolution_AdvanceOneTimeStep, STATGROUP_Chaos);
+	DECLARE_CYCLE_STAT(TEXT("MinEvolution::Integrate"), STAT_MinEvolution_Integrate, STATGROUP_Chaos);
+	DECLARE_CYCLE_STAT(TEXT("MinEvolution::KinematicTargets"), STAT_MinEvolution_KinematicTargets, STATGROUP_Chaos);
+	DECLARE_CYCLE_STAT(TEXT("MinEvolution::PrepareConstraints"), STAT_MinEvolution_PrepareConstraints, STATGROUP_Chaos);
+	DECLARE_CYCLE_STAT(TEXT("MinEvolution::UnprepareConstraints"), STAT_MinEvolution_UnprepareConstraints, STATGROUP_Chaos);
+	DECLARE_CYCLE_STAT(TEXT("MinEvolution::ApplyConstraints"), STAT_MinEvolution_ApplyConstraints, STATGROUP_Chaos);
+	DECLARE_CYCLE_STAT(TEXT("MinEvolution::UpdateVelocities"), STAT_MinEvolution_UpdateVelocites, STATGROUP_Chaos);
+	DECLARE_CYCLE_STAT(TEXT("MinEvolution::ApplyPushOut"), STAT_MinEvolution_ApplyPushOut, STATGROUP_Chaos);
+	DECLARE_CYCLE_STAT(TEXT("MinEvolution::DetectCollisions"), STAT_MinEvolution_DetectCollisions, STATGROUP_Chaos);
+	DECLARE_CYCLE_STAT(TEXT("MinEvolution::UpdatePositions"), STAT_MinEvolution_UpdatePositions, STATGROUP_Chaos);
 
 	FPBDMinEvolution::FPBDMinEvolution(FRigidParticleSOAs& InParticles, FCollisionDetector& InCollisionDetector, const FReal InBoundsExtension)
 		: Particles(InParticles)
@@ -127,11 +127,6 @@ namespace Chaos
 	{
 		SCOPE_CYCLE_COUNTER(STAT_MinEvolution_Integrate);
 
-		TPerParticleEulerStepVelocity<FReal, 3> EulerStepVelocityRule;
-		TPerParticleAddImpulses<FReal, 3> AddImpulsesRule;
-		TPerParticleEtherDrag<FReal, 3> EtherDragRule;
-		TPerParticlePBDEulerStep<FReal, 3> EulerStepPositionRule;
-
 		for (TTransientPBDRigidParticleHandle<FReal, 3>& Particle : Particles.GetActiveParticlesView())
 		{
 			if (Particle.ObjectState() == EObjectStateType::Dynamic)
@@ -139,26 +134,37 @@ namespace Chaos
 				Particle.PreV() = Particle.V();
 				Particle.PreW() = Particle.W();
 
-				EulerStepVelocityRule.Apply(Particle, Dt);
-
-				AddImpulsesRule.Apply(Particle, Dt);
+				// Calculate new velocities from forces, torques and drag
+				const FMatrix33 WorldInvI = Utilities::ComputeWorldSpaceInertia(Particle.R(), Particle.InvI());
+				const FVec3 DV = Particle.InvM() * (Particle.F() * Dt + Particle.LinearImpulse());
+				const FVec3 DW = WorldInvI * (Particle.Torque() * Dt + Particle.AngularImpulse());
+				const FReal LinearDrag = FMath::Max(FReal(0), FReal(1) - (Particle.LinearEtherDrag() * Dt));
+				const FReal AngularDrag = FMath::Max(FReal(0), FReal(1) - (Particle.AngularEtherDrag() * Dt));
+				const FVec3 V = LinearDrag * (Particle.V() + DV);
+				const FVec3 W = AngularDrag * (Particle.W() + DW);
+
+				FVec3 PCoM = FParticleUtilitiesXR::GetCoMWorldPosition(&Particle);
+				FRotation3 QCoM = FParticleUtilitiesXR::GetCoMWorldRotation(&Particle);
+				PCoM = PCoM + V * Dt;
+				QCoM = FRotation3::IntegrateRotationWithAngularVelocity(QCoM, W, Dt);
+
+				// Update particle state (forces are zeroed at the end of the frame)
+				FParticleUtilitiesPQ::SetCoMWorldTransform(&Particle, PCoM, QCoM);
+				Particle.V() = V;
+				Particle.W() = W;
 				Particle.LinearImpulse() = FVec3(0);
 				Particle.AngularImpulse() = FVec3(0);
 
-				EtherDragRule.Apply(Particle, Dt);
-
-				EulerStepPositionRule.Apply(Particle, Dt);
-
-
+				// Update world-space bounds
 				if (Particle.HasBounds())
 				{
 					const TAABB<FReal, 3>& LocalBounds = Particle.LocalBounds();
 					
-					TAABB<FReal, 3> NextWorldSpaceBounds = LocalBounds.TransformedAABB(FRigidTransform3(Particle.P(), Particle.Q()));
-					NextWorldSpaceBounds.ThickenSymmetrically(NextWorldSpaceBounds.Extents() * BoundsExtension);
+					TAABB<FReal, 3> WorldSpaceBounds = LocalBounds.TransformedAABB(FRigidTransform3(Particle.P(), Particle.Q()));
+					WorldSpaceBounds.ThickenSymmetrically(WorldSpaceBounds.Extents() * BoundsExtension);
 
-					TAABB<FReal, 3> WorldSpaceBounds = LocalBounds.TransformedAABB(FRigidTransform3(Particle.X(), Particle.R()));
-					WorldSpaceBounds.GrowToInclude(NextWorldSpaceBounds);
+					//TAABB<FReal, 3> PrevWorldSpaceBounds = LocalBounds.TransformedAABB(FRigidTransform3(Particle.X(), Particle.R()));
+					//WorldSpaceBounds.GrowToInclude(PrevWorldSpaceBounds);
 
 					Particle.SetWorldSpaceInflatedBounds(WorldSpaceBounds);
 				}
@@ -239,11 +245,11 @@ namespace Chaos
 			{
 				const TAABB<FReal, 3>& LocalBounds = Particle.LocalBounds();
 				
-				TAABB<FReal, 3> NextWorldSpaceBounds = LocalBounds.TransformedAABB(FRigidTransform3(Particle.X(), Particle.R()));
-				NextWorldSpaceBounds.ThickenSymmetrically(NextWorldSpaceBounds.Extents() * BoundsExtension);
+				TAABB<FReal, 3> WorldSpaceBounds = LocalBounds.TransformedAABB(FRigidTransform3(Particle.X(), Particle.R()));
+				WorldSpaceBounds.ThickenSymmetrically(WorldSpaceBounds.Extents() * BoundsExtension);
 
-				TAABB<FReal, 3> WorldSpaceBounds = LocalBounds.TransformedAABB(FRigidTransform3(PrevX, PrevR));
-				WorldSpaceBounds.GrowToInclude(NextWorldSpaceBounds);
+				//TAABB<FReal, 3> PrevWorldSpaceBounds = LocalBounds.TransformedAABB(FRigidTransform3(PrevX, PrevR));
+				//WorldSpaceBounds.GrowToInclude(PrevWorldSpaceBounds);
 
 				Particle.SetWorldSpaceInflatedBounds(WorldSpaceBounds);
 			}
diff --git a/Engine/Source/Runtime/Experimental/Chaos/Private/Chaos/PBDCollisionConstraints.cpp b/Engine/Source/Runtime/Experimental/Chaos/Private/Chaos/PBDCollisionConstraints.cpp
--- a/Engine/Source/Runtime/Experimental/Chaos/Private/Chaos/PBDCollisionConstraints.cpp
+++ b/Engine/Source/Runtime/Experimental/Chaos/Private/Chaos/PBDCollisionConstraints.cpp
@@ -78,12 +78,20 @@ namespace Chaos
 		, MAngularFriction(0)
 		, bUseCCD(false)
 		, bEnableCollisions(true)
+		, bHandlesEnabled(true)
 		, LifespanCounter(0)
 		, PostApplyCallback(nullptr)
 		, PostApplyPushOutCallback(nullptr)
 	{
 	}
 
+	template<typename T, int d>
+	void TPBDCollisionConstraints<T, d>::DisableHandles()
+	{
+		check(NumConstraints() == 0);
+		bHandlesEnabled = false;
+	}
+
 	template<typename T, int d>
 	void TPBDCollisionConstraints<T, d>::SetPostApplyCallback(const TRigidBodyContactConstraintsPostApplyCallback<T, d>& Callback)
 	{
@@ -149,18 +157,22 @@ namespace Chaos
 	void TPBDCollisionConstraints<T, d>::AddConstraint(const TRigidBodyPointContactConstraint<FReal, 3>& InConstraint)
 	{
 		int32 Idx = PointConstraints.Add(InConstraint);
-		FConstraintContainerHandle* Handle = HandleAllocator.template AllocHandle< TRigidBodyPointContactConstraint<T, d> >(this, Idx);
-		Handle->GetContact().Timestamp = -INT_MAX; // force point constraints to be deleted.
 
-		PointConstraints[Idx].ConstraintHandle = Handle;
+		if (bHandlesEnabled)
+		{
+			FConstraintContainerHandle* Handle = HandleAllocator.template AllocHandle< TRigidBodyPointContactConstraint<T, d> >(this, Idx);
+			Handle->GetContact().Timestamp = -INT_MAX; // force point constraints to be deleted.
+
+			PointConstraints[Idx].ConstraintHandle = Handle;
 
-		check(Handle != nullptr);
-		Handles.Add(Handle);
+			check(Handle != nullptr);
+			Handles.Add(Handle);
 
 #if CHAOS_COLLISION_PERSISTENCE_ENABLED
-		check(!Manifolds.Contains(Handle->GetKey()));
-		Manifolds.Add(Handle->GetKey(), Handle);
+			check(!Manifolds.Contains(Handle->GetKey()));
+			Manifolds.Add(Handle->GetKey(), Handle);
 #endif
+		}
 
 		UpdateConstraintMaterialProperties(PointConstraints[Idx]);
 	}
@@ -169,18 +181,22 @@ namespace Chaos
 	void TPBDCollisionConstraints<T, d>::AddConstraint(const TRigidBodyMultiPointContactConstraint<FReal, 3>& InConstraint)
 	{
 		int32 Idx = IterativeConstraints.Add(InConstraint);
-		FConstraintContainerHandle* Handle = HandleAllocator.template AllocHandle< TRigidBodyMultiPointContactConstraint<T, d> >(this, Idx);
-		Handle->GetContact().Timestamp = LifespanCounter;
 
-		IterativeConstraints[Idx].ConstraintHandle = Handle;
+		if (bHandlesEnabled)
+		{
+			FConstraintContainerHandle* Handle = HandleAllocator.template AllocHandle< TRigidBodyMultiPointContactConstraint<T, d> >(this, Idx);
+			Handle->GetContact().Timestamp = LifespanCounter;
+
+			IterativeConstraints[Idx].ConstraintHandle = Handle;
 
-		check(Handle != nullptr);
-		Handles.Add(Handle);
+			check(Handle != nullptr);
+			Handles.Add(Handle);
 
 #if CHAOS_COLLISION_PERSISTENCE_ENABLED
-		check(!Manifolds.Contains(Handle->GetKey()));
-		Manifolds.Add(Handle->GetKey(), Handle);
+			check(!Manifolds.Contains(Handle->GetKey()));
+			Manifolds.Add(Handle->GetKey(), Handle);
 #endif
+		}
 
 		UpdateConstraintMaterialProperties(IterativeConstraints[Idx]);
 	}
@@ -199,6 +215,7 @@ namespace Chaos
 		SCOPE_CYCLE_COUNTER(STAT_Collisions_Reset);
 
 #if CHAOS_COLLISION_PERSISTENCE_ENABLED
+		check(bHandlesEnabled);	// This will need fixing for handle-free mode
 		TArray<FConstraintContainerHandle*> CopyOfHandles = Handles;
 		int32 LifespanWindow = LifespanCounter - 1;
 		for (FConstraintContainerHandle* ContactHandle : CopyOfHandles)
@@ -277,7 +294,7 @@ namespace Chaos
 			}
 #endif
 			PointConstraints.RemoveAtSwap(Idx);
-			if (Idx < PointConstraints.Num())
+			if (bHandlesEnabled && (Idx < PointConstraints.Num()))
 			{
 				PointConstraints[Idx].ConstraintHandle->SetConstraintIndex(Idx, FCollisionConstraintBase::FType::SinglePoint);
 			}
@@ -294,7 +311,7 @@ namespace Chaos
 			}
 #endif
 			IterativeConstraints.RemoveAtSwap(Idx);
-			if (Idx < IterativeConstraints.Num())
+			if (bHandlesEnabled && (Idx < IterativeConstraints.Num()))
 			{
 				IterativeConstraints[Idx].ConstraintHandle->SetConstraintIndex(Idx, FCollisionConstraintBase::FType::MultiPoint);
 			}
@@ -304,17 +321,19 @@ namespace Chaos
 			check(false);
 		}
 
-		// @todo(chaos): Collision Manifold
-		//   Add an index to the handle in the Manifold.Value 
-		//   to prevent the search in Handles when removed.
+		if (bHandlesEnabled)
+		{
+			// @todo(chaos): Collision Manifold
+			//   Add an index to the handle in the Manifold.Value 
+			//   to prevent the search in Handles when removed.
 #if CHAOS_COLLISION_PERSISTENCE_ENABLED
-		Manifolds.Remove(KeyToRemove);
+			Manifolds.Remove(KeyToRemove);
 #endif
-		Handles.Remove(Handle);
+			Handles.Remove(Handle);
+			check(Handles.Num() == PointConstraints.Num() + IterativeConstraints.Num());
 
-		ensure(Handles.Num() == PointConstraints.Num() + IterativeConstraints.Num());
-
-		HandleAllocator.FreeHandle(Handle);
+			HandleAllocator.FreeHandle(Handle);
+		}
 	}
 
 
diff --git a/Engine/Source/Runtime/Experimental/Chaos/Private/Chaos/PBDJointConstraints.cpp b/Engine/Source/Runtime/Experimental/Chaos/Private/Chaos/PBDJointConstraints.cpp
--- a/Engine/Source/Runtime/Experimental/Chaos/Private/Chaos/PBDJointConstraints.cpp
+++ b/Engine/Source/Runtime/Experimental/Chaos/Private/Chaos/PBDJointConstraints.cpp
@@ -635,25 +635,23 @@ namespace Chaos
 	//
 	//////////////////////////////////////////////////////////////////////////
 
-	void FPBDJointConstraints::UpdateParticleState(TPBDRigidParticleHandle<FReal, 3>* Rigid, const FReal Dt, const FVec3& P, const FRotation3& Q, const bool bUpdateVelocity)
+	void FPBDJointConstraints::UpdateParticleState(TPBDRigidParticleHandle<FReal, 3>* Rigid, const FReal Dt, const FVec3& PrevP, const FRotation3& PrevQ, const FVec3& P, const FRotation3& Q, const bool bUpdateVelocity)
 	{
 		if ((Rigid != nullptr) && (Rigid->ObjectState() == EObjectStateType::Dynamic))
 		{
+			FParticleUtilities::SetCoMWorldTransform(Rigid, P, Q);
 			if (bUpdateVelocity && (Dt > SMALL_NUMBER))
 			{
-				FVec3 PCoM = FParticleUtilities::GetCoMWorldPosition(Rigid);
-				FRotation3 QCoM = FParticleUtilities::GetCoMWorldRotation(Rigid);
-				FVec3 DV = FVec3::CalculateVelocity(PCoM, P, Dt);
-				FVec3 DW = FRotation3::CalculateAngularVelocity(QCoM, Q, Dt);
+				const FVec3 DV = FVec3::CalculateVelocity(PrevP, P, Dt);
+				const FVec3 DW = FRotation3::CalculateAngularVelocity(PrevQ, Q, Dt);
 				Rigid->SetV(Rigid->V() + DV);
 				Rigid->SetW(Rigid->W() + DW);
 			}
-			FParticleUtilities::SetCoMWorldTransform(Rigid, P, Q);
 		}
 	}
 
 
-	void FPBDJointConstraints::UpdateParticleState(TPBDRigidParticleHandle<FReal, 3>* Rigid, const FReal Dt, const FVec3& P, const FRotation3& Q, const FVec3& V, const FVec3& W)
+	void FPBDJointConstraints::UpdateParticleStateExplicit(TPBDRigidParticleHandle<FReal, 3>* Rigid, const FReal Dt, const FVec3& P, const FRotation3& Q, const FVec3& V, const FVec3& W)
 	{
 		if ((Rigid != nullptr) && (Rigid->ObjectState() == EObjectStateType::Dynamic))
 		{
@@ -679,14 +677,19 @@ namespace Chaos
 		TGenericParticleHandle<FReal, 3> Particle0 = TGenericParticleHandle<FReal, 3>(ConstraintParticles[ConstraintIndex][Index0]);
 		TGenericParticleHandle<FReal, 3> Particle1 = TGenericParticleHandle<FReal, 3>(ConstraintParticles[ConstraintIndex][Index1]);
 
+		const FVec3 P0 = FParticleUtilities::GetCoMWorldPosition(Particle0);
+		const FRotation3 Q0 = FParticleUtilities::GetCoMWorldRotation(Particle0);
+		const FVec3 P1 = FParticleUtilities::GetCoMWorldPosition(Particle1);
+		const FRotation3 Q1 = FParticleUtilities::GetCoMWorldRotation(Particle1);
+
 		Solver.Update(
 			Dt,
-			FParticleUtilities::GetCoMWorldPosition(Particle0),
-			FParticleUtilities::GetCoMWorldRotation(Particle0),
+			P0,
+			Q0,
 			Particle0->V(),
 			Particle0->W(),
-			FParticleUtilities::GetCoMWorldPosition(Particle1),
-			FParticleUtilities::GetCoMWorldRotation(Particle1),
+			P1,
+			Q1,
 			Particle1->V(),
 			Particle1->W());
 
@@ -702,8 +705,8 @@ namespace Chaos
 			}
 		}
 
-		UpdateParticleState(Particle0->CastToRigidParticle(), Dt, Solver.GetP(0), Solver.GetQ(0));
-		UpdateParticleState(Particle1->CastToRigidParticle(), Dt, Solver.GetP(1), Solver.GetQ(1));
+		UpdateParticleState(Particle0->CastToRigidParticle(), Dt, P0, Q0, Solver.GetP(0), Solver.GetQ(0));
+		UpdateParticleState(Particle1->CastToRigidParticle(), Dt, P1, Q1, Solver.GetP(1), Solver.GetQ(1));
 
 		return NetResult;
 	}
@@ -743,8 +746,8 @@ namespace Chaos
 			}
 		}
 
-		UpdateParticleState(Particle0->CastToRigidParticle(), Dt, Solver.GetP(0), Solver.GetQ(0), Solver.GetV(0), Solver.GetW(0));
-		UpdateParticleState(Particle1->CastToRigidParticle(), Dt, Solver.GetP(1), Solver.GetQ(1), Solver.GetV(1), Solver.GetW(1));
+		UpdateParticleStateExplicit(Particle0->CastToRigidParticle(), Dt, Solver.GetP(0), Solver.GetQ(0), Solver.GetV(0), Solver.GetW(0));
+		UpdateParticleStateExplicit(Particle1->CastToRigidParticle(), Dt, Solver.GetP(1), Solver.GetQ(1), Solver.GetV(1), Solver.GetW(1));
 
 		return NetResult;
 	}
diff --git a/Engine/Source/Runtime/Experimental/Chaos/Public/Chaos/PBDCollisionConstraints.h b/Engine/Source/Runtime/Experimental/Chaos/Public/Chaos/PBDCollisionConstraints.h
--- a/Engine/Source/Runtime/Experimental/Chaos/Public/Chaos/PBDCollisionConstraints.h
+++ b/Engine/Source/Runtime/Experimental/Chaos/Public/Chaos/PBDCollisionConstraints.h
@@ -88,6 +88,16 @@ public:
 
 	virtual ~TPBDCollisionConstraints() {}
 
+	/**
+	 * Whether this container provides constraint handles (simple solvers do not need them)
+	 */
+	bool GetHandlesEnabled() const { return bHandlesEnabled; }
+
+	/**
+	 * Put the container in "no handles" mode for use with simple solver. Must be called when empty of constraints (ideally right after creation).
+	 */
+	void DisableHandles();
+
 	/**
 	*  Add the constraint to the container. 
 	*
@@ -250,18 +260,31 @@ public:
 
 	int32 NumConstraints() const
 	{
-		return Handles.Num();
+		return PointConstraints.Num() + IterativeConstraints.Num();
 	}
 
 	FHandles& GetConstraintHandles()
 	{
 		return Handles;
 	}
+
 	const FHandles& GetConstConstraintHandles() const
 	{
 		return Handles;
 	}
 
+	const FConstraintBase& GetConstraint(int32 Index) const 
+	{
+		check(Index < NumConstraints());
+		
+		if (Index < PointConstraints.Num())
+		{
+			return PointConstraints[Index];
+		}
+		
+		return IterativeConstraints[Index - PointConstraints.Num()];
+	}
+
 
 protected:
 	using Base::GetConstraintIndex;
@@ -291,7 +314,7 @@ private:
 	T MAngularFriction;
 	bool bUseCCD;
 	bool bEnableCollisions;
-	bool bEnableParallelFor;
+	bool bHandlesEnabled;
 
 	int32 LifespanCounter;
 
diff --git a/Engine/Source/Runtime/Experimental/Chaos/Public/Chaos/PBDJointConstraints.h b/Engine/Source/Runtime/Experimental/Chaos/Public/Chaos/PBDJointConstraints.h
--- a/Engine/Source/Runtime/Experimental/Chaos/Public/Chaos/PBDJointConstraints.h
+++ b/Engine/Source/Runtime/Experimental/Chaos/Public/Chaos/PBDJointConstraints.h
@@ -165,8 +165,8 @@ namespace Chaos
 
 		void GetConstrainedParticleIndices(const int32 ConstraintIndex, int32& Index0, int32& Index1) const;
 		void CalculateConstraintSpace(int32 ConstraintIndex, FVec3& OutX0, FMatrix33& OutR0, FVec3& OutX1, FMatrix33& OutR1) const;
-		void UpdateParticleState(TPBDRigidParticleHandle<FReal, 3>* Rigid, const FReal Dt, const FVec3& P, const FRotation3& Q, const bool bUpdateVelocity = true);
-		void UpdateParticleState(TPBDRigidParticleHandle<FReal, 3>* Rigid, const FReal Dt, const FVec3& P, const FRotation3& Q, const FVec3& V, const FVec3& W);
+		void UpdateParticleState(TPBDRigidParticleHandle<FReal, 3>* Rigid, const FReal Dt, const FVec3& PrevP, const FRotation3& PrevQ, const FVec3& P, const FRotation3& Q, const bool bUpdateVelocity = true);
+		void UpdateParticleStateExplicit(TPBDRigidParticleHandle<FReal, 3>* Rigid, const FReal Dt, const FVec3& P, const FRotation3& Q, const FVec3& V, const FVec3& W);
 		void SortConstraints();
 
 		FJointSolverResult SolvePosition_GaussSiedel(const FReal Dt, const int32 ConstraintIndex, const int32 NumPairIts, const int32 It, const int32 NumIts);
