Add sleeping kinematic collisions support.

diff --git a/Engine/Source/Runtime/Experimental/Chaos/Public/Chaos/BoundingVolumeUtilities.h b/Engine/Source/Runtime/Experimental/Chaos/Public/Chaos/BoundingVolumeUtilities.h
--- a/Engine/Source/Runtime/Experimental/Chaos/Public/Chaos/BoundingVolumeUtilities.h
+++ b/Engine/Source/Runtime/Experimental/Chaos/Public/Chaos/BoundingVolumeUtilities.h
@@ -36,6 +36,11 @@ inline FVec3 ComputeBoundsThickness(const TPBDRigidParticles<FReal, 3>& InPartic
 	return ComputeBoundsThickness(InParticles.V(BodyIndex), Dt, BoundsThickness, BoundsVelocityInflation);
 }
 
+inline FVec3 ComputeBoundsThickness(const TKinematicGeometryParticles<FReal, 3>& InParticles, FReal Dt, int32 BodyIndex, FReal BoundsThickness, FReal BoundsVelocityInflation)
+{
+	return ComputeBoundsThickness(InParticles.V(BodyIndex), Dt, BoundsThickness, BoundsVelocityInflation);
+}
+
 template <typename THandle>
 FVec3 ComputeBoundsThickness(const THandle& PBDRigid, FReal Dt, FReal BoundsThickness, FReal BoundsVelocityInflation)
 {
diff --git a/Engine/Source/Runtime/Experimental/Chaos/Public/Chaos/Collision/SpatialAccelerationBroadPhase.h b/Engine/Source/Runtime/Experimental/Chaos/Public/Chaos/Collision/SpatialAccelerationBroadPhase.h
--- a/Engine/Source/Runtime/Experimental/Chaos/Public/Chaos/Collision/SpatialAccelerationBroadPhase.h
+++ b/Engine/Source/Runtime/Experimental/Chaos/Public/Chaos/Collision/SpatialAccelerationBroadPhase.h
@@ -135,7 +135,7 @@ namespace Chaos
 
 			TArray<TAccelerationStructureHandle<FReal, 3>> PotentialIntersections;
 
-			if (Particle1.CastToRigidParticle() && Particle1.ObjectState() == EObjectStateType::Dynamic)
+			if (Particle1.CastToRigidParticle() && (Particle1.ObjectState() == EObjectStateType::Dynamic || Particle1.ObjectState() == EObjectStateType::Sleeping))
 			{
 				// @todo(ccaulfield): the spatial acceleration scheme needs to know the expanded bboxes for all particles, not just the one doing the test
 				const bool bBody1Bounded = HasBoundingBox(Particle1);
@@ -199,6 +199,13 @@ namespace Chaos
 						continue;
 					}
 
+					// Sleeping won't collide against another sleeping and sleeping vs dynamic gets picked up by the other direction.
+					const bool bIsParticle2Kinematic = Particle2.CastToKinematicParticle() && Particle2.ObjectState() == EObjectStateType::Kinematic;
+					if (Particle1.ObjectState() == EObjectStateType::Sleeping && !bIsParticle2Kinematic)
+					{
+						continue;
+					}
+
 					const bool bIsParticle2Dynamic = Particle2.CastToRigidParticle() && Particle2.ObjectState() == EObjectStateType::Dynamic;
 					if (bBody1Bounded == bBody2Bounded && bIsParticle2Dynamic)
 					{
@@ -208,8 +215,9 @@ namespace Chaos
 							continue;
 						}
 					}
-
-					const FReal Box2Thickness = bIsParticle2Dynamic ? ComputeBoundsThickness(*Particle2.CastToRigidParticle(), Dt, BoundsThickness, BoundsThicknessVelocityInflation).Size() : (FReal)0;
+				
+					const FReal Box2Thickness = bIsParticle2Dynamic ? ComputeBoundsThickness(*Particle2.CastToRigidParticle(), Dt, BoundsThickness, BoundsThicknessVelocityInflation).Size()
+						: (bIsParticle2Kinematic ? ComputeBoundsThickness(*Particle2.CastToKinematicParticle(), Dt, BoundsThickness, BoundsThicknessVelocityInflation).Size() : (FReal)0);
 
 					FCollisionConstraintsArray NewConstraints;
 					NarrowPhase.GenerateCollisions(NewConstraints, Dt, Particle1.Handle(), Particle2.Handle(), FMath::Max(Box1Thickness, Box2Thickness), StatData);
