Chaos: Remove events if particle proxy is deleted

diff --git a/Engine/Source/Runtime/Experimental/Chaos/Public/Chaos/ExternalCollisionData.h b/Engine/Source/Runtime/Experimental/Chaos/Public/Chaos/ExternalCollisionData.h
--- a/Engine/Source/Runtime/Experimental/Chaos/Public/Chaos/ExternalCollisionData.h
+++ b/Engine/Source/Runtime/Experimental/Chaos/Public/Chaos/ExternalCollisionData.h
@@ -69,6 +69,8 @@ namespace Chaos
 			, LevelsetIndexMesh(InLevelsetIndexMesh)
 		{}
 
+		bool IsValid() { return (ParticleProxy && LevelsetProxy); }
+
 		TVector<T, d> Location;
 		TVector<T, d> AccumulatedImpulse;
 		TVector<T, d> Normal;
diff --git a/Engine/Source/Runtime/Experimental/ChaosSolvers/Private/PBDRigidsSolver.cpp b/Engine/Source/Runtime/Experimental/ChaosSolvers/Private/PBDRigidsSolver.cpp
--- a/Engine/Source/Runtime/Experimental/ChaosSolvers/Private/PBDRigidsSolver.cpp
+++ b/Engine/Source/Runtime/Experimental/ChaosSolvers/Private/PBDRigidsSolver.cpp
@@ -17,6 +17,8 @@
 #include "PhysicsProxy/GeometryCollectionPhysicsProxy.h"
 #include "PhysicsProxy/FieldSystemPhysicsProxy.h"
 #include "EventDefaults.h"
+#include "EventsData.h"
+
 
 DEFINE_LOG_CATEGORY_STATIC(LogPBDRigidsSolverSolver, Log, All);
 
@@ -309,6 +311,34 @@ namespace Chaos
 		{
 			UE_LOG(LogPBDRigidsSolverSolver, Verbose, TEXT("FPBDRigidsSolver::UnregisterObject() ~ Dequeue"));
 
+				// Generally need to remove stale events for particles that no longer exist
+				Solver->GetEventManager()->ClearEvents<FCollisionEventData>([InProxy]
+				(FCollisionEventData& EventDataInOut)
+				{
+					Chaos::FCollisionDataArray const& CollisionData = EventDataInOut.CollisionData.AllCollisionsArray;
+					if (CollisionData.Num() > 0)
+					{
+						check(InProxy);
+						TArray<int32> const* const CollisionIndices = EventDataInOut.PhysicsProxyToCollisionIndices.PhysicsProxyToIndicesMap.Find(InProxy);
+						if (CollisionIndices)
+						{
+							for (int32 EncodedCollisionIdx : *CollisionIndices)
+							{
+								bool bSwapOrder;
+								int32 CollisionIdx = Chaos::FEventManager::DecodeCollisionIndex(EncodedCollisionIdx, bSwapOrder);
+
+								// invalidate but don't delete from array, as this would mean we'd need to reindex PhysicsProxyToIndicesMap to maintain the other collisions lookup
+								Chaos::TCollisionData<float, 3>& CollisionDataItem = EventDataInOut.CollisionData.AllCollisionsArray[CollisionIdx];
+								CollisionDataItem.ParticleProxy = nullptr;
+								CollisionDataItem.LevelsetProxy = nullptr;
+							}
+
+							EventDataInOut.PhysicsProxyToCollisionIndices.PhysicsProxyToIndicesMap.Remove(InProxy);
+						}
+					}
+
+				});
+
 			// Get the physics thread-handle from the proxy, and then delete the proxy.
 			//
 			// NOTE: We have to delete the proxy from its derived version, because the
@@ -339,7 +369,9 @@ namespace Chaos
 
 			// Use the handle to destroy the particle data
 			Solver->GetEvolution()->DestroyParticle(Handle);
+
 		});
+
 	}
 
 	void FPBDRigidsSolver::RegisterObject(FGeometryCollectionPhysicsProxy* InProxy)
diff --git a/Engine/Source/Runtime/Experimental/ChaosSolvers/Public/EventManager.h b/Engine/Source/Runtime/Experimental/ChaosSolvers/Public/EventManager.h
--- a/Engine/Source/Runtime/Experimental/ChaosSolvers/Public/EventManager.h
+++ b/Engine/Source/Runtime/Experimental/ChaosSolvers/Public/EventManager.h
@@ -4,6 +4,7 @@
 
 #include "CoreMinimal.h"
 #include "Chaos/Framework/MultiBufferResource.h"
+#include "Chaos/Framework/PhysicsProxyBase.h"
 
 namespace Chaos
 {
@@ -184,6 +185,11 @@ namespace Chaos
 			InjectedFunction(Solver, *EventBuffer->AccessProducerBuffer());
 		}
 
+		virtual void DestroyStaleEvents(TFunction<void(PayloadType & EventDataInOut)> InFunction)
+		{
+			InFunction(*EventBuffer->AccessProducerBuffer());
+		}
+
 		/**
 		 * Flips the buffer if the buffer type is double or triple
 		 */
@@ -270,6 +276,23 @@ namespace Chaos
 			ContainerLock.WriteUnlock();
 		}
 
+		/**
+		 * Modify the producer side of the event buffer
+		 */
+		template<typename PayloadType>
+		void ClearEvents(TFunction<void(PayloadType & EventData)> InFunction)
+		{
+			ContainerLock.ReadLock();
+			for (FEventContainerPtr EventContainer : EventContainers)
+			{
+				if (EventContainer)
+				{
+					((TEventContainer<PayloadType>*)(EventContainer))->DestroyStaleEvents(InFunction);
+				}
+			}
+			ContainerLock.ReadUnlock();
+		}
+
 		/**
 		 * Unregister specified event from system
 		 */
