Missing instanced implicit for sq serialization

diff --git a/Engine/Source/Runtime/Experimental/Chaos/Private/Chaos/ImplicitObject.cpp b/Engine/Source/Runtime/Experimental/Chaos/Private/Chaos/ImplicitObject.cpp
--- a/Engine/Source/Runtime/Experimental/Chaos/Private/Chaos/ImplicitObject.cpp
+++ b/Engine/Source/Runtime/Experimental/Chaos/Private/Chaos/ImplicitObject.cpp
@@ -385,6 +385,19 @@ FImplicitObject* FImplicitObject::SerializationFactory(FChaosArchive& Ar, FImpli
 		}
 	}
 
+	if(IsInstanced(ObjectType))
+	{
+		EImplicitObjectType InnerType = GetInnerType(ObjectType);
+		switch(InnerType)
+		{
+		case ImplicitObjectType::Convex: if(Ar.IsLoading()) { return new TImplicitObjectInstanced<FConvex>(); } break;
+		case ImplicitObjectType::TriangleMesh: if(Ar.IsLoading()) { return new TImplicitObjectInstanced<FTriangleMeshImplicitObject>(); } break;
+		default: check(false);
+		}
+
+		return nullptr;
+	}
+
 	switch ((EImplicitObjectType)ObjectType)
 	{
 	case ImplicitObjectType::Sphere: if (Ar.IsLoading()) { return new TSphere<FReal, 3>(); } break;
diff --git a/Engine/Source/Runtime/Experimental/Chaos/Public/Chaos/ImplicitObjectScaled.h b/Engine/Source/Runtime/Experimental/Chaos/Public/Chaos/ImplicitObjectScaled.h
--- a/Engine/Source/Runtime/Experimental/Chaos/Public/Chaos/ImplicitObjectScaled.h
+++ b/Engine/Source/Runtime/Experimental/Chaos/Public/Chaos/ImplicitObjectScaled.h
@@ -27,6 +27,11 @@ public:
 
 	using FImplicitObject::GetTypeName;
 
+	//needed for serialization
+	TImplicitObjectInstanced()
+		: FImplicitObject(EImplicitObject::HasBoundingBox,StaticType())
+	{}
+
 	TImplicitObjectInstanced(const ObjectType&& Object)
 		: FImplicitObject(EImplicitObject::HasBoundingBox, Object->GetType() | ImplicitObjectType::IsInstanced)
 		, MObject(MoveTemp(Object))
