Add sub implicit objects in MapImplicitShapes for geometry particles

diff --git a/Engine/Source/Runtime/Experimental/Chaos/Private/Chaos/GeometryParticles.cpp b/Engine/Source/Runtime/Experimental/Chaos/Private/Chaos/GeometryParticles.cpp
--- a/Engine/Source/Runtime/Experimental/Chaos/Private/Chaos/GeometryParticles.cpp
+++ b/Engine/Source/Runtime/Experimental/Chaos/Private/Chaos/GeometryParticles.cpp
@@ -198,10 +198,19 @@ namespace Chaos
 		TShapesArray<T, d>& ShapeArray = MShapesArray[Index];
 		Mapping.Reset();
 
-		int32 ShapeIndex = 0;
-		for (TUniquePtr<TPerShapeData<T, d>>& Shape : ShapeArray)
+		for (int32 ShapeIndex = 0; ShapeIndex < ShapeArray.Num(); ++ShapeIndex)
 		{
-			Mapping.Add(Shape->Geometry.Get(), ShapeIndex++);
+			const FImplicitObject* Geometry = ShapeArray[ShapeIndex]->Geometry.Get();
+			Mapping.Add(Geometry, ShapeIndex);
+
+			Utilities::CastHelper(*Geometry, FRigidTransform3(), [&Mapping, Geometry, ShapeIndex](const auto& InnerGeometry, const FRigidTransform3& TM)
+			{
+				const FImplicitObject* ImplicitObject = static_cast<const FImplicitObject*>(&InnerGeometry);
+				if (ImplicitObject && ImplicitObject != Geometry)
+				{
+					Mapping.Add(ImplicitObject, ShapeIndex);
+				}
+			});
 		}
 
 		if (MGeometry[Index])
