Add TriangleMesh Support to ImplicitChildHelper.
Move MapImplicitShapes to cpp to avoid cyclic include.

diff --git a/Engine/Source/Runtime/Experimental/Chaos/Private/Chaos/ParticleHandle.cpp b/Engine/Source/Runtime/Experimental/Chaos/Private/Chaos/ParticleHandle.cpp
--- a/Engine/Source/Runtime/Experimental/Chaos/Private/Chaos/ParticleHandle.cpp
+++ b/Engine/Source/Runtime/Experimental/Chaos/Private/Chaos/ParticleHandle.cpp
@@ -5,9 +5,61 @@
 #include "Chaos/Framework/PhysicsProxy.h"
 #include "Chaos/Framework/PhysicsSolverBase.h"
 #include "Chaos/Framework/PhysicsProxyBase.h"
+#include "Chaos/CastingUtilities.h"
 
 namespace Chaos
 {
+	template <typename T, int d>
+	void Chaos::TGeometryParticle<T, d>::MapImplicitShapes()
+	{
+		ImplicitShapeMap.Reset();
+		int32 ShapeIndex = 0;
+		for (TUniquePtr<TPerShapeData<T, d>>& ShapeData : MShapesArray)
+		{
+			ImplicitShapeMap.Add(ShapeData->Geometry.Get(), ShapeIndex++);
+		}
+
+		if (MGeometry)
+		{
+			int32 CurrentShapeIndex = INDEX_NONE;
+			if (const auto* Union = MGeometry->template GetObject<FImplicitObjectUnion>())
+			{
+				for (const TUniquePtr<FImplicitObject>& ImplicitObject : Union->GetObjects())
+				{
+					if (ImplicitObject.Get())
+					{
+						if (const FImplicitObject* ImplicitChildObject = Utilities::ImplicitChildHelper(ImplicitObject.Get()))
+						{
+							if (ImplicitShapeMap.Contains(ImplicitObject.Get()))
+							{
+								ImplicitShapeMap.Add(ImplicitChildObject, ImplicitShapeMap[ImplicitObject.Get()]);
+							}
+							else if (ImplicitShapeMap.Contains(ImplicitChildObject))
+							{
+								ImplicitShapeMap.Add(ImplicitObject.Get(), ImplicitShapeMap[ImplicitChildObject]);
+							}
+						}
+					}
+				}
+			}
+			else 
+			{
+				if (const FImplicitObject* ImplicitChildObject = Utilities::ImplicitChildHelper(MGeometry.Get()))
+				{
+					if (ImplicitShapeMap.Contains(MGeometry.Get()))
+					{
+						ImplicitShapeMap.Add(ImplicitChildObject, ImplicitShapeMap[MGeometry.Get()]);
+					}
+					else if (ImplicitShapeMap.Contains(ImplicitChildObject))
+					{
+						ImplicitShapeMap.Add(MGeometry.Get(), ImplicitShapeMap[ImplicitChildObject]);
+					}
+				}
+			}
+		}
+	}
+
+
 
 	template class CHAOS_API TGeometryParticleData<float, 3>;
 	template class CHAOS_API TGeometryParticle<float, 3>;
diff --git a/Engine/Source/Runtime/Experimental/Chaos/Public/Chaos/CastingUtilities.h b/Engine/Source/Runtime/Experimental/Chaos/Public/Chaos/CastingUtilities.h
--- a/Engine/Source/Runtime/Experimental/Chaos/Public/Chaos/CastingUtilities.h
+++ b/Engine/Source/Runtime/Experimental/Chaos/Public/Chaos/CastingUtilities.h
@@ -8,7 +8,7 @@
 #include "Chaos/Sphere.h"
 #include "Chaos/Capsule.h"
 #include "Chaos/Convex.h"
-//#include "Chaos/TriangleMeshImplicitObject.h"
+#include "Chaos/TriangleMeshImplicitObject.h"
 #include "Chaos/Box.h"
 
 namespace Chaos
@@ -81,10 +81,10 @@ namespace Chaos
 			{
 				return ImplicitObject->template GetObject<const TImplicitObjectScaled<TSphere<FReal, 3>>>()->GetUnscaledObject();
 			}
-			//else if (ImplicitType == TImplicitObjectScaled<FTriangleMeshImplicitObject>::StaticType())
-			//{
-			//	return ImplicitObject->template GetObject<const TImplicitObjectScaled<FTriangleMeshImplicitObject>>()->GetUnscaledObject();
-			//}
+			else if (ImplicitType == TImplicitObjectScaled<FTriangleMeshImplicitObject>::StaticType())
+			{
+				return ImplicitObject->template GetObject<const TImplicitObjectScaled<FTriangleMeshImplicitObject>>()->GetUnscaledObject();
+			}
 
 			else if (ImplicitType == TImplicitObjectInstanced<FConvex>::StaticType())
 			{
@@ -102,10 +102,10 @@ namespace Chaos
 			{
 				return ImplicitObject->template GetObject<const TImplicitObjectInstanced<TSphere<FReal, 3>>>()->GetInstancedObject();
 			}
-			//else if (ImplicitType == TImplicitObjectInstanced<FTriangleMeshImplicitObject>::StaticType())
-			//{
-			//	return ImplicitObject->template GetObject<const TImplicitObjectInstanced<FTriangleMeshImplicitObject>>()->GetInstancedObject();
-			//}
+			else if (ImplicitType == TImplicitObjectInstanced<FTriangleMeshImplicitObject>::StaticType())
+			{
+				return ImplicitObject->template GetObject<const TImplicitObjectInstanced<FTriangleMeshImplicitObject>>()->GetInstancedObject();
+			}
 			return nullptr;
 		}
 
diff --git a/Engine/Source/Runtime/Experimental/Chaos/Public/Chaos/ParticleHandle.h b/Engine/Source/Runtime/Experimental/Chaos/Public/Chaos/ParticleHandle.h
--- a/Engine/Source/Runtime/Experimental/Chaos/Public/Chaos/ParticleHandle.h
+++ b/Engine/Source/Runtime/Experimental/Chaos/Public/Chaos/ParticleHandle.h
@@ -2,7 +2,6 @@
 #pragma once
 
 #include "Templates/ChooseClass.h"
-#include "Chaos/CastingUtilities.h"
 #include "Chaos/PBDRigidClusteredParticles.h"
 #include "Chaos/PBDGeometryCollectionParticles.h"
 #include "Chaos/ParticleHandleFwd.h"
@@ -1539,55 +1538,7 @@ protected:
 		MapImplicitShapes();
 	}
 
-	//TGeometryParticle::
-	void MapImplicitShapes()
-	{
-		ImplicitShapeMap.Reset();
-		int32 ShapeIndex = 0;
-		for (TUniquePtr<TPerShapeData<T, d>>& ShapeData : MShapesArray)
-		{
-			ImplicitShapeMap.Add(ShapeData->Geometry.Get(), ShapeIndex++);
-		}
-
-		if (MGeometry)
-		{
-			int32 CurrentShapeIndex = INDEX_NONE;
-			if (const auto* Union = MGeometry->template GetObject<FImplicitObjectUnion>())
-			{
-				for (const TUniquePtr<FImplicitObject>& ImplicitObject : Union->GetObjects())
-				{
-					if (ImplicitObject.Get())
-					{
-						if (const FImplicitObject* ImplicitChildObject = Utilities::ImplicitChildHelper(ImplicitObject.Get()))
-						{
-							if (ImplicitShapeMap.Contains(ImplicitObject.Get()))
-							{
-								ImplicitShapeMap.Add(ImplicitChildObject, ImplicitShapeMap[ImplicitObject.Get()]);
-							}
-							else if (ImplicitShapeMap.Contains(ImplicitChildObject))
-							{
-								ImplicitShapeMap.Add(ImplicitObject.Get(), ImplicitShapeMap[ImplicitChildObject]);
-							}
-						}
-					}
-				}
-			}
-			else 
-			{
-				if (const FImplicitObject* ImplicitChildObject = Utilities::ImplicitChildHelper(MGeometry.Get()))
-				{
-					if (ImplicitShapeMap.Contains(MGeometry.Get()))
-					{
-						ImplicitShapeMap.Add(ImplicitChildObject, ImplicitShapeMap[MGeometry.Get()]);
-					}
-					else if (ImplicitShapeMap.Contains(ImplicitChildObject))
-					{
-						ImplicitShapeMap.Add(MGeometry.Get(), ImplicitShapeMap[ImplicitChildObject]);
-					}
-				}
-			}
-		}
-	}
+	void MapImplicitShapes();
 };
 
 template <typename T, int d>
