Back out changelist 11113745 and 11113737 (changes near CastHelper). Something is being given to casthelper that is not implemented in function, and failing builds.

diff --git a/Engine/Source/Runtime/Experimental/Chaos/Private/Chaos/GeometryParticles.cpp b/Engine/Source/Runtime/Experimental/Chaos/Private/Chaos/GeometryParticles.cpp
--- a/Engine/Source/Runtime/Experimental/Chaos/Private/Chaos/GeometryParticles.cpp
+++ b/Engine/Source/Runtime/Experimental/Chaos/Private/Chaos/GeometryParticles.cpp
@@ -198,19 +198,10 @@ namespace Chaos
 		TShapesArray<T, d>& ShapeArray = MShapesArray[Index];
 		Mapping.Reset();
 
-		for (int32 ShapeIndex = 0; ShapeIndex < ShapeArray.Num(); ++ShapeIndex)
+		int32 ShapeIndex = 0;
+		for (TUniquePtr<TPerShapeData<T, d>>& Shape : ShapeArray)
 		{
-			const FImplicitObject* Geometry = ShapeArray[ShapeIndex]->Geometry.Get();
-			Mapping.Add(Geometry, ShapeIndex);
-
-			Utilities::CastHelper(*Geometry, FRigidTransform3(), [&Mapping, Geometry, ShapeIndex](const auto& InnerGeometry, const FRigidTransform3& TM)
-			{
-				const FImplicitObject* ImplicitObject = static_cast<const FImplicitObject*>(&InnerGeometry);
-				if (ImplicitObject && ImplicitObject != Geometry)
-				{
-					Mapping.Add(ImplicitObject, ShapeIndex);
-				}
-			});
+			Mapping.Add(Shape->Geometry.Get(), ShapeIndex++);
 		}
 
 		if (MGeometry[Index])
diff --git a/Engine/Source/Runtime/Experimental/Chaos/Private/Chaos/ParticleHandle.cpp b/Engine/Source/Runtime/Experimental/Chaos/Private/Chaos/ParticleHandle.cpp
--- a/Engine/Source/Runtime/Experimental/Chaos/Private/Chaos/ParticleHandle.cpp
+++ b/Engine/Source/Runtime/Experimental/Chaos/Private/Chaos/ParticleHandle.cpp
@@ -5,61 +5,9 @@
 #include "Chaos/Framework/PhysicsProxy.h"
 #include "Chaos/Framework/PhysicsSolverBase.h"
 #include "Chaos/Framework/PhysicsProxyBase.h"
-#include "Chaos/CastingUtilities.h"
 
 namespace Chaos
 {
-	template <typename T, int d>
-	void Chaos::TGeometryParticle<T, d>::MapImplicitShapes()
-	{
-		ImplicitShapeMap.Reset();
-		int32 ShapeIndex = 0;
-		for (TUniquePtr<TPerShapeData<T, d>>& ShapeData : MShapesArray)
-		{
-			ImplicitShapeMap.Add(ShapeData->Geometry.Get(), ShapeIndex++);
-		}
-
-		if (MGeometry)
-		{
-			int32 CurrentShapeIndex = INDEX_NONE;
-			if (const auto* Union = MGeometry->template GetObject<FImplicitObjectUnion>())
-			{
-				for (const TUniquePtr<FImplicitObject>& ImplicitObject : Union->GetObjects())
-				{
-					if (ImplicitObject.Get())
-					{
-						if (const FImplicitObject* ImplicitChildObject = Utilities::ImplicitChildHelper(ImplicitObject.Get()))
-						{
-							if (ImplicitShapeMap.Contains(ImplicitObject.Get()))
-							{
-								ImplicitShapeMap.Add(ImplicitChildObject, ImplicitShapeMap[ImplicitObject.Get()]);
-							}
-							else if (ImplicitShapeMap.Contains(ImplicitChildObject))
-							{
-								ImplicitShapeMap.Add(ImplicitObject.Get(), ImplicitShapeMap[ImplicitChildObject]);
-							}
-						}
-					}
-				}
-			}
-			else 
-			{
-				if (const FImplicitObject* ImplicitChildObject = Utilities::ImplicitChildHelper(MGeometry.Get()))
-				{
-					if (ImplicitShapeMap.Contains(MGeometry.Get()))
-					{
-						ImplicitShapeMap.Add(ImplicitChildObject, ImplicitShapeMap[MGeometry.Get()]);
-					}
-					else if (ImplicitShapeMap.Contains(ImplicitChildObject))
-					{
-						ImplicitShapeMap.Add(MGeometry.Get(), ImplicitShapeMap[ImplicitChildObject]);
-					}
-				}
-			}
-		}
-	}
-
-
 
 	template class CHAOS_API TGeometryParticleData<float, 3>;
 	template class CHAOS_API TGeometryParticle<float, 3>;
diff --git a/Engine/Source/Runtime/Experimental/Chaos/Public/Chaos/CastingUtilities.h b/Engine/Source/Runtime/Experimental/Chaos/Public/Chaos/CastingUtilities.h
--- a/Engine/Source/Runtime/Experimental/Chaos/Public/Chaos/CastingUtilities.h
+++ b/Engine/Source/Runtime/Experimental/Chaos/Public/Chaos/CastingUtilities.h
@@ -8,7 +8,7 @@
 #include "Chaos/Sphere.h"
 #include "Chaos/Capsule.h"
 #include "Chaos/Convex.h"
-#include "Chaos/TriangleMeshImplicitObject.h"
+//#include "Chaos/TriangleMeshImplicitObject.h"
 #include "Chaos/Box.h"
 
 namespace Chaos
@@ -81,10 +81,10 @@ namespace Chaos
 			{
 				return ImplicitObject->template GetObject<const TImplicitObjectScaled<TSphere<FReal, 3>>>()->GetUnscaledObject();
 			}
-			else if (ImplicitType == TImplicitObjectScaled<FTriangleMeshImplicitObject>::StaticType())
-			{
-				return ImplicitObject->template GetObject<const TImplicitObjectScaled<FTriangleMeshImplicitObject>>()->GetUnscaledObject();
-			}
+			//else if (ImplicitType == TImplicitObjectScaled<FTriangleMeshImplicitObject>::StaticType())
+			//{
+			//	return ImplicitObject->template GetObject<const TImplicitObjectScaled<FTriangleMeshImplicitObject>>()->GetUnscaledObject();
+			//}
 
 			else if (ImplicitType == TImplicitObjectInstanced<FConvex>::StaticType())
 			{
@@ -102,10 +102,10 @@ namespace Chaos
 			{
 				return ImplicitObject->template GetObject<const TImplicitObjectInstanced<TSphere<FReal, 3>>>()->GetInstancedObject();
 			}
-			else if (ImplicitType == TImplicitObjectInstanced<FTriangleMeshImplicitObject>::StaticType())
-			{
-				return ImplicitObject->template GetObject<const TImplicitObjectInstanced<FTriangleMeshImplicitObject>>()->GetInstancedObject();
-			}
+			//else if (ImplicitType == TImplicitObjectInstanced<FTriangleMeshImplicitObject>::StaticType())
+			//{
+			//	return ImplicitObject->template GetObject<const TImplicitObjectInstanced<FTriangleMeshImplicitObject>>()->GetInstancedObject();
+			//}
 			return nullptr;
 		}
 
diff --git a/Engine/Source/Runtime/Experimental/Chaos/Public/Chaos/ParticleHandle.h b/Engine/Source/Runtime/Experimental/Chaos/Public/Chaos/ParticleHandle.h
--- a/Engine/Source/Runtime/Experimental/Chaos/Public/Chaos/ParticleHandle.h
+++ b/Engine/Source/Runtime/Experimental/Chaos/Public/Chaos/ParticleHandle.h
@@ -2,6 +2,7 @@
 #pragma once
 
 #include "Templates/ChooseClass.h"
+#include "Chaos/CastingUtilities.h"
 #include "Chaos/PBDRigidClusteredParticles.h"
 #include "Chaos/PBDGeometryCollectionParticles.h"
 #include "Chaos/ParticleHandleFwd.h"
@@ -1538,7 +1539,55 @@ protected:
 		MapImplicitShapes();
 	}
 
-	void MapImplicitShapes();
+	//TGeometryParticle::
+	void MapImplicitShapes()
+	{
+		ImplicitShapeMap.Reset();
+		int32 ShapeIndex = 0;
+		for (TUniquePtr<TPerShapeData<T, d>>& ShapeData : MShapesArray)
+		{
+			ImplicitShapeMap.Add(ShapeData->Geometry.Get(), ShapeIndex++);
+		}
+
+		if (MGeometry)
+		{
+			int32 CurrentShapeIndex = INDEX_NONE;
+			if (const auto* Union = MGeometry->template GetObject<FImplicitObjectUnion>())
+			{
+				for (const TUniquePtr<FImplicitObject>& ImplicitObject : Union->GetObjects())
+				{
+					if (ImplicitObject.Get())
+					{
+						if (const FImplicitObject* ImplicitChildObject = Utilities::ImplicitChildHelper(ImplicitObject.Get()))
+						{
+							if (ImplicitShapeMap.Contains(ImplicitObject.Get()))
+							{
+								ImplicitShapeMap.Add(ImplicitChildObject, ImplicitShapeMap[ImplicitObject.Get()]);
+							}
+							else if (ImplicitShapeMap.Contains(ImplicitChildObject))
+							{
+								ImplicitShapeMap.Add(ImplicitObject.Get(), ImplicitShapeMap[ImplicitChildObject]);
+							}
+						}
+					}
+				}
+			}
+			else 
+			{
+				if (const FImplicitObject* ImplicitChildObject = Utilities::ImplicitChildHelper(MGeometry.Get()))
+				{
+					if (ImplicitShapeMap.Contains(MGeometry.Get()))
+					{
+						ImplicitShapeMap.Add(ImplicitChildObject, ImplicitShapeMap[MGeometry.Get()]);
+					}
+					else if (ImplicitShapeMap.Contains(ImplicitChildObject))
+					{
+						ImplicitShapeMap.Add(MGeometry.Get(), ImplicitShapeMap[ImplicitChildObject]);
+					}
+				}
+			}
+		}
+	}
 };
 
 template <typename T, int d>
