Cloth: Add fast position based friction, frcition is set to 0.2 for all assets.

diff --git a/Engine/Plugins/Experimental/ChaosCloth/Source/Chaos/Private/ChaosCloth/ChaosClothingSimulation.cpp b/Engine/Plugins/Experimental/ChaosCloth/Source/Chaos/Private/ChaosCloth/ChaosClothingSimulation.cpp
--- a/Engine/Plugins/Experimental/ChaosCloth/Source/Chaos/Private/ChaosCloth/ChaosClothingSimulation.cpp
+++ b/Engine/Plugins/Experimental/ChaosCloth/Source/Chaos/Private/ChaosCloth/ChaosClothingSimulation.cpp
@@ -57,7 +57,7 @@ namespace ClothingSimulationDefault
 	static const int32 NumIterations = 1;
 	static const float SelfCollisionThickness = 2.f;
 	static const float CollisionThickness = 1.2f;
-	static const float FrictionCoefficient = 0.f;
+	static const float FrictionCoefficient = 0.2f;
 	static const float DampingCoefficient = 0.01f;
 	static const float WorldScale = 100.f;  // World is in cm, but values like wind speed and density are in SI unit and relates to m.
 }
diff --git a/Engine/Source/Runtime/Experimental/Chaos/Public/Chaos/PerParticlePBDCollisionConstraint.h b/Engine/Source/Runtime/Experimental/Chaos/Public/Chaos/PerParticlePBDCollisionConstraint.h
--- a/Engine/Source/Runtime/Experimental/Chaos/Public/Chaos/PerParticlePBDCollisionConstraint.h
+++ b/Engine/Source/Runtime/Experimental/Chaos/Public/Chaos/PerParticlePBDCollisionConstraint.h
@@ -19,7 +19,7 @@ class TPerParticlePBDCollisionConstraint : public TPerParticleRule<T, d>
 
   public:
 	TPerParticlePBDCollisionConstraint(const TKinematicGeometryParticlesImp<T, d, SimType>& InParticles, TArrayCollectionArray<bool>& Collided, const T Thickness = (T)0, const T Friction = (T)0)
-	    : MParticles(InParticles), MCollided(Collided), MThickness(Thickness), MFriction(Friction) {}
+	    : bFastPositionBasedFriction(true), MParticles(InParticles), MCollided(Collided), MThickness(Thickness), MFriction(Friction) {}
 	virtual ~TPerParticlePBDCollisionConstraint() {}
 
 	inline void Apply(TPBDParticles<T, d>& InParticles, const T Dt, const int32 Index) const override //-V762
@@ -33,14 +33,38 @@ class TPerParticlePBDCollisionConstraint : public TPerParticleRule<T, d>
 			T Phi = MParticles.Geometry(i)->PhiWithNormal(Frame.InverseTransformPosition(InParticles.P(Index)), Normal);
 			if (Phi < MThickness)
 			{
-				InParticles.P(Index) += (-Phi + MThickness) * Frame.TransformVector(Normal);
+				const TVector<T, d> NormalWorld = Frame.TransformVector(Normal);
+				InParticles.P(Index) += (-Phi + MThickness) * NormalWorld;
 				if (MFriction > 0)
 				{
-					VelocityConstraint Constraint;
-					TVector<T, d> VectorToPoint = InParticles.P(Index) - MParticles.X(i);
-					Constraint.Velocity = MParticles.V(i) + TVector<T, d>::CrossProduct(MParticles.W(i), VectorToPoint);
-					Constraint.Normal = Frame.TransformVector(Normal);
-					MVelocityConstraints.Add(Index, Constraint);
+					const T MaximumFrictionCorrectionPerStep = 1.0f; // 1cm absolute maximum correction the friction force can provide
+					if (bFastPositionBasedFriction)
+					{
+						const T Penetration = (-Phi + MThickness); // This is related to the Normal impulse 
+						TVector<T, d> VectorToPoint = InParticles.P(Index) - MParticles.X(i);
+						const TVector<T, d> RelativeDisplacement = (InParticles.P(Index) - InParticles.X(Index)) - (MParticles.V(i) + TVector<T, d>::CrossProduct(MParticles.W(i), VectorToPoint)) * Dt; // This corresponds to the tangential velocity multiplied by dt (friction will drive this to zero if it is high enough)
+						const TVector<T, d> RelativeDisplacementTangent = RelativeDisplacement - TVector<T, d>::DotProduct(RelativeDisplacement, NormalWorld) * NormalWorld; // Project displacement into the tangential plane
+						const T RelativeDisplacementTangentLength = RelativeDisplacementTangent.Size();
+						T PositionCorrection = FMath::Clamp<T>(Penetration * MFriction, 0, RelativeDisplacementTangentLength);
+						if (PositionCorrection > MaximumFrictionCorrectionPerStep)
+						{
+							PositionCorrection = MaximumFrictionCorrectionPerStep;
+						}
+						if (PositionCorrection > 0)
+						{
+							InParticles.P(Index) -= (PositionCorrection / RelativeDisplacementTangentLength) * RelativeDisplacementTangent;
+						}
+					}
+					else
+					{
+						// Note, to fix: Only use fast position based friction for now, since adding to TMaps here is not thread safe when calling Apply on multiple threads (will cause crash)
+						VelocityConstraint Constraint;
+						TVector<T, d> VectorToPoint = InParticles.P(Index) - MParticles.X(i);
+						Constraint.Velocity = MParticles.V(i) + TVector<T, d>::CrossProduct(MParticles.W(i), VectorToPoint);
+						Constraint.Normal = Frame.TransformVector(Normal);
+						
+						MVelocityConstraints.Add(Index, Constraint);
+					}
 				}
 			}
 		}
@@ -48,23 +72,27 @@ class TPerParticlePBDCollisionConstraint : public TPerParticleRule<T, d>
 
 	void ApplyFriction(TPBDParticles<T, d>& InParticles, const T Dt, const int32 Index) const
 	{
-		check(MFriction > 0);
-		if (!MVelocityConstraints.Contains(Index))
+		if (!bFastPositionBasedFriction)
 		{
-			return;
+			check(MFriction > 0);
+			if (!MVelocityConstraints.Contains(Index))
+			{
+				return;
+			}
+			T VN = TVector<T, d>::DotProduct(InParticles.V(Index), MVelocityConstraints[Index].Normal);
+			T VNBody = TVector<T, d>::DotProduct(MVelocityConstraints[Index].Velocity, MVelocityConstraints[Index].Normal);
+			TVector<T, d> VTBody = MVelocityConstraints[Index].Velocity - VNBody * MVelocityConstraints[Index].Normal;
+			TVector<T, d> VTRelative = InParticles.V(Index) - VN * MVelocityConstraints[Index].Normal - VTBody;
+			T VTRelativeSize = VTRelative.Size();
+			T VNMax = FGenericPlatformMath::Max(VN, VNBody);
+			T VNDelta = VNMax - VN;
+			T Friction = MFriction * VNDelta < VTRelativeSize ? MFriction * VNDelta / VTRelativeSize : 1;
+			InParticles.V(Index) = VNMax * MVelocityConstraints[Index].Normal + VTBody + VTRelative * (1 - Friction);
 		}
-		T VN = TVector<T, d>::DotProduct(InParticles.V(Index), MVelocityConstraints[Index].Normal);
-		T VNBody = TVector<T, d>::DotProduct(MVelocityConstraints[Index].Velocity, MVelocityConstraints[Index].Normal);
-		TVector<T, d> VTBody = MVelocityConstraints[Index].Velocity - VNBody * MVelocityConstraints[Index].Normal;
-		TVector<T, d> VTRelative = InParticles.V(Index) - VN * MVelocityConstraints[Index].Normal - VTBody;
-		T VTRelativeSize = VTRelative.Size();
-		T VNMax = FGenericPlatformMath::Max(VN, VNBody);
-		T VNDelta = VNMax - VN;
-		T Friction = MFriction * VNDelta < VTRelativeSize ? MFriction * VNDelta / VTRelativeSize : 1;
-		InParticles.V(Index) = VNMax * MVelocityConstraints[Index].Normal + VTBody + VTRelative * (1 - Friction);
 	}
 
   private:
+	  bool bFastPositionBasedFriction;
 	// TODO(mlentine): Need a bb hierarchy
 	const TKinematicGeometryParticlesImp<T, d, SimType>& MParticles;
 	TArrayCollectionArray<bool>& MCollided;
