Chaos - shape type handling in update loop

diff --git a/Engine/Source/Runtime/Experimental/Chaos/Private/Chaos/CollisionResolution.cpp b/Engine/Source/Runtime/Experimental/Chaos/Private/Chaos/CollisionResolution.cpp
--- a/Engine/Source/Runtime/Experimental/Chaos/Private/Chaos/CollisionResolution.cpp
+++ b/Engine/Source/Runtime/Experimental/Chaos/Private/Chaos/CollisionResolution.cpp
@@ -66,52 +66,6 @@ namespace Chaos
 			}
 		}
 
-		/**
-		 * Select the best manifold point as the new contact point and update the constraint data.
-		 */
-		template<class T, int d>
-		void UpdateConstraintFromManifold(TRigidBodyMultiPointContactConstraint<T, d>& Constraint, const TRigidTransform<T, d>& ParticleTransform0, const TRigidTransform<T, d>& ParticleTransform1, const T CullDistance)
-		{
-			// TEMP: Test code - run single-point algorithm instead. Allows us to debug draw the manifold without using it
-			if (Chaos_Collision_UseManifolds_Test)
-			{
-				TRigidBodyPointContactConstraint<T, d> PointConstraint(Constraint.Particle[0], Constraint.Manifold.Implicit[0], Constraint.ImplicitTransform[0], Constraint.Particle[1], Constraint.Manifold.Implicit[1], Constraint.ImplicitTransform[1], Constraint.Manifold.ShapesType);
-				UpdateConstraint<ECollisionUpdateType::Deepest>(PointConstraint, ParticleTransform0, ParticleTransform1, CullDistance);
-				Constraint.Manifold = PointConstraint.Manifold;
-				return;
-			}
-
-			// Get the plane and point transforms (depends which body owns the plane)
-			const FRigidTransform3 Transform0 = Constraint.ImplicitTransform[0] * ParticleTransform0;
-			const FRigidTransform3 Transform1 = Constraint.ImplicitTransform[1] * ParticleTransform1;
-			const FRigidTransform3& PlaneTransform = (Constraint.GetManifoldPlaneOwnerIndex() == 0) ? Transform0 : Transform1;
-			const FRigidTransform3& PointsTransform = (Constraint.GetManifoldPlaneOwnerIndex() == 0) ? Transform1 : Transform0;
-
-			// World-space manifold plane
-			FVec3 PlaneNormal = PlaneTransform.TransformVector(Constraint.GetManifoldPlaneNormal());
-			FVec3 PlanePos = PlaneTransform.TransformPosition(Constraint.GetManifoldPlanePosition());
-
-			// Select the best manifold point
-			const int32 NumPoints = Constraint.NumManifoldPoints();
-			for (int32 PointIndex = 0; PointIndex < NumPoints; ++PointIndex)
-			{
-				// World-space manifold point and distance to manifold plane
-				FVec3 PointPos = PointsTransform.TransformPosition(Constraint.GetManifoldPoint(PointIndex));
-				FReal PointDistance = FVec3::DotProduct(PointPos - PlanePos, PlaneNormal);
-
-				// If this is the deepest hit, select it
-				if (PointDistance < Constraint.Manifold.Phi)
-				{
-					// @todo(chaos): Consider using average of plane and point positions for contact location
-					FVec3 ContactPos = PointPos - PointDistance * PlaneNormal;
-					FVec3 ContactNormal = (Constraint.GetManifoldPlaneOwnerIndex() == 0) ? -PlaneNormal : PlaneNormal;
-					Constraint.Manifold.Phi = PointDistance;
-					Constraint.Manifold.Location = ContactPos;
-					Constraint.Manifold.Normal = ContactNormal;
-				}
-			}
-		}
-
 
 		template <typename T, int d, typename GeometryA, typename GeometryB>
 		TContactPoint<T> GJKContactPoint(const GeometryA& A, const TRigidTransform<T, d>& ATM, const GeometryB& B, const TRigidTransform<T, d>& BTM, const TVector<T, 3>& InitialDir)
@@ -139,6 +93,7 @@ namespace Chaos
 			return Contact;
 		}
 
+
 		template <typename GeometryA, typename GeometryB, typename T, int d>
 		TContactPoint<T> GJKImplicitContactPoint(const FImplicitObject& A, const TRigidTransform<T, d>& ATransform, const GeometryB& B, const TRigidTransform<T, d>& BTransform, const T CullDistance)
 		{
@@ -182,6 +137,43 @@ namespace Chaos
 		}
 
 
+		// This is pretty unnecessary - all instanced shapes have the same implementation so we should be able to
+		// collapse this switch into a generic call. Maybe add a base class to TImplicitObjectInstanced.
+		inline const FImplicitObject* GetInstancedImplicit(const FImplicitObject* Implicit0)
+		{
+			EImplicitObjectType Implicit0OuterType = Implicit0->GetType();
+
+			if (Implicit0OuterType == TImplicitObjectInstanced<FConvex>::StaticType())
+			{
+				return Implicit0->template GetObject<const TImplicitObjectInstanced<FConvex>>()->GetInstancedObject();
+			}
+			else if (Implicit0OuterType == TImplicitObjectInstanced<TBox<FReal, 3>>::StaticType())
+			{
+				return Implicit0->template GetObject<const TImplicitObjectInstanced<TBox<FReal, 3>>>()->GetInstancedObject();
+			}
+			else if (Implicit0OuterType == TImplicitObjectInstanced<TCapsule<FReal>>::StaticType())
+			{
+				return Implicit0->template GetObject<const TImplicitObjectInstanced<TCapsule<FReal>>>()->GetInstancedObject();
+			}
+			else if (Implicit0OuterType == TImplicitObjectInstanced<TSphere<FReal, 3>>::StaticType())
+			{
+				return Implicit0->template GetObject<const TImplicitObjectInstanced<TSphere<FReal, 3>>>()->GetInstancedObject();
+			}
+			else if (Implicit0OuterType == TImplicitObjectInstanced<FConvex>::StaticType())
+			{
+				return Implicit0->template GetObject<const TImplicitObjectInstanced<FConvex>>()->GetInstancedObject();
+			}
+			else if (Implicit0OuterType == TImplicitObjectInstanced<FTriangleMeshImplicitObject>::StaticType())
+			{
+				// Confusingly, instanced Tri Meshes are treated differently
+				//return Implicit0->template GetObject<const TImplicitObjectInstanced<FTriangleMeshImplicitObject>>()->GetInstancedObject();
+				return nullptr;
+			}
+
+			return nullptr;
+		}
+
+
 		template<class T, int d>
 		TContactPoint<T> ConvexConvexContactPoint(const FImplicitObject& A, const TRigidTransform<T, d>& ATM, const FImplicitObject& B, const TRigidTransform<T, d>& BTM, const T CullDistance)
 		{
@@ -195,11 +187,16 @@ namespace Chaos
 		}
 
 		template <typename T, int d>
-		void UpdateSingleShotManifold(TRigidBodyMultiPointContactConstraint<T, d>&  Constraint, const TRigidTransform<T, d>& Transform0, const TRigidTransform<T, d>& Transform1, const T CullDistance)
+		void UpdateSingleShotManifold(TRigidBodyMultiPointContactConstraint<T, d>& Constraint, const TRigidTransform<T, d>& Transform0, const TRigidTransform<T, d>& Transform1, const T CullDistance)
 		{
 			// single shot manifolds for TConvex implicit object in the constraints implicit[0] position. 
 			TContactPoint<T> ContactPoint = ConvexConvexContactPoint(*Constraint.Manifold.Implicit[0], Transform0, *Constraint.Manifold.Implicit[1], Transform1, CullDistance);
 
+			// Cache the nearest point as the initial contact
+			Constraint.Manifold.Phi = ContactPoint.Phi;
+			Constraint.Manifold.Normal = ContactPoint.Normal;
+			Constraint.Manifold.Location = ContactPoint.Location;
+
 			TArray<FVec3> CollisionSamples;
 			//
 			//  @todo(chaos) : Collision Manifold
@@ -246,6 +243,11 @@ namespace Chaos
 			// iterative manifolds for non TConvex implicit objects that require sampling 
 			TContactPoint<T> ContactPoint = ConvexConvexContactPoint(*Constraint.Manifold.Implicit[0], Transform0, *Constraint.Manifold.Implicit[1], Transform1, CullDistance);
 
+			// Cache the nearest point as the initial contact
+			Constraint.Manifold.Phi = ContactPoint.Phi;
+			Constraint.Manifold.Normal = ContactPoint.Normal;
+			Constraint.Manifold.Location = ContactPoint.Location;
+
 			if (!ContactPoint.Normal.Equals(Constraint.GetManifoldPlaneNormal()) || !Constraint.NumManifoldPoints())
 			{
 				Constraint.ResetManifoldPoints();
@@ -336,7 +338,7 @@ namespace Chaos
 			{
 				TRigidTransform<T, d> ParticleImplicit0TM = Transform0.GetRelativeTransform(Collisions::GetTransform(Particle0));
 				TRigidTransform<T, d> ParticleImplicit1TM = Transform1.GetRelativeTransform(Collisions::GetTransform(Particle1));
-				FRigidBodyPointContactConstraint Constraint = FRigidBodyPointContactConstraint(Particle0, Implicit0, ParticleImplicit0TM, Particle1, Implicit1, ParticleImplicit1TM);
+				FRigidBodyPointContactConstraint Constraint = FRigidBodyPointContactConstraint(Particle0, Implicit0, ParticleImplicit0TM, Particle1, Implicit1, ParticleImplicit1TM, EContactShapesType::BoxBox);
 				UpdateBoxBoxConstraint(Object0->BoundingBox(), Transform0, Object1->BoundingBox(), Transform1, CullDistance, Constraint);
 				NewConstraints.TryAdd(CullDistance, Constraint);
 			}
@@ -375,7 +377,7 @@ namespace Chaos
 			{
 				TRigidTransform<T, d> ParticleImplicit0TM = Transform0.GetRelativeTransform(Collisions::GetTransform(Particle0));
 				TRigidTransform<T, d> ParticleImplicit1TM = Transform1.GetRelativeTransform(Collisions::GetTransform(Particle1));
-				FRigidBodyPointContactConstraint Constraint = FRigidBodyPointContactConstraint(Particle0, Implicit0, ParticleImplicit0TM, Particle1, Implicit1, ParticleImplicit1TM);
+				FRigidBodyPointContactConstraint Constraint = FRigidBodyPointContactConstraint(Particle0, Implicit0, ParticleImplicit0TM, Particle1, Implicit1, ParticleImplicit1TM, EContactShapesType::BoxHeightField);
 				UpdateBoxHeightFieldConstraint(Object0->BoundingBox(), Transform0, *Object1, Transform1, CullDistance, Constraint);
 				NewConstraints.TryAdd(CullDistance, Constraint);
 			}
@@ -463,7 +465,7 @@ namespace Chaos
 			{
 				TRigidTransform<T, d> ParticleImplicit0TM = Transform0.GetRelativeTransform(Collisions::GetTransform(Particle0));
 				TRigidTransform<T, d> ParticleImplicit1TM = Transform1.GetRelativeTransform(Collisions::GetTransform(Particle1));
-				FRigidBodyPointContactConstraint Constraint = FRigidBodyPointContactConstraint(Particle0, Implicit0, ParticleImplicit0TM, Particle1, Implicit1, ParticleImplicit1TM);
+				FRigidBodyPointContactConstraint Constraint = FRigidBodyPointContactConstraint(Particle0, Implicit0, ParticleImplicit0TM, Particle1, Implicit1, ParticleImplicit1TM, EContactShapesType::BoxPlane);
 				UpdateBoxPlaneConstraint(Object0->BoundingBox(), Transform0, *Object1, Transform1, CullDistance, Constraint);
 				NewConstraints.TryAdd(CullDistance, Constraint);
 			}
@@ -501,7 +503,7 @@ namespace Chaos
 			{
 				TRigidTransform<T, d> ParticleImplicit0TM = Transform0.GetRelativeTransform(Collisions::GetTransform(Particle0));
 				TRigidTransform<T, d> ParticleImplicit1TM = Transform1.GetRelativeTransform(Collisions::GetTransform(Particle1));
-				FRigidBodyPointContactConstraint Constraint = FRigidBodyPointContactConstraint(Particle0, Implicit0, ParticleImplicit0TM, Particle1, Implicit1, ParticleImplicit1TM);
+				FRigidBodyPointContactConstraint Constraint = FRigidBodyPointContactConstraint(Particle0, Implicit0, ParticleImplicit0TM, Particle1, Implicit1, ParticleImplicit1TM, EContactShapesType::BoxTriMesh);
 				UpdateBoxTriangleMeshConstraint(Object0->GetAABB(), Transform0, *Object1, Transform1, CullDistance, Constraint);
 				NewConstraints.TryAdd(CullDistance, Constraint);
 			}
@@ -551,7 +553,7 @@ namespace Chaos
 			{
 				TRigidTransform<T, d> ParticleImplicit0TM = Transform0.GetRelativeTransform(Collisions::GetTransform(Particle0));
 				TRigidTransform<T, d> ParticleImplicit1TM = Transform1.GetRelativeTransform(Collisions::GetTransform(Particle1));
-				FRigidBodyPointContactConstraint Constraint = FRigidBodyPointContactConstraint(Particle0, Implicit0, ParticleImplicit0TM, Particle1, Implicit1, ParticleImplicit1TM);
+				FRigidBodyPointContactConstraint Constraint = FRigidBodyPointContactConstraint(Particle0, Implicit0, ParticleImplicit0TM, Particle1, Implicit1, ParticleImplicit1TM, EContactShapesType::SphereSphere);
 				UpdateSphereSphereConstraint(*Object0, Transform0, *Object1, Transform1, CullDistance, Constraint);
 				NewConstraints.TryAdd(CullDistance, Constraint);
 			}
@@ -590,7 +592,7 @@ namespace Chaos
 			{
 				TRigidTransform<T, d> ParticleImplicit0TM = Transform0.GetRelativeTransform(Collisions::GetTransform(Particle0));
 				TRigidTransform<T, d> ParticleImplicit1TM = Transform1.GetRelativeTransform(Collisions::GetTransform(Particle1));
-				FRigidBodyPointContactConstraint Constraint = FRigidBodyPointContactConstraint(Particle0, Implicit0, ParticleImplicit0TM, Particle1, Implicit1, ParticleImplicit1TM);
+				FRigidBodyPointContactConstraint Constraint = FRigidBodyPointContactConstraint(Particle0, Implicit0, ParticleImplicit0TM, Particle1, Implicit1, ParticleImplicit1TM, EContactShapesType::SphereHeightField);
 				UpdateSphereHeightFieldConstraint(*Object0, Transform0, *Object1, Transform1, CullDistance, Constraint);
 				NewConstraints.TryAdd(CullDistance, Constraint);
 			}
@@ -637,7 +639,7 @@ namespace Chaos
 			{
 				TRigidTransform<T, d> ParticleImplicit0TM = Transform0.GetRelativeTransform(Collisions::GetTransform(Particle0));
 				TRigidTransform<T, d> ParticleImplicit1TM = Transform1.GetRelativeTransform(Collisions::GetTransform(Particle1));
-				FRigidBodyPointContactConstraint Constraint = FRigidBodyPointContactConstraint(Particle0, Implicit0, ParticleImplicit0TM, Particle1, Implicit1, ParticleImplicit1TM);
+				FRigidBodyPointContactConstraint Constraint = FRigidBodyPointContactConstraint(Particle0, Implicit0, ParticleImplicit0TM, Particle1, Implicit1, ParticleImplicit1TM, EContactShapesType::SpherePlane);
 				UpdateSpherePlaneConstraint(*Object0, Transform0, *Object1, Transform1, CullDistance, Constraint);
 				NewConstraints.TryAdd(CullDistance, Constraint);
 			}
@@ -688,7 +690,7 @@ namespace Chaos
 			{
 				TRigidTransform<T, d> ParticleImplicit0TM = Transform0.GetRelativeTransform(Collisions::GetTransform(Particle0));
 				TRigidTransform<T, d> ParticleImplicit1TM = Transform1.GetRelativeTransform(Collisions::GetTransform(Particle1));
-				FRigidBodyPointContactConstraint Constraint = FRigidBodyPointContactConstraint(Particle0, Implicit0, ParticleImplicit0TM, Particle1, Implicit1, ParticleImplicit1TM);
+				FRigidBodyPointContactConstraint Constraint = FRigidBodyPointContactConstraint(Particle0, Implicit0, ParticleImplicit0TM, Particle1, Implicit1, ParticleImplicit1TM, EContactShapesType::SphereBox);
 				UpdateSphereBoxConstraint(*Object0, Transform0, Object1->BoundingBox(), Transform1, CullDistance, Constraint);
 				NewConstraints.TryAdd(CullDistance, Constraint);
 			}
@@ -745,7 +747,7 @@ namespace Chaos
 			{
 				TRigidTransform<T, d> ParticleImplicit0TM = Transform0.GetRelativeTransform(Collisions::GetTransform(Particle0));
 				TRigidTransform<T, d> ParticleImplicit1TM = Transform1.GetRelativeTransform(Collisions::GetTransform(Particle1));
-				FRigidBodyPointContactConstraint Constraint = FRigidBodyPointContactConstraint(Particle0, Implicit0, ParticleImplicit0TM, Particle1, Implicit1, ParticleImplicit1TM);
+				FRigidBodyPointContactConstraint Constraint = FRigidBodyPointContactConstraint(Particle0, Implicit0, ParticleImplicit0TM, Particle1, Implicit1, ParticleImplicit1TM, EContactShapesType::SphereCapsule);
 				UpdateSphereCapsuleConstraint(*Object0, Transform0, *Object1, Transform1, CullDistance, Constraint);
 				NewConstraints.TryAdd(CullDistance, Constraint);
 			}
@@ -783,7 +785,7 @@ namespace Chaos
 			{
 				TRigidTransform<T, d> ParticleImplicit0TM = Transform0.GetRelativeTransform(Collisions::GetTransform(Particle0));
 				TRigidTransform<T, d> ParticleImplicit1TM = Transform1.GetRelativeTransform(Collisions::GetTransform(Particle1));
-				FRigidBodyPointContactConstraint Constraint = FRigidBodyPointContactConstraint(Particle0, Implicit0, ParticleImplicit0TM, Particle1, Implicit1, ParticleImplicit1TM);
+				FRigidBodyPointContactConstraint Constraint = FRigidBodyPointContactConstraint(Particle0, Implicit0, ParticleImplicit0TM, Particle1, Implicit1, ParticleImplicit1TM, EContactShapesType::SphereTriMesh);
 				UpdateSphereTriangleMeshConstraint(*Object0, Transform0, *Object1, Transform1, CullDistance, Constraint);
 				NewConstraints.TryAdd(CullDistance, Constraint);
 			}
@@ -843,7 +845,7 @@ namespace Chaos
 			{
 				TRigidTransform<T, d> ParticleImplicit0TM = Transform0.GetRelativeTransform(Collisions::GetTransform(Particle0));
 				TRigidTransform<T, d> ParticleImplicit1TM = Transform1.GetRelativeTransform(Collisions::GetTransform(Particle1));
-				FRigidBodyPointContactConstraint Constraint = FRigidBodyPointContactConstraint(Particle0, Implicit0, ParticleImplicit0TM, Particle1, Implicit1, ParticleImplicit1TM);
+				FRigidBodyPointContactConstraint Constraint = FRigidBodyPointContactConstraint(Particle0, Implicit0, ParticleImplicit0TM, Particle1, Implicit1, ParticleImplicit1TM, EContactShapesType::CapsuleCapsule);
 				UpdateCapsuleCapsuleConstraint(*Object0, Transform0, *Object1, Transform1, CullDistance, Constraint);
 				NewConstraints.TryAdd(CullDistance, Constraint);
 			}
@@ -1132,7 +1134,7 @@ namespace Chaos
 			{
 				TRigidTransform<T, d> ParticleImplicit0TM = Transform0.GetRelativeTransform(Collisions::GetTransform(Particle0));
 				TRigidTransform<T, d> ParticleImplicit1TM = Transform1.GetRelativeTransform(Collisions::GetTransform(Particle1));
-				FRigidBodyPointContactConstraint Constraint = FRigidBodyPointContactConstraint(Particle0, Implicit0, ParticleImplicit0TM, Particle1, Implicit1, ParticleImplicit1TM);
+				FRigidBodyPointContactConstraint Constraint = FRigidBodyPointContactConstraint(Particle0, Implicit0, ParticleImplicit0TM, Particle1, Implicit1, ParticleImplicit1TM, EContactShapesType::CapsuleHeightField);
 				UpdateCapsuleHeightFieldConstraint(*Object0, Transform0, *Object1, Transform1, CullDistance, Constraint);
 				NewConstraints.TryAdd(CullDistance, Constraint);
 			}
@@ -1170,7 +1172,7 @@ namespace Chaos
 			{
 				TRigidTransform<T, d> ParticleImplicit0TM = Transform0.GetRelativeTransform(Collisions::GetTransform(Particle0));
 				TRigidTransform<T, d> ParticleImplicit1TM = Transform1.GetRelativeTransform(Collisions::GetTransform(Particle1));
-				FRigidBodyPointContactConstraint Constraint = FRigidBodyPointContactConstraint(Particle0, Implicit0, ParticleImplicit0TM, Particle1, Implicit1, ParticleImplicit1TM);
+				FRigidBodyPointContactConstraint Constraint = FRigidBodyPointContactConstraint(Particle0, Implicit0, ParticleImplicit0TM, Particle1, Implicit1, ParticleImplicit1TM, EContactShapesType::CapsuleTriMesh);
 				UpdateCapsuleTriangleMeshConstraint(*Object0, Transform0, *Object1, Transform1, CullDistance, Constraint);
 				NewConstraints.TryAdd(CullDistance, Constraint);
 			}
@@ -1241,7 +1243,7 @@ namespace Chaos
 		{
 			TRigidTransform<T, d> ParticleImplicit0TM = Transform0.GetRelativeTransform(Collisions::GetTransform(Particle0));
 			TRigidTransform<T, d> ParticleImplicit1TM = Transform1.GetRelativeTransform(Collisions::GetTransform(Particle1));
-			FRigidBodyMultiPointContactConstraint Constraint = FRigidBodyMultiPointContactConstraint(Particle0, Implicit0, ParticleImplicit0TM, Particle1, Implicit1, ParticleImplicit1TM);
+			FRigidBodyMultiPointContactConstraint Constraint = FRigidBodyMultiPointContactConstraint(Particle0, Implicit0, ParticleImplicit0TM, Particle1, Implicit1, ParticleImplicit1TM, EContactShapesType::ConvexConvex);
 			UpdateConvexConvexManifold(Constraint, Transform0, Transform1, CullDistance);
 			UpdateConvexConvexConstraint(*Implicit0, Transform0, *Implicit1, Transform1, CullDistance, Constraint);
 			NewConstraints.TryAdd(CullDistance, Constraint);
@@ -1279,7 +1281,7 @@ namespace Chaos
 			{
 				TRigidTransform<T, d> ParticleImplicit0TM = Transform0.GetRelativeTransform(Collisions::GetTransform(Particle0));
 				TRigidTransform<T, d> ParticleImplicit1TM = Transform1.GetRelativeTransform(Collisions::GetTransform(Particle1));
-				FRigidBodyPointContactConstraint Constraint = FRigidBodyPointContactConstraint(Particle0, Implicit0, ParticleImplicit0TM, Particle1, Implicit1, ParticleImplicit1TM);
+				FRigidBodyPointContactConstraint Constraint = FRigidBodyPointContactConstraint(Particle0, Implicit0, ParticleImplicit0TM, Particle1, Implicit1, ParticleImplicit1TM, EContactShapesType::ConvexHeightField);
 				UpdateConvexHeightFieldConstraint(*Implicit0, Transform0, *Object1, Transform1, CullDistance, Constraint);
 				NewConstraints.TryAdd(CullDistance, Constraint);
 			}
@@ -1316,7 +1318,7 @@ namespace Chaos
 			{
 				TRigidTransform<T, d> ParticleImplicit0TM = Transform0.GetRelativeTransform(Collisions::GetTransform(Particle0));
 				TRigidTransform<T, d> ParticleImplicit1TM = Transform1.GetRelativeTransform(Collisions::GetTransform(Particle1));
-				FRigidBodyPointContactConstraint Constraint = FRigidBodyPointContactConstraint(Particle0, Implicit0, ParticleImplicit0TM, Particle1, Implicit1, ParticleImplicit1TM);
+				FRigidBodyPointContactConstraint Constraint = FRigidBodyPointContactConstraint(Particle0, Implicit0, ParticleImplicit0TM, Particle1, Implicit1, ParticleImplicit1TM, EContactShapesType::ConvexConvex);
 				UpdateConvexTriangleMeshConstraint(*Implicit0, Transform0, *Object1, Transform1, CullDistance, Constraint);
 				NewConstraints.TryAdd(CullDistance, Constraint);
 			}
@@ -1371,7 +1373,7 @@ namespace Chaos
 		{
 			TRigidTransform<T, d> ParticleImplicit0TM = Transform0.GetRelativeTransform(Collisions::GetTransform(Particle0));
 			TRigidTransform<T, d> ParticleImplicit1TM = Transform1.GetRelativeTransform(Collisions::GetTransform(Particle1));
-			FRigidBodyPointContactConstraint Constraint = FRigidBodyPointContactConstraint(Particle0, Implicit0, ParticleImplicit0TM, Particle1, Implicit1, ParticleImplicit1TM);
+			FRigidBodyPointContactConstraint Constraint = FRigidBodyPointContactConstraint(Particle0, Implicit0, ParticleImplicit0TM, Particle1, Implicit1, ParticleImplicit1TM, EContactShapesType::LevelSetLevelSet);
 
 			bool bIsParticleDynamic0 = Particle0->CastToRigidParticle() && Particle0->ObjectState() == EObjectStateType::Dynamic;
 			if (!Particle1->Geometry() || (bIsParticleDynamic0 && !Particle0->CastToRigidParticle()->CollisionParticlesSize() && Particle0->Geometry() && !Particle0->Geometry()->IsUnderlyingUnion()))
@@ -1425,483 +1427,307 @@ namespace Chaos
 		// Constraint API
 		//
 
-		template<typename T, int d>
-		void UpdateManifold(TCollisionConstraintBase<T, d>& ConstraintBase, const TRigidTransform<T, d>& ATM, const TRigidTransform<T, d>& BTM, const T CullDistance)
-		{
-			if (ConstraintBase.GetType() != TCollisionConstraintBase<T, d>::FType::MultiPoint)
-			{
-				return;
-			}
-			TRigidBodyMultiPointContactConstraint<T, d>& Constraint = *ConstraintBase.template As<TRigidBodyMultiPointContactConstraint<T, d>>();
-
-			const FImplicitObject& Implicit0 = *ConstraintBase.Manifold.Implicit[0];
-			const FImplicitObject& Implicit1 = *ConstraintBase.Manifold.Implicit[1];
-
-			const TRigidTransform<T, d>& Transform0 = ConstraintBase.ImplicitTransform[0] * ATM;
-			const TRigidTransform<T, d>& Transform1 = ConstraintBase.ImplicitTransform[1] * BTM;
 
-#if !UE_BUILD_SHIPPING
-			EImplicitObjectType Implicit0OuterType = Implicit0.GetType();
-			EImplicitObjectType Implicit1OuterType = Implicit1.GetType();
-
-			if (Implicit0OuterType == TImplicitObjectTransformed<T, d>::StaticType())
-			{
-				ensure(false);//should not be possible to get this type, it should already be resolved by the constraint. (see ConstructConstraints)
-				return;
-			}
-			else if (Implicit1OuterType == TImplicitObjectTransformed<T, d>::StaticType())
-			{
-				ensure(false);//should not be possible to get this type, it should already be resolved by the constraint. (see ConstructConstraints)
-				return;
-			}
-			else if (Implicit0OuterType != FImplicitObjectUnion::StaticType() && Implicit1OuterType == FImplicitObjectUnion::StaticType())
-			{
-				ensure(false);//should not be possible to get this type, it should already be resolved by the constraint. (see ConstructConstraints)
-				return;
-			}
-			else if (Implicit0OuterType == FImplicitObjectUnion::StaticType() && Implicit1OuterType != FImplicitObjectUnion::StaticType())
-			{
-				ensure(false);//should not be possible to get this type, it should already be resolved by the constraint. (see ConstructConstraints)
-				return;
-			}
-			else if (Implicit0OuterType == FImplicitObjectUnion::StaticType() && Implicit1OuterType == FImplicitObjectUnion::StaticType())
-			{
-				ensure(false);//should not be possible to get this type, it should already be resolved by the constraint. (see ConstructConstraints)
-				return;
-			}
-#endif
-
-			//
-			// @todo(chaos): Collision Constraints (CollisionMap)
-			//    Modify Construct() and Update() to use a CollisionMap indexed on 
-			//    EImplicitObjectType, instead of the if/else chain. Also, remove 
-			//    the blocks with the ensure(false), they are just for validation 
-			//    after the recent change. 
-			//
-			// @todo(chaos): use constraint ShapesType to index manifold update function
-			//
-			EImplicitObjectType Implicit0Type = GetInnerType(Implicit0.GetType());
-			EImplicitObjectType Implicit1Type = GetInnerType(Implicit1.GetType()); 
-			
-			if (Implicit0Type == TBox<T, d>::StaticType() && Implicit1Type == TBox<T, d>::StaticType())
-			{
-				UpdateBoxBoxManifold(ConstraintBase, Transform0, Transform1, CullDistance);
-			}
-			else if (Implicit0Type == TBox<T, d>::StaticType() && Implicit1Type == THeightField<T>::StaticType())
-			{
-				UpdateBoxHeightFieldManifold(ConstraintBase, Transform0, Transform1, CullDistance);
-			}
-			else if (Implicit0Type == TSphere<T, d>::StaticType() && Implicit1Type == TSphere<T, d>::StaticType())
-			{
-				UpdateSphereSphereManifold(ConstraintBase, Transform0, Transform1, CullDistance);
-			}
-			else if (Implicit0Type == TSphere<T, d>::StaticType() && Implicit1Type == THeightField<T>::StaticType())
-			{
-				UpdateSphereHeightFieldManifold(ConstraintBase, Transform0, Transform1, CullDistance);
-			}
-			else if (Implicit0Type == TBox<T, d>::StaticType() && Implicit1Type == TPlane<T, d>::StaticType())
-			{
-				UpdateBoxPlaneManifold(ConstraintBase, Transform0, Transform1, CullDistance);
-			}
-			else if (Implicit0Type == TSphere<T, d>::StaticType() && Implicit1Type == TPlane<T, d>::StaticType())
-			{
-				UpdateSpherePlaneManifold(ConstraintBase, Transform0, Transform1, CullDistance);
-			}
-			else if (Implicit0Type == TSphere<T, d>::StaticType() && Implicit1Type == TBox<T, d>::StaticType())
-			{
-				UpdateSphereBoxManifold(ConstraintBase, Transform0, Transform1, CullDistance);
-			}
-			else if (Implicit0Type == TSphere<T, d>::StaticType() && Implicit1Type == TCapsule<T>::StaticType())
-			{
-				UpdateSphereCapsuleManifold(ConstraintBase, Transform0, Transform1, CullDistance);
-			}
-			else if (Implicit0Type == TCapsule<T>::StaticType() && Implicit1Type == TCapsule<T>::StaticType())
-			{
-				UpdateCapsuleCapsuleManifold(ConstraintBase, Transform0, Transform1, CullDistance);
-			}
-			else if (Implicit0Type == TCapsule<T>::StaticType() && Implicit1Type == TBox<T, d>::StaticType())
-			{
-				UpdateCapsuleBoxManifold<T, d>(*Constraint.Manifold.Implicit[0]->template GetObject<const TCapsule<T>>(), Transform0, Constraint.Manifold.Implicit[1]->template GetObject<const TBox<T, d>>()->GetAABB(), Transform1, CullDistance, Constraint);
-			}
-			else if (Implicit0Type == TCapsule<T>::StaticType() && Implicit1Type == THeightField<T>::StaticType())
-			{
-				UpdateCapsuleHeightFieldManifold(ConstraintBase, Transform0, Transform1, CullDistance);
-			}
-#if !UE_BUILD_SHIPPING
-			else if (Implicit0Type == THeightField<T>::StaticType() && Implicit1Type == TBox<T, d>::StaticType())
-			{
-				ensure(false);
-			}
-			else if (Implicit0Type == TPlane<T, d>::StaticType() && Implicit1Type == TBox<T, d>::StaticType())
-			{
-				//UpdatePlaneBoxManifold(ConstraintBase, Transform0, Transform1, CullDistance);
-			}
-			else if (Implicit0Type == THeightField<T>::StaticType() && Implicit1Type == TSphere<T, d>::StaticType())
-			{
-				ensure(false);
-			}
-			else if (Implicit0Type == TPlane<T, d>::StaticType() && Implicit1Type == TSphere<T, d>::StaticType())
-			{
-				//UpdatePlaneSphereManifold(ConstraintBase, Transform0, Transform1, CullDistance);
-			}
-			else if (Implicit0Type == TBox<T, d>::StaticType() && Implicit1Type == TSphere<T, d>::StaticType())
-			{
-				//UpdateBoxSphereManifold(ConstraintBase, Transform0, Transform1, CullDistance);
-			}
-			else if (Implicit0Type == TBox<T, d>::StaticType() && Implicit1Type == TCapsule<T>::StaticType())
-			{
-				//UpdateBoxCapsuleManifold(ConstraintBase, Transform0, Transform1, CullDistance);
-			}
-			else if (Implicit0Type == TCapsule<T>::StaticType() && Implicit1Type == TSphere<T, d>::StaticType())
-			{
-				//UpdateCapsuleSphereManifold(ConstraintBase, Transform0, Transform1, CullDistance);
-			}
-			else if (Implicit0Type == TBox<T, d>::StaticType() && Implicit1Type == FTriangleMeshImplicitObject::StaticType())
-			{
-				//UpdateBoxTriangleMeshManifold(ConstraintBase, Transform0, Transform1, CullDistance);
-			}
-			else if (Implicit0Type == FTriangleMeshImplicitObject::StaticType() && Implicit1Type == TBox<T, d>::StaticType())
-			{
-				ensure(false);
-			}
-			else if (Implicit0Type == TSphere<T, d>::StaticType() && Implicit1Type == FTriangleMeshImplicitObject::StaticType())
-			{
-				//UpdateSphereTriangleMeshManifold(ConstraintBase, Transform0, Transform1, CullDistance);
-			}
-			else if (Implicit0Type == FTriangleMeshImplicitObject::StaticType() && Implicit1Type == TSphere<T, d>::StaticType())
-			{
-				ensure(false);
-			}
-			else if (Implicit0Type == TCapsule<T>::StaticType() && Implicit1Type == FTriangleMeshImplicitObject::StaticType())
-			{
-				//UpdateCapsuleTriangleMeshManifold(ConstraintBase, Transform0, Transform1, CullDistance);
-			}
-			else if (Implicit0Type == FTriangleMeshImplicitObject::StaticType() && Implicit1Type == TCapsule<T>::StaticType())
-			{
-				ensure(false);
-			}
-			else if (Implicit0Type == FConvex::StaticType() && Implicit1Type == FTriangleMeshImplicitObject::StaticType())
-			{
-				//UpdateConvexTriangleMeshManifold(ConstraintBase, Transform0, Transform1, CullDistance);
-			}
-			else if (Implicit0Type == FTriangleMeshImplicitObject::StaticType() && Implicit1Type == FConvex::StaticType())
-			{
+		template<typename T, int d>
+		inline void UpdateManifold(TRigidBodyMultiPointContactConstraint<T, d>& Constraint, const TRigidTransform<T, d>& ParticleTransform0, const TRigidTransform<T, d>& ParticleTransform1, const T CullDistance)
+		{
+			const FImplicitObject& Implicit0 = *Constraint.Manifold.Implicit[0];
+			const FImplicitObject& Implicit1 = *Constraint.Manifold.Implicit[1];
+			const TRigidTransform<T, d> Transform0 = Constraint.ImplicitTransform[0] * ParticleTransform0;
+			const TRigidTransform<T, d> Transform1 = Constraint.ImplicitTransform[1] * ParticleTransform1;
+
+			switch (Constraint.Manifold.ShapesType)
+			{
+			case EContactShapesType::CapsuleBox:
+				UpdateCapsuleBoxManifold<T, d>(*Implicit0.template GetObject<const TCapsule<T>>(), Transform0, Implicit1.template GetObject<const TBox<T, d>>()->GetAABB(), Transform1, CullDistance, Constraint);
+				break;
+			case EContactShapesType::ConvexConvex:
+				UpdateConvexConvexManifold(Constraint, Transform0, Transform1, CullDistance);
+				break;
+			default:
+				// This switch needs to contain all pair types that generate a Manifold. ConstructConstrainsImpl
 				ensure(false);
+				break;
 			}
-			//
-			// the generic convex bodies are last
-			//
+		}
 
-			else if (Implicit0Type == THeightField<T>::StaticType() && Implicit1.IsConvex())
-			{
+		// Run collision detection for the specified constraint to update the nearest contact point.
+		// NOTE: Transforms are world space shape transforms
+		// @todo(chaos): use a lookup table?
+		// @todo(chaos): add planes back in
+		// @todo(chaos): add the missing cases below
+		template<ECollisionUpdateType UpdateType, typename T, int d>
+		inline void UpdateConstraintFromGeometryImpl(TRigidBodyPointContactConstraint<T, d>& Constraint, const TRigidTransform<T, d>& Transform0, const TRigidTransform<T, d>& Transform1, const T CullDistance)
+		{
+			const FImplicitObject& Implicit0 = *Constraint.Manifold.Implicit[0];
+			const FImplicitObject& Implicit1 = *Constraint.Manifold.Implicit[1];
+
+			switch (Constraint.Manifold.ShapesType)
+			{
+			case EContactShapesType::SphereSphere:
+				UpdateSphereSphereConstraint(*Implicit0.template GetObject<TSphere<T, d>>(), Transform0, *Implicit1.template GetObject<TSphere<T, d>>(), Transform1, CullDistance, Constraint);
+				break;
+			case EContactShapesType::SphereCapsule:
+				UpdateSphereCapsuleConstraint(*Implicit0.template GetObject<TSphere<T, d>>(), Transform0, *Implicit1.template GetObject<TCapsule<T>>(), Transform1, CullDistance, Constraint);
+				break;
+			case EContactShapesType::SphereBox:
+				UpdateSphereBoxConstraint(*Implicit0.template GetObject<TSphere<T, d>>(), Transform0, Implicit1.template GetObject<TBox<T, d>>()->GetAABB(), Transform1, CullDistance, Constraint);
+				break;
+			case EContactShapesType::SphereConvex:
+				// MISSING CASE!!!
+				UpdateConvexConvexConstraint(Implicit0, Transform0, Implicit1, Transform1, CullDistance, Constraint);
+				break;
+			case EContactShapesType::SphereTriMesh:
+				UpdateSphereTriangleMeshConstraint(*Implicit0.template GetObject<TSphere<T, d>>(), Transform0, *GetInnerObject<FTriangleMeshImplicitObject>(Implicit1), Transform1, CullDistance, Constraint);
+				break;
+			case EContactShapesType::SphereHeightField:
+				UpdateSphereHeightFieldConstraint(*Implicit0.template GetObject<TSphere<T, d>>(), Transform0, *Implicit1.template GetObject< THeightField<T> >(), Transform1, CullDistance, Constraint);
+				break;
+			case EContactShapesType::SpherePlane:
+				UpdateSpherePlaneConstraint(*Implicit0.template GetObject<TSphere<T, d>>(), Transform0, *Implicit1.template GetObject<TPlane<T, d>>(), Transform1, CullDistance, Constraint);
+				break;
+			case EContactShapesType::CapsuleCapsule:
+				UpdateCapsuleCapsuleConstraint(*Implicit0.template GetObject<TCapsule<T>>(), Transform0, *Implicit1.template GetObject<TCapsule<T>>(), Transform1, CullDistance, Constraint);
+				break;
+			case EContactShapesType::CapsuleBox:
+				UpdateCapsuleBoxConstraint(*Implicit0.template GetObject<TCapsule<T>>(), Transform0, Implicit1.template GetObject<TBox<T, d>>()->GetAABB(), Transform1, CullDistance, Constraint);
+				break;
+			case EContactShapesType::CapsuleConvex:
+				// MISSING CASE!!!
+				UpdateConvexConvexConstraint(Implicit0, Transform0, Implicit1, Transform1, CullDistance, Constraint);
+				break;
+			case EContactShapesType::CapsuleTriMesh:
+				UpdateCapsuleTriangleMeshConstraint(*Implicit0.template GetObject<TCapsule<T>>(), Transform0, *GetInnerObject<FTriangleMeshImplicitObject>(Implicit1), Transform1, CullDistance, Constraint);
+				break;
+			case EContactShapesType::CapsuleHeightField:
+				UpdateCapsuleHeightFieldConstraint(*Implicit0.template GetObject<TCapsule<T>>(), Transform0, *Implicit1.template GetObject< THeightField<T> >(), Transform1, CullDistance, Constraint);
+				break;
+			case EContactShapesType::BoxBox:
+				UpdateBoxBoxConstraint(Implicit0.template GetObject<TBox<T, d>>()->GetAABB(), Transform0, Implicit1.template GetObject<TBox<T, d>>()->GetAABB(), Transform1, CullDistance, Constraint);
+				break;
+			case EContactShapesType::BoxConvex:
+				// MISSING CASE!!!
+				UpdateConvexConvexConstraint(Implicit0, Transform0, Implicit1, Transform1, CullDistance, Constraint);
+				break;
+			case EContactShapesType::BoxTriMesh:
+				UpdateBoxTriangleMeshConstraint(Implicit0.template GetObject<TBox<T, d>>()->GetAABB(), Transform0, *GetInnerObject<FTriangleMeshImplicitObject>(Implicit1), Transform1, CullDistance, Constraint);
+				break;
+			case EContactShapesType::BoxHeightField:
+				UpdateBoxHeightFieldConstraint(Implicit0.template GetObject<TBox<T, d>>()->GetAABB(), Transform0, *Implicit1.template GetObject< THeightField<T> >(), Transform1, CullDistance, Constraint);
+				break;
+			case EContactShapesType::BoxPlane:
+				UpdateBoxPlaneConstraint(Implicit0.template GetObject<TBox<T, d>>()->GetAABB(), Transform0, *Implicit1.template GetObject<TPlane<T, d>>(), Transform1, CullDistance, Constraint);
+				break;
+			case EContactShapesType::ConvexConvex:
+				UpdateConvexConvexConstraint(Implicit0, Transform0, Implicit1, Transform1, CullDistance, Constraint);
+				break;
+			case EContactShapesType::ConvexTriMesh:
+				UpdateConvexTriangleMeshConstraint(Implicit0, Transform0, *GetInnerObject<FTriangleMeshImplicitObject>(Implicit1), Transform1, CullDistance, Constraint);
+				break;
+			case EContactShapesType::ConvexHeightField:
+				UpdateConvexHeightFieldConstraint(Implicit0, Transform0, *Implicit1.template GetObject< THeightField<T> >(), Transform1, CullDistance, Constraint);
+				break;
+			case EContactShapesType::LevelSetLevelSet:
+				UpdateLevelsetLevelsetConstraint<UpdateType>(CullDistance, Constraint);
+				break;
+			default:
+				// Switch needs updating....
 				ensure(false);
+				break;
 			}
-#endif
-			else if (Implicit0.IsConvex() && Implicit1Type == THeightField<T>::StaticType())
-			{
-				UpdateConvexHeightFieldManifold(ConstraintBase, Transform0, Transform1, CullDistance);
-			}
-			else if (Implicit0.IsConvex() && Implicit1.IsConvex())
-			{
-				UpdateConvexConvexManifold(ConstraintBase, Transform0, Transform1, CullDistance);
-			}
-			else
-			{
-				UpdateLevelsetLevelsetManifold(ConstraintBase, Transform0, Transform1, CullDistance);
-			}
-
 		}
 
-		template<ECollisionUpdateType UpdateType, typename T, int d>
-		void UpdateAnyConstraint(TCollisionConstraintBase<T, d>& ConstraintBase, const FImplicitObject& Implicit0, const FImplicitObject& Implicit1, const TRigidTransform<T, d>& Transform0, const TRigidTransform<T, d>& Transform1, const T CullDistance)
+		template<typename T, int d>
+		void ConstructConstraintsImpl(TGeometryParticleHandle<T, d>* Particle0, TGeometryParticleHandle<T, d>* Particle1, const FImplicitObject* Implicit0, const FImplicitObject* Implicit1, const TRigidTransform<T, d>& Transform0, const TRigidTransform<T, d>& Transform1, const T CullDistance, FCollisionConstraintsArray& NewConstraints)
 		{
-#if !UE_BUILD_SHIPPING
-			EImplicitObjectType Implicit0OuterType = Implicit0.GetType();
-			EImplicitObjectType Implicit1OuterType = Implicit1.GetType();
-
-			if (Implicit0OuterType == TImplicitObjectTransformed<T, d>::StaticType())
-			{
-				ensure(false);//should not be possible to get this type, it should already be resolved by the constraint. (see ConstructConstraints)
-				return;
-			}
-			else if (Implicit1OuterType == TImplicitObjectTransformed<T, d>::StaticType())
-			{
-				ensure(false);//should not be possible to get this type, it should already be resolved by the constraint. (see ConstructConstraints)
-				return;
-			}
-			else if (Implicit0OuterType != FImplicitObjectUnion::StaticType() && Implicit1OuterType == FImplicitObjectUnion::StaticType())
-			{
-				ensure(false);//should not be possible to get this type, it should already be resolved by the constraint. (see ConstructConstraints)
-				return;
-			}
-			else if (Implicit0OuterType == FImplicitObjectUnion::StaticType() && Implicit1OuterType != FImplicitObjectUnion::StaticType())
-			{
-				ensure(false);//should not be possible to get this type, it should already be resolved by the constraint. (see ConstructConstraints)
-				return;
-			}
-			else if (Implicit0OuterType == FImplicitObjectUnion::StaticType() && Implicit1OuterType == FImplicitObjectUnion::StaticType())
-			{
-				ensure(false);//should not be possible to get this type, it should already be resolved by the constraint. (see ConstructConstraints)
-				return;
-			}
-#endif
+			EImplicitObjectType Implicit0Type = Implicit0 ? GetInnerType(Implicit0->GetType()) : ImplicitObjectType::Unknown;
+			EImplicitObjectType Implicit1Type = Implicit1 ? GetInnerType(Implicit1->GetType()) : ImplicitObjectType::Unknown;
 
-			//
-			// @todo(chaos): Collision Constraints (CollisionMap)
-			//    Modify Construct() and Update() use a CollisionMap indexed on 
-			//    EImplicitObjectType, instead of the if/else chain. Also, remove 
-			//    the blocks with the ensure(false), they are just for validation 
-			//    after the recent change. 
-			//
-			EImplicitObjectType Implicit0Type = GetInnerType(Implicit0.GetType());
-			EImplicitObjectType Implicit1Type = GetInnerType(Implicit1.GetType()); 
-			
-			if (Implicit0Type == TBox<T, d>::StaticType() && Implicit1Type == TBox<T, d>::StaticType())
+			if (Implicit0Type == TBox<FReal, 3>::StaticType() && Implicit1Type == TBox<FReal, 3>::StaticType())
 			{
-				UpdateBoxBoxConstraint(Implicit0.template GetObject<TBox<T, d>>()->GetAABB(), Transform0, Implicit1.template GetObject<TBox<T, d>>()->GetAABB(), Transform1, CullDistance, *ConstraintBase.template As<TRigidBodyPointContactConstraint<T, d>>());
+				ConstructBoxBoxConstraints(Particle0, Particle1, Implicit0, Implicit1, Transform0, Transform1, CullDistance, NewConstraints);
 			}
-			else if (Implicit0Type == TBox<T, d>::StaticType() && Implicit1Type == THeightField<T>::StaticType())
+			else if (Implicit0Type == TBox<FReal, 3>::StaticType() && Implicit1Type == THeightField<FReal>::StaticType())
 			{
-				UpdateBoxHeightFieldConstraint(Implicit0.template GetObject<TBox<T, d>>()->GetAABB(), Transform0, *Implicit1.template GetObject< THeightField<T> >(), Transform1, CullDistance, *ConstraintBase.template As<TRigidBodyPointContactConstraint<T, d>>());
+				ConstructBoxHeightFieldConstraints(Particle0, Particle1, Implicit0, Implicit1, Transform0, Transform1, CullDistance, NewConstraints);
 			}
-			else if (Implicit0Type == TSphere<T, d>::StaticType() && Implicit1Type == TSphere<T, d>::StaticType())
+			else if (Implicit0Type == THeightField<FReal>::StaticType() && Implicit1Type == TBox<FReal, 3>::StaticType())
 			{
-				UpdateSphereSphereConstraint(*Implicit0.template GetObject<TSphere<T, d>>(), Transform0, *Implicit1.template GetObject<TSphere<T, d>>(), Transform1, CullDistance, *ConstraintBase.template As<TRigidBodyPointContactConstraint<T, d>>());
+				ConstructBoxHeightFieldConstraints(Particle1, Particle0, Implicit1, Implicit0, Transform1, Transform0, CullDistance, NewConstraints);
 			}
-			else if (Implicit0Type == TSphere<T, d>::StaticType() && Implicit1Type == THeightField<T>::StaticType())
+			else if (Implicit0Type == TBox<FReal, 3>::StaticType() && Implicit1Type == TPlane<FReal, 3>::StaticType())
 			{
-				UpdateSphereHeightFieldConstraint(*Implicit0.template GetObject<TSphere<T, d>>(), Transform0, *Implicit1.template GetObject< THeightField<T> >(), Transform1, CullDistance, *ConstraintBase.template As<TRigidBodyPointContactConstraint<T, d>>());
+				ConstructBoxPlaneConstraints(Particle0, Particle1, Implicit0, Implicit1, Transform0, Transform1, CullDistance, NewConstraints);
 			}
-			else if (Implicit0Type == TBox<T, d>::StaticType() && Implicit1Type == TPlane<T, d>::StaticType())
+			else if (Implicit0Type == TPlane<FReal, 3>::StaticType() && Implicit1Type == TBox<FReal, 3>::StaticType())
 			{
-				UpdateBoxPlaneConstraint(Implicit0.template GetObject<TBox<T, d>>()->GetAABB(), Transform0, *Implicit1.template GetObject<TPlane<T, d>>(), Transform1, CullDistance, *ConstraintBase.template As<TRigidBodyPointContactConstraint<T, d>>());
+				ConstructBoxPlaneConstraints(Particle1, Particle0, Implicit1, Implicit0, Transform1, Transform0, CullDistance, NewConstraints);
 			}
-			else if (Implicit0Type == TSphere<T, d>::StaticType() && Implicit1Type == TPlane<T, d>::StaticType())
+			else if (Implicit0Type == TSphere<FReal, 3>::StaticType() && Implicit1Type == TSphere<FReal, 3>::StaticType())
 			{
-				UpdateSpherePlaneConstraint(*Implicit0.template GetObject<TSphere<T, d>>(), Transform0, *Implicit1.template GetObject<TPlane<T, d>>(), Transform1, CullDistance, *ConstraintBase.template As<TRigidBodyPointContactConstraint<T, d>>());
+				ConstructSphereSphereConstraints(Particle0, Particle1, Implicit0, Implicit1, Transform0, Transform1, CullDistance, NewConstraints);
 			}
-			else if (Implicit0Type == TSphere<T, d>::StaticType() && Implicit1Type == TBox<T, d>::StaticType())
+			else if (Implicit0Type == TSphere<FReal, 3>::StaticType() && Implicit1Type == THeightField<FReal>::StaticType())
 			{
-				UpdateSphereBoxConstraint(*Implicit0.template GetObject<TSphere<T, d>>(), Transform0, Implicit1.template GetObject<TBox<T, d>>()->GetAABB(), Transform1, CullDistance, *ConstraintBase.template As<TRigidBodyPointContactConstraint<T, d>>());
+				ConstructSphereHeightFieldConstraints(Particle0, Particle1, Implicit0, Implicit1, Transform0, Transform1, CullDistance, NewConstraints);
 			}
-			else if (Implicit0Type == TSphere<T, d>::StaticType() && Implicit1Type == TCapsule<T>::StaticType())
+			else if (Implicit0Type == THeightField<FReal>::StaticType() && Implicit1Type == TSphere<FReal, 3>::StaticType())
 			{
-				UpdateSphereCapsuleConstraint(*Implicit0.template GetObject<TSphere<T, d>>(), Transform0, *Implicit1.template GetObject<TCapsule<T>>(), Transform1, CullDistance, *ConstraintBase.template As<TRigidBodyPointContactConstraint<T, d>>());
+				ConstructSphereHeightFieldConstraints(Particle1, Particle0, Implicit1, Implicit0, Transform1, Transform0, CullDistance, NewConstraints);
 			}
-			else if (Implicit0Type == TCapsule<T>::StaticType() && Implicit1Type == TCapsule<T>::StaticType())
+			else if (Implicit0Type == TSphere<FReal, 3>::StaticType() && Implicit1Type == TPlane<FReal, 3>::StaticType())
 			{
-				UpdateCapsuleCapsuleConstraint(*Implicit0.template GetObject<TCapsule<T>>(), Transform0, *Implicit1.template GetObject<TCapsule<T>>(), Transform1, CullDistance, *ConstraintBase.template As<TRigidBodyPointContactConstraint<T, d>>());
+				ConstructSpherePlaneConstraints(Particle0, Particle1, Implicit0, Implicit1, Transform0, Transform1, CullDistance, NewConstraints);
 			}
-			else if (Implicit0Type == TCapsule<T>::StaticType() && Implicit1Type == TBox<T, d>::StaticType())
+			else if (Implicit0Type == TPlane<FReal, 3>::StaticType() && Implicit1Type == TSphere<FReal, 3>::StaticType())
 			{
-				UpdateCapsuleBoxConstraint(*Implicit0.template GetObject<TCapsule<T>>(), Transform0, Implicit1.template GetObject<TBox<T, d>>()->GetAABB(), Transform1, CullDistance, *ConstraintBase.template As<TRigidBodyPointContactConstraint<T, d>>());
+				ConstructSpherePlaneConstraints(Particle1, Particle0, Implicit1, Implicit0, Transform1, Transform0, CullDistance, NewConstraints);
 			}
-			else if (Implicit0Type == TCapsule<T>::StaticType() && Implicit1Type == THeightField<T>::StaticType())
+			else if (Implicit0Type == TSphere<FReal, 3>::StaticType() && Implicit1Type == TBox<FReal, 3>::StaticType())
 			{
-				UpdateCapsuleHeightFieldConstraint(*Implicit0.template GetObject<TCapsule<T>>(), Transform0, *Implicit1.template GetObject< THeightField<T> >(), Transform1, CullDistance, *ConstraintBase.template As<TRigidBodyPointContactConstraint<T, d>>());
+				ConstructSphereBoxConstraints(Particle0, Particle1, Implicit0, Implicit1, Transform0, Transform1, CullDistance, NewConstraints);
 			}
-			else if (Implicit0Type == TPlane<T, d>::StaticType() && Implicit1Type == TBox<T, d>::StaticType())
+			else if (Implicit0Type == TBox<FReal, 3>::StaticType() && Implicit1Type == TSphere<FReal, 3>::StaticType())
 			{
-				TRigidBodyPointContactConstraint<T, d>& Constraint = *ConstraintBase.template As<TRigidBodyPointContactConstraint<T, d> >();
-				TRigidBodyPointContactConstraint<T, d> TmpConstraint = Constraint;
-				UpdateBoxPlaneConstraint(Implicit1.template GetObject<TBox<T, d>>()->GetAABB(), Transform1, *Implicit0.template GetObject<TPlane<T, d>>(), Transform0, CullDistance, TmpConstraint);
-				if (TmpConstraint.GetPhi() < Constraint.GetPhi())
-				{
-					Constraint = TmpConstraint;
-					Constraint.SetNormal(-Constraint.GetNormal());
-				}
+				ConstructSphereBoxConstraints(Particle1, Particle0, Implicit1, Implicit0, Transform1, Transform0, CullDistance, NewConstraints);
 			}
-			else if (Implicit0Type == TPlane<T, d>::StaticType() && Implicit1Type == TSphere<T, d>::StaticType())
+			else if (Implicit0Type == TSphere<FReal, 3>::StaticType() && Implicit1Type == TCapsule<FReal>::StaticType())
 			{
-				TRigidBodyPointContactConstraint<T, d>& Constraint = *ConstraintBase.template As<TRigidBodyPointContactConstraint<T, d> >();
-				TRigidBodyPointContactConstraint<T, d> TmpConstraint = Constraint;
-				UpdateSpherePlaneConstraint(*Implicit1.template GetObject<TSphere<T, d>>(), Transform1, *Implicit0.template GetObject<TPlane<T, d>>(), Transform0, CullDistance, TmpConstraint);
-				if (TmpConstraint.GetPhi() < Constraint.GetPhi())
-				{
-					Constraint = TmpConstraint;
-					Constraint.SetNormal(-Constraint.GetNormal());
-				}
+				ConstructSphereCapsuleConstraints(Particle0, Particle1, Implicit0, Implicit1, Transform0, Transform1, CullDistance, NewConstraints);
 			}
-			else if (Implicit0Type == TBox<T, d>::StaticType() && Implicit1Type == TSphere<T, d>::StaticType())
+			else if (Implicit0Type == TCapsule<FReal>::StaticType() && Implicit1Type == TSphere<FReal, 3>::StaticType())
 			{
-				TRigidBodyPointContactConstraint<T, d>& Constraint = *ConstraintBase.template As<TRigidBodyPointContactConstraint<T, d> >();
-				TRigidBodyPointContactConstraint<T, d> TmpConstraint = Constraint;
-				UpdateSphereBoxConstraint(*Implicit1.template GetObject<TSphere<T, d>>(), Transform1, Implicit0.template GetObject<TBox<T, d>>()->GetAABB(), Transform0, CullDistance, TmpConstraint);
-				if (TmpConstraint.GetPhi() < Constraint.GetPhi())
-				{
-					Constraint = TmpConstraint;
-					Constraint.SetNormal(-Constraint.GetNormal());
-				}
+				ConstructSphereCapsuleConstraints(Particle1, Particle0, Implicit1, Implicit0, Transform1, Transform0, CullDistance, NewConstraints);
 			}
-			else if (Implicit0Type == TBox<T, d>::StaticType() && Implicit1Type == TCapsule<T>::StaticType())
+			else if (Implicit0Type == TCapsule<FReal>::StaticType() && Implicit1Type == TCapsule<FReal>::StaticType())
 			{
-				TRigidBodyPointContactConstraint<T, d>& Constraint = *ConstraintBase.template As<TRigidBodyPointContactConstraint<T, d> >();
-				TRigidBodyPointContactConstraint<T, d> TmpConstraint = Constraint;
-				UpdateCapsuleBoxConstraint(*Implicit1.template GetObject<TCapsule<T>>(), Transform1, Implicit0.template GetObject<TBox<T, d>>()->GetAABB(), Transform0, CullDistance, TmpConstraint);
-				if (TmpConstraint.GetPhi() < Constraint.GetPhi())
-				{
-					Constraint = TmpConstraint;
-					Constraint.SetNormal(-Constraint.GetNormal());
-				}
+				ConstructCapsuleCapsuleConstraints(Particle0, Particle1, Implicit0, Implicit1, Transform0, Transform1, CullDistance, NewConstraints);
 			}
-			else if (Implicit0Type == TCapsule<T>::StaticType() && Implicit1Type == TSphere<T, d>::StaticType())
+			else if (Implicit0Type == TCapsule<FReal>::StaticType() && Implicit1Type == TBox<FReal, 3>::StaticType())
 			{
-				TRigidBodyPointContactConstraint<T, d>& Constraint = *ConstraintBase.template As<TRigidBodyPointContactConstraint<T, d> >();
-				TRigidBodyPointContactConstraint<T, d> TmpConstraint = Constraint;
-				UpdateSphereCapsuleConstraint(*Implicit1.template GetObject<TSphere<T, d>>(), Transform1, *Implicit0.template GetObject<TCapsule<T>>(), Transform0, CullDistance, TmpConstraint);
-				if (TmpConstraint.GetPhi() < Constraint.GetPhi())
-				{
-					Constraint = TmpConstraint;
-					Constraint.SetNormal(-Constraint.GetNormal());
-				}
+				ConstructCapsuleBoxConstraints(Particle0, Particle1, Implicit0, Implicit1, Transform0, Transform1, CullDistance, NewConstraints);
 			}
-			else if (Implicit0Type == TBox<T, d>::StaticType() && Implicit1Type == FTriangleMeshImplicitObject::StaticType())
+			else if (Implicit0Type == TBox<FReal, 3>::StaticType() && Implicit1Type == TCapsule<FReal>::StaticType())
 			{
-				UpdateBoxTriangleMeshConstraint(Implicit0.template GetObject<TBox<T, d>>()->GetAABB(), Transform0, *GetInnerObject<FTriangleMeshImplicitObject>(Implicit1), Transform1, CullDistance, *ConstraintBase.template As<TRigidBodyPointContactConstraint<T, d>>());
+				ConstructCapsuleBoxConstraints(Particle1, Particle0, Implicit1, Implicit0, Transform1, Transform0, CullDistance, NewConstraints);
 			}
-			else if (Implicit0Type == TSphere<T, d>::StaticType() && Implicit1Type == FTriangleMeshImplicitObject::StaticType())
+			else if (Implicit0Type == TCapsule<FReal>::StaticType() && Implicit1Type == THeightField<FReal>::StaticType())
 			{
-				UpdateSphereTriangleMeshConstraint(*Implicit0.template GetObject<TSphere<T, d>>(), Transform0, *GetInnerObject<FTriangleMeshImplicitObject>(Implicit1), Transform1, CullDistance, *ConstraintBase.template As<TRigidBodyPointContactConstraint<T, d>>());
+				ConstructCapsuleHeightFieldConstraints(Particle0, Particle1, Implicit0, Implicit1, Transform0, Transform1, CullDistance, NewConstraints);
 			}
-			else if (Implicit0Type == TCapsule<T>::StaticType() && Implicit1Type == FTriangleMeshImplicitObject::StaticType())
+			else if (Implicit0Type == THeightField<FReal>::StaticType() && Implicit1Type == TCapsule<FReal>::StaticType())
 			{
-				UpdateCapsuleTriangleMeshConstraint(*Implicit0.template GetObject<TCapsule<T>>(), Transform0, *GetInnerObject<FTriangleMeshImplicitObject>(Implicit1), Transform1, CullDistance, *ConstraintBase.template As<TRigidBodyPointContactConstraint<T, d>>());
+				ConstructCapsuleHeightFieldConstraints(Particle1, Particle0, Implicit1, Implicit0, Transform1, Transform0, CullDistance, NewConstraints);
 			}
-#if !UE_BUILD_SHIPPING
-			else if (Implicit0Type == THeightField<T>::StaticType() && Implicit1Type == TBox<T, d>::StaticType())
+			else if (Implicit0Type == TBox<FReal, 3>::StaticType() && Implicit1Type == FTriangleMeshImplicitObject::StaticType())
 			{
-				//     This case should not be necessary. The height fields 
-				//     will only ever be collided against, so ideally will never 
-				//     be in index[0] position of the constraint, also the construction
-				//     of the constraint will just switch the index position so its always 
-				//     second.
-				ensure(false);
+				ConstructBoxTriangleMeshConstraints(Particle0, Particle1, Implicit0, Implicit1, Transform0, Transform1, CullDistance, NewConstraints);
 			}
-			else if (Implicit0Type == THeightField<T>::StaticType() && Implicit1Type == TSphere<T, d>::StaticType())
+			else if (Implicit0Type == FTriangleMeshImplicitObject::StaticType() && Implicit1Type == TBox<FReal, 3>::StaticType())
 			{
-				//     This case should not be necessary. The height fields 
-				//     will only ever be collided against, so ideally will never 
-				//     be in index[0] position of the constraint, also the construction
-				//     of the constraint will just switch the index position so its always 
-				//     second.
-				ensure(false);
+				ConstructBoxTriangleMeshConstraints(Particle1, Particle0, Implicit1, Implicit0, Transform1, Transform0, CullDistance, NewConstraints);
 			}
-			else if (Implicit0Type == THeightField<T>::StaticType() && Implicit1Type == TSphere<T, d>::StaticType())
+			else if (Implicit0Type == TSphere<FReal, 3>::StaticType() && Implicit1Type == FTriangleMeshImplicitObject::StaticType())
 			{
-				//     This case should not be necessary. The height fields 
-				//     will only ever be collided against, so ideally will never 
-				//     be in index[0] position of the constraint, also the construction
-				//     of the constraint will just switch the index position so its always 
-				//     second.
-				ensure(false);
+				ConstructSphereTriangleMeshConstraints(Particle0, Particle1, Implicit0, Implicit1, Transform0, Transform1, CullDistance, NewConstraints);
 			}
-			else if (Implicit0Type == FTriangleMeshImplicitObject::StaticType() && Implicit1Type == TBox<T, d>::StaticType())
+			else if (Implicit0Type == FTriangleMeshImplicitObject::StaticType() && Implicit1Type == TSphere<FReal, 3>::StaticType())
 			{
-				//     This case should not be necessary. The triangle mesh
-				//     will only ever be collided against, so ideally will never
-				//     be in index[0] position of the constraint, also the construction
-				//     of the constraint will just switch the index position so its always
-				//     second.
-				ensure(false);
+				ConstructSphereTriangleMeshConstraints(Particle1, Particle0, Implicit1, Implicit0, Transform1, Transform0, CullDistance, NewConstraints);
 			}
-			else if (Implicit0Type == FTriangleMeshImplicitObject::StaticType() && Implicit1Type == TSphere<T, d>::StaticType())
+			else if (Implicit0Type == TCapsule<FReal>::StaticType() && Implicit1Type == FTriangleMeshImplicitObject::StaticType())
 			{
-				//     This case should not be necessary. The triangle mesh
-				//     will only ever be collided against, so ideally will never
-				//     be in index[0] position of the constraint, also the construction
-				//     of the constraint will just switch the index position so its always
-				//     second.
-				ensure(false);
+				ConstructCapsuleTriangleMeshConstraints(Particle0, Particle1, Implicit0, Implicit1, Transform0, Transform1, CullDistance, NewConstraints);
 			}
-			else if (Implicit0Type == FTriangleMeshImplicitObject::StaticType() && Implicit1Type == TCapsule<T>::StaticType() )
+			else if (Implicit0Type == FTriangleMeshImplicitObject::StaticType() && Implicit1Type == TCapsule<FReal>::StaticType())
 			{
-				//     This case should not be necessary. The triangle mesh
-				//     will only ever be collided against, so ideally will never
-				//     be in index[0] position of the constraint, also the construction
-				//     of the constraint will just switch the index position so its always
-				//     second.
-				ensure(false);
+				ConstructCapsuleTriangleMeshConstraints(Particle1, Particle0, Implicit1, Implicit0, Transform1, Transform0, CullDistance, NewConstraints);
 			}
-
-			//
-			// the generic convex bodies are last
-			//
-
-
-			else if (Implicit0Type == THeightField<T>::StaticType() && Implicit1.IsConvex())
+			else if (Implicit0->IsConvex() && Implicit1Type == THeightField<FReal>::StaticType())
 			{
-				//     This case should not be necessary. The height fields 
-				//     will only ever be collided against, so ideally will never 
-				//     be in index[0] position of the constraint, also the construction
-				//     of the constraint will just switch the index position so its always 
-				//     second.
-				ensure(false);
+				ConstructCapsuleHeightFieldConstraints(Particle0, Particle1, Implicit0, Implicit1, Transform0, Transform1, CullDistance, NewConstraints);
 			}
-			else if (Implicit0Type == FTriangleMeshImplicitObject::StaticType() && Implicit1.IsConvex())
+			else if (Implicit0Type == THeightField<FReal>::StaticType() && Implicit1->IsConvex())
 			{
-				//     This case should not be necessary. The triangle mesh
-				//     will only ever be collided against, so ideally will never
-				//     be in index[0] position of the constraint, also the construction
-				//     of the constraint will just switch the index position so its always
-				//     second.
-				ensure(false);
+				ConstructCapsuleHeightFieldConstraints(Particle1, Particle0, Implicit1, Implicit0, Transform1, Transform0, CullDistance, NewConstraints);
 			}
-#endif
-			else if (Implicit0.IsConvex() && Implicit1Type == THeightField<T>::StaticType())
+			else if (Implicit0->IsConvex() && Implicit1Type == FTriangleMeshImplicitObject::StaticType())
 			{
-				UpdateConvexHeightFieldConstraint(Implicit0, Transform0, *Implicit1.template GetObject< THeightField<T> >(), Transform1, CullDistance, *ConstraintBase.template As<TRigidBodyPointContactConstraint<T, d>>());
+				ConstructConvexTriangleMeshConstraints(Particle0, Particle1, Implicit0, Implicit1, Transform0, Transform1, CullDistance, NewConstraints);
 			}
-			else if (Implicit0Type == FConvex::StaticType() && Implicit1Type == FTriangleMeshImplicitObject::StaticType())
+			else if (Implicit0Type == FTriangleMeshImplicitObject::StaticType() && Implicit1->IsConvex())
 			{
-				UpdateConvexTriangleMeshConstraint(Implicit0, Transform0, *GetInnerObject<FTriangleMeshImplicitObject>(Implicit1), Transform1, CullDistance, *ConstraintBase.template As<TRigidBodyPointContactConstraint<T, d>>());
+				ConstructConvexTriangleMeshConstraints(Particle1, Particle0, Implicit1, Implicit0, Transform1, Transform0, CullDistance, NewConstraints);
 			}
-			else if (Implicit0.IsConvex() && Implicit1.IsConvex())
+			else if (Implicit0->IsConvex() && Implicit1->IsConvex())
 			{
-				UpdateConvexConvexConstraint(Implicit0, Transform0, Implicit1, Transform1, CullDistance, ConstraintBase);
+				ConstructConvexConvexConstraints(Particle0, Particle1, Implicit0, Implicit1, Transform0, Transform1, CullDistance, NewConstraints);
 			}
 			else
 			{
-				UpdateLevelsetLevelsetConstraint<UpdateType>(CullDistance, *ConstraintBase.template As<TRigidBodyPointContactConstraint<T, d> >());
+				ConstructLevelsetLevelsetConstraints(Particle0, Particle1, Implicit0, Implicit1, Transform0, Transform1, CullDistance, NewConstraints);
 			}
 		}
 
+
+		// Run collision detection for the specified constraint to update the nearest contact point.
+		// NOTE: Transforms are world space particle transforms
 		template<ECollisionUpdateType UpdateType, typename T, int d>
-		void UpdateConstraint(TCollisionConstraintBase<T, d>& ConstraintBase, const TRigidTransform<T, d>& ParticleTransform0, const TRigidTransform<T, d>& ParticleTransform1, const T CullDistance)
+		void UpdateConstraintFromGeometry(TRigidBodyPointContactConstraint<T, d>& Constraint, const TRigidTransform<T, d>& ParticleTransform0, const TRigidTransform<T, d>& ParticleTransform1, const T CullDistance)
 		{
-			if (ConstraintBase.GetType() == TCollisionConstraintBase<T, d>::FType::MultiPoint)
+			const TRigidTransform<T, d> Transform0 = Constraint.ImplicitTransform[0] * ParticleTransform0;
+			const TRigidTransform<T, d> Transform1 = Constraint.ImplicitTransform[1] * ParticleTransform1;
+			UpdateConstraintFromGeometryImpl<UpdateType>(Constraint, Transform0, Transform1, CullDistance);
+		}
+
+		// Select the best manifold point as the new contact point.
+		// NOTE: Transforms are world space particle transforms
+		template<class T, int d>
+		void UpdateConstraintFromManifold(TRigidBodyMultiPointContactConstraint<T, d>& Constraint, const TRigidTransform<T, d>& ParticleTransform0, const TRigidTransform<T, d>& ParticleTransform1, const T CullDistance)
+		{
+			const FRigidTransform3 Transform0 = Constraint.ImplicitTransform[0] * ParticleTransform0;
+			const FRigidTransform3 Transform1 = Constraint.ImplicitTransform[1] * ParticleTransform1;
+
+			// TEMP: Test code - run collision detection instead. Allows us to debug draw the manifold without using it
+			if (Chaos_Collision_UseManifolds_Test)
 			{
-				// Select the best point from the available manifold points as the current contact
-				TRigidBodyMultiPointContactConstraint<T, d>& MultiPointConstraint = *ConstraintBase.template As<TRigidBodyMultiPointContactConstraint<T, d>>();
-				UpdateConstraintFromManifold(MultiPointConstraint, ParticleTransform0, ParticleTransform1, CullDistance);
+				// Copy the manifold constraint into point constraint, update it, and copy the results back
+				TRigidBodyPointContactConstraint<T, d> PointConstraint(Constraint.Particle[0], Constraint.Manifold.Implicit[0], Constraint.ImplicitTransform[0], Constraint.Particle[1], Constraint.Manifold.Implicit[1], Constraint.ImplicitTransform[1], Constraint.Manifold.ShapesType);
+				UpdateConstraintFromGeometryImpl<ECollisionUpdateType::Deepest>(PointConstraint, Transform0, Transform1, CullDistance);
+				Constraint.Manifold = PointConstraint.Manifold;
+				return;
 			}
-			else if (ConstraintBase.GetType() == TCollisionConstraintBase<T, d>::FType::SinglePoint)
+
+			// Get the plane and point transforms (depends which body owns the plane)
+			const FRigidTransform3& PlaneTransform = (Constraint.GetManifoldPlaneOwnerIndex() == 0) ? Transform0 : Transform1;
+			const FRigidTransform3& PointsTransform = (Constraint.GetManifoldPlaneOwnerIndex() == 0) ? Transform1 : Transform0;
+
+			// World-space manifold plane
+			FVec3 PlaneNormal = PlaneTransform.TransformVector(Constraint.GetManifoldPlaneNormal());
+			FVec3 PlanePos = PlaneTransform.TransformPosition(Constraint.GetManifoldPlanePosition());
+
+			// Select the best manifold point
+			const int32 NumPoints = Constraint.NumManifoldPoints();
+			for (int32 PointIndex = 0; PointIndex < NumPoints; ++PointIndex)
 			{
-				const FImplicitObject& Implicit0 = *ConstraintBase.Manifold.Implicit[0];
-				const FImplicitObject& Implicit1 = *ConstraintBase.Manifold.Implicit[1];
-				const TRigidTransform<T, d> Transform0 = ConstraintBase.ImplicitTransform[0] * ParticleTransform0;
-				const TRigidTransform<T, d> Transform1 = ConstraintBase.ImplicitTransform[1] * ParticleTransform1;
+				// World-space manifold point and distance to manifold plane
+				FVec3 PointPos = PointsTransform.TransformPosition(Constraint.GetManifoldPoint(PointIndex));
+				FReal PointDistance = FVec3::DotProduct(PointPos - PlanePos, PlaneNormal);
 
-				// Re-run collision detection to generate a new contact
-				// @todo(chaos): all constraints should have a ShapesType set in ConstructConstraints. Then we can use a function table for all.
-				TRigidBodyPointContactConstraint<T, d>& PointConstraint = *ConstraintBase.template As<TRigidBodyPointContactConstraint<T, d>>();
-				switch (ConstraintBase.Manifold.ShapesType)
+				// If this is the deepest hit, select it
+				if (PointDistance < Constraint.Manifold.Phi)
 				{
-				case EContactShapesType::CapsuleCapsule:
-					UpdateCapsuleCapsuleConstraint(*Implicit0.template GetObject<TCapsule<T>>(), Transform0, *Implicit1.template GetObject<TCapsule<T>>(), Transform1, CullDistance, PointConstraint);
-					break;
-				case EContactShapesType::CapsuleBox:
-					UpdateCapsuleBoxConstraint(*Implicit0.template GetObject<TCapsule<T>>(), Transform0, Implicit1.template GetObject<TBox<T, d>>()->GetAABB(), Transform1, CullDistance, PointConstraint);
-					break;
-				case EContactShapesType::BoxBox:
-					UpdateBoxBoxConstraint(Implicit0.template GetObject<TBox<T, d>>()->GetAABB(), Transform0, Implicit1.template GetObject<TBox<T, d>>()->GetAABB(), Transform1, CullDistance, PointConstraint);
-					break;
-				default:
-					UpdateAnyConstraint<UpdateType, T, d>(ConstraintBase, Implicit0, Implicit1, Transform0, Transform1, CullDistance);
-					break;
+					// @todo(chaos): Consider using average of plane and point positions for contact location
+					FVec3 ContactPos = PointPos - PointDistance * PlaneNormal;
+					FVec3 ContactNormal = (Constraint.GetManifoldPlaneOwnerIndex() == 0) ? -PlaneNormal : PlaneNormal;
+					Constraint.Manifold.Phi = PointDistance;
+					Constraint.Manifold.Location = ContactPos;
+					Constraint.Manifold.Normal = ContactNormal;
 				}
 			}
 		}
 
+
 		typedef uint8 FMaskFilter;
 		enum { NumExtraFilterBits = 6 };
 		enum { NumCollisionChannelBits = 5 };
@@ -1925,7 +1751,7 @@ namespace Chaos
 		}
 
 		template<class T, int d>
-		bool DoCollide(EImplicitObjectType Implicit0Type, const TPerShapeData<T, d>* Shape0, EImplicitObjectType Implicit1Type, const TPerShapeData<T, d>* Shape1)
+		inline bool DoCollide(EImplicitObjectType Implicit0Type, const TPerShapeData<T, d>* Shape0, EImplicitObjectType Implicit1Type, const TPerShapeData<T, d>* Shape1)
 		{
 			//
 			// Disabled shapes do not collide
@@ -1995,6 +1821,45 @@ namespace Chaos
 			return true;
 		}
 
+
+		template<typename T, int d>
+		void ConstructPointConstraint(EContactShapesType ShapesType, TGeometryParticleHandle<T, d>* Particle0, TGeometryParticleHandle<T, d>* Particle1, const FImplicitObject* Implicit0, const FImplicitObject* Implicit1, const TRigidTransform<T, d>& Transform0, const TRigidTransform<T, d>& Transform1, const T CullDistance, FCollisionConstraintsArray& NewConstraints)
+		{
+			if (ensure(Implicit0 && Implicit1))
+			{
+				TRigidTransform<T, d> ParticleImplicit0TM = Transform0.GetRelativeTransform(Collisions::GetTransform(Particle0));
+				TRigidTransform<T, d> ParticleImplicit1TM = Transform1.GetRelativeTransform(Collisions::GetTransform(Particle1));
+				if (FRigidBodyPointContactConstraint* Constraint = NewConstraints.AddPointConstraint(Particle0, Implicit0, ParticleImplicit0TM, Particle1, Implicit1, ParticleImplicit1TM, ShapesType))
+				{
+					UpdateConstraintFromGeometryImpl<ECollisionUpdateType::Deepest>(*Constraint, Transform0, Transform1, CullDistance);
+					if (Constraint->GetPhi() >= CullDistance)
+					{
+						NewConstraints.PopPointConstraint();
+					}
+				}
+			}
+		}
+
+
+		template<typename T, int d>
+		void ConstructManifoldConstraint(EContactShapesType ShapesType, TGeometryParticleHandle<T, d>* Particle0, TGeometryParticleHandle<T, d>* Particle1, const FImplicitObject* Implicit0, const FImplicitObject* Implicit1, const TRigidTransform<T, d>& Transform0, const TRigidTransform<T, d>& Transform1, const T CullDistance, FCollisionConstraintsArray& NewConstraints)
+		{
+			if (ensure(Implicit0 && Implicit1))
+			{
+				TRigidTransform<T, d> ParticleImplicit0TM = Transform0.GetRelativeTransform(Collisions::GetTransform(Particle0));
+				TRigidTransform<T, d> ParticleImplicit1TM = Transform1.GetRelativeTransform(Collisions::GetTransform(Particle1));
+				if (FRigidBodyMultiPointContactConstraint* Constraint = NewConstraints.AddMultiPointConstraint(Particle0, Implicit0, ParticleImplicit0TM, Particle1, Implicit1, ParticleImplicit1TM, ShapesType))
+				{
+					UpdateManifold(*Constraint, Transform0, Transform1, CullDistance);
+					if (Constraint->GetPhi() >= CullDistance)
+					{
+						NewConstraints.PopMultiPointConstraint();
+					}
+				}
+			}
+		}
+
+
 		template<typename T, int d>
 		void ConstructConstraints(TGeometryParticleHandle<T, d>* Particle0, TGeometryParticleHandle<T, d>* Particle1, const FImplicitObject* Implicit0, const FImplicitObject* Implicit1, const TRigidTransform<T, d>& Transform0, const TRigidTransform<T, d>& Transform1, const T CullDistance, FCollisionConstraintsArray& NewConstraints)
 		{
@@ -2006,21 +1871,26 @@ namespace Chaos
 				return;
 			}
 
-
 			if (!Implicit0 || !Implicit1)
 			{
 				ConstructLevelsetLevelsetConstraints(Particle0, Particle1, Implicit0, Implicit1, Transform0, Transform1, CullDistance, NewConstraints);
 				return;
 			}
 
-			//
-			// @todo(chaos): Collision Constraints (CollisionMap)
-			//    Modify Construct() and Update() use a CollisionMap indexed on EImplicitObjectType, instead of the if/else chain
-			//
 			EImplicitObjectType Implicit0OuterType = Implicit0->GetType();
 			EImplicitObjectType Implicit1OuterType = Implicit1->GetType();
 
-			if (Implicit0OuterType == TImplicitObjectTransformed<T, d>::StaticType())
+			// Handle transform wrapper shape
+			if ((Implicit0OuterType == TImplicitObjectTransformed<T, d>::StaticType()) && (Implicit1OuterType == TImplicitObjectTransformed<T, d>::StaticType()))
+			{
+				const TImplicitObjectTransformed<FReal, 3>* TransformedImplicit0 = Implicit0->template GetObject<const TImplicitObjectTransformed<FReal, 3>>();
+				const TImplicitObjectTransformed<FReal, 3>* TransformedImplicit1 = Implicit1->template GetObject<const TImplicitObjectTransformed<FReal, 3>>();
+				TRigidTransform<T, d> TransformedTransform0 = TransformedImplicit0->GetTransform() * Transform0;
+				TRigidTransform<T, d> TransformedTransform1 = TransformedImplicit1->GetTransform() * Transform1;
+				ConstructConstraints(Particle0, Particle1, TransformedImplicit0->GetTransformedObject(), TransformedImplicit1->GetTransformedObject(), TransformedTransform0, TransformedTransform1, CullDistance, NewConstraints);
+				return;
+			}
+			else if (Implicit0OuterType == TImplicitObjectTransformed<T, d>::StaticType())
 			{
 				const TImplicitObjectTransformed<FReal, 3>* TransformedImplicit0 = Implicit0->template GetObject<const TImplicitObjectTransformed<FReal, 3>>();
 				TRigidTransform<T, d> TransformedTransform0 = TransformedImplicit0->GetTransform() * Transform0;
@@ -2035,71 +1905,29 @@ namespace Chaos
 				return;
 			}
 
-			else if (Implicit0OuterType == TImplicitObjectInstanced<FConvex>::StaticType())
-			{
-				const TImplicitObjectInstanced<FConvex>* TransformedImplicit0 = Implicit0->template GetObject<const TImplicitObjectInstanced<FConvex>>();
-				ConstructConstraints(Particle0, Particle1, TransformedImplicit0->GetInstancedObject(), Implicit1, Transform0, Transform1, CullDistance, NewConstraints);
-				return;
-			}
-			else if (Implicit1OuterType == TImplicitObjectInstanced<FConvex>::StaticType())
-			{
-				const TImplicitObjectInstanced<FConvex>* TransformedImplicit1 = Implicit1->template GetObject<const TImplicitObjectInstanced<FConvex>>();
-				ConstructConstraints(Particle0, Particle1, Implicit0, TransformedImplicit1->GetInstancedObject(), Transform0, Transform1, CullDistance, NewConstraints);
-				return;
-			}
-
-			else if (Implicit0OuterType == TImplicitObjectInstanced<TBox<FReal,3>>::StaticType())
-			{
-				const TImplicitObjectInstanced<TBox<FReal, 3>>* TransformedImplicit0 = Implicit0->template GetObject<const TImplicitObjectInstanced<TBox<FReal, 3>>>();
-				ConstructConstraints(Particle0, Particle1, TransformedImplicit0->GetInstancedObject(), Implicit1, Transform0, Transform1, CullDistance, NewConstraints);
-				return;
-			}
-			else if (Implicit1OuterType == TImplicitObjectInstanced<TBox<FReal, 3>>::StaticType())
+			// Handle Instanced shapes
+			// NOTE: Tri Meshes are handled differently. We should probably do something about this...
+			if (((uint32)Implicit0OuterType & ImplicitObjectType::IsInstanced) || ((uint32)Implicit1OuterType & ImplicitObjectType::IsInstanced))
 			{
-				const TImplicitObjectInstanced<TBox<FReal, 3>>* TransformedImplicit1 = Implicit1->template GetObject<const TImplicitObjectInstanced<TBox<FReal, 3>>>();
-				ConstructConstraints(Particle0, Particle1, Implicit0, TransformedImplicit1->GetInstancedObject(), Transform0, Transform1, CullDistance, NewConstraints);
-				return;
-			}
-
-			else if (Implicit0OuterType == TImplicitObjectInstanced<TCapsule<FReal>>::StaticType())
-			{
-				const TImplicitObjectInstanced<TCapsule<FReal>>* TransformedImplicit0 = Implicit0->template GetObject<const TImplicitObjectInstanced<TCapsule<FReal>>>();
-				ConstructConstraints(Particle0, Particle1, TransformedImplicit0->GetInstancedObject(), Implicit1, Transform0, Transform1, CullDistance, NewConstraints);
-				return;
-			}
-			else if (Implicit1OuterType == TImplicitObjectInstanced<TCapsule<FReal>>::StaticType())
-			{
-				const TImplicitObjectInstanced<TCapsule<FReal>>* TransformedImplicit1 = Implicit1->template GetObject<const TImplicitObjectInstanced<TCapsule<FReal>>>();
-				ConstructConstraints(Particle0, Particle1, Implicit0, TransformedImplicit1->GetInstancedObject(), Transform0, Transform1, CullDistance, NewConstraints);
-				return;
-			}
-
-			else if (Implicit0OuterType == TImplicitObjectInstanced<TSphere<FReal, 3>>::StaticType())
-			{
-				const TImplicitObjectInstanced<TSphere<FReal, 3>>* TransformedImplicit0 = Implicit0->template GetObject<const TImplicitObjectInstanced<TSphere<FReal, 3>>>();
-				ConstructConstraints(Particle0, Particle1, TransformedImplicit0->GetInstancedObject(), Implicit1, Transform0, Transform1, CullDistance, NewConstraints);
-				return;
-			}
-			else if (Implicit1OuterType == TImplicitObjectInstanced<TSphere<FReal, 3>>::StaticType())
-			{
-				const TImplicitObjectInstanced<TSphere<FReal, 3>>* TransformedImplicit1 = Implicit1->template GetObject<const TImplicitObjectInstanced<TSphere<FReal, 3>>>();
-				ConstructConstraints(Particle0, Particle1, Implicit0, TransformedImplicit1->GetInstancedObject(), Transform0, Transform1, CullDistance, NewConstraints);
-				return;
+				const FImplicitObject* InnerImplicit0 = Implicit0;
+				const FImplicitObject* InnerImplicit1 = Implicit1;
+				if ((uint32)Implicit0OuterType & ImplicitObjectType::IsInstanced)
+				{
+					InnerImplicit0 = GetInstancedImplicit(Implicit0);
+				}
+				if ((uint32)Implicit1OuterType & ImplicitObjectType::IsInstanced)
+				{
+					InnerImplicit1 = GetInstancedImplicit(Implicit1);
+				}
+				if (InnerImplicit0 && InnerImplicit1)
+				{
+					ConstructConstraints(Particle0, Particle1, InnerImplicit0, InnerImplicit1, Transform0, Transform1, CullDistance, NewConstraints);
+					return;
+				}
 			}
 
-			else if (Implicit0OuterType == TImplicitObjectInstanced<FConvex>::StaticType())
-			{
-				const TImplicitObjectInstanced<FConvex>* TransformedImplicit0 = Implicit0->template GetObject<const TImplicitObjectInstanced<FConvex>>();
-				ConstructConstraints(Particle0, Particle1, TransformedImplicit0->GetInstancedObject(), Implicit1, Transform0, Transform1, CullDistance, NewConstraints);
-				return;
-			}
-			else if (Implicit1OuterType == TImplicitObjectInstanced<FConvex>::StaticType())
-			{
-				const TImplicitObjectInstanced<FConvex>* TransformedImplicit1 = Implicit1->template GetObject<const TImplicitObjectInstanced<FConvex>>();
-				ConstructConstraints(Particle0, Particle1, Implicit0, TransformedImplicit1->GetInstancedObject(), Transform0, Transform1, CullDistance, NewConstraints);
-				return;
-			}
-			else if (Implicit0OuterType != FImplicitObjectUnion::StaticType() && Implicit1OuterType == FImplicitObjectUnion::StaticType())
+			// Handle Unions
+			if (Implicit0OuterType != FImplicitObjectUnion::StaticType() && Implicit1OuterType == FImplicitObjectUnion::StaticType())
 			{
 				const TArray<Pair<const FImplicitObject*, TRigidTransform<T, d>>> LevelsetShapes = FindRelevantShapes(Implicit0, Transform0, *Implicit1, Transform1, CullDistance);
 				for (const Pair<const FImplicitObject*, TRigidTransform<T, d>>& LevelsetObjPair : LevelsetShapes)
@@ -2128,136 +1956,9 @@ namespace Chaos
 				return;
 			}
 
-
-			if (Implicit0Type == TBox<T, d>::StaticType() && Implicit1Type == TBox<T, d>::StaticType())
-			{
-				ConstructBoxBoxConstraints(Particle0, Particle1, Implicit0, Implicit1, Transform0, Transform1, CullDistance, NewConstraints);
-			}
-			else if (Implicit0Type == TBox<T, d>::StaticType() && Implicit1Type == THeightField<T>::StaticType())
-			{
-				ConstructBoxHeightFieldConstraints(Particle0, Particle1, Implicit0, Implicit1, Transform0, Transform1, CullDistance, NewConstraints);
-			}
-			else if (Implicit0Type == THeightField<T>::StaticType() && Implicit1Type == TBox<T, d>::StaticType())
-			{
-				ConstructBoxHeightFieldConstraints(Particle1, Particle0, Implicit1, Implicit0, Transform1, Transform0, CullDistance, NewConstraints);
-			}
-			else if (Implicit0Type == TBox<T, d>::StaticType() && Implicit1Type == TPlane<T, d>::StaticType())
-			{
-				ConstructBoxPlaneConstraints(Particle0, Particle1, Implicit0, Implicit1, Transform0, Transform1, CullDistance, NewConstraints);
-			}
-			else if (Implicit0Type == TPlane<T, d>::StaticType() && Implicit1Type == TBox<T, d>::StaticType())
-			{
-				ConstructBoxPlaneConstraints(Particle1, Particle0, Implicit1, Implicit0, Transform1, Transform0, CullDistance, NewConstraints);
-			}
-			else if (Implicit0Type == TSphere<T, d>::StaticType() && Implicit1Type == TSphere<T, d>::StaticType())
-			{
-				ConstructSphereSphereConstraints(Particle0, Particle1, Implicit0, Implicit1, Transform0, Transform1, CullDistance, NewConstraints);
-			}
-			else if (Implicit0Type == TSphere<T, d>::StaticType() && Implicit1Type == THeightField<T>::StaticType())
-			{
-				ConstructSphereHeightFieldConstraints(Particle0, Particle1, Implicit0, Implicit1, Transform0, Transform1, CullDistance, NewConstraints);
-			}
-			else if (Implicit0Type == THeightField<T>::StaticType() && Implicit1Type == TSphere<T, d>::StaticType())
-			{
-				ConstructSphereHeightFieldConstraints(Particle1, Particle0, Implicit1, Implicit0, Transform1, Transform0, CullDistance, NewConstraints);
-			}
-			else if (Implicit0Type == TSphere<T, d>::StaticType() && Implicit1Type == TPlane<T, d>::StaticType())
-			{
-				ConstructSpherePlaneConstraints(Particle0, Particle1, Implicit0, Implicit1, Transform0, Transform1, CullDistance, NewConstraints);
-			}
-			else if (Implicit0Type == TPlane<T, d>::StaticType() && Implicit1Type == TSphere<T, d>::StaticType())
-			{
-				ConstructSpherePlaneConstraints(Particle1, Particle0, Implicit1, Implicit0, Transform1, Transform0, CullDistance, NewConstraints);
-			}
-			else if (Implicit0Type == TSphere<T, d>::StaticType() && Implicit1Type == TBox<T, d>::StaticType())
-			{
-				ConstructSphereBoxConstraints(Particle0, Particle1, Implicit0, Implicit1, Transform0, Transform1, CullDistance, NewConstraints);
-			}
-			else if (Implicit0Type == TBox<T, d>::StaticType() && Implicit1Type == TSphere<T, d>::StaticType())
-			{
-				ConstructSphereBoxConstraints(Particle1, Particle0, Implicit1, Implicit0, Transform1, Transform0, CullDistance, NewConstraints);
-			}
-			else if (Implicit0Type == TSphere<T, d>::StaticType() && Implicit1Type == TCapsule<T>::StaticType())
-			{
-				ConstructSphereCapsuleConstraints(Particle0, Particle1, Implicit0, Implicit1, Transform0, Transform1, CullDistance, NewConstraints);
-			}
-			else if (Implicit0Type == TCapsule<T>::StaticType() && Implicit1Type == TSphere<T, d>::StaticType())
-			{
-				ConstructSphereCapsuleConstraints(Particle1, Particle0, Implicit1, Implicit0, Transform1, Transform0, CullDistance, NewConstraints);
-			}
-			else if (Implicit0Type == TCapsule<T>::StaticType() && Implicit1Type == TCapsule<T>::StaticType())
-			{
-				ConstructCapsuleCapsuleConstraints(Particle0, Particle1, Implicit0, Implicit1, Transform0, Transform1, CullDistance, NewConstraints);
-			}
-			else if (Implicit0Type == TCapsule<T>::StaticType() && Implicit1Type == TBox<T, d>::StaticType())
-			{
-				ConstructCapsuleBoxConstraints(Particle0, Particle1, Implicit0, Implicit1, Transform0, Transform1, CullDistance, NewConstraints);
-			}
-			else if (Implicit0Type == TBox<T, d>::StaticType() && Implicit1Type == TCapsule<T>::StaticType())
-			{
-				ConstructCapsuleBoxConstraints(Particle1, Particle0, Implicit1, Implicit0, Transform1, Transform0, CullDistance, NewConstraints);
-			}
-			else if (Implicit0Type == TCapsule<T>::StaticType() && Implicit1Type == THeightField<T>::StaticType())
-			{
-				ConstructCapsuleHeightFieldConstraints(Particle0, Particle1, Implicit0, Implicit1, Transform0, Transform1, CullDistance, NewConstraints);
-			}
-			else if (Implicit0Type == THeightField<T>::StaticType() && Implicit1Type == TCapsule<T>::StaticType())
-			{
-				ConstructCapsuleHeightFieldConstraints(Particle1, Particle0, Implicit1, Implicit0, Transform1, Transform0, CullDistance, NewConstraints);
-			}
-			else if (Implicit0Type == TBox<T, d>::StaticType() && Implicit1Type == FTriangleMeshImplicitObject::StaticType())
-			{
-				ConstructBoxTriangleMeshConstraints(Particle0, Particle1, Implicit0, Implicit1, Transform0, Transform1, CullDistance, NewConstraints);
-			}
-			else if (Implicit0Type == FTriangleMeshImplicitObject::StaticType() && Implicit1Type == TBox<T, d>::StaticType())
-			{
-				ConstructBoxTriangleMeshConstraints(Particle1, Particle0, Implicit1, Implicit0, Transform1, Transform0, CullDistance, NewConstraints);
-			}
-			else if (Implicit0Type == TSphere<T, d>::StaticType() && Implicit1Type == FTriangleMeshImplicitObject::StaticType())
-			{
-				ConstructSphereTriangleMeshConstraints(Particle0, Particle1, Implicit0, Implicit1, Transform0, Transform1, CullDistance, NewConstraints);
-			}
-			else if (Implicit0Type == FTriangleMeshImplicitObject::StaticType() && Implicit1Type == TSphere<T, d>::StaticType())
-			{
-				ConstructSphereTriangleMeshConstraints(Particle1, Particle0, Implicit1, Implicit0, Transform1, Transform0, CullDistance, NewConstraints);
-			}
-			else if (Implicit0Type == TCapsule<T>::StaticType() && Implicit1Type == FTriangleMeshImplicitObject::StaticType())
-			{
-				ConstructCapsuleTriangleMeshConstraints(Particle0, Particle1, Implicit0, Implicit1, Transform0, Transform1, CullDistance, NewConstraints);
-			}
-			else if (Implicit0Type == FTriangleMeshImplicitObject::StaticType() && Implicit1Type == TCapsule<T>::StaticType())
-			{
-				ConstructCapsuleTriangleMeshConstraints(Particle1, Particle0, Implicit1, Implicit0, Transform1, Transform0, CullDistance, NewConstraints);
-			}
-
-			//
-			// the generic convex bodies are last
-			//
-
-			else if (Implicit0->IsConvex() && Implicit1Type == THeightField<T>::StaticType())
-			{
-				ConstructConvexHeightFieldConstraints(Particle0, Particle1, Implicit0, Implicit1, Transform0, Transform1, CullDistance, NewConstraints);
-			}
-			else if (Implicit0Type == THeightField<T>::StaticType() && Implicit1->IsConvex())
-			{
-				ConstructConvexHeightFieldConstraints(Particle1, Particle0, Implicit1, Implicit0, Transform1, Transform0, CullDistance, NewConstraints);
-			}
-			else if (Implicit0->IsConvex() && Implicit1Type == FTriangleMeshImplicitObject::StaticType())
-			{
-				ConstructConvexTriangleMeshConstraints(Particle0, Particle1, Implicit0, Implicit1, Transform0, Transform1, CullDistance, NewConstraints);
-			}
-			else if (Implicit0Type == FTriangleMeshImplicitObject::StaticType() && Implicit1->IsConvex())
-			{
-				ConstructConvexTriangleMeshConstraints(Particle1, Particle0, Implicit1, Implicit0, Transform1, Transform0, CullDistance, NewConstraints);
-			}
-			else if (Implicit0->IsConvex() && Implicit1->IsConvex())
-			{
-				ConstructConvexConvexConstraints(Particle0, Particle1, Implicit0, Implicit1, Transform0, Transform1, CullDistance, NewConstraints);
-			}
-			else
-			{
-				ConstructLevelsetLevelsetConstraints(Particle0, Particle1, Implicit0, Implicit1, Transform0, Transform1, CullDistance, NewConstraints);
-			}
+			// If we get here, we have a pair of concrete shapes (i.e., no wrappers or containers)
+			// Create a constraint for the shape pair
+			ConstructConstraintsImpl(Particle0, Particle1, Implicit0, Implicit1, Transform0, Transform1, CullDistance, NewConstraints);
 		}
 
 		template void UpdateBoxBoxConstraint<float, 3>(const TAABB<float, 3>& Box1, const TRigidTransform<float, 3>& Box1Transform, const TAABB<float, 3>& Box2, const TRigidTransform<float, 3>& Box2Transform, const float CullDistance, TRigidBodyPointContactConstraint<float, 3>& Constraint);
@@ -2321,9 +2022,9 @@ namespace Chaos
 
 		template void UpdateSingleShotManifold<float, 3>(TRigidBodyMultiPointContactConstraint<float, 3>&  Constraint, const TRigidTransform<float, 3>& Transform0, const TRigidTransform<float, 3>& Transform1, const float CullDistance);
 		template void UpdateIterativeManifold<float, 3>(TRigidBodyMultiPointContactConstraint<float, 3>&  Constraint, const TRigidTransform<float, 3>& Transform0, const TRigidTransform<float, 3>& Transform1, const float CullDistance);
-		template void UpdateManifold<float, 3>(TCollisionConstraintBase<float, 3>& ConstraintBase, const TRigidTransform<float, 3>& ATM, const TRigidTransform<float, 3>& BTM, const float CullDistance);
-		template void UpdateConstraint<ECollisionUpdateType::Any, float, 3>(TCollisionConstraintBase<float, 3>& ConstraintBase, const TRigidTransform<float, 3>& Transform0, const TRigidTransform<float, 3>& Transform1, const float CullDistance);
-		template void UpdateConstraint<ECollisionUpdateType::Deepest, float, 3>(TCollisionConstraintBase<float, 3>& ConstraintBase, const TRigidTransform<float, 3>& Transform0, const TRigidTransform<float, 3>& Transform1, const float CullDistance);
+		template void UpdateManifold<float, 3>(TRigidBodyMultiPointContactConstraint<float, 3>& Constraint, const TRigidTransform<float, 3>& ATM, const TRigidTransform<float, 3>& BTM, const float CullDistance);
+		template void UpdateConstraintFromGeometry<ECollisionUpdateType::Any, float, 3>(TRigidBodyPointContactConstraint<float, 3>& ConstraintBase, const TRigidTransform<float, 3>& Transform0, const TRigidTransform<float, 3>& Transform1, const float CullDistance);
+		template void UpdateConstraintFromGeometry<ECollisionUpdateType::Deepest, float, 3>(TRigidBodyPointContactConstraint<float, 3>& ConstraintBase, const TRigidTransform<float, 3>& Transform0, const TRigidTransform<float, 3>& Transform1, const float CullDistance);
 		template void UpdateConstraintFromManifold(TRigidBodyMultiPointContactConstraint<float, 3>& Constraint, const TRigidTransform<float, 3>& Transform0, const TRigidTransform<float, 3>& Transform1, const float CullDistance);
 		template void ConstructConstraints<float, 3>(TGeometryParticleHandle<float, 3>* Particle0, TGeometryParticleHandle<float, 3>* Particle1, const FImplicitObject* Implicit0, const FImplicitObject* Implicit1, const TRigidTransform<float, 3>& Transform0, const TRigidTransform<float, 3>& Transform1, const float CullDistance, FCollisionConstraintsArray& NewConstraints);
 
diff --git a/Engine/Source/Runtime/Experimental/Chaos/Private/Chaos/PBDCollisionConstraintsContact.cpp b/Engine/Source/Runtime/Experimental/Chaos/Private/Chaos/PBDCollisionConstraintsContact.cpp
--- a/Engine/Source/Runtime/Experimental/Chaos/Private/Chaos/PBDCollisionConstraintsContact.cpp
+++ b/Engine/Source/Runtime/Experimental/Chaos/Private/Chaos/PBDCollisionConstraintsContact.cpp
@@ -30,7 +30,7 @@ namespace Chaos
 			const FRigidTransform3 Transform1 = GetTransform(Constraint.Particle[1]);
 
 			Constraint.ResetPhi(CullDistance);
-			UpdateConstraint<ECollisionUpdateType::Deepest>(Constraint, Transform0, Transform1, CullDistance);
+			UpdateConstraintFromGeometry<ECollisionUpdateType::Deepest>(Constraint, Transform0, Transform1, CullDistance);
 		}
 
 		void Update(FRigidBodyMultiPointContactConstraint& Constraint, const FReal CullDistance)
diff --git a/Engine/Source/Runtime/Experimental/Chaos/Public/Chaos/CollisionResolution.h b/Engine/Source/Runtime/Experimental/Chaos/Public/Chaos/CollisionResolution.h
--- a/Engine/Source/Runtime/Experimental/Chaos/Public/Chaos/CollisionResolution.h
+++ b/Engine/Source/Runtime/Experimental/Chaos/Public/Chaos/CollisionResolution.h
@@ -252,14 +252,14 @@ namespace Chaos
 		void CHAOS_API UpdateIterativeManifold(TRigidBodyMultiPointContactConstraint<T, d>&  Constraint, const TRigidTransform<T, d>& Transform0, const TRigidTransform<T, d>& Transform1, const T CullDistance);
 
 		template<typename T, int d>
-		void CHAOS_API UpdateManifold(TCollisionConstraintBase<T, d>& ConstraintBase, const TRigidTransform<T, d>& ATM, const TRigidTransform<T, d>& BTM, const T CullDistance);
-
-		template<ECollisionUpdateType UpdateType, typename T, int d>
-		void CHAOS_API UpdateConstraint(TCollisionConstraintBase<T, d>& ConstraintBase, const TRigidTransform<T, d>& Transform0, const TRigidTransform<T, d>& Transform1, const T CullDistance);
+		void CHAOS_API UpdateManifold(TRigidBodyMultiPointContactConstraint<T, d>& Constraint, const TRigidTransform<T, d>& ATM, const TRigidTransform<T, d>& BTM, const T CullDistance);
 
 		template<class T, int d>
 		void UpdateConstraintFromManifold(TRigidBodyMultiPointContactConstraint<T, d>& Constraint, const TRigidTransform<T, d>& Transform0, const TRigidTransform<T, d>& Transform1, const T CullDistance);
 
+		template<ECollisionUpdateType UpdateType, typename T, int d>
+		void UpdateConstraintFromGeometry(TRigidBodyPointContactConstraint<T, d>& Constraint, const TRigidTransform<T, d>& ParticleTransform0, const TRigidTransform<T, d>& ParticleTransform1, const T CullDistance);
+
 		template<typename T, int d>
 		void CHAOS_API ConstructConstraints(TGeometryParticleHandle<T, d>* Particle0, TGeometryParticleHandle<T, d>* Particle1, const FImplicitObject* Implicit0, const FImplicitObject* Implicit1, const TRigidTransform<T, d>& Transform0, const TRigidTransform<T, d>& Transform1, const T CullDistance, FCollisionConstraintsArray& NewConstraints);
 
diff --git a/Engine/Source/Runtime/Experimental/Chaos/Public/Chaos/CollisionResolutionTypes.h b/Engine/Source/Runtime/Experimental/Chaos/Public/Chaos/CollisionResolutionTypes.h
--- a/Engine/Source/Runtime/Experimental/Chaos/Public/Chaos/CollisionResolutionTypes.h
+++ b/Engine/Source/Runtime/Experimental/Chaos/Public/Chaos/CollisionResolutionTypes.h
@@ -41,9 +41,29 @@ namespace Chaos
 	enum class CHAOS_API EContactShapesType
 	{
 		Unknown,
+		SphereSphere,
+		SphereCapsule,
+		SphereBox,
+		SphereConvex,
+		SphereTriMesh,
+		SphereHeightField,
+		SpherePlane,
 		CapsuleCapsule,
 		CapsuleBox,
+		CapsuleConvex,
+		CapsuleTriMesh,
+		CapsuleHeightField,
 		BoxBox,
+		BoxConvex,
+		BoxTriMesh,
+		BoxHeightField,
+		BoxPlane,
+		ConvexConvex,
+		ConvexTriMesh,
+		ConvexHeightField,
+		LevelSetLevelSet,
+
+		NumShapesTypes
 	};
 
 	/*
@@ -187,7 +207,7 @@ namespace Chaos
 		TRigidBodyPointContactConstraint(
 			FGeometryParticleHandle* Particle0, const FImplicitObject* Implicit0, const TRigidTransform<T, d>& Transform0,
 			FGeometryParticleHandle* Particle1, const FImplicitObject* Implicit1, const TRigidTransform<T, d>& Transform1,
-			EContactShapesType ShapesType = EContactShapesType::Unknown)
+			EContactShapesType ShapesType)
 			: Base(Particle0, Implicit0, Transform0, Particle1, Implicit1, Transform1, Base::FType::SinglePoint, ShapesType) {}
 
 		static typename Base::FType StaticType() { return Base::FType::SinglePoint; };
@@ -211,7 +231,7 @@ namespace Chaos
 		TRigidBodyMultiPointContactConstraint(
 			FGeometryParticleHandle* Particle0, const FImplicitObject* Implicit0, const TRigidTransform<T, d>& Transform0,
 			FGeometryParticleHandle* Particle1, const FImplicitObject* Implicit1, const TRigidTransform<T, d>& Transform1,
-			EContactShapesType ShapesType = EContactShapesType::Unknown)
+			EContactShapesType ShapesType)
 			: Base(Particle0, Implicit0, Transform0, Particle1, Implicit1, Transform1, Base::FType::MultiPoint, ShapesType)
 			, PlaneOwnerIndex(INDEX_NONE), PlaneFaceIndex(INDEX_NONE), PlaneNormal(0), PlanePosition(0)
 		{}
@@ -386,8 +406,8 @@ namespace Chaos
 	template<int T_MAXCONSTRAINTS>
 	struct TCollisionConstraintsStore
 	{
-		TArray<TRigidBodyPointContactConstraint<FReal, 3>, TFixedAllocator<T_MAXCONSTRAINTS>> SinglePointConstraints;
-		TArray<TRigidBodyMultiPointContactConstraint<FReal, 3>, TFixedAllocator<T_MAXCONSTRAINTS>> MultiPointConstraints;
+		TArray<FRigidBodyPointContactConstraint, TFixedAllocator<T_MAXCONSTRAINTS>> SinglePointConstraints;
+		TArray<FRigidBodyMultiPointContactConstraint, TFixedAllocator<T_MAXCONSTRAINTS>> MultiPointConstraints;
 
 		int32 Num() const { return SinglePointConstraints.Num() + MultiPointConstraints.Num(); }
 
@@ -397,6 +417,45 @@ namespace Chaos
 			MultiPointConstraints.Empty();
 		}
 
+		FRigidBodyPointContactConstraint* AddPointConstraint(
+			TGeometryParticleHandle<FReal, 3>* Particle0, const FImplicitObject* Implicit0, const FRigidTransform3& Transform0,
+			TGeometryParticleHandle<FReal, 3>* Particle1, const FImplicitObject* Implicit1, const FRigidTransform3& Transform1,
+			EContactShapesType ShapesType)
+		{
+			if (SinglePointConstraints.Num() < T_MAXCONSTRAINTS)
+			{
+				int32 ConstraintIndex = SinglePointConstraints.Emplace(FRigidBodyPointContactConstraint(Particle0, Implicit0, Transform0, Particle1, Implicit1, Transform1, ShapesType));
+				return &SinglePointConstraints[ConstraintIndex];
+			}
+			return nullptr;
+		}
+
+		void PopPointConstraint()
+		{
+			SinglePointConstraints.Pop(false);
+		}
+
+		FRigidBodyMultiPointContactConstraint* AddMultiPointConstraint(
+			TGeometryParticleHandle<FReal, 3>* Particle0, const FImplicitObject* Implicit0, const FRigidTransform3& Transform0,
+			TGeometryParticleHandle<FReal, 3>* Particle1, const FImplicitObject* Implicit1, const FRigidTransform3& Transform1,
+			EContactShapesType ShapesType)
+		{
+			if (MultiPointConstraints.Num() < T_MAXCONSTRAINTS)
+			{
+				int32 ConstraintIndex = MultiPointConstraints.Emplace(FRigidBodyMultiPointContactConstraint(Particle0, Implicit0, Transform0, Particle1, Implicit1, Transform1, ShapesType));
+				return &MultiPointConstraints[ConstraintIndex];
+			}
+			return nullptr;
+		}
+
+		void PopMultiPointConstraint()
+		{
+			MultiPointConstraints.Pop();
+		}
+
+
+
+		// Get rid of these
 		TRigidBodyPointContactConstraint<FReal, 3>* TryAdd(FReal MaxPhi, const TRigidBodyPointContactConstraint<FReal, 3>& C)
 		{
 			if ((SinglePointConstraints.Num() < T_MAXCONSTRAINTS) && (C.GetPhi() < MaxPhi))
