Chaos - Fix headless Chaos Deformable Gravity test.

diff --git a/Engine/Source/Runtime/Experimental/Chaos/Private/Chaos/PBDEvolution.cpp b/Engine/Source/Runtime/Experimental/Chaos/Private/Chaos/PBDEvolution.cpp
--- a/Engine/Source/Runtime/Experimental/Chaos/Private/Chaos/PBDEvolution.cpp
+++ b/Engine/Source/Runtime/Experimental/Chaos/Private/Chaos/PBDEvolution.cpp
@@ -43,6 +43,7 @@ TPBDEvolution<T, d>::TPBDEvolution(TPBDParticles<T, d>&& InParticles, TKinematic
 {
 	MCollisionParticles.AddArray(&MCollided);
 	MParticles.AddArray(&MParticleGroupIds);
+	MPerGroupDamping.Add(Damping);
 
 	SetParticleUpdateFunction(
 		[PBDUpdateRule = 
