Re-add child implicits to ImplicitShapeMap to fix contact mod.

diff --git a/Engine/Source/Runtime/Experimental/Chaos/Private/Chaos/GeometryParticles.cpp b/Engine/Source/Runtime/Experimental/Chaos/Private/Chaos/GeometryParticles.cpp
--- a/Engine/Source/Runtime/Experimental/Chaos/Private/Chaos/GeometryParticles.cpp
+++ b/Engine/Source/Runtime/Experimental/Chaos/Private/Chaos/GeometryParticles.cpp
@@ -198,10 +198,16 @@ namespace Chaos
 		TShapesArray<T, d>& ShapeArray = MShapesArray[Index];
 		Mapping.Reset();
 
-		int32 ShapeIndex = 0;
-		for (TUniquePtr<TPerShapeData<T, d>>& Shape : ShapeArray)
+		for (int32 ShapeIndex = 0; ShapeIndex < ShapeArray.Num(); ++ ShapeIndex)
 		{
-			Mapping.Add(Shape->Geometry.Get(), ShapeIndex++);
+			const FImplicitObject* ImplicitObject = ShapeArray[ShapeIndex]->Geometry.Get();
+			Mapping.Add(ImplicitObject, ShapeIndex);
+
+			const FImplicitObject* ImplicitChildObject = Utilities::ImplicitChildHelper(ImplicitObject);
+			if (ImplicitChildObject != ImplicitObject)
+			{
+				Mapping.Add(ImplicitChildObject, ShapeIndex);
+			}
 		}
 
 		if (MGeometry[Index])
diff --git a/Engine/Source/Runtime/Experimental/Chaos/Private/Chaos/ParticleHandle.cpp b/Engine/Source/Runtime/Experimental/Chaos/Private/Chaos/ParticleHandle.cpp
--- a/Engine/Source/Runtime/Experimental/Chaos/Private/Chaos/ParticleHandle.cpp
+++ b/Engine/Source/Runtime/Experimental/Chaos/Private/Chaos/ParticleHandle.cpp
@@ -13,10 +13,17 @@ namespace Chaos
 	void Chaos::TGeometryParticle<T, d>::MapImplicitShapes()
 	{
 		ImplicitShapeMap.Reset();
-		int32 ShapeIndex = 0;
-		for (TUniquePtr<TPerShapeData<T, d>>& ShapeData : MShapesArray)
+
+		for (int32 ShapeIndex = 0; ShapeIndex < MShapesArray.Num(); ++ ShapeIndex)
 		{
-			ImplicitShapeMap.Add(ShapeData->Geometry.Get(), ShapeIndex++);
+			const FImplicitObject* ImplicitObject = MShapesArray[ShapeIndex]->Geometry.Get();
+			ImplicitShapeMap.Add(ImplicitObject, ShapeIndex);
+
+			const FImplicitObject* ImplicitChildObject = Utilities::ImplicitChildHelper(ImplicitObject);
+			if (ImplicitChildObject != ImplicitObject)
+			{
+				ImplicitShapeMap.Add(ImplicitChildObject, ShapeIndex);
+			}
 		}
 
 		if (MGeometry)
