Chaos - fix Convex-TrMesh

diff --git a/Engine/Source/Runtime/Experimental/Chaos/Private/Chaos/CollisionResolution.cpp b/Engine/Source/Runtime/Experimental/Chaos/Private/Chaos/CollisionResolution.cpp
--- a/Engine/Source/Runtime/Experimental/Chaos/Private/Chaos/CollisionResolution.cpp
+++ b/Engine/Source/Runtime/Experimental/Chaos/Private/Chaos/CollisionResolution.cpp
@@ -36,6 +36,11 @@ FAutoConsoleVariableRef CVarChaosCollisionUseManifolds(TEXT("p.Chaos.Collision.U
 bool Chaos_Collision_UseManifolds_Test = false;
 FAutoConsoleVariableRef CVarChaosCollisionUseManifoldsTest(TEXT("p.Chaos.Collision.UseManifoldsTest"), Chaos_Collision_UseManifolds_Test, TEXT("Enable/Disable use of manifoldes in collision."));
 
+float Chaos_Collision_ManifoldFaceAngle = 5.0f;
+float Chaos_Collision_ManifoldFaceEpsilon = FMath::Sin(FMath::DegreesToRadians(Chaos_Collision_ManifoldFaceAngle));
+FConsoleVariableDelegate Chaos_Collision_ManifoldFaceDelegate = FConsoleVariableDelegate::CreateLambda([](IConsoleVariable* CVar) { Chaos_Collision_ManifoldFaceEpsilon = FMath::Sin(FMath::DegreesToRadians(Chaos_Collision_ManifoldFaceAngle)); });
+FAutoConsoleVariableRef CVarChaosCollisionManifoldFaceAngle(TEXT("p.Chaos.Collision.ManifoldFaceAngle"), Chaos_Collision_ManifoldFaceAngle, TEXT("Angle above which a face is rejected and we switch to point collision"), Chaos_Collision_ManifoldFaceDelegate);
+
 
 namespace Chaos
 {
@@ -909,7 +914,7 @@ namespace Chaos
 
 			// Find the box feature that the near point is on
 			// Face, Edge, or Vertex can be determined from number of non-zero elements in the box-space normal.
-			const FReal ComponentEpsilon = 0.09f; // ~Sin(5deg) = angle tolerance
+			const FReal ComponentEpsilon = Chaos_Collision_ManifoldFaceEpsilon;
 			int32 NumNonZeroNormalComponents = 0;
 			int32 MaxComponentIndex = INDEX_NONE;
 			FReal MaxComponentValue = 0;
@@ -1318,7 +1323,7 @@ namespace Chaos
 			{
 				TRigidTransform<T, d> ParticleImplicit0TM = Transform0.GetRelativeTransform(Collisions::GetTransform(Particle0));
 				TRigidTransform<T, d> ParticleImplicit1TM = Transform1.GetRelativeTransform(Collisions::GetTransform(Particle1));
-				FRigidBodyPointContactConstraint Constraint = FRigidBodyPointContactConstraint(Particle0, Implicit0, ParticleImplicit0TM, Particle1, Implicit1, ParticleImplicit1TM, EContactShapesType::ConvexConvex);
+				FRigidBodyPointContactConstraint Constraint = FRigidBodyPointContactConstraint(Particle0, Implicit0, ParticleImplicit0TM, Particle1, Implicit1, ParticleImplicit1TM, EContactShapesType::ConvexTriMesh);
 				UpdateConvexTriangleMeshConstraint(*Implicit0, Transform0, *Object1, Transform1, CullDistance, Constraint);
 				NewConstraints.TryAdd(CullDistance, Constraint);
 			}
@@ -1454,8 +1459,8 @@ namespace Chaos
 		// Run collision detection for the specified constraint to update the nearest contact point.
 		// NOTE: Transforms are world space shape transforms
 		// @todo(chaos): use a lookup table?
-		// @todo(chaos): add planes back in
 		// @todo(chaos): add the missing cases below
+		// @todo(chaos): see use GetInnerObject below - we should try to uise the leaf types for all (currently only TriMesh needs this)
 		template<ECollisionUpdateType UpdateType, typename T, int d>
 		inline void UpdateConstraintFromGeometryImpl(TRigidBodyPointContactConstraint<T, d>& Constraint, const TRigidTransform<T, d>& Transform0, const TRigidTransform<T, d>& Transform1, const T CullDistance)
 		{
@@ -1540,6 +1545,7 @@ namespace Chaos
 		template<typename T, int d>
 		void ConstructConstraintsImpl(TGeometryParticleHandle<T, d>* Particle0, TGeometryParticleHandle<T, d>* Particle1, const FImplicitObject* Implicit0, const FImplicitObject* Implicit1, const TRigidTransform<T, d>& Transform0, const TRigidTransform<T, d>& Transform1, const T CullDistance, FCollisionConstraintsArray& NewConstraints)
 		{
+			// @todo(chaos): We use GetInnerType here because TriMeshes are lef with their "Instanced" wrapper, unlike all other instanced implicits. Should we strip the instance on Tri Mesh too?
 			EImplicitObjectType Implicit0Type = Implicit0 ? GetInnerType(Implicit0->GetType()) : ImplicitObjectType::Unknown;
 			EImplicitObjectType Implicit1Type = Implicit1 ? GetInnerType(Implicit1->GetType()) : ImplicitObjectType::Unknown;
 
@@ -1645,11 +1651,11 @@ namespace Chaos
 			}
 			else if (Implicit0->IsConvex() && Implicit1Type == THeightField<FReal>::StaticType())
 			{
-				ConstructCapsuleHeightFieldConstraints(Particle0, Particle1, Implicit0, Implicit1, Transform0, Transform1, CullDistance, NewConstraints);
+				ConstructConvexHeightFieldConstraints(Particle0, Particle1, Implicit0, Implicit1, Transform0, Transform1, CullDistance, NewConstraints);
 			}
 			else if (Implicit0Type == THeightField<FReal>::StaticType() && Implicit1->IsConvex())
 			{
-				ConstructCapsuleHeightFieldConstraints(Particle1, Particle0, Implicit1, Implicit0, Transform1, Transform0, CullDistance, NewConstraints);
+				ConstructConvexHeightFieldConstraints(Particle1, Particle0, Implicit1, Implicit0, Transform1, Transform0, CullDistance, NewConstraints);
 			}
 			else if (Implicit0->IsConvex() && Implicit1Type == FTriangleMeshImplicitObject::StaticType())
 			{
