Fix Scale not being respected for collision against triangle meshes.

diff --git a/Engine/Source/Runtime/Experimental/Chaos/Private/Chaos/CollisionResolution.cpp b/Engine/Source/Runtime/Experimental/Chaos/Private/Chaos/CollisionResolution.cpp
--- a/Engine/Source/Runtime/Experimental/Chaos/Private/Chaos/CollisionResolution.cpp
+++ b/Engine/Source/Runtime/Experimental/Chaos/Private/Chaos/CollisionResolution.cpp
@@ -1454,6 +1454,7 @@ namespace Chaos
 				ensure(false);
 				break;
 			}
+
 		}
 
 		// Run collision detection for the specified constraint to update the nearest contact point.
@@ -1910,7 +1911,20 @@ namespace Chaos
 				ConstructConstraints(Particle0, Particle1, Implicit0, TransformedImplicit1->GetTransformedObject(), Transform0, TransformedTransform1, CullDistance, NewConstraints);
 				return;
 			}
-
+			else if (Implicit0OuterType == TImplicitObjectScaled<FTriangleMeshImplicitObject>::StaticType())
+			{
+				const TImplicitObjectScaled<FTriangleMeshImplicitObject>* ScaledImplicit0 = Implicit0->template GetObject<const TImplicitObjectScaled<FTriangleMeshImplicitObject>>();
+				TRigidTransform<T, d> ScaledTransform0 = TRigidTransform<T, d>(TVec3<T>(0.0f), FQuat::Identity, ScaledImplicit0->GetScale()) * Transform0;
+				ConstructConstraints(Particle0, Particle1, ScaledImplicit0->Object().Get(), Implicit1, ScaledTransform0, Transform1, CullDistance, NewConstraints);
+				return;
+			}
+			else if (Implicit1OuterType == TImplicitObjectScaled<FTriangleMeshImplicitObject>::StaticType())
+			{
+				const TImplicitObjectScaled<FTriangleMeshImplicitObject>* ScaledImplicit1 = Implicit1->template GetObject<const TImplicitObjectScaled<FTriangleMeshImplicitObject>>();
+				TRigidTransform<T, d> ScaledTransform1 = TRigidTransform<T, d>(TVec3<T>(0.0f), FQuat::Identity, ScaledImplicit1->GetScale()) * Transform1;
+				ConstructConstraints(Particle0, Particle1, Implicit0, ScaledImplicit1->Object().Get(), Transform0, ScaledTransform1, CullDistance, NewConstraints);
+				return;
+			}
 			// Handle Instanced shapes
 			// NOTE: Tri Meshes are handled differently. We should probably do something about this...
 			if (((uint32)Implicit0OuterType & ImplicitObjectType::IsInstanced) || ((uint32)Implicit1OuterType & ImplicitObjectType::IsInstanced))
