Chaos : Filtering

- Filter out shapes with invalid SimData

diff --git a/Engine/Source/Runtime/Experimental/Chaos/Private/Chaos/CollisionResolution.cpp b/Engine/Source/Runtime/Experimental/Chaos/Private/Chaos/CollisionResolution.cpp
--- a/Engine/Source/Runtime/Experimental/Chaos/Private/Chaos/CollisionResolution.cpp
+++ b/Engine/Source/Runtime/Experimental/Chaos/Private/Chaos/CollisionResolution.cpp
@@ -1763,8 +1763,8 @@ namespace Chaos
 			//
 			// Disabled shapes do not collide
 			//
-			if (Shape0 && Shape0->bDisable) return false;
-			if (Shape1 && Shape1->bDisable) return false;
+			if (Shape0 && (Shape0->bDisable || !IsValid(Shape0->SimData) ) ) return false;
+			if (Shape1 && (Shape1->bDisable || !IsValid(Shape1->SimData) ) ) return false;
 
 			//
 			// Triangle Mesh geometry is only used if the shape specifies UseComplexAsSimple
