Chaos
- workaround for capsule-box manifold issues
- pass a collision context from game into low-level collision functions
- allows collision detection to know about world-space (as well as sim space)

diff --git a/Engine/Source/Programs/HeadlessChaos/Public/HeadlessChaosCollisionConstraints.h b/Engine/Source/Programs/HeadlessChaos/Public/HeadlessChaosCollisionConstraints.h
--- a/Engine/Source/Programs/HeadlessChaos/Public/HeadlessChaosCollisionConstraints.h
+++ b/Engine/Source/Programs/HeadlessChaos/Public/HeadlessChaosCollisionConstraints.h
@@ -8,6 +8,7 @@
 #include "Chaos/PBDCollisionConstraints.h"
 #include "Chaos/PBDCollisionConstraintsContact.h"
 #include "Chaos/CollisionResolution.h"
+#include "Chaos/Collision/CollisionContext.h"
 #include "Chaos/Collision/CollisionDetector.h"
 #include "Chaos/CollisionResolutionUtil.h"
 #include "Chaos/Plane.h"
@@ -75,11 +76,11 @@ public:
 
 	void UpdateManifold(FContactConstraintBase& Constraint, T CullDistance = T(0))
 	{
-		const FRigidTransform3 Transform0 = Collisions::GetTransform(Constraint.Particle[0]);
-		const FRigidTransform3 Transform1 = Collisions::GetTransform(Constraint.Particle[1]);
-
-		if( Constraint.GetType()== FCollisionConstraintBase::FType::MultiPoint)
-			Collisions::UpdateManifold(*Constraint.As<FMultiPointContactConstraint>(), CullDistance);
+		if (Constraint.GetType() == FContactConstraintBase::FType::MultiPoint)
+		{
+			FCollisionContext Context;
+			Collisions::UpdateManifold(*Constraint.As<FRigidBodyMultiPointContactConstraint>(), CullDistance, Context);
+		}
 	}
 
 
diff --git a/Engine/Source/Runtime/AnimGraphRuntime/Private/BoneControllers/AnimNode_RigidBody.cpp b/Engine/Source/Runtime/AnimGraphRuntime/Private/BoneControllers/AnimNode_RigidBody.cpp
--- a/Engine/Source/Runtime/AnimGraphRuntime/Private/BoneControllers/AnimNode_RigidBody.cpp
+++ b/Engine/Source/Runtime/AnimGraphRuntime/Private/BoneControllers/AnimNode_RigidBody.cpp
@@ -286,10 +286,7 @@ void FAnimNode_RigidBody::EvaluateSkeletalControl_AnyThread(FComponentSpacePoseC
 		}
 		const FTransform BaseBoneTM = Output.Pose.GetComponentSpaceTransform(BaseBoneRef.GetCompactPoseIndex(BoneContainer));
 
-#if !UE_BUILD_SHIPPING
-		// Only used for debug draw...
 		PhysicsSimulation->SetSimulationSpaceTransform(SpaceToWorldTransform(SimulationSpace, CompWorldSpaceTM, BaseBoneTM));
-#endif
 
 		// Initialize potential new bodies because of LOD change.
 		if (ResetSimulatedTeleportType == ETeleportType::None && bCheckForBodyTransformInit)
diff --git a/Engine/Source/Runtime/AnimGraphRuntime/Private/BoneControllers/AnimNode_RigidBody_Chaos.cpp b/Engine/Source/Runtime/AnimGraphRuntime/Private/BoneControllers/AnimNode_RigidBody_Chaos.cpp
--- a/Engine/Source/Runtime/AnimGraphRuntime/Private/BoneControllers/AnimNode_RigidBody_Chaos.cpp
+++ b/Engine/Source/Runtime/AnimGraphRuntime/Private/BoneControllers/AnimNode_RigidBody_Chaos.cpp
@@ -271,10 +271,7 @@ void FAnimNode_RigidBody_Chaos::EvaluateSkeletalControl_AnyThread(FComponentSpac
 		}
 		const FTransform BaseBoneTM = Output.Pose.GetComponentSpaceTransform(BaseBoneRef.GetCompactPoseIndex(BoneContainer));
 
-#if !UE_BUILD_SHIPPING
-		// Only used for debug draw...
 		PhysicsSimulation->SetSimulationSpaceTransform(SpaceToWorldTransformChaos(SimulationSpace, CompWorldSpaceTM, BaseBoneTM));
-#endif
 
 		// Initialize potential new bodies because of LOD change.
 		if (ResetSimulatedTeleportType == ETeleportType::None && bCheckForBodyTransformInit)
diff --git a/Engine/Source/Runtime/Engine/Private/PhysicsEngine/ImmediatePhysics/ImmediatePhysicsChaos/ImmediatePhysicsSimulation_Chaos.cpp b/Engine/Source/Runtime/Engine/Private/PhysicsEngine/ImmediatePhysics/ImmediatePhysicsChaos/ImmediatePhysicsSimulation_Chaos.cpp
--- a/Engine/Source/Runtime/Engine/Private/PhysicsEngine/ImmediatePhysics/ImmediatePhysicsChaos/ImmediatePhysicsSimulation_Chaos.cpp
+++ b/Engine/Source/Runtime/Engine/Private/PhysicsEngine/ImmediatePhysics/ImmediatePhysicsChaos/ImmediatePhysicsSimulation_Chaos.cpp
@@ -136,29 +136,6 @@ FAutoConsoleVariableRef CVarChaosImmPhysDebugDrawJointFeatures(TEXT("p.Chaos.Imm
 
 namespace ImmediatePhysics_Chaos
 {
-	template<typename T, int d>
-	bool ShouldIgnoreCollisionConstraint(
-		const Chaos::TGeometryParticleHandle<T, d>* ParticleA,
-		const Chaos::TGeometryParticleHandle<T, d>* ParticleB,
-		const TMap<const Chaos::TGeometryParticleHandle<T, d>*, TSet<const Chaos::TGeometryParticleHandle<T, d>*>>& IgnoreSetMap)
-	{
-		using namespace Chaos;
-
-		if (!ChaosImmediate_Collision_Enabled)
-		{
-			return true;
-		}
-		if (const TSet<const TGeometryParticleHandle<T, d>*>* IgnoreSet = IgnoreSetMap.Find(ParticleA))
-		{
-			return IgnoreSet->Contains(ParticleB);
-		}
-		if (const TSet<const TGeometryParticleHandle<T, d>*>* IgnoreSet = IgnoreSetMap.Find(ParticleB))
-		{
-			return IgnoreSet->Contains(ParticleA);
-		}
-		return false;
-	}
-
 	//
 	//
 	//
@@ -189,11 +166,13 @@ namespace ImmediatePhysics_Chaos
 		Particles.GetParticleHandles().AddArray(&ParticleMaterials);
 		Particles.GetParticleHandles().AddArray(&PerParticleMaterials);
 
-		Collisions.DisableHandles();
-
 		Evolution.AddConstraintRule(&CollisionsRule);
 		Evolution.AddConstraintRule(&JointsRule);
 
+		// RBAN collision customization
+		Collisions.DisableHandles();
+		CollisionDetector.GetContext().bFilteringEnabled = false;
+
 #if CHAOS_DEBUG_DRAW
 		Evolution.SetPostIntegrateCallback(
 			[this]()
@@ -525,6 +504,13 @@ namespace ImmediatePhysics_Chaos
 		}
 	}
 
+	void FSimulation::SetSimulationSpaceTransform(const FTransform& Transform)
+	{ 
+		SimulationSpaceTransform = Transform;
+		CollisionDetector.GetContext().SpaceTransform = Transform;
+	}
+
+
 	FReal FSimulation::UpdateStepTime(const FReal DeltaTime, const FReal MaxStepTime)
 	{
 		RollingAverageStepTime = RollingAverageStepTime + (DeltaTime - RollingAverageStepTime) / NumRollingAverageStepTimes;
diff --git a/Engine/Source/Runtime/Engine/Public/Physics/ImmediatePhysics/ImmediatePhysicsChaos/ImmediatePhysicsSimulation_Chaos.h b/Engine/Source/Runtime/Engine/Public/Physics/ImmediatePhysics/ImmediatePhysicsChaos/ImmediatePhysicsSimulation_Chaos.h
--- a/Engine/Source/Runtime/Engine/Public/Physics/ImmediatePhysics/ImmediatePhysicsChaos/ImmediatePhysicsSimulation_Chaos.h
+++ b/Engine/Source/Runtime/Engine/Public/Physics/ImmediatePhysics/ImmediatePhysicsChaos/ImmediatePhysicsSimulation_Chaos.h
@@ -67,7 +67,7 @@ namespace ImmediatePhysics_Chaos
 		void Simulate(float DeltaTime, float MaxStepTime, int32 MaxSubSteps, const FVector& InGravity);
 		void Simulate_AssumesLocked(float DeltaTime, float MaxStepTime, int32 MaxSubSteps, const FVector& InGravity) { Simulate(DeltaTime, MaxStepTime, MaxSubSteps, InGravity); }
 
-		void SetSimulationSpaceTransform(const FTransform& Transform) { SimulationSpaceTransform = Transform; }
+		void SetSimulationSpaceTransform(const FTransform& Transform);
 
 		/** Set new iteration counts. A negative value with leave that iteration count unchanged */
 		void SetSolverIterations(const int32 SolverIts, const int32 JointIts, const int32 CollisionIts, const int32 SolverPushOutIts, const int32 JointPushOutIts, const int32 CollisionPushOutIts);
diff --git a/Engine/Source/Runtime/Experimental/Chaos/Private/Chaos/ChaosDebugDraw.cpp b/Engine/Source/Runtime/Experimental/Chaos/Private/Chaos/ChaosDebugDraw.cpp
--- a/Engine/Source/Runtime/Experimental/Chaos/Private/Chaos/ChaosDebugDraw.cpp
+++ b/Engine/Source/Runtime/Experimental/Chaos/Private/Chaos/ChaosDebugDraw.cpp
@@ -24,7 +24,7 @@ namespace Chaos
 		float BodyAxisLen = 4.0f;
 		float ContactLen = 4.0f;
 		float ContactWidth = 2.0f;
-		float ContactPhiWidth = 1.5f;
+		float ContactPhiWidth = 0.3f;
 		float ContactOwnerWidth = 0.0f;
 		float ConstraintAxisLen = 5.0f;
 		float JointComSize = 2.0f;
@@ -162,23 +162,31 @@ namespace Chaos
 
 		void DrawCollisionImpl(const FRigidTransform3& SpaceTransform, const FCollisionConstraintBase& Contact, float ColorScale)
 		{
-			const FReal ActiveColorScale = (Contact.GetPhi() > 0)? ColorScale * 0.1f : ColorScale;
-
 			FVec3 Location = SpaceTransform.TransformPosition(Contact.GetLocation());
 			FVec3 Normal = SpaceTransform.TransformVector(Contact.GetNormal());
 
 			if (ContactWidth > 0)
 			{
-				bool bIsManifold = (Contact.GetType() == FCollisionConstraintBase::FType::MultiPoint);
-				FColor C0 = bIsManifold? (ColorScale * FColor(0, 0, 200)).ToFColor(false) : (ColorScale * FColor(200, 0, 0)).ToFColor(false);
-				FMatrix Axes = FRotationMatrix::MakeFromX(Normal);
-				FDebugDrawQueue::GetInstance().DrawDebugCircle(Location, DrawScale * ContactWidth, 12, C0, false, KINDA_SMALL_NUMBER, DrawPriority, LineThickness, Axes.GetUnitAxis(EAxis::Y), Axes.GetUnitAxis(EAxis::Z), false);
-
-				if (bIsManifold)
+				bool bIsManifold = (Contact.GetType() == FCollisionConstraintBase::FType::MultiPoint) && (Contact.As<FRigidBodyMultiPointContactConstraint>()->GetManifoldPlaneOwnerIndex() >= 0);
+				if (!bIsManifold)
+				{
+					FColor C0 = (ColorScale * FColor(200, 0, 0)).ToFColor(false);
+					FMatrix Axes = FRotationMatrix::MakeFromX(Normal);
+					FDebugDrawQueue::GetInstance().DrawDebugCircle(Location, DrawScale * ContactWidth, 12, C0, false, KINDA_SMALL_NUMBER, DrawPriority, LineThickness, Axes.GetUnitAxis(EAxis::Y), Axes.GetUnitAxis(EAxis::Z), false);
+				}
+				else
 				{
 					const FRigidBodyMultiPointContactConstraint& MultiPointConstraint = *Contact.As<FRigidBodyMultiPointContactConstraint>();
-					TConstGenericParticleHandle<FReal, 3> PointsParticle = MultiPointConstraint.PointsParticleHandle();
-					FRigidTransform3 PointsTransform = FParticleUtilities::GetActorWorldTransform(PointsParticle) * SpaceTransform;
+					int32 PlaneOwnerIndex = MultiPointConstraint.GetManifoldPlaneOwnerIndex();
+					int32 PointsOwnerIndex = 1 - PlaneOwnerIndex;
+					FColor C0 = (PlaneOwnerIndex == 0) ? (ColorScale * FColor(0, 200, 0)).ToFColor(false) : (ColorScale * FColor(0, 0, 200)).ToFColor(false);
+					
+					FMatrix Axes = FRotationMatrix::MakeFromX(Normal);
+					FDebugDrawQueue::GetInstance().DrawDebugCircle(Location, DrawScale * ContactWidth, 12, C0, false, KINDA_SMALL_NUMBER, DrawPriority, LineThickness, Axes.GetUnitAxis(EAxis::Y), Axes.GetUnitAxis(EAxis::Z), false);
+					
+					TConstGenericParticleHandle<FReal, 3> PointsParticle = MultiPointConstraint.Particle[PointsOwnerIndex];
+					FRigidTransform3 PointsImplicitTransform = MultiPointConstraint.ImplicitTransform[PointsOwnerIndex];
+					FRigidTransform3 PointsTransform = PointsImplicitTransform * FParticleUtilities::GetActorWorldTransform(PointsParticle) * SpaceTransform;
 					for (int32 SampleIndex = 1; SampleIndex < MultiPointConstraint.NumManifoldPoints(); ++SampleIndex)
 					{
 						FVec3 S0 = PointsTransform.TransformPosition(MultiPointConstraint.GetManifoldPoint(SampleIndex - 1));
@@ -189,12 +197,12 @@ namespace Chaos
 			}
 			if (ContactLen > 0)
 			{
-				FColor C1 = (ActiveColorScale * FColor(255, 0, 0)).ToFColor(false);
+				FColor C1 = (ColorScale * FColor(255, 0, 0)).ToFColor(false);
 				FDebugDrawQueue::GetInstance().DrawDebugLine(Location, Location + DrawScale * ContactLen * Normal, C1, false, KINDA_SMALL_NUMBER, DrawPriority, LineThickness);
 			}
 			if (ContactPhiWidth > 0 && Contact.GetPhi() < FLT_MAX)
 			{
-				FColor C2 = (ActiveColorScale * FColor(128, 128, 0)).ToFColor(false);
+				FColor C2 = (ColorScale * FColor(128, 128, 0)).ToFColor(false);
 				FMatrix Axes = FRotationMatrix::MakeFromX(Normal);
 				FDebugDrawQueue::GetInstance().DrawDebugCircle(Location - Contact.GetPhi() * Normal, DrawScale * ContactPhiWidth, 12, C2, false, KINDA_SMALL_NUMBER, DrawPriority, LineThickness, Axes.GetUnitAxis(EAxis::Y), Axes.GetUnitAxis(EAxis::Z), false);
 			}
diff --git a/Engine/Source/Runtime/Experimental/Chaos/Private/Chaos/CollisionResolution.cpp b/Engine/Source/Runtime/Experimental/Chaos/Private/Chaos/CollisionResolution.cpp
--- a/Engine/Source/Runtime/Experimental/Chaos/Private/Chaos/CollisionResolution.cpp
+++ b/Engine/Source/Runtime/Experimental/Chaos/Private/Chaos/CollisionResolution.cpp
@@ -7,6 +7,7 @@
 #include "Chaos/Capsule.h"
 #include "Chaos/CollisionResolutionTypes.h"
 #include "Chaos/CollisionResolutionUtil.h"
+#include "Chaos/Collision/CollisionContext.h"
 #include "Chaos/Convex.h"
 #include "Chaos/Defines.h"
 #include "Chaos/HeightField.h"
@@ -41,6 +42,10 @@ float Chaos_Collision_ManifoldFaceEpsilon = FMath::Sin(FMath::DegreesToRadians(C
 FConsoleVariableDelegate Chaos_Collision_ManifoldFaceDelegate = FConsoleVariableDelegate::CreateLambda([](IConsoleVariable* CVar) { Chaos_Collision_ManifoldFaceEpsilon = FMath::Sin(FMath::DegreesToRadians(Chaos_Collision_ManifoldFaceAngle)); });
 FAutoConsoleVariableRef CVarChaosCollisionManifoldFaceAngle(TEXT("p.Chaos.Collision.ManifoldFaceAngle"), Chaos_Collision_ManifoldFaceAngle, TEXT("Angle above which a face is rejected and we switch to point collision"), Chaos_Collision_ManifoldFaceDelegate);
 
+float Chaos_Collision_CapsuleBoxManifoldAngle = 0.0f;
+float Chaos_Collision_CapsuleBoxManifoldTolerance = FMath::Sin(FMath::DegreesToRadians(Chaos_Collision_CapsuleBoxManifoldAngle));
+FConsoleVariableDelegate Chaos_Collision_CapsuleBoxManifoldDelegate = FConsoleVariableDelegate::CreateLambda([](IConsoleVariable* CVar) { Chaos_Collision_CapsuleBoxManifoldTolerance = FMath::Sin(FMath::DegreesToRadians(Chaos_Collision_CapsuleBoxManifoldAngle)); });
+FAutoConsoleVariableRef CVarChaosCollisionBoxCapsuleManifoldAngle(TEXT("p.Chaos.Collision.CapsuleBoxManifoldAngle"), Chaos_Collision_CapsuleBoxManifoldAngle, TEXT("If a capsule is more than this angle from vertical, do not use a manifold"), Chaos_Collision_CapsuleBoxManifoldDelegate);
 
 namespace Chaos
 {
@@ -886,7 +891,7 @@ namespace Chaos
 		 * 
 		 */
 		template<class T, int d>
-		void UpdateCapsuleBoxManifold(const TCapsule<T>& Capsule, const TRigidTransform<T, d>& CapsuleTM, const TAABB<T, d>& Box, const TRigidTransform<T, d>& BoxTM, const T CullDistance, FRigidBodyMultiPointContactConstraint& Constraint)
+		void UpdateCapsuleBoxManifold(const TCapsule<T>& Capsule, const TRigidTransform<T, d>& CapsuleTM, const TAABB<T, d>& Box, const TRigidTransform<T, d>& BoxTM, const T CullDistance, const FCollisionContext& Context, FRigidBodyMultiPointContactConstraint& Constraint)
 		{
 			Constraint.ResetManifoldPoints();
 
@@ -907,7 +912,7 @@ namespace Chaos
 				}
 			}
 
-			// Cache the closest point so we don't need to re-iterate over the manifold on the first iteration
+			// Cache the closest point so we don't need to re-iterate over the manifold on the first iteration.
 			Constraint.Manifold.Location = BoxTM.TransformPosition(BoxClosestBoxSpace);
 			Constraint.Manifold.Normal = BoxTM.TransformVector(NormalBoxSpace);
 			Constraint.Manifold.Phi = -Penetration;
@@ -1072,7 +1077,7 @@ namespace Chaos
 		}
 
 		template<typename T, int d>
-		void ConstructCapsuleBoxConstraints(TGeometryParticleHandle<T, d>* Particle0, TGeometryParticleHandle<T, d>* Particle1, const FImplicitObject* Implicit0, const FImplicitObject* Implicit1, const TRigidTransform<T, d>& Transform0, const TRigidTransform<T, d>& Transform1, const T CullDistance, FCollisionConstraintsArray& NewConstraints)
+		void ConstructCapsuleBoxConstraints(TGeometryParticleHandle<T, d>* Particle0, TGeometryParticleHandle<T, d>* Particle1, const FImplicitObject* Implicit0, const FImplicitObject* Implicit1, const TRigidTransform<T, d>& Transform0, const TRigidTransform<T, d>& Transform1, const T CullDistance, const FCollisionContext& Context, FCollisionConstraintsArray& NewConstraints)
 		{
 			const TCapsule<T> * Object0 = Implicit0->template GetObject<const TCapsule<T> >();
 			const TBox<T, d> * Object1 = Implicit1->template GetObject<const TBox<T, d> >();
@@ -1090,10 +1095,22 @@ namespace Chaos
 					TRigidTransform<T, d> ParticleImplicit0TM = Transform0.GetRelativeTransform(Collisions::GetTransform(Particle0));
 					TRigidTransform<T, d> ParticleImplicit1TM = Transform1.GetRelativeTransform(Collisions::GetTransform(Particle1));
 
-					if (Chaos_Collision_UseManifolds)
-					{ 
+					bool bAllowManifold = Chaos_Collision_UseManifolds;
+
+					if (Chaos_Collision_CapsuleBoxManifoldTolerance > KINDA_SMALL_NUMBER)
+					{
+						// @todo(ccaulfield): weak sauce - fix capsule-box manifolds.
+						// HACK: Disable manifolds for "horizontal" capsules. Manifolds don't work well when joints are pulling boxes down
+						// (under gravity) when the upper boxes are draped over a horizontal capsule. The box rotations about the manifold
+						// points(line) is too great and we end up with jitter.
+						const FVector CapsuleAxis = Context.SpaceTransform.TransformVector(Transform0.TransformVector(Object0->GetAxis()));
+						bAllowManifold = (FMath::Abs(CapsuleAxis.Z) > Chaos_Collision_CapsuleBoxManifoldTolerance);
+					}
+
+					if (bAllowManifold)
+					{
 						FRigidBodyMultiPointContactConstraint Constraint = FRigidBodyMultiPointContactConstraint(Particle0, Implicit0, ParticleImplicit0TM, Particle1, Implicit1, ParticleImplicit1TM, EContactShapesType::CapsuleBox);
-						UpdateCapsuleBoxManifold(*Object0, Transform0, Object1->BoundingBox(), Transform1, CullDistance, Constraint);
+						UpdateCapsuleBoxManifold(*Object0, Transform0, Object1->BoundingBox(), Transform1, CullDistance, Context, Constraint);
 						NewConstraints.TryAdd(CullDistance, Constraint);
 					}
 					else
@@ -1405,7 +1422,7 @@ namespace Chaos
 
 		DECLARE_CYCLE_STAT(TEXT("TPBDCollisionConstraints::ConstructUnionUnionConstraints"), STAT_ConstructUnionUnionConstraints, STATGROUP_ChaosWide);
 		template<typename T, int d>
-		void ConstructUnionUnionConstraints(TGeometryParticleHandle<T, d>* Particle0, TGeometryParticleHandle<T, d>* Particle1, const FImplicitObject* Implicit0, const FImplicitObject* Implicit1, const TRigidTransform<T, d>& Transform0, const TRigidTransform<T, d>& Transform1, const T CullDistance, FCollisionConstraintsArray& NewConstraints)
+		void ConstructUnionUnionConstraints(TGeometryParticleHandle<T, d>* Particle0, TGeometryParticleHandle<T, d>* Particle1, const FImplicitObject* Implicit0, const FImplicitObject* Implicit1, const TRigidTransform<T, d>& Transform0, const TRigidTransform<T, d>& Transform1, const T CullDistance, const FCollisionContext& Context, FCollisionConstraintsArray& NewConstraints)
 		{
 			SCOPE_CYCLE_COUNTER(STAT_ConstructUnionUnionConstraints);
 
@@ -1423,7 +1440,7 @@ namespace Chaos
 				for (const Pair<const FImplicitObject*, TRigidTransform<T, d>>& ParticlePair : ParticleShapes)
 				{
 					const FImplicitObject* ParticleInnerObj = ParticlePair.First;
-					ConstructConstraints<T, d>(Particle0, Particle1, ParticleInnerObj, LevelsetInnerObj, Transform0, Transform1, CullDistance, NewConstraints);
+					ConstructConstraints<T, d>(Particle0, Particle1, ParticleInnerObj, LevelsetInnerObj, Transform0, Transform1, CullDistance, Context, NewConstraints);
 				}
 			}
 		}
@@ -1434,7 +1451,7 @@ namespace Chaos
 
 
 		template<typename T, int d>
-		inline void UpdateManifold(TRigidBodyMultiPointContactConstraint<T, d>& Constraint, const TRigidTransform<T, d>& ParticleTransform0, const TRigidTransform<T, d>& ParticleTransform1, const T CullDistance)
+		inline void UpdateManifold(TRigidBodyMultiPointContactConstraint<T, d>& Constraint, const TRigidTransform<T, d>& ParticleTransform0, const TRigidTransform<T, d>& ParticleTransform1, const T CullDistance, const FCollisionContext& Context)
 		{
 			const FImplicitObject& Implicit0 = *Constraint.Manifold.Implicit[0];
 			const FImplicitObject& Implicit1 = *Constraint.Manifold.Implicit[1];
@@ -1444,7 +1461,7 @@ namespace Chaos
 			switch (Constraint.Manifold.ShapesType)
 			{
 			case EContactShapesType::CapsuleBox:
-				UpdateCapsuleBoxManifold<T, d>(*Implicit0.template GetObject<const TCapsule<T>>(), Transform0, Implicit1.template GetObject<const TBox<T, d>>()->GetAABB(), Transform1, CullDistance, Constraint);
+				UpdateCapsuleBoxManifold<T, d>(*Implicit0.template GetObject<const TCapsule<T>>(), Transform0, Implicit1.template GetObject<const TBox<T, d>>()->GetAABB(), Transform1, CullDistance, Context, Constraint);
 				break;
 			case EContactShapesType::ConvexConvex:
 				UpdateConvexConvexManifold(Constraint, Transform0, Transform1, CullDistance);
@@ -1544,9 +1561,9 @@ namespace Chaos
 		}
 
 		template<typename T, int d>
-		void ConstructConstraintsImpl(TGeometryParticleHandle<T, d>* Particle0, TGeometryParticleHandle<T, d>* Particle1, const FImplicitObject* Implicit0, const FImplicitObject* Implicit1, const TRigidTransform<T, d>& Transform0, const TRigidTransform<T, d>& Transform1, const T CullDistance, FCollisionConstraintsArray& NewConstraints)
+		void ConstructConstraintsImpl(TGeometryParticleHandle<T, d>* Particle0, TGeometryParticleHandle<T, d>* Particle1, const FImplicitObject* Implicit0, const FImplicitObject* Implicit1, const TRigidTransform<T, d>& Transform0, const TRigidTransform<T, d>& Transform1, const T CullDistance, const FCollisionContext& Context, FCollisionConstraintsArray& NewConstraints)
 		{
-			// @todo(chaos): We use GetInnerType here because TriMeshes are lef with their "Instanced" wrapper, unlike all other instanced implicits. Should we strip the instance on Tri Mesh too?
+			// @todo(chaos): We use GetInnerType here because TriMeshes are left with their "Instanced" wrapper, unlike all other instanced implicits. Should we strip the instance on Tri Mesh too?
 			EImplicitObjectType Implicit0Type = Implicit0 ? GetInnerType(Implicit0->GetType()) : ImplicitObjectType::Unknown;
 			EImplicitObjectType Implicit1Type = Implicit1 ? GetInnerType(Implicit1->GetType()) : ImplicitObjectType::Unknown;
 
@@ -1612,11 +1629,11 @@ namespace Chaos
 			}
 			else if (Implicit0Type == TCapsule<FReal>::StaticType() && Implicit1Type == TBox<FReal, 3>::StaticType())
 			{
-				ConstructCapsuleBoxConstraints(Particle0, Particle1, Implicit0, Implicit1, Transform0, Transform1, CullDistance, NewConstraints);
+				ConstructCapsuleBoxConstraints(Particle0, Particle1, Implicit0, Implicit1, Transform0, Transform1, CullDistance, Context, NewConstraints);
 			}
 			else if (Implicit0Type == TBox<FReal, 3>::StaticType() && Implicit1Type == TCapsule<FReal>::StaticType())
 			{
-				ConstructCapsuleBoxConstraints(Particle1, Particle0, Implicit1, Implicit0, Transform1, Transform0, CullDistance, NewConstraints);
+				ConstructCapsuleBoxConstraints(Particle1, Particle0, Implicit1, Implicit0, Transform1, Transform0, CullDistance, Context, NewConstraints);
 			}
 			else if (Implicit0Type == TCapsule<FReal>::StaticType() && Implicit1Type == THeightField<FReal>::StaticType())
 			{
@@ -1694,14 +1711,12 @@ namespace Chaos
 		{
 			const FRigidTransform3 Transform0 = Constraint.ImplicitTransform[0] * ParticleTransform0;
 			const FRigidTransform3 Transform1 = Constraint.ImplicitTransform[1] * ParticleTransform1;
+			const int32 NumPoints = Constraint.NumManifoldPoints();
 
-			// TEMP: Test code - run collision detection instead. Allows us to debug draw the manifold without using it
-			if (Chaos_Collision_UseManifolds_Test)
+			// Fall back to full collision detection if we have no manifold (or for testing)
+			if ((NumPoints == 0) || Chaos_Collision_UseManifolds_Test)
 			{
-				// Copy the manifold constraint into point constraint, update it, and copy the results back
-				TRigidBodyPointContactConstraint<T, d> PointConstraint(Constraint.Particle[0], Constraint.Manifold.Implicit[0], Constraint.ImplicitTransform[0], Constraint.Particle[1], Constraint.Manifold.Implicit[1], Constraint.ImplicitTransform[1], Constraint.Manifold.ShapesType);
-				UpdateConstraintFromGeometryImpl<ECollisionUpdateType::Deepest>(PointConstraint, Transform0, Transform1, CullDistance);
-				Constraint.Manifold = PointConstraint.Manifold;
+				UpdateConstraintFromGeometryImpl<ECollisionUpdateType::Deepest>(Constraint, Transform0, Transform1, CullDistance);
 				return;
 			}
 
@@ -1714,7 +1729,6 @@ namespace Chaos
 			FVec3 PlanePos = PlaneTransform.TransformPosition(Constraint.GetManifoldPlanePosition());
 
 			// Select the best manifold point
-			const int32 NumPoints = Constraint.NumManifoldPoints();
 			for (int32 PointIndex = 0; PointIndex < NumPoints; ++PointIndex)
 			{
 				// World-space manifold point and distance to manifold plane
@@ -1849,7 +1863,7 @@ namespace Chaos
 
 
 		template<typename T, int d>
-		void ConstructManifoldConstraint(EContactShapesType ShapesType, TGeometryParticleHandle<T, d>* Particle0, TGeometryParticleHandle<T, d>* Particle1, const FImplicitObject* Implicit0, const FImplicitObject* Implicit1, const TRigidTransform<T, d>& Transform0, const TRigidTransform<T, d>& Transform1, const T CullDistance, FCollisionConstraintsArray& NewConstraints)
+		void ConstructManifoldConstraint(EContactShapesType ShapesType, TGeometryParticleHandle<T, d>* Particle0, TGeometryParticleHandle<T, d>* Particle1, const FImplicitObject* Implicit0, const FImplicitObject* Implicit1, const TRigidTransform<T, d>& Transform0, const TRigidTransform<T, d>& Transform1, const T CullDistance, const FCollisionContext& Context, FCollisionConstraintsArray& NewConstraints)
 		{
 			if (ensure(Implicit0 && Implicit1))
 			{
@@ -1857,7 +1871,7 @@ namespace Chaos
 				TRigidTransform<T, d> ParticleImplicit1TM = Transform1.GetRelativeTransform(Collisions::GetTransform(Particle1));
 				if (FRigidBodyMultiPointContactConstraint* Constraint = NewConstraints.AddMultiPointConstraint(Particle0, Implicit0, ParticleImplicit0TM, Particle1, Implicit1, ParticleImplicit1TM, ShapesType))
 				{
-					UpdateManifold(*Constraint, Transform0, Transform1, CullDistance);
+					UpdateManifold(*Constraint, Transform0, Transform1, CullDistance, Context);
 					if (Constraint->GetPhi() >= CullDistance)
 					{
 						NewConstraints.PopMultiPointConstraint();
@@ -1868,7 +1882,7 @@ namespace Chaos
 
 
 		template<typename T, int d>
-		void ConstructConstraints(TGeometryParticleHandle<T, d>* Particle0, TGeometryParticleHandle<T, d>* Particle1, const FImplicitObject* Implicit0, const FImplicitObject* Implicit1, const TRigidTransform<T, d>& Transform0, const TRigidTransform<T, d>& Transform1, const T CullDistance, FCollisionConstraintsArray& NewConstraints)
+		void ConstructConstraints(TGeometryParticleHandle<T, d>* Particle0, TGeometryParticleHandle<T, d>* Particle1, const FImplicitObject* Implicit0, const FImplicitObject* Implicit1, const TRigidTransform<T, d>& Transform0, const TRigidTransform<T, d>& Transform1, const T CullDistance, const FCollisionContext& Context, FCollisionConstraintsArray& NewConstraints)
 		{
 			EImplicitObjectType Implicit0Type = Implicit0 ? GetInnerType(Implicit0->GetType()) : ImplicitObjectType::Unknown;
 			EImplicitObjectType Implicit1Type = Implicit1 ? GetInnerType(Implicit1->GetType()) : ImplicitObjectType::Unknown;
@@ -1890,35 +1904,35 @@ namespace Chaos
 				const TImplicitObjectTransformed<FReal, 3>* TransformedImplicit1 = Implicit1->template GetObject<const TImplicitObjectTransformed<FReal, 3>>();
 				TRigidTransform<T, d> TransformedTransform0 = TransformedImplicit0->GetTransform() * Transform0;
 				TRigidTransform<T, d> TransformedTransform1 = TransformedImplicit1->GetTransform() * Transform1;
-				ConstructConstraints(Particle0, Particle1, TransformedImplicit0->GetTransformedObject(), TransformedImplicit1->GetTransformedObject(), TransformedTransform0, TransformedTransform1, CullDistance, NewConstraints);
+				ConstructConstraints(Particle0, Particle1, TransformedImplicit0->GetTransformedObject(), TransformedImplicit1->GetTransformedObject(), TransformedTransform0, TransformedTransform1, CullDistance, Context, NewConstraints);
 				return;
 			}
 			else if (Implicit0OuterType == TImplicitObjectTransformed<T, d>::StaticType())
 			{
 				const TImplicitObjectTransformed<FReal, 3>* TransformedImplicit0 = Implicit0->template GetObject<const TImplicitObjectTransformed<FReal, 3>>();
 				TRigidTransform<T, d> TransformedTransform0 = TransformedImplicit0->GetTransform() * Transform0;
-				ConstructConstraints(Particle0, Particle1, TransformedImplicit0->GetTransformedObject(), Implicit1, TransformedTransform0, Transform1, CullDistance, NewConstraints);
+				ConstructConstraints(Particle0, Particle1, TransformedImplicit0->GetTransformedObject(), Implicit1, TransformedTransform0, Transform1, CullDistance, Context, NewConstraints);
 				return;
 			}
 			else if (Implicit1OuterType == TImplicitObjectTransformed<T, d>::StaticType())
 			{
 				const TImplicitObjectTransformed<FReal, 3>* TransformedImplicit1 = Implicit1->template GetObject<const TImplicitObjectTransformed<FReal, 3>>();
 				TRigidTransform<T, d> TransformedTransform1 = TransformedImplicit1->GetTransform() * Transform1;
-				ConstructConstraints(Particle0, Particle1, Implicit0, TransformedImplicit1->GetTransformedObject(), Transform0, TransformedTransform1, CullDistance, NewConstraints);
+				ConstructConstraints(Particle0, Particle1, Implicit0, TransformedImplicit1->GetTransformedObject(), Transform0, TransformedTransform1, CullDistance, Context, NewConstraints);
 				return;
 			}
 			else if (Implicit0OuterType == TImplicitObjectScaled<FTriangleMeshImplicitObject>::StaticType())
 			{
 				const TImplicitObjectScaled<FTriangleMeshImplicitObject>* ScaledImplicit0 = Implicit0->template GetObject<const TImplicitObjectScaled<FTriangleMeshImplicitObject>>();
 				TRigidTransform<T, d> ScaledTransform0 = TRigidTransform<T, d>(TVec3<T>(0.0f), FQuat::Identity, ScaledImplicit0->GetScale()) * Transform0;
-				ConstructConstraints(Particle0, Particle1, ScaledImplicit0->Object().Get(), Implicit1, ScaledTransform0, Transform1, CullDistance, NewConstraints);
+				ConstructConstraints(Particle0, Particle1, ScaledImplicit0->Object().Get(), Implicit1, ScaledTransform0, Transform1, CullDistance, Context, NewConstraints);
 				return;
 			}
 			else if (Implicit1OuterType == TImplicitObjectScaled<FTriangleMeshImplicitObject>::StaticType())
 			{
 				const TImplicitObjectScaled<FTriangleMeshImplicitObject>* ScaledImplicit1 = Implicit1->template GetObject<const TImplicitObjectScaled<FTriangleMeshImplicitObject>>();
 				TRigidTransform<T, d> ScaledTransform1 = TRigidTransform<T, d>(TVec3<T>(0.0f), FQuat::Identity, ScaledImplicit1->GetScale()) * Transform1;
-				ConstructConstraints(Particle0, Particle1, Implicit0, ScaledImplicit1->Object().Get(), Transform0, ScaledTransform1, CullDistance, NewConstraints);
+				ConstructConstraints(Particle0, Particle1, Implicit0, ScaledImplicit1->Object().Get(), Transform0, ScaledTransform1, CullDistance, Context, NewConstraints);
 				return;
 			}
 			// Handle Instanced shapes
@@ -1937,7 +1951,7 @@ namespace Chaos
 				}
 				if (InnerImplicit0 && InnerImplicit1)
 				{
-					ConstructConstraints(Particle0, Particle1, InnerImplicit0, InnerImplicit1, Transform0, Transform1, CullDistance, NewConstraints);
+					ConstructConstraints(Particle0, Particle1, InnerImplicit0, InnerImplicit1, Transform0, Transform1, CullDistance, Context, NewConstraints);
 					return;
 				}
 			}
@@ -1950,7 +1964,7 @@ namespace Chaos
 				{
 					const FImplicitObject* Implicit1InnerObj = LevelsetObjPair.First;
 					const TRigidTransform<T, d> Implicit1InnerObjTM = LevelsetObjPair.Second * Transform1;
-					ConstructConstraints(Particle0, Particle1, Implicit0, Implicit1InnerObj, Transform0, Implicit1InnerObjTM, CullDistance, NewConstraints);
+					ConstructConstraints(Particle0, Particle1, Implicit0, Implicit1InnerObj, Transform0, Implicit1InnerObjTM, CullDistance, Context, NewConstraints);
 				}
 				return;
 			}
@@ -1962,24 +1976,24 @@ namespace Chaos
 				{
 					const FImplicitObject* Implicit0InnerObj = LevelsetObjPair.First;
 					const TRigidTransform<T, d> Implicit0InnerObjTM = LevelsetObjPair.Second * Transform0;
-					ConstructConstraints(Particle0, Particle1, Implicit0InnerObj, Implicit1, Implicit0InnerObjTM, Transform1, CullDistance, NewConstraints);
+					ConstructConstraints(Particle0, Particle1, Implicit0InnerObj, Implicit1, Implicit0InnerObjTM, Transform1, CullDistance, Context, NewConstraints);
 				}
 				return;
 			}
 			else if (Implicit0OuterType == FImplicitObjectUnion::StaticType() && Implicit1OuterType == FImplicitObjectUnion::StaticType())
 			{
-				ConstructUnionUnionConstraints(Particle0, Particle1, Implicit0, Implicit1, Transform0, Transform1, CullDistance, NewConstraints);
+				ConstructUnionUnionConstraints(Particle0, Particle1, Implicit0, Implicit1, Transform0, Transform1, CullDistance, Context, NewConstraints);
 				return;
 			}
 
-			if (!DoCollide(Implicit0Type, Particle0->GetImplicitShape(Implicit0), Implicit1Type, Particle1->GetImplicitShape(Implicit1)))
+			if (Context.bFilteringEnabled && !DoCollide(Implicit0Type, Particle0->GetImplicitShape(Implicit0), Implicit1Type, Particle1->GetImplicitShape(Implicit1)))
 			{
 				return;
 			}
 
 			// If we get here, we have a pair of concrete shapes (i.e., no wrappers or containers)
 			// Create a constraint for the shape pair
-			ConstructConstraintsImpl(Particle0, Particle1, Implicit0, Implicit1, Transform0, Transform1, CullDistance, NewConstraints);
+			ConstructConstraintsImpl(Particle0, Particle1, Implicit0, Implicit1, Transform0, Transform1, CullDistance, Context, NewConstraints);
 		}
 
 		template void UpdateBoxBoxConstraint<float, 3>(const TAABB<float, 3>& Box1, const TRigidTransform<float, 3>& Box1Transform, const TAABB<float, 3>& Box2, const TRigidTransform<float, 3>& Box2Transform, const float CullDistance, TRigidBodyPointContactConstraint<float, 3>& Constraint);
@@ -2019,8 +2033,7 @@ namespace Chaos
 		template void ConstructCapsuleCapsuleConstraints<float, 3>(TGeometryParticleHandle<float, 3>* Particle0, TGeometryParticleHandle<float, 3>* Particle1, const FImplicitObject* Implicit0, const FImplicitObject* Implicit1, const TRigidTransform<float, 3>& Transform0, const TRigidTransform<float, 3>& Transform1, const float CullDistance, FCollisionConstraintsArray& NewConstraints);
 
 		template void UpdateCapsuleBoxConstraint<float, 3>(const TCapsule<float>& A, const TRigidTransform<float, 3>& ATransform, const TAABB<float, 3>& B, const TRigidTransform<float, 3>& BTransform, const float CullDistance, TRigidBodyPointContactConstraint<float, 3>& Constraint);
-		template void UpdateCapsuleBoxManifold<float, 3>(const TCapsule<float>& Capsule, const TRigidTransform<float, 3>& CapsuleTM, const TAABB<float, 3>& Box, const TRigidTransform<float, 3>& BoxTM, const float CullDistance, FRigidBodyMultiPointContactConstraint& Constraint);
-		template void ConstructCapsuleBoxConstraints<float, 3>(TGeometryParticleHandle<float, 3>* Particle0, TGeometryParticleHandle<float, 3>* Particle1, const FImplicitObject* Implicit0, const FImplicitObject* Implicit1, const TRigidTransform<float, 3>& Transform0, const TRigidTransform<float, 3>& Transform1, const float CullDistance, FCollisionConstraintsArray& NewConstraints);
+		template void ConstructCapsuleBoxConstraints<float, 3>(TGeometryParticleHandle<float, 3>* Particle0, TGeometryParticleHandle<float, 3>* Particle1, const FImplicitObject* Implicit0, const FImplicitObject* Implicit1, const TRigidTransform<float, 3>& Transform0, const TRigidTransform<float, 3>& Transform1, const float CullDistance, const FCollisionContext& Context, FCollisionConstraintsArray& NewConstraints);
 
 		template void UpdateCapsuleHeightFieldConstraint<float, 3>(const TCapsule<float>& A, const TRigidTransform<float, 3>& ATransform, const THeightField<float>& B, const TRigidTransform<float, 3>& BTransform, const float CullDistance, TRigidBodyPointContactConstraint<float, 3>& Constraint);
 		template void UpdateCapsuleHeightFieldManifold<float, 3>(TCollisionConstraintBase<float, 3>&  Constraint, const TRigidTransform<float, 3>& ATM, const TRigidTransform<float, 3>& BTM, const float CullDistance);
@@ -2039,15 +2052,15 @@ namespace Chaos
 		template void UpdateLevelsetLevelsetManifold<float, 3>(TCollisionConstraintBase<float, 3>&  Constraint, const TRigidTransform<float, 3>& ATM, const TRigidTransform<float, 3>& BTM, const float CullDistance);
 		template void ConstructLevelsetLevelsetConstraints<float, 3>(TGeometryParticleHandle<float, 3>* Particle0, TGeometryParticleHandle<float, 3>* Particle1, const FImplicitObject* Implicit0, const FImplicitObject* Implicit1, const TRigidTransform<float, 3>& Transform0, const TRigidTransform<float, 3>& Transform1, const float CullDistance, FCollisionConstraintsArray& NewConstraints);
 
-		template void ConstructUnionUnionConstraints<float, 3>(TGeometryParticleHandle<float, 3>* Particle0, TGeometryParticleHandle<float, 3>* Particle1, const FImplicitObject* Implicit0, const FImplicitObject* Implicit1, const TRigidTransform<float, 3>& Transform0, const TRigidTransform<float, 3>& Transform1, const float CullDistance, FCollisionConstraintsArray& NewConstraints);
+		template void ConstructUnionUnionConstraints<float, 3>(TGeometryParticleHandle<float, 3>* Particle0, TGeometryParticleHandle<float, 3>* Particle1, const FImplicitObject* Implicit0, const FImplicitObject* Implicit1, const TRigidTransform<float, 3>& Transform0, const TRigidTransform<float, 3>& Transform1, const float CullDistance, const FCollisionContext& Context, FCollisionConstraintsArray& NewConstraints);
 
 		template void UpdateSingleShotManifold<float, 3>(TRigidBodyMultiPointContactConstraint<float, 3>&  Constraint, const TRigidTransform<float, 3>& Transform0, const TRigidTransform<float, 3>& Transform1, const float CullDistance);
 		template void UpdateIterativeManifold<float, 3>(TRigidBodyMultiPointContactConstraint<float, 3>&  Constraint, const TRigidTransform<float, 3>& Transform0, const TRigidTransform<float, 3>& Transform1, const float CullDistance);
-		template void UpdateManifold<float, 3>(TRigidBodyMultiPointContactConstraint<float, 3>& Constraint, const TRigidTransform<float, 3>& ATM, const TRigidTransform<float, 3>& BTM, const float CullDistance);
+		template void UpdateManifold<float, 3>(TRigidBodyMultiPointContactConstraint<float, 3>& Constraint, const TRigidTransform<float, 3>& ATM, const TRigidTransform<float, 3>& BTM, const float CullDistance, const FCollisionContext& Context);
 		template void UpdateConstraintFromGeometry<ECollisionUpdateType::Any, float, 3>(TRigidBodyPointContactConstraint<float, 3>& ConstraintBase, const TRigidTransform<float, 3>& Transform0, const TRigidTransform<float, 3>& Transform1, const float CullDistance);
 		template void UpdateConstraintFromGeometry<ECollisionUpdateType::Deepest, float, 3>(TRigidBodyPointContactConstraint<float, 3>& ConstraintBase, const TRigidTransform<float, 3>& Transform0, const TRigidTransform<float, 3>& Transform1, const float CullDistance);
 		template void UpdateConstraintFromManifold(TRigidBodyMultiPointContactConstraint<float, 3>& Constraint, const TRigidTransform<float, 3>& Transform0, const TRigidTransform<float, 3>& Transform1, const float CullDistance);
-		template void ConstructConstraints<float, 3>(TGeometryParticleHandle<float, 3>* Particle0, TGeometryParticleHandle<float, 3>* Particle1, const FImplicitObject* Implicit0, const FImplicitObject* Implicit1, const TRigidTransform<float, 3>& Transform0, const TRigidTransform<float, 3>& Transform1, const float CullDistance, FCollisionConstraintsArray& NewConstraints);
+		template void ConstructConstraints<float, 3>(TGeometryParticleHandle<float, 3>* Particle0, TGeometryParticleHandle<float, 3>* Particle1, const FImplicitObject* Implicit0, const FImplicitObject* Implicit1, const TRigidTransform<float, 3>& Transform0, const TRigidTransform<float, 3>& Transform1, const float CullDistance, const FCollisionContext& Context, FCollisionConstraintsArray& NewConstraints);
 
 	} // Collisions
 
diff --git a/Engine/Source/Runtime/Experimental/Chaos/Private/Chaos/PBDCollisionConstraints.cpp b/Engine/Source/Runtime/Experimental/Chaos/Private/Chaos/PBDCollisionConstraints.cpp
--- a/Engine/Source/Runtime/Experimental/Chaos/Private/Chaos/PBDCollisionConstraints.cpp
+++ b/Engine/Source/Runtime/Experimental/Chaos/Private/Chaos/PBDCollisionConstraints.cpp
@@ -8,6 +8,7 @@
 #include "Chaos/PBDCollisionConstraintsContact.h"
 #include "Chaos/CollisionResolutionUtil.h"
 #include "Chaos/CollisionResolution.h"
+#include "Chaos/Collision/CollisionContext.h"
 #include "Chaos/Defines.h"
 #include "Chaos/GeometryQueries.h"
 #include "Chaos/ImplicitObjectUnion.h"
@@ -388,9 +389,11 @@ namespace Chaos
 		//	Collisions::Update(MCullDistance, MShapePadding, ConstraintHandle->GetContact());
 		//}, bDisableCollisionParallelFor);
 
+		FCollisionContext Context;
+
 		for (FMultiPointContactConstraint& Contact : IterativeConstraints)
 		{
-			Collisions::UpdateManifold(Contact, MCullDistance);
+			Collisions::UpdateManifold(Contact, MCullDistance, Context);
 			if (Contact.GetPhi() < MCullDistance)
 			{
 				Contact.Timestamp = LifespanCounter;
diff --git a/Engine/Source/Runtime/Experimental/Chaos/Private/Chaos/PBDCollisionConstraintsContact.cpp b/Engine/Source/Runtime/Experimental/Chaos/Private/Chaos/PBDCollisionConstraintsContact.cpp
--- a/Engine/Source/Runtime/Experimental/Chaos/Private/Chaos/PBDCollisionConstraintsContact.cpp
+++ b/Engine/Source/Runtime/Experimental/Chaos/Private/Chaos/PBDCollisionConstraintsContact.cpp
@@ -16,12 +16,12 @@ namespace Chaos
 		int32 Chaos_Collision_EnergyClampEnabled = 1;
 		FAutoConsoleVariableRef CVarChaosImmPhysStepTime(TEXT("p.Chaos.Collision.EnergyClampEnabled"), Chaos_Collision_EnergyClampEnabled, TEXT("Whether to use energy clamping in collision apply step"));
 
-		extern void UpdateManifold(FRigidBodyMultiPointContactConstraint& Constraint, const FReal CullDistance)
+		extern void UpdateManifold(FRigidBodyMultiPointContactConstraint& Constraint, const FReal CullDistance, const FCollisionContext& Context)
 		{
 			const FRigidTransform3 Transform0 = GetTransform(Constraint.Particle[0]);
 			const FRigidTransform3 Transform1 = GetTransform(Constraint.Particle[1]);
 
-			UpdateManifold(Constraint, Transform0, Transform1, CullDistance);
+			UpdateManifold(Constraint, Transform0, Transform1, CullDistance, Context);
 		}
 
 		void Update(FRigidBodyPointContactConstraint& Constraint, const FReal CullDistance)
diff --git a/Engine/Source/Runtime/Experimental/Chaos/Public/Chaos/Collision/CollisionContext.h b/Engine/Source/Runtime/Experimental/Chaos/Public/Chaos/Collision/CollisionContext.h
new file mode 100644
--- /dev/null
+++ b/Engine/Source/Runtime/Experimental/Chaos/Public/Chaos/Collision/CollisionContext.h
@@ -0,0 +1,26 @@
+// Copyright 1998-2019 Epic Games, Inc. All Rights Reserved.
+#pragma once
+
+#include "Chaos/Core.h"
+
+namespace Chaos
+{
+	/**
+	 * Data passed down into the collision detection functions.
+	 */
+	class FCollisionContext
+	{
+	public:
+		FCollisionContext()
+			: SpaceTransform(FVec3(0), FRotation3::FromIdentity())
+			, bFilteringEnabled(true)
+		{
+		}
+
+		// The simulation space to world-space transform (e.g., Rigid Body nodes are often simulated relative to a component or bone)
+		FRigidTransform3 SpaceTransform;
+
+		// Whether to check the shape query flags (e.g., Rigid Body nodes have already performed filtering prior to collision detection)
+		bool bFilteringEnabled;
+	};
+}
\ No newline at end of file
diff --git a/Engine/Source/Runtime/Experimental/Chaos/Public/Chaos/Collision/CollisionDetector.h b/Engine/Source/Runtime/Experimental/Chaos/Public/Chaos/Collision/CollisionDetector.h
--- a/Engine/Source/Runtime/Experimental/Chaos/Public/Chaos/Collision/CollisionDetector.h
+++ b/Engine/Source/Runtime/Experimental/Chaos/Public/Chaos/Collision/CollisionDetector.h
@@ -1,6 +1,7 @@
 // Copyright Epic Games, Inc. All Rights Reserved.
 #pragma once
 
+#include "Chaos/Collision/CollisionContext.h"
 #include "Chaos/Collision/CollisionReceiver.h"
 #include "Chaos/Collision/NarrowPhase.h"
 #include "Chaos/Collision/ParticlePairBroadPhase.h"
@@ -30,6 +31,7 @@ namespace Chaos
 
 		FBroadPhase& GetBroadPhase() { return BroadPhase; }
 		FContainer& GetCollisionContainer() { return CollisionContainer; }
+		FCollisionContext& GetContext() { return Context; }
 
 		void DetectCollisions(const FReal Dt)
 		{
@@ -54,7 +56,7 @@ namespace Chaos
 			// Receivers and NarrowPhase are assumed to be stateless atm. If we change that, they need to
 			// be passed into the constructor with the BroadPhase and Container.
 			FReceiver Receiver(CollisionContainer);
-			FNarrowPhase NarrowPhase;
+			FNarrowPhase NarrowPhase(Context);
 			BroadPhase.ProduceOverlaps(Dt, NarrowPhase, Receiver, StatData);
 			Receiver.ProcessCollisions();
 		}
@@ -62,6 +64,7 @@ namespace Chaos
 	private:
 		FBroadPhase& BroadPhase;
 		FContainer& CollisionContainer;
+		FCollisionContext Context;
 	};
 
 }
diff --git a/Engine/Source/Runtime/Experimental/Chaos/Public/Chaos/Collision/NarrowPhase.h b/Engine/Source/Runtime/Experimental/Chaos/Public/Chaos/Collision/NarrowPhase.h
--- a/Engine/Source/Runtime/Experimental/Chaos/Public/Chaos/Collision/NarrowPhase.h
+++ b/Engine/Source/Runtime/Experimental/Chaos/Public/Chaos/Collision/NarrowPhase.h
@@ -19,6 +19,8 @@ namespace Chaos
 	#define SCOPE_CYCLE_COUNTER_NAROWPHASE()
 #endif
 
+	class FCollisionContext;
+
 	/**
 	 * Generate contact manifolds for particle pairs.
 	 *
@@ -27,6 +29,11 @@ namespace Chaos
 	class CHAOS_API FNarrowPhase
 	{
 	public:
+		FNarrowPhase(const FCollisionContext& InContext)
+			: Context(InContext)
+		{
+		}
+
 		// @todo(chaos): COLLISION Transient Handle version
 		/**
 		 * /param CullDistance The contact separation at which we ignore the constraint
@@ -43,11 +50,14 @@ namespace Chaos
 				//   determine if the constraint is already defined, and then opt out of 
 				//   the creation process. 
 				//
-				Collisions::ConstructConstraints<FReal, 3>(Particle0, Particle1, Particle0->Geometry().Get(), Particle1->Geometry().Get(), Collisions::GetTransform(Particle0), Collisions::GetTransform(Particle1), CullDistance, NewConstraints);
+				Collisions::ConstructConstraints<FReal, 3>(Particle0, Particle1, Particle0->Geometry().Get(), Particle1->Geometry().Get(), Collisions::GetTransform(Particle0), Collisions::GetTransform(Particle1), CullDistance, Context, NewConstraints);
 
 				CHAOS_COLLISION_STAT(if (NewConstraints.Num()) { StatData.IncrementCountNP(NewConstraints.Num()); });
 				CHAOS_COLLISION_STAT(if (!NewConstraints.Num()) { StatData.IncrementRejectedNP(); });
 			}
 		}
+
+	private:
+		const FCollisionContext& Context;
 	};
 }
diff --git a/Engine/Source/Runtime/Experimental/Chaos/Public/Chaos/CollisionResolution.h b/Engine/Source/Runtime/Experimental/Chaos/Public/Chaos/CollisionResolution.h
--- a/Engine/Source/Runtime/Experimental/Chaos/Public/Chaos/CollisionResolution.h
+++ b/Engine/Source/Runtime/Experimental/Chaos/Public/Chaos/CollisionResolution.h
@@ -33,6 +33,8 @@ namespace Chaos
 
 	class FTriangleMeshImplicitObject;
 
+	class FCollisionContext;
+
 
 	namespace Collisions
 	{
@@ -178,11 +180,8 @@ namespace Chaos
 		template <typename T, int d>
 		void CHAOS_API UpdateCapsuleBoxConstraint(const TCapsule<T>& A, const TRigidTransform<T, d>& ATransform, const TAABB<T, d>& B, const TRigidTransform<T, d>& BTransform, const T CullDistance, TRigidBodyPointContactConstraint<T, d>& Constraint);
 
-		template<class T, int d>
-		void CHAOS_API UpdateCapsuleBoxManifold(TCollisionConstraintBase<T, d>&  Constraint, const TRigidTransform<T, d>& ATM, const TRigidTransform<T, d>& BTM, const T CullDistance);
-
 		template<typename T, int d>
-		void CHAOS_API ConstructCapsuleBoxConstraints(TGeometryParticleHandle<T, d>* Particle0, TGeometryParticleHandle<T, d>* Particle1, const FImplicitObject* Implicit0, const FImplicitObject* Implicit1, const TRigidTransform<T, d>& Transform0, const TRigidTransform<T, d>& Transform1, const T CullDistance, FCollisionConstraintsArray& NewConstraints);
+		void CHAOS_API ConstructCapsuleBoxConstraints(TGeometryParticleHandle<T, d>* Particle0, TGeometryParticleHandle<T, d>* Particle1, const FImplicitObject* Implicit0, const FImplicitObject* Implicit1, const TRigidTransform<T, d>& Transform0, const TRigidTransform<T, d>& Transform1, const T CullDistance, const FCollisionContext& Context, FCollisionConstraintsArray& NewConstraints);
 
 		//
 		// Capsule-HeightField
@@ -252,7 +251,7 @@ namespace Chaos
 		void CHAOS_API UpdateIterativeManifold(TRigidBodyMultiPointContactConstraint<T, d>&  Constraint, const TRigidTransform<T, d>& Transform0, const TRigidTransform<T, d>& Transform1, const T CullDistance);
 
 		template<typename T, int d>
-		void CHAOS_API UpdateManifold(TRigidBodyMultiPointContactConstraint<T, d>& Constraint, const TRigidTransform<T, d>& ATM, const TRigidTransform<T, d>& BTM, const T CullDistance);
+		void CHAOS_API UpdateManifold(TRigidBodyMultiPointContactConstraint<T, d>& Constraint, const TRigidTransform<T, d>& ATM, const TRigidTransform<T, d>& BTM, const T CullDistance, const FCollisionContext& Context);
 
 		template<class T, int d>
 		void UpdateConstraintFromManifold(TRigidBodyMultiPointContactConstraint<T, d>& Constraint, const TRigidTransform<T, d>& Transform0, const TRigidTransform<T, d>& Transform1, const T CullDistance);
@@ -261,7 +260,7 @@ namespace Chaos
 		void UpdateConstraintFromGeometry(TRigidBodyPointContactConstraint<T, d>& Constraint, const TRigidTransform<T, d>& ParticleTransform0, const TRigidTransform<T, d>& ParticleTransform1, const T CullDistance);
 
 		template<typename T, int d>
-		void CHAOS_API ConstructConstraints(TGeometryParticleHandle<T, d>* Particle0, TGeometryParticleHandle<T, d>* Particle1, const FImplicitObject* Implicit0, const FImplicitObject* Implicit1, const TRigidTransform<T, d>& Transform0, const TRigidTransform<T, d>& Transform1, const T CullDistance, FCollisionConstraintsArray& NewConstraints);
+		void CHAOS_API ConstructConstraints(TGeometryParticleHandle<T, d>* Particle0, TGeometryParticleHandle<T, d>* Particle1, const FImplicitObject* Implicit0, const FImplicitObject* Implicit1, const TRigidTransform<T, d>& Transform0, const TRigidTransform<T, d>& Transform1, const T CullDistance, const FCollisionContext& Context, FCollisionConstraintsArray& NewConstraints);
 
 	}
 }
diff --git a/Engine/Source/Runtime/Experimental/Chaos/Public/Chaos/CollisionResolutionTypes.h b/Engine/Source/Runtime/Experimental/Chaos/Public/Chaos/CollisionResolutionTypes.h
--- a/Engine/Source/Runtime/Experimental/Chaos/Public/Chaos/CollisionResolutionTypes.h
+++ b/Engine/Source/Runtime/Experimental/Chaos/Public/Chaos/CollisionResolutionTypes.h
@@ -210,7 +210,16 @@ namespace Chaos
 			EContactShapesType ShapesType)
 			: Base(Particle0, Implicit0, Transform0, Particle1, Implicit1, Transform1, Base::FType::SinglePoint, ShapesType) {}
 
-		static typename Base::FType StaticType() { return Base::FType::SinglePoint; };
+		static typename TCollisionConstraintBase<T, d>::FType StaticType() { return Base::FType::SinglePoint; };
+
+	protected:
+		// For use by derived types that can be used as point constraints in Update
+		TRigidBodyPointContactConstraint(typename Base::FType InType) : Base(InType) {}
+		TRigidBodyPointContactConstraint(
+			FGeometryParticleHandle* Particle0, const FImplicitObject* Implicit0, const TRigidTransform<T, d>& Transform0,
+			FGeometryParticleHandle* Particle1, const FImplicitObject* Implicit1, const TRigidTransform<T, d>& Transform1,
+			typename Base::FType InType, EContactShapesType ShapesType)
+			: Base(Particle0, Implicit0, Transform0, Particle1, Implicit1, Transform1, InType, ShapesType) {}
 	};
 	typedef TRigidBodyPointContactConstraint<float, 3> FRigidBodyPointContactConstraint;
 
@@ -219,15 +228,15 @@ namespace Chaos
 	*
 	*/
 	template<class T, int d>
-	class TRigidBodyMultiPointContactConstraint : public TCollisionConstraintBase<T, d>
+	class TRigidBodyMultiPointContactConstraint : public TRigidBodyPointContactConstraint<T, d>
 	{
 	public:
-		using Base = TCollisionConstraintBase<T, d>;
+		using Base = TRigidBodyPointContactConstraint<T, d>;
 		using FGeometryParticleHandle = TGeometryParticleHandle<T, d>;
 		using FManifold = TCollisionContact<T, d>;
 		using Base::Particle;
 
-		TRigidBodyMultiPointContactConstraint() : Base(Base::FType::MultiPoint), PlaneOwnerIndex(INDEX_NONE), PlaneFaceIndex(INDEX_NONE), PlaneNormal(0), PlanePosition(0) {}
+		TRigidBodyMultiPointContactConstraint() : Base(TCollisionConstraintBase<T, d>::FType::MultiPoint), PlaneOwnerIndex(INDEX_NONE), PlaneFaceIndex(INDEX_NONE), PlaneNormal(0), PlanePosition(0) {}
 		TRigidBodyMultiPointContactConstraint(
 			FGeometryParticleHandle* Particle0, const FImplicitObject* Implicit0, const TRigidTransform<T, d>& Transform0,
 			FGeometryParticleHandle* Particle1, const FImplicitObject* Implicit1, const TRigidTransform<T, d>& Transform1,
@@ -236,7 +245,7 @@ namespace Chaos
 			, PlaneOwnerIndex(INDEX_NONE), PlaneFaceIndex(INDEX_NONE), PlaneNormal(0), PlanePosition(0)
 		{}
 
-		static typename Base::FType StaticType() { return Base::FType::MultiPoint; };
+		static typename TCollisionConstraintBase<T, d>::FType StaticType() { return Base::FType::MultiPoint; };
 
 		// Get the particle that owns the plane
 		FGeometryParticleHandle* PlaneParticleHandle() const { check(PlaneOwnerIndex >= 0 && PlaneOwnerIndex < 2); return Particle[PlaneOwnerIndex]; }
diff --git a/Engine/Source/Runtime/Experimental/Chaos/Public/Chaos/PBDCollisionConstraintsContact.h b/Engine/Source/Runtime/Experimental/Chaos/Public/Chaos/PBDCollisionConstraintsContact.h
--- a/Engine/Source/Runtime/Experimental/Chaos/Public/Chaos/PBDCollisionConstraintsContact.h
+++ b/Engine/Source/Runtime/Experimental/Chaos/Public/Chaos/PBDCollisionConstraintsContact.h
@@ -7,8 +7,11 @@
 
 namespace Chaos
 {
+	class FCollisionContext;
+
 	namespace Collisions
 	{
+
 		template<class T = float>
 		struct TContactParticleParameters {
 			T CullDistance;
@@ -28,7 +31,7 @@ namespace Chaos
 		using FContactIterationParameters = TContactIterationParameters<FReal>;
 
 		// Regenerate (one-shot or incremental) the manifold plane and points
-		extern void UpdateManifold(FRigidBodyMultiPointContactConstraint& Constraint, const FReal CullDistance);
+		extern void UpdateManifold(FRigidBodyMultiPointContactConstraint& Constraint, const FReal CullDistance, const FCollisionContext& Context);
 
 		// Update the constraint (re-runs collision detection for this contact)
 		extern void Update(FRigidBodyPointContactConstraint& Constraint, const FReal CullDistance);
