Chaos - added missing forwarding of GetMaterialIndex solves getting physical materials from TriMesh

diff --git a/Engine/Source/Runtime/Experimental/Chaos/Public/Chaos/ImplicitObjectScaled.h b/Engine/Source/Runtime/Experimental/Chaos/Public/Chaos/ImplicitObjectScaled.h
--- a/Engine/Source/Runtime/Experimental/Chaos/Public/Chaos/ImplicitObjectScaled.h
+++ b/Engine/Source/Runtime/Experimental/Chaos/Public/Chaos/ImplicitObjectScaled.h
@@ -143,6 +143,11 @@ public:
 		return MObject->OverlapGeom(B,BToATM,Thickness, OutMTD);
 	}
 
+	virtual uint16 GetMaterialIndex(uint32 HintIndex) const
+	{
+		return MObject->GetMaterialIndex(HintIndex);
+	}
+
 protected:
 	ObjectType MObject;
 
