Chaos - fix potential null deref (static analysis)

diff --git a/Engine/Source/Runtime/Experimental/Chaos/Private/Chaos/CollisionResolution.cpp b/Engine/Source/Runtime/Experimental/Chaos/Private/Chaos/CollisionResolution.cpp
--- a/Engine/Source/Runtime/Experimental/Chaos/Private/Chaos/CollisionResolution.cpp
+++ b/Engine/Source/Runtime/Experimental/Chaos/Private/Chaos/CollisionResolution.cpp
@@ -1297,7 +1297,6 @@ namespace Chaos
 		template<typename T, int d>
 		void ConstructConvexHeightFieldConstraints(TGeometryParticleHandle<T, d>* Particle0, TGeometryParticleHandle<T, d>* Particle1, const FImplicitObject* Implicit0, const FImplicitObject* Implicit1, const TRigidTransform<T, d>& Transform0, const TRigidTransform<T, d>& Transform1, const T CullDistance, FCollisionConstraintsArray& NewConstraints)
 		{
-
 			const THeightField<T> * Object1 = Implicit1->template GetObject<const THeightField<T> >();
 			if (ensure(Implicit0->IsConvex() && Object1))
 			{
@@ -1566,6 +1565,8 @@ namespace Chaos
 			// @todo(chaos): We use GetInnerType here because TriMeshes are left with their "Instanced" wrapper, unlike all other instanced implicits. Should we strip the instance on Tri Mesh too?
 			EImplicitObjectType Implicit0Type = Implicit0 ? GetInnerType(Implicit0->GetType()) : ImplicitObjectType::Unknown;
 			EImplicitObjectType Implicit1Type = Implicit1 ? GetInnerType(Implicit1->GetType()) : ImplicitObjectType::Unknown;
+			bool bIsConvex0 = Implicit0 && Implicit0->IsConvex();
+			bool bIsConvex1 = Implicit1 && Implicit1->IsConvex();
 
 			if (Implicit0Type == TBox<FReal, 3>::StaticType() && Implicit1Type == TBox<FReal, 3>::StaticType())
 			{
@@ -1667,23 +1668,23 @@ namespace Chaos
 			{
 				ConstructCapsuleTriangleMeshConstraints(Particle1, Particle0, Implicit1, Implicit0, Transform1, Transform0, CullDistance, NewConstraints);
 			}
-			else if (Implicit0->IsConvex() && Implicit1Type == THeightField<FReal>::StaticType())
+			else if (bIsConvex0 && Implicit1Type == THeightField<FReal>::StaticType())
 			{
 				ConstructConvexHeightFieldConstraints(Particle0, Particle1, Implicit0, Implicit1, Transform0, Transform1, CullDistance, NewConstraints);
 			}
-			else if (Implicit0Type == THeightField<FReal>::StaticType() && Implicit1->IsConvex())
+			else if (Implicit0Type == THeightField<FReal>::StaticType() && bIsConvex1)
 			{
 				ConstructConvexHeightFieldConstraints(Particle1, Particle0, Implicit1, Implicit0, Transform1, Transform0, CullDistance, NewConstraints);
 			}
-			else if (Implicit0->IsConvex() && Implicit1Type == FTriangleMeshImplicitObject::StaticType())
+			else if (bIsConvex0 && Implicit1Type == FTriangleMeshImplicitObject::StaticType())
 			{
 				ConstructConvexTriangleMeshConstraints(Particle0, Particle1, Implicit0, Implicit1, Transform0, Transform1, CullDistance, NewConstraints);
 			}
-			else if (Implicit0Type == FTriangleMeshImplicitObject::StaticType() && Implicit1->IsConvex())
+			else if (Implicit0Type == FTriangleMeshImplicitObject::StaticType() && bIsConvex1)
 			{
 				ConstructConvexTriangleMeshConstraints(Particle1, Particle0, Implicit1, Implicit0, Transform1, Transform0, CullDistance, NewConstraints);
 			}
-			else if (Implicit0->IsConvex() && Implicit1->IsConvex())
+			else if (bIsConvex0 && bIsConvex1)
 			{
 				ConstructConvexConvexConstraints(Particle0, Particle1, Implicit0, Implicit1, Transform0, Transform1, CullDistance, NewConstraints);
 			}
