Fix CIS

diff --git a/Engine/Source/Runtime/Experimental/Chaos/Public/Chaos/GeometryParticles.h b/Engine/Source/Runtime/Experimental/Chaos/Public/Chaos/GeometryParticles.h
--- a/Engine/Source/Runtime/Experimental/Chaos/Public/Chaos/GeometryParticles.h
+++ b/Engine/Source/Runtime/Experimental/Chaos/Public/Chaos/GeometryParticles.h
@@ -230,8 +230,8 @@ namespace Chaos
 		CHAOS_API FUniqueIdx UniqueIdx(const int32 Index) const { return MUniqueIdx[Index]; }
 		CHAOS_API FUniqueIdx& UniqueIdx(const int32 Index) { return MUniqueIdx[Index]; }
 
-		void*& UserData(const int32 Index) { return MUserData[Index]; }
-		const void* UserData(const int32 Index) const { return MUserData[Index]; }
+		CHAOS_API void*& UserData(const int32 Index) { return MUserData[Index]; }
+		CHAOS_API const void* UserData(const int32 Index) const { return MUserData[Index]; }
 
 		CHAOS_API TSerializablePtr<FImplicitObject> Geometry(const int32 Index) const { return MGeometry[Index]; }
 
diff --git a/Engine/Source/Runtime/Experimental/Chaos/Public/Chaos/RigidParticles.h b/Engine/Source/Runtime/Experimental/Chaos/Public/Chaos/RigidParticles.h
--- a/Engine/Source/Runtime/Experimental/Chaos/Public/Chaos/RigidParticles.h
+++ b/Engine/Source/Runtime/Experimental/Chaos/Public/Chaos/RigidParticles.h
@@ -175,7 +175,7 @@ class TRigidParticles : public TKinematicGeometryParticles<T, d>
 		MSleepData.Add(SleepData);
 		SleepDataLock.WriteUnlock();
 	}
-	void ClearSleepData()
+	CHAOS_API void ClearSleepData()
 	{
 		SleepDataLock.WriteLock();
 		MSleepData.Empty();
