Force copies when inserting child implicits into the shape map using their parent index so we don't alias the TMap allocation.

diff --git a/Engine/Source/Runtime/Experimental/Chaos/Private/Chaos/GeometryParticles.cpp b/Engine/Source/Runtime/Experimental/Chaos/Private/Chaos/GeometryParticles.cpp
--- a/Engine/Source/Runtime/Experimental/Chaos/Private/Chaos/GeometryParticles.cpp
+++ b/Engine/Source/Runtime/Experimental/Chaos/Private/Chaos/GeometryParticles.cpp
@@ -234,11 +234,11 @@ namespace Chaos
 						{
 							if (ImplicitShapeMap[Index].Contains(ImplicitObject.Get()))
 							{
-								ImplicitShapeMap[Index].Add(ImplicitChildObject, ImplicitShapeMap[Index][ImplicitObject.Get()]);
+								ImplicitShapeMap[Index].Add(ImplicitChildObject, CopyTemp(ImplicitShapeMap[Index][ImplicitObject.Get()]));
 							}
 							else if (ImplicitShapeMap[Index].Contains(ImplicitChildObject))
 							{
-								ImplicitShapeMap[Index].Add(ImplicitObject.Get(), ImplicitShapeMap[Index][ImplicitChildObject]);
+								ImplicitShapeMap[Index].Add(ImplicitObject.Get(), CopyTemp(ImplicitShapeMap[Index][ImplicitChildObject]));
 							}
 							
 						}
@@ -251,11 +251,11 @@ namespace Chaos
 				{
 					if (ImplicitShapeMap[Index].Contains(MGeometry[Index].Get()))
 					{
-						ImplicitShapeMap[Index].Add(ImplicitChildObject, ImplicitShapeMap[Index][MGeometry[Index].Get()]);
+						ImplicitShapeMap[Index].Add(ImplicitChildObject, CopyTemp(ImplicitShapeMap[Index][MGeometry[Index].Get()]));
 					}
 					else if (ImplicitShapeMap[Index].Contains(ImplicitChildObject))
 					{
-						ImplicitShapeMap[Index].Add(MGeometry[Index].Get(), ImplicitShapeMap[Index][ImplicitChildObject]);
+						ImplicitShapeMap[Index].Add(MGeometry[Index].Get(), CopyTemp(ImplicitShapeMap[Index][ImplicitChildObject]));
 					}
 					
 				}
diff --git a/Engine/Source/Runtime/Experimental/Chaos/Private/Chaos/ParticleHandle.cpp b/Engine/Source/Runtime/Experimental/Chaos/Private/Chaos/ParticleHandle.cpp
--- a/Engine/Source/Runtime/Experimental/Chaos/Private/Chaos/ParticleHandle.cpp
+++ b/Engine/Source/Runtime/Experimental/Chaos/Private/Chaos/ParticleHandle.cpp
@@ -39,11 +39,11 @@ namespace Chaos
 						{
 							if (ImplicitShapeMap.Contains(ImplicitObject.Get()))
 							{
-								ImplicitShapeMap.Add(ImplicitChildObject, ImplicitShapeMap[ImplicitObject.Get()]);
+								ImplicitShapeMap.Add(ImplicitChildObject, CopyTemp(ImplicitShapeMap[ImplicitObject.Get()]));
 							}
 							else if (ImplicitShapeMap.Contains(ImplicitChildObject))
 							{
-								ImplicitShapeMap.Add(ImplicitObject.Get(), ImplicitShapeMap[ImplicitChildObject]);
+								ImplicitShapeMap.Add(ImplicitObject.Get(), CopyTemp(ImplicitShapeMap[ImplicitChildObject]));
 							}
 						}
 					}
@@ -55,11 +55,11 @@ namespace Chaos
 				{
 					if (ImplicitShapeMap.Contains(MGeometry.Get()))
 					{
-						ImplicitShapeMap.Add(ImplicitChildObject, ImplicitShapeMap[MGeometry.Get()]);
+						ImplicitShapeMap.Add(ImplicitChildObject, CopyTemp(ImplicitShapeMap[MGeometry.Get()]));
 					}
 					else if (ImplicitShapeMap.Contains(ImplicitChildObject))
 					{
-						ImplicitShapeMap.Add(MGeometry.Get(), ImplicitShapeMap[ImplicitChildObject]);
+						ImplicitShapeMap.Add(MGeometry.Get(), CopyTemp(ImplicitShapeMap[ImplicitChildObject]));
 					}
 				}
 			}
