Removed exploratory logging for previous tree rebuild issue.

diff --git a/Engine/Source/Runtime/Experimental/Chaos/Public/Chaos/AABBTree.h b/Engine/Source/Runtime/Experimental/Chaos/Public/Chaos/AABBTree.h
--- a/Engine/Source/Runtime/Experimental/Chaos/Public/Chaos/AABBTree.h
+++ b/Engine/Source/Runtime/Experimental/Chaos/Public/Chaos/AABBTree.h
@@ -499,7 +499,6 @@ public:
 
 		if(DirtyElements.Num() > MaxDirtyElements)
 		{
-			UE_LOG(LogChaos, Warning, TEXT("Bounding volume exceeded maximum dirty elements (%d dirty of max %d) and is forcing a tree rebuild."), DirtyElements.Num(), MaxDirtyElements);
 			ReoptimizeTree();
 		}
 	}
