Chaos RBAN stability improvements
- remove post-rotation position fixup as it caused swimming
- retuned stiffnes scaling for parity
- increased joint pair iterations
- set capsule-box manifold cutoff angle

diff --git a/Engine/Source/Runtime/Engine/Private/PhysicsEngine/ImmediatePhysics/ImmediatePhysicsChaos/ImmediatePhysicsJointHandle_Chaos.cpp b/Engine/Source/Runtime/Engine/Private/PhysicsEngine/ImmediatePhysics/ImmediatePhysicsChaos/ImmediatePhysicsJointHandle_Chaos.cpp
--- a/Engine/Source/Runtime/Engine/Private/PhysicsEngine/ImmediatePhysics/ImmediatePhysicsChaos/ImmediatePhysicsJointHandle_Chaos.cpp
+++ b/Engine/Source/Runtime/Engine/Private/PhysicsEngine/ImmediatePhysics/ImmediatePhysicsChaos/ImmediatePhysicsJointHandle_Chaos.cpp
@@ -22,16 +22,18 @@ namespace ImmediatePhysics_Chaos
 {
 	float ChaosImmediate_JointStiffness = 1.0f;
 
-	float ChaosImmediate_DriveStiffnessScale = 0.85f;
-	float ChaosImmediate_DriveDampingScale = 0.85f;
+	float ChaosImmediate_LinearDriveStiffnessScale = 1.0f;
+	float ChaosImmediate_LinearDriveDampingScale = 1.0f;
+	float ChaosImmediate_AngularDriveStiffnessScale = 1.5f;
+	float ChaosImmediate_AngularDriveDampingScale = 1.5f;
 
 	int ChaosImmediate_SoftLinearForceMode = (int)Chaos::EJointForceMode::Acceleration;
-	float ChaosImmediate_SoftLinearStiffnessScale = 1.2f;
+	float ChaosImmediate_SoftLinearStiffnessScale = 1.5f;
 	float ChaosImmediate_SoftLinearDampingScale = 1.2f;
 
 	int ChaosImmediate_SoftAngularForceMode = (int)Chaos::EJointForceMode::Acceleration;
-	float ChaosImmediate_SoftAngularStiffnessScale = 500;
-	float ChaosImmediate_SoftAngularDampingScale = 5;
+	float ChaosImmediate_SoftAngularStiffnessScale = 100000;
+	float ChaosImmediate_SoftAngularDampingScale = 1000;
 
 	float ChaosImmediate_JointMinLinearProjection = 0.0f;
 	float ChaosImmediate_JointMaxLinearProjection = 0.8f;
@@ -39,8 +41,10 @@ namespace ImmediatePhysics_Chaos
 	float ChaosImmediate_JointMaxAngularProjection = 0.1f;
 
 	FAutoConsoleVariableRef CVarJointStiffness(TEXT("p.Chaos.ImmPhys.JointStiffness"), ChaosImmediate_JointStiffness, TEXT("Hard-joint solver stiffness."));
-	FAutoConsoleVariableRef CVarDriveStiffnessScale(TEXT("p.Chaos.ImmPhys.DriveStiffnessScale"), ChaosImmediate_DriveStiffnessScale, TEXT("Conversion factor for drive stiffness."));
-	FAutoConsoleVariableRef CVarDriveDampingScale(TEXT("p.Chaos.ImmPhys.DriveDampingScale"), ChaosImmediate_DriveDampingScale, TEXT("Conversion factor for drive damping."));
+	FAutoConsoleVariableRef CVarLinearDriveStiffnessScale(TEXT("p.Chaos.ImmPhys.LinearDriveStiffnessScale"), ChaosImmediate_LinearDriveStiffnessScale, TEXT("Conversion factor for Linear drive stiffness."));
+	FAutoConsoleVariableRef CVarLinearDriveDampingScale(TEXT("p.Chaos.ImmPhys.LinaearDriveDampingScale"), ChaosImmediate_LinearDriveDampingScale, TEXT("Conversion factor for Linear drive damping."));
+	FAutoConsoleVariableRef CVarAngularDriveStiffnessScale(TEXT("p.Chaos.ImmPhys.AngularDriveStiffnessScale"), ChaosImmediate_AngularDriveStiffnessScale, TEXT("Conversion factor for Angular drive stiffness."));
+	FAutoConsoleVariableRef CVarAngularDriveDampingScale(TEXT("p.Chaos.ImmPhys.AngularDriveDampingScale"), ChaosImmediate_AngularDriveDampingScale, TEXT("Conversion factor for Angular drive damping."));
 	FAutoConsoleVariableRef CVarSoftLinearStiffnessScale(TEXT("p.Chaos.ImmPhys.SoftLinearStiffnessScale"), ChaosImmediate_SoftLinearStiffnessScale, TEXT("Conversion factor for soft-joint stiffness."));
 	FAutoConsoleVariableRef CVarSoftLinearDampingScale(TEXT("p.Chaos.ImmPhys.SoftLinearDampingScale"), ChaosImmediate_SoftLinearDampingScale, TEXT("Conversion factor for soft-joint damping."));
 	FAutoConsoleVariableRef CVarSoftLinearForceMode(TEXT("p.Chaos.ImmPhys.SoftLinearForceMode"), ChaosImmediate_SoftLinearForceMode, TEXT("Soft Linear constraint force mode (0: Acceleration; 1: Force"));
@@ -98,8 +102,8 @@ namespace ImmediatePhysics_Chaos
 		ConstraintSettings.bLinearVelocityDriveEnabled[0] = Profile.LinearDrive.XDrive.bEnableVelocityDrive;
 		ConstraintSettings.bLinearVelocityDriveEnabled[1] = Profile.LinearDrive.YDrive.bEnableVelocityDrive;
 		ConstraintSettings.bLinearVelocityDriveEnabled[2] = Profile.LinearDrive.ZDrive.bEnableVelocityDrive;
-		ConstraintSettings.LinearDriveStiffness = Profile.LinearDrive.XDrive.Stiffness;
-		ConstraintSettings.LinearDriveDamping = Profile.LinearDrive.XDrive.Damping;
+		ConstraintSettings.LinearDriveStiffness = ChaosImmediate_LinearDriveStiffnessScale * Profile.LinearDrive.XDrive.Stiffness;
+		ConstraintSettings.LinearDriveDamping = ChaosImmediate_LinearDriveDampingScale * Profile.LinearDrive.XDrive.Damping;
 		ConstraintSettings.LinearDriveForceMode = EJointForceMode::Acceleration;
 
 		ConstraintSettings.AngularDrivePositionTarget = FQuat(Profile.AngularDrive.OrientationTarget);
@@ -120,8 +124,8 @@ namespace ImmediatePhysics_Chaos
 			ConstraintSettings.bAngularSwingPositionDriveEnabled = Profile.AngularDrive.SwingDrive.bEnablePositionDrive;
 			ConstraintSettings.bAngularSwingVelocityDriveEnabled = Profile.AngularDrive.SwingDrive.bEnableVelocityDrive;
 		}
-		ConstraintSettings.AngularDriveStiffness = ChaosImmediate_DriveStiffnessScale * Profile.AngularDrive.TwistDrive.Stiffness;
-		ConstraintSettings.AngularDriveDamping = ChaosImmediate_DriveDampingScale * Profile.AngularDrive.TwistDrive.Damping;
+		ConstraintSettings.AngularDriveStiffness = ChaosImmediate_AngularDriveStiffnessScale * Profile.AngularDrive.TwistDrive.Stiffness;
+		ConstraintSettings.AngularDriveDamping = ChaosImmediate_AngularDriveDampingScale * Profile.AngularDrive.TwistDrive.Damping;
 		ConstraintSettings.AngularDriveForceMode = EJointForceMode::Acceleration;
 
 		// UE Disables Soft Limits when the Limit is less than some threshold. This is not necessary in Chaos but for now we also do it for parity's sake (See FLinearConstraint::UpdateLinearLimit_AssumesLocked).
diff --git a/Engine/Source/Runtime/Engine/Private/PhysicsEngine/ImmediatePhysics/ImmediatePhysicsChaos/ImmediatePhysicsSimulation_Chaos.cpp b/Engine/Source/Runtime/Engine/Private/PhysicsEngine/ImmediatePhysics/ImmediatePhysicsChaos/ImmediatePhysicsSimulation_Chaos.cpp
--- a/Engine/Source/Runtime/Engine/Private/PhysicsEngine/ImmediatePhysics/ImmediatePhysicsChaos/ImmediatePhysicsSimulation_Chaos.cpp
+++ b/Engine/Source/Runtime/Engine/Private/PhysicsEngine/ImmediatePhysics/ImmediatePhysicsChaos/ImmediatePhysicsSimulation_Chaos.cpp
@@ -75,7 +75,7 @@ float ChaosImmediate_Joint_AngularDriveStiffness = 0.0f;
 float ChaosImmediate_Joint_AngularDriveDamping = 0.0f;
 float ChaosImmediate_Joint_MinParentMassRatio = 0.2f;
 float ChaosImmediate_Joint_MaxInertiaRatio = 5.0f;
-float ChaosImmediate_Joint_AngularPositionCorrection = 1.0f;
+float ChaosImmediate_Joint_AngularPositionCorrection = 0.0f;
 FAutoConsoleVariableRef CVarChaosImmPhysJointPairIterations(TEXT("p.Chaos.ImmPhys.Joint.PairIterations"), ChaosImmediate_Joint_PairIterations, TEXT("Override joint pair iterations (if >= 0)"));
 FAutoConsoleVariableRef CVarChaosImmPhysJointPushOutPairIterations(TEXT("p.Chaos.ImmPhys.Joint.PushOutPairIterations"), ChaosImmediate_Joint_PushOutPairIterations, TEXT("Override joint push-out pair iterations (if >= 0)"));
 FAutoConsoleVariableRef CVarChaosImmPhysJointSwingTwistAngleTolerance(TEXT("p.Chaos.ImmPhys.Joint.SwingTwistAngleTolerance"), ChaosImmediate_Joint_SwingTwistAngleTolerance, TEXT("SwingTwistAngleTolerance."));
diff --git a/Engine/Source/Runtime/Engine/Private/PhysicsEngine/PhysicsAsset.cpp b/Engine/Source/Runtime/Engine/Private/PhysicsEngine/PhysicsAsset.cpp
--- a/Engine/Source/Runtime/Engine/Private/PhysicsEngine/PhysicsAsset.cpp
+++ b/Engine/Source/Runtime/Engine/Private/PhysicsEngine/PhysicsAsset.cpp
@@ -23,7 +23,7 @@
 
 FSolverIterations::FSolverIterations()
 	: SolverIterations(3)
-	, JointIterations(1)
+	, JointIterations(2)
 	, CollisionIterations(1)
 	, SolverPushOutIterations(1)
 	, JointPushOutIterations(1)
diff --git a/Engine/Source/Runtime/Experimental/Chaos/Private/Chaos/CollisionResolution.cpp b/Engine/Source/Runtime/Experimental/Chaos/Private/Chaos/CollisionResolution.cpp
--- a/Engine/Source/Runtime/Experimental/Chaos/Private/Chaos/CollisionResolution.cpp
+++ b/Engine/Source/Runtime/Experimental/Chaos/Private/Chaos/CollisionResolution.cpp
@@ -30,6 +30,7 @@ DECLARE_CYCLE_STAT(TEXT("Collisions::GJK"), STAT_Collisions_GJK, STATGROUP_Chaos
 #endif
 
 //#pragma optimize("", off)
+//PRAGMA_DISABLE_OPTIMIZATION_ACTUAL
 
 bool Chaos_Collision_UseManifolds = true;
 FAutoConsoleVariableRef CVarChaosCollisionUseManifolds(TEXT("p.Chaos.Collision.UseManifolds"), Chaos_Collision_UseManifolds, TEXT("Enable/Disable use of manifoldes in collision."));
@@ -42,7 +43,7 @@ float Chaos_Collision_ManifoldFaceEpsilon = FMath::Sin(FMath::DegreesToRadians(C
 FConsoleVariableDelegate Chaos_Collision_ManifoldFaceDelegate = FConsoleVariableDelegate::CreateLambda([](IConsoleVariable* CVar) { Chaos_Collision_ManifoldFaceEpsilon = FMath::Sin(FMath::DegreesToRadians(Chaos_Collision_ManifoldFaceAngle)); });
 FAutoConsoleVariableRef CVarChaosCollisionManifoldFaceAngle(TEXT("p.Chaos.Collision.ManifoldFaceAngle"), Chaos_Collision_ManifoldFaceAngle, TEXT("Angle above which a face is rejected and we switch to point collision"), Chaos_Collision_ManifoldFaceDelegate);
 
-float Chaos_Collision_CapsuleBoxManifoldAngle = 0.0f;
+float Chaos_Collision_CapsuleBoxManifoldAngle = 20.0f;
 float Chaos_Collision_CapsuleBoxManifoldTolerance = FMath::Sin(FMath::DegreesToRadians(Chaos_Collision_CapsuleBoxManifoldAngle));
 FConsoleVariableDelegate Chaos_Collision_CapsuleBoxManifoldDelegate = FConsoleVariableDelegate::CreateLambda([](IConsoleVariable* CVar) { Chaos_Collision_CapsuleBoxManifoldTolerance = FMath::Sin(FMath::DegreesToRadians(Chaos_Collision_CapsuleBoxManifoldAngle)); });
 FAutoConsoleVariableRef CVarChaosCollisionBoxCapsuleManifoldAngle(TEXT("p.Chaos.Collision.CapsuleBoxManifoldAngle"), Chaos_Collision_CapsuleBoxManifoldAngle, TEXT("If a capsule is more than this angle from vertical, do not use a manifold"), Chaos_Collision_CapsuleBoxManifoldDelegate);
@@ -1097,7 +1098,7 @@ namespace Chaos
 
 					bool bAllowManifold = Chaos_Collision_UseManifolds;
 
-					if (Chaos_Collision_CapsuleBoxManifoldTolerance > KINDA_SMALL_NUMBER)
+					if (bAllowManifold && (Chaos_Collision_CapsuleBoxManifoldTolerance > KINDA_SMALL_NUMBER))
 					{
 						// @todo(ccaulfield): weak sauce - fix capsule-box manifolds.
 						// HACK: Disable manifolds for "horizontal" capsules. Manifolds don't work well when joints are pulling boxes down
diff --git a/Engine/Source/Runtime/Experimental/Chaos/Private/Chaos/Joint/PBDJointSolverGaussSeidel.cpp b/Engine/Source/Runtime/Experimental/Chaos/Private/Chaos/Joint/PBDJointSolverGaussSeidel.cpp
--- a/Engine/Source/Runtime/Experimental/Chaos/Private/Chaos/Joint/PBDJointSolverGaussSeidel.cpp
+++ b/Engine/Source/Runtime/Experimental/Chaos/Private/Chaos/Joint/PBDJointSolverGaussSeidel.cpp
@@ -10,6 +10,7 @@
 #endif
 
 //#pragma optimize("", off)
+//PRAGMA_DISABLE_OPTIMIZATION_ACTUAL
 
 #if !INTEL_ISPC
 const bool bChaos_Joint_ISPC_Enabled = false;
@@ -901,17 +902,17 @@ namespace Chaos
 		// This significantly improves angular stiffness at lower iterations, although the same effect
 		// is achieved by increasing iterations.
 		// @todo(ccaulfield): this position correction needs to have components in direction of inactive position constraints removed
-		if (AngularPositionCorrection > 0)
-		{
-			const FVec3 PrevX0 = Xs[0];
-			const FVec3 PrevX1 = Xs[1];
-			UpdateDerivedState();
-
-			const FVec3 DX = (Xs[1] - Xs[0]) - (PrevX1 - PrevX0);
-			const FVec3 DP0 = DX * (InvMs[0] / (InvMs[0] + InvMs[1]));
-			const FVec3 DP1 = DX * (-InvMs[1] / (InvMs[0] + InvMs[1]));
-			ApplyPositionDelta(Stiffness * AngularPositionCorrection, DP0, DP1);
-		}
+		//if (AngularPositionCorrection > 0)
+		//{
+		//	const FVec3 PrevX0 = Xs[0];
+		//	const FVec3 PrevX1 = Xs[1];
+		//	UpdateDerivedState();
+
+		//	const FVec3 DX = (Xs[1] - Xs[0]) - (PrevX1 - PrevX0);
+		//	const FVec3 DP0 = DX * (InvMs[0] / (InvMs[0] + InvMs[1]));
+		//	const FVec3 DP1 = DX * (-InvMs[1] / (InvMs[0] + InvMs[1]));
+		//	ApplyPositionDelta(Stiffness * AngularPositionCorrection, DP0, DP1);
+		//}
 
 		UpdateDerivedState();
 	}
@@ -948,16 +949,16 @@ namespace Chaos
 
 			// If we are correcting the position, we need to adjust the constraint effective mass using parallel axis theorem
 			// @todo(ccaulfield): the IMs here are constant per constraint...pre-build and cache it
-			if (AngularPositionCorrection > 0)
-			{
-				const FVec3 LinearAxis1 = FVec3::CrossProduct(Xs[DIndex] - Ps[DIndex], Axis);
-				const FReal LinearAxisLen1Sq = LinearAxis1.SizeSquared();
-				if (LinearAxisLen1Sq > KINDA_SMALL_NUMBER)
-				{
-					const FReal IM1 = InvMs[DIndex] / LinearAxisLen1Sq;
-					II1 = II1 * IM1 / (II1 + IM1);
-				}
-			}
+			//if (AngularPositionCorrection > 0)
+			//{
+			//	const FVec3 LinearAxis1 = FVec3::CrossProduct(Xs[DIndex] - Ps[DIndex], Axis);
+			//	const FReal LinearAxisLen1Sq = LinearAxis1.SizeSquared();
+			//	if (LinearAxisLen1Sq > KINDA_SMALL_NUMBER)
+			//	{
+			//		const FReal IM1 = InvMs[DIndex] / LinearAxisLen1Sq;
+			//		II1 = II1 * IM1 / (II1 + IM1);
+			//	}
+			//}
 			const FReal II = II1;
 
 			// Damping angular velocity
@@ -985,15 +986,15 @@ namespace Chaos
 			// This significantly improves angular stiffness at lower iterations, although the same effect
 			// is achieved by increasing iterations.
 			// @todo(ccaulfield): this position correction needs to have components in direction of inactive position constraints removed
-			if (AngularPositionCorrection > 0)
-			{
-				const FVec3 PrevX0 = Xs[KIndex];
-				const FVec3 PrevX1 = Xs[DIndex];
-				UpdateDerivedState(DIndex);
+			//if (AngularPositionCorrection > 0)
+			//{
+			//	const FVec3 PrevX0 = Xs[KIndex];
+			//	const FVec3 PrevX1 = Xs[DIndex];
+			//	UpdateDerivedState(DIndex);
 
-				const FVec3 DP1 = (PrevX1 - PrevX0) - (Xs[DIndex] - Xs[KIndex]);
-				ApplyPositionDelta(DIndex, AngularPositionCorrection, DP1);
-			}
+			//	const FVec3 DP1 = (PrevX1 - PrevX0) - (Xs[DIndex] - Xs[KIndex]);
+			//	ApplyPositionDelta(DIndex, AngularPositionCorrection, DP1);
+			//}
 
 			UpdateDerivedState(DIndex);
 		}
@@ -1032,23 +1033,23 @@ namespace Chaos
 
 			// If we are correcting the position, we need to adjust the constraint effective mass using parallel axis theorem
 			// @todo(ccaulfield): the IMs here are constant per constraint...pre-build and cache it
-			if (AngularPositionCorrection > 0)
-			{
-				const FVec3 LinearAxis0 = FVec3::CrossProduct(Xs[0] - Ps[0], Axis);
-				const FReal LinearAxisLen0Sq = LinearAxis0.SizeSquared();
-				if (LinearAxisLen0Sq > KINDA_SMALL_NUMBER)
-				{
-					const FReal IM0 = InvMs[0] / LinearAxisLen0Sq;
-					II0 = II0 * IM0 / (II0 + IM0);
-				}
-				const FVec3 LinearAxis1 = FVec3::CrossProduct(Xs[1] - Ps[1], Axis);
-				const FReal LinearAxisLen1Sq = LinearAxis1.SizeSquared();
-				if (LinearAxisLen1Sq > KINDA_SMALL_NUMBER)
-				{
-					const FReal IM1 = InvMs[1] / LinearAxisLen1Sq;
-					II1 = II1 * IM1 / (II1 + IM1);
-				}
-			}
+			//if (AngularPositionCorrection > 0)
+			//{
+			//	const FVec3 LinearAxis0 = FVec3::CrossProduct(Xs[0] - Ps[0], Axis);
+			//	const FReal LinearAxisLen0Sq = LinearAxis0.SizeSquared();
+			//	if (LinearAxisLen0Sq > KINDA_SMALL_NUMBER)
+			//	{
+			//		const FReal IM0 = InvMs[0] / LinearAxisLen0Sq;
+			//		II0 = II0 * IM0 / (II0 + IM0);
+			//	}
+			//	const FVec3 LinearAxis1 = FVec3::CrossProduct(Xs[1] - Ps[1], Axis);
+			//	const FReal LinearAxisLen1Sq = LinearAxis1.SizeSquared();
+			//	if (LinearAxisLen1Sq > KINDA_SMALL_NUMBER)
+			//	{
+			//		const FReal IM1 = InvMs[1] / LinearAxisLen1Sq;
+			//		II1 = II1 * IM1 / (II1 + IM1);
+			//	}
+			//}
 			const FReal II = (II0 + II1);
 
 			// Damping angular velocity
@@ -1079,17 +1080,17 @@ namespace Chaos
 			// This significantly improves angular stiffness at lower iterations, although the same effect
 			// is achieved by increasing iterations.
 			// @todo(ccaulfield): this position correction needs to have components in direction of inactive position constraints removed
-			if (AngularPositionCorrection > 0)
-			{
-				const FVec3 PrevX0 = Xs[0];
-				const FVec3 PrevX1 = Xs[1];
-				UpdateDerivedState();
-
-				const FVec3 DX = (Xs[1] - Xs[0]) - (PrevX1 - PrevX0);
-				const FVec3 DP0 = DX * (InvMs[0] / (InvMs[0] + InvMs[1]));
-				const FVec3 DP1 = DX * (-InvMs[1] / (InvMs[0] + InvMs[1]));
-				ApplyPositionDelta(AngularPositionCorrection, DP0, DP1);
-			}
+			//if (AngularPositionCorrection > 0)
+			//{
+			//	const FVec3 PrevX0 = Xs[0];
+			//	const FVec3 PrevX1 = Xs[1];
+			//	UpdateDerivedState();
+
+			//	const FVec3 DX = (Xs[1] - Xs[0]) - (PrevX1 - PrevX0);
+			//	const FVec3 DP0 = DX * (InvMs[0] / (InvMs[0] + InvMs[1]));
+			//	const FVec3 DP1 = DX * (-InvMs[1] / (InvMs[0] + InvMs[1]));
+			//	ApplyPositionDelta(AngularPositionCorrection, DP0, DP1);
+			//}
 
 			UpdateDerivedState();
 		}
diff --git a/Engine/Source/Runtime/Experimental/Chaos/Private/Chaos/Joint/PBDJointSolverGaussSeidel.ispc b/Engine/Source/Runtime/Experimental/Chaos/Private/Chaos/Joint/PBDJointSolverGaussSeidel.ispc
--- a/Engine/Source/Runtime/Experimental/Chaos/Private/Chaos/Joint/PBDJointSolverGaussSeidel.ispc
+++ b/Engine/Source/Runtime/Experimental/Chaos/Private/Chaos/Joint/PBDJointSolverGaussSeidel.ispc
@@ -278,16 +278,16 @@ export void ApplyRotationConstraintSoftKD(
 
 	// If we are correcting the position, we need to adjust the constraint effective mass using parallel axis theorem
 	// @todo(ccaulfield): the IMs here are constant per constraint...pre-build and cache it
-	if (M->AngularPositionCorrection > 0)
-	{
-		const uniform FVector LinearAxis1 = VectorCross(M->Xs[DIndex] - M->Ps[DIndex], Axis);
-		const uniform float LinearAxisLen1Sq = VectorSizeSquared(LinearAxis1);
-		if (LinearAxisLen1Sq > KINDA_SMALL_NUMBER)
-		{
-			const uniform float IM1 = M->InvMs[DIndex] / LinearAxisLen1Sq;
-			II1 = II1 * IM1 / (II1 + IM1);
-		}
-	}
+	////if (M->AngularPositionCorrection > 0)
+	//{
+	//	const uniform FVector LinearAxis1 = VectorCross(M->Xs[DIndex] - M->Ps[DIndex], Axis);
+	//	const uniform float LinearAxisLen1Sq = VectorSizeSquared(LinearAxis1);
+	//	if (LinearAxisLen1Sq > KINDA_SMALL_NUMBER)
+	//	{
+	//		const uniform float IM1 = M->InvMs[DIndex] / LinearAxisLen1Sq;
+	//		II1 = II1 * IM1 / (II1 + IM1);
+	//	}
+	//}
 	const uniform float II = II1;
 
 	// Damping angular velocity
@@ -315,15 +315,15 @@ export void ApplyRotationConstraintSoftKD(
 	// This significantly improves angular stiffness at lower iterations, although the same effect
 	// is achieved by increasing iterations.
 	// @todo(ccaulfield): this position correction needs to have components in direction of inactive position constraints removed
-	if (M->AngularPositionCorrection > 0)
-	{
-		const uniform FVector PrevX0 = M->Xs[KIndex];
-		const uniform FVector PrevX1 = M->Xs[DIndex];
-		UpdateDerivedState(M);
-
-		const uniform FVector DP1 = (PrevX1 - PrevX0) - (M->Xs[DIndex] - M->Xs[KIndex]);
-		ApplyPositionDelta(M, DIndex, M->AngularPositionCorrection, DP1);
-	}
+	//if (M->AngularPositionCorrection > 0)
+	//{
+	//	const uniform FVector PrevX0 = M->Xs[KIndex];
+	//	const uniform FVector PrevX1 = M->Xs[DIndex];
+	//	UpdateDerivedState(M);
+	//
+	//	const uniform FVector DP1 = (PrevX1 - PrevX0) - (M->Xs[DIndex] - M->Xs[KIndex]);
+	//	ApplyPositionDelta(M, DIndex, M->AngularPositionCorrection, DP1);
+	//}
 
 	UpdateDerivedState(M, DIndex);
 }
@@ -350,23 +350,23 @@ export void ApplyRotationConstraintSoftDD(
 
 	// If we are correcting the position, we need to adjust the constraint effective mass using parallel axis theorem
 	// @todo(ccaulfield): the IMs here are constant per constraint...pre-build and cache it
-	if (M->AngularPositionCorrection > 0)
-	{
-		const uniform FVector LinearAxis0 = VectorCross(M->Xs[0] - M->Ps[0], Axis);
-		const uniform float LinearAxisLen0Sq = VectorSizeSquared(LinearAxis0);
-		if (LinearAxisLen0Sq > KINDA_SMALL_NUMBER)
-		{
-			const uniform float IM0 = M->InvMs[0] / LinearAxisLen0Sq;
-			II0 = II0 * IM0 / (II0 + IM0);
-		}
-		const uniform FVector LinearAxis1 = VectorCross(M->Xs[1] - M->Ps[1], Axis);
-		const uniform float LinearAxisLen1Sq = VectorSizeSquared(LinearAxis1);
-		if (LinearAxisLen1Sq > KINDA_SMALL_NUMBER)
-		{
-			const uniform float IM1 = M->InvMs[1] / LinearAxisLen1Sq;
-			II1 = II1 * IM1 / (II1 + IM1);
-		}
-	}
+	//if (M->AngularPositionCorrection > 0)
+	//{
+	//	const uniform FVector LinearAxis0 = VectorCross(M->Xs[0] - M->Ps[0], Axis);
+	//	const uniform float LinearAxisLen0Sq = VectorSizeSquared(LinearAxis0);
+	//	if (LinearAxisLen0Sq > KINDA_SMALL_NUMBER)
+	//	{
+	//		const uniform float IM0 = M->InvMs[0] / LinearAxisLen0Sq;
+	//		II0 = II0 * IM0 / (II0 + IM0);
+	//	}
+	//	const uniform FVector LinearAxis1 = VectorCross(M->Xs[1] - M->Ps[1], Axis);
+	//	const uniform float LinearAxisLen1Sq = VectorSizeSquared(LinearAxis1);
+	//	if (LinearAxisLen1Sq > KINDA_SMALL_NUMBER)
+	//	{
+	//		const uniform float IM1 = M->InvMs[1] / LinearAxisLen1Sq;
+	//		II1 = II1 * IM1 / (II1 + IM1);
+	//	}
+	//}
 	const uniform float II = (II0 + II1);
 
 	// Damping angular velocity
@@ -397,17 +397,17 @@ export void ApplyRotationConstraintSoftDD(
 	// This significantly improves angular stiffness at lower iterations, although the same effect
 	// is achieved by increasing iterations.
 	// @todo(ccaulfield): this position correction needs to have components in direction of inactive position constraints removed
-	if (M->AngularPositionCorrection > 0)
-	{
-		const uniform FVector PrevX0 = M->Xs[0];
-		const uniform FVector PrevX1 = M->Xs[1];
-		UpdateDerivedState(M);
-
-		const uniform FVector DX = (M->Xs[1] - M->Xs[0]) - (PrevX1 - PrevX0);
-		const uniform FVector DP0 = DX * (M->InvMs[0] / (M->InvMs[0] + M->InvMs[1]));
-		const uniform FVector DP1 = DX * (-M->InvMs[1] / (M->InvMs[0] + M->InvMs[1]));
-		ApplyPositionDelta(M, M->AngularPositionCorrection, DP0, DP1);
-	}
+	//if (M->AngularPositionCorrection > 0)
+	//{
+	//	const uniform FVector PrevX0 = M->Xs[0];
+	//	const uniform FVector PrevX1 = M->Xs[1];
+	//	UpdateDerivedState(M);
+	//
+	//	const uniform FVector DX = (M->Xs[1] - M->Xs[0]) - (PrevX1 - PrevX0);
+	//	const uniform FVector DP0 = DX * (M->InvMs[0] / (M->InvMs[0] + M->InvMs[1]));
+	//	const uniform FVector DP1 = DX * (-M->InvMs[1] / (M->InvMs[0] + M->InvMs[1]));
+	//	ApplyPositionDelta(M, M->AngularPositionCorrection, DP0, DP1);
+	//}
 
 	UpdateDerivedState(M);
 }
diff --git a/Engine/Source/Runtime/Experimental/Chaos/Private/Chaos/PBDCollisionConstraints.cpp b/Engine/Source/Runtime/Experimental/Chaos/Private/Chaos/PBDCollisionConstraints.cpp
--- a/Engine/Source/Runtime/Experimental/Chaos/Private/Chaos/PBDCollisionConstraints.cpp
+++ b/Engine/Source/Runtime/Experimental/Chaos/Private/Chaos/PBDCollisionConstraints.cpp
@@ -30,6 +30,7 @@
 #endif
 
 //#pragma optimize("", off)
+//PRAGMA_DISABLE_OPTIMIZATION_ACTUAL
 
 namespace Chaos
 {
@@ -413,12 +414,12 @@ namespace Chaos
 
 			for (FPointContactConstraint& Contact : PointConstraints)
 			{
-				Collisions::Apply(Contact, IterationParameters, ParticleParameters);
+				Collisions::ApplySinglePoint(Contact, IterationParameters, ParticleParameters);
 			}
 
 			for (FMultiPointContactConstraint& Contact : IterativeConstraints)
 			{
-				Collisions::Apply(Contact, IterationParameters, ParticleParameters);
+				Collisions::ApplyMultiPoint(Contact, IterationParameters, ParticleParameters);
 			}
 		}
 
@@ -442,12 +443,12 @@ namespace Chaos
 
 			for (FPointContactConstraint& Contact : PointConstraints)
 			{
-				Collisions::ApplyPushOut(Contact, TempStatic, IterationParameters, ParticleParameters);
+				Collisions::ApplyPushOutSinglePoint(Contact, TempStatic, IterationParameters, ParticleParameters);
 			}
 
 			for (FMultiPointContactConstraint& Contact : IterativeConstraints)
 			{
-				Collisions::ApplyPushOut(Contact, TempStatic, IterationParameters, ParticleParameters);
+				Collisions::ApplyPushOutMultiPoint(Contact, TempStatic, IterationParameters, ParticleParameters);
 			}
 		}
 
diff --git a/Engine/Source/Runtime/Experimental/Chaos/Private/Chaos/PBDCollisionConstraintsContact.cpp b/Engine/Source/Runtime/Experimental/Chaos/Private/Chaos/PBDCollisionConstraintsContact.cpp
--- a/Engine/Source/Runtime/Experimental/Chaos/Private/Chaos/PBDCollisionConstraintsContact.cpp
+++ b/Engine/Source/Runtime/Experimental/Chaos/Private/Chaos/PBDCollisionConstraintsContact.cpp
@@ -262,14 +262,14 @@ namespace Chaos
 			}
 		}
 
-		void Apply(FRigidBodyPointContactConstraint& Constraint, const FContactIterationParameters & IterationParameters, const FContactParticleParameters & ParticleParameters)
+		void ApplySinglePoint(FRigidBodyPointContactConstraint& Constraint, const FContactIterationParameters & IterationParameters, const FContactParticleParameters & ParticleParameters)
 		{
-			ApplyImpl(Constraint, IterationParameters, ParticleParameters);
+			ApplyImpl<FRigidBodyPointContactConstraint>(Constraint, IterationParameters, ParticleParameters);
 		}
 
-		void Apply(FRigidBodyMultiPointContactConstraint& Constraint, const FContactIterationParameters & IterationParameters, const FContactParticleParameters & ParticleParameters)
+		void ApplyMultiPoint(FRigidBodyMultiPointContactConstraint& Constraint, const FContactIterationParameters & IterationParameters, const FContactParticleParameters & ParticleParameters)
 		{
-			ApplyImpl(Constraint, IterationParameters, ParticleParameters);
+			ApplyImpl<FRigidBodyMultiPointContactConstraint>(Constraint, IterationParameters, ParticleParameters);
 		}
 
 
@@ -407,24 +407,24 @@ namespace Chaos
 		{
 			if (Constraint.GetType() == FCollisionConstraintBase::FType::SinglePoint)
 			{
-				ApplyPushOutImpl(*Constraint.As<FRigidBodyPointContactConstraint>(), IsTemporarilyStatic, IterationParameters, ParticleParameters);
+				ApplyPushOutImpl<FRigidBodyPointContactConstraint>(*Constraint.As<FRigidBodyPointContactConstraint>(), IsTemporarilyStatic, IterationParameters, ParticleParameters);
 			}
 			else if (Constraint.GetType() == FCollisionConstraintBase::FType::MultiPoint)
 			{
-				ApplyPushOutImpl(*Constraint.As<FRigidBodyMultiPointContactConstraint>(), IsTemporarilyStatic, IterationParameters, ParticleParameters);
+				ApplyPushOutImpl<FRigidBodyMultiPointContactConstraint>(*Constraint.As<FRigidBodyMultiPointContactConstraint>(), IsTemporarilyStatic, IterationParameters, ParticleParameters);
 			}
 		}
 
-		void ApplyPushOut(FRigidBodyPointContactConstraint& Constraint, const TSet<const TGeometryParticleHandle<FReal, 3>*>& IsTemporarilyStatic,
+		void ApplyPushOutSinglePoint(FRigidBodyPointContactConstraint& Constraint, const TSet<const TGeometryParticleHandle<FReal, 3>*>& IsTemporarilyStatic,
 			const FContactIterationParameters & IterationParameters, const FContactParticleParameters & ParticleParameters)
 		{
-			ApplyPushOutImpl(Constraint, IsTemporarilyStatic, IterationParameters, ParticleParameters);
+			ApplyPushOutImpl<FRigidBodyPointContactConstraint>(Constraint, IsTemporarilyStatic, IterationParameters, ParticleParameters);
 		}
 
-		void ApplyPushOut(FRigidBodyMultiPointContactConstraint& Constraint, const TSet<const TGeometryParticleHandle<FReal, 3>*>& IsTemporarilyStatic,
+		void ApplyPushOutMultiPoint(FRigidBodyMultiPointContactConstraint& Constraint, const TSet<const TGeometryParticleHandle<FReal, 3>*>& IsTemporarilyStatic,
 			const FContactIterationParameters & IterationParameters, const FContactParticleParameters & ParticleParameters)
 		{
-			ApplyPushOutImpl(Constraint, IsTemporarilyStatic, IterationParameters, ParticleParameters);
+			ApplyPushOutImpl<FRigidBodyMultiPointContactConstraint>(Constraint, IsTemporarilyStatic, IterationParameters, ParticleParameters);
 		}
 
 	} // Collisions
diff --git a/Engine/Source/Runtime/Experimental/Chaos/Public/Chaos/PBDCollisionConstraintsContact.h b/Engine/Source/Runtime/Experimental/Chaos/Public/Chaos/PBDCollisionConstraintsContact.h
--- a/Engine/Source/Runtime/Experimental/Chaos/Public/Chaos/PBDCollisionConstraintsContact.h
+++ b/Engine/Source/Runtime/Experimental/Chaos/Public/Chaos/PBDCollisionConstraintsContact.h
@@ -40,14 +40,14 @@ namespace Chaos
 		extern void Update(FRigidBodyMultiPointContactConstraint& Constraint, const FReal CullDistance);
 
 		extern void Apply(FCollisionConstraintBase& Constraint, const FContactIterationParameters& IterationParameters, const FContactParticleParameters& ParticleParameters);
-		extern void Apply(FRigidBodyPointContactConstraint& Constraint, const FContactIterationParameters& IterationParameters, const FContactParticleParameters& ParticleParameters);
-		extern void Apply(FRigidBodyMultiPointContactConstraint& Constraint, const FContactIterationParameters& IterationParameters, const FContactParticleParameters& ParticleParameters);
+		extern void ApplySinglePoint(FRigidBodyPointContactConstraint& Constraint, const FContactIterationParameters& IterationParameters, const FContactParticleParameters& ParticleParameters);
+		extern void ApplyMultiPoint(FRigidBodyMultiPointContactConstraint& Constraint, const FContactIterationParameters& IterationParameters, const FContactParticleParameters& ParticleParameters);
 
-		void ApplyPushOut(FCollisionConstraintBase& Constraint, const TSet<const TGeometryParticleHandle<FReal, 3>*>& IsTemporarilyStatic,
+		extern void ApplyPushOut(FCollisionConstraintBase& Constraint, const TSet<const TGeometryParticleHandle<FReal, 3>*>& IsTemporarilyStatic,
 			const FContactIterationParameters& IterationParameters, const FContactParticleParameters& ParticleParameters);
-		void ApplyPushOut(FRigidBodyPointContactConstraint& Constraint, const TSet<const TGeometryParticleHandle<FReal, 3>*>& IsTemporarilyStatic,
+		extern void ApplyPushOutSinglePoint(FRigidBodyPointContactConstraint& Constraint, const TSet<const TGeometryParticleHandle<FReal, 3>*>& IsTemporarilyStatic,
 			const FContactIterationParameters& IterationParameters, const FContactParticleParameters& ParticleParameters);
-		void ApplyPushOut(FRigidBodyPointContactConstraint& Constraint, const TSet<const TGeometryParticleHandle<FReal, 3>*>& IsTemporarilyStatic,
+		extern void ApplyPushOutMultiPoint(FRigidBodyMultiPointContactConstraint& Constraint, const TSet<const TGeometryParticleHandle<FReal, 3>*>& IsTemporarilyStatic,
 			const FContactIterationParameters& IterationParameters, const FContactParticleParameters& ParticleParameters);
 
 
