Chaos - clean up Union contact generation to eliminate allocations and extra transforms

diff --git a/Engine/Source/Runtime/Experimental/Chaos/Private/Chaos/CollisionResolution.cpp b/Engine/Source/Runtime/Experimental/Chaos/Private/Chaos/CollisionResolution.cpp
--- a/Engine/Source/Runtime/Experimental/Chaos/Private/Chaos/CollisionResolution.cpp
+++ b/Engine/Source/Runtime/Experimental/Chaos/Private/Chaos/CollisionResolution.cpp
@@ -93,7 +93,7 @@ namespace Chaos
 			if (ensure(GJKPenetration<true>(A, B, BToATM, Penetration, ClosestA, ClosestB, Normal, (T)0, InitialDir, (T)0, &NumIterations)))
 			{
 				Contact.Location = ATM.TransformPosition(ClosestA);
-				Contact.Normal = -ATM.TransformVector(Normal);
+				Contact.Normal = -ATM.TransformVectorNoScale(Normal);
 				Contact.Phi = -Penetration;
 			}
 
@@ -119,7 +119,7 @@ namespace Chaos
 				{
 					Contact.Phi = Penetration;
 					Contact.Location = BTransform.TransformPosition(Location);
-					Contact.Normal = BTransform.TransformVector(Normal);
+					Contact.Normal = BTransform.TransformVectorNoScale(Normal);
 				}
 			}
 			else if (const TImplicitObjectInstanced<GeometryA>* InstancedConvexImplicit = A.template GetObject<const TImplicitObjectInstanced<GeometryA> >())
@@ -130,7 +130,7 @@ namespace Chaos
 					{
 						Contact.Phi = Penetration;
 						Contact.Location = BTransform.TransformPosition(Location);
-						Contact.Normal = BTransform.TransformVector(Normal);
+						Contact.Normal = BTransform.TransformVectorNoScale(Normal);
 					}
 				}
 			}
@@ -140,7 +140,7 @@ namespace Chaos
 				{
 					Contact.Phi = Penetration;
 					Contact.Location = BTransform.TransformPosition(Location);
-					Contact.Normal = BTransform.TransformVector(Normal);
+					Contact.Normal = BTransform.TransformVectorNoScale(Normal);
 				}
 			}
 
@@ -221,7 +221,7 @@ namespace Chaos
 			bool bNewManifold = (FaceIndex != Constraint.GetManifoldPlaneFaceIndex()) || (Constraint.NumManifoldPoints() == 0);
 			if (bNewManifold)
 			{
-				const FVec3 PlaneNormal = Transform1.InverseTransformVector(ContactPoint.Normal);
+				const FVec3 PlaneNormal = Transform1.InverseTransformVectorNoScale(ContactPoint.Normal);
 				const FVec3 PlanePos = Transform1.InverseTransformPosition(ContactPoint.Location - ContactPoint.Phi*ContactPoint.Normal);
 				Constraint.SetManifoldPlane(1, FaceIndex, PlaneNormal, PlanePos);
 
@@ -262,7 +262,7 @@ namespace Chaos
 			if (!ContactPoint.Normal.Equals(Constraint.GetManifoldPlaneNormal()) || !Constraint.NumManifoldPoints())
 			{
 				Constraint.ResetManifoldPoints();
-				FVec3 PlaneNormal = Transform1.InverseTransformVector(ContactPoint.Normal);
+				FVec3 PlaneNormal = Transform1.InverseTransformVectorNoScale(ContactPoint.Normal);
 				FVec3 PlanePosition = Transform1.InverseTransformPosition(ContactPoint.Location - ContactPoint.Phi*ContactPoint.Normal);
 				Constraint.SetManifoldPlane(1, INDEX_NONE, PlaneNormal, PlanePosition);
 			}
@@ -442,7 +442,7 @@ namespace Chaos
 						NumConstraints = 0;
 					}
 					Contact.Phi = NewPhi;
-					Contact.Normal = PlaneTransform.TransformVector(Normal);
+					Contact.Normal = PlaneTransform.TransformVectorNoScale(Normal);
 					Contact.Location = PlaneTransform.TransformPosition(Corners[i]);
 					PotentialConstraints[NumConstraints++] = Contact.Location;
 				}
@@ -915,7 +915,7 @@ namespace Chaos
 
 			// Cache the closest point so we don't need to re-iterate over the manifold on the first iteration.
 			Constraint.Manifold.Location = BoxTM.TransformPosition(BoxClosestBoxSpace);
-			Constraint.Manifold.Normal = BoxTM.TransformVector(NormalBoxSpace);
+			Constraint.Manifold.Normal = BoxTM.TransformVectorNoScale(NormalBoxSpace);
 			Constraint.Manifold.Phi = -Penetration;
 
 			// Find the box feature that the near point is on
@@ -944,7 +944,7 @@ namespace Chaos
 				return;
 			}
 
-			FVec3 CapsuleAxis = CapsuleToBoxTM.TransformVector(Capsule.GetAxis()); // Box space capsule axis
+			FVec3 CapsuleAxis = CapsuleToBoxTM.TransformVectorNoScale(Capsule.GetAxis()); // Box space capsule axis
 			FReal CapsuleAxisNormal = FVec3::DotProduct(CapsuleAxis, NormalBoxSpace);
 			const bool bIsBoxFaceContact = (NumNonZeroNormalComponents == 1);
 			const bool bIsBoxEdgeContact = (NumNonZeroNormalComponents == 2);
@@ -1027,7 +1027,7 @@ namespace Chaos
 			if (bIsBoxEdgeContact && bIsCapsuleEdgeContact)
 			{
 				// Initialize the manifold with the plane and point on the capsule (move to capsule space)
-				FVec3 NormalCapsuleSpace = CapsuleToBoxTM.InverseTransformVector(NormalBoxSpace);
+				FVec3 NormalCapsuleSpace = CapsuleToBoxTM.InverseTransformVectorNoScale(NormalBoxSpace);
 				FVec3 CapsuleClosestCapsuleSpace = CapsuleToBoxTM.InverseTransformPosition(CapsuleClosestBoxSpace);
 				Constraint.SetManifoldPlane(0, INDEX_NONE, -NormalCapsuleSpace, CapsuleClosestCapsuleSpace);
 				Constraint.AddManifoldPoint(BoxClosestBoxSpace);
@@ -1104,7 +1104,7 @@ namespace Chaos
 						// HACK: Disable manifolds for "horizontal" capsules. Manifolds don't work well when joints are pulling boxes down
 						// (under gravity) when the upper boxes are draped over a horizontal capsule. The box rotations about the manifold
 						// points(line) is too great and we end up with jitter.
-						const FVector CapsuleAxis = Context.SpaceTransform.TransformVector(Transform0.TransformVector(Object0->GetAxis()));
+						const FVector CapsuleAxis = Context.SpaceTransform.TransformVectorNoScale(Transform0.TransformVectorNoScale(Object0->GetAxis()));
 						bAllowManifold = (FMath::Abs(CapsuleAxis.Z) > Chaos_Collision_CapsuleBoxManifoldTolerance);
 					}
 
@@ -1235,7 +1235,7 @@ namespace Chaos
 					if (ContactPoint.Phi > Phi)
 					{
 						ContactPoint.Phi = Phi;
-						ContactPoint.Normal = Transform1.TransformVector(Constraint.GetManifoldPlaneNormal());
+						ContactPoint.Normal = Transform1.TransformVectorNoScale(Constraint.GetManifoldPlaneNormal());
 						ContactPoint.Location = Transform0.TransformPosition(Location);
 					}
 				}
@@ -1416,34 +1416,6 @@ namespace Chaos
 			NewConstraints.TryAdd(CullDistance, Constraint);
 		}
 
-		//
-		//  Union-Union
-		//
-
-		DECLARE_CYCLE_STAT(TEXT("TPBDCollisionConstraints::ConstructUnionUnionConstraints"), STAT_ConstructUnionUnionConstraints, STATGROUP_ChaosWide);
-		template<typename T, int d>
-		void ConstructUnionUnionConstraints(TGeometryParticleHandle<T, d>* Particle0, TGeometryParticleHandle<T, d>* Particle1, const FImplicitObject* Implicit0, const FImplicitObject* Implicit1, const TRigidTransform<T, d>& Transform0, const TRigidTransform<T, d>& Transform1, const T CullDistance, const FCollisionContext& Context, FCollisionConstraintsArray& NewConstraints)
-		{
-			SCOPE_CYCLE_COUNTER(STAT_ConstructUnionUnionConstraints);
-
-			const TArray<Pair<const FImplicitObject*, TRigidTransform<T, d>>> LevelsetShapes = FindRelevantShapes(Implicit0, Transform0, *Implicit1, Transform1, CullDistance);
-
-			for (const Pair<const FImplicitObject*, TRigidTransform<T, d>>& LevelsetObjPair : LevelsetShapes)
-			{
-				const FImplicitObject* LevelsetInnerObj = LevelsetObjPair.First;
-				const TRigidTransform<T, d>& LevelsetInnerObjTM = LevelsetObjPair.Second * Transform1;
-
-				//now find all particle inner objects
-				const TArray<Pair<const FImplicitObject*, TRigidTransform<T, d>>> ParticleShapes = FindRelevantShapes(LevelsetInnerObj, LevelsetInnerObjTM, *Implicit0, Transform0, CullDistance);
-
-				//for each inner obj pair, update constraint
-				for (const Pair<const FImplicitObject*, TRigidTransform<T, d>>& ParticlePair : ParticleShapes)
-				{
-					const FImplicitObject* ParticleInnerObj = ParticlePair.First;
-					ConstructConstraints<T, d>(Particle0, Particle1, ParticleInnerObj, LevelsetInnerObj, Transform0, Transform1, CullDistance, Context, NewConstraints);
-				}
-			}
-		}
 
 		//
 		// Constraint API
@@ -1727,7 +1699,7 @@ namespace Chaos
 			const FRigidTransform3& PointsTransform = (Constraint.GetManifoldPlaneOwnerIndex() == 0) ? Transform1 : Transform0;
 
 			// World-space manifold plane
-			FVec3 PlaneNormal = PlaneTransform.TransformVector(Constraint.GetManifoldPlaneNormal());
+			FVec3 PlaneNormal = PlaneTransform.TransformVectorNoScale(Constraint.GetManifoldPlaneNormal());
 			FVec3 PlanePos = PlaneTransform.TransformPosition(Constraint.GetManifoldPlanePosition());
 
 			// Select the best manifold point
@@ -1845,44 +1817,6 @@ namespace Chaos
 		}
 
 
-		template<typename T, int d>
-		void ConstructPointConstraint(EContactShapesType ShapesType, TGeometryParticleHandle<T, d>* Particle0, TGeometryParticleHandle<T, d>* Particle1, const FImplicitObject* Implicit0, const FImplicitObject* Implicit1, const TRigidTransform<T, d>& Transform0, const TRigidTransform<T, d>& Transform1, const T CullDistance, FCollisionConstraintsArray& NewConstraints)
-		{
-			if (ensure(Implicit0 && Implicit1))
-			{
-				TRigidTransform<T, d> ParticleImplicit0TM = Transform0.GetRelativeTransform(Collisions::GetTransform(Particle0));
-				TRigidTransform<T, d> ParticleImplicit1TM = Transform1.GetRelativeTransform(Collisions::GetTransform(Particle1));
-				if (FRigidBodyPointContactConstraint* Constraint = NewConstraints.AddPointConstraint(Particle0, Implicit0, ParticleImplicit0TM, Particle1, Implicit1, ParticleImplicit1TM, ShapesType))
-				{
-					UpdateConstraintFromGeometryImpl<ECollisionUpdateType::Deepest>(*Constraint, Transform0, Transform1, CullDistance);
-					if (Constraint->GetPhi() >= CullDistance)
-					{
-						NewConstraints.PopPointConstraint();
-					}
-				}
-			}
-		}
-
-
-		template<typename T, int d>
-		void ConstructManifoldConstraint(EContactShapesType ShapesType, TGeometryParticleHandle<T, d>* Particle0, TGeometryParticleHandle<T, d>* Particle1, const FImplicitObject* Implicit0, const FImplicitObject* Implicit1, const TRigidTransform<T, d>& Transform0, const TRigidTransform<T, d>& Transform1, const T CullDistance, const FCollisionContext& Context, FCollisionConstraintsArray& NewConstraints)
-		{
-			if (ensure(Implicit0 && Implicit1))
-			{
-				TRigidTransform<T, d> ParticleImplicit0TM = Transform0.GetRelativeTransform(Collisions::GetTransform(Particle0));
-				TRigidTransform<T, d> ParticleImplicit1TM = Transform1.GetRelativeTransform(Collisions::GetTransform(Particle1));
-				if (FRigidBodyMultiPointContactConstraint* Constraint = NewConstraints.AddMultiPointConstraint(Particle0, Implicit0, ParticleImplicit0TM, Particle1, Implicit1, ParticleImplicit1TM, ShapesType))
-				{
-					UpdateManifold(*Constraint, Transform0, Transform1, CullDistance, Context);
-					if (Constraint->GetPhi() >= CullDistance)
-					{
-						NewConstraints.PopMultiPointConstraint();
-					}
-				}
-			}
-		}
-
-
 		template<typename T, int d>
 		void ConstructConstraints(TGeometryParticleHandle<T, d>* Particle0, TGeometryParticleHandle<T, d>* Particle1, const FImplicitObject* Implicit0, const FImplicitObject* Implicit1, const TRigidTransform<T, d>& Transform0, const TRigidTransform<T, d>& Transform1, const T CullDistance, const FCollisionContext& Context, FCollisionConstraintsArray& NewConstraints)
 		{
@@ -1959,35 +1893,26 @@ namespace Chaos
 			}
 
 			// Handle Unions
-			if (Implicit0OuterType != FImplicitObjectUnion::StaticType() && Implicit1OuterType == FImplicitObjectUnion::StaticType())
+			if (Implicit0OuterType == FImplicitObjectUnion::StaticType())
 			{
-				const TArray<Pair<const FImplicitObject*, TRigidTransform<T, d>>> LevelsetShapes = FindRelevantShapes(Implicit0, Transform0, *Implicit1, Transform1, CullDistance);
-				for (const Pair<const FImplicitObject*, TRigidTransform<T, d>>& LevelsetObjPair : LevelsetShapes)
+				const FImplicitObjectUnion* Union0 = Implicit0->template GetObject<FImplicitObjectUnion>();
+				for (const auto& Child0 : Union0->GetObjects())
 				{
-					const FImplicitObject* Implicit1InnerObj = LevelsetObjPair.First;
-					const TRigidTransform<T, d> Implicit1InnerObjTM = LevelsetObjPair.Second * Transform1;
-					ConstructConstraints(Particle0, Particle1, Implicit0, Implicit1InnerObj, Transform0, Implicit1InnerObjTM, CullDistance, Context, NewConstraints);
+					ConstructConstraints(Particle0, Particle1, Child0.Get(), Implicit1, Transform0, Transform1, CullDistance, Context, NewConstraints);
 				}
 				return;
 			}
-			else if (Implicit0OuterType == FImplicitObjectUnion::StaticType() && Implicit1OuterType != FImplicitObjectUnion::StaticType())
+			if (Implicit1OuterType == FImplicitObjectUnion::StaticType())
 			{
-				// [Note] forces non-unions into particle[0] position
-				const TArray<Pair<const FImplicitObject*, TRigidTransform<T, d>>> LevelsetShapes = FindRelevantShapes(Implicit1, Transform1, *Implicit0, Transform0, CullDistance);
-				for (const Pair<const FImplicitObject*, TRigidTransform<T, d>>& LevelsetObjPair : LevelsetShapes)
+				const FImplicitObjectUnion* Union1 = Implicit1->template GetObject<FImplicitObjectUnion>();
+				for (const auto& Child1 : Union1->GetObjects())
 				{
-					const FImplicitObject* Implicit0InnerObj = LevelsetObjPair.First;
-					const TRigidTransform<T, d> Implicit0InnerObjTM = LevelsetObjPair.Second * Transform0;
-					ConstructConstraints(Particle0, Particle1, Implicit0InnerObj, Implicit1, Implicit0InnerObjTM, Transform1, CullDistance, Context, NewConstraints);
+					ConstructConstraints(Particle0, Particle1, Implicit0, Child1.Get(), Transform0, Transform1, CullDistance, Context, NewConstraints);
 				}
 				return;
 			}
-			else if (Implicit0OuterType == FImplicitObjectUnion::StaticType() && Implicit1OuterType == FImplicitObjectUnion::StaticType())
-			{
-				ConstructUnionUnionConstraints(Particle0, Particle1, Implicit0, Implicit1, Transform0, Transform1, CullDistance, Context, NewConstraints);
-				return;
-			}
 
+			// Check shape pair filtering if enable
 			if (Context.bFilteringEnabled && !DoCollide(Implicit0Type, Particle0->GetImplicitShape(Implicit0), Implicit1Type, Particle1->GetImplicitShape(Implicit1)))
 			{
 				return;
@@ -2054,8 +1979,6 @@ namespace Chaos
 		template void UpdateLevelsetLevelsetManifold<float, 3>(TCollisionConstraintBase<float, 3>&  Constraint, const TRigidTransform<float, 3>& ATM, const TRigidTransform<float, 3>& BTM, const float CullDistance);
 		template void ConstructLevelsetLevelsetConstraints<float, 3>(TGeometryParticleHandle<float, 3>* Particle0, TGeometryParticleHandle<float, 3>* Particle1, const FImplicitObject* Implicit0, const FImplicitObject* Implicit1, const TRigidTransform<float, 3>& Transform0, const TRigidTransform<float, 3>& Transform1, const float CullDistance, FCollisionConstraintsArray& NewConstraints);
 
-		template void ConstructUnionUnionConstraints<float, 3>(TGeometryParticleHandle<float, 3>* Particle0, TGeometryParticleHandle<float, 3>* Particle1, const FImplicitObject* Implicit0, const FImplicitObject* Implicit1, const TRigidTransform<float, 3>& Transform0, const TRigidTransform<float, 3>& Transform1, const float CullDistance, const FCollisionContext& Context, FCollisionConstraintsArray& NewConstraints);
-
 		template void UpdateSingleShotManifold<float, 3>(TRigidBodyMultiPointContactConstraint<float, 3>&  Constraint, const TRigidTransform<float, 3>& Transform0, const TRigidTransform<float, 3>& Transform1, const float CullDistance);
 		template void UpdateIterativeManifold<float, 3>(TRigidBodyMultiPointContactConstraint<float, 3>&  Constraint, const TRigidTransform<float, 3>& Transform0, const TRigidTransform<float, 3>& Transform1, const float CullDistance);
 		template void UpdateManifold<float, 3>(TRigidBodyMultiPointContactConstraint<float, 3>& Constraint, const TRigidTransform<float, 3>& ATM, const TRigidTransform<float, 3>& BTM, const float CullDistance, const FCollisionContext& Context);
diff --git a/Engine/Source/Runtime/Experimental/Chaos/Public/Chaos/CollisionResolution.h b/Engine/Source/Runtime/Experimental/Chaos/Public/Chaos/CollisionResolution.h
--- a/Engine/Source/Runtime/Experimental/Chaos/Public/Chaos/CollisionResolution.h
+++ b/Engine/Source/Runtime/Experimental/Chaos/Public/Chaos/CollisionResolution.h
@@ -235,12 +235,6 @@ namespace Chaos
 		void CHAOS_API ConstructLevelsetLevelsetConstraints(TGeometryParticleHandle<T, d>* Particle0, TGeometryParticleHandle<T, d>* Particle1, const FImplicitObject* Implicit0, const FImplicitObject* Implicit1, const TRigidTransform<T, d>& Transform0, const TRigidTransform<T, d>& Transform1, const T CullDistance, FCollisionConstraintsArray& NewConstraints);
 
 
-		//
-		// Union-Union
-		//
-		template<typename T, int d>
-		void CHAOS_API ConstructUnionUnionConstraints(TGeometryParticleHandle<T, d>* Particle0, TGeometryParticleHandle<T, d>* Particle1, const FImplicitObject* Implicit0, const FImplicitObject* Implicit1, const TRigidTransform<T, d>& Transform0, const TRigidTransform<T, d>& Transform1, const T CullDistance, FCollisionConstraintsArray& NewConstraints);
-
 		//
 		// Constraint API
 		//
diff --git a/Engine/Source/Runtime/Experimental/Chaos/Public/Chaos/CollisionResolutionTypes.h b/Engine/Source/Runtime/Experimental/Chaos/Public/Chaos/CollisionResolutionTypes.h
--- a/Engine/Source/Runtime/Experimental/Chaos/Public/Chaos/CollisionResolutionTypes.h
+++ b/Engine/Source/Runtime/Experimental/Chaos/Public/Chaos/CollisionResolutionTypes.h
@@ -412,11 +412,12 @@ namespace Chaos
 	typedef TPBDCollisionConstraintHandle<float, 3> FPBDCollisionConstraintHandle;
 
 
-	template<int T_MAXCONSTRAINTS>
-	struct TCollisionConstraintsStore
+	struct FCollisionConstraintsArray
 	{
-		TArray<FRigidBodyPointContactConstraint, TFixedAllocator<T_MAXCONSTRAINTS>> SinglePointConstraints;
-		TArray<FRigidBodyMultiPointContactConstraint, TFixedAllocator<T_MAXCONSTRAINTS>> MultiPointConstraints;
+		static const int32 InlineMaxConstraints = 8;
+
+		TArray<FRigidBodyPointContactConstraint, TInlineAllocator<InlineMaxConstraints>> SinglePointConstraints;
+		TArray<FRigidBodyMultiPointContactConstraint, TInlineAllocator<InlineMaxConstraints>> MultiPointConstraints;
 
 		int32 Num() const { return SinglePointConstraints.Num() + MultiPointConstraints.Num(); }
 
@@ -426,48 +427,9 @@ namespace Chaos
 			MultiPointConstraints.Empty();
 		}
 
-		FRigidBodyPointContactConstraint* AddPointConstraint(
-			TGeometryParticleHandle<FReal, 3>* Particle0, const FImplicitObject* Implicit0, const FRigidTransform3& Transform0,
-			TGeometryParticleHandle<FReal, 3>* Particle1, const FImplicitObject* Implicit1, const FRigidTransform3& Transform1,
-			EContactShapesType ShapesType)
-		{
-			if (SinglePointConstraints.Num() < T_MAXCONSTRAINTS)
-			{
-				int32 ConstraintIndex = SinglePointConstraints.Emplace(FRigidBodyPointContactConstraint(Particle0, Implicit0, Transform0, Particle1, Implicit1, Transform1, ShapesType));
-				return &SinglePointConstraints[ConstraintIndex];
-			}
-			return nullptr;
-		}
-
-		void PopPointConstraint()
-		{
-			SinglePointConstraints.Pop(false);
-		}
-
-		FRigidBodyMultiPointContactConstraint* AddMultiPointConstraint(
-			TGeometryParticleHandle<FReal, 3>* Particle0, const FImplicitObject* Implicit0, const FRigidTransform3& Transform0,
-			TGeometryParticleHandle<FReal, 3>* Particle1, const FImplicitObject* Implicit1, const FRigidTransform3& Transform1,
-			EContactShapesType ShapesType)
-		{
-			if (MultiPointConstraints.Num() < T_MAXCONSTRAINTS)
-			{
-				int32 ConstraintIndex = MultiPointConstraints.Emplace(FRigidBodyMultiPointContactConstraint(Particle0, Implicit0, Transform0, Particle1, Implicit1, Transform1, ShapesType));
-				return &MultiPointConstraints[ConstraintIndex];
-			}
-			return nullptr;
-		}
-
-		void PopMultiPointConstraint()
-		{
-			MultiPointConstraints.Pop();
-		}
-
-
-
-		// Get rid of these
 		TRigidBodyPointContactConstraint<FReal, 3>* TryAdd(FReal MaxPhi, const TRigidBodyPointContactConstraint<FReal, 3>& C)
 		{
-			if ((SinglePointConstraints.Num() < T_MAXCONSTRAINTS) && (C.GetPhi() < MaxPhi))
+			if (C.GetPhi() < MaxPhi)
 			{
 				int32 ConstraintIndex = SinglePointConstraints.Add(C);
 				return &SinglePointConstraints[ConstraintIndex];
@@ -477,7 +439,7 @@ namespace Chaos
 
 		TRigidBodyMultiPointContactConstraint<FReal, 3>* TryAdd(FReal MaxPhi, const TRigidBodyMultiPointContactConstraint<FReal, 3>& C)
 		{
-			if ((MultiPointConstraints.Num() < T_MAXCONSTRAINTS) && (C.GetPhi() < MaxPhi))
+			if (C.GetPhi() < MaxPhi)
 			{
 				int32 ConstraintIndex = MultiPointConstraints.Add(C);
 				return &MultiPointConstraints[ConstraintIndex];
@@ -486,8 +448,6 @@ namespace Chaos
 		}
 	};
 
-	using FCollisionConstraintsArray = TCollisionConstraintsStore<8>;
-
 
 #if PLATFORM_MAC || PLATFORM_LINUX
 	extern template class CHAOS_API TCollisionContact<float, 3>;
