Fix for convex support failing with zero particles leading to crashes in GJK.
Planar shapes now turn into thin prisms when running through the convex path
Triangles in the convex path now work correctly (previously they generated halfspaces)

diff --git a/Engine/Source/Runtime/Engine/Private/PhysicsEngine/Experimental/ChaosDerivedData.cpp b/Engine/Source/Runtime/Engine/Private/PhysicsEngine/Experimental/ChaosDerivedData.cpp
--- a/Engine/Source/Runtime/Engine/Private/PhysicsEngine/Experimental/ChaosDerivedData.cpp
+++ b/Engine/Source/Runtime/Engine/Private/PhysicsEngine/Experimental/ChaosDerivedData.cpp
@@ -24,7 +24,7 @@ const TCHAR* FChaosDerivedDataCooker::GetPluginName() const
 
 const TCHAR* FChaosDerivedDataCooker::GetVersionString() const
 {
-	return TEXT("CD1FE02F2AB9453FB6E173E28A076DC4");
+	return TEXT("F1F25234039247878ECAF926DBB9594C");
 }
 
 FString FChaosDerivedDataCooker::GetPluginSpecificCacheKeySuffix() const
diff --git a/Engine/Source/Runtime/Experimental/Chaos/Public/Chaos/CollisionConvexMesh.h b/Engine/Source/Runtime/Experimental/Chaos/Public/Chaos/CollisionConvexMesh.h
--- a/Engine/Source/Runtime/Experimental/Chaos/Public/Chaos/CollisionConvexMesh.h
+++ b/Engine/Source/Runtime/Experimental/Chaos/Public/Chaos/CollisionConvexMesh.h
@@ -6,42 +6,136 @@
 #include "Chaos/Box.h"
 #include "TriangleMesh.h"
 #include "Particles.h"
+#include "ChaosLog.h"
 
 namespace Chaos
 {
 	class FConvexBuilder
 	{
 	public:
-		static bool IsValidTriangle(const FVec3& A, const FVec3& B, const FVec3& C)
+
+		// When encountering a triangle or quad in hull generation (3-points or 4 coplanar points) we will instead generate
+		// a prism with a small thickness to emulate the desired result as a hull. Otherwise hull generation will fail on
+		// these cases. Verbose logging on LogChaos will point out when this path is taken for further scrutiny about
+		// the geometry
+		static constexpr float TriQuadPrismInflation = 0.1f;
+
+		static bool IsValidTriangle(const FVec3& A, const FVec3& B, const FVec3& C, FVec3& OutNormal)
 		{
 			const FVec3 BA = B - A;
 			const FVec3 CA = C - A;
 			const FVec3 Cross = FVec3::CrossProduct(BA, CA);
+			OutNormal = Cross.GetUnsafeNormal();
 			return Cross.Size() > 1e-4;
 		}
 
+		static bool IsValidTriangle(const FVec3& A, const FVec3& B, const FVec3& C)
+		{
+			FVec3 Normal(0);
+			return IsValidTriangle(A, B, C, Normal);
+		}
+
+		static bool IsValidQuad(const FVec3& A, const FVec3& B, const FVec3& C, const FVec3& D, FVec3& OutNormal)
+		{
+			FPlane TriPlane(A, B, C);
+			const float DPointDistance = FMath::Abs(TriPlane.PlaneDot(D));
+			OutNormal = FVec3(TriPlane.X, TriPlane.Y, TriPlane.Z);
+			return FMath::IsNearlyEqual(DPointDistance, 0, KINDA_SMALL_NUMBER);
+		}
+
+		static bool IsPlanarShape(const TParticles<FReal, 3>& InParticles, FVec3& OutNormal)
+		{
+			bool bResult = false;
+			const int32 NumParticles = InParticles.Size();
+			
+			if(NumParticles <= 3)
+			{
+				// Nothing, point, line or triangle, not a planar set
+				return false;
+			}
+			else // > 3 points
+			{
+				FPlane TriPlane(InParticles.X(0), InParticles.X(1), InParticles.X(2));
+				OutNormal = FVec3(TriPlane.X, TriPlane.Y, TriPlane.Z);
+
+				for(int32 Index = 3; Index < NumParticles; ++Index)
+				{
+					const float PointPlaneDot = FMath::Abs(TriPlane.PlaneDot(InParticles.X(Index)));
+					if(!FMath::IsNearlyEqual(PointPlaneDot, 0, KINDA_SMALL_NUMBER))
+					{
+						return false;
+					}
+				}
+			}
+
+			return true;
+		}
+
 		static void Build(const TParticles<FReal, 3>& InParticles, TArray <TPlaneConcrete<FReal, 3>>& OutPlanes, TArray<TArray<int32>>& OutFaceIndices, TParticles<FReal, 3>& OutSurfaceParticles, TAABB<FReal, 3>& OutLocalBounds)
 		{
 			OutPlanes.Reset();
 			OutSurfaceParticles.Resize(0);
 			OutLocalBounds = TAABB<FReal, 3>::EmptyAABB();
 
-			const uint32 NumParticles = InParticles.Size();
-			if(NumParticles == 0)
+			const uint32 NumParticlesIn = InParticles.Size();
+			if(NumParticlesIn == 0)
 			{
 				return;
 			}
 
-			OutLocalBounds = TAABB<FReal, 3>(InParticles.X(0), InParticles.X(0));
-			for(uint32 ParticleIndex = 0; ParticleIndex < NumParticles; ++ParticleIndex)
+			const TParticles<FReal, 3>* ParticlesToUse = &InParticles;
+			TParticles<FReal, 3> ModifiedParticles;
+
+			// For triangles and planar shapes, create a very thin prism as a convex
+			auto Inflate = [](const TParticles<FReal, 3>& Source, TParticles<FReal, 3>& Destination, const FVec3& Normal, float Inflation)
+			{
+				const int32 NumSource = Source.Size();
+				Destination.Resize(0);
+				Destination.AddParticles(NumSource * 2);
+
+				for(int32 Index = 0; Index < NumSource; ++Index)
+				{
+					Destination.X(Index) = Source.X(Index);
+					Destination.X(NumSource + Index) = Source.X(Index) + Normal * Inflation;
+				}
+			};
+
+			FVec3 PlanarNormal(0);
+			if(NumParticlesIn == 3)
+			{
+				const bool bIsValidTriangle = IsValidTriangle(InParticles.X(0), InParticles.X(1), InParticles.X(2), PlanarNormal);
+
+				//TODO_SQ_IMPLEMENTATION: should do proper cleanup to avoid this
+				if(ensureMsgf(bIsValidTriangle, TEXT("FConvexBuilder::Build(): Generated invalid triangle!")))
+				{
+					Inflate(InParticles, ModifiedParticles, PlanarNormal, TriQuadPrismInflation);
+					ParticlesToUse = &ModifiedParticles;
+					UE_LOG(LogChaos, Verbose, TEXT("Encountered a triangle in convex hull generation. Will prepare a prism of thickness %.5f in place of a triangle."), TriQuadPrismInflation);
+				}
+				else
+				{
+					return;
+				}
+			}
+			else if(IsPlanarShape(InParticles, PlanarNormal))
+			{
+				Inflate(InParticles, ModifiedParticles, PlanarNormal, TriQuadPrismInflation);
+				ParticlesToUse = &ModifiedParticles;
+				UE_LOG(LogChaos, Verbose, TEXT("Encountered a planar shape in convex hull generation. Will prepare a prism of thickness %.5f in place of a triangle."), TriQuadPrismInflation);
+			}
+
+			const int32 NumParticlesToUse = ParticlesToUse->Size();
+
+			OutLocalBounds = TAABB<FReal, 3>(ParticlesToUse->X(0), ParticlesToUse->X(0));
+			for(int32 ParticleIndex = 0; ParticleIndex < NumParticlesToUse; ++ParticleIndex)
 			{
-				OutLocalBounds.GrowToInclude(InParticles.X(ParticleIndex));
+				OutLocalBounds.GrowToInclude(ParticlesToUse->X(ParticleIndex));
 			}
 
-			if(NumParticles >= 4)
+			if(NumParticlesToUse >= 4)
 			{
 				TArray<TVector<int32, 3>> Indices;
-				BuildConvexHull(InParticles, Indices);
+				BuildConvexHull(*ParticlesToUse, Indices);
 				OutPlanes.Reserve(Indices.Num());
 				TMap<int32, int32> IndexMap; // maps original particle indices to output particle indices
 				int32 NewIdx = 0;
@@ -58,7 +152,7 @@ namespace Chaos
 
 				for(const TVector<int32, 3>& Idx : Indices)
 				{
-					FVec3 Vs[3] = {InParticles.X(Idx[0]), InParticles.X(Idx[1]), InParticles.X(Idx[2])};
+					FVec3 Vs[3] = {ParticlesToUse->X(Idx[0]), ParticlesToUse->X(Idx[1]), ParticlesToUse->X(Idx[2])};
 					const FVec3 Normal = FVec3::CrossProduct(Vs[1] - Vs[0], Vs[2] - Vs[0]).GetUnsafeNormal();
 					OutPlanes.Add(TPlaneConcrete<FReal, 3>(Vs[0], Normal));
 					TArray<int32> FaceIndices;
@@ -72,32 +166,11 @@ namespace Chaos
 				OutSurfaceParticles.AddParticles(IndexMap.Num());
 				for(const auto& Elem : IndexMap)
 				{
-					OutSurfaceParticles.X(Elem.Value) = InParticles.X(Elem.Key);
+					OutSurfaceParticles.X(Elem.Value) = ParticlesToUse->X(Elem.Key);
 				}
 			}
-			else if(NumParticles == 3)
-			{
-				//special support for triangle
-				const bool bIsValidTriangle = IsValidTriangle(InParticles.X(0), InParticles.X(1), InParticles.X(2));
 
-				//TODO_SQ_IMPLEMENTATION: should do proper cleanup to avoid this
-				if (ensureMsgf(bIsValidTriangle, TEXT("FConvexBuilder::Build(): Generated invalid triangle!")))
-				{
-					FVec3 Normal = FVec3::CrossProduct(InParticles.X(1) - InParticles.X(0), InParticles.X(2) - InParticles.X(0)).GetSafeNormal();
-					OutPlanes.Add(TPlaneConcrete<FReal, 3>(InParticles.X(0), Normal));
-					OutSurfaceParticles.AddParticles(3);
-					OutSurfaceParticles.X(0) = InParticles.X(0);
-					OutSurfaceParticles.X(1) = InParticles.X(1);
-					OutSurfaceParticles.X(2) = InParticles.X(2);
-
-					TArray<int32> FaceIndices;
-					FaceIndices.SetNum(3);
-					FaceIndices[0] = 0; 
-					FaceIndices[1] = 1;
-					FaceIndices[2] = 2;
-					OutFaceIndices.Add(FaceIndices);
-				}
-			}
+			UE_CLOG(OutSurfaceParticles.Size() == 0, LogChaos, Warning, TEXT("Convex hull generation produced zero convex particles, collision will fail for this primitive."));
 		}
 
 		static void BuildConvexHull(const TParticles<FReal, 3>& InParticles, TArray<TVector<int32, 3>>& OutIndices)
diff --git a/Engine/Source/Runtime/Experimental/Chaos/Public/Chaos/Convex.h b/Engine/Source/Runtime/Experimental/Chaos/Public/Chaos/Convex.h
--- a/Engine/Source/Runtime/Experimental/Chaos/Public/Chaos/Convex.h
+++ b/Engine/Source/Runtime/Experimental/Chaos/Public/Chaos/Convex.h
@@ -8,6 +8,7 @@
 #include "ChaosArchive.h"
 #include "GJK.h"
 #include "ChaosCheck.h"
+#include "ChaosLog.h"
 
 namespace Chaos
 {
@@ -167,16 +168,23 @@ namespace Chaos
 			int32 MaxVIdx = 0;
 			const int32 NumVertices = SurfaceParticles.Size();
 
-			check(NumVertices > 0);
-			for (int32 Idx = 0; Idx < NumVertices; ++Idx)
+			if(ensure(NumVertices > 0))
 			{
-				const FReal Dot = FVec3::DotProduct(SurfaceParticles.X(Idx), Direction);
-				if (Dot > MaxDot)
+				for(int32 Idx = 0; Idx < NumVertices; ++Idx)
 				{
-					MaxDot = Dot;
-					MaxVIdx = Idx;
+					const FReal Dot = FVec3::DotProduct(SurfaceParticles.X(Idx), Direction);
+					if(Dot > MaxDot)
+					{
+						MaxDot = Dot;
+						MaxVIdx = Idx;
+					}
 				}
 			}
+			else
+			{
+				UE_LOG(LogChaos, Warning, TEXT("Attempting to get a support for an empty convex. Returning object center."));
+				return FVec3(0);
+			}
 
 			if (Thickness)
 			{
