Adding defualt friction cvar.

diff --git a/Engine/Source/Runtime/Experimental/Chaos/Private/Chaos/PBDCollisionConstraints.cpp b/Engine/Source/Runtime/Experimental/Chaos/Private/Chaos/PBDCollisionConstraints.cpp
--- a/Engine/Source/Runtime/Experimental/Chaos/Private/Chaos/PBDCollisionConstraints.cpp
+++ b/Engine/Source/Runtime/Experimental/Chaos/Private/Chaos/PBDCollisionConstraints.cpp
@@ -50,6 +50,9 @@ namespace Chaos
 
 	CHAOS_API int32 EnableCollisions = 1;
 	FAutoConsoleVariableRef CVarEnableCollisions(TEXT("p.EnableCollisions"), EnableCollisions, TEXT("Enable/Disable collisions on the Chaos solver."));
+	
+	float DefaultCollisionFriction = 0;
+	FAutoConsoleVariableRef CVarDefaultCollisionFriction(TEXT("p.DefaultCollisionFriction"), DefaultCollisionFriction, TEXT("Collision friction default value if no materials are found."));
 
 	DECLARE_CYCLE_STAT(TEXT("Collisions::Reset"), STAT_Collisions_Reset, STATGROUP_ChaosCollision);
 	DECLARE_CYCLE_STAT(TEXT("Collisions::UpdatePointConstraints"), STAT_Collisions_UpdatePointConstraints, STATGROUP_ChaosCollision);
@@ -143,7 +146,7 @@ namespace Chaos
 		}
 		else
 		{
-			Contact.Friction = 0;
+			Contact.Friction = DefaultCollisionFriction;
 			Contact.Restitution = 0;
 		}
 		Contact.AngularFriction = MAngularFriction;
