Fix scaled trimesh/convex collision cases.
Add support for GJKContactPoint for scaled trimesh.
Fix instanced object paths in CollisionResolution.cpp.
Update Collisions EngineTest with more complex scaled cases.

diff --git a/Engine/Source/Runtime/Experimental/Chaos/Private/Chaos/CollisionResolution.cpp b/Engine/Source/Runtime/Experimental/Chaos/Private/Chaos/CollisionResolution.cpp
--- a/Engine/Source/Runtime/Experimental/Chaos/Private/Chaos/CollisionResolution.cpp
+++ b/Engine/Source/Runtime/Experimental/Chaos/Private/Chaos/CollisionResolution.cpp
@@ -176,9 +176,7 @@ namespace Chaos
 			}
 			else if (Implicit0OuterType == TImplicitObjectInstanced<FTriangleMeshImplicitObject>::StaticType())
 			{
-				// Confusingly, instanced Tri Meshes are treated differently
-				//return Implicit0->template GetObject<const TImplicitObjectInstanced<FTriangleMeshImplicitObject>>()->GetInstancedObject();
-				return nullptr;
+				return Implicit0->template GetObject<const TImplicitObjectInstanced<FTriangleMeshImplicitObject>>()->GetInstancedObject();
 			}
 
 			return nullptr;
@@ -486,14 +484,14 @@ namespace Chaos
 		// Box-TriangleMesh
 		//
 
-		template <typename T, int d>
-		TContactPoint<T> BoxTriangleMeshContactPoint(const TAABB<T, d>& A, const TRigidTransform<T, d>& ATransform, const FTriangleMeshImplicitObject& B, const TRigidTransform<T, d>& BTransform, const T CullDistance)
+		template <typename TriMeshType, typename T, int d>
+		TContactPoint<T> BoxTriangleMeshContactPoint(const TAABB<T, d>& A, const TRigidTransform<T, d>& ATransform, const TriMeshType& B, const TRigidTransform<T, d>& BTransform, const T CullDistance)
 		{
 			return GJKImplicitContactPoint< TBox<float, 3> >(TBox<float, 3>(A), ATransform, B, BTransform, CullDistance);
 		}
 
-		template <typename T, int d>
-		void UpdateBoxTriangleMeshConstraint(const TAABB<T, d>& Box0, const TRigidTransform<T, d>& Transform0, const FTriangleMeshImplicitObject& TriangleMesh1, const TRigidTransform<T, d>& Transform1, const T CullDistance, TRigidBodyPointContactConstraint<T, d>& Constraint)
+		template <typename TriMeshType, typename T, int d>
+		void UpdateBoxTriangleMeshConstraint(const TAABB<T, d>& Box0, const TRigidTransform<T, d>& Transform0, const TriMeshType& TriangleMesh1, const TRigidTransform<T, d>& Transform1, const T CullDistance, TRigidBodyPointContactConstraint<T, d>& Constraint)
 		{
 			UpdateContactPoint(Constraint.Manifold, BoxTriangleMeshContactPoint(Box0, Transform0, TriangleMesh1, Transform1, CullDistance));
 		}
@@ -509,14 +507,28 @@ namespace Chaos
 		void ConstructBoxTriangleMeshConstraints(TGeometryParticleHandle<T, d>* Particle0, TGeometryParticleHandle<T, d>* Particle1, const FImplicitObject* Implicit0, const FImplicitObject* Implicit1, const TRigidTransform<T, d>& Transform0, const TRigidTransform<T, d>& Transform1, const T CullDistance, FCollisionConstraintsArray& NewConstraints)
 		{
 			const TBox<T, d> * Object0 = Implicit0->template GetObject<const TBox<T, d> >();
-			const FTriangleMeshImplicitObject * Object1 = GetInnerObject<FTriangleMeshImplicitObject>(*Implicit1);
-			if (ensure(Object0 && Object1))
+			if (ensure(Object0))
 			{
-				TRigidTransform<T, d> ParticleImplicit0TM = Transform0.GetRelativeTransform(Collisions::GetTransform(Particle0));
-				TRigidTransform<T, d> ParticleImplicit1TM = Transform1.GetRelativeTransform(Collisions::GetTransform(Particle1));
-				FRigidBodyPointContactConstraint Constraint = FRigidBodyPointContactConstraint(Particle0, Implicit0, ParticleImplicit0TM, Particle1, Implicit1, ParticleImplicit1TM, EContactShapesType::BoxTriMesh);
-				UpdateBoxTriangleMeshConstraint(Object0->GetAABB(), Transform0, *Object1, Transform1, CullDistance, Constraint);
-				NewConstraints.TryAdd(CullDistance, Constraint);
+				if (const TImplicitObjectScaled<FTriangleMeshImplicitObject>* ScaledTriangleMesh = Implicit1->template GetObject<const TImplicitObjectScaled<FTriangleMeshImplicitObject>>())
+				{
+					TRigidTransform<T, d> ParticleImplicit0TM = Transform0.GetRelativeTransform(Collisions::GetTransform(Particle0));
+					TRigidTransform<T, d> ParticleImplicit1TM = Transform1.GetRelativeTransform(Collisions::GetTransform(Particle1));
+					FRigidBodyPointContactConstraint Constraint = FRigidBodyPointContactConstraint(Particle0, Implicit0, ParticleImplicit0TM, Particle1, Implicit1, ParticleImplicit1TM, EContactShapesType::BoxTriMesh);
+					UpdateBoxTriangleMeshConstraint(Object0->GetAABB(), Transform0, *ScaledTriangleMesh, Transform1, CullDistance, Constraint);
+					NewConstraints.TryAdd(CullDistance, Constraint);
+				}
+				else if(const FTriangleMeshImplicitObject* TriangleMesh = Implicit1->template GetObject<const FTriangleMeshImplicitObject>())
+				{
+					TRigidTransform<T, d> ParticleImplicit0TM = Transform0.GetRelativeTransform(Collisions::GetTransform(Particle0));
+					TRigidTransform<T, d> ParticleImplicit1TM = Transform1.GetRelativeTransform(Collisions::GetTransform(Particle1));
+					FRigidBodyPointContactConstraint Constraint = FRigidBodyPointContactConstraint(Particle0, Implicit0, ParticleImplicit0TM, Particle1, Implicit1, ParticleImplicit1TM, EContactShapesType::BoxTriMesh);
+					UpdateBoxTriangleMeshConstraint(Object0->GetAABB(), Transform0, *TriangleMesh, Transform1, CullDistance, Constraint);
+					NewConstraints.TryAdd(CullDistance, Constraint);
+				}
+				else
+				{
+					ensure(false);
+				}
 			}
 		}
 
@@ -768,14 +780,14 @@ namespace Chaos
 		// Sphere-TriangleMesh
 		//
 
-		template <typename T, int d>
-		TContactPoint<T> SphereTriangleMeshContactPoint(const TSphere<T, d>& A, const TRigidTransform<T, d>& ATransform, const FTriangleMeshImplicitObject& B, const TRigidTransform<T, d>& BTransform, const T CullDistance)
+		template <typename T, int d, typename TriMeshType>
+		TContactPoint<T> SphereTriangleMeshContactPoint(const TSphere<T, d>& A, const TRigidTransform<T, d>& ATransform, const TriMeshType& B, const TRigidTransform<T, d>& BTransform, const T CullDistance)
 		{
 			return GJKImplicitContactPoint< TSphere<float, 3> >(TSphere<float, 3>(A), ATransform, B, BTransform, CullDistance);
 		}
 
-		template <typename T, int d>
-		void UpdateSphereTriangleMeshConstraint(const TSphere<T, d>& Sphere0, const TRigidTransform<T, d>& Transform0, const FTriangleMeshImplicitObject& TriangleMesh1, const TRigidTransform<T, d>& Transform1, const T CullDistance, TRigidBodyPointContactConstraint<T, d>& Constraint)
+		template <typename T, int d, typename TriMeshType>
+		void UpdateSphereTriangleMeshConstraint(const TSphere<T, d>& Sphere0, const TRigidTransform<T, d>& Transform0, const TriMeshType& TriangleMesh1, const TRigidTransform<T, d>& Transform1, const T CullDistance, TRigidBodyPointContactConstraint<T, d>& Constraint)
 		{
 			UpdateContactPoint(Constraint.Manifold, SphereTriangleMeshContactPoint(Sphere0, Transform0, TriangleMesh1, Transform1, CullDistance));
 		}
@@ -790,15 +802,29 @@ namespace Chaos
 		template<typename T, int d>
 		void ConstructSphereTriangleMeshConstraints(TGeometryParticleHandle<T, d>* Particle0, TGeometryParticleHandle<T, d>* Particle1, const FImplicitObject* Implicit0, const FImplicitObject* Implicit1, const TRigidTransform<T, d>& Transform0, const TRigidTransform<T, d>& Transform1, const T CullDistance, FCollisionConstraintsArray& NewConstraints)
 		{
-			const TSphere<T, d> * Object0 = Implicit0->template GetObject<const TSphere<T, d> >();
-			const FTriangleMeshImplicitObject * Object1 = GetInnerObject<FTriangleMeshImplicitObject>(*Implicit1);
-			if (ensure(Object0 && Object1))
+			const TSphere<T, d>* Object0 = Implicit0->template GetObject<const TSphere<T, d> >();
+			if (ensure(Object0))
 			{
-				TRigidTransform<T, d> ParticleImplicit0TM = Transform0.GetRelativeTransform(Collisions::GetTransform(Particle0));
-				TRigidTransform<T, d> ParticleImplicit1TM = Transform1.GetRelativeTransform(Collisions::GetTransform(Particle1));
-				FRigidBodyPointContactConstraint Constraint = FRigidBodyPointContactConstraint(Particle0, Implicit0, ParticleImplicit0TM, Particle1, Implicit1, ParticleImplicit1TM, EContactShapesType::SphereTriMesh);
-				UpdateSphereTriangleMeshConstraint(*Object0, Transform0, *Object1, Transform1, CullDistance, Constraint);
-				NewConstraints.TryAdd(CullDistance, Constraint);
+				if (const TImplicitObjectScaled<FTriangleMeshImplicitObject>* ScaledTriangleMesh = Implicit1->template GetObject<const TImplicitObjectScaled<FTriangleMeshImplicitObject>>())
+				{
+					TRigidTransform<T, d> ParticleImplicit0TM = Transform0.GetRelativeTransform(Collisions::GetTransform(Particle0));
+					TRigidTransform<T, d> ParticleImplicit1TM = Transform1.GetRelativeTransform(Collisions::GetTransform(Particle1));
+					FRigidBodyPointContactConstraint Constraint = FRigidBodyPointContactConstraint(Particle0, Implicit0, ParticleImplicit0TM, Particle1, Implicit1, ParticleImplicit1TM, EContactShapesType::SphereTriMesh);
+					UpdateSphereTriangleMeshConstraint(*Object0, Transform0, *ScaledTriangleMesh, Transform1, CullDistance, Constraint);
+					NewConstraints.TryAdd(CullDistance, Constraint);
+				}
+				else if (const FTriangleMeshImplicitObject* TriangleMesh = Implicit1->template GetObject<const FTriangleMeshImplicitObject>())
+				{
+					TRigidTransform<T, d> ParticleImplicit0TM = Transform0.GetRelativeTransform(Collisions::GetTransform(Particle0));
+					TRigidTransform<T, d> ParticleImplicit1TM = Transform1.GetRelativeTransform(Collisions::GetTransform(Particle1));
+					FRigidBodyPointContactConstraint Constraint = FRigidBodyPointContactConstraint(Particle0, Implicit0, ParticleImplicit0TM, Particle1, Implicit1, ParticleImplicit1TM, EContactShapesType::SphereTriMesh);
+					UpdateSphereTriangleMeshConstraint(*Object0, Transform0, *TriangleMesh, Transform1, CullDistance, Constraint);
+					NewConstraints.TryAdd(CullDistance, Constraint);
+				}
+				else
+				{
+					ensure(false);
+				}
 			}
 		}
 
@@ -1167,14 +1193,14 @@ namespace Chaos
 		// Capsule-TriangleMesh
 		//
 
-		template <typename T, int d>
-		TContactPoint<T> CapsuleTriangleMeshContactPoint(const TCapsule<T>& A, const TRigidTransform<T, d>& ATransform, const FTriangleMeshImplicitObject& B, const TRigidTransform<T, d>& BTransform, const T CullDistance)
+		template <typename TriMeshType, typename T, int d>
+		TContactPoint<T> CapsuleTriangleMeshContactPoint(const TCapsule<T>& A, const TRigidTransform<T, d>& ATransform, const TriMeshType& B, const TRigidTransform<T, d>& BTransform, const T CullDistance)
 		{
 			return GJKImplicitContactPoint< TCapsule<T> >(A, ATransform, B, BTransform, CullDistance);
 		}
 
-		template <typename T, int d>
-		void UpdateCapsuleTriangleMeshConstraint(const TCapsule<T>& Capsule0, const TRigidTransform<T, d>& Transform0, const FTriangleMeshImplicitObject& TriangleMesh1, const TRigidTransform<T, d>& Transform1, const T CullDistance, TRigidBodyPointContactConstraint<T, d>& Constraint)
+		template <typename TriMeshType, typename T, int d>
+		void UpdateCapsuleTriangleMeshConstraint(const TCapsule<T>& Capsule0, const TRigidTransform<T, d>& Transform0, const TriMeshType& TriangleMesh1, const TRigidTransform<T, d>& Transform1, const T CullDistance, TRigidBodyPointContactConstraint<T, d>& Constraint)
 		{
 			UpdateContactPoint(Constraint.Manifold, CapsuleTriangleMeshContactPoint(Capsule0, Transform0, TriangleMesh1, Transform1, CullDistance));
 		}
@@ -1190,18 +1216,31 @@ namespace Chaos
 		void ConstructCapsuleTriangleMeshConstraints(TGeometryParticleHandle<T, d>* Particle0, TGeometryParticleHandle<T, d>* Particle1, const FImplicitObject* Implicit0, const FImplicitObject* Implicit1, const TRigidTransform<T, d>& Transform0, const TRigidTransform<T, d>& Transform1, const T CullDistance, FCollisionConstraintsArray& NewConstraints)
 		{
 			const TCapsule<T> * Object0 = Implicit0->template GetObject<const TCapsule<T> >();
-			const FTriangleMeshImplicitObject * Object1 = GetInnerObject<FTriangleMeshImplicitObject>(*Implicit1);
-			if (ensure(Object0 && Object1))
+			if (ensure(Object0))
 			{
-				TRigidTransform<T, d> ParticleImplicit0TM = Transform0.GetRelativeTransform(Collisions::GetTransform(Particle0));
-				TRigidTransform<T, d> ParticleImplicit1TM = Transform1.GetRelativeTransform(Collisions::GetTransform(Particle1));
-				FRigidBodyPointContactConstraint Constraint = FRigidBodyPointContactConstraint(Particle0, Implicit0, ParticleImplicit0TM, Particle1, Implicit1, ParticleImplicit1TM, EContactShapesType::CapsuleTriMesh);
-				UpdateCapsuleTriangleMeshConstraint(*Object0, Transform0, *Object1, Transform1, CullDistance, Constraint);
-				NewConstraints.TryAdd(CullDistance, Constraint);
+				if (const TImplicitObjectScaled<FTriangleMeshImplicitObject>* ScaledTriangleMesh = Implicit1->template GetObject<const TImplicitObjectScaled<FTriangleMeshImplicitObject>>())
+				{
+					TRigidTransform<T, d> ParticleImplicit0TM = Transform0.GetRelativeTransform(Collisions::GetTransform(Particle0));
+					TRigidTransform<T, d> ParticleImplicit1TM = Transform1.GetRelativeTransform(Collisions::GetTransform(Particle1));
+					FRigidBodyPointContactConstraint Constraint = FRigidBodyPointContactConstraint(Particle0, Implicit0, ParticleImplicit0TM, Particle1, Implicit1, ParticleImplicit1TM, EContactShapesType::CapsuleTriMesh);
+					UpdateCapsuleTriangleMeshConstraint(*Object0, Transform0, *ScaledTriangleMesh, Transform1, CullDistance, Constraint);
+					NewConstraints.TryAdd(CullDistance, Constraint);
+				}
+				else if (const FTriangleMeshImplicitObject* TriangleMesh = Implicit1->template GetObject<const FTriangleMeshImplicitObject>())
+				{
+					TRigidTransform<T, d> ParticleImplicit0TM = Transform0.GetRelativeTransform(Collisions::GetTransform(Particle0));
+					TRigidTransform<T, d> ParticleImplicit1TM = Transform1.GetRelativeTransform(Collisions::GetTransform(Particle1));
+					FRigidBodyPointContactConstraint Constraint = FRigidBodyPointContactConstraint(Particle0, Implicit0, ParticleImplicit0TM, Particle1, Implicit1, ParticleImplicit1TM, EContactShapesType::CapsuleTriMesh);
+					UpdateCapsuleTriangleMeshConstraint(*Object0, Transform0, *TriangleMesh, Transform1, CullDistance, Constraint);
+					NewConstraints.TryAdd(CullDistance, Constraint);
+				}
+				else
+				{
+					ensure(false);
+				}
 			}
 		}
 
-
 		//
 		// Convex - Convex
 		//
@@ -1313,14 +1352,14 @@ namespace Chaos
 		// Convex-TriangleMesh
 		//
 
-		template <typename T, int d>
-		TContactPoint<T> ConvexTriangleMeshContactPoint(const FImplicitObject& A, const TRigidTransform<T, d>& ATransform, const FTriangleMeshImplicitObject& B, const TRigidTransform<T, d>& BTransform, const T CullDistance)
+		template <typename TriMeshType, typename T, int d>
+		TContactPoint<T> ConvexTriangleMeshContactPoint(const FImplicitObject& A, const TRigidTransform<T, d>& ATransform, const TriMeshType& B, const TRigidTransform<T, d>& BTransform, const T CullDistance)
 		{
 			return GJKImplicitContactPoint< FConvex >(A, ATransform, B, BTransform, CullDistance);
 		}
 
-		template <typename T, int d>
-		void UpdateConvexTriangleMeshConstraint(const FImplicitObject& Convex0, const TRigidTransform<T, d>& Transform0, const FTriangleMeshImplicitObject& TriangleMesh1, const TRigidTransform<T, d>& Transform1, const T CullDistance, TRigidBodyPointContactConstraint<T, d>& Constraint)
+		template <typename TriMeshType, typename T, int d>
+		void UpdateConvexTriangleMeshConstraint(const FImplicitObject& Convex0, const TRigidTransform<T, d>& Transform0, const TriMeshType& TriangleMesh1, const TRigidTransform<T, d>& Transform1, const T CullDistance, TRigidBodyPointContactConstraint<T, d>& Constraint)
 		{
 			UpdateContactPoint(Constraint.Manifold, ConvexTriangleMeshContactPoint(Convex0, Transform0, TriangleMesh1, Transform1, CullDistance));
 		}
@@ -1335,14 +1374,28 @@ namespace Chaos
 		template<typename T, int d>
 		void ConstructConvexTriangleMeshConstraints(TGeometryParticleHandle<T, d>* Particle0, TGeometryParticleHandle<T, d>* Particle1, const FImplicitObject* Implicit0, const FImplicitObject* Implicit1, const TRigidTransform<T, d>& Transform0, const TRigidTransform<T, d>& Transform1, const T CullDistance, FCollisionConstraintsArray& NewConstraints)
 		{
-			const FTriangleMeshImplicitObject * Object1 = GetInnerObject<FTriangleMeshImplicitObject>(*Implicit1);
-			if (ensure(Implicit0->IsConvex() && Object1) )
+			if (ensure(Implicit0->IsConvex()))
 			{
-				TRigidTransform<T, d> ParticleImplicit0TM = Transform0.GetRelativeTransform(Collisions::GetTransform(Particle0));
-				TRigidTransform<T, d> ParticleImplicit1TM = Transform1.GetRelativeTransform(Collisions::GetTransform(Particle1));
-				FRigidBodyPointContactConstraint Constraint = FRigidBodyPointContactConstraint(Particle0, Implicit0, ParticleImplicit0TM, Particle1, Implicit1, ParticleImplicit1TM, EContactShapesType::ConvexTriMesh);
-				UpdateConvexTriangleMeshConstraint(*Implicit0, Transform0, *Object1, Transform1, CullDistance, Constraint);
-				NewConstraints.TryAdd(CullDistance, Constraint);
+				if (const TImplicitObjectScaled<FTriangleMeshImplicitObject>* ScaledTriangleMesh = Implicit1->template GetObject<const TImplicitObjectScaled<FTriangleMeshImplicitObject>>())
+				{
+					TRigidTransform<T, d> ParticleImplicit0TM = Transform0.GetRelativeTransform(Collisions::GetTransform(Particle0));
+					TRigidTransform<T, d> ParticleImplicit1TM = Transform1.GetRelativeTransform(Collisions::GetTransform(Particle1));
+					FRigidBodyPointContactConstraint Constraint = FRigidBodyPointContactConstraint(Particle0, Implicit0, ParticleImplicit0TM, Particle1, Implicit1, ParticleImplicit1TM, EContactShapesType::ConvexTriMesh);
+					UpdateConvexTriangleMeshConstraint(*Implicit0, Transform0, *ScaledTriangleMesh , Transform1, CullDistance, Constraint);
+					NewConstraints.TryAdd(CullDistance, Constraint);
+				}
+				else if (const FTriangleMeshImplicitObject* TriangleMesh = Implicit1->template GetObject<const FTriangleMeshImplicitObject>())
+				{
+					TRigidTransform<T, d> ParticleImplicit0TM = Transform0.GetRelativeTransform(Collisions::GetTransform(Particle0));
+					TRigidTransform<T, d> ParticleImplicit1TM = Transform1.GetRelativeTransform(Collisions::GetTransform(Particle1));
+					FRigidBodyPointContactConstraint Constraint = FRigidBodyPointContactConstraint(Particle0, Implicit0, ParticleImplicit0TM, Particle1, Implicit1, ParticleImplicit1TM, EContactShapesType::ConvexTriMesh);
+					UpdateConvexTriangleMeshConstraint(*Implicit0, Transform0, *TriangleMesh, Transform1, CullDistance, Constraint);
+					NewConstraints.TryAdd(CullDistance, Constraint);
+				}
+				else
+				{
+					ensure(false);
+				}
 			}
 		}
 
@@ -1473,7 +1526,17 @@ namespace Chaos
 				UpdateConvexConvexConstraint(Implicit0, Transform0, Implicit1, Transform1, CullDistance, Constraint);
 				break;
 			case EContactShapesType::SphereTriMesh:
-				UpdateSphereTriangleMeshConstraint(*Implicit0.template GetObject<TSphere<T, d>>(), Transform0, *GetInnerObject<FTriangleMeshImplicitObject>(Implicit1), Transform1, CullDistance, Constraint);
+				if (const TImplicitObjectScaled<FTriangleMeshImplicitObject>* ScaledTriMesh = Implicit1.template GetObject<const TImplicitObjectScaled<FTriangleMeshImplicitObject> >())
+				{
+					UpdateSphereTriangleMeshConstraint(*Implicit0.template GetObject<TSphere<T, d>>(), Transform0, *ScaledTriMesh, Transform1, CullDistance, Constraint);
+					break;
+				}
+				else if (const FTriangleMeshImplicitObject* TriangleMeshImplicit = Implicit1.template GetObject<const FTriangleMeshImplicitObject>())
+				{
+					UpdateSphereTriangleMeshConstraint(*Implicit0.template GetObject<TSphere<T, d>>(), Transform0, *TriangleMeshImplicit, Transform1, CullDistance, Constraint);
+					break;
+				}
+				ensure(false);
 				break;
 			case EContactShapesType::SphereHeightField:
 				UpdateSphereHeightFieldConstraint(*Implicit0.template GetObject<TSphere<T, d>>(), Transform0, *Implicit1.template GetObject< THeightField<T> >(), Transform1, CullDistance, Constraint);
@@ -1492,7 +1555,17 @@ namespace Chaos
 				UpdateConvexConvexConstraint(Implicit0, Transform0, Implicit1, Transform1, CullDistance, Constraint);
 				break;
 			case EContactShapesType::CapsuleTriMesh:
-				UpdateCapsuleTriangleMeshConstraint(*Implicit0.template GetObject<TCapsule<T>>(), Transform0, *GetInnerObject<FTriangleMeshImplicitObject>(Implicit1), Transform1, CullDistance, Constraint);
+				if (const TImplicitObjectScaled<FTriangleMeshImplicitObject>* ScaledTriMesh = Implicit1.template GetObject<const TImplicitObjectScaled<FTriangleMeshImplicitObject> >())
+				{
+					UpdateCapsuleTriangleMeshConstraint(*Implicit0.template GetObject<TCapsule<T>>(), Transform0, *ScaledTriMesh, Transform1, CullDistance, Constraint);
+					break;
+				}
+				else if (const FTriangleMeshImplicitObject* TriangleMeshImplicit = Implicit1.template GetObject<const FTriangleMeshImplicitObject>())
+				{
+					UpdateCapsuleTriangleMeshConstraint(*Implicit0.template GetObject<TCapsule<T>>(), Transform0, *TriangleMeshImplicit, Transform1, CullDistance, Constraint);
+					break;
+				}
+				ensure(false);
 				break;
 			case EContactShapesType::CapsuleHeightField:
 				UpdateCapsuleHeightFieldConstraint(*Implicit0.template GetObject<TCapsule<T>>(), Transform0, *Implicit1.template GetObject< THeightField<T> >(), Transform1, CullDistance, Constraint);
@@ -1505,7 +1578,17 @@ namespace Chaos
 				UpdateConvexConvexConstraint(Implicit0, Transform0, Implicit1, Transform1, CullDistance, Constraint);
 				break;
 			case EContactShapesType::BoxTriMesh:
-				UpdateBoxTriangleMeshConstraint(Implicit0.template GetObject<TBox<T, d>>()->GetAABB(), Transform0, *GetInnerObject<FTriangleMeshImplicitObject>(Implicit1), Transform1, CullDistance, Constraint);
+				if (const TImplicitObjectScaled<FTriangleMeshImplicitObject>* ScaledTriMesh = Implicit1.template GetObject<const TImplicitObjectScaled<FTriangleMeshImplicitObject> >())
+				{
+					UpdateBoxTriangleMeshConstraint(Implicit0.template GetObject<TBox<T, d>>()->GetAABB(), Transform0, *ScaledTriMesh, Transform1, CullDistance, Constraint);
+					break;
+				}
+				else if (const FTriangleMeshImplicitObject* TriangleMeshImplicit = Implicit1.template GetObject<const FTriangleMeshImplicitObject>())
+				{
+					UpdateBoxTriangleMeshConstraint(Implicit0.template GetObject<TBox<T, d>>()->GetAABB(), Transform0, *TriangleMeshImplicit, Transform1, CullDistance, Constraint);
+					break;
+				}
+				ensure(false);
 				break;
 			case EContactShapesType::BoxHeightField:
 				UpdateBoxHeightFieldConstraint(Implicit0.template GetObject<TBox<T, d>>()->GetAABB(), Transform0, *Implicit1.template GetObject< THeightField<T> >(), Transform1, CullDistance, Constraint);
@@ -1517,7 +1600,17 @@ namespace Chaos
 				UpdateConvexConvexConstraint(Implicit0, Transform0, Implicit1, Transform1, CullDistance, Constraint);
 				break;
 			case EContactShapesType::ConvexTriMesh:
-				UpdateConvexTriangleMeshConstraint(Implicit0, Transform0, *GetInnerObject<FTriangleMeshImplicitObject>(Implicit1), Transform1, CullDistance, Constraint);
+				if (const TImplicitObjectScaled<FTriangleMeshImplicitObject>* ScaledTriMesh = Implicit1.template GetObject<const TImplicitObjectScaled<FTriangleMeshImplicitObject> >())
+				{
+					UpdateConvexTriangleMeshConstraint(Implicit0, Transform0, *ScaledTriMesh, Transform1, CullDistance, Constraint);
+					break;
+				}
+				else if (const FTriangleMeshImplicitObject* TriangleMeshImplicit = Implicit1.template GetObject<const FTriangleMeshImplicitObject>())
+				{
+					UpdateConvexTriangleMeshConstraint(Implicit0, Transform0, *TriangleMeshImplicit, Transform1, CullDistance, Constraint);
+					break;
+				}
+				ensure(false);
 				break;
 			case EContactShapesType::ConvexHeightField:
 				UpdateConvexHeightFieldConstraint(Implicit0, Transform0, *Implicit1.template GetObject< THeightField<T> >(), Transform1, CullDistance, Constraint);
@@ -1857,20 +1950,6 @@ namespace Chaos
 				ConstructConstraints(Particle0, Particle1, Implicit0, TransformedImplicit1->GetTransformedObject(), Transform0, TransformedTransform1, CullDistance, Context, NewConstraints);
 				return;
 			}
-			else if (Implicit0OuterType == TImplicitObjectScaled<FTriangleMeshImplicitObject>::StaticType())
-			{
-				const TImplicitObjectScaled<FTriangleMeshImplicitObject>* ScaledImplicit0 = Implicit0->template GetObject<const TImplicitObjectScaled<FTriangleMeshImplicitObject>>();
-				TRigidTransform<T, d> ScaledTransform0 = TRigidTransform<T, d>(TVec3<T>(0.0f), FQuat::Identity, ScaledImplicit0->GetScale()) * Transform0;
-				ConstructConstraints(Particle0, Particle1, ScaledImplicit0->Object().Get(), Implicit1, ScaledTransform0, Transform1, CullDistance, Context, NewConstraints);
-				return;
-			}
-			else if (Implicit1OuterType == TImplicitObjectScaled<FTriangleMeshImplicitObject>::StaticType())
-			{
-				const TImplicitObjectScaled<FTriangleMeshImplicitObject>* ScaledImplicit1 = Implicit1->template GetObject<const TImplicitObjectScaled<FTriangleMeshImplicitObject>>();
-				TRigidTransform<T, d> ScaledTransform1 = TRigidTransform<T, d>(TVec3<T>(0.0f), FQuat::Identity, ScaledImplicit1->GetScale()) * Transform1;
-				ConstructConstraints(Particle0, Particle1, Implicit0, ScaledImplicit1->Object().Get(), Transform0, ScaledTransform1, CullDistance, Context, NewConstraints);
-				return;
-			}
 			// Handle Instanced shapes
 			// NOTE: Tri Meshes are handled differently. We should probably do something about this...
 			if (((uint32)Implicit0OuterType & ImplicitObjectType::IsInstanced) || ((uint32)Implicit1OuterType & ImplicitObjectType::IsInstanced))
@@ -1890,6 +1969,17 @@ namespace Chaos
 					ConstructConstraints(Particle0, Particle1, InnerImplicit0, InnerImplicit1, Transform0, Transform1, CullDistance, Context, NewConstraints);
 					return;
 				}
+				else if (InnerImplicit0 && !InnerImplicit1)
+				{
+					ConstructConstraints(Particle0, Particle1, InnerImplicit0, Implicit1, Transform0, Transform1, CullDistance, Context, NewConstraints);
+					return;
+				}
+				else if (!InnerImplicit0 && InnerImplicit1)
+				{
+					ConstructConstraints(Particle0, Particle1, Implicit0, InnerImplicit1, Transform0, Transform1, CullDistance, Context, NewConstraints);
+					return;
+				}
+
 			}
 
 			// Handle Unions
diff --git a/Engine/Source/Runtime/Experimental/Chaos/Private/Chaos/TriangleMeshImplicitObject.cpp b/Engine/Source/Runtime/Experimental/Chaos/Private/Chaos/TriangleMeshImplicitObject.cpp
--- a/Engine/Source/Runtime/Experimental/Chaos/Private/Chaos/TriangleMeshImplicitObject.cpp
+++ b/Engine/Source/Runtime/Experimental/Chaos/Private/Chaos/TriangleMeshImplicitObject.cpp
@@ -248,10 +248,16 @@ bool FTriangleMeshImplicitObject::Raycast(const FVec3& StartPoint, const FVec3&
 
 
 template <typename QueryGeomType>
-bool FTriangleMeshImplicitObject::GJKContactPointImp(const QueryGeomType& QueryGeom, const FRigidTransform3& QueryTM, const FReal Thickness, FVec3& Location, FVec3& Normal, FReal& Penetration) const
+bool FTriangleMeshImplicitObject::GJKContactPointImp(const QueryGeomType& QueryGeom, const FRigidTransform3& QueryTM, const FReal Thickness, FVec3& Location, FVec3& Normal, FReal& Penetration, FVec3 TriMeshScale) const
 {
+	ensure(TriMeshScale != FVec3(0.0f));
 	bool bResult = false;
 
+	const auto& WorldScaleGeom = ScaleGeomIntoWorldHelper(QueryGeom, TriMeshScale);
+
+	TRigidTransform<FReal, 3> WorldScaleQueryTM;
+	ScaleTransformHelper(TriMeshScale, QueryTM, WorldScaleQueryTM);
+
 	auto OverlapTriangle = [&](const FVec3& A, const FVec3& B, const FVec3& C,
 		FVec3& LocalContactLocation, FVec3& LocalContactNormal, FReal& LocalContactPhi) -> bool
 	{
@@ -261,10 +267,8 @@ bool FTriangleMeshImplicitObject::GJKContactPointImp(const QueryGeomType& QueryG
 
 		FReal LambdaPenetration;
 		FVec3 ClosestA, ClosestB, LambdaNormal;
-		if (GJKPenetration(TriangleConvex, QueryGeom, QueryTM, LambdaPenetration, ClosestA, ClosestB, LambdaNormal, (FReal)0))
+		if (GJKPenetration(TriangleConvex, WorldScaleGeom, WorldScaleQueryTM, LambdaPenetration, ClosestA, ClosestB, LambdaNormal, (FReal)0))
 		{
-			FVec3 TestVector = QueryTM.InverseTransformVector(LambdaNormal);
-
 			LocalContactLocation = ClosestB;
 			LocalContactNormal = LambdaNormal;
 			LocalContactPhi = -LambdaPenetration;
@@ -290,7 +294,7 @@ bool FTriangleMeshImplicitObject::GJKContactPointImp(const QueryGeomType& QueryG
 			//It's most likely that the query object is in front of the triangle since queries tend to be on the outside.
 			//However, maybe we should check if it's behind the triangle plane. Also, we should enforce this winding in some way
 			TVec3<FReal> A, B, C;
-			TransformVertsHelper(QueryGeom, TriIdx, MParticles, Elements, A, B, C);
+			TransformVertsHelper(TriMeshScale, TriIdx, MParticles, Elements, A, B, C);
 
 			if (OverlapTriangle(A, B, C, LocalContactLocation, LocalContactNormal, LocalContactPhi))
 			{
@@ -344,24 +348,24 @@ bool FTriangleMeshImplicitObject::GJKContactPoint(const FConvex& QueryGeom, cons
 	return GJKContactPointImp(QueryGeom, QueryTM, Thickness, Location, Normal, Penetration);
 }
 
-bool FTriangleMeshImplicitObject::GJKContactPoint(const TImplicitObjectScaled< TSphere<FReal, 3> >& QueryGeom, const FRigidTransform3& QueryTM, const FReal Thickness, FVec3& Location, FVec3& Normal, FReal& Penetration) const
+bool FTriangleMeshImplicitObject::GJKContactPoint(const TImplicitObjectScaled< TSphere<FReal, 3> >& QueryGeom, const FRigidTransform3& QueryTM, const FReal Thickness, FVec3& Location, FVec3& Normal, FReal& Penetration, FVec3 TriMeshScale) const
 {
-	return GJKContactPointImp(QueryGeom, QueryTM, Thickness, Location, Normal, Penetration);
+	return GJKContactPointImp(QueryGeom, QueryTM, Thickness, Location, Normal, Penetration, TriMeshScale);
 }
 
-bool FTriangleMeshImplicitObject::GJKContactPoint(const TImplicitObjectScaled< TBox<FReal, 3> >& QueryGeom, const FRigidTransform3& QueryTM, const FReal Thickness, FVec3& Location, FVec3& Normal, FReal& Penetration) const
+bool FTriangleMeshImplicitObject::GJKContactPoint(const TImplicitObjectScaled< TBox<FReal, 3> >& QueryGeom, const FRigidTransform3& QueryTM, const FReal Thickness, FVec3& Location, FVec3& Normal, FReal& Penetration, FVec3 TriMeshScale) const
 {
-	return GJKContactPointImp(QueryGeom, QueryTM, Thickness, Location, Normal, Penetration);
+	return GJKContactPointImp(QueryGeom, QueryTM, Thickness, Location, Normal, Penetration, TriMeshScale);
 }
 
-bool FTriangleMeshImplicitObject::GJKContactPoint(const TImplicitObjectScaled< TCapsule<FReal> >& QueryGeom, const FRigidTransform3& QueryTM, const FReal Thickness, FVec3& Location, FVec3& Normal, FReal& Penetration) const
+bool FTriangleMeshImplicitObject::GJKContactPoint(const TImplicitObjectScaled< TCapsule<FReal> >& QueryGeom, const FRigidTransform3& QueryTM, const FReal Thickness, FVec3& Location, FVec3& Normal, FReal& Penetration, FVec3 TriMeshScale) const
 {
-	return GJKContactPointImp(QueryGeom, QueryTM, Thickness, Location, Normal, Penetration);
+	return GJKContactPointImp(QueryGeom, QueryTM, Thickness, Location, Normal, Penetration, TriMeshScale);
 }
 
-bool FTriangleMeshImplicitObject::GJKContactPoint(const TImplicitObjectScaled< FConvex >& QueryGeom, const FRigidTransform3& QueryTM, const FReal Thickness, FVec3& Location, FVec3& Normal, FReal& Penetration) const
+bool FTriangleMeshImplicitObject::GJKContactPoint(const TImplicitObjectScaled< FConvex >& QueryGeom, const FRigidTransform3& QueryTM, const FReal Thickness, FVec3& Location, FVec3& Normal, FReal& Penetration, FVec3 TriMeshScale) const
 {
-	return GJKContactPointImp(QueryGeom, QueryTM, Thickness, Location, Normal, Penetration);
+	return GJKContactPointImp(QueryGeom, QueryTM, Thickness, Location, Normal, Penetration, TriMeshScale);
 }
 
 int32 FTriangleMeshImplicitObject::GetExternalFaceIndexFromInternal(int32 InternalFaceIndex) const
@@ -437,6 +441,24 @@ bool FTriangleMeshImplicitObject::Overlap(const FVec3& Point, const FReal Thickn
 	}
 }
 
+template <typename QueryGeomType, typename T, int d>
+static auto MakeScaledHelper(const QueryGeomType& B, const TVector<T,d>& InvScale )
+{
+	TUniquePtr<QueryGeomType> HackBPtr(const_cast<QueryGeomType*>(&B));	//todo: hack, need scaled object to accept raw ptr similar to transformed implicit
+	TImplicitObjectScaled<QueryGeomType> ScaledB(MakeSerializable(HackBPtr), InvScale);
+	HackBPtr.Release();
+	return ScaledB;
+}
+
+template <typename QueryGeomType, typename T, int d>
+static auto MakeScaledHelper(const TImplicitObjectScaled<QueryGeomType>& B, const TVector<T,d>& InvScale)
+{
+	//if scaled of scaled just collapse into one scaled
+	TImplicitObjectScaled<QueryGeomType> ScaledB(B.Object(), InvScale * B.GetScale());
+	return ScaledB;
+}
+
+
 template <typename QueryGeomType, typename IdxType>
 void TransformVertsHelper(const QueryGeomType& QueryGeom, int32 TriIdx, const TParticles<FReal, 3>& Particles,
 	const TArray<TVector<IdxType, 3>>& Elements, TVec3<FReal>& OutA, TVec3<FReal>& OutB, TVec3<FReal>& OutC)
@@ -462,10 +484,32 @@ const QueryGeomType& GetGeomHelper(const QueryGeomType& QueryGeom)
 	return QueryGeom;
 }
 
+void ScaleTransformHelper(const FVec3& TriMeshScale, const FRigidTransform3& QueryTM, FRigidTransform3& OutScaledQueryTM)
+{
+	OutScaledQueryTM = TRigidTransform<FReal, 3>(QueryTM.GetLocation() * TriMeshScale, QueryTM.GetRotation());
+}
+
+
+template <typename IdxType>
+void TransformVertsHelper(const FVec3& TriMeshScale, int32 TriIdx, const TParticles<FReal, 3>& Particles,
+	const TArray<TVector<IdxType, 3>>& Elements, TVec3<FReal>& OutA, TVec3<FReal>& OutB, TVec3<FReal>& OutC)
+{
+	OutA = Particles.X(Elements[TriIdx][0]) * TriMeshScale;
+	OutB = Particles.X(Elements[TriIdx][1]) * TriMeshScale;
+	OutC = Particles.X(Elements[TriIdx][2]) * TriMeshScale;
+}
+
+template <typename QueryGeomType>
+const QueryGeomType& ScaleGeomIntoWorldHelper(const QueryGeomType& QueryGeom, const FVec3& TriMeshScale)
+{
+	return QueryGeom;
+}
+
 template <typename QueryGeomType>
-const QueryGeomType& GetGeomHelper(const TImplicitObjectScaled<QueryGeomType>& QueryGeom)
+TImplicitObjectScaled<QueryGeomType> ScaleGeomIntoWorldHelper(const TImplicitObjectScaled<QueryGeomType>& QueryGeom, const FVec3& TriMeshScale)
 {
-	return *QueryGeom.GetUnscaledObject();
+	// This will apply TriMeshScale to QueryGeom and return a new scaled implicit in world space.
+	return MakeScaledHelper(QueryGeom, TriMeshScale);
 }
 
 template <typename QueryGeomType>
diff --git a/Engine/Source/Runtime/Experimental/Chaos/Public/Chaos/CollisionResolution.h b/Engine/Source/Runtime/Experimental/Chaos/Public/Chaos/CollisionResolution.h
--- a/Engine/Source/Runtime/Experimental/Chaos/Public/Chaos/CollisionResolution.h
+++ b/Engine/Source/Runtime/Experimental/Chaos/Public/Chaos/CollisionResolution.h
@@ -85,8 +85,8 @@ namespace Chaos
 		// Box-TriangleMesh
 		//
 
-		template <typename T, int d>
-		void CHAOS_API UpdateBoxTriangleMeshConstraint(const TAABB<T, d>& Box0, const TRigidTransform<T, d>& Transform0, const FTriangleMeshImplicitObject& TriangleMesh1, const TRigidTransform<T, d>& Transformw1, const T CullDistance, TRigidBodyPointContactConstraint<T, d>& Constraint);
+		template <typename TriMeshType, typename T, int d>
+		void CHAOS_API UpdateBoxTriangleMeshConstraint(const TAABB<T, d>& Box0, const TRigidTransform<T, d>& Transform0, const TriMeshType & TriangleMesh1, const TRigidTransform<T, d>& Transformw1, const T CullDistance, TRigidBodyPointContactConstraint<T, d>& Constraint);
 
 		template <typename T, int d>
 		void CHAOS_API UpdateBoxTriangleMeshManifold(TCollisionConstraintBase<T, d>&  Constraint, const TRigidTransform<T, d>& ATM, const TRigidTransform<T, d>& BTM, const T CullDistance);
diff --git a/Engine/Source/Runtime/Experimental/Chaos/Public/Chaos/ImplicitObjectScaled.h b/Engine/Source/Runtime/Experimental/Chaos/Public/Chaos/ImplicitObjectScaled.h
--- a/Engine/Source/Runtime/Experimental/Chaos/Public/Chaos/ImplicitObjectScaled.h
+++ b/Engine/Source/Runtime/Experimental/Chaos/Public/Chaos/ImplicitObjectScaled.h
@@ -394,6 +394,15 @@ public:
 		return false;
 	}
 
+	template <typename QueryGeomType>
+	bool GJKContactPoint(const QueryGeomType& A, const FRigidTransform3& AToBTM, const FReal Thickness, FVec3& Location, FVec3& Normal, FReal& Penetration) const
+	{
+		TRigidTransform<T, d> AToBTMNoScale(AToBTM.GetLocation() * MInvScale, AToBTM.GetRotation());
+
+		auto ScaledA = MakeScaledHelper(A, MInvScale);
+		return MObject->GJKContactPoint(ScaledA, AToBTMNoScale, MInternalThickness + Thickness, Location, Normal, Penetration, MScale);
+	}
+
 	/** This is a low level function and assumes the internal object has a OverlapGeom function. Should not be called directly. See GeometryQueries.h : OverlapQuery */
 	template <typename QueryGeomType>
 	bool LowLevelOverlapGeom(const QueryGeomType& B, const TRigidTransform<T, d>& BToATM, T Thickness = 0, FMTDInfo* OutMTD = nullptr) const
diff --git a/Engine/Source/Runtime/Experimental/Chaos/Public/Chaos/TriangleMeshImplicitObject.h b/Engine/Source/Runtime/Experimental/Chaos/Public/Chaos/TriangleMeshImplicitObject.h
--- a/Engine/Source/Runtime/Experimental/Chaos/Public/Chaos/TriangleMeshImplicitObject.h
+++ b/Engine/Source/Runtime/Experimental/Chaos/Public/Chaos/TriangleMeshImplicitObject.h
@@ -150,10 +150,11 @@ namespace Chaos
 		bool GJKContactPoint(const TBox<FReal, 3>& QueryGeom, const FRigidTransform3& QueryTM, const FReal Thickness, FVec3& Location, FVec3& Normal, FReal& Penetration) const;
 		bool GJKContactPoint(const TCapsule<FReal>& QueryGeom, const FRigidTransform3& QueryTM, const FReal Thickness, FVec3& Location, FVec3& Normal, FReal& Penetration) const;
 		bool GJKContactPoint(const FConvex& QueryGeom, const FRigidTransform3& QueryTM, const FReal Thickness, FVec3& Location, FVec3& Normal, FReal& Penetration) const;
-		bool GJKContactPoint(const TImplicitObjectScaled < TSphere<FReal, 3> >& QueryGeom, const FRigidTransform3& QueryTM, const FReal Thickness, FVec3& Location, FVec3& Normal, FReal& Penetration) const;
-		bool GJKContactPoint(const TImplicitObjectScaled < TBox<FReal, 3> >& QueryGeom, const FRigidTransform3& QueryTM, const FReal Thickness, FVec3& Location, FVec3& Normal, FReal& Penetration) const;
-		bool GJKContactPoint(const TImplicitObjectScaled < TCapsule<FReal> >& QueryGeom, const FRigidTransform3& QueryTM, const FReal Thickness, FVec3& Location, FVec3& Normal, FReal& Penetration) const;
-		bool GJKContactPoint(const TImplicitObjectScaled < FConvex >& QueryGeom, const FRigidTransform3& QueryTM, const FReal Thickness, FVec3& Location, FVec3& Normal, FReal& Penetration) const;
+
+		bool GJKContactPoint(const TImplicitObjectScaled < TSphere<FReal, 3> >& QueryGeom, const FRigidTransform3& QueryTM, const FReal Thickness, FVec3& Location, FVec3& Normal, FReal& Penetration, FVec3 TriMeshScale = FVec3(1.0f)) const;
+		bool GJKContactPoint(const TImplicitObjectScaled < TBox<FReal, 3> >& QueryGeom, const FRigidTransform3& QueryTM, const FReal Thickness, FVec3& Location, FVec3& Normal, FReal& Penetration, FVec3 TriMeshScale = FVec3(1.0f)) const;
+		bool GJKContactPoint(const TImplicitObjectScaled < TCapsule<FReal> >& QueryGeom, const FRigidTransform3& QueryTM, const FReal Thickness, FVec3& Location, FVec3& Normal, FReal& Penetration, FVec3 TriMeshScale = FVec3(1.0f)) const;
+		bool GJKContactPoint(const TImplicitObjectScaled < FConvex >& QueryGeom, const FRigidTransform3& QueryTM, const FReal Thickness, FVec3& Location, FVec3& Normal, FReal& Penetration, FVec3 TriMeshScale = FVec3(1.0f)) const;
 
 		// Returns -1 if InternalFaceIndex is not in map, or map is invalid.
 		int32 GetExternalFaceIndexFromInternal(int32 InternalFaceIndex) const;
@@ -323,7 +324,7 @@ namespace Chaos
 		    : FImplicitObject(EImplicitObject::HasBoundingBox, ImplicitObjectType::TriangleMesh){};
 
 		template <typename QueryGeomType>
-		bool GJKContactPointImp(const QueryGeomType& QueryGeom, const FRigidTransform3& QueryTM, const FReal Thickness, FVec3& Location, FVec3& Normal, FReal& Penetration) const;
+		bool GJKContactPointImp(const QueryGeomType& QueryGeom, const FRigidTransform3& QueryTM, const FReal Thickness, FVec3& Location, FVec3& Normal, FReal& Penetration, FVec3 TriMeshScale = FVec3(1.0)) const;
 
 		template<typename QueryGeomType>
 		bool OverlapGeomImp(const QueryGeomType& QueryGeom, const FRigidTransform3& QueryTM, const FReal Thickness, FMTDInfo* OutMTD = nullptr) const;
