Chaos Cloth - Fix crash in geodesic long range attachment by allowing several paths from multiples kinematic nodes per single target node, even if these paths have some nodes in common.
- Also enable multiple kinematic attachment zones to split the cloth.

diff --git a/Engine/Source/Runtime/Experimental/Chaos/Private/Chaos/PBDLongRangeConstraintsBase.cpp b/Engine/Source/Runtime/Experimental/Chaos/Private/Chaos/PBDLongRangeConstraintsBase.cpp
--- a/Engine/Source/Runtime/Experimental/Chaos/Private/Chaos/PBDLongRangeConstraintsBase.cpp
+++ b/Engine/Source/Runtime/Experimental/Chaos/Private/Chaos/PBDLongRangeConstraintsBase.cpp
@@ -248,6 +248,7 @@ void TPBDLongRangeConstraintsBase<T, d>::ComputeGeodesicConstraints(
 			const TSet<uint32>& Neighbors = PointToNeighbors[PairElem.Second];
 			for (const uint32 Neighbor : Neighbors)
 			{
+				if (InParticles.InvM(Neighbor) == (T)0.) { continue; }
 				check(Neighbor != PairElem.Second);
 				const TVector<uint32, 2> NeighborStartEnd(Element, Neighbor);
 				const Pair<T, TArray<uint32>>& NeighborDistancePath = GeodesicPaths[NeighborStartEnd];
@@ -274,9 +275,8 @@ void TPBDLongRangeConstraintsBase<T, d>::ComputeGeodesicConstraints(
 		{
 			if (!Elements.Num()) { continue; }  // Empty island 
 
-			int32 ClosestElement = Elements[0];
-			T ClosestDistance = GeodesicPaths[TVector<uint32, 2>(ClosestElement, i)].First;
-			if (ClosestDistance == FLT_MAX) { continue; }  // Not on this island
+			int32 ClosestElement = INDEX_NONE;
+			T ClosestDistance = FLT_MAX;
 
 			for (const uint32 Element : Elements)
 			{
@@ -287,6 +287,7 @@ void TPBDLongRangeConstraintsBase<T, d>::ComputeGeodesicConstraints(
 					ClosestElement = Element;
 				}
 			}
+			if (ClosestElement == INDEX_NONE) { continue; }  // Not on this island
 
 			const TVector<uint32, 2> Index(ClosestElement, i);
 			check(GeodesicPaths[Index].First != FLT_MAX);
@@ -314,15 +315,16 @@ void TPBDLongRangeConstraintsBase<T, d>::ComputeGeodesicConstraints(
 	MConstraints.Sort([](const TArray<uint32>& Elem1, const TArray<uint32>& Elem2) { return Elem1.Num() > Elem2.Num(); });
 	TArray<TArray<uint32>> NewConstraints;
 	TArray<T> NewDists;
-	TMap<uint32, TArray<uint32>> ProcessedPairs;
-	for (uint32 i = 1; i < static_cast<uint32>(MConstraints.Num()); ++i)
+	TMap<TVector<uint32, 2>, TArray<uint32>> ProcessedPairs;
+	for (uint32 i = 0; i < static_cast<uint32>(MConstraints.Num()); ++i)
 	{
-		if (ProcessedPairs.Contains(MConstraints[i].Last()))
+		const TVector<uint32, 2> Traverse(MConstraints[i][0], MConstraints[i].Last());
+		if (const TArray<uint32>* TraversePath = ProcessedPairs.Find(Traverse))
 		{
-			check(MConstraints[i].Num() == ProcessedPairs[MConstraints[i].Last()].Num());
-			for (uint32 j = 0; j < static_cast<uint32>(ProcessedPairs[MConstraints[i].Last()].Num()); ++j)
+			check(MConstraints[i].Num() == TraversePath->Num());
+			for (uint32 j = 0; j < static_cast<uint32>(TraversePath->Num()); ++j)
 			{
-				check(ProcessedPairs[MConstraints[i].Last()][j] == MConstraints[i][j]);
+				check((*TraversePath)[j] == MConstraints[i][j]);
 			}
 			continue;
 		}
@@ -338,7 +340,8 @@ void TPBDLongRangeConstraintsBase<T, d>::ComputeGeodesicConstraints(
 			//NewDists.Add(Dist);
 			NewConstraints.Add({Path[0], Path[Path.Num() - 1]});  // Use shortest distance instead 
 			NewDists.Add(ComputeDistance(InParticles, Path[0], Path[Path.Num() - 1]));
-			ProcessedPairs.Add(MConstraints[i][j], Path);
+			const TVector<uint32, 2> SubTraverse(MConstraints[i][0], MConstraints[i][j]);
+			ProcessedPairs.Add(SubTraverse, Path);
 		}
 	}
 	MDists = NewDists;
