Fix clang builds.
(TriangleMeshImplicit HelperFunctions being defined after usage in TriangleMeshImplicitObject.cpp)

diff --git a/Engine/Source/Runtime/Experimental/Chaos/Private/Chaos/TriangleMeshImplicitObject.cpp b/Engine/Source/Runtime/Experimental/Chaos/Private/Chaos/TriangleMeshImplicitObject.cpp
--- a/Engine/Source/Runtime/Experimental/Chaos/Private/Chaos/TriangleMeshImplicitObject.cpp
+++ b/Engine/Source/Runtime/Experimental/Chaos/Private/Chaos/TriangleMeshImplicitObject.cpp
@@ -11,6 +11,113 @@
 {
 
 
+
+template <typename QueryGeomType, typename T, int d>
+static auto MakeScaledHelper(const QueryGeomType& B, const TVector<T, d>& InvScale)
+{
+	TUniquePtr<QueryGeomType> HackBPtr(const_cast<QueryGeomType*>(&B));	//todo: hack, need scaled object to accept raw ptr similar to transformed implicit
+	TImplicitObjectScaled<QueryGeomType> ScaledB(MakeSerializable(HackBPtr), InvScale);
+	HackBPtr.Release();
+	return ScaledB;
+}
+
+template <typename QueryGeomType, typename T, int d>
+static auto MakeScaledHelper(const TImplicitObjectScaled<QueryGeomType>& B, const TVector<T, d>& InvScale)
+{
+	//if scaled of scaled just collapse into one scaled
+	TImplicitObjectScaled<QueryGeomType> ScaledB(B.Object(), InvScale * B.GetScale());
+	return ScaledB;
+}
+
+
+template <typename QueryGeomType, typename IdxType>
+void TransformVertsHelper(const QueryGeomType& QueryGeom, int32 TriIdx, const TParticles<FReal, 3>& Particles,
+	const TArray<TVector<IdxType, 3>>& Elements, TVec3<FReal>& OutA, TVec3<FReal>& OutB, TVec3<FReal>& OutC)
+{
+	OutA = Particles.X(Elements[TriIdx][0]);
+	OutB = Particles.X(Elements[TriIdx][1]);
+	OutC = Particles.X(Elements[TriIdx][2]);
+}
+
+template <typename QueryGeomType, typename IdxType>
+void TransformVertsHelper(const TImplicitObjectScaled<QueryGeomType>& QueryGeom, int32 TriIdx, const TParticles<FReal, 3>& Particles,
+	const TArray<TVector<IdxType, 3>>& Elements, TVec3<FReal>& OutA, TVec3<FReal>& OutB, TVec3<FReal>& OutC)
+{
+	const TVec3<FReal> InvScale = QueryGeom.GetInvScale();
+	OutA = Particles.X(Elements[TriIdx][0]) * InvScale;
+	OutB = Particles.X(Elements[TriIdx][1]) * InvScale;
+	OutC = Particles.X(Elements[TriIdx][2]) * InvScale;
+}
+
+template <typename QueryGeomType>
+const QueryGeomType& GetGeomHelper(const QueryGeomType& QueryGeom)
+{
+	return QueryGeom;
+}
+
+void ScaleTransformHelper(const FVec3& TriMeshScale, const FRigidTransform3& QueryTM, FRigidTransform3& OutScaledQueryTM)
+{
+	OutScaledQueryTM = TRigidTransform<FReal, 3>(QueryTM.GetLocation() * TriMeshScale, QueryTM.GetRotation());
+}
+
+
 template <typename IdxType>
+void TransformVertsHelper(const FVec3& TriMeshScale, int32 TriIdx, const TParticles<FReal, 3>& Particles,
+	const TArray<TVector<IdxType, 3>>& Elements, TVec3<FReal>& OutA, TVec3<FReal>& OutB, TVec3<FReal>& OutC)
+{
+	OutA = Particles.X(Elements[TriIdx][0]) * TriMeshScale;
+	OutB = Particles.X(Elements[TriIdx][1]) * TriMeshScale;
+	OutC = Particles.X(Elements[TriIdx][2]) * TriMeshScale;
+}
+
+template <typename QueryGeomType>
+const QueryGeomType& ScaleGeomIntoWorldHelper(const QueryGeomType& QueryGeom, const FVec3& TriMeshScale)
+{
+	return QueryGeom;
+}
+
+template <typename QueryGeomType>
+TImplicitObjectScaled<QueryGeomType> ScaleGeomIntoWorldHelper(const TImplicitObjectScaled<QueryGeomType>& QueryGeom, const FVec3& TriMeshScale)
+{
+	// This will apply TriMeshScale to QueryGeom and return a new scaled implicit in world space.
+	return MakeScaledHelper(QueryGeom, TriMeshScale);
+}
+
+template <typename QueryGeomType>
+void TransformSweepOutputsHelper(const QueryGeomType& QueryGeom, const TVec3<FReal>& HitNormal, const TVec3<FReal>& HitPosition, const FReal LengthScale,
+	const FReal Time, TVec3<FReal>& OutNormal, TVec3<FReal>& OutPosition, FReal& OutTime)
+{
+	OutNormal = HitNormal;
+	OutPosition = HitPosition;
+	OutTime = Time;
+}
+
+template <typename QueryGeomType>
+void TransformSweepOutputsHelper(const TImplicitObjectScaled<QueryGeomType>& QueryGeom, const TVec3<FReal>& HitNormal, const TVec3<FReal>& HitPosition, const FReal LengthScale,
+	const FReal Time, TVec3<FReal>& OutNormal, TVec3<FReal>& OutPosition, FReal& OutTime)
+{
+	const TVec3<FReal> InvScale = QueryGeom.GetInvScale();
+	const TVec3<FReal> Scale = QueryGeom.GetScale();
+
+	OutTime = Time / LengthScale;
+	OutNormal = (InvScale * HitNormal).GetSafeNormal();
+	OutPosition = Scale * HitPosition;
+}
+
+template <typename QueryGeomType>
+void TransformOverlapInputsHelper(const QueryGeomType& QueryGeom, const FRigidTransform3& QueryTM, FRigidTransform3& OutScaledQueryTM)
+{
+	OutScaledQueryTM = QueryTM;
+}
+
+template <typename QueryGeomType>
+void TransformOverlapInputsHelper(const TImplicitObjectScaled<QueryGeomType>& QueryGeom, const FRigidTransform3& QueryTM, FRigidTransform3& OutScaledQueryTM)
+{
+	const TVec3<FReal> InvScale = QueryGeom.GetInvScale();
+	OutScaledQueryTM = TRigidTransform<FReal, 3>(QueryTM.GetLocation() * InvScale, QueryTM.GetRotation());
+}
+
+
+template <typename IdxType>
 struct FTriangleMeshRaycastVisitor
 {
@@ -441,111 +548,6 @@
 	}
 }
 
-template <typename QueryGeomType, typename T, int d>
-static auto MakeScaledHelper(const QueryGeomType& B, const TVector<T,d>& InvScale )
-{
-	TUniquePtr<QueryGeomType> HackBPtr(const_cast<QueryGeomType*>(&B));	//todo: hack, need scaled object to accept raw ptr similar to transformed implicit
-	TImplicitObjectScaled<QueryGeomType> ScaledB(MakeSerializable(HackBPtr), InvScale);
-	HackBPtr.Release();
-	return ScaledB;
-}
-
-template <typename QueryGeomType, typename T, int d>
-static auto MakeScaledHelper(const TImplicitObjectScaled<QueryGeomType>& B, const TVector<T,d>& InvScale)
-{
-	//if scaled of scaled just collapse into one scaled
-	TImplicitObjectScaled<QueryGeomType> ScaledB(B.Object(), InvScale * B.GetScale());
-	return ScaledB;
-}
-
-
-template <typename QueryGeomType, typename IdxType>
-void TransformVertsHelper(const QueryGeomType& QueryGeom, int32 TriIdx, const TParticles<FReal, 3>& Particles,
-	const TArray<TVector<IdxType, 3>>& Elements, TVec3<FReal>& OutA, TVec3<FReal>& OutB, TVec3<FReal>& OutC)
-{
-	OutA = Particles.X(Elements[TriIdx][0]);
-	OutB = Particles.X(Elements[TriIdx][1]);
-	OutC = Particles.X(Elements[TriIdx][2]);
-}
-
-template <typename QueryGeomType, typename IdxType>
-void TransformVertsHelper(const TImplicitObjectScaled<QueryGeomType>& QueryGeom, int32 TriIdx, const TParticles<FReal, 3>& Particles,
-	const TArray<TVector<IdxType, 3>>& Elements, TVec3<FReal>& OutA, TVec3<FReal>& OutB, TVec3<FReal>& OutC)
-{
-	const TVec3<FReal> InvScale = QueryGeom.GetInvScale();
-	OutA = Particles.X(Elements[TriIdx][0]) * InvScale;
-	OutB = Particles.X(Elements[TriIdx][1]) * InvScale;
-	OutC = Particles.X(Elements[TriIdx][2]) * InvScale;
-}
-
 template <typename QueryGeomType>
-const QueryGeomType& GetGeomHelper(const QueryGeomType& QueryGeom)
-{
-	return QueryGeom;
-}
-
-void ScaleTransformHelper(const FVec3& TriMeshScale, const FRigidTransform3& QueryTM, FRigidTransform3& OutScaledQueryTM)
-{
-	OutScaledQueryTM = TRigidTransform<FReal, 3>(QueryTM.GetLocation() * TriMeshScale, QueryTM.GetRotation());
-}
-
-
-template <typename IdxType>
-void TransformVertsHelper(const FVec3& TriMeshScale, int32 TriIdx, const TParticles<FReal, 3>& Particles,
-	const TArray<TVector<IdxType, 3>>& Elements, TVec3<FReal>& OutA, TVec3<FReal>& OutB, TVec3<FReal>& OutC)
-{
-	OutA = Particles.X(Elements[TriIdx][0]) * TriMeshScale;
-	OutB = Particles.X(Elements[TriIdx][1]) * TriMeshScale;
-	OutC = Particles.X(Elements[TriIdx][2]) * TriMeshScale;
-}
-
-template <typename QueryGeomType>
-const QueryGeomType& ScaleGeomIntoWorldHelper(const QueryGeomType& QueryGeom, const FVec3& TriMeshScale)
-{
-	return QueryGeom;
-}
-
-template <typename QueryGeomType>
-TImplicitObjectScaled<QueryGeomType> ScaleGeomIntoWorldHelper(const TImplicitObjectScaled<QueryGeomType>& QueryGeom, const FVec3& TriMeshScale)
-{
-	// This will apply TriMeshScale to QueryGeom and return a new scaled implicit in world space.
-	return MakeScaledHelper(QueryGeom, TriMeshScale);
-}
-
-template <typename QueryGeomType>
-void TransformSweepOutputsHelper(const QueryGeomType& QueryGeom, const TVec3<FReal>& HitNormal, const TVec3<FReal>& HitPosition, const FReal LengthScale,
-	const FReal Time,  TVec3<FReal>& OutNormal, TVec3<FReal>& OutPosition, FReal& OutTime)
-{
-	OutNormal = HitNormal;
-	OutPosition = HitPosition;
-	OutTime = Time;
-}
-
-template <typename QueryGeomType>
-void TransformSweepOutputsHelper(const TImplicitObjectScaled<QueryGeomType>& QueryGeom, const TVec3<FReal>& HitNormal, const TVec3<FReal>& HitPosition,  const FReal LengthScale,
-	const FReal Time,  TVec3<FReal>& OutNormal, TVec3<FReal>& OutPosition, FReal& OutTime)
-{
-	const TVec3<FReal> InvScale = QueryGeom.GetInvScale();
-	const TVec3<FReal> Scale = QueryGeom.GetScale();
-
-	OutTime = Time / LengthScale;
-	OutNormal = (InvScale * HitNormal).GetSafeNormal();
-	OutPosition = Scale * HitPosition;
-}
-
-template <typename QueryGeomType>
-void TransformOverlapInputsHelper(const QueryGeomType& QueryGeom, const FRigidTransform3& QueryTM, FRigidTransform3& OutScaledQueryTM)
-{
-	OutScaledQueryTM = QueryTM;
-}
-
-template <typename QueryGeomType>
-void TransformOverlapInputsHelper(const TImplicitObjectScaled<QueryGeomType>& QueryGeom, const FRigidTransform3& QueryTM, FRigidTransform3& OutScaledQueryTM)
-{
-	const TVec3<FReal> InvScale = QueryGeom.GetInvScale();
-	OutScaledQueryTM = TRigidTransform<FReal, 3>(QueryTM.GetLocation() * InvScale, QueryTM.GetRotation());
-}
-
-template <typename QueryGeomType>
 bool FTriangleMeshImplicitObject::OverlapGeomImp(const QueryGeomType& QueryGeom, const FRigidTransform3& QueryTM, const FReal Thickness, FMTDInfo* OutMTD) const
 {
