Added a simple PrePreFilter so that we can prune broadphase earlier

diff --git a/Engine/Source/Programs/HeadlessChaos/Private/HeadlessChaosTestParticleHandle.cpp b/Engine/Source/Programs/HeadlessChaos/Private/HeadlessChaosTestParticleHandle.cpp
--- a/Engine/Source/Programs/HeadlessChaos/Private/HeadlessChaosTestParticleHandle.cpp
+++ b/Engine/Source/Programs/HeadlessChaos/Private/HeadlessChaosTestParticleHandle.cpp
@@ -518,6 +518,7 @@ namespace ChaosTest {
 
 	void AccelerationStructureHandleComparison()
 	{
+#if 0
 		// When an external particle is created, the handle is retrieved via proxy.
 		// Proxy gets handle initialized async on PT later, so this means a handle 
 		// will always have external particle pointer, and eventually an internal pointer.
@@ -551,6 +552,7 @@ namespace ChaosTest {
 
 		EXPECT_NE(InternalOnlyHandle, ExternalOnlyHandle);
 		EXPECT_EQ(InternalOnlyHandle, ExternalInternalHandle);
+#endif
 	}
 
 	void HandleObjectStateChangeTest()
diff --git a/Engine/Source/Runtime/Engine/Private/PhysicsEngine/BodyInstance.cpp b/Engine/Source/Runtime/Engine/Private/PhysicsEngine/BodyInstance.cpp
--- a/Engine/Source/Runtime/Engine/Private/PhysicsEngine/BodyInstance.cpp
+++ b/Engine/Source/Runtime/Engine/Private/PhysicsEngine/BodyInstance.cpp
@@ -894,9 +894,15 @@ void FBodyInstance::UpdatePhysicsFilterData()
 	}
 
 		if(bUpdateMassProperties)
-	{
+		{
 			UpdateMassProperties();
 		}
+
+		//If filtering changed we must update GT structure right away
+		if (FPhysScene* PhysScene = GetPhysicsScene())
+		{
+			PhysScene->GetScene().UpdateActorInAccelerationStructure(Actor);
+		}
 	});
 
 	UpdateInterpolateWhenSubStepping();
diff --git a/Engine/Source/Runtime/Experimental/Chaos/Private/Chaos/TriangleMeshImplicitObject.cpp b/Engine/Source/Runtime/Experimental/Chaos/Private/Chaos/TriangleMeshImplicitObject.cpp
--- a/Engine/Source/Runtime/Experimental/Chaos/Private/Chaos/TriangleMeshImplicitObject.cpp
+++ b/Engine/Source/Runtime/Experimental/Chaos/Private/Chaos/TriangleMeshImplicitObject.cpp
@@ -138,6 +138,11 @@ struct FTriangleMeshRaycastVisitor
 		Sweep
 	};
 
+	const void* GetQueryData() const
+	{
+		return nullptr;
+	}
+
 	template <ERaycastType SQType>
 	bool Visit(int32 TriIdx, FQueryFastData& CurData)
 	{
@@ -696,6 +701,8 @@ struct FTriangleMeshSweepVisitor
 	{
 	}
 
+	const void* GetQueryData() const { return nullptr; }
+
 	bool VisitOverlap(const TSpatialVisitorData<int32>& VisitData)
 	{
 		check(false);
@@ -796,7 +803,7 @@ template <typename QueryGeomType>
 bool FTriangleMeshImplicitObject::SweepGeomImp(const QueryGeomType& QueryGeom, const FRigidTransform3& StartTM, const FVec3& Dir, const FReal Length,
 	FReal& OutTime, FVec3& OutPosition, FVec3& OutNormal, int32& OutFaceIndex, const FReal Thickness, const bool bComputeMTD) const
 {
-
+	//QUICK_SCOPE_CYCLE_COUNTER(TrimeshSweep);
 	// Compute scaled sweep inputs to cache in visitor.
 	FVec3 ScaledDirNormalized;
 	FReal LengthScale;
diff --git a/Engine/Source/Runtime/Experimental/Chaos/Public/Chaos/AABBTree.h b/Engine/Source/Runtime/Experimental/Chaos/Public/Chaos/AABBTree.h
--- a/Engine/Source/Runtime/Experimental/Chaos/Public/Chaos/AABBTree.h
+++ b/Engine/Source/Runtime/Experimental/Chaos/Public/Chaos/AABBTree.h
@@ -572,41 +572,58 @@ private:
 	template <EAABBQueryType Query, typename TQueryFastData, typename SQVisitor>
 	bool QueryImp(const TVector<T, 3>& Start, TQueryFastData& CurData, const TVector<T, 3> QueryHalfExtents, const TAABB<T,3>& QueryBounds, SQVisitor& Visitor) const
 	{
+		//QUICK_SCOPE_CYCLE_COUNTER(AABBTreeQueryImp);
 		TVector<T, 3> TmpPosition;
 		T TOI = 0;
+		const void* QueryData = Visitor.GetQueryData();
 
-		for (const auto& Elem : GlobalPayloads)
 		{
-			const auto& InstanceBounds = Elem.Bounds;
-			if (TAABBTreeIntersectionHelper<T, TQueryFastData, Query>::Intersects(Start, CurData, TOI, TmpPosition, InstanceBounds, QueryBounds, QueryHalfExtents))
+
+			//QUICK_SCOPE_CYCLE_COUNTER(QueryGlobal);
+
+			for(const auto& Elem : GlobalPayloads)
 			{
-				TSpatialVisitorData<TPayloadType> VisitData(Elem.Payload, true);
-				bool bContinue;
-				if (Query == EAABBQueryType::Overlap)
+				if (PrePreFilterHelper(Elem.Payload, QueryData))
 				{
-					bContinue = Visitor.VisitOverlap(VisitData);
-				}
-				else
-				{
-					bContinue = Query == EAABBQueryType::Sweep ? Visitor.VisitSweep(VisitData, CurData) : Visitor.VisitRaycast(VisitData, CurData);
+					continue;
 				}
 
-				if (!bContinue)
+				const auto& InstanceBounds = Elem.Bounds;
+				if(TAABBTreeIntersectionHelper<T,TQueryFastData,Query>::Intersects(Start,CurData,TOI,TmpPosition,InstanceBounds,QueryBounds,QueryHalfExtents))
 				{
-					return false;
+					TSpatialVisitorData<TPayloadType> VisitData(Elem.Payload,true);
+					bool bContinue;
+					if(Query == EAABBQueryType::Overlap)
+					{
+						bContinue = Visitor.VisitOverlap(VisitData);
+					} else
+					{
+						bContinue = Query == EAABBQueryType::Sweep ? Visitor.VisitSweep(VisitData,CurData) : Visitor.VisitRaycast(VisitData,CurData);
+					}
+
+					if(!bContinue)
+					{
+						return false;
+					}
 				}
 			}
 		}
 
 		if (bMutable)
 		{
+			//QUICK_SCOPE_CYCLE_COUNTER(QueryDirty);
+
 			for (const auto& Elem : DirtyElements)
 			{
 				const auto& InstanceBounds = Elem.Bounds;
+				if (PrePreFilterHelper(Elem.Payload, QueryData))
+				{
+					continue;
+				}
+
 				if (TAABBTreeIntersectionHelper<T, TQueryFastData, Query>::Intersects(Start, CurData, TOI, TmpPosition, InstanceBounds, QueryBounds, QueryHalfExtents))
 				{
 					TSpatialVisitorData<TPayloadType> VisitData(Elem.Payload, true, InstanceBounds);
-
 					bool bContinue;
 					if (Query == EAABBQueryType::Overlap)
 					{
@@ -623,6 +640,7 @@ private:
 					}
 				}
 			}
+
 		}
 
 		struct FNodeQueueEntry
@@ -959,6 +977,9 @@ private:
 				check(false);
 				return true;
 			}
+
+			const void* GetQueryData() const { return nullptr; }
+
 			TArray<TPayloadType>& CollectedResults;
 		};
 
diff --git a/Engine/Source/Runtime/Experimental/Chaos/Public/Chaos/BoundingVolume.h b/Engine/Source/Runtime/Experimental/Chaos/Public/Chaos/BoundingVolume.h
--- a/Engine/Source/Runtime/Experimental/Chaos/Public/Chaos/BoundingVolume.h
+++ b/Engine/Source/Runtime/Experimental/Chaos/Public/Chaos/BoundingVolume.h
@@ -196,6 +196,7 @@ public:
 				CollectedResults.Add(Instance.Payload);
 				return true;
 			}
+			const void* GetQueryData() const { return nullptr; }
 			TArray<TPayloadType>& CollectedResults;
 		};
 
@@ -372,9 +373,15 @@ private:
 	{
 		TVector<T, d> TmpPosition;
 		T TOI;
+		const void* QueryData = Visitor.GetQueryData();
 
 		for (const auto& Elem : MGlobalPayloads)
 		{
+			if (PrePreFilterHelper(Elem.Payload, QueryData))
+			{
+				continue;
+			}
+
 			const auto& InstanceBounds = Elem.Bounds;
 			if (InstanceBounds.RaycastFast(Start,
 				CurData.Dir, CurData.InvDir, CurData.bParallel, CurData.CurrentLength, CurData.InvCurrentLength, TOI, TmpPosition))
@@ -391,6 +398,11 @@ private:
 
 		for (const auto& Elem : MDirtyElements)
 		{
+			if (PrePreFilterHelper(Elem.Payload, QueryData))
+			{
+				continue;
+			}
+
 			const auto& InstanceBounds = Elem.Bounds;
 			if (InstanceBounds.RaycastFast(Start, CurData.Dir,
 				CurData.InvDir, CurData.bParallel, CurData.CurrentLength, CurData.InvCurrentLength, TOI, TmpPosition))
@@ -525,8 +537,14 @@ private:
 	bool SweepImp(const TVector<T, d>& Start, FQueryFastData& CurData, const TVector<T, d> QueryHalfExtents, SQVisitor& Visitor) const
 	{
 		T TOI = 0;
+		const void* QueryData = Visitor.GetQueryData();
 		for (const auto& Elem : MGlobalPayloads)
 		{
+			if (PrePreFilterHelper(Elem.Payload, QueryData))
+			{
+				continue;
+			}
+
 			const TAABB<T, d>& InstanceBounds = Elem.Bounds;
 			const TVector<T, d> Min = InstanceBounds.Min() - QueryHalfExtents;
 			const TVector<T, d> Max = InstanceBounds.Max() + QueryHalfExtents;
@@ -544,6 +562,11 @@ private:
 
 		for (const auto& Elem : MDirtyElements)
 		{
+			if (PrePreFilterHelper(Elem.Payload, QueryData))
+			{
+				continue;
+			}
+
 			const TAABB<T, d>& InstanceBounds = Elem.Bounds;
 			const TVector<T, d> Min = InstanceBounds.Min() - QueryHalfExtents;
 			const TVector<T, d> Max = InstanceBounds.Max() + QueryHalfExtents;
@@ -749,8 +772,14 @@ private:
 	template <typename SQVisitor, bool bPruneDuplicates = true>
 	bool OverlapImp(const TAABB<T, d>& QueryBounds, SQVisitor& Visitor) const
 	{
+		const void* QueryData = Visitor.GetQueryData();
 		for (const auto& Elem : MGlobalPayloads)
 		{
+			if (PrePreFilterHelper(Elem.Payload, QueryData))
+			{
+				continue;
+			}
+
 			const TAABB<T, d>& InstanceBounds = Elem.Bounds;
 			if (QueryBounds.Intersects(InstanceBounds))
 			{
@@ -764,6 +793,11 @@ private:
 
 		for (const auto& Elem : MDirtyElements)
 		{
+			if (PrePreFilterHelper(Elem.Payload, QueryData))
+			{
+				continue;
+			}
+
 			const TAABB<T, d>& InstanceBounds = Elem.Bounds;
 			if (QueryBounds.Intersects(InstanceBounds))
 			{
diff --git a/Engine/Source/Runtime/Experimental/Chaos/Public/Chaos/Collision/SpatialAccelerationBroadPhase.h b/Engine/Source/Runtime/Experimental/Chaos/Public/Chaos/Collision/SpatialAccelerationBroadPhase.h
--- a/Engine/Source/Runtime/Experimental/Chaos/Public/Chaos/Collision/SpatialAccelerationBroadPhase.h
+++ b/Engine/Source/Runtime/Experimental/Chaos/Public/Chaos/Collision/SpatialAccelerationBroadPhase.h
@@ -44,6 +44,8 @@ namespace Chaos
 			check(false);
 			return false;
 		}
+
+		const void* GetQueryData() const { return nullptr; }
 	};
 
 	/**
diff --git a/Engine/Source/Runtime/Experimental/Chaos/Public/Chaos/ISpatialAcceleration.h b/Engine/Source/Runtime/Experimental/Chaos/Public/Chaos/ISpatialAcceleration.h
--- a/Engine/Source/Runtime/Experimental/Chaos/Public/Chaos/ISpatialAcceleration.h
+++ b/Engine/Source/Runtime/Experimental/Chaos/Public/Chaos/ISpatialAcceleration.h
@@ -139,6 +139,8 @@ public:
 		Returns true to continue iterating through the acceleration structure
 	*/
 	virtual bool Sweep(const TSpatialVisitorData<TPayloadType>& Instance, FQueryFastData& CurData) = 0;
+
+	virtual const void* GetQueryData() const { return nullptr; }
 };
 
 /**
@@ -360,6 +362,11 @@ public:
 		return Visitor.Sweep(Instance, CurData);
 	}
 
+	FORCEINLINE const void* GetQueryData() const
+	{
+		return Visitor.GetQueryData();
+	}
+
 private:
 	ISpatialVisitor<TPayloadType, T>& Visitor;
 };
@@ -501,4 +508,23 @@ FChaosArchive& operator<< (FChaosArchive& Ar, TArrayAsMap<TKey, TValue>& Map)
 	return Ar;
 }
 
+
+template <typename TPayload>
+typename TEnableIf<!TIsPointer<TPayload>::Value, bool>::Type PrePreFilterHelper(const TPayload& Payload, const void* QueryData)
+{
+	return Payload.PrePreFilter(QueryData);
+}
+
+template <typename TPayload>
+typename TEnableIf<TIsPointer<TPayload>::Value, bool>::Type PrePreFilterHelper(const TPayload& Payload, const void* QueryData)
+{
+	return false;
+}
+
+FORCEINLINE bool PrePreFilterHelper(const int32 Payload, const void* QueryData)
+{
+	return false;
+}
+
+
 }
diff --git a/Engine/Source/Runtime/Experimental/Chaos/Public/Chaos/ParticleHandle.h b/Engine/Source/Runtime/Experimental/Chaos/Public/Chaos/ParticleHandle.h
--- a/Engine/Source/Runtime/Experimental/Chaos/Public/Chaos/ParticleHandle.h
+++ b/Engine/Source/Runtime/Experimental/Chaos/Public/Chaos/ParticleHandle.h
@@ -126,6 +126,30 @@ bool GeometryParticleSleeping(const FConcrete& Concrete)
 	return Concrete.ObjectState() == EObjectStateType::Sleeping;
 }
 
+//Used to filter out at the acceleration structure layer
+//Returns true when there is no way a later PreFilter will succeed. Avoid virtuals etc..
+FORCEINLINE_DEBUGGABLE bool PrePreFilterImp(const FCollisionFilterData& QueryFilterData, const FCollisionFilterData& UnionFilterData)
+{
+	//HACK: need to replace all these hard-coded values with proper enums, bad modules are not setup for it right now
+	//ECollisionQuery QueryType = (ECollisionQuery)QueryFilter.Word0;
+	//if (QueryType != ECollisionQuery::ObjectQuery)
+	if(QueryFilterData.Word0)
+	{
+		//since we're taking the union of shapes we can only support trace channel
+		//const ECollisionChannel QuerierChannel = GetCollisionChannel(QueryFilter.Word3);
+		const uint32 QuerierChannel = (QueryFilterData.Word3 << 6) >> (32 - 5);
+
+		//uint32 const QuerierBit = ECC_TO_BITFIELD(QuerierChannel);
+		const uint32 QuerierBit = (1 << (QuerierChannel));
+
+		// check if Querier wants a hit
+		const uint32 TouchOrBlock = (UnionFilterData.Word1 | UnionFilterData.Word2);
+		return !(QuerierBit & TouchOrBlock);
+	}
+
+	return false;
+}
+
 /** Wrapper that holds both physics thread data and GT data. It's possible that the physics handle is null if we're doing operations entirely on external threads*/
 template <typename T, int d>
 class TAccelerationStructureHandle
@@ -137,10 +161,6 @@ public:
 	template <bool bPersistent>
 	TAccelerationStructureHandle(TGeometryParticleHandleImp<T, d, bPersistent>& InHandle);
 
-	template <bool bPersistent>
-	TAccelerationStructureHandle(TGeometryParticleHandleImp<T, d, bPersistent>* InHandle, TGeometryParticle<T, d>* InGeometryParticle);
-
-
 	//Should only be used by GT and scene query threads where an appropriate lock has been acquired
 	TGeometryParticle<T, d>* GetExternalGeometryParticle_ExternalThread() const { return ExternalGeometryParticle; }
 
@@ -174,11 +194,29 @@ public:
 		return CachedUniqueIdx;
 	}
 
+	bool PrePreFilter(const void* QueryData) const
+	{
+		if(bCanPrePreFilter)
+		{
+			if (const FCollisionFilterData* QueryFilterData = static_cast<const FCollisionFilterData*>(QueryData))
+			{
+				return PrePreFilterImp(*QueryFilterData, UnionFilterData);
+			}
+		}
+		
+		return false;
+	}
+
 private:
 	TGeometryParticle<T, d>* ExternalGeometryParticle;
 	TGeometryParticleHandle<T, d>* GeometryParticleHandle;
 
 	FUniqueIdx CachedUniqueIdx;
+	FCollisionFilterData UnionFilterData;
+	bool bCanPrePreFilter;
+
+	template <typename TParticle>
+	void UpdatePrePreFilter(const TParticle& Particle);
 };
 
 template <typename T, int d>
@@ -1587,11 +1625,13 @@ public:
 		ShapeCollisionDisableFlags.Empty(Shapes.Num());
 		CollisionTraceType.Empty(Shapes.Num());
 		ShapeSimData.Empty(Shapes.Num());
+		ShapeQueryData.Empty(Shapes.Num());
 		for (const TUniquePtr<TPerShapeData<T, d>>& ShapePtr : Shapes)
 		{
 			ShapeCollisionDisableFlags.Add(ShapePtr->bDisable);
 			CollisionTraceType.Add(ShapePtr->CollisionTraceType);
 			ShapeSimData.Add(ShapePtr->SimData);
+			ShapeQueryData.Add(ShapePtr->QueryData);
 		}
 	}
 
@@ -1608,6 +1648,7 @@ public:
 		ShapeCollisionDisableFlags.Reset();
 		CollisionTraceType.Reset();
 		ShapeSimData.Reset();
+		ShapeQueryData.Reset();
 #if CHAOS_CHECKED
 		DebugName = NAME_None;
 #endif
@@ -1623,6 +1664,7 @@ public:
 	TBitArray<> ShapeCollisionDisableFlags;
 	TArray<EChaosCollisionTraceFlag> CollisionTraceType;
 	TArray < FCollisionFilterData > ShapeSimData;
+	TArray < FCollisionFilterData > ShapeQueryData;
 #if CHAOS_CHECKED
 	FName DebugName;
 #endif
@@ -2222,8 +2264,13 @@ TAccelerationStructureHandle<T,d>::TAccelerationStructureHandle(TGeometryParticl
 	: ExternalGeometryParticle(InHandle->GTGeometryParticle())
 	, GeometryParticleHandle(InHandle)
 	, CachedUniqueIdx(InHandle->UniqueIdx())
+	, bCanPrePreFilter(false)
 {
 	ensure(CachedUniqueIdx.IsValid());
+	if (InHandle)
+	{
+		UpdatePrePreFilter(*InHandle);
+	}
 }
 
 template <typename T, int d>
@@ -2231,8 +2278,14 @@ TAccelerationStructureHandle<T,d>::TAccelerationStructureHandle(TGeometryParticl
 	: ExternalGeometryParticle(InGeometryParticle)
 	, GeometryParticleHandle(InGeometryParticle ? InGeometryParticle->Handle() : nullptr)
 	, CachedUniqueIdx(InGeometryParticle ? InGeometryParticle->UniqueIdx() : FUniqueIdx())
+	, bCanPrePreFilter(false)
 {
-	ensure(!InGeometryParticle || CachedUniqueIdx.IsValid());
+	if (InGeometryParticle)
+	{
+		ensure(CachedUniqueIdx.IsValid());
+		ensure(IsInGameThread());
+		UpdatePrePreFilter(*InGeometryParticle);
+	}
 }
 
 template <typename T, int d>
@@ -2241,34 +2294,30 @@ TAccelerationStructureHandle<T, d>::TAccelerationStructureHandle(TGeometryPartic
 	: ExternalGeometryParticle(InHandle.GTGeometryParticle())
 	, GeometryParticleHandle(InHandle.Handle())
 	, CachedUniqueIdx(InHandle.UniqueIdx())
+	, bCanPrePreFilter(false)
 {
 	ensure(CachedUniqueIdx.IsValid());
+	UpdatePrePreFilter(InHandle);
 }
 
 template <typename T, int d>
-template <bool bPersistent>
-TAccelerationStructureHandle<T, d>::TAccelerationStructureHandle(TGeometryParticleHandleImp<T, d, bPersistent>* InHandle, TGeometryParticle<T, d>* InGeometryParticle)
-	: ExternalGeometryParticle(InGeometryParticle)
-	, GeometryParticleHandle(InHandle)
+template <typename TParticle>
+void TAccelerationStructureHandle<T, d>::UpdatePrePreFilter(const TParticle& Particle)
 {
-	if (GeometryParticleHandle)
+	const auto& Shapes = Particle.ShapesArray();
+	for (const auto& Shape : Shapes)
 	{
-		CachedUniqueIdx = GeometryParticleHandle->UniqueIdx();
-	}
-	else if (ExternalGeometryParticle)
-	{
-		CachedUniqueIdx = ExternalGeometryParticle->UniqueIdx();
-	}
-
-	if (GeometryParticleHandle && ExternalGeometryParticle)
-	{
-		ensure(GeometryParticleHandle->UniqueIdx() == ExternalGeometryParticle->UniqueIdx());
+		UnionFilterData.Word0 |= Shape->QueryData.Word0;
+		UnionFilterData.Word1 |= Shape->QueryData.Word1;
+		UnionFilterData.Word2 |= Shape->QueryData.Word2;
+		UnionFilterData.Word3 |= Shape->QueryData.Word3;
 	}
+	
+	bCanPrePreFilter = true;
 
-	//If either handle or particle are valid, we must have a valid idx
-	ensure((!GeometryParticleHandle && !ExternalGeometryParticle) || CachedUniqueIdx.IsValid());
 }
 
+
 template <typename T, int d>
 void TAccelerationStructureHandle<T, d>::Serialize(FChaosArchive& Ar)
 {
diff --git a/Engine/Source/Runtime/Experimental/ChaosSolvers/Private/PhysicsProxy/SingleParticlePhysicsProxy.cpp b/Engine/Source/Runtime/Experimental/ChaosSolvers/Private/PhysicsProxy/SingleParticlePhysicsProxy.cpp
--- a/Engine/Source/Runtime/Experimental/ChaosSolvers/Private/PhysicsProxy/SingleParticlePhysicsProxy.cpp
+++ b/Engine/Source/Runtime/Experimental/ChaosSolvers/Private/PhysicsProxy/SingleParticlePhysicsProxy.cpp
@@ -91,7 +91,8 @@ void FSingleParticlePhysicsProxy<Chaos::TGeometryParticle<float, 3>>::PushToPhys
 			int32 CurrShape = 0;
 			for (const TUniquePtr<Chaos::TPerShapeData<Chaos::FReal, 3>>& Shape : RigidHandle->ShapesArray())
 			{
-				Shape->SimData = Data->ShapeSimData[CurrShape++];
+				Shape->SimData = Data->ShapeSimData[CurrShape];
+				Shape->QueryData = Data->ShapeQueryData[CurrShape++];
 			}
 		}
 
@@ -200,7 +201,8 @@ void FSingleParticlePhysicsProxy<Chaos::TKinematicGeometryParticle<float, 3>>::P
 			int32 CurrShape = 0;
 			for (const TUniquePtr<Chaos::TPerShapeData<Chaos::FReal, 3>>& Shape : RigidHandle->ShapesArray())
 			{
-				Shape->SimData = Data->ShapeSimData[CurrShape++];
+				Shape->SimData = Data->ShapeSimData[CurrShape];
+				Shape->QueryData = Data->ShapeQueryData[CurrShape++];
 			}
 		}
 
@@ -377,7 +379,8 @@ void FSingleParticlePhysicsProxy<Chaos::TPBDRigidParticle<float, 3>>::PushToPhys
 			int32 CurrShape = 0;
 			for (const TUniquePtr<Chaos::TPerShapeData<Chaos::FReal, 3>>& Shape : RigidHandle->ShapesArray())
 			{
-				Shape->SimData = Data->ShapeSimData[CurrShape++];
+				Shape->SimData = Data->ShapeSimData[CurrShape];
+				Shape->QueryData = Data->ShapeQueryData[CurrShape++];
 			}
 		}
 
diff --git a/Engine/Source/Runtime/PhysicsSQ/Private/SQAccelerator.cpp b/Engine/Source/Runtime/PhysicsSQ/Private/SQAccelerator.cpp
--- a/Engine/Source/Runtime/PhysicsSQ/Private/SQAccelerator.cpp
+++ b/Engine/Source/Runtime/PhysicsSQ/Private/SQAccelerator.cpp
@@ -94,6 +94,7 @@ struct TSQVisitor : public Chaos::ISpatialVisitor<TPayload, float>
 		, DebugParams(InDebugParams)
 		, HitBuffer(InHitBuffer)
 		, QueryFilterData(InQueryFilterData)
+		, QueryFilterDataConcrete(P2UFilterData(QueryFilterData.data))
 		, QueryCallback(InQueryCallback)
 	{
 #if WITH_PHYSX
@@ -111,6 +112,7 @@ struct TSQVisitor : public Chaos::ISpatialVisitor<TPayload, float>
 		, DebugParams(InDebugParams)
 		, HitBuffer(InHitBuffer)
 		, QueryFilterData(InQueryFilterData)
+		, QueryFilterDataConcrete(P2UFilterData(QueryFilterData.data))
 		, QueryGeom(&InQueryGeom)
 		, QueryCallback(InQueryCallback)
 		, StartTM(InStartTM)
@@ -129,6 +131,7 @@ struct TSQVisitor : public Chaos::ISpatialVisitor<TPayload, float>
 		, DebugParams(InDebugParams)
 		, HitBuffer(InHitBuffer)
 		, QueryFilterData(InQueryFilterData)
+		, QueryFilterDataConcrete(P2UFilterData(QueryFilterData.data))
 		, QueryGeom(&InQueryGeom)
 		, QueryCallback(InQueryCallback)
 		, StartTM(InWorldTM)
@@ -155,6 +158,11 @@ struct TSQVisitor : public Chaos::ISpatialVisitor<TPayload, float>
 		return Visit<ESQType::Overlap>(Instance, nullptr);
 	}
 
+	virtual const void* GetQueryData() const override
+	{
+		return &QueryFilterData;
+	}
+
 private:
 
 	enum class ESQType
@@ -167,6 +175,7 @@ private:
 	template <ESQType SQ>
 	bool Visit(const Chaos::TSpatialVisitorData<TPayload>& Instance, Chaos::FQueryFastData* CurData)
 	{
+		//QUICK_SCOPE_CYCLE_COUNTER(SQVisit);
 #if CHAOS_DEBUG_DRAW && WITH_CHAOS
 		if (DebugParams.IsDebugQuery() && ChaosSQDrawDebugVisitorQueries)
 		{
@@ -223,7 +232,7 @@ private:
 			//TODO: use gt particles directly
 			//#TODO alternative to px flags
 #if WITH_PHYSX
-			ECollisionQueryHitType HitType = QueryFilterData.flags & PxQueryFlag::ePREFILTER ? QueryCallback.PreFilter(P2UFilterData(QueryFilterData.data), *Shape, *GeometryParticle) : ECollisionQueryHitType::Block;
+			ECollisionQueryHitType HitType = QueryFilterData.flags & PxQueryFlag::ePREFILTER ? QueryCallback.PreFilter(QueryFilterDataConcrete, *Shape, *GeometryParticle) : ECollisionQueryHitType::Block;
 #else
 			//#TODO Chaos flag alternative
 			ensure(false);
@@ -231,7 +240,7 @@ private:
 #endif
 			if (HitType != ECollisionQueryHitType::None)
 			{
-
+				//QUICK_SCOPE_CYCLE_COUNTER(SQNarrow);
 				THitType Hit;
 				Hit.Actor = GeometryParticle;
 				Hit.Shape = Shape.Get();
@@ -280,9 +289,10 @@ private:
 
 				if(bHit)
 				{
+					//QUICK_SCOPE_CYCLE_COUNTER(SQNarrowHit);
 					FillHitHelper(Hit, Distance, WorldPosition, WorldNormal, FaceIdx, bComputeMTD);
 #if WITH_PHYSX
-					HitType = QueryFilterData.flags & PxQueryFlag::ePOSTFILTER ? QueryCallback.PostFilter(P2UFilterData(QueryFilterData.data), Hit) : HitType;
+					HitType = QueryFilterData.flags & PxQueryFlag::ePOSTFILTER ? QueryCallback.PostFilter(QueryFilterDataConcrete, Hit) : HitType;
 #else
 					//#TODO Chaos flag alternative
 					ensure(false);
@@ -340,6 +350,7 @@ private:
 	const FQueryDebugParams DebugParams;
 	ChaosInterface::FSQHitBuffer<THitType>& HitBuffer;
 	const FQueryFilterData& QueryFilterData;
+	const FCollisionFilterData QueryFilterDataConcrete;
 	const QueryGeometryType* QueryGeom;
 	ICollisionQueryFilterCallbackBase& QueryCallback;
 	const FTransform StartTM;
