Fix issues with scaled convex against trimesh sweeps and overlaps.
Add unit tests.

diff --git a/Engine/Source/Runtime/Experimental/Chaos/Private/Chaos/TriangleMeshImplicitObject.cpp b/Engine/Source/Runtime/Experimental/Chaos/Private/Chaos/TriangleMeshImplicitObject.cpp
--- a/Engine/Source/Runtime/Experimental/Chaos/Private/Chaos/TriangleMeshImplicitObject.cpp
+++ b/Engine/Source/Runtime/Experimental/Chaos/Private/Chaos/TriangleMeshImplicitObject.cpp
@@ -29,38 +29,11 @@ static auto MakeScaledHelper(const TImplicitObjectScaled<QueryGeomType>& B, cons
 	return ScaledB;
 }
 
-
-template <typename QueryGeomType, typename IdxType>
-void TransformVertsHelper(const QueryGeomType& QueryGeom, int32 TriIdx, const TParticles<FReal, 3>& Particles,
-	const TArray<TVector<IdxType, 3>>& Elements, TVec3<FReal>& OutA, TVec3<FReal>& OutB, TVec3<FReal>& OutC)
-{
-	OutA = Particles.X(Elements[TriIdx][0]);
-	OutB = Particles.X(Elements[TriIdx][1]);
-	OutC = Particles.X(Elements[TriIdx][2]);
-}
-
-template <typename QueryGeomType, typename IdxType>
-void TransformVertsHelper(const TImplicitObjectScaled<QueryGeomType>& QueryGeom, int32 TriIdx, const TParticles<FReal, 3>& Particles,
-	const TArray<TVector<IdxType, 3>>& Elements, TVec3<FReal>& OutA, TVec3<FReal>& OutB, TVec3<FReal>& OutC)
-{
-	const TVec3<FReal> InvScale = QueryGeom.GetInvScale();
-	OutA = Particles.X(Elements[TriIdx][0]) * InvScale;
-	OutB = Particles.X(Elements[TriIdx][1]) * InvScale;
-	OutC = Particles.X(Elements[TriIdx][2]) * InvScale;
-}
-
-template <typename QueryGeomType>
-const QueryGeomType& GetGeomHelper(const QueryGeomType& QueryGeom)
-{
-	return QueryGeom;
-}
-
 void ScaleTransformHelper(const FVec3& TriMeshScale, const FRigidTransform3& QueryTM, FRigidTransform3& OutScaledQueryTM)
 {
 	OutScaledQueryTM = TRigidTransform<FReal, 3>(QueryTM.GetLocation() * TriMeshScale, QueryTM.GetRotation());
 }
 
-
 template <typename IdxType>
 void TransformVertsHelper(const FVec3& TriMeshScale, int32 TriIdx, const TParticles<FReal, 3>& Particles,
 	const TArray<TVector<IdxType, 3>>& Elements, TVec3<FReal>& OutA, TVec3<FReal>& OutB, TVec3<FReal>& OutC)
@@ -83,41 +56,19 @@ TImplicitObjectScaled<QueryGeomType> ScaleGeomIntoWorldHelper(const TImplicitObj
 	return MakeScaledHelper(QueryGeom, TriMeshScale);
 }
 
-template <typename QueryGeomType>
-void TransformSweepOutputsHelper(const QueryGeomType& QueryGeom, const TVec3<FReal>& HitNormal, const TVec3<FReal>& HitPosition, const FReal LengthScale,
-	const FReal Time, TVec3<FReal>& OutNormal, TVec3<FReal>& OutPosition, FReal& OutTime)
-{
-	OutNormal = HitNormal;
-	OutPosition = HitPosition;
-	OutTime = Time;
-}
-
-template <typename QueryGeomType>
-void TransformSweepOutputsHelper(const TImplicitObjectScaled<QueryGeomType>& QueryGeom, const TVec3<FReal>& HitNormal, const TVec3<FReal>& HitPosition, const FReal LengthScale,
+void TransformSweepOutputsHelper(FVec3 TriMeshScale, const TVec3<FReal>& HitNormal, const TVec3<FReal>& HitPosition, const FReal LengthScale,
 	const FReal Time, TVec3<FReal>& OutNormal, TVec3<FReal>& OutPosition, FReal& OutTime)
 {
-	const TVec3<FReal> InvScale = QueryGeom.GetInvScale();
-	const TVec3<FReal> Scale = QueryGeom.GetScale();
-
-	OutTime = Time / LengthScale;
-	OutNormal = (InvScale * HitNormal).GetSafeNormal();
-	OutPosition = Scale * HitPosition;
-}
-
-template <typename QueryGeomType>
-void TransformOverlapInputsHelper(const QueryGeomType& QueryGeom, const FRigidTransform3& QueryTM, FRigidTransform3& OutScaledQueryTM)
-{
-	OutScaledQueryTM = QueryTM;
-}
+	if (ensure(TriMeshScale != FVec3(0.0f)))
+	{
+		FVec3 InvTriMeshScale = 1.0f / TriMeshScale;
 
-template <typename QueryGeomType>
-void TransformOverlapInputsHelper(const TImplicitObjectScaled<QueryGeomType>& QueryGeom, const FRigidTransform3& QueryTM, FRigidTransform3& OutScaledQueryTM)
-{
-	const TVec3<FReal> InvScale = QueryGeom.GetInvScale();
-	OutScaledQueryTM = TRigidTransform<FReal, 3>(QueryTM.GetLocation() * InvScale, QueryTM.GetRotation());
+		OutTime = Time / LengthScale;
+		OutNormal = (TriMeshScale * HitNormal).GetSafeNormal();
+		OutPosition = InvTriMeshScale * HitPosition;
+	}
 }
 
-
 template <typename IdxType>
 struct FTriangleMeshRaycastVisitor
 {
@@ -554,7 +505,7 @@ bool FTriangleMeshImplicitObject::Overlap(const FVec3& Point, const FReal Thickn
 }
 
 template <typename QueryGeomType>
-bool FTriangleMeshImplicitObject::OverlapGeomImp(const QueryGeomType& QueryGeom, const FRigidTransform3& QueryTM, const FReal Thickness, FMTDInfo* OutMTD) const
+bool FTriangleMeshImplicitObject::OverlapGeomImp(const QueryGeomType& QueryGeom, const FRigidTransform3& QueryTM, const FReal Thickness, FMTDInfo* OutMTD, FVec3 TriMeshScale) const
 {
 	bool bResult = false;
 	TAABB<FReal, 3> QueryBounds = QueryGeom.BoundingBox();
@@ -568,10 +519,10 @@ bool FTriangleMeshImplicitObject::OverlapGeomImp(const QueryGeomType& QueryGeom,
 		OutMTD->Penetration = TNumericLimits<FReal>::Lowest();
 	}
 
-	const auto& InnerQueryGeom = GetGeomHelper(QueryGeom);
+	const auto& WorldScaleQueryGeom = ScaleGeomIntoWorldHelper(QueryGeom, TriMeshScale);
 
-	TRigidTransform<FReal, 3> TransformedQueryTM;
-	TransformOverlapInputsHelper(QueryGeom, QueryTM, TransformedQueryTM);
+	TRigidTransform<FReal, 3> WorldScaleQueryTM;
+	ScaleTransformHelper(TriMeshScale, QueryTM, WorldScaleQueryTM);
 
 	auto LambdaHelper = [&](const auto& Elements, FMTDInfo* InnerMTD)
 	{
@@ -581,13 +532,13 @@ bool FTriangleMeshImplicitObject::OverlapGeomImp(const QueryGeomType& QueryGeom,
 			for (int32 TriIdx : PotentialIntersections)
 			{
 				TVec3<FReal> A, B, C;
-				TransformVertsHelper(QueryGeom, TriIdx, MParticles, Elements, A, B, C);
+				TransformVertsHelper(TriMeshScale, TriIdx, MParticles, Elements, A, B, C);
 
 				FVec3 TriangleNormal(0.0);
 				FReal Penetration = 0.0;
 				FVec3 ClosestA(0.0);
 				FVec3 ClosestB(0.0);
-				if (GJKPenetration(TTriangle<FReal>(A, B, C), InnerQueryGeom, TransformedQueryTM, Penetration, ClosestA, ClosestB, TriangleNormal, Thickness))
+				if (GJKPenetration(TTriangle<FReal>(A, B, C), WorldScaleQueryGeom, WorldScaleQueryTM, Penetration, ClosestA, ClosestB, TriangleNormal, Thickness))
 				{
 					bOverlap = true;
 
@@ -607,7 +558,7 @@ bool FTriangleMeshImplicitObject::OverlapGeomImp(const QueryGeomType& QueryGeom,
 			for (int32 TriIdx : PotentialIntersections)
 			{
 				TVec3<FReal> A, B, C;
-				TransformVertsHelper(QueryGeom, TriIdx, MParticles, Elements, A, B, C);
+				TransformVertsHelper(TriMeshScale, TriIdx, MParticles, Elements, A, B, C);
 
 				const FVec3 AB = B - A;
 				const FVec3 AC = C - A;
@@ -616,7 +567,7 @@ bool FTriangleMeshImplicitObject::OverlapGeomImp(const QueryGeomType& QueryGeom,
 				//However, maybe we should check if it's behind the triangle plane. Also, we should enforce this winding in some way
 				const FVec3 Offset = FVec3::CrossProduct(AB, AC);
 
-				if (GJKIntersection(TTriangle<FReal>(A, B, C), InnerQueryGeom, TransformedQueryTM, Thickness, Offset))
+				if (GJKIntersection(TTriangle<FReal>(A, B, C), WorldScaleQueryGeom, WorldScaleQueryTM, Thickness, Offset))
 				{
 					return true;
 				}
@@ -656,37 +607,31 @@ bool FTriangleMeshImplicitObject::OverlapGeom(const FConvex& QueryGeom, const FR
 	return OverlapGeomImp(QueryGeom, QueryTM, Thickness, OutMTD);
 }
 
-bool FTriangleMeshImplicitObject::OverlapGeom(const TImplicitObjectScaled<TSphere<FReal, 3>>& QueryGeom, const FRigidTransform3& QueryTM, const FReal Thickness, FMTDInfo* OutMTD) const
+bool FTriangleMeshImplicitObject::OverlapGeom(const TImplicitObjectScaled<TSphere<FReal, 3>>& QueryGeom, const FRigidTransform3& QueryTM, const FReal Thickness, FMTDInfo* OutMTD, FVec3 TriMeshScale) const
 {
-	return OverlapGeomImp(QueryGeom, QueryTM, Thickness, OutMTD);
+	return OverlapGeomImp(QueryGeom, QueryTM, Thickness, OutMTD, TriMeshScale);
 }
 
-bool FTriangleMeshImplicitObject::OverlapGeom(const TImplicitObjectScaled<TBox<FReal, 3>>& QueryGeom, const FRigidTransform3& QueryTM, const FReal Thickness, FMTDInfo* OutMTD) const
+bool FTriangleMeshImplicitObject::OverlapGeom(const TImplicitObjectScaled<TBox<FReal, 3>>& QueryGeom, const FRigidTransform3& QueryTM, const FReal Thickness, FMTDInfo* OutMTD, FVec3 TriMeshScale) const
 {
-	return OverlapGeomImp(QueryGeom, QueryTM, Thickness, OutMTD);
-}
-
-bool FTriangleMeshImplicitObject::OverlapGeom(const TImplicitObjectScaled<TCapsule<FReal>>& QueryGeom, const FRigidTransform3& QueryTM, const FReal Thickness, FMTDInfo* OutMTD) const
-{
-	return OverlapGeomImp(QueryGeom, QueryTM, Thickness, OutMTD);
+	return OverlapGeomImp(QueryGeom, QueryTM, Thickness, OutMTD, TriMeshScale);
 }
 
-bool FTriangleMeshImplicitObject::OverlapGeom(const TImplicitObjectScaled<FConvex>& QueryGeom, const FRigidTransform3& QueryTM, const FReal Thickness, FMTDInfo* OutMTD) const
+bool FTriangleMeshImplicitObject::OverlapGeom(const TImplicitObjectScaled<TCapsule<FReal>>& QueryGeom, const FRigidTransform3& QueryTM, const FReal Thickness, FMTDInfo* OutMTD, FVec3 TriMeshScale) const
 {
-	return OverlapGeomImp(QueryGeom, QueryTM, Thickness, OutMTD);
+	return OverlapGeomImp(QueryGeom, QueryTM, Thickness, OutMTD, TriMeshScale);
 }
 
-bool FTriangleMeshImplicitObject::OverlapGeom(const TImplicitObjectScaled<TImplicitObjectScaled<FConvex>>& QueryGeom, const FRigidTransform3& QueryTM, const FReal Thickness, FMTDInfo* OutMTD) const
+bool FTriangleMeshImplicitObject::OverlapGeom(const TImplicitObjectScaled<FConvex>& QueryGeom, const FRigidTransform3& QueryTM, const FReal Thickness, FMTDInfo* OutMTD, FVec3 TriMeshScale) const
 {
-	return OverlapGeomImp(QueryGeom, QueryTM, Thickness, OutMTD);
+	return OverlapGeomImp(QueryGeom, QueryTM, Thickness, OutMTD, TriMeshScale);
 }
 
-
 template <typename QueryGeomType, typename IdxType>
 struct FTriangleMeshSweepVisitor
 {
 	FTriangleMeshSweepVisitor(const FTriangleMeshImplicitObject& InTriMesh, const TArray<TVec3<IdxType>>& InElements, const QueryGeomType& InQueryGeom, const TRigidTransform<FReal,3>& InStartTM, const FVec3& InDir,
-		const FVec3& InScaledDirNormalized, const FReal InLengthScale, const FRigidTransform3& InScaledStartTM, const FReal InThickness, const bool InComputeMTD)
+		const FVec3& InScaledDirNormalized, const FReal InLengthScale, const FRigidTransform3& InScaledStartTM, const FReal InThickness, const bool InComputeMTD, FVec3 InTriMeshScale)
 	: TriMesh(InTriMesh)
 	, Elements(InElements)
 	, StartTM(InStartTM)
@@ -698,6 +643,7 @@ struct FTriangleMeshSweepVisitor
 	, LengthScale(InLengthScale)
 	, ScaledStartTM(InScaledStartTM)
 	, OutTime(TNumericLimits<FReal>::Max())
+	, TriMeshScale(InTriMeshScale)
 	{
 	}
 
@@ -724,16 +670,16 @@ struct FTriangleMeshSweepVisitor
 		FVec3 HitNormal;
 
 		FVec3 A, B, C;
-		TransformVertsHelper(QueryGeom,TriIdx,TriMesh.MParticles,Elements,A,B,C);
+		TransformVertsHelper(TriMeshScale,TriIdx,TriMesh.MParticles,Elements,A,B,C);
 		TTriangle<FReal> Tri(A, B, C);
 
-		const auto& InnerQueryGeom = GetGeomHelper(QueryGeom);
+		const auto& WorldScaleQueryGeom = ScaleGeomIntoWorldHelper(QueryGeom, TriMeshScale);
 
-		if(GJKRaycast2<FReal>(Tri, InnerQueryGeom, ScaledStartTM, ScaledDirNormalized, LengthScale * CurData.CurrentLength, Time, HitPosition, HitNormal, Thickness, bComputeMTD))
+		if(GJKRaycast2<FReal>(Tri, WorldScaleQueryGeom, ScaledStartTM, ScaledDirNormalized, LengthScale * CurData.CurrentLength, Time, HitPosition, HitNormal, Thickness, bComputeMTD))
 		{
 			if(Time < OutTime)
 			{
-				TransformSweepOutputsHelper(QueryGeom, HitNormal, HitPosition, LengthScale, Time, OutNormal, OutPosition, OutTime);
+				TransformSweepOutputsHelper(TriMeshScale, HitNormal, HitPosition, LengthScale, Time, OutNormal, OutPosition, OutTime);
 
 				OutFaceIndex = TriIdx;
 
@@ -769,24 +715,14 @@ struct FTriangleMeshSweepVisitor
 	FVec3 OutPosition;
 	FVec3 OutNormal;
 	int32 OutFaceIndex;
-};
 
-template <typename QueryGeomType>
-void ComputeScaledSweepInputs(const QueryGeomType& QueryGeom, const FRigidTransform3& StartTM, const FVec3& Dir, const FReal Length,
-	FVec3& OutScaledDirNormalized, FReal& OutLengthScale, FRigidTransform3& OutScaledStartTM)
-{
-	OutScaledDirNormalized = Dir;
-	OutLengthScale = 1.0f;
-	OutScaledStartTM = StartTM;
-}
+	FVec3 TriMeshScale;
+};
 
-template<typename QueryGeomType>
-void ComputeScaledSweepInputs(const TImplicitObjectScaled<QueryGeomType>& QueryGeom, const FRigidTransform3& StartTM, const FVec3& Dir, const FReal Length,
+void ComputeScaledSweepInputs(FVec3 TriMeshScale, const FRigidTransform3& StartTM, const FVec3& Dir, const FReal Length,
 	FVec3& OutScaledDirNormalized, FReal& OutLengthScale, FRigidTransform3& OutScaledStartTM)
 {
-	const FVec3& InvScale = QueryGeom.GetInvScale();
-
-	const FVec3 UnscaledDirDenorm = InvScale * Dir;
+	const FVec3 UnscaledDirDenorm = TriMeshScale * Dir;
 	const FReal LengthScale = UnscaledDirDenorm.Size();
 	if (CHAOS_ENSURE(LengthScale > TNumericLimits<FReal>::Min()))
 	{
@@ -796,25 +732,25 @@ void ComputeScaledSweepInputs(const TImplicitObjectScaled<QueryGeomType>& QueryG
 
 
 	OutLengthScale = LengthScale;
-	OutScaledStartTM = FRigidTransform3(StartTM.GetLocation() * InvScale, StartTM.GetRotation());
+	OutScaledStartTM = FRigidTransform3(StartTM.GetLocation() * TriMeshScale, StartTM.GetRotation());
 }
 
 template <typename QueryGeomType>
 bool FTriangleMeshImplicitObject::SweepGeomImp(const QueryGeomType& QueryGeom, const FRigidTransform3& StartTM, const FVec3& Dir, const FReal Length,
-	FReal& OutTime, FVec3& OutPosition, FVec3& OutNormal, int32& OutFaceIndex, const FReal Thickness, const bool bComputeMTD) const
+	FReal& OutTime, FVec3& OutPosition, FVec3& OutNormal, int32& OutFaceIndex, const FReal Thickness, const bool bComputeMTD, FVec3 TriMeshScale) const
 {
 	//QUICK_SCOPE_CYCLE_COUNTER(TrimeshSweep);
 	// Compute scaled sweep inputs to cache in visitor.
 	FVec3 ScaledDirNormalized;
 	FReal LengthScale;
 	FRigidTransform3 ScaledStartTM;
-	ComputeScaledSweepInputs(QueryGeom, StartTM, Dir, Length, ScaledDirNormalized, LengthScale, ScaledStartTM);
+	ComputeScaledSweepInputs(TriMeshScale, StartTM, Dir, Length, ScaledDirNormalized, LengthScale, ScaledStartTM);
 
 	bool bHit = false;
 	auto LambdaHelper = [&](const auto& Elements)
 	{
 		using VisitorType = FTriangleMeshSweepVisitor<QueryGeomType,decltype(Elements[0][0])>;
-		VisitorType SQVisitor(*this,Elements, QueryGeom,StartTM,Dir,ScaledDirNormalized,LengthScale,ScaledStartTM,Thickness,bComputeMTD);
+		VisitorType SQVisitor(*this,Elements, QueryGeom,StartTM,Dir,ScaledDirNormalized,LengthScale,ScaledStartTM,Thickness,bComputeMTD, TriMeshScale);
 
 
 		const TAABB<FReal,3> QueryBounds = QueryGeom.BoundingBox().TransformedAABB(FRigidTransform3(FVec3::ZeroVector,StartTM.GetRotation()));
@@ -863,24 +799,24 @@ bool FTriangleMeshImplicitObject::SweepGeom(const FConvex& QueryGeom, const FRig
 	return SweepGeomImp(QueryGeom, StartTM, Dir, Length, OutTime, OutPosition, OutNormal, OutFaceIndex, Thickness, bComputeMTD);
 }
 
-bool FTriangleMeshImplicitObject::SweepGeom(const TImplicitObjectScaled<TSphere<FReal, 3>>& QueryGeom, const FRigidTransform3& StartTM, const FVec3& Dir, const FReal Length, FReal& OutTime, FVec3& OutPosition, FVec3& OutNormal, int32& OutFaceIndex, const FReal Thickness, const bool bComputeMTD) const
+bool FTriangleMeshImplicitObject::SweepGeom(const TImplicitObjectScaled<TSphere<FReal, 3>>& QueryGeom, const FRigidTransform3& StartTM, const FVec3& Dir, const FReal Length, FReal& OutTime, FVec3& OutPosition, FVec3& OutNormal, int32& OutFaceIndex, const FReal Thickness, const bool bComputeMTD, FVec3 TriMeshScale) const
 {
-	return SweepGeomImp(QueryGeom, StartTM, Dir, Length, OutTime, OutPosition, OutNormal, OutFaceIndex, Thickness, bComputeMTD);
+	return SweepGeomImp(QueryGeom, StartTM, Dir, Length, OutTime, OutPosition, OutNormal, OutFaceIndex, Thickness, bComputeMTD, TriMeshScale);
 }
 
-bool FTriangleMeshImplicitObject::SweepGeom(const TImplicitObjectScaled<TBox<FReal, 3>>& QueryGeom, const FRigidTransform3& StartTM, const FVec3& Dir, const FReal Length, FReal& OutTime, FVec3& OutPosition, FVec3& OutNormal, int32& OutFaceIndex, const FReal Thickness, const bool bComputeMTD) const
+bool FTriangleMeshImplicitObject::SweepGeom(const TImplicitObjectScaled<TBox<FReal, 3>>& QueryGeom, const FRigidTransform3& StartTM, const FVec3& Dir, const FReal Length, FReal& OutTime, FVec3& OutPosition, FVec3& OutNormal, int32& OutFaceIndex, const FReal Thickness, const bool bComputeMTD, FVec3 TriMeshScale) const
 {
-	return SweepGeomImp(QueryGeom, StartTM, Dir, Length, OutTime, OutPosition, OutNormal, OutFaceIndex, Thickness, bComputeMTD);
+	return SweepGeomImp(QueryGeom, StartTM, Dir, Length, OutTime, OutPosition, OutNormal, OutFaceIndex, Thickness, bComputeMTD, TriMeshScale);
 }
 
-bool FTriangleMeshImplicitObject::SweepGeom(const TImplicitObjectScaled<TCapsule<FReal>>& QueryGeom, const FRigidTransform3& StartTM, const FVec3& Dir, const FReal Length, FReal& OutTime, FVec3& OutPosition, FVec3& OutNormal, int32& OutFaceIndex, const FReal Thickness, const bool bComputeMTD) const
+bool FTriangleMeshImplicitObject::SweepGeom(const TImplicitObjectScaled<TCapsule<FReal>>& QueryGeom, const FRigidTransform3& StartTM, const FVec3& Dir, const FReal Length, FReal& OutTime, FVec3& OutPosition, FVec3& OutNormal, int32& OutFaceIndex, const FReal Thickness, const bool bComputeMTD, FVec3 TriMeshScale) const
 {
-	return SweepGeomImp(QueryGeom, StartTM, Dir, Length, OutTime, OutPosition, OutNormal, OutFaceIndex, Thickness, bComputeMTD);
+	return SweepGeomImp(QueryGeom, StartTM, Dir, Length, OutTime, OutPosition, OutNormal, OutFaceIndex, Thickness, bComputeMTD, TriMeshScale);
 }
 
-bool FTriangleMeshImplicitObject::SweepGeom(const TImplicitObjectScaled<FConvex>& QueryGeom, const FRigidTransform3& StartTM, const FVec3& Dir, const FReal Length, FReal& OutTime, FVec3& OutPosition, FVec3& OutNormal, int32& OutFaceIndex, const FReal Thickness, const bool bComputeMTD) const
+bool FTriangleMeshImplicitObject::SweepGeom(const TImplicitObjectScaled<FConvex>& QueryGeom, const FRigidTransform3& StartTM, const FVec3& Dir, const FReal Length, FReal& OutTime, FVec3& OutPosition, FVec3& OutNormal, int32& OutFaceIndex, const FReal Thickness, const bool bComputeMTD, FVec3 TriMeshScale) const
 {
-	return SweepGeomImp(QueryGeom, StartTM, Dir, Length, OutTime, OutPosition, OutNormal, OutFaceIndex, Thickness, bComputeMTD);
+	return SweepGeomImp(QueryGeom, StartTM, Dir, Length, OutTime, OutPosition, OutNormal, OutFaceIndex, Thickness, bComputeMTD, TriMeshScale);
 }
 
 template <typename IdxType>
diff --git a/Engine/Source/Runtime/Experimental/Chaos/Public/Chaos/ImplicitObjectScaled.h b/Engine/Source/Runtime/Experimental/Chaos/Public/Chaos/ImplicitObjectScaled.h
--- a/Engine/Source/Runtime/Experimental/Chaos/Public/Chaos/ImplicitObjectScaled.h
+++ b/Engine/Source/Runtime/Experimental/Chaos/Public/Chaos/ImplicitObjectScaled.h
@@ -377,7 +377,7 @@ public:
 
 			TRigidTransform<T, d> BToATMNoScale(BToATM.GetLocation() * MInvScale, BToATM.GetRotation());
 			
-			if (MObject->SweepGeom(ScaledB, BToATMNoScale, UnscaledDir, UnscaledLength, UnscaledTime, UnscaledPosition, UnscaledNormal, OutFaceIndex, MInternalThickness + Thickness, bComputeMTD))
+			if (MObject->SweepGeom(ScaledB, BToATMNoScale, UnscaledDir, UnscaledLength, UnscaledTime, UnscaledPosition, UnscaledNormal, OutFaceIndex, MInternalThickness + Thickness, bComputeMTD, MScale))
 			{
 				const T NewTime = LengthScaleInv * UnscaledTime;
 				//We double check that NewTime < Length because of potential precision issues. When that happens we always keep the shortest hit first
@@ -411,7 +411,7 @@ public:
 
 		auto ScaledB = MakeScaledHelper(B, MInvScale);
 		TRigidTransform<T, d> BToATMNoScale(BToATM.GetLocation() * MInvScale, BToATM.GetRotation());
-		return MObject->OverlapGeom(ScaledB, BToATMNoScale, MInternalThickness + Thickness, OutMTD);
+		return MObject->OverlapGeom(ScaledB, BToATMNoScale, MInternalThickness + Thickness, OutMTD, MScale);
 	}
 
 	virtual int32 FindMostOpposingFace(const TVector<T, d>& Position, const TVector<T, d>& UnitDir, int32 HintFaceIndex, T SearchDist) const override
diff --git a/Engine/Source/Runtime/Experimental/Chaos/Public/Chaos/TriangleMeshImplicitObject.h b/Engine/Source/Runtime/Experimental/Chaos/Public/Chaos/TriangleMeshImplicitObject.h
--- a/Engine/Source/Runtime/Experimental/Chaos/Public/Chaos/TriangleMeshImplicitObject.h
+++ b/Engine/Source/Runtime/Experimental/Chaos/Public/Chaos/TriangleMeshImplicitObject.h
@@ -131,20 +131,21 @@ namespace Chaos
 		bool OverlapGeom(const TBox<FReal, 3>& QueryGeom, const FRigidTransform3& QueryTM, const FReal Thickness, FMTDInfo* OutMTD = nullptr) const;
 		bool OverlapGeom(const TCapsule<FReal>& QueryGeom, const FRigidTransform3& QueryTM, const FReal Thickness, FMTDInfo* OutMTD = nullptr) const;
 		bool OverlapGeom(const FConvex& QueryGeom, const FRigidTransform3& QueryTM, const FReal Thickness, FMTDInfo* OutMTD = nullptr) const;
-		bool OverlapGeom(const TImplicitObjectScaled<TSphere<FReal, 3>>& QueryGeom, const FRigidTransform3& QueryTM, const FReal Thickness, FMTDInfo* OutMTD = nullptr) const;
-		bool OverlapGeom(const TImplicitObjectScaled<TBox<FReal, 3>>& QueryGeom, const FRigidTransform3& QueryTM, const FReal Thickness, FMTDInfo* OutMTD = nullptr) const;
-		bool OverlapGeom(const TImplicitObjectScaled<TCapsule<FReal>>& QueryGeom, const FRigidTransform3& QueryTM, const FReal Thickness, FMTDInfo* OutMTD = nullptr) const;
-		bool OverlapGeom(const TImplicitObjectScaled<FConvex>& QueryGeom, const FRigidTransform3& QueryTM, const FReal Thickness, FMTDInfo* OutMTD = nullptr) const;
-		bool OverlapGeom(const TImplicitObjectScaled<TImplicitObjectScaled<FConvex>>& QueryGeom, const FRigidTransform3& QueryTM, const FReal Thickness, FMTDInfo* OutMTD = nullptr) const;
+
+		bool OverlapGeom(const TImplicitObjectScaled<TSphere<FReal, 3>>& QueryGeom, const FRigidTransform3& QueryTM, const FReal Thickness, FMTDInfo* OutMTD = nullptr, FVec3 TriMeshScale = FVec3(1.0f)) const;
+		bool OverlapGeom(const TImplicitObjectScaled<TBox<FReal, 3>>& QueryGeom, const FRigidTransform3& QueryTM, const FReal Thickness, FMTDInfo* OutMTD = nullptr, FVec3 TriMeshScale = FVec3(1.0f)) const;
+		bool OverlapGeom(const TImplicitObjectScaled<TCapsule<FReal>>& QueryGeom, const FRigidTransform3& QueryTM, const FReal Thickness, FMTDInfo* OutMTD = nullptr, FVec3 TriMeshScale = FVec3(1.0f)) const;
+		bool OverlapGeom(const TImplicitObjectScaled<FConvex>& QueryGeom, const FRigidTransform3& QueryTM, const FReal Thickness, FMTDInfo* OutMTD = nullptr, FVec3 TriMeshScale = FVec3(1.0f)) const;
 
 		bool SweepGeom(const TSphere<FReal, 3>& QueryGeom, const FRigidTransform3& StartTM, const FVec3& Dir, const FReal Length, FReal& OutTime, FVec3& OutPosition, FVec3& OutNormal, int32& OutFaceIndex, const FReal Thickness = 0, const bool bComputeMTD = false) const;
 		bool SweepGeom(const TBox<FReal, 3>& QueryGeom, const FRigidTransform3& StartTM, const FVec3& Dir, const FReal Length, FReal& OutTime, FVec3& OutPosition, FVec3& OutNormal, int32& OutFaceIndex, const FReal Thickness = 0, const bool bComputeMTD = false) const;
 		bool SweepGeom(const TCapsule<FReal>& QueryGeom, const FRigidTransform3& StartTM, const FVec3& Dir, const FReal Length, FReal& OutTime, FVec3& OutPosition, FVec3& OutNormal, int32& OutFaceIndex, const FReal Thickness = 0, const bool bComputeMTD = false) const;
 		bool SweepGeom(const FConvex& QueryGeom, const FRigidTransform3& StartTM, const FVec3& Dir, const FReal Length, FReal& OutTime, FVec3& OutPosition, FVec3& OutNormal, int32& OutFaceIndex, const FReal Thickness = 0, const bool bComputeMTD = false) const;
-		bool SweepGeom(const TImplicitObjectScaled<TSphere<FReal, 3>>& QueryGeom, const FRigidTransform3& StartTM, const FVec3& Dir, const FReal Length, FReal& OutTime, FVec3& OutPosition, FVec3& OutNormal, int32& OutFaceIndex, const FReal Thickness = 0, const bool bComputeMTD = false) const;
-		bool SweepGeom(const TImplicitObjectScaled<TBox<FReal, 3>>& QueryGeom, const FRigidTransform3& StartTM, const FVec3& Dir, const FReal Length, FReal& OutTime, FVec3& OutPosition, FVec3& OutNormal, int32& OutFaceIndex, const FReal Thickness = 0, const bool bComputeMTD = false) const;
-		bool SweepGeom(const TImplicitObjectScaled<TCapsule<FReal>>& QueryGeom, const FRigidTransform3& StartTM, const FVec3& Dir, const FReal Length, FReal& OutTime, FVec3& OutPosition, FVec3& OutNormal, int32& OutFaceIndex, const FReal Thickness = 0, const bool bComputeMTD = false) const;
-		bool SweepGeom(const TImplicitObjectScaled<FConvex>& QueryGeom, const FRigidTransform3& StartTM, const FVec3& Dir, const FReal Length, FReal& OutTime, FVec3& OutPosition, FVec3& OutNormal, int32& OutFaceIndex, const FReal Thickness = 0, const bool bComputeMTD = false) const;
+
+		bool SweepGeom(const TImplicitObjectScaled<TSphere<FReal, 3>>& QueryGeom, const FRigidTransform3& StartTM, const FVec3& Dir, const FReal Length, FReal& OutTime, FVec3& OutPosition, FVec3& OutNormal, int32& OutFaceIndex, const FReal Thickness = 0, const bool bComputeMTD = false, FVec3 TriMeshScale = FVec3(1.0f)) const;
+		bool SweepGeom(const TImplicitObjectScaled<TBox<FReal, 3>>& QueryGeom, const FRigidTransform3& StartTM, const FVec3& Dir, const FReal Length, FReal& OutTime, FVec3& OutPosition, FVec3& OutNormal, int32& OutFaceIndex, const FReal Thickness = 0, const bool bComputeMTD = false, FVec3 TriMeshScale = FVec3(1.0f)) const;
+		bool SweepGeom(const TImplicitObjectScaled<TCapsule<FReal>>& QueryGeom, const FRigidTransform3& StartTM, const FVec3& Dir, const FReal Length, FReal& OutTime, FVec3& OutPosition, FVec3& OutNormal, int32& OutFaceIndex, const FReal Thickness = 0, const bool bComputeMTD = false, FVec3 TriMeshScale = FVec3(1.0f)) const;
+		bool SweepGeom(const TImplicitObjectScaled<FConvex>& QueryGeom, const FRigidTransform3& StartTM, const FVec3& Dir, const FReal Length, FReal& OutTime, FVec3& OutPosition, FVec3& OutNormal, int32& OutFaceIndex, const FReal Thickness = 0, const bool bComputeMTD = false, FVec3 TriMeshScale = FVec3(1.0f)) const;
 
 		bool GJKContactPoint(const TSphere<FReal, 3>& QueryGeom, const FRigidTransform3& QueryTM, const FReal Thickness, FVec3& Location, FVec3& Normal, FReal& Penetration) const;
 		bool GJKContactPoint(const TBox<FReal, 3>& QueryGeom, const FRigidTransform3& QueryTM, const FReal Thickness, FVec3& Location, FVec3& Normal, FReal& Penetration) const;
@@ -327,10 +328,10 @@ namespace Chaos
 		bool GJKContactPointImp(const QueryGeomType& QueryGeom, const FRigidTransform3& QueryTM, const FReal Thickness, FVec3& Location, FVec3& Normal, FReal& Penetration, FVec3 TriMeshScale = FVec3(1.0)) const;
 
 		template<typename QueryGeomType>
-		bool OverlapGeomImp(const QueryGeomType& QueryGeom, const FRigidTransform3& QueryTM, const FReal Thickness, FMTDInfo* OutMTD = nullptr) const;
+		bool OverlapGeomImp(const QueryGeomType& QueryGeom, const FRigidTransform3& QueryTM, const FReal Thickness, FMTDInfo* OutMTD = nullptr, FVec3 TriMeshScale = FVec3(1.0f)) const;
 
 		template<typename QueryGeomType>
-		bool SweepGeomImp(const QueryGeomType& QueryGeom, const FRigidTransform3& StartTM, const FVec3& Dir, const FReal Length, FReal& OutTime, FVec3& OutPosition, FVec3& OutNormal, int32& OutFaceIndex, const FReal Thickness, const bool bComputeMTD) const;
+		bool SweepGeomImp(const QueryGeomType& QueryGeom, const FRigidTransform3& StartTM, const FVec3& Dir, const FReal Length, FReal& OutTime, FVec3& OutPosition, FVec3& OutNormal, int32& OutFaceIndex, const FReal Thickness, const bool bComputeMTD, FVec3 TriMeshScale = FVec3(1.0f)) const;
 
 		template <typename IdxType>
 		bool RaycastImp(const TArray<TVector<IdxType, 3>>& Elements, const FVec3& StartPoint, const FVec3& Dir, const FReal Length, const FReal Thickness, FReal& OutTime, FVec3& OutPosition, FVec3& OutNormal, int32& OutFaceIndex) const;
