Fix scale ensure in ImplicitObjectScaled.
Make uniform scale check in ImplicitObjectScaled use epsilon.

diff --git a/Engine/Source/Runtime/Experimental/Chaos/Public/Chaos/ImplicitObjectScaled.h b/Engine/Source/Runtime/Experimental/Chaos/Public/Chaos/ImplicitObjectScaled.h
--- a/Engine/Source/Runtime/Experimental/Chaos/Public/Chaos/ImplicitObjectScaled.h
+++ b/Engine/Source/Runtime/Experimental/Chaos/Public/Chaos/ImplicitObjectScaled.h
@@ -511,7 +511,9 @@ public:
 	{
 		if (T UnscaledMargin = MObject->GetMargin())
 		{
-			ensure(MScale[0] == MScale[1] && MScale[1] == MScale[2]);
+			ensure(FMath::IsNearlyEqual(MScale[0], MScale[1], KINDA_SMALL_NUMBER));
+			ensure(FMath::IsNearlyEqual(MScale[1], MScale[2], KINDA_SMALL_NUMBER));
+
 			return UnscaledMargin * FMath::Abs(MScale[0]);
 		}
 
