Set Initialized flag for particles after initialization so that they can get wake events normally.

diff --git a/Engine/Source/Runtime/Engine/Private/PhysicsEngine/BodyInstance.cpp b/Engine/Source/Runtime/Engine/Private/PhysicsEngine/BodyInstance.cpp
--- a/Engine/Source/Runtime/Engine/Private/PhysicsEngine/BodyInstance.cpp
+++ b/Engine/Source/Runtime/Engine/Private/PhysicsEngine/BodyInstance.cpp
@@ -3952,6 +3952,7 @@ void FBodyInstance::InitDynamicProperties_AssumesLocked()
 				FPhysicsInterface::PutToSleep_AssumesLocked(ActorHandle);
 			}
 		}
+		FPhysicsInterface::SetInitialized_AssumesLocked(ActorHandle, true);
 	}
 }
 
diff --git a/Engine/Source/Runtime/Engine/Private/PhysicsEngine/Experimental/PhysInterface_Chaos.cpp b/Engine/Source/Runtime/Engine/Private/PhysicsEngine/Experimental/PhysInterface_Chaos.cpp
--- a/Engine/Source/Runtime/Engine/Private/PhysicsEngine/Experimental/PhysInterface_Chaos.cpp
+++ b/Engine/Source/Runtime/Engine/Private/PhysicsEngine/Experimental/PhysInterface_Chaos.cpp
@@ -959,6 +959,15 @@ void FPhysInterface_Chaos::SetWakeCounter_AssumesLocked(const FPhysicsActorHandl
 	// #todo : Implement
 }
 
+void FPhysInterface_Chaos::SetInitialized_AssumesLocked(const FPhysicsActorHandle& InHandle, bool InInitialized)
+{
+	Chaos::TPBDRigidParticle<float, 3>* Rigid = InHandle->CastToRigidParticle();
+	if (Rigid)
+	{
+		Rigid->SetInitialized(InInitialized);
+	}
+}
+
 SIZE_T FPhysInterface_Chaos::GetResourceSizeEx(const FPhysicsActorHandle& InActorRef)
 {
     return sizeof(FPhysicsActorHandle);
diff --git a/Engine/Source/Runtime/Engine/Public/Physics/Experimental/PhysInterface_Chaos.h b/Engine/Source/Runtime/Engine/Public/Physics/Experimental/PhysInterface_Chaos.h
--- a/Engine/Source/Runtime/Engine/Public/Physics/Experimental/PhysInterface_Chaos.h
+++ b/Engine/Source/Runtime/Engine/Public/Physics/Experimental/PhysInterface_Chaos.h
@@ -300,6 +300,8 @@ public:
 	static float GetWakeCounter_AssumesLocked(const FPhysicsActorHandle& InHandle);
 	static void SetWakeCounter_AssumesLocked(const FPhysicsActorHandle& InHandle, float InWakeCounter);
 
+	static void SetInitialized_AssumesLocked(const FPhysicsActorHandle& InHandle, bool InInitialized);
+
 	static SIZE_T GetResourceSizeEx(const FPhysicsActorHandle& InActorRef);
 	
     static FPhysicsConstraintReference_Chaos CreateConstraint(const FPhysicsActorHandle& InActorRef1, const FPhysicsActorHandle& InActorRef2, const FTransform& InLocalFrame1, const FTransform& InLocalFrame2);
diff --git a/Engine/Source/Runtime/Experimental/Chaos/Public/Chaos/ParticleHandle.h b/Engine/Source/Runtime/Experimental/Chaos/Public/Chaos/ParticleHandle.h
--- a/Engine/Source/Runtime/Experimental/Chaos/Public/Chaos/ParticleHandle.h
+++ b/Engine/Source/Runtime/Experimental/Chaos/Public/Chaos/ParticleHandle.h
@@ -110,6 +110,7 @@ void PBDRigidParticleDefaultConstruct(TPBDRigidParticle<T,d>& Concrete, const TP
 	Concrete.SetObjectState(Params.bStartSleeping ? EObjectStateType::Sleeping : EObjectStateType::Dynamic);
 	Concrete.SetGravityEnabled(Params.bGravityEnabled);
 	Concrete.ClearEvents();
+	Concrete.SetInitialized(false);
 }
 
 
@@ -1868,6 +1869,12 @@ public:
 		this->MGravityEnabled = InGravityEnabled;
 	}
 
+	bool IsInitialized() const { return MInitialized; }
+	void SetInitialized(const bool InInitialized)
+	{
+		this->MInitialized = InInitialized;
+	}
+
 	// Named to match signature of TPBDRigidParticleHandle, as both are used in templated functions.
 	// See its comment for details.
 	bool& SetDisabledLowLevel() { return MDisabled; }
@@ -2035,6 +2042,7 @@ private:
 	bool MDisabled;
 	bool MToBeRemovedOnFracture;
 	bool MGravityEnabled;
+	bool MInitialized;
 	bool MAwakeEvent;
 };
 
@@ -2090,6 +2098,7 @@ public:
 		, MDisabled(false)
 		, MToBeRemovedOnFracture(false)
 		, MGravityEnabled(false)
+		, MInitialized(false)
 	{}
 
 	TPBDRigidParticleData(const TPBDRigidParticle<T, d>& InParticle)
@@ -2115,6 +2124,7 @@ public:
 		, MDisabled(InParticle.Disabled())
 		, MToBeRemovedOnFracture(InParticle.ToBeRemovedOnFracture())
 		, MGravityEnabled(InParticle.IsGravityEnabled())
+		, MInitialized(InParticle.IsInitialized())
 	{
 		Type = EParticleType::Rigid;
 	}
@@ -2141,6 +2151,7 @@ public:
 	bool MDisabled;
 	bool MToBeRemovedOnFracture;
 	bool MGravityEnabled;
+	bool MInitialized;
 
 	void Reset() {
 		TKinematicGeometryParticleData<T, d>::Reset();
diff --git a/Engine/Source/Runtime/Experimental/ChaosSolvers/Private/PhysicsProxy/SingleParticlePhysicsProxy.cpp b/Engine/Source/Runtime/Experimental/ChaosSolvers/Private/PhysicsProxy/SingleParticlePhysicsProxy.cpp
--- a/Engine/Source/Runtime/Experimental/ChaosSolvers/Private/PhysicsProxy/SingleParticlePhysicsProxy.cpp
+++ b/Engine/Source/Runtime/Experimental/ChaosSolvers/Private/PhysicsProxy/SingleParticlePhysicsProxy.cpp
@@ -391,7 +391,8 @@ void FSingleParticlePhysicsProxy<Chaos::TPBDRigidParticle<float, 3>>::PushToPhys
 				GetSolver()->GetEvolution()->SetParticleObjectState(RigidHandle, Chaos::EObjectStateType::Dynamic);
 			}
 		}
-		else
+		
+		if(Data->MInitialized)
 		{
 			// wait for the first pass with nothing updated to claim its initialized
 			bInitialized = true;
