Don't use checks

diff --git a/Engine/Source/Runtime/Experimental/Chaos/Public/Chaos/PBDRigidsSOAs.h b/Engine/Source/Runtime/Experimental/Chaos/Public/Chaos/PBDRigidsSOAs.h
--- a/Engine/Source/Runtime/Experimental/Chaos/Public/Chaos/PBDRigidsSOAs.h
+++ b/Engine/Source/Runtime/Experimental/Chaos/Public/Chaos/PBDRigidsSOAs.h
@@ -294,13 +294,15 @@ public:
 		auto PBDRigid = Particle->CastToRigidParticle();
 		if(PBDRigid && PBDRigid->ObjectState() == EObjectStateType::Dynamic)
 		{
-			check(!PBDRigid->Disabled());
-			if (auto PBDRigidClustered = Particle->CastToClustered())
+			if (ensure(!PBDRigid->Disabled()))
 			{
-				InsertToMapAndArray(PBDRigidClustered, ActiveClusteredToIndex, ActiveClusteredArray);
-			}
+				if (auto PBDRigidClustered = Particle->CastToClustered())
+				{
+					InsertToMapAndArray(PBDRigidClustered, ActiveClusteredToIndex, ActiveClusteredArray);
+				}
 
-			InsertToMapAndArray(PBDRigid, ActiveParticlesToIndex, ActiveParticlesArray);
+				InsertToMapAndArray(PBDRigid, ActiveParticlesToIndex, ActiveParticlesArray);
+			}
 		}
 		
 		UpdateViews();
@@ -314,13 +316,15 @@ public:
 			if (   PBDRigid->ObjectState() == EObjectStateType::Dynamic
 				|| PBDRigid->ObjectState() == EObjectStateType::Sleeping)
 			{
-				check(!PBDRigid->Disabled());
-				if (auto PBDRigidClustered = Particle->CastToClustered())
+				if (ensure(!PBDRigid->Disabled()))
 				{
-					RemoveFromMapAndArray(PBDRigidClustered, ActiveClusteredToIndex, ActiveClusteredArray);
-				}
+					if (auto PBDRigidClustered = Particle->CastToClustered())
+					{
+						RemoveFromMapAndArray(PBDRigidClustered, ActiveClusteredToIndex, ActiveClusteredArray);
+					}
 
-				RemoveFromMapAndArray(PBDRigid, ActiveParticlesToIndex, ActiveParticlesArray);
+					RemoveFromMapAndArray(PBDRigid, ActiveParticlesToIndex, ActiveParticlesArray);
+				}
 			}
 		}
 
