Chaos - hook up manifold disable cvar to convex-convex manifold

diff --git a/Engine/Source/Runtime/Experimental/Chaos/Private/Chaos/CollisionResolution.cpp b/Engine/Source/Runtime/Experimental/Chaos/Private/Chaos/CollisionResolution.cpp
--- a/Engine/Source/Runtime/Experimental/Chaos/Private/Chaos/CollisionResolution.cpp
+++ b/Engine/Source/Runtime/Experimental/Chaos/Private/Chaos/CollisionResolution.cpp
@@ -1286,8 +1286,9 @@ namespace Chaos
 		template<class T, int d>
 		void UpdateConvexConvexManifold(TCollisionConstraintBase<T, d>&  ConstraintBase, const TRigidTransform<T, d>& Transform0, const TRigidTransform<T, d>& Transform1, const T CullDistance)
 		{
-			if (TRigidBodyMultiPointContactConstraint<T,d>* Constraint = ConstraintBase.template As<TRigidBodyMultiPointContactConstraint<T, d>>())
+			if (ConstraintBase.GetType() == FRigidBodyMultiPointContactConstraint::StaticType())
 			{
+				TRigidBodyMultiPointContactConstraint<T, d>* Constraint = ConstraintBase.template As<TRigidBodyMultiPointContactConstraint<T, d>>();
 				if (GetInnerType(ConstraintBase.Manifold.Implicit[0]->GetType()) == ImplicitObjectType::Convex)
 				{
 					UpdateSingleShotManifold(*Constraint, Transform0, Transform1, CullDistance);
@@ -1305,10 +1306,19 @@ namespace Chaos
 		{
 			TRigidTransform<T, d> ParticleImplicit0TM = Transform0.GetRelativeTransform(Collisions::GetTransform(Particle0));
 			TRigidTransform<T, d> ParticleImplicit1TM = Transform1.GetRelativeTransform(Collisions::GetTransform(Particle1));
-			FRigidBodyMultiPointContactConstraint Constraint = FRigidBodyMultiPointContactConstraint(Particle0, Implicit0, ParticleImplicit0TM, Particle1, Implicit1, ParticleImplicit1TM, EContactShapesType::ConvexConvex);
-			UpdateConvexConvexManifold(Constraint, Transform0, Transform1, CullDistance);
-			UpdateConvexConvexConstraint(*Implicit0, Transform0, *Implicit1, Transform1, CullDistance, Constraint);
-			NewConstraints.TryAdd(CullDistance, Constraint);
+			if (Chaos_Collision_UseManifolds)
+			{
+				FRigidBodyMultiPointContactConstraint Constraint = FRigidBodyMultiPointContactConstraint(Particle0, Implicit0, ParticleImplicit0TM, Particle1, Implicit1, ParticleImplicit1TM, EContactShapesType::ConvexConvex);
+				UpdateConvexConvexManifold(Constraint, Transform0, Transform1, CullDistance);
+				UpdateConvexConvexConstraint(*Implicit0, Transform0, *Implicit1, Transform1, CullDistance, Constraint);
+				NewConstraints.TryAdd(CullDistance, Constraint);
+			}
+			else
+			{
+				FRigidBodyPointContactConstraint Constraint = FRigidBodyPointContactConstraint(Particle0, Implicit0, ParticleImplicit0TM, Particle1, Implicit1, ParticleImplicit1TM, EContactShapesType::ConvexConvex);
+				UpdateConvexConvexConstraint(*Implicit0, Transform0, *Implicit1, Transform1, CullDistance, Constraint);
+				NewConstraints.TryAdd(CullDistance, Constraint);
+			}
 		}
 
 		//
