Chaos - ISPC TransformAABB for Chaos transform type

diff --git a/Engine/Source/Runtime/Experimental/Chaos/Private/Chaos/AABB.cpp b/Engine/Source/Runtime/Experimental/Chaos/Private/Chaos/AABB.cpp
--- a/Engine/Source/Runtime/Experimental/Chaos/Private/Chaos/AABB.cpp
+++ b/Engine/Source/Runtime/Experimental/Chaos/Private/Chaos/AABB.cpp
@@ -179,7 +179,7 @@ bool TAABB<T, d>::Raycast(const TVector<T, d>& StartPoint, const TVector<T, d>&
 }
 
 template<typename T, int d, class TTRANSFORM>
-TAABB<T, d> TransformedAABBHelper(const TAABB<T, d>& AABB, const TTRANSFORM& SpaceTransform)
+inline TAABB<T, d> TransformedAABBHelper(const TAABB<T, d>& AABB, const TTRANSFORM& SpaceTransform)
 {
 	TVector<T, d> CurrentExtents = AABB.Extents();
 	int32 Idx = 0;
@@ -196,10 +196,31 @@ TAABB<T, d> TransformedAABBHelper(const TAABB<T, d>& AABB, const TTRANSFORM& Spa
 	return NewAABB;
 }
 
+inline TAABB<float, 3> TransformedAABBHelperISPC(const TAABB<float, 3>& AABB, const FTransform& SpaceTransform)
+{
+#if INTEL_ISPC
+	TVector<float, 3> NewMin, NewMax;
+	ispc::TransformedAABB((const ispc::FTransform&)SpaceTransform, (const ispc::FVector&)AABB.Min(), (const ispc::FVector&)AABB.Max(), (ispc::FVector&)NewMin, (ispc::FVector&)NewMax);
+
+	TAABB<float, 3> NewAABB(NewMin, NewMax);
+	return NewAABB;
+#else
+	check(false);
+	return TAABB<float, 3>::EmptyAABB();
+#endif
+}
+
 template<typename T, int d>
 TAABB<T, d> TAABB<T, d>::TransformedAABB(const Chaos::TRigidTransform<FReal, 3>& SpaceTransform) const
 {
-	return TransformedAABBHelper<T, d>(*this, SpaceTransform);
+	if (INTEL_ISPC)
+	{
+		return TransformedAABBHelperISPC(*this, SpaceTransform);
+	}
+	else
+	{
+		return TransformedAABBHelper<T, d>(*this, SpaceTransform);
+	}
 }
 
 template<typename T, int d>
@@ -220,29 +241,11 @@ TAABB<T, d> TAABB<T, d>::TransformedAABB(const FTransform& SpaceTransform) const
 {
 	if (INTEL_ISPC)
 	{
-#if INTEL_ISPC
-		TVector<float, 3> NewMin, NewMax;
-		ispc::TransformedAABB((const ispc::FTransform&)SpaceTransform, (const ispc::FVector&)MMin, (const ispc::FVector&)MMax, (ispc::FVector&)NewMin, (ispc::FVector&)NewMax);
-
-		TAABB<float, 3> NewAABB(NewMin, NewMax);
-		return NewAABB;
-#endif
+		return TransformedAABBHelperISPC(*this, SpaceTransform);
 	}
 	else
 	{
-		TVector<float, 3> CurrentExtents = Extents();
-		int32 Idx = 0;
-		const TVector<float, 3> MinToNewSpace = SpaceTransform.TransformPosition(MMin);
-		TAABB<float, 3> NewAABB(MinToNewSpace, MinToNewSpace);
-		NewAABB.GrowToInclude(SpaceTransform.TransformPosition(MMax));
-
-		for (int32 j = 0; j < 3; ++j)
-		{
-			NewAABB.GrowToInclude(SpaceTransform.TransformPosition(MMin + TVector<float, 3>::AxisVector(j) * CurrentExtents));
-			NewAABB.GrowToInclude(SpaceTransform.TransformPosition(MMax - TVector<float, 3>::AxisVector(j) * CurrentExtents));
-		}
-
-		return NewAABB;
+		return TransformedAABBHelper<T, d>(*this, SpaceTransform);
 	}
 }
 }
