Make it so actors are not marked dirty unless they move outside of their original bounds

diff --git a/Engine/Source/Programs/HeadlessChaos/Private/HeadlessChaosTestBroadphase.cpp b/Engine/Source/Programs/HeadlessChaos/Private/HeadlessChaosTestBroadphase.cpp
--- a/Engine/Source/Programs/HeadlessChaos/Private/HeadlessChaosTestBroadphase.cpp
+++ b/Engine/Source/Programs/HeadlessChaos/Private/HeadlessChaosTestBroadphase.cpp
@@ -496,7 +496,7 @@ namespace ChaosTest
 
 			//fill up until dirty limit
 			int32 Count;
-			for(Count = 0; Count < 10; ++Count)
+			for(Count = 1; Count <= 10; ++Count)
 			{
 				auto Boxes2 = BuildBoxes<T>(Box);
 				for(uint32 Idx = 0; Idx < Boxes2->Size(); ++Idx)
@@ -504,14 +504,14 @@ namespace ChaosTest
 					Spatial.UpdateElement(Idx + Boxes->Size() * Count,Boxes2->WorldSpaceInflatedBounds(Idx),true);
 				}
 
-				EXPECT_EQ(Spatial.NumDirtyElements(), (Count+1)*Boxes->Size());
+				EXPECT_EQ(Spatial.NumDirtyElements(), (Count)*Boxes->Size());
 			}
 
 			//finally pass dirty limit so reset to 0 and then add the remaining new boxes
 			auto Boxes2 = BuildBoxes<T>(Box);
 			for(uint32 Idx = 0; Idx < Boxes2->Size(); ++Idx)
 			{
-				Spatial.UpdateElement(Idx + Boxes->Size() * Count,Boxes2->WorldSpaceInflatedBounds(Idx),true);
+				Spatial.UpdateElement(Idx + Boxes->Size() * (Count),Boxes2->WorldSpaceInflatedBounds(Idx),true);
 			}
 
 			EXPECT_EQ(Spatial.NumDirtyElements(),Boxes->Size() - 1);
diff --git a/Engine/Source/Runtime/Experimental/Chaos/Public/Chaos/AABBTree.h b/Engine/Source/Runtime/Experimental/Chaos/Public/Chaos/AABBTree.h
--- a/Engine/Source/Runtime/Experimental/Chaos/Public/Chaos/AABBTree.h
+++ b/Engine/Source/Runtime/Experimental/Chaos/Public/Chaos/AABBTree.h
@@ -69,14 +69,57 @@ struct TAABBTreeIntersectionHelper<T, FQueryFastDataVoid, EAABBQueryType::Overla
 	}
 };
 
+template <typename TPayloadType, typename T, bool bComputeBounds>
+struct TBoundsWrapperHelper
+{
+
+};
+
+template <typename TPayloadType, typename T>
+struct TBoundsWrapperHelper<TPayloadType, T, true>
+{
+	void ComputeBounds(const TArray<TPayloadBoundsElement<TPayloadType, T>>& Elems)
+	{
+		Bounds = TAABB<T, 3>::EmptyAABB();
+
+		for (const auto& Elem : Elems)
+		{
+			Bounds.GrowToInclude(Elem.Bounds);
+		}
+	}
+
+	const TAABB<T, 3>& GetBounds() const { return Bounds; }
+
+private:
+	TAABB<T, 3> Bounds;
+};
+
 template <typename TPayloadType, typename T>
-struct TAABBTreeLeafArray
+struct TBoundsWrapperHelper<TPayloadType, T, false>
+{
+	void ComputeBounds(const TArray<TPayloadBoundsElement<TPayloadType, T>>&)
+	{
+	}
+
+	const TAABB<T, 3> GetBounds() const
+	{
+		return TAABB<T, 3>::EmptyAABB();
+	}
+};
+
+template <typename TPayloadType, typename T, bool bComputeBounds = true>
+struct TAABBTreeLeafArray : public TBoundsWrapperHelper<TPayloadType, T, bComputeBounds>
 {
 	TAABBTreeLeafArray() {}
 	//todo: avoid copy?
 	TAABBTreeLeafArray(const TArray<TPayloadBoundsElement<TPayloadType, T>>& InElems)
 		: Elems(InElems)
 	{
+		ComputeBounds(Elems);
+		if (bComputeBounds)
+		{
+		
+		}
 	}
 
 	void GatherElements(TArray<TPayloadBoundsElement<TPayloadType, T>>& OutElements)
@@ -148,16 +191,23 @@ struct TAABBTreeLeafArray
 		}
 	}
 
+	const TAABB<T, 3>& GetBounds() const
+	{
+		return Bounds;
+	}
+
 	void Serialize(FChaosArchive& Ar)
 	{
 		Ar << Elems;
 	}
 
+
 	TArray<TPayloadBoundsElement<TPayloadType, T>> Elems;
+	TAABB<T,3> Bounds;
 };
 
-template <typename TPayloadType, typename T>
-FChaosArchive& operator<<(FChaosArchive& Ar, TAABBTreeLeafArray<TPayloadType, T>& LeafArray)
+template <typename TPayloadType, typename T, bool bComputeBounds>
+FChaosArchive& operator<<(FChaosArchive& Ar, TAABBTreeLeafArray<TPayloadType, T, bComputeBounds>& LeafArray)
 {
 	LeafArray.Serialize(Ar);
 	return Ar;
@@ -427,6 +477,16 @@ public:
 			{
 				if (PayloadInfo->LeafIdx != INDEX_NONE)
 				{
+					//If we are still within the same leaf bounds, do nothing
+					if (bHasBounds)
+					{
+						const TAABB<T,3>& LeafBounds = Leaves[PayloadInfo->LeafIdx].GetBounds();
+						if (LeafBounds.Contains(NewBounds.Min()) && LeafBounds.Contains(NewBounds.Max()))
+						{
+							return;
+						}
+					}
+
 					Leaves[PayloadInfo->LeafIdx].RemoveElement(Payload);
 					PayloadInfo->LeafIdx = INDEX_NONE;
 				}
diff --git a/Engine/Source/Runtime/Experimental/Chaos/Public/Chaos/TriangleMeshImplicitObject.h b/Engine/Source/Runtime/Experimental/Chaos/Public/Chaos/TriangleMeshImplicitObject.h
--- a/Engine/Source/Runtime/Experimental/Chaos/Public/Chaos/TriangleMeshImplicitObject.h
+++ b/Engine/Source/Runtime/Experimental/Chaos/Public/Chaos/TriangleMeshImplicitObject.h
@@ -266,7 +266,7 @@ namespace Chaos
 		bool bCullsBackFaceRaycast;
 
 		//using BVHType = TBoundingVolume<int32, T, 3>;
-		using BVHType = TAABBTree<int32, TAABBTreeLeafArray<int32, FReal>, FReal, /*bMutable=*/false>;
+		using BVHType = TAABBTree<int32, TAABBTreeLeafArray<int32, FReal, /*bComputeBounds=*/false>, FReal, /*bMutable=*/false>;
 
 		template<typename InStorageType, typename InRealType>
 		friend struct FBvEntry;
