Missing file

diff --git a/Engine/Source/Runtime/Experimental/Chaos/Public/Chaos/BoundingVolume.h b/Engine/Source/Runtime/Experimental/Chaos/Public/Chaos/BoundingVolume.h
--- a/Engine/Source/Runtime/Experimental/Chaos/Public/Chaos/BoundingVolume.h
+++ b/Engine/Source/Runtime/Experimental/Chaos/Public/Chaos/BoundingVolume.h
@@ -353,14 +353,31 @@ public:
 			OutElements.Add(TPayloadBoundsElement<TPayloadType,T>{Elem.Payload,Elem.Bounds});
 		}
 
-		for(int32 Idx = 0; Idx < MElements.Num(); ++Idx)
+		const auto& Counts = MGrid.Counts();
+		for (int32 X = 0; X < Counts[0]; ++X)
 		{
-			const TArray<FCellElement>& Elems = MElements[Idx];
-			for(const FCellElement& Elem : Elems)
+			for (int32 Y = 0; Y < Counts[1]; ++Y)
 			{
-				OutElements.Add(TPayloadBoundsElement<TPayloadType,T>{Elem.Payload,Elem.Bounds});
+				for (int32 Z = 0; Z < Counts[2]; ++Z)
+				{
+					const auto& Elems = MElements(X, Y, Z);
+					for (const auto& Elem : Elems)
+					{
+						//elements can be in multiple cells, only add for the first cell
+						if (Elem.StartIdx == TVector<int32, 3>(X, Y, Z))
+						{
+							OutElements.Add(TPayloadBoundsElement<TPayloadType, T>{Elem.Payload, Elem.Bounds});
+						}
+					}
+				}
 			}
 		}
+
+	}
+
+	TAABB<T, d> GetBounds() const
+	{
+		return TAABB<T, d>(MGrid.MinCorner(), MGrid.MaxCorner());
 	}
 
 private:
