Fix CIS

diff --git a/Engine/Source/Runtime/Experimental/Chaos/Public/Chaos/AABBTree.h b/Engine/Source/Runtime/Experimental/Chaos/Public/Chaos/AABBTree.h
--- a/Engine/Source/Runtime/Experimental/Chaos/Public/Chaos/AABBTree.h
+++ b/Engine/Source/Runtime/Experimental/Chaos/Public/Chaos/AABBTree.h
@@ -115,11 +115,7 @@ struct TAABBTreeLeafArray : public TBoundsWrapperHelper<TPayloadType, T, bComput
 	TAABBTreeLeafArray(const TArray<TPayloadBoundsElement<TPayloadType, T>>& InElems)
 		: Elems(InElems)
 	{
-		ComputeBounds(Elems);
-		if (bComputeBounds)
-		{
-		
-		}
+		this->ComputeBounds(Elems);
 	}
 
 	void GatherElements(TArray<TPayloadBoundsElement<TPayloadType, T>>& OutElements)
