Fix the DynamicState field so that it's less restrictive about the current RBD state for geometry collection particles.

diff --git a/Engine/Source/Runtime/Experimental/ChaosSolvers/Private/PhysicsProxy/FieldSystemPhysicsProxy.cpp b/Engine/Source/Runtime/Experimental/ChaosSolvers/Private/PhysicsProxy/FieldSystemPhysicsProxy.cpp
--- a/Engine/Source/Runtime/Experimental/ChaosSolvers/Private/PhysicsProxy/FieldSystemPhysicsProxy.cpp
+++ b/Engine/Source/Runtime/Experimental/ChaosSolvers/Private/PhysicsProxy/FieldSystemPhysicsProxy.cpp
@@ -129,20 +129,24 @@ void FFieldSystemPhysicsProxy::FieldParameterUpdateCallback(
 					int32 i = 0;
 					for(Chaos::TGeometryParticleHandle<float,3>* Handle : Handles)
 					{
-						const Chaos::EObjectStateType CurrState = Handle->ObjectState();
-						if(CurrState == Chaos::EObjectStateType::Kinematic)
+						const Chaos::EObjectStateType ObjectState = Handle->ObjectState();
+						switch (ObjectState)
 						{
-							DynamicState[i] = (int)EObjectStateTypeEnum::Chaos_Object_Kinematic;
+						case Chaos::EObjectStateType::Kinematic:
+							DynamicState[i++] = (int)EObjectStateTypeEnum::Chaos_Object_Kinematic;
+							break;
+						case Chaos::EObjectStateType::Static:
+							DynamicState[i++] = (int)EObjectStateTypeEnum::Chaos_Object_Static;
+							break;
+						case Chaos::EObjectStateType::Sleeping:
+							DynamicState[i++] = (int)EObjectStateTypeEnum::Chaos_Object_Sleeping;
+							break;
+						case Chaos::EObjectStateType::Dynamic:
+						case Chaos::EObjectStateType::Uninitialized:
+						default:
+							DynamicState[i++] = (int)EObjectStateTypeEnum::Chaos_Object_Dynamic;
+							break;
 						}
-						else if (CurrState == Chaos::EObjectStateType::Static)
-						{
-							DynamicState[i] = (int)EObjectStateTypeEnum::Chaos_Object_Static;
-						}
-						else
-						{
-							DynamicState[i] = (int)EObjectStateTypeEnum::Chaos_Object_Dynamic;
-						}
-						++i;
 					}
 
 					TArrayView<int32> DynamicStateView(&(DynamicState[0]), DynamicState.Num());
@@ -168,11 +172,22 @@ void FFieldSystemPhysicsProxy::FieldParameterUpdateCallback(
 						// we're just going to ignore non-dynamic particles.  This has the added
 						// benefit of not needing to deal with the floor, as it's pretty likely to
 						// not be dynamic.  Har.
-						Chaos::TPBDRigidParticleHandle<float, 3>* RigidHandle = Handle->CastToRigidParticle();
-						if(RigidHandle && RigidHandle->ObjectState() == Chaos::EObjectStateType::Dynamic)
+						if(Chaos::TPBDRigidParticleHandle<float, 3>* RigidHandle = Handle->CastToRigidParticle())
 						{
-							const int32 FieldState = DynamicStateView[i];
+							const bool bIsGC = (Handle->GetParticleType() == Chaos::EParticleType::GeometryCollection);
 							const EObjectStateType HandleState = RigidHandle->ObjectState();
+
+							// Non-Geometry Collection rigid bodies are more restricted as to how
+							// they can change state, as they'd have to promote or demote to a 
+							// different particle type, necessitating changing to a different
+							// particle SOA, which would be expensive.  So, we only permit 
+							// downgrading from Dynamic in those cases.
+							if (!bIsGC && HandleState != Chaos::EObjectStateType::Dynamic)
+							{
+								continue;
+							}
+
+							const int32 FieldState = DynamicStateView[i];
 							if (FieldState == (int32)EObjectStateTypeEnum::Chaos_Object_Dynamic)
 							{
 								if ((HandleState == Chaos::EObjectStateType::Static ||
@@ -190,7 +205,7 @@ void FFieldSystemPhysicsProxy::FieldParameterUpdateCallback(
 							}
 							else if (FieldState == (int32)EObjectStateTypeEnum::Chaos_Object_Kinematic)
 							{
-								if (HandleState == Chaos::EObjectStateType::Dynamic)
+								if (HandleState != Chaos::EObjectStateType::Kinematic)
 								{
 									RigidHandle->SetObjectStateLowLevel(Chaos::EObjectStateType::Kinematic);
 									RigidHandle->SetV(Chaos::TVector<float, 3>(0));
@@ -200,7 +215,7 @@ void FFieldSystemPhysicsProxy::FieldParameterUpdateCallback(
 							}
 							else if (FieldState == (int32)EObjectStateTypeEnum::Chaos_Object_Static)
 							{
-								if (HandleState == Chaos::EObjectStateType::Dynamic)
+								if (HandleState != Chaos::EObjectStateType::Static)
 								{
 									RigidHandle->SetObjectStateLowLevel(Chaos::EObjectStateType::Static);
 									RigidHandle->SetV(Chaos::TVector<float, 3>(0));
@@ -210,7 +225,7 @@ void FFieldSystemPhysicsProxy::FieldParameterUpdateCallback(
 							}
 							else if (FieldState == (int32)EObjectStateTypeEnum::Chaos_Object_Sleeping)
 							{
-								if (HandleState == Chaos::EObjectStateType::Dynamic)
+								if (HandleState != Chaos::EObjectStateType::Sleeping)
 								{
 									RigidHandle->SetObjectStateLowLevel(Chaos::EObjectStateType::Sleeping);
 									StateChanged = true;
