PushPhysicsState Optim
 -Use a pool of particle data instead of alloc/delete them each frame
 -Batch lamda to reduce the overhead of the taskGraph when we have lots of dirty particles
 -Process each dirty particle array individally

diff --git a/Engine/Source/Runtime/Experimental/Chaos/Public/Chaos/ParticleHandle.h b/Engine/Source/Runtime/Experimental/Chaos/Public/Chaos/ParticleHandle.h
--- a/Engine/Source/Runtime/Experimental/Chaos/Public/Chaos/ParticleHandle.h
+++ b/Engine/Source/Runtime/Experimental/Chaos/Public/Chaos/ParticleHandle.h
@@ -1655,6 +1655,34 @@ public:
 #endif
 	}
 
+	void Init(const TGeometryParticle<T, d>& InParticle)
+	{
+		Type = EParticleType::Static;
+		X = InParticle.X();
+		R = InParticle.R();
+		Geometry = InParticle.GeometrySharedLowLevel();
+		SpatialIdx = InParticle.SpatialIdx();
+		UniqueIdx = InParticle.UniqueIdx();
+		UserData = InParticle.UserData();
+		DirtyFlags = InParticle.DirtyFlags();
+#if CHAOS_CHECKED
+		DebugName = InParticle.DebugName();
+#endif
+		const TShapesArray<T, d>& Shapes = InParticle.ShapesArray();
+		ShapeCollisionDisableFlags.Empty(Shapes.Num());
+		CollisionTraceType.Empty(Shapes.Num());
+		ShapeSimData.Empty(Shapes.Num());
+		ShapeQueryData.Empty(Shapes.Num());
+		for (const TUniquePtr<TPerShapeData<T, d>>& ShapePtr : Shapes)
+		{
+			ShapeCollisionDisableFlags.Add(ShapePtr->bDisable);
+			CollisionTraceType.Add(ShapePtr->CollisionTraceType);
+			ShapeSimData.Add(ShapePtr->SimData);
+			ShapeSimData.Add(ShapePtr->SimData);
+			ShapeQueryData.Add(ShapePtr->QueryData);
+		}
+	}
+
 	TVector<T, d> X;
 	TRotation<T, d> R;
 	TSharedPtr<FImplicitObject, ESPMode::ThreadSafe> Geometry;
@@ -1781,6 +1809,15 @@ public:
 		MRotationOfMass = TRotation<T, d>(FQuat(EForceInit::ForceInit));
 	}
 
+	void Init(const TKinematicGeometryParticle<T, d>& InParticle) {
+		Base::Init(InParticle);
+		MV = InParticle.V();
+		MW = InParticle.W();
+		MCenterOfMass = InParticle.CenterOfMass();
+		MRotationOfMass = InParticle.RotationOfMass();
+		Type = EParticleType::Kinematic;
+	}
+
 	TVector<T, d> MV;
 	TVector<T, d> MW;
 	TVector<T, d> MCenterOfMass;
@@ -2178,6 +2215,32 @@ public:
 		MToBeRemovedOnFracture = false;
 		MGravityEnabled = false;
 	}
+
+	void Init(const TPBDRigidParticle<T, d>& InParticle) {
+		Base::Init(InParticle);
+		MQ = InParticle.Q();
+		MPreV = InParticle.PreV();
+		MPreW = InParticle.PreW();
+		MP = InParticle.P();
+		MF = InParticle.F();
+		MTorque = InParticle.Torque();
+		MLinearImpulse = TVector<T, d>(0);
+		MAngularImpulse = TVector<T, d>(0);
+		MI = InParticle.I();
+		MInvI = InParticle.InvI();
+		MCollisionParticles = nullptr;
+		MM = InParticle.M();
+		MInvM = InParticle.InvM();
+		MLinearEtherDrag = InParticle.LinearEtherDrag();
+		MAngularEtherDrag = InParticle.AngularEtherDrag();
+		MIsland = InParticle.Island();
+		MCollisionGroup = InParticle.CollisionGroup();
+		MObjectState = InParticle.ObjectState();
+		MDisabled = InParticle.Disabled();
+		MToBeRemovedOnFracture = InParticle.ToBeRemovedOnFracture();
+		MGravityEnabled = InParticle.IsGravityEnabled();
+		Type = EParticleType::Rigid;
+	}
 };
 
 
diff --git a/Engine/Source/Runtime/Experimental/ChaosSolvers/Private/PBDRigidsSolver.cpp b/Engine/Source/Runtime/Experimental/ChaosSolvers/Private/PBDRigidsSolver.cpp
--- a/Engine/Source/Runtime/Experimental/ChaosSolvers/Private/PBDRigidsSolver.cpp
+++ b/Engine/Source/Runtime/Experimental/ChaosSolvers/Private/PBDRigidsSolver.cpp
@@ -18,9 +18,10 @@
 #include "PhysicsProxy/FieldSystemPhysicsProxy.h"
 #include "EventDefaults.h"
 #include "EventsData.h"
-
+//#include "ProfilingDebugging/CsvProfiler.h"
 
 DEFINE_LOG_CATEGORY_STATIC(LogPBDRigidsSolverSolver, Log, All);
+//CSV_DEFINE_CATEGORY(Chaos, true);
 
 // DebugDraw CVars
 #if CHAOS_DEBUG_DRAW
@@ -494,51 +495,6 @@ namespace Chaos
 		GetEventManager()->DispatchEvents();
 	}
 
-	/**
-	 * ProxyType is FSingleParticlePhysicsProxy<T>, where T is:
-	 *    Chaos::TPBDRigidParticle<float,3>
-	 *    Chaos::TKinematicGeometryParticle<float,3>
-	 *    Chaos::TGeometryParticle<float,3>
-	 */
-	template<typename ProxyType>
-	void PushPhysicsStateExec(FPBDRigidsSolver * Solver, ProxyType* Proxy, Chaos::IDispatcher* Dispatcher)
-	{
-		LLM_SCOPE(ELLMTag::Chaos);
-
-		if (Chaos::FParticleData* ProxyData = Proxy->NewData())
-		{
-			if (auto* RigidHandle = static_cast<Chaos::TGeometryParticleHandle<float, 3>*>(Proxy->GetHandle()))
-			{
-				if (Dispatcher)
-				{
-					Dispatcher->EnqueueCommandImmediate([Proxy, ProxyData, Solver](Chaos::FPersistentPhysicsTask* PhysThread)
-					{
-						// make sure the handle is still valid
-						if (auto* Handle = static_cast<Chaos::TGeometryParticleHandle<float, 3>*>(Proxy->GetHandle()))
-						{
-							Solver->GetEvolution()->DirtyParticle(*Handle);
-							Proxy->PushToPhysicsState(ProxyData);
-						}
-						delete ProxyData;
-					});
-				}
-				else
-				{
-					Solver->GetEvolution()->DirtyParticle(*RigidHandle);
-					Proxy->PushToPhysicsState(ProxyData);
-					delete ProxyData;
-				}
-				Solver->RemoveDirtyProxy(Proxy);
-
-				Proxy->ClearAccumulatedData();
-			}
-			else
-			{
-				delete ProxyData;
-			}
-		}
-	}
-
 	void PushPhysicsStateExec(FPBDRigidsSolver* Solver, FGeometryCollectionPhysicsProxy* Proxy, Chaos::IDispatcher* Dispatcher)
 	{
 		Proxy->NewData();
@@ -580,33 +536,58 @@ namespace Chaos
 
 	void FPBDRigidsSolver::PushPhysicsState(IDispatcher* Dispatcher)
 	{
-		if (DirtyProxiesSet.Num() == 0)
-			return;
-		TArray<IPhysicsProxyBase*> DirtyProxiesArray = DirtyProxiesSet.Array();
-		for (auto & Proxy : DirtyProxiesArray)
+		//CSV_SCOPED_TIMING_STAT(Chaos, PushPhysicsState);
+
+		TArray< IPhysicsProxyBase *> DirtyProxiesArray = DirtyProxiesSet.Array();
+
+		// init the pooled particle data
+		for (int32 i = 0; i < DirtyProxiesArray.Num(); i++)
 		{
+			auto Proxy = DirtyProxiesArray[i];
 			switch (Proxy->GetType())
 			{
-			//case EPhysicsProxyType::NoneType: // 0
-			//case EPhysicsProxyType::StaticMeshType: // 1
+				// @todo(RyanK): May need to port these to using pools
 			case EPhysicsProxyType::GeometryCollectionType: // 2
 				PushPhysicsStateExec(this, static_cast<FGeometryCollectionPhysicsProxy*>(Proxy), Dispatcher);
 				break;
 			case EPhysicsProxyType::FieldType: // 3
 				PushPhysicsStateExec(this, static_cast<FFieldSystemPhysicsProxy*>(Proxy), Dispatcher);
 				break;
-			//case EPhysicsProxyType::SkeletalMeshType: // 4
-			case EPhysicsProxyType::SingleRigidParticleType: // 7
-				PushPhysicsStateExec(this, static_cast<FSingleParticlePhysicsProxy< Chaos::TPBDRigidParticle<float, 3> >*>(Proxy), Dispatcher);
-				break;
-			case EPhysicsProxyType::SingleKinematicParticleType: // 6
-				PushPhysicsStateExec(this, static_cast<FSingleParticlePhysicsProxy< Chaos::TKinematicGeometryParticle<float, 3> >*>(Proxy), Dispatcher);
-				break;
-			case EPhysicsProxyType::SingleGeometryParticleType: // 5
-				PushPhysicsStateExec(this, static_cast<FSingleParticlePhysicsProxy< Chaos::TGeometryParticle<float, 3> >*>(Proxy), Dispatcher);
-				break;
-			default:
-				ensure("Unknown proxy type in physics solver.");
+			case EPhysicsProxyType::SingleRigidParticleType:
+			{
+				FSingleParticlePhysicsProxy< Chaos::TPBDRigidParticle<float, 3> >* ProxyType = static_cast<FSingleParticlePhysicsProxy< Chaos::TPBDRigidParticle<float, 3> >*>(Proxy);
+				if (auto* RigidHandle = static_cast<Chaos::TGeometryParticleHandle<float, 3>*>(ProxyType->GetHandle()))
+				{
+					TPoolParticleData<Chaos::TPBDRigidParticleData<float, 3>, FSingleParticlePhysicsProxy< Chaos::TPBDRigidParticle<float, 3> >*, 800>::PoolEntry& Entry = RigidParticlePool.GetNewEntry();
+					Entry.Proxy = ProxyType;
+					Entry.ParticleData.Init(*ProxyType->GetParticle());
+					DirtyProxiesSet.Remove(Proxy);
+				}
+			}
+			break;
+			case EPhysicsProxyType::SingleKinematicParticleType:
+			{
+				FSingleParticlePhysicsProxy< Chaos::TKinematicGeometryParticle<float, 3> >* ProxyType = static_cast<FSingleParticlePhysicsProxy< Chaos::TKinematicGeometryParticle<float, 3> >*>(Proxy);
+				if (auto* RigidHandle = static_cast<Chaos::TGeometryParticleHandle<float, 3>*>(ProxyType->GetHandle()))
+				{
+					TPoolParticleData<Chaos::TKinematicGeometryParticleData<float, 3>, FSingleParticlePhysicsProxy< Chaos::TKinematicGeometryParticle<float, 3> >*, 100>::PoolEntry& Entry = KinematicGeometryParticlePool.GetNewEntry();
+					Entry.Proxy = ProxyType;
+					Entry.ParticleData.Init(*ProxyType->GetParticle());
+					DirtyProxiesSet.Remove(Proxy);
+				}
+			}
+			break;
+			case EPhysicsProxyType::SingleGeometryParticleType:
+			{
+				FSingleParticlePhysicsProxy< Chaos::TGeometryParticle<float, 3> >* ProxyType = static_cast<FSingleParticlePhysicsProxy< Chaos::TGeometryParticle<float, 3> >*>(Proxy);
+				if (auto* RigidHandle = static_cast<Chaos::TGeometryParticleHandle<float, 3>*>(ProxyType->GetHandle()))
+				{
+					TPoolParticleData<Chaos::TGeometryParticleData<float, 3>, FSingleParticlePhysicsProxy< Chaos::TGeometryParticle<float, 3> >*, 100>::PoolEntry& Entry = GeometryParticlePool.GetNewEntry();
+					Entry.Proxy = ProxyType;
+					Entry.ParticleData.Init(*ProxyType->GetParticle());
+					DirtyProxiesSet.Remove(Proxy);
+				}
+			}
 			}
 		}
 	}
diff --git a/Engine/Source/Runtime/Experimental/ChaosSolvers/Public/PBDRigidsSolver.h b/Engine/Source/Runtime/Experimental/ChaosSolvers/Public/PBDRigidsSolver.h
--- a/Engine/Source/Runtime/Experimental/ChaosSolvers/Public/PBDRigidsSolver.h
+++ b/Engine/Source/Runtime/Experimental/ChaosSolvers/Public/PBDRigidsSolver.h
@@ -377,6 +377,118 @@ namespace Chaos
 
 		template<ELockType>
 		friend struct TSolverQueryMaterialScope;
+
+		template<typename DirtyParticleData, typename DirtyProxy, SIZE_T PreAllocCount>
+		class TPoolParticleData
+		{
+		public:
+			struct PoolEntry
+			{
+				PoolEntry() {}
+
+				DirtyParticleData ParticleData;
+				DirtyProxy Proxy;
+			};
+
+			TPoolParticleData()
+			{
+				NumDirty = 0;
+				// Prealloc default object in the pool. 
+				// Must call DirtyParticleData.Init + DirtyParticleData.Reset each 
+				// time you use an entry in the pool.
+				Pool.AddDefaulted(PreAllocCount);
+			}
+
+			~TPoolParticleData()
+			{
+				Pool.Empty();
+			}
+
+			void Enqueue(FPBDRigidsSolver * Solver, Chaos::IDispatcher* Dispatcher)
+			{
+				if (GetNum() == 0)
+				{
+					return;
+				}
+
+				static const uint32 TASK_GRANULARITY = 10;
+				uint32 StartIndex = 0;
+				uint32 EndIndex = 0;
+
+				if (Dispatcher != nullptr)
+				{
+					do
+					{
+						uint32 NextUpperBound = StartIndex + TASK_GRANULARITY;
+						EndIndex = (NextUpperBound > GetNum()) ? GetNum() : NextUpperBound;
+
+						Dispatcher->EnqueueCommandImmediate([this, Solver, Dispatcher, StartIndex, EndIndex](Chaos::FPersistentPhysicsTask* PhysThread)
+						{
+							for (uint32 SubIndex = StartIndex; SubIndex < EndIndex; SubIndex++)
+							{
+								PoolEntry& Entry = GetEntry(SubIndex);
+								if (Entry.Proxy != nullptr)
+								{
+									// make sure the handle is still valid
+									if (auto* Handle = static_cast<Chaos::TGeometryParticleHandle<float, 3>*>(Entry.Proxy->GetHandle()))
+									{
+										Solver->GetEvolution()->DirtyParticle(*Handle);
+										Entry.Proxy->PushToPhysicsState(static_cast<Chaos::FParticleData*>(&Entry.ParticleData));
+									}
+									Entry.Proxy->ClearAccumulatedData();
+									Entry.ParticleData.Reset();
+								}
+							}
+						});
+
+						uint32 NextLowerBound = StartIndex + TASK_GRANULARITY;
+						StartIndex = (NextLowerBound > GetNum()) ? StartIndex : NextLowerBound;
+
+					} while (EndIndex < GetNum());
+				}
+				else
+				{
+					for (uint32 i = 0; i < GetNum(); i++)
+					{
+						PoolEntry& Entry = GetEntry(i);
+						if (Entry.Proxy != nullptr)
+						{
+							// make sure the handle is still valid
+							if (auto* Handle = static_cast<Chaos::TGeometryParticleHandle<float, 3>*>(Entry.Proxy->GetHandle()))
+							{
+								Solver->GetEvolution()->DirtyParticle(*Handle);
+								Entry.Proxy->PushToPhysicsState(static_cast<Chaos::FParticleData*>(&Entry.ParticleData));
+							}
+							Entry.Proxy->ClearAccumulatedData();
+							Entry.ParticleData.Reset();
+						}
+					}
+				}
+			}
+
+			uint32 GetNum() {return NumDirty;}
+			void Reset() { NumDirty = 0; }
+
+			PoolEntry& GetNewEntry() {
+				if (NumDirty >= Pool.Num())
+				{
+					Pool.Add(PoolEntry());
+				}
+				return Pool[NumDirty++];
+			}
+
+			PoolEntry& GetEntry(int32 index) {
+				ensure(index < Pool.Num());
+				return Pool[index];
+			}
+
+			TArray<PoolEntry> Pool;
+			int32 NumDirty;
+		};
+
+		TPoolParticleData<Chaos::TPBDRigidParticleData<float, 3>, FSingleParticlePhysicsProxy< Chaos::TPBDRigidParticle<float, 3> >*, 800> RigidParticlePool;
+		TPoolParticleData<Chaos::TKinematicGeometryParticleData<float, 3>, FSingleParticlePhysicsProxy< Chaos::TKinematicGeometryParticle<float, 3> >*, 100> KinematicGeometryParticlePool;
+		TPoolParticleData<Chaos::TGeometryParticleData<float, 3>, FSingleParticlePhysicsProxy< Chaos::TGeometryParticle<float, 3> >*, 100> GeometryParticlePool;
 	};
 
 	template<>
diff --git a/Engine/Source/Runtime/Experimental/ChaosSolvers/Public/PhysicsProxy/SingleParticlePhysicsProxy.h b/Engine/Source/Runtime/Experimental/ChaosSolvers/Public/PhysicsProxy/SingleParticlePhysicsProxy.h
--- a/Engine/Source/Runtime/Experimental/ChaosSolvers/Public/PhysicsProxy/SingleParticlePhysicsProxy.h
+++ b/Engine/Source/Runtime/Experimental/ChaosSolvers/Public/PhysicsProxy/SingleParticlePhysicsProxy.h
@@ -145,6 +145,11 @@ public:
 	/**/
 	void ClearEvents();
 
+	PARTICLE_TYPE* GetParticle()
+	{
+		return Particle;
+	}
+
 private:
 	bool bInitialized;
 	TArray<int32> InitializedIndices;
