Change the way AABBTree rebuild task is handled. Fixes several spikes and makes it easier to tune the timeslice.

diff --git a/Engine/Source/Runtime/Experimental/Chaos/Private/Chaos/PBDRigidsEvolution.cpp b/Engine/Source/Runtime/Experimental/Chaos/Private/Chaos/PBDRigidsEvolution.cpp
--- a/Engine/Source/Runtime/Experimental/Chaos/Private/Chaos/PBDRigidsEvolution.cpp
+++ b/Engine/Source/Runtime/Experimental/Chaos/Private/Chaos/PBDRigidsEvolution.cpp
@@ -29,7 +29,7 @@ namespace Chaos
 			AABBMaxChildrenInLeaf = 500;
 			AABBMaxTreeDepth = 200;
 			MaxPayloadSize = 100000;
-			IterationsPerTimeSlice = 40000;
+			IterationsPerTimeSlice = 4000;
 		}
 	} ConfigSettings;
 
diff --git a/Engine/Source/Runtime/Experimental/Chaos/Public/Chaos/AABBTree.h b/Engine/Source/Runtime/Experimental/Chaos/Public/Chaos/AABBTree.h
--- a/Engine/Source/Runtime/Experimental/Chaos/Public/Chaos/AABBTree.h
+++ b/Engine/Source/Runtime/Experimental/Chaos/Public/Chaos/AABBTree.h
@@ -304,28 +304,11 @@ public:
 		}
 
 		// still has work to complete
-		if (!TimeSliceWorkToComplete.IsEmpty())
+		if (WorkStack.Num())
 		{
-			FWorkSnapshot Store;
-			TimeSliceWorkToComplete.Dequeue(Store);
 			NumProcessedThisSlice = 0;
-			if (Store.TimeslicePhase == eTimeSlicePhase::PHASE1)
-			{
-				SplitNode(Store.Bounds, Store.Elems, Store.NodeLevel, Store.NewNodeIdx);
-			}
-			else
-			{
-				check(Store.TimeslicePhase == eTimeSlicePhase::PHASE2);
-				SplitNode(Store.Bounds, Store.RealBounds, Store.Elems, Store.SplitBoundsSize2, Store.NewNodeIdx, Store.NodeLevel, Store.BoxIdx);
-			}
+			SplitNode();
 		}
-
-		// are done now?
-		if (TimeSliceWorkToComplete.IsEmpty())
-		{
-			this->SetAsyncTimeSlicingComplete(true);
-		}
-
 	}
 
 	template <typename TParticles>
@@ -758,15 +741,44 @@ private:
 		return true;
 	}
 
+	int32 GetNewWorkSnapshot()
+	{
+		if(WorkPoolFreeList.Num())
+		{
+			return WorkPoolFreeList.Pop();
+		}
+		else
+		{
+			return WorkPool.AddDefaulted(1);
+		}
+	}
+
+	void FreeWorkSnapshot(int32 WorkSnapshotIdx)
+	{
+		//Reset for next time they want to use it
+		WorkPool[WorkSnapshotIdx] = FWorkSnapshot();
+		WorkPoolFreeList.Add(WorkSnapshotIdx);
+		
+	}
+
 	template <typename TParticles>
 	void GenerateTree(const TParticles& Particles)
 	{
 		SCOPE_CYCLE_COUNTER(STAT_AABBTreeGenerateTree);
 		this->SetAsyncTimeSlicingComplete(false);
 
-		TArray<FElement> ElemsWithBounds;
-		ElemsWithBounds.Reserve(Particles.Num());
+		ensure(WorkStack.Num() == 0);
+
+		const int32 ExpectedNumLeaves = Particles.Num() / MaxChildrenInLeaf;
+		const int32 ExpectedNumNodes = ExpectedNumLeaves;
+
+		WorkStack.Reserve(ExpectedNumNodes);
 
+		const int32 CurIdx = GetNewWorkSnapshot();
+		FWorkSnapshot& WorkSnapshot = WorkPool[CurIdx];
+		WorkSnapshot.Elems.Reserve(Particles.Num());
+		
+		
 		GlobalPayloads.Reset();
 		Leaves.Reset();
 		Nodes.Reset();
@@ -774,13 +786,14 @@ private:
 		PayloadToInfo.Reset();
 		NumProcessedThisSlice = 0;
 
-		FullBounds = TAABB<T, 3>::EmptyAABB();
+		WorkSnapshot.Bounds = TAABB<T, 3>::EmptyAABB();
 
 		{
 			SCOPE_CYCLE_COUNTER(STAT_AABBTreeTimeSliceSetup);
 
 			int32 Idx = 0;
 
+			//TODO: we need a better way to time-slice this case since there can be a huge number of Particles. Can't do it right now without making full copy
 			for (auto& Particle : Particles)
 			{
 				bool bHasBoundingBox = HasBoundingBox(Particle);
@@ -795,8 +808,8 @@ private:
 					}
 					else
 					{
-						ElemsWithBounds.Add(FElement{ Payload, ElemBounds });
-						FullBounds.GrowToInclude(ElemBounds);
+						WorkSnapshot.Elems.Add(FElement{ Payload, ElemBounds });
+						WorkSnapshot.Bounds.GrowToInclude(ElemBounds);
 					}
 				}
 				else
@@ -816,201 +829,287 @@ private:
 				++Idx;
 				//todo: payload info
 			}
-
 		}
 
+		NumProcessedThisSlice = Particles.Num();	//todo: give chance to time slice out of next phase
+
 		{
 			SCOPE_CYCLE_COUNTER(STAT_AABBTreeInitialTimeSlice);
+			WorkSnapshot.NewNodeIdx = 0;
+			WorkSnapshot.NodeLevel = 0;
 
-			SplitNode(FullBounds, ElemsWithBounds, 0, 0, true);
+			//push root onto stack
+			WorkStack.Add(CurIdx);
 
-			if (TimeSliceWorkToComplete.IsEmpty())
+			SplitNode();
+		}
+
+		/*  Helper validation code 
+		int32 Count = 0;
+		TSet<int32> Seen;
+		if(WorkStack.Num() == 0)
+		{
+			int32 LeafIdx = 0;
+			for (const auto& Leaf : Leaves)
 			{
-				this->SetAsyncTimeSlicingComplete(true);
+				Validate(Seen, Count, Leaf);
+				bool bHasParent = false;
+				for (const auto& Node : Nodes)
+				{
+					if (Node.bLeaf && Node.ChildrenNodes[0] == LeafIdx)
+					{
+						bHasParent = true;
+
+						break;
+					}
+				}
+				ensure(bHasParent);
+				++LeafIdx;
 			}
+			ensure(Count == 0 || Seen.Num() == Count);
+			ensure(Count == 0 || Count == Particles.Num());
 		}
+		*/
 
 	}
 
-	void SplitNode(TAABB<T, 3>& SplitBounds, TAABB<T, 3>& RealBounds, TArray<FElement>& Children, T& SplitBoundsSize2, int NewNodeIdx, int NodeLevel, int BoxIdx)
+	enum eTimeSlicePhase
 	{
-		//SCOPE_CYCLE_COUNTER(STAT_AABBTreeChildrenPhase);
-
-		check (Children.Num() > 0)
-		{
-			Nodes[NewNodeIdx].bLeaf = false;
-
-			// increment this early so the amount of recursion work left on the stack is smaller
-			NumProcessedThisSlice += Children.Num();
-
-			const int32 ChildIdx = Nodes.Num();
-			Nodes[NewNodeIdx].ChildrenBounds[BoxIdx] = RealBounds;
-			Nodes[NewNodeIdx].ChildrenNodes[BoxIdx] = ChildIdx;
-
-			SplitNode(RealBounds, Children, NodeLevel + 1, ChildIdx);
-		}
+		PreFindBestBounds,
+		DuringFindBestBounds,
+		ProcessingChildren
+	};
 
-	}
+	struct FSplitInfo
+	{
+		TAABB<T, 3> SplitBounds;	//Even split of parent bounds
+		TAABB<T, 3> RealBounds;	//Actual bounds as children are added
+		int32 WorkSnapshotIdx;	//Idx into work snapshot pool
+		T SplitBoundsSize2;
+	};
 
-	void SplitNode(const TAABB<T, 3>& Bounds, const TArray<FElement>& Elems, int32 NodeLevel, int32 NewNodeIdx, bool Initial=false)
+	struct FWorkSnapshot
 	{
-		// create the actual node space but might no be filled in (YET) due to time slicing exit
-		if (NewNodeIdx >= Nodes.Num())
+		FWorkSnapshot()
+			: TimeslicePhase(eTimeSlicePhase::PreFindBestBounds)
 		{
-		 	Nodes.AddDefaulted();	//todo: remove TAABB
+
 		}
 
-		bool WeAreTimeslicing = (MaxNumToProcess > 0);
+		eTimeSlicePhase TimeslicePhase;
 
-		if (WeAreTimeslicing && (NumProcessedThisSlice >= MaxNumToProcess))
-		{
-			// done enough work, capture stack
-			FWorkSnapshot Store;
-			Store.TimeslicePhase = eTimeSlicePhase::PHASE1;
-			Store.Bounds = Bounds;
-			Store.Elems = Elems;			
-			Store.NodeLevel = NodeLevel;
-			Store.NewNodeIdx = NewNodeIdx;
+		TAABB<T, 3> Bounds;
+		TArray<FElement> Elems;
 
-			TimeSliceWorkToComplete.Enqueue(Store);
+		int32 NodeLevel;
+		int32 NewNodeIdx;
 
-			return; // done enough
-		}
+		int32 BestBoundsCurIdx;
 
-		auto& PayloadToInfoRef = PayloadToInfo;
-		auto& LeavesRef = Leaves;
-		auto& NodesRef = Nodes;
-		auto MakeLeaf = [NewNodeIdx, &PayloadToInfoRef, &Elems, &LeavesRef, &NodesRef]()
+		FSplitInfo SplitInfos[2];
+	};
+
+	void FindBestBounds(const int32 StartElemIdx, const int32 LastElem, FWorkSnapshot& CurrentSnapshot)
+	{
+		// add all elements to one of the two split infos at this level - root level [ not taking into account the max number allowed or anything
+		for(int32 ElemIdx = StartElemIdx; ElemIdx < LastElem; ++ElemIdx)
 		{
-			if (bMutable)
+			const FElement& Elem = CurrentSnapshot.Elems[ElemIdx];
+			int32 MinBoxIdx = INDEX_NONE;
+			T MinDelta2 = TNumericLimits<T>::Max();
+			int32 BoxIdx = 0;
+			for (const FSplitInfo& SplitInfo : CurrentSnapshot.SplitInfos)
 			{
-				for (const FElement& Elem : Elems)
+				TAABB<T, 3> NewBox = SplitInfo.SplitBounds;
+				NewBox.GrowToInclude(Elem.Bounds);
+				const T Delta2 = NewBox.Extents().SizeSquared() - SplitInfo.SplitBoundsSize2;
+				if (Delta2 < MinDelta2)
 				{
-					PayloadToInfoRef.Add(Elem.Payload, FAABBTreePayloadInfo{ INDEX_NONE, INDEX_NONE, LeavesRef.Num() });
+					MinDelta2 = Delta2;
+					MinBoxIdx = BoxIdx;
 				}
+				++BoxIdx;
 			}
 
-			NodesRef[NewNodeIdx].bLeaf = true;
-			NodesRef[NewNodeIdx].ChildrenNodes[0] = LeavesRef.Add(TLeafType{ Elems }); //todo: avoid copy?
+			if (CHAOS_ENSURE(MinBoxIdx != INDEX_NONE))
+			{
+				FSplitInfo& SplitInfo = CurrentSnapshot.SplitInfos[MinBoxIdx];
+				WorkPool[SplitInfo.WorkSnapshotIdx].Elems.Add(Elem);
+				SplitInfo.RealBounds.GrowToInclude(Elem.Bounds);
+			}
+		}
 
-		};
+		NumProcessedThisSlice += LastElem - StartElemIdx;
+	}
 
-		if (Elems.Num() <= MaxChildrenInLeaf || NodeLevel >= MaxTreeDepth)
+	void SplitNode()
+	{
+		const bool WeAreTimeslicing = (MaxNumToProcess > 0);
+
+		while (WorkStack.Num())
 		{
-			MakeLeaf();
-			return; // keep working
-		}
+			//NOTE: remember to be careful with this since it's a pointer on a tarray
+			const int32 CurIdx = WorkStack.Last();
 
-		const TVector<T, 3> Extents = Bounds.Extents();
-		const int32 MaxAxis = Bounds.LargestAxis();
-		bool ChildrenInBothHalves = false;
+			if (WorkPool[CurIdx].TimeslicePhase == eTimeSlicePhase::ProcessingChildren)
+			{
+				//If we got to this it must be that my children are done, so I'm done as well
+				WorkStack.Pop(/*bResize=*/false);
+				FreeWorkSnapshot(CurIdx);
+				continue;
+			}
 
-		FSplitInfo SplitInfos[2];
-		SplitInfos[0].SplitBounds = TAABB<T, 3>(Bounds.Min(), Bounds.Min());
-		SplitInfos[1].SplitBounds = TAABB<T, 3>(Bounds.Max(), Bounds.Max());
+			const int32 NewNodeIdx = WorkPool[CurIdx].NewNodeIdx;
 
-		const TVector<T, 3> Center = Bounds.Center();
-		for (FSplitInfo& SplitInfo : SplitInfos)
-		{
-			SplitInfo.RealBounds = TAABB<T, 3>::EmptyAABB();
+			// create the actual node space but might no be filled in (YET) due to time slicing exit
+			if (NewNodeIdx >= Nodes.Num())
+			{
+				Nodes.AddDefaulted((1 + NewNodeIdx) - Nodes.Num());
+			}
 
-			for (int32 Axis = 0; Axis < 3; ++Axis)
+			if (WeAreTimeslicing && (NumProcessedThisSlice >= MaxNumToProcess))
 			{
-				TVector<T, 3> NewPt0 = Center;
-				TVector<T, 3> NewPt1 = Center;
-				if (Axis != MaxAxis)
+				return; // done enough
+			}
+
+			auto& PayloadToInfoRef = PayloadToInfo;
+			auto& LeavesRef = Leaves;
+			auto& NodesRef = Nodes;
+			auto& WorkPoolRef = WorkPool;
+			auto MakeLeaf = [NewNodeIdx, &PayloadToInfoRef, &WorkPoolRef, CurIdx, &LeavesRef, &NodesRef]()
+			{
+				if (bMutable)
 				{
-					NewPt0[Axis] = Bounds.Min()[Axis];
-					NewPt1[Axis] = Bounds.Max()[Axis];
-					SplitInfo.SplitBounds.GrowToInclude(NewPt0);
-					SplitInfo.SplitBounds.GrowToInclude(NewPt1);
+					//todo: does this need time slicing in the case when we have a ton of elements that can't be split?
+					//hopefully not a real issue
+					for (const FElement& Elem : WorkPoolRef[CurIdx].Elems)
+					{
+						PayloadToInfoRef.Add(Elem.Payload, FAABBTreePayloadInfo{ INDEX_NONE, INDEX_NONE, LeavesRef.Num() });
+					}
 				}
-			}
 
-			SplitInfo.SplitBoundsSize2 = SplitInfo.SplitBounds.Extents().SizeSquared();
-		}
-		
-		{
-			//SCOPE_CYCLE_COUNTER(STAT_AABBTreeGrowPhase);
+				NodesRef[NewNodeIdx].bLeaf = true;
+				NodesRef[NewNodeIdx].ChildrenNodes[0] = LeavesRef.Add(TLeafType{ WorkPoolRef[CurIdx].Elems }); //todo: avoid copy?
+
+			};
 
-			// add all elements to one of the two split infos at this level - root level [ not taking into account the max number allowed or anything
-			for (const FElement& Elem : Elems)
+			if (WorkPool[CurIdx].Elems.Num() <= MaxChildrenInLeaf || WorkPool[CurIdx].NodeLevel >= MaxTreeDepth)
 			{
-				int32 MinBoxIdx = INDEX_NONE;
-				T MinDelta2 = TNumericLimits<T>::Max();
-				int32 BoxIdx = 0;
-				for (const FSplitInfo& SplitInfo : SplitInfos)
+
+				MakeLeaf();
+				WorkStack.Pop(/*bResize=*/false);	//finished with this node
+				FreeWorkSnapshot(CurIdx);
+				continue;
+			}
+
+			if (WorkPool[CurIdx].TimeslicePhase == eTimeSlicePhase::PreFindBestBounds)
+			{
+				const TVector<T, 3> Extents = WorkPool[CurIdx].Bounds.Extents();
+				const int32 MaxAxis = WorkPool[CurIdx].Bounds.LargestAxis();
+
+				//Add two children, remember this invalidates any pointers to current snapshot
+				const int32 FirstChildIdx = GetNewWorkSnapshot();
+				const int32 SecondChildIdx = GetNewWorkSnapshot();
+
+				//mark idx of children into the work pool
+				WorkPool[CurIdx].SplitInfos[0].WorkSnapshotIdx = FirstChildIdx;
+				WorkPool[CurIdx].SplitInfos[1].WorkSnapshotIdx = SecondChildIdx;
+
+				//these are the hypothetical bounds for the perfect 50/50 split
+				WorkPool[CurIdx].SplitInfos[0].SplitBounds = TAABB<T, 3>(WorkPool[CurIdx].Bounds.Min(), WorkPool[CurIdx].Bounds.Min());
+				WorkPool[CurIdx].SplitInfos[1].SplitBounds = TAABB<T, 3>(WorkPool[CurIdx].Bounds.Max(), WorkPool[CurIdx].Bounds.Max());
+
+				const TVector<T, 3> Center = WorkPool[CurIdx].Bounds.Center();
+				for (FSplitInfo& SplitInfo : WorkPool[CurIdx].SplitInfos)
 				{
-					TAABB<T, 3> NewBox = SplitInfo.SplitBounds;
-					NewBox.GrowToInclude(Elem.Bounds);
-					const T Delta2 = NewBox.Extents().SizeSquared() - SplitInfo.SplitBoundsSize2;
-					if (Delta2 < MinDelta2)
+					SplitInfo.RealBounds = TAABB<T, 3>::EmptyAABB();
+
+					for (int32 Axis = 0; Axis < 3; ++Axis)
 					{
-						MinDelta2 = Delta2;
-						MinBoxIdx = BoxIdx;
+						TVector<T, 3> NewPt0 = Center;
+						TVector<T, 3> NewPt1 = Center;
+						if (Axis != MaxAxis)
+						{
+							NewPt0[Axis] = WorkPool[CurIdx].Bounds.Min()[Axis];
+							NewPt1[Axis] = WorkPool[CurIdx].Bounds.Max()[Axis];
+							SplitInfo.SplitBounds.GrowToInclude(NewPt0);
+							SplitInfo.SplitBounds.GrowToInclude(NewPt1);
+						}
 					}
-					++BoxIdx;
+
+					SplitInfo.SplitBoundsSize2 = SplitInfo.SplitBounds.Extents().SizeSquared();
 				}
 
-				if (CHAOS_ENSURE(MinBoxIdx != INDEX_NONE))
+				WorkPool[CurIdx].BestBoundsCurIdx = 0;
+				WorkPool[CurIdx].TimeslicePhase = eTimeSlicePhase::DuringFindBestBounds;
+				const int32 ExpectedNumPerChild = WorkPool[CurIdx].Elems.Num() / 2;
 				{
-					SplitInfos[MinBoxIdx].Children.Add(Elem);
-					SplitInfos[MinBoxIdx].RealBounds.GrowToInclude(Elem.Bounds);
+					WorkPool[FirstChildIdx].Elems.Reserve(ExpectedNumPerChild);
+					WorkPool[SecondChildIdx].Elems.Reserve(ExpectedNumPerChild);
 				}
 			}
 
-			NumProcessedThisSlice += Elems.Num();
-
-			ChildrenInBothHalves = SplitInfos[0].Children.Num() && SplitInfos[1].Children.Num();
-
-			if (ChildrenInBothHalves && WeAreTimeslicing && (NumProcessedThisSlice >= MaxNumToProcess))
+			if (WorkPool[CurIdx].TimeslicePhase == eTimeSlicePhase::DuringFindBestBounds)
 			{
-				for (int32 BoxIdx = 0; BoxIdx < 2; ++BoxIdx)
+				const int32 NumWeCanProcess = MaxNumToProcess - NumProcessedThisSlice;
+				const int32 LastIdxToProcess = WeAreTimeslicing ? FMath::Min(WorkPool[CurIdx].BestBoundsCurIdx + NumWeCanProcess, WorkPool[CurIdx].Elems.Num()) : WorkPool[CurIdx].Elems.Num();
+				FindBestBounds(WorkPool[CurIdx].BestBoundsCurIdx, LastIdxToProcess, WorkPool[CurIdx]);
+				WorkPool[CurIdx].BestBoundsCurIdx = LastIdxToProcess;
+
+				if (WeAreTimeslicing && (NumProcessedThisSlice >= MaxNumToProcess))
 				{
-					// done enough work, capture stack
-					FWorkSnapshot Store;
-					Store.TimeslicePhase = eTimeSlicePhase::PHASE2;
-					Store.NodeLevel = NodeLevel;
-					Store.NewNodeIdx = NewNodeIdx;
-					Store.BoxIdx = BoxIdx;
-					Store.Bounds = MoveTemp(SplitInfos[BoxIdx].SplitBounds);
-					Store.Elems = MoveTemp(SplitInfos[BoxIdx].Children);
-					Store.RealBounds = MoveTemp(SplitInfos[BoxIdx].RealBounds);
-					Store.SplitBoundsSize2 = SplitInfos[BoxIdx].SplitBoundsSize2;
-					TimeSliceWorkToComplete.Enqueue(Store);
+					return; // done enough
 				}
-				return; // done enough
-			}	
-		}
+			}
+
 
-		{			
-			//SCOPE_CYCLE_COUNTER(STAT_AABBTreeChildrenPhase);
+			const int32 FirstChildIdx = WorkPool[CurIdx].SplitInfos[0].WorkSnapshotIdx;
+			const int32 SecondChildIdx = WorkPool[CurIdx].SplitInfos[1].WorkSnapshotIdx;
 
-			// if children in both halves, recurse a level down
-			if (ChildrenInBothHalves)
+			const bool bChildrenInBothHalves = WorkPool[FirstChildIdx].Elems.Num() && WorkPool[SecondChildIdx].Elems.Num();
+
+			// if children in both halves, push them on the stack to continue the split
+			if (bChildrenInBothHalves)
 			{
 				Nodes[NewNodeIdx].bLeaf = false;
 
-				// increment this early so the amount of recursion work left on the stack is smaller
-				NumProcessedThisSlice += SplitInfos[0].Children.Num() + SplitInfos[1].Children.Num();
+				Nodes[NewNodeIdx].ChildrenBounds[0] = WorkPool[CurIdx].SplitInfos[0].RealBounds;
+				WorkPool[FirstChildIdx].Bounds = Nodes[NewNodeIdx].ChildrenBounds[0];
+				Nodes[NewNodeIdx].ChildrenNodes[0] = Nodes.Num();
 
-				for (int32 BoxIdx = 0; BoxIdx < 2; ++BoxIdx)
-				{
-					const int32 ChildIdx = Nodes.Num();
-					Nodes[NewNodeIdx].ChildrenBounds[BoxIdx] = SplitInfos[BoxIdx].RealBounds;
-					Nodes[NewNodeIdx].ChildrenNodes[BoxIdx] = ChildIdx;
+				Nodes[NewNodeIdx].ChildrenBounds[1] = WorkPool[CurIdx].SplitInfos[1].RealBounds;
+				WorkPool[SecondChildIdx].Bounds = Nodes[NewNodeIdx].ChildrenBounds[1];
+				Nodes[NewNodeIdx].ChildrenNodes[1] = Nodes.Num() + 1;
 
-					SplitNode(SplitInfos[BoxIdx].RealBounds, SplitInfos[BoxIdx].Children, NodeLevel + 1, ChildIdx);
-				}
+				WorkPool[FirstChildIdx].NodeLevel = WorkPool[CurIdx].NodeLevel + 1;
+				WorkPool[SecondChildIdx].NodeLevel = WorkPool[CurIdx].NodeLevel + 1;
+
+				WorkPool[FirstChildIdx].NewNodeIdx = Nodes[NewNodeIdx].ChildrenNodes[0];
+				WorkPool[SecondChildIdx].NewNodeIdx = Nodes[NewNodeIdx].ChildrenNodes[1];
+
+				//push these two new nodes onto the stack
+				WorkStack.Add(SecondChildIdx);
+				WorkStack.Add(FirstChildIdx);
+
+				// create the actual node so that no one else can use our children node indices
+				const int32 HighestNodeIdx = Nodes[NewNodeIdx].ChildrenNodes[1];
+				Nodes.AddDefaulted((1 + HighestNodeIdx) - Nodes.Num());
+				
+				WorkPool[CurIdx].TimeslicePhase = eTimeSlicePhase::ProcessingChildren;
 			}
 			else
 			{
 				//couldn't split so just make a leaf - THIS COULD CONTAIN MORE THAN MaxChildrenInLeaf!!!
 				MakeLeaf();
+				WorkStack.Pop(/*bResize=*/false);	//we are done with this node
+				FreeWorkSnapshot(CurIdx);
 			}
 		}
 
-		return; // keep working
+		check(WorkStack.Num() == 0);
+		//Stack is empty, clean up pool and mark task as complete
+		
+		this->SetAsyncTimeSlicingComplete(true);
 	}
 
 	TArray<TPayloadType> FindAllIntersectionsImp(const TAABB<T, 3>& Intersection) const
@@ -1066,8 +1165,12 @@ private:
 
 	TAABBTree<TPayloadType,TLeafType,T,bMutable>& operator=(const TAABBTree<TPayloadType,TLeafType,T,bMutable>& Rhs)
 	{
-		ensure(Rhs.TimeSliceWorkToComplete.IsEmpty());
-		TimeSliceWorkToComplete.Empty();
+		ensure(Rhs.WorkStack.Num() == 0);
+		ensure(Rhs.WorkPool.Num() == 0);
+		ensure(Rhs.WorkPoolFreeList.Num() == 0);
+		WorkStack.Empty();
+		WorkPool.Empty();
+		WorkPoolFreeList.Empty();
 		if(this != &Rhs)
 		{
 			FullBounds = Rhs.FullBounds;
@@ -1086,14 +1189,6 @@ private:
 		return *this;
 	}
 
-	struct FSplitInfo
-	{
-		TAABB<T, 3> SplitBounds;	//Even split of parent bounds
-		TAABB<T, 3> RealBounds;	//Actual bounds as children are added
-		TArray<FElement> Children;
-		T SplitBoundsSize2;
-	};
-
 	TAABB<T, 3> FullBounds;
 	TArray<FNode> Nodes;
 	TArray<TLeafType> Leaves;
@@ -1105,32 +1200,11 @@ private:
 	int32 MaxTreeDepth;
 	T MaxPayloadBounds;
 	int32 MaxNumToProcess;
-	
-	enum eTimeSlicePhase
-	{
-		PHASE1,
-		PHASE2
-	};
-
-	struct FWorkSnapshot
-	{
-		eTimeSlicePhase TimeslicePhase;
-
-		// used for phase 1 & 2
-		TAABB<T, 3> Bounds;
-		TArray<FElement> Elems;
-
-		int32 NodeLevel;
-		int32 NewNodeIdx;
-	
-		// phase 2 data only
-		TAABB<T, 3> RealBounds;
-		int8 BoxIdx;
-		T SplitBoundsSize2;
-	};
 
 	int32 NumProcessedThisSlice;
-	TQueue<FWorkSnapshot> TimeSliceWorkToComplete;
+	TArray<int32> WorkStack;
+	TArray<int32> WorkPoolFreeList;
+	TArray<FWorkSnapshot> WorkPool;
 };
 
 template<typename TPayloadType, typename TLeafType, class T, bool bMutable>
