Fix float comparison with zero.

diff --git a/Engine/Source/Runtime/Experimental/Chaos/Private/Chaos/HeightField.cpp b/Engine/Source/Runtime/Experimental/Chaos/Private/Chaos/HeightField.cpp
--- a/Engine/Source/Runtime/Experimental/Chaos/Private/Chaos/HeightField.cpp
+++ b/Engine/Source/Runtime/Experimental/Chaos/Private/Chaos/HeightField.cpp
@@ -720,7 +720,7 @@ namespace Chaos
 		T InvCurrentLength = 1 / CurrentLength;
 		for(int Axis = 0; Axis < 3; ++Axis)
 		{
-			bParallel[Axis] = Dir[Axis] == 0;
+			bParallel[Axis] = FMath::IsNearlyZero(Dir[Axis], 1.e-8f);
 			InvDir[Axis] = bParallel[Axis] ? 0 : 1 / Dir[Axis];
 		}
 
@@ -920,7 +920,7 @@ namespace Chaos
 			T InvCurrentLength = 1 / CurrentLength;
 			for(int Axis = 0; Axis < 3; ++Axis)
 			{
-				bParallel[Axis] = Dir[Axis] == 0;
+				bParallel[Axis] = FMath::IsNearlyZero(Dir[Axis], 1.e-8f);
 				InvDir[Axis] = bParallel[Axis] ? 0 : 1 / Dir[Axis];
 			}
 
