Back out changelist 11191723. EdTest ensure

diff --git a/Engine/Source/Runtime/Experimental/Chaos/Private/Chaos/PBDRigidsEvolution.cpp b/Engine/Source/Runtime/Experimental/Chaos/Private/Chaos/PBDRigidsEvolution.cpp
--- a/Engine/Source/Runtime/Experimental/Chaos/Private/Chaos/PBDRigidsEvolution.cpp
+++ b/Engine/Source/Runtime/Experimental/Chaos/Private/Chaos/PBDRigidsEvolution.cpp
@@ -29,7 +29,7 @@ namespace Chaos
 			AABBMaxChildrenInLeaf = 500;
 			AABBMaxTreeDepth = 200;
 			MaxPayloadSize = 100000;
-			IterationsPerTimeSlice = 4000;
+			IterationsPerTimeSlice = 40000;
 		}
 	} ConfigSettings;
 
diff --git a/Engine/Source/Runtime/Experimental/Chaos/Public/Chaos/AABBTree.h b/Engine/Source/Runtime/Experimental/Chaos/Public/Chaos/AABBTree.h
--- a/Engine/Source/Runtime/Experimental/Chaos/Public/Chaos/AABBTree.h
+++ b/Engine/Source/Runtime/Experimental/Chaos/Public/Chaos/AABBTree.h
@@ -304,11 +304,28 @@ public:
 		}
 
 		// still has work to complete
-		if (WorkStack.Num())
+		if (!TimeSliceWorkToComplete.IsEmpty())
 		{
+			FWorkSnapshot Store;
+			TimeSliceWorkToComplete.Dequeue(Store);
 			NumProcessedThisSlice = 0;
-			SplitNode();
+			if (Store.TimeslicePhase == eTimeSlicePhase::PHASE1)
+			{
+				SplitNode(Store.Bounds, Store.Elems, Store.NodeLevel, Store.NewNodeIdx);
+			}
+			else
+			{
+				check(Store.TimeslicePhase == eTimeSlicePhase::PHASE2);
+				SplitNode(Store.Bounds, Store.RealBounds, Store.Elems, Store.SplitBoundsSize2, Store.NewNodeIdx, Store.NodeLevel, Store.BoxIdx);
+			}
 		}
+
+		// are done now?
+		if (TimeSliceWorkToComplete.IsEmpty())
+		{
+			this->SetAsyncTimeSlicingComplete(true);
+		}
+
 	}
 
 	template <typename TParticles>
@@ -741,44 +758,15 @@ private:
 		return true;
 	}
 
-	int32 GetNewWorkSnapshot()
-	{
-		if(WorkPoolFreeList.Num())
-		{
-			return WorkPoolFreeList.Pop();
-		}
-		else
-		{
-			return WorkPool.AddDefaulted(1);
-		}
-	}
-
-	void FreeWorkSnapshot(int32 WorkSnapshotIdx)
-	{
-		//Reset for next time they want to use it
-		WorkPool[WorkSnapshotIdx] = FWorkSnapshot();
-		WorkPoolFreeList.Add(WorkSnapshotIdx);
-		
-	}
-
 	template <typename TParticles>
 	void GenerateTree(const TParticles& Particles)
 	{
 		SCOPE_CYCLE_COUNTER(STAT_AABBTreeGenerateTree);
 		this->SetAsyncTimeSlicingComplete(false);
 
-		ensure(WorkStack.Num() == 0);
-
-		const int32 ExpectedNumLeaves = Particles.Num() / MaxChildrenInLeaf;
-		const int32 ExpectedNumNodes = ExpectedNumLeaves;
-
-		WorkStack.Reserve(ExpectedNumNodes);
+		TArray<FElement> ElemsWithBounds;
+		ElemsWithBounds.Reserve(Particles.Num());
 
-		const int32 CurIdx = GetNewWorkSnapshot();
-		FWorkSnapshot& WorkSnapshot = WorkPool[CurIdx];
-		WorkSnapshot.Elems.Reserve(Particles.Num());
-		
-		
 		GlobalPayloads.Reset();
 		Leaves.Reset();
 		Nodes.Reset();
@@ -786,14 +774,13 @@ private:
 		PayloadToInfo.Reset();
 		NumProcessedThisSlice = 0;
 
-		WorkSnapshot.Bounds = TAABB<T, 3>::EmptyAABB();
+		FullBounds = TAABB<T, 3>::EmptyAABB();
 
 		{
 			SCOPE_CYCLE_COUNTER(STAT_AABBTreeTimeSliceSetup);
 
 			int32 Idx = 0;
 
-			//TODO: we need a better way to time-slice this case since there can be a huge number of Particles. Can't do it right now without making full copy
 			for (auto& Particle : Particles)
 			{
 				bool bHasBoundingBox = HasBoundingBox(Particle);
@@ -808,8 +795,8 @@ private:
 					}
 					else
 					{
-						WorkSnapshot.Elems.Add(FElement{ Payload, ElemBounds });
-						WorkSnapshot.Bounds.GrowToInclude(ElemBounds);
+						ElemsWithBounds.Add(FElement{ Payload, ElemBounds });
+						FullBounds.GrowToInclude(ElemBounds);
 					}
 				}
 				else
@@ -829,287 +816,201 @@ private:
 				++Idx;
 				//todo: payload info
 			}
-		}
 
-		NumProcessedThisSlice = Particles.Num();	//todo: give chance to time slice out of next phase
+		}
 
 		{
 			SCOPE_CYCLE_COUNTER(STAT_AABBTreeInitialTimeSlice);
-			WorkSnapshot.NewNodeIdx = 0;
-			WorkSnapshot.NodeLevel = 0;
 
-			//push root onto stack
-			WorkStack.Add(CurIdx);
+			SplitNode(FullBounds, ElemsWithBounds, 0, 0, true);
 
-			SplitNode();
-		}
-
-		/*  Helper validation code 
-		int32 Count = 0;
-		TSet<int32> Seen;
-		if(WorkStack.Num() == 0)
-		{
-			int32 LeafIdx = 0;
-			for (const auto& Leaf : Leaves)
+			if (TimeSliceWorkToComplete.IsEmpty())
 			{
-				Validate(Seen, Count, Leaf);
-				bool bHasParent = false;
-				for (const auto& Node : Nodes)
-				{
-					if (Node.bLeaf && Node.ChildrenNodes[0] == LeafIdx)
-					{
-						bHasParent = true;
-
-						break;
-					}
-				}
-				ensure(bHasParent);
-				++LeafIdx;
+				this->SetAsyncTimeSlicingComplete(true);
 			}
-			ensure(Count == 0 || Seen.Num() == Count);
-			ensure(Count == 0 || Count == Particles.Num());
 		}
-		*/
 
 	}
 
-	enum eTimeSlicePhase
-	{
-		PreFindBestBounds,
-		DuringFindBestBounds,
-		ProcessingChildren
-	};
-
-	struct FSplitInfo
+	void SplitNode(TAABB<T, 3>& SplitBounds, TAABB<T, 3>& RealBounds, TArray<FElement>& Children, T& SplitBoundsSize2, int NewNodeIdx, int NodeLevel, int BoxIdx)
 	{
-		TAABB<T, 3> SplitBounds;	//Even split of parent bounds
-		TAABB<T, 3> RealBounds;	//Actual bounds as children are added
-		int32 WorkSnapshotIdx;	//Idx into work snapshot pool
-		T SplitBoundsSize2;
-	};
+		//SCOPE_CYCLE_COUNTER(STAT_AABBTreeChildrenPhase);
 
-	struct FWorkSnapshot
-	{
-		FWorkSnapshot()
-			: TimeslicePhase(eTimeSlicePhase::PreFindBestBounds)
+		check (Children.Num() > 0)
 		{
+			Nodes[NewNodeIdx].bLeaf = false;
+
+			// increment this early so the amount of recursion work left on the stack is smaller
+			NumProcessedThisSlice += Children.Num();
+
+			const int32 ChildIdx = Nodes.Num();
+			Nodes[NewNodeIdx].ChildrenBounds[BoxIdx] = RealBounds;
+			Nodes[NewNodeIdx].ChildrenNodes[BoxIdx] = ChildIdx;
 
+			SplitNode(RealBounds, Children, NodeLevel + 1, ChildIdx);
 		}
 
-		eTimeSlicePhase TimeslicePhase;
+	}
 
-		TAABB<T, 3> Bounds;
-		TArray<FElement> Elems;
+	void SplitNode(const TAABB<T, 3>& Bounds, const TArray<FElement>& Elems, int32 NodeLevel, int32 NewNodeIdx, bool Initial=false)
+	{
+		// create the actual node space but might no be filled in (YET) due to time slicing exit
+		if (NewNodeIdx >= Nodes.Num())
+		{
+		 	Nodes.AddDefaulted();	//todo: remove TAABB
+		}
 
-		int32 NodeLevel;
-		int32 NewNodeIdx;
+		bool WeAreTimeslicing = (MaxNumToProcess > 0);
 
-		int32 BestBoundsCurIdx;
+		if (WeAreTimeslicing && (NumProcessedThisSlice >= MaxNumToProcess))
+		{
+			// done enough work, capture stack
+			FWorkSnapshot Store;
+			Store.TimeslicePhase = eTimeSlicePhase::PHASE1;
+			Store.Bounds = Bounds;
+			Store.Elems = Elems;			
+			Store.NodeLevel = NodeLevel;
+			Store.NewNodeIdx = NewNodeIdx;
 
-		FSplitInfo SplitInfos[2];
-	};
+			TimeSliceWorkToComplete.Enqueue(Store);
 
-	void FindBestBounds(const int32 StartElemIdx, const int32 LastElem, FWorkSnapshot& CurrentSnapshot)
-	{
-		// add all elements to one of the two split infos at this level - root level [ not taking into account the max number allowed or anything
-		for(int32 ElemIdx = StartElemIdx; ElemIdx < LastElem; ++ElemIdx)
+			return; // done enough
+		}
+
+		auto& PayloadToInfoRef = PayloadToInfo;
+		auto& LeavesRef = Leaves;
+		auto& NodesRef = Nodes;
+		auto MakeLeaf = [NewNodeIdx, &PayloadToInfoRef, &Elems, &LeavesRef, &NodesRef]()
 		{
-			const FElement& Elem = CurrentSnapshot.Elems[ElemIdx];
-			int32 MinBoxIdx = INDEX_NONE;
-			T MinDelta2 = TNumericLimits<T>::Max();
-			int32 BoxIdx = 0;
-			for (const FSplitInfo& SplitInfo : CurrentSnapshot.SplitInfos)
+			if (bMutable)
 			{
-				TAABB<T, 3> NewBox = SplitInfo.SplitBounds;
-				NewBox.GrowToInclude(Elem.Bounds);
-				const T Delta2 = NewBox.Extents().SizeSquared() - SplitInfo.SplitBoundsSize2;
-				if (Delta2 < MinDelta2)
+				for (const FElement& Elem : Elems)
 				{
-					MinDelta2 = Delta2;
-					MinBoxIdx = BoxIdx;
+					PayloadToInfoRef.Add(Elem.Payload, FAABBTreePayloadInfo{ INDEX_NONE, INDEX_NONE, LeavesRef.Num() });
 				}
-				++BoxIdx;
-			}
-
-			if (CHAOS_ENSURE(MinBoxIdx != INDEX_NONE))
-			{
-				FSplitInfo& SplitInfo = CurrentSnapshot.SplitInfos[MinBoxIdx];
-				WorkPool[SplitInfo.WorkSnapshotIdx].Elems.Add(Elem);
-				SplitInfo.RealBounds.GrowToInclude(Elem.Bounds);
 			}
-		}
 
-		NumProcessedThisSlice += LastElem - StartElemIdx;
-	}
+			NodesRef[NewNodeIdx].bLeaf = true;
+			NodesRef[NewNodeIdx].ChildrenNodes[0] = LeavesRef.Add(TLeafType{ Elems }); //todo: avoid copy?
 
-	void SplitNode()
-	{
-		const bool WeAreTimeslicing = (MaxNumToProcess > 0);
+		};
 
-		while (WorkStack.Num())
+		if (Elems.Num() <= MaxChildrenInLeaf || NodeLevel >= MaxTreeDepth)
 		{
-			//NOTE: remember to be careful with this since it's a pointer on a tarray
-			const int32 CurIdx = WorkStack.Last();
-
-			if (WorkPool[CurIdx].TimeslicePhase == eTimeSlicePhase::ProcessingChildren)
-			{
-				//If we got to this it must be that my children are done, so I'm done as well
-				WorkStack.Pop(/*bResize=*/false);
-				FreeWorkSnapshot(CurIdx);
-				continue;
-			}
+			MakeLeaf();
+			return; // keep working
+		}
 
-			const int32 NewNodeIdx = WorkPool[CurIdx].NewNodeIdx;
+		const TVector<T, 3> Extents = Bounds.Extents();
+		const int32 MaxAxis = Bounds.LargestAxis();
+		bool ChildrenInBothHalves = false;
 
-			// create the actual node space but might no be filled in (YET) due to time slicing exit
-			if (NewNodeIdx >= Nodes.Num())
-			{
-				Nodes.AddDefaulted((1 + NewNodeIdx) - Nodes.Num());
-			}
+		FSplitInfo SplitInfos[2];
+		SplitInfos[0].SplitBounds = TAABB<T, 3>(Bounds.Min(), Bounds.Min());
+		SplitInfos[1].SplitBounds = TAABB<T, 3>(Bounds.Max(), Bounds.Max());
 
-			if (WeAreTimeslicing && (NumProcessedThisSlice >= MaxNumToProcess))
-			{
-				return; // done enough
-			}
+		const TVector<T, 3> Center = Bounds.Center();
+		for (FSplitInfo& SplitInfo : SplitInfos)
+		{
+			SplitInfo.RealBounds = TAABB<T, 3>::EmptyAABB();
 
-			auto& PayloadToInfoRef = PayloadToInfo;
-			auto& LeavesRef = Leaves;
-			auto& NodesRef = Nodes;
-			auto& WorkPoolRef = WorkPool;
-			auto MakeLeaf = [NewNodeIdx, &PayloadToInfoRef, &WorkPoolRef, CurIdx, &LeavesRef, &NodesRef]()
+			for (int32 Axis = 0; Axis < 3; ++Axis)
 			{
-				if (bMutable)
+				TVector<T, 3> NewPt0 = Center;
+				TVector<T, 3> NewPt1 = Center;
+				if (Axis != MaxAxis)
 				{
-					//todo: does this need time slicing in the case when we have a ton of elements that can't be split?
-					//hopefully not a real issue
-					for (const FElement& Elem : WorkPoolRef[CurIdx].Elems)
-					{
-						PayloadToInfoRef.Add(Elem.Payload, FAABBTreePayloadInfo{ INDEX_NONE, INDEX_NONE, LeavesRef.Num() });
-					}
+					NewPt0[Axis] = Bounds.Min()[Axis];
+					NewPt1[Axis] = Bounds.Max()[Axis];
+					SplitInfo.SplitBounds.GrowToInclude(NewPt0);
+					SplitInfo.SplitBounds.GrowToInclude(NewPt1);
 				}
-
-				NodesRef[NewNodeIdx].bLeaf = true;
-				NodesRef[NewNodeIdx].ChildrenNodes[0] = LeavesRef.Add(TLeafType{ WorkPoolRef[CurIdx].Elems }); //todo: avoid copy?
-
-			};
-
-			if (WorkPool[CurIdx].Elems.Num() <= MaxChildrenInLeaf || WorkPool[CurIdx].NodeLevel >= MaxTreeDepth)
-			{
-
-				MakeLeaf();
-				WorkStack.Pop(/*bResize=*/false);	//finished with this node
-				FreeWorkSnapshot(CurIdx);
-				continue;
 			}
 
-			if (WorkPool[CurIdx].TimeslicePhase == eTimeSlicePhase::PreFindBestBounds)
-			{
-				const TVector<T, 3> Extents = WorkPool[CurIdx].Bounds.Extents();
-				const int32 MaxAxis = WorkPool[CurIdx].Bounds.LargestAxis();
-
-				//Add two children, remember this invalidates any pointers to current snapshot
-				const int32 FirstChildIdx = GetNewWorkSnapshot();
-				const int32 SecondChildIdx = GetNewWorkSnapshot();
-
-				//mark idx of children into the work pool
-				WorkPool[CurIdx].SplitInfos[0].WorkSnapshotIdx = FirstChildIdx;
-				WorkPool[CurIdx].SplitInfos[1].WorkSnapshotIdx = SecondChildIdx;
-
-				//these are the hypothetical bounds for the perfect 50/50 split
-				WorkPool[CurIdx].SplitInfos[0].SplitBounds = TAABB<T, 3>(WorkPool[CurIdx].Bounds.Min(), WorkPool[CurIdx].Bounds.Min());
-				WorkPool[CurIdx].SplitInfos[1].SplitBounds = TAABB<T, 3>(WorkPool[CurIdx].Bounds.Max(), WorkPool[CurIdx].Bounds.Max());
+			SplitInfo.SplitBoundsSize2 = SplitInfo.SplitBounds.Extents().SizeSquared();
+		}
+		
+		{
+			//SCOPE_CYCLE_COUNTER(STAT_AABBTreeGrowPhase);
 
-				const TVector<T, 3> Center = WorkPool[CurIdx].Bounds.Center();
-				for (FSplitInfo& SplitInfo : WorkPool[CurIdx].SplitInfos)
+			// add all elements to one of the two split infos at this level - root level [ not taking into account the max number allowed or anything
+			for (const FElement& Elem : Elems)
+			{
+				int32 MinBoxIdx = INDEX_NONE;
+				T MinDelta2 = TNumericLimits<T>::Max();
+				int32 BoxIdx = 0;
+				for (const FSplitInfo& SplitInfo : SplitInfos)
 				{
-					SplitInfo.RealBounds = TAABB<T, 3>::EmptyAABB();
-
-					for (int32 Axis = 0; Axis < 3; ++Axis)
+					TAABB<T, 3> NewBox = SplitInfo.SplitBounds;
+					NewBox.GrowToInclude(Elem.Bounds);
+					const T Delta2 = NewBox.Extents().SizeSquared() - SplitInfo.SplitBoundsSize2;
+					if (Delta2 < MinDelta2)
 					{
-						TVector<T, 3> NewPt0 = Center;
-						TVector<T, 3> NewPt1 = Center;
-						if (Axis != MaxAxis)
-						{
-							NewPt0[Axis] = WorkPool[CurIdx].Bounds.Min()[Axis];
-							NewPt1[Axis] = WorkPool[CurIdx].Bounds.Max()[Axis];
-							SplitInfo.SplitBounds.GrowToInclude(NewPt0);
-							SplitInfo.SplitBounds.GrowToInclude(NewPt1);
-						}
+						MinDelta2 = Delta2;
+						MinBoxIdx = BoxIdx;
 					}
-
-					SplitInfo.SplitBoundsSize2 = SplitInfo.SplitBounds.Extents().SizeSquared();
+					++BoxIdx;
 				}
 
-				WorkPool[CurIdx].BestBoundsCurIdx = 0;
-				WorkPool[CurIdx].TimeslicePhase = eTimeSlicePhase::DuringFindBestBounds;
-				const int32 ExpectedNumPerChild = WorkPool[CurIdx].Elems.Num() / 2;
+				if (CHAOS_ENSURE(MinBoxIdx != INDEX_NONE))
 				{
-					WorkPool[FirstChildIdx].Elems.Reserve(ExpectedNumPerChild);
-					WorkPool[SecondChildIdx].Elems.Reserve(ExpectedNumPerChild);
+					SplitInfos[MinBoxIdx].Children.Add(Elem);
+					SplitInfos[MinBoxIdx].RealBounds.GrowToInclude(Elem.Bounds);
 				}
 			}
 
-			if (WorkPool[CurIdx].TimeslicePhase == eTimeSlicePhase::DuringFindBestBounds)
-			{
-				const int32 NumWeCanProcess = MaxNumToProcess - NumProcessedThisSlice;
-				const int32 LastIdxToProcess = WeAreTimeslicing ? FMath::Min(WorkPool[CurIdx].BestBoundsCurIdx + NumWeCanProcess, WorkPool[CurIdx].Elems.Num()) : WorkPool[CurIdx].Elems.Num();
-				FindBestBounds(WorkPool[CurIdx].BestBoundsCurIdx, LastIdxToProcess, WorkPool[CurIdx]);
-				WorkPool[CurIdx].BestBoundsCurIdx = LastIdxToProcess;
+			NumProcessedThisSlice += Elems.Num();
+
+			ChildrenInBothHalves = SplitInfos[0].Children.Num() && SplitInfos[1].Children.Num();
 
-				if (WeAreTimeslicing && (NumProcessedThisSlice >= MaxNumToProcess))
+			if (ChildrenInBothHalves && WeAreTimeslicing && (NumProcessedThisSlice >= MaxNumToProcess))
+			{
+				for (int32 BoxIdx = 0; BoxIdx < 2; ++BoxIdx)
 				{
-					return; // done enough
+					// done enough work, capture stack
+					FWorkSnapshot Store;
+					Store.TimeslicePhase = eTimeSlicePhase::PHASE2;
+					Store.NodeLevel = NodeLevel;
+					Store.NewNodeIdx = NewNodeIdx;
+					Store.BoxIdx = BoxIdx;
+					Store.Bounds = MoveTemp(SplitInfos[BoxIdx].SplitBounds);
+					Store.Elems = MoveTemp(SplitInfos[BoxIdx].Children);
+					Store.RealBounds = MoveTemp(SplitInfos[BoxIdx].RealBounds);
+					Store.SplitBoundsSize2 = SplitInfos[BoxIdx].SplitBoundsSize2;
+					TimeSliceWorkToComplete.Enqueue(Store);
 				}
-			}
-
-
-			const int32 FirstChildIdx = WorkPool[CurIdx].SplitInfos[0].WorkSnapshotIdx;
-			const int32 SecondChildIdx = WorkPool[CurIdx].SplitInfos[1].WorkSnapshotIdx;
+				return; // done enough
+			}	
+		}
 
-			const bool bChildrenInBothHalves = WorkPool[FirstChildIdx].Elems.Num() && WorkPool[SecondChildIdx].Elems.Num();
+		{			
+			//SCOPE_CYCLE_COUNTER(STAT_AABBTreeChildrenPhase);
 
-			// if children in both halves, push them on the stack to continue the split
-			if (bChildrenInBothHalves)
+			// if children in both halves, recurse a level down
+			if (ChildrenInBothHalves)
 			{
 				Nodes[NewNodeIdx].bLeaf = false;
 
-				Nodes[NewNodeIdx].ChildrenBounds[0] = WorkPool[CurIdx].SplitInfos[0].RealBounds;
-				WorkPool[FirstChildIdx].Bounds = Nodes[NewNodeIdx].ChildrenBounds[0];
-				Nodes[NewNodeIdx].ChildrenNodes[0] = Nodes.Num();
-
-				Nodes[NewNodeIdx].ChildrenBounds[1] = WorkPool[CurIdx].SplitInfos[1].RealBounds;
-				WorkPool[SecondChildIdx].Bounds = Nodes[NewNodeIdx].ChildrenBounds[1];
-				Nodes[NewNodeIdx].ChildrenNodes[1] = Nodes.Num() + 1;
-
-				WorkPool[FirstChildIdx].NodeLevel = WorkPool[CurIdx].NodeLevel + 1;
-				WorkPool[SecondChildIdx].NodeLevel = WorkPool[CurIdx].NodeLevel + 1;
-
-				WorkPool[FirstChildIdx].NewNodeIdx = Nodes[NewNodeIdx].ChildrenNodes[0];
-				WorkPool[SecondChildIdx].NewNodeIdx = Nodes[NewNodeIdx].ChildrenNodes[1];
+				// increment this early so the amount of recursion work left on the stack is smaller
+				NumProcessedThisSlice += SplitInfos[0].Children.Num() + SplitInfos[1].Children.Num();
 
-				//push these two new nodes onto the stack
-				WorkStack.Add(SecondChildIdx);
-				WorkStack.Add(FirstChildIdx);
+				for (int32 BoxIdx = 0; BoxIdx < 2; ++BoxIdx)
+				{
+					const int32 ChildIdx = Nodes.Num();
+					Nodes[NewNodeIdx].ChildrenBounds[BoxIdx] = SplitInfos[BoxIdx].RealBounds;
+					Nodes[NewNodeIdx].ChildrenNodes[BoxIdx] = ChildIdx;
 
-				// create the actual node so that no one else can use our children node indices
-				const int32 HighestNodeIdx = Nodes[NewNodeIdx].ChildrenNodes[1];
-				Nodes.AddDefaulted((1 + HighestNodeIdx) - Nodes.Num());
-				
-				WorkPool[CurIdx].TimeslicePhase = eTimeSlicePhase::ProcessingChildren;
+					SplitNode(SplitInfos[BoxIdx].RealBounds, SplitInfos[BoxIdx].Children, NodeLevel + 1, ChildIdx);
+				}
 			}
 			else
 			{
 				//couldn't split so just make a leaf - THIS COULD CONTAIN MORE THAN MaxChildrenInLeaf!!!
 				MakeLeaf();
-				WorkStack.Pop(/*bResize=*/false);	//we are done with this node
-				FreeWorkSnapshot(CurIdx);
 			}
 		}
 
-		check(WorkStack.Num() == 0);
-		//Stack is empty, clean up pool and mark task as complete
-		
-		this->SetAsyncTimeSlicingComplete(true);
+		return; // keep working
 	}
 
 	TArray<TPayloadType> FindAllIntersectionsImp(const TAABB<T, 3>& Intersection) const
@@ -1165,12 +1066,8 @@ private:
 
 	TAABBTree<TPayloadType,TLeafType,T,bMutable>& operator=(const TAABBTree<TPayloadType,TLeafType,T,bMutable>& Rhs)
 	{
-		ensure(Rhs.WorkStack.Num() == 0);
-		ensure(Rhs.WorkPool.Num() == 0);
-		ensure(Rhs.WorkPoolFreeList.Num() == 0);
-		WorkStack.Empty();
-		WorkPool.Empty();
-		WorkPoolFreeList.Empty();
+		ensure(Rhs.TimeSliceWorkToComplete.IsEmpty());
+		TimeSliceWorkToComplete.Empty();
 		if(this != &Rhs)
 		{
 			FullBounds = Rhs.FullBounds;
@@ -1189,6 +1086,14 @@ private:
 		return *this;
 	}
 
+	struct FSplitInfo
+	{
+		TAABB<T, 3> SplitBounds;	//Even split of parent bounds
+		TAABB<T, 3> RealBounds;	//Actual bounds as children are added
+		TArray<FElement> Children;
+		T SplitBoundsSize2;
+	};
+
 	TAABB<T, 3> FullBounds;
 	TArray<FNode> Nodes;
 	TArray<TLeafType> Leaves;
@@ -1200,11 +1105,32 @@ private:
 	int32 MaxTreeDepth;
 	T MaxPayloadBounds;
 	int32 MaxNumToProcess;
+	
+	enum eTimeSlicePhase
+	{
+		PHASE1,
+		PHASE2
+	};
+
+	struct FWorkSnapshot
+	{
+		eTimeSlicePhase TimeslicePhase;
+
+		// used for phase 1 & 2
+		TAABB<T, 3> Bounds;
+		TArray<FElement> Elems;
+
+		int32 NodeLevel;
+		int32 NewNodeIdx;
+	
+		// phase 2 data only
+		TAABB<T, 3> RealBounds;
+		int8 BoxIdx;
+		T SplitBoundsSize2;
+	};
 
 	int32 NumProcessedThisSlice;
-	TArray<int32> WorkStack;
-	TArray<int32> WorkPoolFreeList;
-	TArray<FWorkSnapshot> WorkPool;
+	TQueue<FWorkSnapshot> TimeSliceWorkToComplete;
 };
 
 template<typename TPayloadType, typename TLeafType, class T, bool bMutable>
