Resample Phi/Normal when Phi is tiny for Convex/Convex.
Should give more reliable normal to prevent falling through convex from boat.

diff --git a/Engine/Source/Runtime/Experimental/Chaos/Private/Chaos/CollisionResolution.cpp b/Engine/Source/Runtime/Experimental/Chaos/Private/Chaos/CollisionResolution.cpp
--- a/Engine/Source/Runtime/Experimental/Chaos/Private/Chaos/CollisionResolution.cpp
+++ b/Engine/Source/Runtime/Experimental/Chaos/Private/Chaos/CollisionResolution.cpp
@@ -32,6 +32,12 @@ DECLARE_CYCLE_STAT(TEXT("Collisions::GJK"), STAT_Collisions_GJK, STATGROUP_Chaos
 //#pragma optimize("", off)
 //PRAGMA_DISABLE_OPTIMIZATION_ACTUAL
 
+
+// If GJKPenetration returns a phi of abs value < this number, we use PhiWithNormal to resample phi and normal.
+// We have observed bad normals coming from GJKPenetration when barely in contact.
+#define PHI_RESAMPLE_THRESHOLD 0.001
+
+
 bool Chaos_Collision_UseManifolds = true;
 FAutoConsoleVariableRef CVarChaosCollisionUseManifolds(TEXT("p.Chaos.Collision.UseManifolds"), Chaos_Collision_UseManifolds, TEXT("Enable/Disable use of manifoldes in collision."));
 
@@ -186,13 +192,25 @@ namespace Chaos
 		template<class T, int d>
 		TContactPoint<T> ConvexConvexContactPoint(const FImplicitObject& A, const TRigidTransform<T, d>& ATM, const FImplicitObject& B, const TRigidTransform<T, d>& BTM, const T CullDistance)
 		{
-			return Utilities::CastHelper(A, ATM, [&](const auto& ADowncast, const TRigidTransform<T,d>& AFullTM)
+			TContactPoint<T> ContactPoint = Utilities::CastHelper(A, ATM, [&](const auto& ADowncast, const TRigidTransform<T,d>& AFullTM)
 			{
 				return Utilities::CastHelper(B, BTM, [&](const auto& BDowncast, const TRigidTransform<T,d>& BFullTM)
 				{
 					return GJKContactPoint(ADowncast, AFullTM, BDowncast, BFullTM, TVector<T, d>(1, 0, 0));
 				});
 			});
+
+			if (FMath::Abs(ContactPoint.Phi) < (T)(PHI_RESAMPLE_THRESHOLD))
+			{
+				// If GJKPenetration returns a phi of abs value < this number, we use PhiWithNormal to resample phi and normal.
+				// We have observed bad normals coming from GJKPenetration when barely in contact.
+
+				TVector<T, d> ContactLocalB = BTM.InverseTransformPosition(ContactPoint.Location);
+				ContactPoint.Phi = B.PhiWithNormal(ContactLocalB, ContactPoint.Normal);
+				ContactPoint.Normal = BTM.TransformVectorNoScale(ContactPoint.Normal);
+			}
+
+			return ContactPoint;
 		}
 
 		template <typename T, int d>
