Fixed UniqueIdx release causing index entries being returned to the free list too early leading to bad particles remaining in the async tree

diff --git a/Engine/Source/Runtime/Experimental/Chaos/Private/Chaos/PBDRigidsEvolution.cpp b/Engine/Source/Runtime/Experimental/Chaos/Private/Chaos/PBDRigidsEvolution.cpp
--- a/Engine/Source/Runtime/Experimental/Chaos/Private/Chaos/PBDRigidsEvolution.cpp
+++ b/Engine/Source/Runtime/Experimental/Chaos/Private/Chaos/PBDRigidsEvolution.cpp
@@ -381,8 +381,18 @@ namespace Chaos
 			ApplyParticlePendingData(Itr.Value, *AsyncInternalAcceleration, true); //only the first queue needs to update the cached acceleration
 			if (!bIsSingleThreaded)
 			{
-			ApplyParticlePendingData(Itr.Value, *AsyncExternalAcceleration, false);
-		}
+				ApplyParticlePendingData(Itr.Value, *AsyncExternalAcceleration, false);
+			}
+
+			// Async queue deletes complete, unique index free to be consumed by new particles.
+			if(Itr.Value.bDelete)
+			{
+				//NOTE: This assumes that we are never creating a PT particle that is replicated to GT
+				//At the moment that is true, and it seems like we have enough mechanisms to avoid this direction
+				//If we want to support that, the UniqueIndex must be kept around until GT goes away
+				//This is hard to do, but would probably mean the ownership of the index is in the proxy
+				Particles.GetUniqueIndices().ReleaseIdx(Itr.Value.DeleteAccelerationHandle.UniqueIdx());
+			}
 		}
 		AsyncAccelerationQueue.Empty();
 
diff --git a/Engine/Source/Runtime/Experimental/Chaos/Public/Chaos/ISpatialAcceleration.h b/Engine/Source/Runtime/Experimental/Chaos/Public/Chaos/ISpatialAcceleration.h
--- a/Engine/Source/Runtime/Experimental/Chaos/Public/Chaos/ISpatialAcceleration.h
+++ b/Engine/Source/Runtime/Experimental/Chaos/Public/Chaos/ISpatialAcceleration.h
@@ -440,10 +440,13 @@ public:
 	void Remove(const TKey& Key)
 	{
 		const int32 Idx = GetUniqueIdx(Key).Idx;
-		Entries[Idx] = FEntry();	//Mark as free, also resets default values for next use of value
+		if(Idx >= 0 && Idx < Entries.Num())
+		{
+			Entries[Idx] = FEntry();	//Mark as free, also resets default values for next use of value
 #if CHAOS_SERIALIZE_OUT
-		KeysToSerializeOut[Idx] = TKey();
+			KeysToSerializeOut[Idx] = TKey();
 #endif
+		}
 	}
 
 	void Reset()
diff --git a/Engine/Source/Runtime/Experimental/Chaos/Public/Chaos/PBDRigidsSOAs.h b/Engine/Source/Runtime/Experimental/Chaos/Public/Chaos/PBDRigidsSOAs.h
--- a/Engine/Source/Runtime/Experimental/Chaos/Public/Chaos/PBDRigidsSOAs.h
+++ b/Engine/Source/Runtime/Experimental/Chaos/Public/Chaos/PBDRigidsSOAs.h
@@ -210,11 +210,6 @@ public:
 			GetDynamicParticles().GetSleepDataLock().WriteUnlock();
 		}
 
-		//NOTE: This assumes that we are never creating a PT particle that is replicated to GT
-		//At the moment that is true, and it seems like we have enough mechanisms to avoid this direction
-		//If we want to support that, the UniqueIndex must be kept around until GT goes away
-		//This is hard to do, but would probably mean the ownership of the index is in the proxy
-		UniqueIndices.ReleaseIdx(Particle->UniqueIdx());
 		ParticleHandles.DestroyHandleSwap(Particle);
 		
 		UpdateViews();
