Updating copyright for Engine Runtime.

#rnx

diff --git a/Engine/Source/Runtime/ApplicationCore/Private/Windows/WindowsPlatformSplash.cpp b/Engine/Source/Runtime/ApplicationCore/Private/Windows/WindowsPlatformSplash.cpp
--- a/Engine/Source/Runtime/ApplicationCore/Private/Windows/WindowsPlatformSplash.cpp
+++ b/Engine/Source/Runtime/ApplicationCore/Private/Windows/WindowsPlatformSplash.cpp
@@ -1,4 +1,4 @@
-// Copyright 1998-2019 Epic Games, Inc. All Rights Reserved.
+// Copyright Epic Games, Inc. All Rights Reserved.
 
 #include "Windows/WindowsPlatformSplash.h"
 #include "HAL/PlatformMemory.h"
@@ -658,7 +658,7 @@ void FWindowsPlatformSplash::Show()
 
 				// Display copyright information in editor splash screen
 				{
-					const FString CopyrightInfo = NSLOCTEXT( "UnrealEd", "SplashScreen_CopyrightInfo", "Copyright \x00a9 1998-2019   Epic Games, Inc.   All rights reserved." ).ToString();
+					const FString CopyrightInfo = NSLOCTEXT( "UnrealEd", "SplashScreen_CopyrightInfo", "Copyright \x00a9   Epic Games, Inc.   All rights reserved." ).ToString();
 					StartSetSplashText( SplashTextType::CopyrightInfo, *CopyrightInfo );
 				}
 			}
diff --git a/Engine/Source/Runtime/Engine/Private/Streaming/Texture2DMipDataProvider_IO.cpp b/Engine/Source/Runtime/Engine/Private/Streaming/Texture2DMipDataProvider_IO.cpp
--- a/Engine/Source/Runtime/Engine/Private/Streaming/Texture2DMipDataProvider_IO.cpp
+++ b/Engine/Source/Runtime/Engine/Private/Streaming/Texture2DMipDataProvider_IO.cpp
@@ -1,4 +1,4 @@
-// Copyright 1998-2019 Epic Games, Inc. All Rights Reserved.FTexture2DMipDataProvider_IO
+// Copyright Epic Games, Inc. All Rights Reserved.FTexture2DMipDataProvider_IO
 
 /*=============================================================================
 Texture2DMipDataProvider_IO.cpp : Implementation of FTextureMipDataProvider using cooked file IO.
diff --git a/Engine/Source/Runtime/Experimental/GeometryCollectionCore/Private/GeometryCollection/GeometryCollection.cpp b/Engine/Source/Runtime/Experimental/GeometryCollectionCore/Private/GeometryCollection/GeometryCollection.cpp
--- a/Engine/Source/Runtime/Experimental/GeometryCollectionCore/Private/GeometryCollection/GeometryCollection.cpp
+++ b/Engine/Source/Runtime/Experimental/GeometryCollectionCore/Private/GeometryCollection/GeometryCollection.cpp
@@ -1,4 +1,4 @@
-// Copyright 1998-2019 Epic Games, Inc. All Rights Reserved.
+// Copyright Epic Games, Inc. All Rights Reserved.
 
 /*=============================================================================
 	GeometryCollection.cpp: FGeometryCollection methods.
@@ -1142,7 +1142,7 @@ void FGeometryCollection::WriteDataToHeaderFile(const FString &Name, const FStri
 
 	ofstream DataFile;
 	DataFile.open(string(TCHAR_TO_UTF8(*FullPath)));
-	DataFile << "// Copyright 1998-2019 Epic Games, Inc. All Rights Reserved." << endl << endl;
+	DataFile << "// Copyright Epic Games, Inc. All Rights Reserved." << endl << endl;
 	DataFile << "#pragma once" << endl << endl;
 	DataFile << "class " << TCHAR_TO_UTF8(*Name) << endl;
 	DataFile << "{" << endl;
diff --git a/Engine/Source/Runtime/Windows/XAudio2/Private/XAudio2Effects.cpp b/Engine/Source/Runtime/Windows/XAudio2/Private/XAudio2Effects.cpp
--- a/Engine/Source/Runtime/Windows/XAudio2/Private/XAudio2Effects.cpp
+++ b/Engine/Source/Runtime/Windows/XAudio2/Private/XAudio2Effects.cpp
@@ -1,4 +1,4 @@
-// Copyright 1998-2019 Epic Games, Inc. All Rights Reserved.
+// Copyright Epic Games, Inc. All Rights Reserved.
 
 /*=============================================================================
 	XeAudioDevice.cpp: Unreal XAudio2 Audio interface object.
@@ -295,7 +295,7 @@ XAPO_REGISTRATION_PROPERTIES FXAudio2RadioEffect::Registration =
 {
 	__uuidof( FXAudio2RadioEffect ),
 	TEXT( "FXAudio2RadioEffect" ), 
-	TEXT( "Copyright 1998-2019 Epic Games, Inc. All Rights Reserved." ),
+	TEXT( "Copyright Epic Games, Inc. All Rights Reserved." ),
 	1, 0,
 	XAPO_FLAG_INPLACE_REQUIRED	| XAPO_FLAG_CHANNELS_MUST_MATCH
 								| XAPO_FLAG_FRAMERATE_MUST_MATCH
