Chaos - reduce DirtyProxy map removal overhead

diff --git a/Engine/Source/Runtime/Experimental/ChaosSolvers/Private/PBDRigidsSolver.cpp b/Engine/Source/Runtime/Experimental/ChaosSolvers/Private/PBDRigidsSolver.cpp
--- a/Engine/Source/Runtime/Experimental/ChaosSolvers/Private/PBDRigidsSolver.cpp
+++ b/Engine/Source/Runtime/Experimental/ChaosSolvers/Private/PBDRigidsSolver.cpp
@@ -528,7 +528,6 @@ namespace Chaos
 					Proxy->PushToPhysicsState(ProxyData);
 					delete ProxyData;
 				}
-				Solver->RemoveDirtyProxy(Proxy);
 
 				Proxy->ClearAccumulatedData();
 			}
@@ -582,8 +581,7 @@ namespace Chaos
 	{
 		if (DirtyProxiesSet.Num() == 0)
 			return;
-		TArray<IPhysicsProxyBase*> DirtyProxiesArray = DirtyProxiesSet.Array();
-		for (auto & Proxy : DirtyProxiesArray)
+		for (auto & Proxy : DirtyProxiesSet)
 		{
 			switch (Proxy->GetType())
 			{
@@ -609,6 +607,8 @@ namespace Chaos
 				ensure("Unknown proxy type in physics solver.");
 			}
 		}
+
+		DirtyProxiesSet.Reset();
 	}
 
 	void FPBDRigidsSolver::BufferPhysicsResults()
