Fix CIS

diff --git a/Engine/Source/Runtime/Experimental/Chaos/Private/Chaos/ParticleHandle.cpp b/Engine/Source/Runtime/Experimental/Chaos/Private/Chaos/ParticleHandle.cpp
--- a/Engine/Source/Runtime/Experimental/Chaos/Private/Chaos/ParticleHandle.cpp
+++ b/Engine/Source/Runtime/Experimental/Chaos/Private/Chaos/ParticleHandle.cpp
@@ -94,12 +94,6 @@ namespace Chaos
 		}
 	}
 
-	template<>
-	void Chaos::TGeometryParticle<float, 3>::MarkClean(const EParticleFlags CleanBits)
-	{
-		this->MDirtyFlags.MarkClean(CleanBits);
-	}
-
 	template <>
 	template <>
 	int32 TGeometryParticleHandleImp<float, 3, true>::GetPayload<int32>(int32 Idx)
diff --git a/Engine/Source/Runtime/Experimental/Chaos/Public/Chaos/ParticleHandle.h b/Engine/Source/Runtime/Experimental/Chaos/Public/Chaos/ParticleHandle.h
--- a/Engine/Source/Runtime/Experimental/Chaos/Public/Chaos/ParticleHandle.h
+++ b/Engine/Source/Runtime/Experimental/Chaos/Public/Chaos/ParticleHandle.h
@@ -1516,7 +1516,11 @@ public:
 		this->MDirtyFlags.Clear();
 	}
 
-	void MarkClean(const EParticleFlags DirtyBits);
+	void MarkClean(const EParticleFlags CleanBits)
+	{
+		this->MDirtyFlags.MarkClean(CleanBits);
+	}
+
 
 	TGeometryParticleHandle<T, d>* Handle() const
 	{
