Treat v=0 kinematic as static.

diff --git a/Engine/Source/Runtime/Experimental/Chaos/Public/Chaos/Collision/SpatialAccelerationBroadPhase.h b/Engine/Source/Runtime/Experimental/Chaos/Public/Chaos/Collision/SpatialAccelerationBroadPhase.h
--- a/Engine/Source/Runtime/Experimental/Chaos/Public/Chaos/Collision/SpatialAccelerationBroadPhase.h
+++ b/Engine/Source/Runtime/Experimental/Chaos/Public/Chaos/Collision/SpatialAccelerationBroadPhase.h
@@ -202,7 +202,7 @@ namespace Chaos
 					}
 
 					// Sleeping won't collide against another sleeping and sleeping vs dynamic gets picked up by the other direction.
-					const bool bIsParticle2Kinematic = Particle2.CastToKinematicParticle() && Particle2.ObjectState() == EObjectStateType::Kinematic;
+					const bool bIsParticle2Kinematic = Particle2.CastToKinematicParticle() && (Particle2.ObjectState() == EObjectStateType::Kinematic && Particle2.CastToKinematicParticle()->V().SizeSquared() > 1e-4);
 					if (Particle1.ObjectState() == EObjectStateType::Sleeping && !bIsParticle2Kinematic)
 					{
 						continue;
