Exempt v=0 capsules.

diff --git a/Engine/Source/Runtime/Experimental/Chaos/Public/Chaos/Collision/SpatialAccelerationBroadPhase.h b/Engine/Source/Runtime/Experimental/Chaos/Public/Chaos/Collision/SpatialAccelerationBroadPhase.h
--- a/Engine/Source/Runtime/Experimental/Chaos/Public/Chaos/Collision/SpatialAccelerationBroadPhase.h
+++ b/Engine/Source/Runtime/Experimental/Chaos/Public/Chaos/Collision/SpatialAccelerationBroadPhase.h
@@ -7,6 +7,7 @@
 #include "Chaos/ISpatialAccelerationCollection.h"
 #include "Chaos/ParticleHandle.h"
 #include "Chaos/PBDRigidsSOAs.h"
+#include "Chaos/Capsule.h"
 
 namespace Chaos
 {
@@ -202,7 +203,10 @@ namespace Chaos
 					}
 
 					// Sleeping won't collide against another sleeping and sleeping vs dynamic gets picked up by the other direction.
-					const bool bIsParticle2Kinematic = Particle2.CastToKinematicParticle() && (Particle2.ObjectState() == EObjectStateType::Kinematic && Particle2.CastToKinematicParticle()->V().SizeSquared() > 1e-4);
+					const bool bIsParticle2Kinematic = Particle2.CastToKinematicParticle() &&
+						(Particle2.ObjectState() == EObjectStateType::Kinematic &&
+							(Particle2.CastToKinematicParticle()->V().SizeSquared() > 1e-4 ||
+								Particle2.Geometry()->GetType(true) == TCapsule<float>::StaticType()));
 					if (Particle1.ObjectState() == EObjectStateType::Sleeping && !bIsParticle2Kinematic)
 					{
 						continue;
