Force sleep islands of 1 that are not moving.

diff --git a/Engine/Source/Runtime/Experimental/Chaos/Private/Chaos/PBDConstraintGraph.cpp b/Engine/Source/Runtime/Experimental/Chaos/Private/Chaos/PBDConstraintGraph.cpp
--- a/Engine/Source/Runtime/Experimental/Chaos/Private/Chaos/PBDConstraintGraph.cpp
+++ b/Engine/Source/Runtime/Experimental/Chaos/Private/Chaos/PBDConstraintGraph.cpp
@@ -531,11 +531,28 @@ void FPBDConstraintGraph::ComputeIsland(const int32 InNode, const int32 Island,
 
 bool FPBDConstraintGraph::SleepInactive(const int32 Island, const TArrayCollectionArray<TSerializablePtr<FChaosPhysicsMaterial>>& PerParticleMaterialAttributes)
 {
+	FReal DefaultLinearSleepingThreshold = (FReal)1;
+	FReal DefaultAngularSleepingThreshold = (FReal)0.1f;
 	// @todo(ccaulfield): should be able to eliminate this when island is already sleeping
 
 	const TArray<TGeometryParticleHandle<FReal, 3>*>& IslandParticles = GetIslandParticles(Island);
 	check(IslandParticles.Num());
 
+	if (IslandToParticles[Island].Num() == 1)
+	{
+		if (TPBDRigidParticleHandle<FReal, 3>* PBDRigid = IslandToParticles[Island][0]->CastToRigidParticle())
+		{
+			if (PBDRigid->V().Size() < DefaultLinearSleepingThreshold && PBDRigid->W().Size() < DefaultAngularSleepingThreshold)
+			{
+				PBDRigid->SetSleeping(true);
+				PBDRigid->V() = TVector<FReal, 3>(0);
+				PBDRigid->W() = TVector<FReal, 3>(0);
+				return true;
+			}
+		}
+	}
+
+
 	if (!IslandToData[Island].bIsIslandPersistant)
 	{
 		return false;
@@ -548,15 +565,13 @@ bool FPBDConstraintGraph::SleepInactive(const int32 Island, const TArrayCollecti
 	FReal M = 0;
 	FReal LinearSleepingThreshold = FLT_MAX;
 	FReal AngularSleepingThreshold = FLT_MAX;
-	FReal DefaultLinearSleepingThreshold = (FReal)1;
-	FReal DefaultAngularSleepingThreshold = (FReal)0.1f;
 
 	int32 NumDynamicParticles = 0;
 
-	for (const TGeometryParticleHandle<FReal, 3>* Particle : IslandToParticles[Island])
+	for (TGeometryParticleHandle<FReal, 3>* Particle : IslandToParticles[Island])
 	{
-		const TPBDRigidParticleHandle<FReal, 3>* PBDRigid = Particle->CastToRigidParticle();
-		if(PBDRigid && (PBDRigid->ObjectState() == EObjectStateType::Dynamic))
+		TPBDRigidParticleHandle<FReal, 3>* PBDRigid = Particle->CastToRigidParticle();
+		if (PBDRigid && (PBDRigid->ObjectState() == EObjectStateType::Dynamic))
 		{
 			NumDynamicParticles++;
 
