Don't clobber object state.

diff --git a/Engine/Source/Runtime/Experimental/ChaosSolvers/Private/PhysicsProxy/SingleParticlePhysicsProxy.cpp b/Engine/Source/Runtime/Experimental/ChaosSolvers/Private/PhysicsProxy/SingleParticlePhysicsProxy.cpp
--- a/Engine/Source/Runtime/Experimental/ChaosSolvers/Private/PhysicsProxy/SingleParticlePhysicsProxy.cpp
+++ b/Engine/Source/Runtime/Experimental/ChaosSolvers/Private/PhysicsProxy/SingleParticlePhysicsProxy.cpp
@@ -436,7 +436,10 @@ void FSingleParticlePhysicsProxy<Chaos::TPBDRigidParticle<float, 3>>::PullFromPh
 		Particle->SetV(Buffer->MV, false);
 		Particle->SetW(Buffer->MW, false);
 		Particle->UpdateShapeBounds();
-		Particle->SetObjectState(Buffer->MObjectState, true);
+		if (!Particle->IsDirty(Chaos::EParticleFlags::ObjectState))
+		{
+			Particle->SetObjectState(Buffer->MObjectState, true);
+		}
 	}
 }
 
