Chaos : Sleeping Thresholds Minimal Changes
- Expose CVars for controlling minimum thresholds.

[Mike- This is the change list you are looking for]

diff --git a/Engine/Source/Runtime/Experimental/Chaos/Private/Chaos/PBDConstraintGraph.cpp b/Engine/Source/Runtime/Experimental/Chaos/Private/Chaos/PBDConstraintGraph.cpp
--- a/Engine/Source/Runtime/Experimental/Chaos/Private/Chaos/PBDConstraintGraph.cpp
+++ b/Engine/Source/Runtime/Experimental/Chaos/Private/Chaos/PBDConstraintGraph.cpp
@@ -14,6 +14,12 @@
 
 using namespace Chaos;
 
+float ChaosSolverCollisionDefaultLinearSleepThresholdCVar = 1.f;
+FAutoConsoleVariableRef CVarChaosSolverCollisionDefaultLinearSleepThreshold(TEXT("p.ChaosSolverCollisionDefaultLinearSleepThreshold"), ChaosSolverCollisionDefaultLinearSleepThresholdCVar, TEXT("Default linear threshold for sleeping.[def:2]"));
+
+float ChaosSolverCollisionDefaultAngularSleepThresholdCVar = 0.1f;
+FAutoConsoleVariableRef CVarChaosSolverCollisionDefaultAngularSleepThreshold(TEXT("p.ChaosSolverCollisionDefaultAngularSleepThreshold"), ChaosSolverCollisionDefaultAngularSleepThresholdCVar, TEXT("Default angular threshold for sleeping.[def:2]"));
+
 FPBDConstraintGraph::FPBDConstraintGraph() : VisitToken(0)
 {
 }
@@ -531,8 +537,8 @@ void FPBDConstraintGraph::ComputeIsland(const int32 InNode, const int32 Island,
 
 bool FPBDConstraintGraph::SleepInactive(const int32 Island, const TArrayCollectionArray<TSerializablePtr<FChaosPhysicsMaterial>>& PerParticleMaterialAttributes)
 {
-	FReal DefaultLinearSleepingThreshold = (FReal)1;
-	FReal DefaultAngularSleepingThreshold = (FReal)0.1f;
+	FReal DefaultLinearSleepingThreshold = (FReal)ChaosSolverCollisionDefaultLinearSleepThresholdCVar;
+	FReal DefaultAngularSleepingThreshold = (FReal)ChaosSolverCollisionDefaultAngularSleepThresholdCVar;
 	// @todo(ccaulfield): should be able to eliminate this when island is already sleeping
 
 	const TArray<TGeometryParticleHandle<FReal, 3>*>& IslandParticles = GetIslandParticles(Island);
