Fixed memory leak in Chaos proxies when we leave a world during streaming. Solver command lists contained data allocs for newly added proxies that were lost if we didn't run the command queue again before destroying the solver.
Added p.Chaos.Solver.CleanupCommandsOnDestruction to control the cleanup (0 to disable, >0 to cleanup all commands)

diff --git a/Engine/Source/Runtime/Experimental/ChaosSolvers/Private/PBDRigidsSolver.cpp b/Engine/Source/Runtime/Experimental/ChaosSolvers/Private/PBDRigidsSolver.cpp
--- a/Engine/Source/Runtime/Experimental/ChaosSolvers/Private/PBDRigidsSolver.cpp
+++ b/Engine/Source/Runtime/Experimental/ChaosSolvers/Private/PBDRigidsSolver.cpp
@@ -33,6 +33,8 @@ FAutoConsoleVariableRef CVarChaosSolverDrawBPBounds(TEXT("p.Chaos.Solver.DrawBPB
 float ChaosSolverCollisionDefaultIterationsCVar = 1;
 FAutoConsoleVariableRef CVarChaosSolverCollisionDefaultIterations(TEXT("p.ChaosSolverCollisionDefaultIterations"), ChaosSolverCollisionDefaultIterationsCVar, TEXT("Default collision iterations for the solver.[def:1]"));
 
+int32 ChaosSolverCleanupCommandsOnDestruction = 1;
+FAutoConsoleVariableRef CVarChaosSolverCleanupCommandsOnDestruction(TEXT("p.Chaos.Solver.CleanupCommandsOnDestruction"), ChaosSolverCleanupCommandsOnDestruction, TEXT("Whether or not to run internal command queue cleanup on solver destruction (0 = no cleanup, >0 = cleanup all commands)"));
 
 namespace Chaos
 {
@@ -161,6 +163,18 @@ namespace Chaos
 		Reset();
 	}
 
+	FPBDRigidsSolver::~FPBDRigidsSolver()
+	{
+		if(ChaosSolverCleanupCommandsOnDestruction != 0)
+		{
+			TFunction<void(FPhysicsSolver*)> Command;
+			while(CommandQueue.Dequeue(Command))
+			{
+				Command(this);
+			}
+		}
+	}
+
 	float MaxBoundsForTree = 10000;
 	FAutoConsoleVariableRef CVarMaxBoundsForTree(
 		TEXT("p.MaxBoundsForTree"),
diff --git a/Engine/Source/Runtime/Experimental/ChaosSolvers/Public/PBDRigidsSolver.h b/Engine/Source/Runtime/Experimental/ChaosSolvers/Public/PBDRigidsSolver.h
--- a/Engine/Source/Runtime/Experimental/ChaosSolvers/Public/PBDRigidsSolver.h
+++ b/Engine/Source/Runtime/Experimental/ChaosSolvers/Public/PBDRigidsSolver.h
@@ -61,6 +61,7 @@ namespace Chaos
 	{
 
 		FPBDRigidsSolver(const EMultiBufferMode BufferingModeIn);
+		~FPBDRigidsSolver();
 
 	public:
 		typedef FPhysicsSolverBase Super;
