CIS fix for geom collection unit tests

diff --git a/Engine/Source/Runtime/Experimental/ChaosSolvers/Public/PBDRigidsSolver.h b/Engine/Source/Runtime/Experimental/ChaosSolvers/Public/PBDRigidsSolver.h
--- a/Engine/Source/Runtime/Experimental/ChaosSolvers/Public/PBDRigidsSolver.h
+++ b/Engine/Source/Runtime/Experimental/ChaosSolvers/Public/PBDRigidsSolver.h
@@ -61,9 +61,12 @@ namespace Chaos
 	{
 
 		FPBDRigidsSolver(const EMultiBufferMode BufferingModeIn);
-		~FPBDRigidsSolver();
 
 	public:
+
+		// #BGTODO ensure no external callers directly deleting, make private and push everything through DestroySolver.
+		~FPBDRigidsSolver();
+
 		typedef FPhysicsSolverBase Super;
 
 		friend class FPersistentPhysicsTask;
