ToolsFramework: add missing file
#rnx

diff --git a/Engine/Source/Runtime/Experimental/InteractiveToolsFramework/Private/ComponentSourceInterfaces.cpp b/Engine/Source/Runtime/Experimental/InteractiveToolsFramework/Private/ComponentSourceInterfaces.cpp
--- a/Engine/Source/Runtime/Experimental/InteractiveToolsFramework/Private/ComponentSourceInterfaces.cpp
+++ b/Engine/Source/Runtime/Experimental/InteractiveToolsFramework/Private/ComponentSourceInterfaces.cpp
@@ -73,6 +73,19 @@ FPrimitiveComponentTarget::GetMaterial(int32 MaterialIndex) const
 	return Component->GetMaterial(MaterialIndex);
 }
 
+void
+FPrimitiveComponentTarget::GetMaterialSet(FComponentMaterialSet& MaterialSetOut) const
+{
+	int32 NumMaterials = Component->GetNumMaterials(); 
+	MaterialSetOut.Materials.SetNum(NumMaterials);
+	for (int32 k = 0; k < NumMaterials; ++k)
+	{
+		MaterialSetOut.Materials[k] = Component->GetMaterial(k);
+	}
+}
+
+
+
 FTransform
 FPrimitiveComponentTarget::GetWorldTransform() const
 {
