ToolsFramework: refactor AnyButtonInputBehavior to allow target button to be externally configured. Add ULocalClickDragInputBehavior which uses itself as Target and defines target interface functions with repalceable lambdas. Add InputRouter::RegisterBehavior
#rnx

diff --git a/Engine/Source/Runtime/Experimental/InteractiveToolsFramework/Private/BaseBehaviors/AnyButtonInputBehavior.cpp b/Engine/Source/Runtime/Experimental/InteractiveToolsFramework/Private/BaseBehaviors/AnyButtonInputBehavior.cpp
--- a/Engine/Source/Runtime/Experimental/InteractiveToolsFramework/Private/BaseBehaviors/AnyButtonInputBehavior.cpp
+++ b/Engine/Source/Runtime/Experimental/InteractiveToolsFramework/Private/BaseBehaviors/AnyButtonInputBehavior.cpp
@@ -6,7 +6,7 @@
 
 UAnyButtonInputBehavior::UAnyButtonInputBehavior()
 {
-	ButtonNumber = 0;
+	SetUseLeftMouseButton();
 }
 
 
@@ -21,7 +21,7 @@ bool UAnyButtonInputBehavior::IsPressed(const FInputDeviceState& input)
 	if (input.IsFromDevice(EInputDevices::Mouse)) 
 	{
 		ActiveDevice = EInputDevices::Mouse;
-		return GetMouseButtonState(input).bPressed;
+		return GetButtonStateFunc(input).bPressed;
 	} 
 	else if (input.IsFromDevice(EInputDevices::TabletFingers))
 	{
@@ -37,7 +37,7 @@ bool UAnyButtonInputBehavior::IsDown(const FInputDeviceState& input)
 	if (input.IsFromDevice(EInputDevices::Mouse))
 	{
 		ActiveDevice = EInputDevices::Mouse;
-		return GetMouseButtonState(input).bDown;
+		return GetButtonStateFunc(input).bDown;
 	}
 	return false;
 }
@@ -47,7 +47,7 @@ bool UAnyButtonInputBehavior::IsReleased(const FInputDeviceState& input)
 	if (input.IsFromDevice(EInputDevices::Mouse))
 	{
 		ActiveDevice = EInputDevices::Mouse;
-		return GetMouseButtonState(input).bReleased;
+		return GetButtonStateFunc(input).bReleased;
 	}
 	return false;
 }
@@ -94,19 +94,42 @@ EInputDevices UAnyButtonInputBehavior::GetActiveDevice() const
 
 
 
-
-FDeviceButtonState UAnyButtonInputBehavior::GetMouseButtonState(const FInputDeviceState& input)
+void UAnyButtonInputBehavior::SetUseLeftMouseButton()
 {
-	if (ButtonNumber == 2)
+	GetMouseButtonStateFunc = [](const FInputDeviceState& input)
 	{
-		return input.Mouse.Right;
-	}
-	else if (ButtonNumber == 1)
+		return input.Mouse.Left;
+	};
+}
+
+void UAnyButtonInputBehavior::SetUseMiddleMouseButton()
+{
+	GetMouseButtonStateFunc = [](const FInputDeviceState& input)
 	{
 		return input.Mouse.Middle;
-	}
-	else
+	};
+}
+
+void UAnyButtonInputBehavior::SetUseRightMouseButton()
+{
+	GetMouseButtonStateFunc = [](const FInputDeviceState& input)
 	{
-		return input.Mouse.Left;
+		return input.Mouse.Right;
+	};
+}
+
+void UAnyButtonInputBehavior::SetUseCustomMouseButton(TUniqueFunction<FDeviceButtonState(const FInputDeviceState& Input)> ButtonFunc)
+{
+	GetMouseButtonStateFunc = MoveTemp(ButtonFunc);
+}
+
+
+
+FDeviceButtonState UAnyButtonInputBehavior::GetButtonStateFunc(const FInputDeviceState& Input)
+{
+	if (ActiveDevice == EInputDevices::Mouse)
+	{
+		return GetMouseButtonStateFunc(Input);
 	}
+	return FDeviceButtonState();
 }
\ No newline at end of file
diff --git a/Engine/Source/Runtime/Experimental/InteractiveToolsFramework/Private/BaseBehaviors/ClickDragBehavior.cpp b/Engine/Source/Runtime/Experimental/InteractiveToolsFramework/Private/BaseBehaviors/ClickDragBehavior.cpp
--- a/Engine/Source/Runtime/Experimental/InteractiveToolsFramework/Private/BaseBehaviors/ClickDragBehavior.cpp
+++ b/Engine/Source/Runtime/Experimental/InteractiveToolsFramework/Private/BaseBehaviors/ClickDragBehavior.cpp
@@ -33,7 +33,7 @@ FInputCaptureRequest UClickDragInputBehavior::WantsCapture(const FInputDeviceSta
 FInputCaptureUpdate UClickDragInputBehavior::BeginCapture(const FInputDeviceState& Input, EInputCaptureSide Side)
 {
 	Modifiers.UpdateModifiers(Input, Target);
-	OnClickPress(Input, Side);
+	OnClickPressInternal(Input, Side);
 	return FInputCaptureUpdate::Begin(this, EInputCaptureSide::Any);
 }
 
@@ -44,12 +44,12 @@ FInputCaptureUpdate UClickDragInputBehavior::UpdateCapture(const FInputDeviceSta
 
 	if (IsReleased(Input)) 
 	{
-		OnClickRelease(Input, Data);
+		OnClickReleaseInternal(Input, Data);
 		return FInputCaptureUpdate::End();
 	}
 	else
 	{
-		OnClickDrag(Input, Data);
+		OnClickDragInternal(Input, Data);
 		return FInputCaptureUpdate::Continue();
 	}
 }
@@ -62,17 +62,17 @@ void UClickDragInputBehavior::ForceEndCapture(const FInputCaptureData& Data)
 
 
 
-void UClickDragInputBehavior::OnClickPress(const FInputDeviceState& Input, EInputCaptureSide Side)
+void UClickDragInputBehavior::OnClickPressInternal(const FInputDeviceState& Input, EInputCaptureSide Side)
 {
 	Target->OnClickPress(GetDeviceRay(Input));
 }
 
-void UClickDragInputBehavior::OnClickDrag(const FInputDeviceState& Input, const FInputCaptureData& Data)
+void UClickDragInputBehavior::OnClickDragInternal(const FInputDeviceState& Input, const FInputCaptureData& Data)
 {
 	Target->OnClickDrag(GetDeviceRay(Input));
 }
 
-void UClickDragInputBehavior::OnClickRelease(const FInputDeviceState& Input, const FInputCaptureData& Data)
+void UClickDragInputBehavior::OnClickReleaseInternal(const FInputDeviceState& Input, const FInputCaptureData& Data)
 {
 	Target->OnClickRelease(GetDeviceRay(Input));
 }
diff --git a/Engine/Source/Runtime/Experimental/InteractiveToolsFramework/Private/InputRouter.cpp b/Engine/Source/Runtime/Experimental/InteractiveToolsFramework/Private/InputRouter.cpp
--- a/Engine/Source/Runtime/Experimental/InteractiveToolsFramework/Private/InputRouter.cpp
+++ b/Engine/Source/Runtime/Experimental/InteractiveToolsFramework/Private/InputRouter.cpp
@@ -48,6 +48,11 @@ void UInputRouter::DeregisterSource(IInputBehaviorSource* Source)
 }
 
 
+void UInputRouter::RegisterBehavior(UInputBehavior* Behavior, void* Source, const FString& GroupName)
+{
+	ActiveInputBehaviors->Add(Behavior, Source, GroupName);
+}
+
 
 void UInputRouter::PostInputEvent(const FInputDeviceState& Input)
 {
diff --git a/Engine/Source/Runtime/Experimental/InteractiveToolsFramework/Public/BaseBehaviors/AnyButtonInputBehavior.h b/Engine/Source/Runtime/Experimental/InteractiveToolsFramework/Public/BaseBehaviors/AnyButtonInputBehavior.h
--- a/Engine/Source/Runtime/Experimental/InteractiveToolsFramework/Public/BaseBehaviors/AnyButtonInputBehavior.h
+++ b/Engine/Source/Runtime/Experimental/InteractiveToolsFramework/Public/BaseBehaviors/AnyButtonInputBehavior.h
@@ -51,20 +51,25 @@ public:
 	/** @return last-active supported Device */
 	EInputDevices GetActiveDevice() const;
 
-protected:
 
-	/** 
-	 * Button number on target device. Button 0 is "default" button on all devices.
-	 * Mouse: Left=0, Middle=1, Right=2
-	 */
-	UPROPERTY()
-	int ButtonNumber;
+public:
+	/** Configure the target Mouse button to be the left button */
+	virtual void SetUseLeftMouseButton();
+	/** Configure the target Mouse button to be the middle button */
+	virtual void SetUseMiddleMouseButton();
+	/** Configure the target Mouse button to be the right button */
+	virtual void SetUseRightMouseButton();
+	/** Configure a custom target mouse button */
+	virtual void SetUseCustomMouseButton(TUniqueFunction<FDeviceButtonState(const FInputDeviceState& Input)>);
 
 
 protected:
 	/** Which device is currently active */
 	EInputDevices ActiveDevice;
 
-	/** @return mouse button state for Target Button */
-	virtual FDeviceButtonState GetMouseButtonState(const FInputDeviceState& input);
+	/** Returns FDeviceButtonState for target mouse button */
+	TUniqueFunction<FDeviceButtonState(const FInputDeviceState& Input)> GetMouseButtonStateFunc;
+
+	/** Returns FDeviceButtonState for target active device button */
+	FDeviceButtonState GetButtonStateFunc(const FInputDeviceState& Input);
 };
\ No newline at end of file
diff --git a/Engine/Source/Runtime/Experimental/InteractiveToolsFramework/Public/BaseBehaviors/ClickDragBehavior.h b/Engine/Source/Runtime/Experimental/InteractiveToolsFramework/Public/BaseBehaviors/ClickDragBehavior.h
--- a/Engine/Source/Runtime/Experimental/InteractiveToolsFramework/Public/BaseBehaviors/ClickDragBehavior.h
+++ b/Engine/Source/Runtime/Experimental/InteractiveToolsFramework/Public/BaseBehaviors/ClickDragBehavior.h
@@ -61,16 +61,84 @@ protected:
 	/**
 	 * Internal function that forwards click evens to Target::OnClickPress, you can customize behavior here
 	 */
-	virtual void OnClickPress(const FInputDeviceState& Input, EInputCaptureSide Side);
+	virtual void OnClickPressInternal(const FInputDeviceState& Input, EInputCaptureSide Side);
 
 	/**
 	 * Internal function that forwards click evens to Target::OnClickDrag, you can customize behavior here
 	 */
-	virtual void OnClickDrag(const FInputDeviceState& Input, const FInputCaptureData& Data);
+	virtual void OnClickDragInternal(const FInputDeviceState& Input, const FInputCaptureData& Data);
 
 	/**
 	 * Internal function that forwards click evens to Target::OnClickRelease, you can customize behavior here
 	 */
-	virtual void OnClickRelease(const FInputDeviceState& Input, const FInputCaptureData& Data);
+	virtual void OnClickReleaseInternal(const FInputDeviceState& Input, const FInputCaptureData& Data);
 };
 
+
+
+
+
+/**
+ * ULocalClickDragInputBehavior is an implementation of UClickDragInputBehavior that also implements IClickDragBehaviorTarget directly,
+ * via a set of local lambda functions. To use/customize this class the client replaces the lambda functions with their own.
+ * This avoids having to create a second IClickDragBehaviorTarget implementation for trivial use-cases.
+ */
+UCLASS()
+class INTERACTIVETOOLSFRAMEWORK_API ULocalClickDragInputBehavior : public UClickDragInputBehavior, public IClickDragBehaviorTarget
+{
+	GENERATED_BODY()
+protected:
+	using UClickDragInputBehavior::Initialize;
+
+public:
+	/** Call this to initialize the class */
+	virtual void Initialize()
+	{
+		this->Initialize(this);
+	}
+
+	/** lambda implementation of CanBeginClickDragSequence */
+	TUniqueFunction<FInputRayHit(const FInputDeviceRay& PressPos)> CanBeginClickDragFunc = [](const FInputDeviceRay&) { return FInputRayHit(); };
+
+	/** lambda implementation of OnClickPress */
+	TUniqueFunction<void(const FInputDeviceRay& PressPos)> OnClickPressFunc = [](const FInputDeviceRay&) {};
+
+	/** lambda implementation of OnClickDrag */
+	TUniqueFunction<void(const FInputDeviceRay& PressPos)> OnClickDragFunc = [](const FInputDeviceRay&) {};
+
+	/** lambda implementation of OnClickRelease */
+	TUniqueFunction<void(const FInputDeviceRay& ReleasePos)> OnClickReleaseFunc = [](const FInputDeviceRay&) {};
+
+	/** lambda implementation of OnTerminateDragSequence */
+	TUniqueFunction<void()> OnTerminateFunc = []() {};
+
+
+public:
+	//
+	// IClickDragBehaviorTarget implementation
+	//
+	virtual FInputRayHit CanBeginClickDragSequence(const FInputDeviceRay& PressPos) override
+	{
+		return CanBeginClickDragFunc(PressPos);
+	}
+
+	virtual void OnClickPress(const FInputDeviceRay& PressPos) override
+	{
+		OnClickPressFunc(PressPos);
+	}
+
+	virtual void OnClickDrag(const FInputDeviceRay& DragPos) override
+	{
+		OnClickDragFunc(DragPos);
+	}
+
+	virtual void OnClickRelease(const FInputDeviceRay& ReleasePos) override
+	{
+		OnClickReleaseFunc(ReleasePos);
+	}
+
+	virtual void OnTerminateDragSequence() override
+	{
+		OnTerminateFunc();
+	}
+};
diff --git a/Engine/Source/Runtime/Experimental/InteractiveToolsFramework/Public/InputRouter.h b/Engine/Source/Runtime/Experimental/InteractiveToolsFramework/Public/InputRouter.h
--- a/Engine/Source/Runtime/Experimental/InteractiveToolsFramework/Public/InputRouter.h
+++ b/Engine/Source/Runtime/Experimental/InteractiveToolsFramework/Public/InputRouter.h
@@ -47,6 +47,10 @@ public:
 	/** Remove Behaviors from this Source from the active set */
 	virtual void DeregisterSource(IInputBehaviorSource* Source);
 
+	/** Add a new InputBehavior to the active behavior set, with optional Source and GroupName */
+	virtual void RegisterBehavior(UInputBehavior* Behavior, void* Source = nullptr, const FString& GroupName = "");
+
+
 	/** Insert a new input event which is used to check for new captures, or forwarded to active capture */
 	virtual void PostInputEvent(const FInputDeviceState& Input);
 
