AUTO MERGED/UNRESOLVED

Merging //UE4/Main @ 9388175 to //UE4/Dev-Core (stefan.boberg_rdu-p4merge-02_6162)

(ushell-p4-mergedown)

diff --git a/Engine/Config/BaseEngine.ini b/Engine/Config/BaseEngine.ini
--- a/Engine/Config/BaseEngine.ini
+++ b/Engine/Config/BaseEngine.ini
@@ -2,7 +2,7 @@
 
 [/Script/EngineSettings.GameMapsSettings]
 GameInstanceClass=/Script/Engine.GameInstance
-EditorStartupMap=/Engine/Maps/Templates/Template_Default 
+EditorStartupMap=/Engine/Maps/Templates/Template_Default
 GameDefaultMap=/Engine/Maps/Entry
 ServerDefaultMap=/Engine/Maps/Entry
 GlobalDefaultGameMode="/Script/Engine.GameModeBase"
@@ -73,15 +73,12 @@ bAllowMatureLanguage=false
 GameEngine=/Script/Engine.GameEngine
 EditorEngine=/Script/UnrealEd.EditorEngine
 UnrealEdEngine=/Script/UnrealEd.UnrealEdEngine
-; Leaving these unset causes the engine to use FCoreStyle::GetDefaultFont() instead
-;TinyFontName=/Engine/EngineFonts/RobotoTiny.RobotoTiny
-;SmallFontName=/Engine/EngineFonts/Roboto.Roboto
-;MediumFontName=/Engine/EngineFonts/Roboto.Roboto
-;LargeFontName=/Engine/EngineFonts/Roboto.Roboto
-;SubtitleFontName=/Engine/EngineFonts/Roboto.Roboto
+TinyFontName=/Engine/EngineFonts/RobotoTiny.RobotoTiny
+SmallFontName=/Engine/EngineFonts/Roboto.Roboto
+MediumFontName=/Engine/EngineFonts/Roboto.Roboto
+LargeFontName=/Engine/EngineFonts/Roboto.Roboto
+SubtitleFontName=/Engine/EngineFonts/Roboto.Roboto
 WireframeMaterialName=/Engine/EngineDebugMaterials/WireframeMaterial.WireframeMaterial
-HairDefaultMaterialName=/Engine/EngineMaterials/HairDefaultMaterial.HairDefaultMaterial
-HairDebugMaterialName=/Engine/EngineMaterials/HairDebugMaterial.HairDebugMaterial
 DefaultMaterialName=/Engine/EngineMaterials/WorldGridMaterial.WorldGridMaterial
 DefaultLightFunctionMaterialName=/Engine/EngineMaterials/DefaultLightFunctionMaterial.DefaultLightFunctionMaterial
 DefaultTextureName=/Engine/EngineResources/DefaultTexture.DefaultTexture
@@ -115,7 +112,7 @@ ClothPaintMaterialName=/Engine/EngineDebugMaterials/ClothMaterial.ClothMaterial
 ClothPaintMaterialWireframeName=/Engine/EngineDebugMaterials/ClothMaterial_WF.ClothMaterial_WF
 DebugEditorMaterialName=/Engine/EngineDebugMaterials/DebugEditorMaterial.DebugEditorMaterial
 InvalidLightmapSettingsMaterialName=/Engine/EngineMaterials/M_InvalidLightmapSettings.M_InvalidLightmapSettings
-PreviewShadowsIndicatorMaterialName=/Engine/EditorMaterials/PreviewShadowIndicatorMaterial.PreviewShadowIndicatorMaterial 
+PreviewShadowsIndicatorMaterialName=/Engine/EditorMaterials/PreviewShadowIndicatorMaterial.PreviewShadowIndicatorMaterial
 EditorBrushMaterialName=/Engine/EngineMaterials/EditorBrushMaterial.EditorBrushMaterial
 DefaultPhysMaterialName=/Engine/EngineMaterials/DefaultPhysicalMaterial.DefaultPhysicalMaterial
 DefaultDeferredDecalMaterialName=/Engine/EngineMaterials/DefaultDeferredDecalMaterial.DefaultDeferredDecalMaterial
@@ -268,7 +265,7 @@ bUseStreamingPause=false
 ; All redirects prior to UE 4.10 are now in FCoreRedirects::RegisterNativeRedirects, new ones should be added to the bottom of this section
 ; Redirects can now go in any ini file, including Plugin inis
 ; This new redirect format has replaced the old ActiveClassRedirects/K2FieldRedirects format, and the syntax is the same for every type
-; 
+;
 ; +ClassRedirects=(OldName="OldClass", NewName="/Script/Module.NewClass")
 ; +ClassRedirects=(OldName="/Script/Module.OldClass", NewName="/Script/Module.NewClass")
 ;
@@ -692,24 +689,19 @@ bUseStreamingPause=false
 
 ; 4.24
 +PropertyRedirects=(OldName="NavDataConfig.NavigationDataClassName", NewName="NavDataConfig.NavDataClass")
-+ClassRedirects=(OldName="/Script/EditorInteractiveToolsFramework.StaticMeshComponentTargetFactory",NewName="/Script/UnrealEd.StaticMeshComponentTargetFactory")
-+ClassRedirects=(OldName="/Script/EditorInteractiveToolsFramework.StandardToolModeCommands",NewName="/Script/UnrealEd.StandardToolModeCommands")
-+EnumRedirects==(OldName="/Script/EditorInteractiveToolsFramework.StandardToolModeCommands",NewName="/Script/UnrealEd.StandardToolModeCommands")
-+ClassRedirects=(OldName="/Script/EditorInteractiveToolsFramework.InteractiveToolCommands",NewName="/Script/UnrealEd.InteractiveToolCommands")
-+ClassRedirects=(OldName="/Script/EditorInteractiveToolsFramework.EditorToolAssetAPI",NewName="/Script/UnrealEd.EditorToolAssetAPI")
-+ClassRedirects=(OldName="/Script/EditorInteractiveToolsFramework.StaticMeshComponentTarget",NewName="/Script/UnrealEd.StaticMeshComponentTarget")
-
-; Rename ClothingSystemRuntime to ClothingSystemRuntimeNv
++FunctionRedirects=(OldName="Actor.GetComponentsByClass", NewName="Actor.K2_GetComponentsByClass")
+
+; Rename ClothingSystemRuntime to ClothingSystemRuntimeNv, base functionality moved to ClothingSystemRuntimeInterface.
 +PackageRedirects=(OldName="/Script/ClothingSystemRuntime",NewName="/Script/ClothingSystemRuntimeNv")
++StructRedirects=(OldName="/Script/ClothingSystemRuntimeNv.ClothConfig",NewName="/Script/ClothingSystemRuntimeNv.ClothConfig_Legacy")
 +ClassRedirects=(OldName="/Script/ClothingSystemRuntime.ClothingSimulationFactoryNv",NewName="/Script/ClothingSystemRuntimeNv.ClothingSimulationFactoryNv")
 +ClassRedirects=(OldName="/Script/ClothingSystemRuntime.ClothingSimulationInteractorNv",NewName="/Script/ClothingSystemRuntimeNv.ClothingSimulationInteractorNv")
++ClassRedirects=(OldName="/Script/ClothingSystemRuntime.ClothingAsset",NewName="/Script/ClothingSystemRuntimeNv.ClothingAssetNv")
 +FunctionRedirects=(OldName="ClothingSystemRuntime.ClothingSimulationInteractorNv.SetAnimDriveSpringStiffness",NewName="ClothingSystemRuntimeNv.ClothingSimulationInteractorNv.SetAnimDriveSpringStiffness")
 +FunctionRedirects=(OldName="ClothingSystemRuntime.ClothingSimulationInteractorNv.SetAnimDriveDamperStiffness",NewName="ClothingSystemRuntimeNv.ClothingSimulationInteractorNv.SetAnimDriveDamperStiffness")
 +FunctionRedirects=(OldName="ClothingSystemRuntime.ClothingSimulationInteractorNv.EnableGravityOverride",NewName="ClothingSystemRuntimeNv.ClothingSimulationInteractorNv.EnableGravityOverride")
 +FunctionRedirects=(OldName="ClothingSystemRuntime.ClothingSimulationInteractorNv.DisableGravityOverride",NewName="ClothingSystemRuntimeNv.ClothingSimulationInteractorNv.DisableGravityOverride")
 
-+FunctionRedirects=(OldName="Actor.GetComponentsByClass", NewName="Actor.K2_GetComponentsByClass")
-
 +FunctionRedirects=(OldName="SkeletalMeshComponent.GetSubInstanceByTag",NewName="SkeletalMeshComponent.GetLinkedAnimGraphInstanceByTag")
 +FunctionRedirects=(OldName="SkeletalMeshComponent.GetSubInstancesByTag",NewName="SkeletalMeshComponent.GetLinkedAnimGraphInstancesByTag")
 +PropertyRedirects=(OldName="SkeletalMeshComponent.GetLinkedAnimGraphInstancesByTag.OutSubInstances",NewName="OutLinkedInstances")
@@ -819,8 +811,6 @@ bFullyCompressStartupPackages=false
 +Package=/Engine/EngineDebugMaterials/VertexColorViewMode_GreenOnly
 +Package=/Engine/EngineDebugMaterials/VertexColorViewMode_RedOnly
 +Package=/Engine/EngineDebugMaterials/WireframeMaterial
-+Package=/Engine/EngineDebugMaterials/HairDefaultMaterial
-+Package=/Engine/EngineDebugMaterials/HairDebugMaterial
 
 +Package=/Engine/EngineSounds/WhiteNoise
 
@@ -1044,7 +1034,6 @@ InEditorGameURLOptions=
 [/Script/UnrealEd.UnrealEdEngine]
 AutoSaveIndex=0
 +TemplateMapInfos=(ThumbnailTexture=Texture2D'/Engine/Maps/Templates/Thumbnails/Default.Default',Map="/Engine/Maps/Templates/Template_Default")
-+TemplateMapInfos=(ThumbnailTexture=Texture2D'/Engine/Maps/Templates/Thumbnails/TimeOfDay.TimeOfDay',Map="/Engine/Maps/Templates/TimeOfDay_Default")
 +TemplateMapInfos=(ThumbnailTexture=Texture2D'/Engine/Maps/Templates/Thumbnails/VR-Basic.VR-Basic',Map="/Engine/Maps/Templates/VR-Basic")
 
 +PackagesToBeFullyLoadedAtStartup=/Engine/EditorMaterials/Cross
@@ -1352,7 +1341,7 @@ r.RHICmdBypass=0
 ; Use medium detail mode in splitscreen, this allows LD's to mark meshes as High detail and they won't render in SS
 ; Removed for now because this value is getting applied all the time, even in non-splitscreen. Splitscreen generally needs
 ; work anyway, so this won't cause any issues.
-;r.DetailMode=1 
+;r.DetailMode=1
 
 [OnlineSubsystem]
 bHasVoiceEnabled=true
@@ -1729,11 +1718,6 @@ bFilterAssetDataWithNoTags=false
 bFilterDependenciesWithNoTags=false
 bFilterSearchableNames=true
 
-; Control which collections are added as asset registry tags during cook (none by default)
-; If bTagAllCollections is true then you can exclude specfic collections via the CollectionsToExcludeAsTags arrray
-; If bTagAllCollections is false then you can include specfic collections via the CollectionsToIncludeAsTags arrray
-bTagAllCollections=false
-
 [AutomationTesting]
 ImportTestPath=../../Content/EditorAutomation/
 ImportTestPackagePath=/Engine/Content/EditorAutomation
@@ -1944,6 +1928,7 @@ bSupportsLandscapeLeftOrientation=True
 bSupportsLandscapeRightOrientation=True
 PreferredLandscapeOrientation=LandscapeLeft
 bSupportsMetal=True
+bSupportsOpenGLES2=False
 bCookPVRTCTextures=True
 bCookASTCTextures=False
 bSupportsMetalMRT=False
@@ -1959,8 +1944,6 @@ bAllowRemoteRotation=True
 bUseRemoteAsVirtualJoystick=True
 bUseAbsoluteDpadValues=False
 bAllowControllers=True
-bBuildAsFramework=False
-bGenerateFrameworkWrapperProject=True
 bGeneratedSYMFile=False
 bDisableHTTPS=false
 bUseRSync=True
@@ -1969,7 +1952,7 @@ BundleName=[PROJECT_NAME]
 BundleIdentifier=com.YourCompany.[PROJECT_NAME]
 VersionInfo=1.0
 FrameRateLock=PUFRL_30
-MinimumiOSVersion=IOS_11
+MinimumiOSVersion=IOS_10
 bSupportsIPad=True
 bSupportsIPhone=True
 AdditionalPlistData=
@@ -2138,13 +2121,16 @@ RemoteHoloLensIP=
 MaxBitrate=4000
 
 [/Script/LinuxTargetPlatform.LinuxTargetSettings]
-; When neither -vulkan nor -opengl4 are passed on the commandline, the engine will default to the first targeted RHI.
+; Possible values for architecture: X86_64UnknownLinuxGnu, ArmUnknownLinuxGnueabihf, AArch64UnknownLinuxGnueabi, I686UnknownLinuxGnu
+TargetArchitecture=X86_64UnknownLinuxGnu
+
+; the engine will default to the first targeted RHI.
 +TargetedRHIs=SF_VULKAN_SM5
 
 ; OpenGL4 is deprecated, you can comment this back in to add it to your targeted RHI list
 ; +TargetedRHIs=GLSL_430
 
- 
+
 [/Script/MacTargetPlatform.MacTargetSettings]
 MaxShaderLanguageVersion=3
 +TargetedRHIs=SF_METAL_SM5
@@ -2324,8 +2310,7 @@ ServerUpdateCameraTimeout=2.0
 +FilmbackPresets=(Name="APS-C (Canon)",FilmbackSettings=(SensorWidth=22.2,SensorHeight=14.8))
 +FilmbackPresets=(Name="Full Frame DSLR",FilmbackSettings=(SensorWidth=36,SensorHeight=24))
 +FilmbackPresets=(Name="Micro Four Thirds",FilmbackSettings=(SensorWidth=17.3,SensorHeight=13))
-DefaultFilmbackPresetName="16:9 DSLR" ; Deprecated, here for backwards compatibility
-DefaultFilmbackPreset="16:9 Digital Film"
+DefaultFilmbackPresetName="16:9 DSLR"
 +LensPresets=(Name="12mm Prime f/2.8",LensSettings=(MinFocalLength=12,MaxFocalLength=12,MinFStop=2.8,MaxFStop=22,DiaphragmBladeCount=7))
 +LensPresets=(Name="30mm Prime f/1.4",LensSettings=(MinFocalLength=30,MaxFocalLength=30,MinFStop=1.4,MaxFStop=22,DiaphragmBladeCount=7))
 +LensPresets=(Name="50mm Prime f/1.8",LensSettings=(MinFocalLength=50,MaxFocalLength=50,MinFStop=1.8,MaxFStop=22,DiaphragmBladeCount=7))
diff --git a/Engine/Source/Runtime/Core/Public/HAL/LowLevelMemTracker.h b/Engine/Source/Runtime/Core/Public/HAL/LowLevelMemTracker.h
--- a/Engine/Source/Runtime/Core/Public/HAL/LowLevelMemTracker.h
+++ b/Engine/Source/Runtime/Core/Public/HAL/LowLevelMemTracker.h
@@ -183,7 +183,9 @@ enum class ELLMTagSet : uint8
 	macro(Niagara,								"Niagara",						GET_STATFNAME(STAT_NiagaraLLM),								GET_STATFNAME(STAT_NiagaraSummaryLLM),			-1)\
 	macro(GC,									"GC",							GET_STATFNAME(STAT_GCLLM),									GET_STATFNAME(STAT_EngineSummaryLLM),			-1)\
 	macro(UI,									"UI",							GET_STATFNAME(STAT_UILLM),									GET_STATFNAME(STAT_UISummaryLLM),				-1)\
-	macro(PhysX,								"PhysX",						GET_STATFNAME(STAT_PhysXLLM),								GET_STATFNAME(STAT_PhysXSummaryLLM),			-1)\
+	macro(Physics,								"Physics",						GET_STATFNAME(STAT_PhysicsLLM),								GET_STATFNAME(STAT_PhysicsSummaryLLM),			-1)\
+	macro(PhysX,								"PhysX",						GET_STATFNAME(STAT_PhysXLLM),								GET_STATFNAME(STAT_PhysXSummaryLLM),			ELLMTag::Physics)\
+	macro(Chaos,								"Chaos",						GET_STATFNAME(STAT_ChaosLLM),								GET_STATFNAME(STAT_ChaosSummaryLLM),			ELLMTag::Physics)\
 	macro(EnginePreInitMemory,					"EnginePreInit",				GET_STATFNAME(STAT_EnginePreInitLLM),						GET_STATFNAME(STAT_EngineSummaryLLM),			-1)\
 	macro(EngineInitMemory,						"EngineInit",					GET_STATFNAME(STAT_EngineInitLLM),							GET_STATFNAME(STAT_EngineSummaryLLM),			-1)\
 	macro(RenderingThreadMemory,				"RenderingThread",				GET_STATFNAME(STAT_RenderingThreadLLM),						GET_STATFNAME(STAT_EngineSummaryLLM),			-1)\
diff --git a/Engine/Source/Runtime/Engine/Engine.Build.cs b/Engine/Source/Runtime/Engine/Engine.Build.cs
--- a/Engine/Source/Runtime/Engine/Engine.Build.cs
+++ b/Engine/Source/Runtime/Engine/Engine.Build.cs
@@ -376,15 +376,12 @@ public class Engine : ModuleRules
 			DynamicallyLoadedModuleNames.Add("PhysXCooking");
 		}
 
-		if (Target.bCompileChaos || Target.bUseChaos)
-        {
-            PublicDependencyModuleNames.AddRange(
-				new string[] {
-					"PhysicsSQ",
-					"ChaosSolvers"
-				}
-			);
-        }
+        PublicDependencyModuleNames.AddRange(
+			new string[] {
+				"PhysicsSQ",
+				"ChaosSolvers"
+			}
+		);
 
         // Engine public headers need to know about some types (enums etc.)
         PublicIncludePathModuleNames.Add("ClothingSystemRuntimeInterface");
diff --git a/Engine/Source/Runtime/Engine/Private/PhysicsEngine/BodySetup.cpp b/Engine/Source/Runtime/Engine/Private/PhysicsEngine/BodySetup.cpp
--- a/Engine/Source/Runtime/Engine/Private/PhysicsEngine/BodySetup.cpp
+++ b/Engine/Source/Runtime/Engine/Private/PhysicsEngine/BodySetup.cpp
@@ -42,9 +42,7 @@
 #include "ProfilingDebugging/CookStats.h"
 #include "UObject/AnimPhysObjectVersion.h"
 
-#if INCLUDE_CHAOS
 #include "Chaos/TriangleMeshImplicitObject.h"
-#endif
 
 #if WITH_CHAOS
 	#include "Experimental/ChaosDerivedData.h"
diff --git a/Engine/Source/Runtime/Experimental/ChaosCore/Public/Chaos/Vector.h b/Engine/Source/Runtime/Experimental/ChaosCore/Public/Chaos/Vector.h
--- a/Engine/Source/Runtime/Experimental/ChaosCore/Public/Chaos/Vector.h
+++ b/Engine/Source/Runtime/Experimental/ChaosCore/Public/Chaos/Vector.h
@@ -14,269 +14,336 @@
 #include "Containers/StaticArray.h"
 #include <iostream>
 #include <utility>
+#include <initializer_list>
+#include <limits>
 
 namespace Chaos
 {
 	template<class T, int d>
-	class TVector : public TArray<T>
+	struct TVectorTraits
+	{
+		static const bool RequireRangeCheck = true;
+	};
+
+
+	template<class T, int d>
+	class TVector
 	{
 	public:
-		using TArray<T>::SetNum;
+		using FElement = T;
+		static const int NumElements = d;
+		using FTraits = TVectorTraits<T, d>;
 
-		TVector()
-		    : TArray<T>()
+		TVector() {}
+		explicit TVector(const FElement& Element)
 		{
-			SetNum(d);
+			for (int32 i = 0; i < NumElements; ++i)
+			{
+				V[i] = Element;
+			}
 		}
-		explicit TVector(const T element)
-		    : TArray<T>()
+		~TVector() {}
+
+		TVector(std::initializer_list<T> InElements)
 		{
-			SetNum(d);
-			for (int32 i = 0; i < d; ++i)
-				(*this)[i] = element;
+			check(InElements.size() == NumElements);
+			FElement* Element = &V[0];
+			for (auto InIt = InElements.begin(); InIt != InElements.end(); ++InIt)
+			{
+				*Element++ = *InIt;
+			}
 		}
-		TVector(const T s1, const T s2)
+
+		TVector(const T& V0, const T& V1)
 		{
-			check(d == 2);
-			SetNum(2);
-			(*this)[0] = s1;
-			(*this)[1] = s2;
+			static_assert(NumElements == 2, "Two-element constructor is only for two-element vectors");
+			V[0] = V0;
+			V[1] = V1;
 		}
-		TVector(const T s1, const T s2, const T s3)
+
+		TVector(const T& V0, const T& V1, const T& V2)
 		{
-			check(d == 3);
-			SetNum(3);
-			(*this)[0] = s1;
-			(*this)[1] = s2;
-			(*this)[2] = s3;
+			static_assert(NumElements == 3, "Three-element constructor is only for three-element vectors");
+			V[0] = V0;
+			V[1] = V1;
+			V[2] = V2;
 		}
-		TVector(const T s1, const T s2, const T s3, const T s4)
+
+		TVector(const T& V0, const T& V1, const T& V2, const T& V3)
 		{
-			check(d == 4);
-			SetNum(4);
-			(*this)[0] = s1;
-			(*this)[1] = s2;
-			(*this)[2] = s3;
-			(*this)[3] = s4;
+			static_assert(NumElements == 4, "Four-element constructor is only for four-element vectors");
+			V[0] = V0;
+			V[1] = V1;
+			V[2] = V2;
+			V[3] = V3;
 		}
+
 #if !COMPILE_WITHOUT_UNREAL_SUPPORT
 		TVector(const FVector& Other)
 		{
-			check(d == 3);
-			SetNum(3);
-			(*this)[0] = static_cast<T>(Other.X);
-			(*this)[1] = static_cast<T>(Other.Y);
-			(*this)[2] = static_cast<T>(Other.Z);
+			static_assert(NumElements == 3, "Three-element constructor is only for three-element vectors");
+			V[0] = static_cast<FElement>(Other.X);
+			V[1] = static_cast<FElement>(Other.Y);
+			V[2] = static_cast<FElement>(Other.Z);
 		}
 #endif
+
 		template<class T2>
 		TVector(const TVector<T2, d>& Other)
 		{
-			SetNum(d);
-			for (int32 i = 0; i < d; ++i)
+			for (int32 i = 0; i < NumElements; ++i)
 			{
-				(*this)[i] = static_cast<T>(Other[i]);
-			}
-		}
-		TVector(std::istream& Stream)
-		{
-			for (int32 i = 0; i < d; ++i)
-			{
-				Stream.read(reinterpret_cast<char*>(&(*this)[i]), sizeof(T));
-			}
-		}
-		~TVector() {}
-		void Write(std::ostream& Stream) const
-		{
-			for (int32 i = 0; i < d; ++i)
-			{
-				Stream.write(reinterpret_cast<const char*>(&(*this)[i]), sizeof(T));
+				V[i] = static_cast<T>(Other.V[i]);
 			}
 		}
+
 		TVector<T, d>& operator=(const TVector<T, d>& Other)
 		{
-			for (int32 i = 0; i < d; ++i)
+			for (int32 i = 0; i < NumElements; ++i)
 			{
 				(*this)[i] = Other[i];
 			}
 			return *this;
 		}
-		T Size() const
-		{
-			T SquaredSum = 0;
-			for (int32 i = 0; i < d; ++i)
-			{
-				SquaredSum += ((*this)[i] * (*this)[i]);
-			}
-			return sqrt(SquaredSum);
-		}
-		T Product() const
-		{
-			T Result = 1;
-			for (int32 i = 0; i < d; ++i)
-			{
-				Result *= (*this)[i];
-			}
-			return Result;
-		}
-		static TVector<T, d> AxisVector(const int32 Axis)
-		{
-			check(Axis >= 0 && Axis < d);
-			TVector<T, d> Result(0);
-			Result[Axis] = (T)1;
-			return Result;
-		}
-		T SizeSquared() const
-		{
-			T Result = 0;
-			for (int32 i = 0; i < d; ++i)
-			{
-				Result += ((*this)[i] * (*this)[i]);
-			}
-			return Result;
-		}
-		TVector<T, d> GetSafeNormal() const
-		{
-			//We want N / ||N|| and to avoid inf
-			//So we want N / ||N|| < 1 / eps => N eps < ||N||, but this is clearly true for all eps < 1 and N > 0
-			T SizeSqr = SizeSquared();
-			if (SizeSqr <= TNumericLimits<T>::Min())
-				return AxisVector(0);
-			return (*this) / sqrt(SizeSqr);
-		}
-		T SafeNormalize()
-		{
-			T Size = SizeSquared();
-			if (Size < (T)1e-4)
-			{
-				*this = AxisVector(0);
-				return (T)0.;
-			}
-			Size = sqrt(Size);
-			*this = (*this) / Size;
-			return Size;
-		}
-		TVector<T, d> operator-() const
-		{
-			TVector<T, d> Result;
-			for (int32 i = 0; i < d; ++i)
-			{
-				Result[i] = -(*this)[i];
-			}
-			return Result;
-		}
-		TVector<T, d> operator*(const TVector<T, d>& Other) const
-		{
-			TVector<T, d> Result;
-			for (int32 i = 0; i < d; ++i)
-			{
-				Result[i] = (*this)[i] * Other[i];
-			}
-			return Result;
-		}
-		TVector<T, d> operator/(const TVector<T, d>& Other) const
-		{
-			TVector<T, d> Result;
-			for (int32 i = 0; i < d; ++i)
-			{
-				Result[i] = (*this)[i] / Other[i];
-			}
-			return Result;
-		}
-		TVector<T, d> operator+(const TVector<T, d>& Other) const
+
+		FElement& operator[](int Index)
 		{
-			TVector<T, d> Result;
-			for (int32 i = 0; i < d; ++i)
+			if (FTraits::RequireRangeCheck)
 			{
-				Result[i] = (*this)[i] + Other[i];
+				check(Index >= 0);
+				check(Index < NumElements);
 			}
-			return Result;
+			return V[Index];
 		}
-		TVector<T, d> operator-(const TVector<T, d>& Other) const
+
+		const FElement& operator[](int Index) const
 		{
-			TVector<T, d> Result;
-			for (int32 i = 0; i < d; ++i)
+			if (FTraits::RequireRangeCheck)
 			{
-				Result[i] = (*this)[i] - Other[i];
+				check(Index >= 0);
+				check(Index < NumElements);
 			}
-			return Result;
+			return V[Index];
 		}
-		TVector<T, d>& operator+=(const TVector<T, d>& Other)
+
+		int32 Num() const
 		{
-			for (int32 i = 0; i < d; ++i)
-			{
-				(*this)[i] += Other[i];
-			}
-			return *this;
+			return NumElements;
 		}
-		TVector<T, d>& operator-=(const TVector<T, d>& Other)
-		{
-			for (int32 i = 0; i < d; ++i)
-			{
-				(*this)[i] -= Other[i];
-			}
-			return *this;
-		}
-		TVector<T, d>& operator/=(const TVector<T, d>& Other)
+
+		TVector(std::istream& Stream)
 		{
-			for (int32 i = 0; i < d; ++i)
+			for (int32 i = 0; i < NumElements; ++i)
 			{
-				(*this)[i] /= Other[i];
+				Stream.read(reinterpret_cast<char*>(&V[i]), sizeof(FElement));
 			}
-			return *this;
 		}
-		TVector<T, d> operator*(const T& S) const
+		void Write(std::ostream& Stream) const
 		{
-			TVector<T, d> Result;
-			for (int32 i = 0; i < d; ++i)
+			for (int32 i = 0; i < NumElements; ++i)
 			{
-				Result[i] = (*this)[i] * S;
+				Stream.write(reinterpret_cast<const char*>(&V[i]), sizeof(FElement));
 			}
-			return Result;
 		}
-		TVector<T, d>& operator*=(const T& S)
+
+		friend bool operator==(const TVector<T, d>& L, const TVector<T, d>& R)
 		{
-			for (int32 i = 0; i < d; ++i)
+			for (int32 i = 0; i < NumElements; ++i)
 			{
-				(*this)[i] *= S;
+				if (L.V[i] != R.V[i])
+				{
+					return false;
+				}
 			}
-			return *this;
-		}
-#if COMPILE_WITHOUT_UNREAL_SUPPORT
-		static inline float DotProduct(const Vector<float, 3>& V1, const Vector<float, 3>& V2)
-		{
-			return V1[0] * V2[0] + V1[1] * V2[1] + V1[2] * V2[2];
-		}
-		static inline Vector<float, 3> CrossProduct(const Vector<float, 3>& V1, const Vector<float, 3>& V2)
-		{
-			Vector<float, 3> Result;
-			Result[0] = V1[1] * V2[2] - V1[2] * V2[1];
-			Result[1] = V1[2] * V2[0] - V1[0] * V2[2];
-			Result[2] = V1[0] * V2[1] - V1[1] * V2[0];
-			return Result;
-		}
-#endif
+			return true;
+		}
+
+		friend bool operator!=(const TVector<T, d>& L, const TVector<T, d>& R)
+		{
+			return !(L == R);
+		}
+
+
+		// @todo(ccaulfield): the following should only be available for TVector of numeric types
+//		T Size() const
+//		{
+//			T SquaredSum = 0;
+//			for (int32 i = 0; i < NumElements; ++i)
+//			{
+//				SquaredSum += ((*this)[i] * (*this)[i]);
+//			}
+//			return sqrt(SquaredSum);
+//		}
+//		T Product() const
+//		{
+//			T Result = 1;
+//			for (int32 i = 0; i < NumElements; ++i)
+//			{
+//				Result *= (*this)[i];
+//			}
+//			return Result;
+//		}
+//		static TVector<T, d> AxisVector(const int32 Axis)
+//		{
+//			check(Axis >= 0 && Axis < NumElements);
+//			TVector<T, d> Result(0);
+//			Result[Axis] = (T)1;
+//			return Result;
+//		}
+//		T SizeSquared() const
+//		{
+//			T Result = 0;
+//			for (int32 i = 0; i < d; ++i)
+//			{
+//				Result += ((*this)[i] * (*this)[i]);
+//			}
+//			return Result;
+//		}
+//		TVector<T, d> GetSafeNormal() const
+//		{
+//			//We want N / ||N|| and to avoid inf
+//			//So we want N / ||N|| < 1 / eps => N eps < ||N||, but this is clearly true for all eps < 1 and N > 0
+//			T SizeSqr = SizeSquared();
+//			if (SizeSqr <= TNumericLimits<T>::Min())
+//				return AxisVector(0);
+//			return (*this) / sqrt(SizeSqr);
+//		}
+//		T SafeNormalize()
+//		{
+//			T Size = SizeSquared();
+//			if (Size < (T)1e-4)
+//			{
+//				*this = AxisVector(0);
+//				return (T)0.;
+//			}
+//			Size = sqrt(Size);
+//			*this = (*this) / Size;
+//			return Size;
+//		}
+//		TVector<T, d> operator-() const
+//		{
+//			TVector<T, d> Result;
+//			for (int32 i = 0; i < Num; ++i)
+//			{
+//				Result[i] = -(*this)[i];
+//			}
+//			return Result;
+//		}
+//		TVector<T, d> operator*(const TVector<T, d>& Other) const
+//		{
+//			TVector<T, d> Result;
+//			for (int32 i = 0; i < d; ++i)
+//			{
+//				Result[i] = (*this)[i] * Other[i];
+//			}
+//			return Result;
+//		}
+//		TVector<T, d> operator/(const TVector<T, d>& Other) const
+//		{
+//			TVector<T, d> Result;
+//			for (int32 i = 0; i < d; ++i)
+//			{
+//				Result[i] = (*this)[i] / Other[i];
+//			}
+//			return Result;
+//		}
+//		TVector<T, d> operator+(const TVector<T, d>& Other) const
+//		{
+//			TVector<T, d> Result;
+//			for (int32 i = 0; i < d; ++i)
+//			{
+//				Result[i] = (*this)[i] + Other[i];
+//			}
+//			return Result;
+//		}
+//		TVector<T, d> operator-(const TVector<T, d>& Other) const
+//		{
+//			TVector<T, d> Result;
+//			for (int32 i = 0; i < d; ++i)
+//			{
+//				Result[i] = (*this)[i] - Other[i];
+//			}
+//			return Result;
+//		}
+//		TVector<T, d>& operator+=(const TVector<T, d>& Other)
+//		{
+//			for (int32 i = 0; i < d; ++i)
+//			{
+//				(*this)[i] += Other[i];
+//			}
+//			return *this;
+//		}
+//		TVector<T, d>& operator-=(const TVector<T, d>& Other)
+//		{
+//			for (int32 i = 0; i < d; ++i)
+//			{
+//				(*this)[i] -= Other[i];
+//			}
+//			return *this;
+//		}
+//		TVector<T, d>& operator/=(const TVector<T, d>& Other)
+//		{
+//			for (int32 i = 0; i < d; ++i)
+//			{
+//				(*this)[i] /= Other[i];
+//			}
+//			return *this;
+//		}
+//		TVector<T, d> operator*(const T& S) const
+//		{
+//			TVector<T, d> Result;
+//			for (int32 i = 0; i < d; ++i)
+//			{
+//				Result[i] = (*this)[i] * S;
+//			}
+//			return Result;
+//		}
+//		friend TVector<T, d> operator*(const T S, const TVector<T, d>& V)
+//		{
+//			TVector<T, d> Ret;
+//			for (int32 i = 0; i < d; ++i)
+//			{
+//				Ret[i] = S * V[i];
+//			}
+//			return Ret;
+//		}
+//		TVector<T, d>& operator*=(const T& S)
+//		{
+//			for (int32 i = 0; i < d; ++i)
+//			{
+//				(*this)[i] *= S;
+//			}
+//			return *this;
+//		}
+//		friend TVector<T, d> operator/(const T S, const TVector<T, d>& V)
+//		{
+//			TVector<T, d> Ret;
+//			for (int32 i = 0; i < d; ++i)
+//			{
+//				Ret = S / V[i];
+//			}
+//			return Ret;
+//		}
+//
+//#if COMPILE_WITHOUT_UNREAL_SUPPORT
+//		static inline float DotProduct(const Vector<float, 3>& V1, const Vector<float, 3>& V2)
+//		{
+//			return V1[0] * V2[0] + V1[1] * V2[1] + V1[2] * V2[2];
+//		}
+//		static inline Vector<float, 3> CrossProduct(const Vector<float, 3>& V1, const Vector<float, 3>& V2)
+//		{
+//			Vector<float, 3> Result;
+//			Result[0] = V1[1] * V2[2] - V1[2] * V2[1];
+//			Result[1] = V1[2] * V2[0] - V1[0] * V2[2];
+//			Result[2] = V1[0] * V2[1] - V1[1] * V2[0];
+//			return Result;
+//		}
+//#endif
+
+	private:
+		FElement V[NumElements];
 	};
-	template<class T, int d>
-	inline TVector<T, d> operator*(const T S, const TVector<T, d>& V)
-	{
-		TVector<T, d> Ret;
-		for (int32 i = 0; i < d; ++i)
-		{
-			Ret[i] = S * V[i];
-		}
-		return Ret;
-	}
-	template<class T, int d>
-	inline TVector<T, d> operator/(const T S, const TVector<T, d>& V)
-	{
-		TVector<T, d> Ret;
-		for (int32 i = 0; i < d; ++i)
-		{
-			Ret = S / V[i];
-		}
-		return Ret;
-	}
+
 
 #if !COMPILE_WITHOUT_UNREAL_SUPPORT
 	template<>
@@ -355,6 +422,22 @@ namespace Chaos
 		{
 			return TVector<float, 3>(X / Other, Y / Other, Z / Other);
 		}
+		friend TVector<float, 3> operator/(const float S, const TVector<float, 3>& V)
+		{
+			return TVector<float, 3>(S / V.X, S / V.Y, S / V.Z);
+		}
+		TVector<float, 3> operator-(const TVector<float, 3>& Other) const
+		{
+			return TVector<float, 3>(X - Other[0], Y - Other[1], Z - Other[2]);
+		}
+		TVector<float, 3> operator*(const TVector<float, 3>& Other) const
+		{
+			return TVector<float, 3>(X * Other[0], Y * Other[1], Z * Other[2]);
+		}
+		TVector<float, 3> operator/(const TVector<float, 3>& Other) const
+		{
+			return TVector<float, 3>(X / Other[0], Y / Other[1], Z / Other[2]);
+		}
 		template<class T2>
 		TVector<float, 3> operator-(const TVector<T2, 3>& Other) const
 		{
@@ -443,16 +526,15 @@ namespace Chaos
 			return atan2(s, c);
 		}
 	};
-	template<>
-	inline TVector<float, 3> operator/(const float S, const TVector<float, 3>& V)
-	{
-		return TVector<float, 3>(S / V.X, S / V.Y, S / V.Z);
-	}
 
 	template<>
 	class TVector<float, 2> : public FVector2D
 	{
 	public:
+
+		using FVector2D::X;
+		using FVector2D::Y;
+
 		TVector()
 		    : FVector2D() {}
 		TVector(const float x)
@@ -466,6 +548,12 @@ namespace Chaos
 			Stream.read(reinterpret_cast<char*>(&X), sizeof(X));
 			Stream.read(reinterpret_cast<char*>(&Y), sizeof(Y));
 		}
+		template <typename OtherT>
+		TVector(const TVector<OtherT, 2>& InVector)
+		{
+			X = ((float)InVector[0]);
+			Y = ((float)InVector[1]);
+		}
 		~TVector() {}
 		void Write(std::ostream& Stream) const
 		{
@@ -506,6 +594,15 @@ namespace Chaos
 				return MakePair(max.Y, 1);
 			}
 		}
+		template<class T2>
+		TVector<float, 2> operator/(const TVector<T2, 2>& Other) const
+		{
+			return TVector<float, 2>(X / Other[0], Y / Other[1]);
+		}
+		TVector<float, 2> operator/(const float Other) const
+		{
+			return TVector<float, 2>(X / Other, Y / Other);
+		}
 	};
 #endif // !COMPILE_WITHOUT_UNREAL_SUPPORT
 
@@ -673,93 +770,103 @@ namespace Chaos
 	class TVector<int32, 2>
 	{
 	public:
+		using FElement = int32;
+
 		FORCEINLINE TVector()
 		{}
-		FORCEINLINE explicit TVector(const int32 InX)
+		FORCEINLINE explicit TVector(const FElement InX)
 		    : X(InX), Y(InX)
 		{}
-		FORCEINLINE TVector(const int32 InX, const int32 InY)
+		FORCEINLINE TVector(const FElement InX, const FElement InY)
 		    : X(InX), Y(InY)
 		{}
+		template<typename OtherT>
+		FORCEINLINE TVector(const TVector<OtherT, 2>& InVector)
+			: X((int32)InVector.X)
+			, Y((int32)InVector.Y)
+		{}
 		FORCEINLINE ~TVector()
 		{}
 
 		FORCEINLINE int32 Num() const { return 2; }
 
-		FORCEINLINE int32 Product() const
+		FORCEINLINE FElement Product() const
 		{
 			return X * Y;
 		}
 
-		FORCEINLINE static TVector<int32, 2>
-		AxisVector(const int32 Axis)
+		FORCEINLINE static TVector<FElement, 2> AxisVector(const int32 Axis)
 		{
-			TVector<int32, 2> Result(0);
+			TVector<FElement, 2> Result(0);
 			Result[Axis] = 1;
 			return Result;
 		}
 
 		FORCEINLINE void Write(std::ostream& Stream) const
 		{
-			Stream.write(reinterpret_cast<const char*>(&X), sizeof(int32));
-			Stream.write(reinterpret_cast<const char*>(&Y), sizeof(int32));
+			Stream.write(reinterpret_cast<const char*>(&X), sizeof(FElement));
+			Stream.write(reinterpret_cast<const char*>(&Y), sizeof(FElement));
 		}
 
-		FORCEINLINE TVector<int32, 2>&
-		operator=(const TVector<int32, 2>& Other)
+		FORCEINLINE TVector<FElement, 2>& operator=(const TVector<FElement, 2>& Other)
 		{
 			X = Other.X;
 			Y = Other.Y;
 			return *this;
 		}
 
-		FORCEINLINE bool
-		operator==(const TVector<int32, 2>& Other) const
-		{
-			return X == Other.X && Y == Other.Y;
-		}
-
-		FORCEINLINE int32 operator[](const int32 Idx) const { return (static_cast<const int32*>(&X))[Idx]; }
-		FORCEINLINE int32& operator[](const int32 Idx) { return (static_cast<int32*>(&X))[Idx]; }
+		FORCEINLINE FElement operator[](const int32 Idx) const { return (static_cast<const FElement*>(&X))[Idx]; }
+		FORCEINLINE FElement& operator[](const int32 Idx) { return (static_cast<FElement*>(&X))[Idx]; }
 
-		FORCEINLINE TVector<int32, 2> operator-() const { return {-X, -Y}; }
-		FORCEINLINE TVector<int32, 2> operator*(const TVector<int32, 2>& Other) const { return {X * Other.X, Y * Other.Y}; }
-		FORCEINLINE TVector<int32, 2> operator/(const TVector<int32, 2>& Other) const { return {X / Other.X, Y / Other.Y}; }
-		FORCEINLINE TVector<int32, 2> operator+(const TVector<int32, 2>& Other) const { return {X + Other.X, Y + Other.Y}; }
-		FORCEINLINE TVector<int32, 2> operator-(const TVector<int32, 2>& Other) const { return {X - Other.X, Y - Other.Y}; }
+		FORCEINLINE TVector<FElement, 2> operator-() const { return {-X, -Y}; }
+		FORCEINLINE TVector<FElement, 2> operator*(const TVector<FElement, 2>& Other) const { return {X * Other.X, Y * Other.Y}; }
+		FORCEINLINE TVector<FElement, 2> operator/(const TVector<FElement, 2>& Other) const { return {X / Other.X, Y / Other.Y}; }
+		FORCEINLINE TVector<FElement, 2> operator+(const TVector<FElement, 2>& Other) const { return {X + Other.X, Y + Other.Y}; }
+		FORCEINLINE TVector<FElement, 2> operator-(const TVector<FElement, 2>& Other) const { return {X - Other.X, Y - Other.Y}; }
 
-		FORCEINLINE TVector<int32, 2>& operator+=(const TVector<int32, 2>& Other)
+		FORCEINLINE TVector<FElement, 2>& operator+=(const TVector<FElement, 2>& Other)
 		{
 			X += Other.X;
 			Y += Other.Y;
 			return *this;
 		}
-		FORCEINLINE TVector<int32, 2>& operator-=(const TVector<int32, 2>& Other)
+		FORCEINLINE TVector<FElement, 2>& operator-=(const TVector<FElement, 2>& Other)
 		{
 			X -= Other.X;
 			Y -= Other.Y;
 			return *this;
 		}
-		FORCEINLINE TVector<int32, 2>& operator/=(const TVector<int32, 2>& Other)
+		FORCEINLINE TVector<FElement, 2>& operator/=(const TVector<FElement, 2>& Other)
 		{
 			X /= Other.X;
 			Y /= Other.Y;
 			return *this;
 		}
-		FORCEINLINE TVector<int32, 2> operator*(const int32& S) const
+		FORCEINLINE TVector<FElement, 2> operator*(const FElement& S) const
 		{
 			return {X * S, Y * S};
 		}
-		FORCEINLINE TVector<int32, 2>& operator*=(const int32& S)
+		FORCEINLINE TVector<FElement, 2>& operator*=(const FElement& S)
 		{
 			X *= S;
 			Y *= S;
 			return *this;
 		}
 
+		FORCEINLINE friend bool operator==(const TVector<FElement, 2>& L, const TVector<FElement, 2>& R)
+		{
+			return (L.X == R.X) && (L.Y == R.Y);
+		}
+
+		FORCEINLINE friend bool operator!=(const TVector<FElement, 2>& L, const TVector<FElement, 2>& R)
+		{
+			return !(L == R);
+		}
+
+
 	private:
-		int32 X;
-		int32 Y;
+		FElement X;
+		FElement Y;
 	};
 
 	template<class T>
