Compile fix

- Hide the microsoft warning behind a _MSC_VER check, not a !PLATFORM_PS4, as some of our clang based platforms seem to have Wunknown-pragmas enabled.

#rnx

diff --git a/Engine/Source/Runtime/Experimental/Chaos/Public/Chaos/PBDRigidParticles.h b/Engine/Source/Runtime/Experimental/Chaos/Public/Chaos/PBDRigidParticles.h
--- a/Engine/Source/Runtime/Experimental/Chaos/Public/Chaos/PBDRigidParticles.h
+++ b/Engine/Source/Runtime/Experimental/Chaos/Public/Chaos/PBDRigidParticles.h
@@ -5,7 +5,7 @@
 #include "Chaos/RigidParticles.h"
 #include "Chaos/Rotation.h"
 
-#if !PLATFORM_PS4
+#if defined(_MSC_VER)
 #pragma warning(push)
 #pragma warning(disable:4946)
 #endif
@@ -186,6 +186,6 @@ FChaosArchive& operator<<(FChaosArchive& Ar, TPBDRigidParticles<T, d>& Particles
 }
 }
 
-#if !PLATFORM_PS4
+#if defined(_MSC_VER)
 #pragma warning(pop)
 #endif
