Merging

//UE4/Private-Reverb/Engine/Source/Runtime/Experimental/Chaos/Public/Chaos/PBDRigidsEvolution.h
//UE4/Private-Reverb/Engine/Source/Runtime/Experimental/ChaosSolvers/Private/PhysicsProxy/FieldSystemPhysicsProxy.cpp

Disable field works

diff --git a/Engine/Source/Runtime/Experimental/Chaos/Public/Chaos/PBDRigidsEvolution.h b/Engine/Source/Runtime/Experimental/Chaos/Public/Chaos/PBDRigidsEvolution.h
--- a/Engine/Source/Runtime/Experimental/Chaos/Public/Chaos/PBDRigidsEvolution.h
+++ b/Engine/Source/Runtime/Experimental/Chaos/Public/Chaos/PBDRigidsEvolution.h
@@ -388,6 +388,15 @@ class TPBDRigidsEvolutionBase
 	const auto& GetNonDisabledClusteredArray() const { return Particles.GetNonDisabledClusteredArray(); }
 
 	CHAOS_API TSerializablePtr<FChaosPhysicsMaterial> GetPhysicsMaterial(const TGeometryParticleHandle<T, d>* Particle) const { return Particle->AuxilaryValue(PhysicsMaterials); }
+	
+	CHAOS_API const TUniquePtr<FChaosPhysicsMaterial> &GetPerParticlePhysicsMaterial(const TGeometryParticleHandle<T, d>* Particle) const { return Particle->AuxilaryValue(PerParticlePhysicsMaterials); }
+
+	CHAOS_API void SetPerParticlePhysicsMaterial(TGeometryParticleHandle<T, d>* Particle, TUniquePtr<FChaosPhysicsMaterial> &InMaterial)
+	{
+		check(!Particle->AuxilaryValue(PerParticlePhysicsMaterials)); //shouldn't be setting non unique material if a unique one already exists
+		Particle->AuxilaryValue(PerParticlePhysicsMaterials) = MoveTemp(InMaterial);
+	}
+
 	CHAOS_API void SetPhysicsMaterial(TGeometryParticleHandle<T,d>* Particle, TSerializablePtr<FChaosPhysicsMaterial> InMaterial)
 	{
 		check(!Particle->AuxilaryValue(PerParticlePhysicsMaterials)); //shouldn't be setting non unique material if a unique one already exists
diff --git a/Engine/Source/Runtime/Experimental/ChaosSolvers/Private/PhysicsProxy/FieldSystemPhysicsProxy.cpp b/Engine/Source/Runtime/Experimental/ChaosSolvers/Private/PhysicsProxy/FieldSystemPhysicsProxy.cpp
--- a/Engine/Source/Runtime/Experimental/ChaosSolvers/Private/PhysicsProxy/FieldSystemPhysicsProxy.cpp
+++ b/Engine/Source/Runtime/Experimental/ChaosSolvers/Private/PhysicsProxy/FieldSystemPhysicsProxy.cpp
@@ -590,119 +590,59 @@ void FFieldSystemPhysicsProxy::FieldParameterUpdateCallback(
 				if (ensureMsgf(Command.RootNode->Type() == FFieldNode<float>::StaticType(),
 					TEXT("Field based evaluation of the simulations 'Disable' parameter expects scale field inputs.")))
 				{
-#if TODO_REIMPLEMENT_PHYSICS_PROXY_REVERSE_MAPPING
-					const Chaos::TArrayCollectionArray<PhysicsProxyWrapper>& PhysicsProxyMapping = CurrentSolver->GetPhysicsProxyReverseMapping();
 					
-					FFieldSystemPhysicsProxy::ContiguousIndices(IndicesArray, CurrentSolver, ResolutionType, IndicesArray.Num() != Particles.Size());
-					if (IndicesArray.Num())
+					if (Handles.Num())
 					{
-						TArrayView<ContextIndex> IndexView(&(IndicesArray[0]), IndicesArray.Num());
-
-						FVector * tptr = &(Particles.X(0));
-						TArrayView<FVector> SamplesView(tptr, int32(Particles.Size()));
+						TArrayView<FVector> SamplePointsView(&(SamplePoints[0]), SamplePoints.Num());
+						TArrayView<ContextIndex> SampleIndicesView(&(SampleIndices[0]), SampleIndices.Num());
+						FFieldContext Context(
+							SampleIndicesView,
+							SamplePointsView,
+							Command.MetaData);
 
-						FFieldContext Context{
-							IndexView,
-							SamplesView,
-							Command.MetaData
-						};
+						TArray<float> LocalResults;
+						LocalResults.AddUninitialized(Handles.Num());
+						TArrayView<float> ResultsView(&(LocalResults[0]), LocalResults.Num());
+						static_cast<const FFieldNode<float> *>(Command.RootNode.Get())->Evaluate(
+							Context, ResultsView);
 
-						TArray<float> Results;
-						Results.AddUninitialized(Particles.Size());
-						for (const ContextIndex Index : IndicesArray)
+						int32 i = 0;
+						for (Chaos::TGeometryParticleHandle<float, 3>* Handle : Handles)
 						{
-							const SolverObjectWrapper& ParticleObjectWrapper = SolverObjectMapping[Index.Result];
-							TSerializablePtr<Chaos::FChaosPhysicsMaterial> Material = CurrentSolver->GetPhysicsMaterial(Index.Result);
-							if (ensure(Material) && ParticleObjectWrapper.SolverObject)
-							{
-								const TUniquePtr<Chaos::FChaosPhysicsMaterial>& InstanceMaterial = CurrentSolver->GetPerParticlePhysicsMaterial(Index.Result);
-								if (InstanceMaterial)
-								{
-									Results[Index.Result] = InstanceMaterial->DisabledLinearThreshold;
-								}
-								else
+							Chaos::TPBDRigidParticleHandle<float, 3>* RigidHandle = Handle->CastToRigidParticle();
+							
+							if (RigidHandle && RigidHandle->ObjectState() == Chaos::EObjectStateType::Dynamic && ResultsView.Num() > 0)
+							{						
+								// if no per particle physics material is set, make one
+								if (!CurrentSolver->GetEvolution()->GetPerParticlePhysicsMaterial(RigidHandle).IsValid())
 								{
-									Results[Index.Result] = Material->DisabledLinearThreshold;
-								}
-							}
-							else
-							{
-								Results[Index.Result] = 0.0f;
-							}
-						}
-
-						TArrayView<float> ResultsView(&(Results[0]), Results.Num());
-						static_cast<const FFieldNode<float> *>(Command.RootNode.Get())->Evaluate(Context, ResultsView);
-
-						for (const ContextIndex Index : IndicesArray)
-						{
-							const int32 i = Index.Result;
-							const PhysicsProxyWrapper& ParticleObjectWrapper = PhysicsProxyMapping[i];
-							TSerializablePtr<Chaos::FChaosPhysicsMaterial> Material = CurrentSolver->GetPhysicsMaterial(i);
-							if (!ensure(Material) || !ParticleObjectWrapper.PhysicsProxy)	//question: do we actually need to check for solver object?
-							{
-								continue;
-							}
 
-							//per instance override
-							if (!CurrentSolver->GetPerParticlePhysicsMaterial(Index.Result))
-							{
-								if (Results[i] != Material->DisabledLinearThreshold)
-								{
-									// value changed from shared material, make unique material.
-									TUniquePtr<Chaos::FChaosPhysicsMaterial> NewMaterial = MakeUnique< Chaos::FChaosPhysicsMaterial>(*Material);
-									CurrentSolver->SetPerParticlePhysicsMaterial(Index.Result, MoveTemp(NewMaterial));
-									const TUniquePtr<FChaosPhysicsMaterial>& InstanceMaterial = CurrentSolver->GetPerParticlePhysicsMaterial(i);
+									TUniquePtr<Chaos::FChaosPhysicsMaterial> NewMaterial = MakeUnique< Chaos::FChaosPhysicsMaterial>();
+									NewMaterial->DisabledLinearThreshold = ResultsView[i];
+									NewMaterial->DisabledAngularThreshold = ResultsView[i];
 
-									InstanceMaterial->DisabledLinearThreshold = Results[i];
-									InstanceMaterial->DisabledAngularThreshold = Results[i];
-								}
-							}
-							else
-							{
-								const TUniquePtr<FChaosPhysicsMaterial>& InstanceMaterial = CurrentSolver->GetPerParticlePhysicsMaterial(i);
 
-								if (InstanceMaterial->DisabledLinearThreshold != Results[i])
+									CurrentSolver->GetEvolution()->SetPhysicsMaterial(RigidHandle, MakeSerializable(NewMaterial));
+									CurrentSolver->GetEvolution()->SetPerParticlePhysicsMaterial(RigidHandle, NewMaterial);
+								} 
+								else
 								{
-									InstanceMaterial->DisabledLinearThreshold = Results[i];
-									InstanceMaterial->DisabledAngularThreshold = Results[i];
+									const TUniquePtr<FChaosPhysicsMaterial> &InstanceMaterial = CurrentSolver->GetEvolution()->GetPerParticlePhysicsMaterial(RigidHandle);
+		
+									if (ResultsView[i] != InstanceMaterial->DisabledLinearThreshold)
+									{
+										InstanceMaterial->DisabledLinearThreshold = ResultsView[i];
+										InstanceMaterial->DisabledAngularThreshold = ResultsView[i];
+									}
 								}
 							}
+							++i;
 						}
-					}
-#endif
+					}									
 				}
-				CommandsToRemove.Add(CommandIndex);
-			}
-#if TODO_REIMPLEMENT_RIGID_CLUSTERING
-			else if (Command.TargetAttribute == GetFieldPhysicsName(EFieldPhysicsType::Field_InternalClusterStrain))
-			{
-				SCOPE_CYCLE_COUNTER(STAT_ParamUpdateField_InternalClusterStrain);
-				if (ensureMsgf(Command.RootNode->Type() == FFieldNode<float>::StaticType(),
-					TEXT("Field based evaluation of the simulations 'ExternalClusterStrain' parameter expects scalar field inputs.")))
-				{
-					FFieldSystemPhysicsProxy::ContiguousIndices(IndicesArray, CurrentSolver, ResolutionType, IndicesArray.Num() != Particles.Size());
-					if (IndicesArray.Num())
-					{
-						TArrayView<ContextIndex> IndexView(&(IndicesArray[0]), IndicesArray.Num());
-
-						FVector * tptr = &(Particles.X(0));
-						TArrayView<FVector> SamplesView(tptr, int32(Particles.Size()));
-
-						FFieldContext Context{
-							IndexView,
-							SamplesView,
-							Command.MetaData
-						};
 
-						float * vptr = &(Strains[0]);
-						TArrayView<float> ResultsView(vptr, Particles.Size());
-						static_cast<const FFieldNode<float> *>(Command.RootNode.Get())->Evaluate(Context, ResultsView);
-					}
-				}
 				CommandsToRemove.Add(CommandIndex);
 			}
-#endif
 			else if (Command.TargetAttribute == GetFieldPhysicsName(EFieldPhysicsType::Field_CollisionGroup))
 			{
 				if (ensureMsgf(Command.RootNode->Type() == FFieldNode<int32>::StaticType(),
