Fix for AABB tree queries failing when filter data changes.

diff --git a/Engine/Source/Runtime/Experimental/Chaos/Private/Chaos/AABBTree.cpp b/Engine/Source/Runtime/Experimental/Chaos/Private/Chaos/AABBTree.cpp
new file mode 100644
--- /dev/null
+++ b/Engine/Source/Runtime/Experimental/Chaos/Private/Chaos/AABBTree.cpp
@@ -0,0 +1,6 @@
+// Copyright Epic Games, Inc. All Rights Reserved.
+
+#include "Chaos/AABBTree.h"
+
+int32 FAABBTreeCVars::UpdateDirtyElementPayloadData = 1;
+FAutoConsoleVariableRef FAABBTreeCVars::CVarUpdateDirtyElementPayloadData(TEXT("p.aabbtree.updatedirtyelementpayloads"), FAABBTreeCVars::UpdateDirtyElementPayloadData, TEXT("Allow AABB tree elements to update internal payload data when they recieve a payload update"));
\ No newline at end of file
diff --git a/Engine/Source/Runtime/Experimental/Chaos/Public/Chaos/AABBTree.h b/Engine/Source/Runtime/Experimental/Chaos/Public/Chaos/AABBTree.h
--- a/Engine/Source/Runtime/Experimental/Chaos/Public/Chaos/AABBTree.h
+++ b/Engine/Source/Runtime/Experimental/Chaos/Public/Chaos/AABBTree.h
@@ -10,6 +10,12 @@
 #include "Templates/Models.h"
 #include "Chaos/BoundingVolume.h"
 
+struct FAABBTreeCVars
+{
+	static int32 UpdateDirtyElementPayloadData;
+	static FAutoConsoleVariableRef CVarUpdateDirtyElementPayloadData;
+};
+
 namespace Chaos
 {
 
@@ -272,6 +278,27 @@ inline FArchive& operator<<(FArchive& Ar, FAABBTreePayloadInfo& PayloadInfo)
 
 extern CHAOS_API int32 MaxDirtyElements;
 
+struct CIsUpdatableElement
+{
+	template<typename ElementT>
+	auto Requires(ElementT& InElem, const ElementT& InOtherElem) -> decltype(InElem.UpdateFrom(InOtherElem));
+};
+
+template<typename T, typename TEnableIf<!TModels<CIsUpdatableElement, T>::Value>::Type* = nullptr>
+static void UpdateElementHelper(T& InElem, const T& InFrom)
+{
+
+}
+
+template<typename T, typename TEnableIf<TModels<CIsUpdatableElement, T>::Value>::Type* = nullptr>
+static void UpdateElementHelper(T& InElem, const T& InFrom)
+{
+	if(FAABBTreeCVars::UpdateDirtyElementPayloadData != 0)
+	{
+		InElem.UpdateFrom(InFrom);
+	}
+}
+
 template <typename TPayloadType, typename TLeafType, typename T, bool bMutable = true>
 class TAABBTree final : public ISpatialAcceleration<TPayloadType, T, 3>
 {
@@ -494,6 +521,7 @@ public:
 				else
 				{
 					DirtyElements[PayloadInfo->DirtyPayloadIdx].Bounds = NewBounds;
+					UpdateElementHelper(DirtyElements[PayloadInfo->DirtyPayloadIdx].Payload, Payload);
 				}
 
 				// Handle something that previously did not have bounds that may be in global elements.
@@ -519,6 +547,7 @@ public:
 				else
 				{
 					GlobalPayloads[PayloadInfo->GlobalPayloadIdx].Bounds = GlobalBounds;
+					UpdateElementHelper(GlobalPayloads[PayloadInfo->GlobalPayloadIdx].Payload, Payload);
 				}
 
 				// Handle something that previously had bounds that may be in dirty elements.
diff --git a/Engine/Source/Runtime/Experimental/Chaos/Public/Chaos/ParticleHandle.h b/Engine/Source/Runtime/Experimental/Chaos/Public/Chaos/ParticleHandle.h
--- a/Engine/Source/Runtime/Experimental/Chaos/Public/Chaos/ParticleHandle.h
+++ b/Engine/Source/Runtime/Experimental/Chaos/Public/Chaos/ParticleHandle.h
@@ -139,10 +139,10 @@ FORCEINLINE_DEBUGGABLE bool PrePreFilterImp(const FCollisionFilterData& QueryFil
 		//since we're taking the union of shapes we can only support trace channel
 		//const ECollisionChannel QuerierChannel = GetCollisionChannel(QueryFilter.Word3);
 		const uint32 QuerierChannel = (QueryFilterData.Word3 << 6) >> (32 - 5);
-
+	
 		//uint32 const QuerierBit = ECC_TO_BITFIELD(QuerierChannel);
 		const uint32 QuerierBit = (1 << (QuerierChannel));
-
+	
 		// check if Querier wants a hit
 		const uint32 TouchOrBlock = (UnionFilterData.Word1 | UnionFilterData.Word2);
 		return !(QuerierBit & TouchOrBlock);
@@ -208,6 +208,14 @@ public:
 		return false;
 	}
 
+	void UpdateFrom(const TAccelerationStructureHandle<T, d>& InOther)
+	{
+		UnionFilterData.Word0 = InOther.UnionFilterData.Word0;
+		UnionFilterData.Word1 = InOther.UnionFilterData.Word1;
+		UnionFilterData.Word2 = InOther.UnionFilterData.Word2;
+		UnionFilterData.Word3 = InOther.UnionFilterData.Word3;
+	}
+
 private:
 	TGeometryParticle<T, d>* ExternalGeometryParticle;
 	TGeometryParticleHandle<T, d>* GeometryParticleHandle;
