Chaos - fix swing-twist decomposition
- add swing-twist unit tests

diff --git a/Engine/Source/Programs/HeadlessChaos/Private/HeadlessChaosTestJointUtilities.cpp b/Engine/Source/Programs/HeadlessChaos/Private/HeadlessChaosTestJointUtilities.cpp
new file mode 100644
--- /dev/null
+++ b/Engine/Source/Programs/HeadlessChaos/Private/HeadlessChaosTestJointUtilities.cpp
@@ -0,0 +1,119 @@
+// Copyright Epic Games, Inc. All Rights Reserved.
+
+#include "HeadlessChaos.h"
+
+#include "Chaos/PBDJointConstraintUtilities.h"
+#include "ChaosLog.h"
+
+namespace ChaosTest
+{
+	using namespace Chaos;
+
+	struct SwingTwistCase
+	{
+	public:
+		FVec3 SwingAxis;
+		FReal SwingAngleDeg;
+		FReal TwistAngleDeg;
+
+		friend std::ostream& operator<<(std::ostream& s, const SwingTwistCase& Case)
+		{
+			return s << "Twist/Swing: " << Case.TwistAngleDeg << "/" << Case.SwingAngleDeg << " Swing Axis: (" << Case.SwingAxis.X << ", " << Case.SwingAxis.Y << ", " << Case.SwingAxis.Z << ")";
+		}
+	};
+
+	void TestAnglesDeg(const SwingTwistCase& Case, FReal A0, FReal A1, FReal Tolerance)
+	{
+		if (!(FMath::IsNearlyEqual(FMath::Abs(A0 - A1), 0.0f, Tolerance) || FMath::IsNearlyEqual(FMath::Abs(A0 - A1), 360.0f, Tolerance)))
+		{
+			GTEST_FAIL() << "Angle Test Fail: " << A0 << " != " << A1 << " " << Case;
+		}
+	}
+
+	void TestSwingTwistOrder(const SwingTwistCase& Case)
+	{
+		FRotation3 SwingRot = FRotation3::FromAxisAngle(Case.SwingAxis, FMath::DegreesToRadians(Case.SwingAngleDeg));
+
+		FVec3 TwistAxis = FVec3(1, 0, 0);
+		FRotation3 TwistRot = FRotation3::FromAxisAngle(TwistAxis, FMath::DegreesToRadians(Case.TwistAngleDeg));
+
+		FRotation3 RST = SwingRot * TwistRot;
+		FRotation3 RS = SwingRot;
+
+		// Verify that a vector along the X Axis is unaffected by twist
+		FVec3 X = FVec3(100, 0, 0);
+		FVec3 XST = RST * X;	// Swing and Twist applied
+		FVec3 XS = RS * X;		// Just Swing applied
+
+		EXPECT_NEAR(XS.X, XST.X, KINDA_SMALL_NUMBER) << Case;
+		EXPECT_NEAR(XS.Y, XST.Y, KINDA_SMALL_NUMBER) << Case;
+		EXPECT_NEAR(XS.Z, XST.Z, KINDA_SMALL_NUMBER) << Case;
+	}
+
+	void TestSwingTwistDecomposition(const SwingTwistCase& Case)
+	{
+		FRotation3 SwingRot = FRotation3::FromAxisAngle(Case.SwingAxis, FMath::DegreesToRadians(Case.SwingAngleDeg));
+
+		FVec3 TwistAxis = FVec3(1, 0, 0);
+		FRotation3 TwistRot = FRotation3::FromAxisAngle(TwistAxis, FMath::DegreesToRadians(Case.TwistAngleDeg));
+
+		FRotation3 R0 = FRotation3::Identity;
+		FRotation3 R1 = R0 * SwingRot * TwistRot;
+
+		FVec3 OutTwistAxis, OutSwingAxisLocal;
+		FReal OutTwistAngle, OutSwingAngle;
+
+		FPBDJointUtilities::GetTwistAxisAngle(R0, R1, OutTwistAxis, OutTwistAngle);
+		FPBDJointUtilities::GetConeAxisAngleLocal(R0, R1,  1.e-6f, OutSwingAxisLocal, OutSwingAngle);
+		FReal OutTwistAngleDeg = FMath::RadiansToDegrees(OutTwistAngle);
+		FReal OutSwingAngleDeg = FMath::RadiansToDegrees(OutSwingAngle);
+
+		// Degenerate behavior at 180 degrees
+		if (Case.SwingAngleDeg == 180)
+		{
+			TestAnglesDeg(Case, 180.0f, OutSwingAngleDeg, 0.1f);
+			TestAnglesDeg(Case, 0.0f, OutTwistAngleDeg, 0.1f);
+			return;
+		}
+
+		// If we expect a non-zero swing, make sure we recovered the swing axis
+		FReal ExpectedSwingAngleDeg = (FVec3::DotProduct(Case.SwingAxis, OutSwingAxisLocal) >= 0.0f) ? Case.SwingAngleDeg : 360 - Case.SwingAngleDeg;
+		if (ExpectedSwingAngleDeg != 0)
+		{
+			EXPECT_NEAR(FMath::Abs(FVec3::DotProduct(Case.SwingAxis, OutSwingAxisLocal)), 1.0f, 1.e-2f) << Case;
+		}
+
+		TestAnglesDeg(Case, ExpectedSwingAngleDeg, OutSwingAngleDeg, 0.1f);
+		TestAnglesDeg(Case, Case.TwistAngleDeg, OutTwistAngleDeg, 0.1f);
+	}
+
+	GTEST_TEST(JointUtilitiesTests, TestSwingTwistDecomposition)
+	{
+		FVec3 SwingAxes[] = 
+		{
+			FVec3(0, 1, 0),
+			FVec3(0, 0, 1),
+			FVec3(0, 1, 1).GetSafeNormal(),
+		};
+
+		//int32 TwistIndex = 3;
+		for (int32 TwistIndex = 0; TwistIndex < 360; ++TwistIndex)
+		{
+			FReal TwistAngleDeg = (FReal)TwistIndex;
+
+			//int32 SwingIndex = 1;
+			for (int32 SwingIndex = 0; SwingIndex < 360; ++SwingIndex)
+			{
+				FReal SwingAngleDeg = (FReal)SwingIndex;
+
+				for (int32 SwingAxisIndex = 0; SwingAxisIndex < UE_ARRAY_COUNT(SwingAxes); ++SwingAxisIndex)
+				{
+					FVec3 SwingAxis = SwingAxes[SwingAxisIndex];
+					TestSwingTwistOrder({ SwingAxis, SwingAngleDeg, TwistAngleDeg });
+					TestSwingTwistDecomposition({ SwingAxis, SwingAngleDeg, TwistAngleDeg });
+				}
+			}
+		}
+	}
+
+}
\ No newline at end of file
diff --git a/Engine/Source/Runtime/Experimental/Chaos/Private/Chaos/PBDJointConstraintUtilities.cpp b/Engine/Source/Runtime/Experimental/Chaos/Private/Chaos/PBDJointConstraintUtilities.cpp
--- a/Engine/Source/Runtime/Experimental/Chaos/Private/Chaos/PBDJointConstraintUtilities.cpp
+++ b/Engine/Source/Runtime/Experimental/Chaos/Private/Chaos/PBDJointConstraintUtilities.cpp
@@ -14,8 +14,10 @@ namespace Chaos
 		if (S > KINDA_SMALL_NUMBER)
 		{
 			FReal SInv = 1.0f / S;
+			//OutTwist = FRotation3::FromElements(Q.X * SInv, 0, 0, Q.W * SInv);
+			//OutSwing = FRotation3::FromElements(0, (Q.W * Q.Y + Q.X * Q.Z) * SInv, (Q.W * Q.Z - Q.X * Q.Y) * SInv, S);
 			OutTwist = FRotation3::FromElements(Q.X * SInv, 0, 0, Q.W * SInv);
-			OutSwing = FRotation3::FromElements(0, (Q.W * Q.Y + Q.X * Q.Z) * SInv, (Q.W * Q.Z - Q.X * Q.Y) * SInv, S);
+			OutSwing = FRotation3::FromElements(0, (Q.W * Q.Y - Q.X * Q.Z) * SInv, (Q.W * Q.Z + Q.X * Q.Y) * SInv, S);
 		}
 		else
 		{
@@ -24,13 +26,132 @@ namespace Chaos
 		}
 	}
 
-
 	void FPBDJointUtilities::DecomposeSwingTwistLocal(const FRotation3& R0, const FRotation3& R1, FRotation3& R01Swing, FRotation3& R01Twist)
 	{
 		const FRotation3 R01 = R0.Inverse() * R1;
 		ToSwingTwistX(R01, R01Swing, R01Twist);
 	}
 
+	FReal FPBDJointUtilities::GetTwistAngle(const FRotation3& InTwist)
+	{
+		FRotation3 Twist = InTwist.GetNormalized();
+		ensure(FMath::Abs(Twist.W) <= 1.0f);
+		FReal Angle = Twist.GetAngle();
+		if (Angle > PI)
+		{
+			Angle = Angle - (FReal)2 * PI;
+		}
+		if (Twist.X < 0.0f)
+		{
+			Angle = -Angle;
+		}
+		return Angle;
+	}
+
+
+	void FPBDJointUtilities::GetTwistAxisAngle(
+		const FRotation3& R0,
+		const FRotation3& R1,
+		FVec3& Axis,
+		FReal& Angle)
+	{
+		FRotation3 R01Twist, R01Swing;
+		FPBDJointUtilities::DecomposeSwingTwistLocal(R0, R1, R01Swing, R01Twist);
+
+		Axis = R0 * FJointConstants::TwistAxis();
+		Angle = FPBDJointUtilities::GetTwistAngle(R01Twist);
+	}
+
+
+	void FPBDJointUtilities::GetConeAxisAngleLocal(
+		const FRotation3& R0,
+		const FRotation3& R1,
+		const FReal AngleTolerance,
+		FVec3& AxisLocal,
+		FReal& Angle)
+	{
+		// Decompose rotation of body 1 relative to body 0 into swing and twist rotations, assuming twist is X axis
+		FRotation3 R01Twist, R01Swing;
+		FPBDJointUtilities::DecomposeSwingTwistLocal(R0, R1, R01Swing, R01Twist);
+
+		R01Swing.ToAxisAndAngleSafe(AxisLocal, Angle, FJointConstants::Swing1Axis(), AngleTolerance);
+		if (Angle > PI)
+		{
+			Angle = Angle - (FReal)2 * PI;
+		}
+	}
+
+
+	void FPBDJointUtilities::GetSwingAxisAngle(
+		const FRotation3& R0,
+		const FRotation3& R1,
+		const FReal AngleTolerance,
+		const EJointAngularConstraintIndex SwingConstraintIndex,
+		const EJointAngularAxisIndex SwingAxisIndex,
+		FVec3& Axis,
+		FReal& Angle)
+	{
+#if 0
+		const FMatrix33 RM0 = R0.ToMatrix();
+		const FMatrix33 RM1 = R1.ToMatrix();
+		const FVec3& Twist0 = RM0.GetAxis((int32)EJointAngularAxisIndex::Twist);
+		const FVec3& Twist1 = RM1.GetAxis((int32)EJointAngularAxisIndex::Twist);
+
+		Axis = RM0.GetAxis((int32)SwingAxisIndex);
+		FReal SinAngle = FVec3::DotProduct(FVec3::CrossProduct(Twist1, Twist0), Axis);
+		Angle = FMath::Asin(FMath::Clamp(SinAngle, (FReal)-1, (FReal)1));
+#elif 0
+		// Decompose rotation of body 1 relative to body 0 into swing and twist rotations, assuming twist is X axis
+		FRotation3 R01Twist, R01Swing;
+		FPBDJointUtilities::DecomposeSwingTwistLocal(R0, R1, R01Swing, R01Twist);
+		const FReal R01SwingYorZ = ((int32)SwingAxisIndex == 2) ? R01Swing.Z : R01Swing.Y;	// Can't index a quat :(
+		FReal Angle = 4.0f * FMath::Atan2(R01SwingYorZ, 1.0f + R01Swing.W);
+		if (R01Twist.X < 0)
+		{
+			if (Angle > 0)
+			{
+				Angle = (FReal)PI - Angle;
+			}
+			else
+			{
+				Angle = (FReal)-PI - Angle;
+			}
+		}
+		const FVec3& AxisLocal = (SwingConstraintIndex == EJointAngularConstraintIndex::Swing1) ? FJointConstants::Swing1Axis() : FJointConstants::Swing2Axis();
+		Axis = R0 * AxisLocal;
+#else
+		// Do something better than this!!
+		FRotation3 R01Twist, R01Swing;
+		FPBDJointUtilities::DecomposeSwingTwistLocal(R0, R1, R01Swing, R01Twist);
+		FVec3 TwistAxis01 = FJointConstants::TwistAxis();
+		const FVec3 TwistAxis = R0 * TwistAxis01;
+
+		const FRotation3 R1NoTwist = R1 * R01Twist.Inverse();
+		const FMatrix33 Axes0 = R0.ToMatrix();
+		const FMatrix33 Axes1 = R1NoTwist.ToMatrix();
+
+		Axis = Axes0.GetAxis((int32)SwingAxisIndex);
+		Angle = 0;
+
+		const EJointAngularAxisIndex OtherSwingAxis = (SwingAxisIndex == EJointAngularAxisIndex::Swing1) ? EJointAngularAxisIndex::Swing2 : EJointAngularAxisIndex::Swing1;
+		FVec3 SwingCross = FVec3::CrossProduct(Axes0.GetAxis((int32)OtherSwingAxis), Axes1.GetAxis((int32)OtherSwingAxis));
+		SwingCross = SwingCross - FVec3::DotProduct(TwistAxis, SwingCross) * TwistAxis;
+		const FReal SwingCrossLen = SwingCross.Size();
+		if (SwingCrossLen > KINDA_SMALL_NUMBER)
+		{
+			Axis = SwingCross / SwingCrossLen;
+
+			Angle = FMath::Asin(FMath::Clamp(SwingCrossLen, (FReal)0, (FReal)1));
+			const FReal SwingDot = FVec3::DotProduct(Axes0.GetAxis((int32)OtherSwingAxis), Axes1.GetAxis((int32)OtherSwingAxis));
+			if (SwingDot < (FReal)0)
+			{
+				Angle = (FReal)PI - Angle;
+			}
+		}
+#endif
+	}
+
+
 	// @todo(ccaulfield): separate linear soft and stiff
 	FReal FPBDJointUtilities::GetLinearStiffness(
 		const FPBDJointSolverSettings& SolverSettings,
diff --git a/Engine/Source/Runtime/Experimental/Chaos/Public/Chaos/PBDJointConstraintUtilities.h b/Engine/Source/Runtime/Experimental/Chaos/Public/Chaos/PBDJointConstraintUtilities.h
--- a/Engine/Source/Runtime/Experimental/Chaos/Public/Chaos/PBDJointConstraintUtilities.h
+++ b/Engine/Source/Runtime/Experimental/Chaos/Public/Chaos/PBDJointConstraintUtilities.h
@@ -12,7 +12,36 @@ namespace Chaos
 	class FPBDJointUtilities
 	{
 	public:
-		static void DecomposeSwingTwistLocal(const FRotation3& R0, const FRotation3& R1, FRotation3& R01Swing, FRotation3& R01Twist);
+		static void DecomposeSwingTwistLocal(
+			const FRotation3& R0, 
+			const FRotation3& R1, 
+			FRotation3& R01Swing, 
+			FRotation3& R01Twist);
+		
+		static FReal GetTwistAngle(
+			const FRotation3& InTwist);
+		
+		static void GetTwistAxisAngle(
+			const FRotation3& R0,
+			const FRotation3& R1,
+			FVec3& Axis,
+			FReal& Angle);
+
+		static void GetConeAxisAngleLocal(
+			const FRotation3& R0,
+			const FRotation3& R1,
+			const FReal AngleTolerance,
+			FVec3& AxisLocal,
+			FReal& Angle);
+
+		static void GetSwingAxisAngle(
+			const FRotation3& R0,
+			const FRotation3& R1,
+			const FReal AngleTolerance,
+			const EJointAngularConstraintIndex SwingConstraintIndex,
+			const EJointAngularAxisIndex SwingAxisIndex,
+			FVec3& Axis,
+			FReal& Angle);
 
 		/**
 		 * Increase the lower inertia components to ensure that the maximum ratio between any pair of elements is MaxRatio.
