Fix incorrect removal of particle from acceleration structure.

diff --git a/Engine/Source/Runtime/Experimental/Chaos/Private/Chaos/PBDRigidsEvolution.cpp b/Engine/Source/Runtime/Experimental/Chaos/Private/Chaos/PBDRigidsEvolution.cpp
--- a/Engine/Source/Runtime/Experimental/Chaos/Private/Chaos/PBDRigidsEvolution.cpp
+++ b/Engine/Source/Runtime/Experimental/Chaos/Private/Chaos/PBDRigidsEvolution.cpp
@@ -9,6 +9,9 @@
 
 namespace Chaos
 {
+	CHAOS_API int32 FixBadAccelerationStructureRemoval = 1;
+	FAutoConsoleVariableRef CVarFixBadAccelerationStructureRemoval(TEXT("p.FixBadAccelerationStructureRemoval"), FixBadAccelerationStructureRemoval, TEXT(""));
+
 	struct FAccelerationConfig
 	{
 		int32 BroadphaseType;
diff --git a/Engine/Source/Runtime/Experimental/Chaos/Public/Chaos/PBDRigidsEvolution.h b/Engine/Source/Runtime/Experimental/Chaos/Public/Chaos/PBDRigidsEvolution.h
--- a/Engine/Source/Runtime/Experimental/Chaos/Public/Chaos/PBDRigidsEvolution.h
+++ b/Engine/Source/Runtime/Experimental/Chaos/Public/Chaos/PBDRigidsEvolution.h
@@ -20,6 +20,9 @@ namespace ChaosTest { void TestPendingSpatialDataHandlePointerConflict(); }
 
 namespace Chaos
 {
+
+extern CHAOS_API int32 FixBadAccelerationStructureRemoval;
+
 template<class T, int d> class TPBDRigidsEvolutionGBF;
 class FChaosArchive;
 
@@ -587,9 +590,18 @@ protected:
 		// Delete data should match async, and update is not set, so this operation should be safe.
 		ExternalAccelerationQueue.FindOrAdd(Particle) = AsyncSpatialData;
 
-		//remove particle immediately for intermediate structure
 		InternalAccelerationQueue.Remove(Particle);
-		InternalAcceleration->RemoveElementFrom(AsyncSpatialData.DeleteAccelerationHandle, AsyncSpatialData.DeletedSpatialIdx);	//even though we remove immediately, future adds are still pending
+
+		//remove particle immediately for intermediate structure
+		if (FixBadAccelerationStructureRemoval)
+		{
+			InternalAcceleration->RemoveElementFrom(TAccelerationStructureHandle<T, d>(ParticleHandle), ParticleHandle.SpatialIdx());	//even though we remove immediately, future adds are still pending
+		}
+		else
+		{
+			// This is wrong! Be afraid!
+			InternalAcceleration->RemoveElementFrom(AsyncSpatialData.DeleteAccelerationHandle, AsyncSpatialData.DeletedSpatialIdx);	//even though we remove immediately, future adds are still pending
+		}
 	}
 
 	void UpdateConstraintPositionBasedState(T Dt)
