Default to no rotation if one cannot be found.

diff --git a/Engine/Source/Runtime/Experimental/Chaos/Private/Chaos/MassProperties.cpp b/Engine/Source/Runtime/Experimental/Chaos/Private/Chaos/MassProperties.cpp
--- a/Engine/Source/Runtime/Experimental/Chaos/Private/Chaos/MassProperties.cpp
+++ b/Engine/Source/Runtime/Experimental/Chaos/Private/Chaos/MassProperties.cpp
@@ -66,7 +66,10 @@ namespace Chaos
 		PMatrix<T, d, d> RotationMatrix = DoSwap ? PMatrix<T, d, d>(Eigenvector2, Eigenvector1, -Eigenvector0) : PMatrix<T, d, d>(Eigenvector0, Eigenvector1, Eigenvector2);
 		// NOTE: UE Matrix are column-major, so the PMatrix constructor is not setting eigenvectors - we need to transpose it to get a UE rotation matrix.
 		FinalRotation = TRotation<T,d>(RotationMatrix.GetTransposed());
-		check(FMath::IsNearlyEqual(FinalRotation.Size(), 1.0f, KINDA_SMALL_NUMBER));
+		if (!ensure(FMath::IsNearlyEqual(FinalRotation.Size(), 1.0f, KINDA_SMALL_NUMBER)))
+		{
+			return TRotation<T, d>::FromElements(TVector<T, d>(0), 1);
+		}
 		
 		return FinalRotation;
 	}
