Make it so async acceleration queue is keyed by UniqueIdx instead of pointer. Gives better perf and simplifies code

diff --git a/Engine/Source/Runtime/Experimental/Chaos/Private/Chaos/PBDRigidsEvolution.cpp b/Engine/Source/Runtime/Experimental/Chaos/Private/Chaos/PBDRigidsEvolution.cpp
--- a/Engine/Source/Runtime/Experimental/Chaos/Private/Chaos/PBDRigidsEvolution.cpp
+++ b/Engine/Source/Runtime/Experimental/Chaos/Private/Chaos/PBDRigidsEvolution.cpp
@@ -297,57 +297,46 @@ namespace Chaos
 	template<class FPBDRigidsEvolution, class FPBDCollisionConstraint, class T, int d>
 	void TPBDRigidsEvolutionBase<FPBDRigidsEvolution, FPBDCollisionConstraint, T, d>::ApplyParticlePendingData(const FPendingSpatialData& SpatialData, FAccelerationStructure& AccelerationStructure, bool bUpdateCache)
 	{
-		//Note: we collapsed several update delete events into one struct. If memory is reused this can lead to problems
-		//Luckily there are only 3 states we care about:
-		//While pending we updated an object several times, this collapses into one update
-		//While pending we may have updated an object, we may have also created and destroyed the object, but the final event is a delete, so just remove from acceleration structure
-		//While pending we destroyed, recreated using the same memory address, and then did an update. In this case we should remove first and then update as global bounds may have changed
-		//As long as we delete first and update second this will be respected
-
 		if (SpatialData.bDelete)
 		{
-			AccelerationStructure.RemoveElementFrom(SpatialData.DeleteAccelerationHandle, SpatialData.DeletedSpatialIdx);
-			TGeometryParticleHandle<T, d>* DeleteParticle = SpatialData.DeleteAccelerationHandle.GetGeometryParticleHandle_PhysicsThread();
+			AccelerationStructure.RemoveElementFrom(SpatialData.AccelerationHandle, SpatialData.SpatialIdx);
 
 			if (bUpdateCache)
 			{
-				if (uint32* InnerIdxPtr = ParticleToCacheInnerIdx.Find(DeleteParticle))
+				if (uint32* InnerIdxPtr = ParticleToCacheInnerIdx.Find(SpatialData.UniqueIdx()))
 				{
-					const auto SpatialIdx = SpatialData.DeletedSpatialIdx;
-					TSpatialAccelerationCache<T, d>& Cache = *SpatialAccelerationCache.FindChecked(SpatialIdx);	//can't delete from cache that doesn't exist
+					TSpatialAccelerationCache<T, d>& Cache = *SpatialAccelerationCache.FindChecked(SpatialData.SpatialIdx);	//can't delete from cache that doesn't exist
 					const uint32 CacheInnerIdx = *InnerIdxPtr;
 					if (CacheInnerIdx + 1 < Cache.Size())	//will get swapped with last element, so update it
 					{
-						TGeometryParticleHandle<T, d>* LastParticleInCache = Cache.Payload(Cache.Size() - 1).GetGeometryParticleHandle_PhysicsThread();
-						ParticleToCacheInnerIdx.FindChecked(LastParticleInCache) = CacheInnerIdx;
+						const FUniqueIdx LastParticleInCacheUniqueIdx = Cache.Payload(Cache.Size() - 1).UniqueIdx();
+						ParticleToCacheInnerIdx.FindChecked(LastParticleInCacheUniqueIdx) = CacheInnerIdx;
 					}
 
 					Cache.DestroyElement(CacheInnerIdx);
-					ParticleToCacheInnerIdx.Remove(DeleteParticle);
+					ParticleToCacheInnerIdx.Remove(SpatialData.UniqueIdx());
 				}
 			}
 		}
-
-		if (SpatialData.bUpdate)
+		else
 		{
+			TGeometryParticleHandle<T, d>* UpdateParticle = SpatialData.AccelerationHandle.GetGeometryParticleHandle_PhysicsThread();
 
-			TGeometryParticleHandle<T, d>* UpdateParticle = SpatialData.UpdateAccelerationHandle.GetGeometryParticleHandle_PhysicsThread();
-
-			AccelerationStructure.UpdateElementIn(UpdateParticle, UpdateParticle->WorldSpaceInflatedBounds(), UpdateParticle->HasBounds(), SpatialData.UpdatedSpatialIdx);
+			AccelerationStructure.UpdateElementIn(UpdateParticle, UpdateParticle->WorldSpaceInflatedBounds(), UpdateParticle->HasBounds(), SpatialData.SpatialIdx);
 
 			if (bUpdateCache)
 			{
-				TUniquePtr<TSpatialAccelerationCache<T, d>>* CachePtrPtr = SpatialAccelerationCache.Find(SpatialData.UpdatedSpatialIdx);
+				TUniquePtr<TSpatialAccelerationCache<T, d>>* CachePtrPtr = SpatialAccelerationCache.Find(SpatialData.SpatialIdx);
 				if (CachePtrPtr == nullptr)
 				{
-					CachePtrPtr = &SpatialAccelerationCache.Add(SpatialData.UpdatedSpatialIdx, TUniquePtr<TSpatialAccelerationCache<T, d>>(new TSpatialAccelerationCache<T, d>));
+					CachePtrPtr = &SpatialAccelerationCache.Add(SpatialData.SpatialIdx, TUniquePtr<TSpatialAccelerationCache<T, d>>(new TSpatialAccelerationCache<T, d>));
 				}
 
 				TSpatialAccelerationCache<T, d>& Cache = **CachePtrPtr;
 
 				//make sure in mapping
 				uint32 CacheInnerIdx;
-				if (uint32* CacheInnerIdxPtr = ParticleToCacheInnerIdx.Find(UpdateParticle))
+				if (uint32* CacheInnerIdxPtr = ParticleToCacheInnerIdx.Find(SpatialData.UniqueIdx()))
 				{
 					CacheInnerIdx = *CacheInnerIdxPtr;
 				}
@@ -355,13 +344,13 @@ namespace Chaos
 				{
 					CacheInnerIdx = Cache.Size();
 					Cache.AddElements(1);
-					ParticleToCacheInnerIdx.Add(UpdateParticle, CacheInnerIdx);
+					ParticleToCacheInnerIdx.Add(SpatialData.UniqueIdx(), CacheInnerIdx);
 				}
 
 				//update cache entry
 				Cache.HasBounds(CacheInnerIdx) = UpdateParticle->HasBounds();
 				Cache.Bounds(CacheInnerIdx) = UpdateParticle->WorldSpaceInflatedBounds();
-				Cache.Payload(CacheInnerIdx) = SpatialData.UpdateAccelerationHandle;
+				Cache.Payload(CacheInnerIdx) = SpatialData.AccelerationHandle;
 			}
 		}
 	}
@@ -369,49 +358,49 @@ namespace Chaos
 	template<class FPBDRigidsEvolution, class FPBDCollisionConstraint, class T, int d>
 	void TPBDRigidsEvolutionBase<FPBDRigidsEvolution, FPBDCollisionConstraint, T, d>::FlushInternalAccelerationQueue()
 	{
-		for (auto Itr : InternalAccelerationQueue)
+		for (const FPendingSpatialData& PendingData : InternalAccelerationQueue.PendingData)
 		{
-			ApplyParticlePendingData(Itr.Value, *InternalAcceleration, false);
+			ApplyParticlePendingData(PendingData, *InternalAcceleration, false);
 		}
-		InternalAccelerationQueue.Empty();
+		InternalAccelerationQueue.Reset();
 	}
 
 	template<class FPBDRigidsEvolution, class FPBDCollisionConstraint, class T, int d>
 	void TPBDRigidsEvolutionBase<FPBDRigidsEvolution, FPBDCollisionConstraint, T, d>::FlushAsyncAccelerationQueue()
 	{
-		for (auto Itr : AsyncAccelerationQueue)
+		for (const FPendingSpatialData& PendingData : AsyncAccelerationQueue.PendingData)
 		{
-			ApplyParticlePendingData(Itr.Value, *AsyncInternalAcceleration, true); //only the first queue needs to update the cached acceleration
+			ApplyParticlePendingData(PendingData, *AsyncInternalAcceleration, true); //only the first queue needs to update the cached acceleration
 			if (!bIsSingleThreaded)
 			{
-				ApplyParticlePendingData(Itr.Value, *AsyncExternalAcceleration, false);
+				ApplyParticlePendingData(PendingData, *AsyncExternalAcceleration, false);
 			}
 
 			// Async queue deletes complete, unique index free to be consumed by new particles.
-			if(Itr.Value.bDelete)
+			if(PendingData.bDelete)
 			{
 				//NOTE: This assumes that we are never creating a PT particle that is replicated to GT
 				//At the moment that is true, and it seems like we have enough mechanisms to avoid this direction
 				//If we want to support that, the UniqueIndex must be kept around until GT goes away
 				//This is hard to do, but would probably mean the ownership of the index is in the proxy
-				Particles.GetUniqueIndices().ReleaseIdx(Itr.Value.DeleteAccelerationHandle.UniqueIdx());
+				Particles.GetUniqueIndices().ReleaseIdx(PendingData.UniqueIdx());
 			}
 		}
-		AsyncAccelerationQueue.Empty();
+		AsyncAccelerationQueue.Reset();
 
 		//other queues are no longer needed since we've flushed all operations and now have a pristine structure
-		InternalAccelerationQueue.Empty();
-		ExternalAccelerationQueue.Empty();
+		InternalAccelerationQueue.Reset();
+		ExternalAccelerationQueue.Reset();
 	}
 
 	template<class FPBDRigidsEvolution, class FPBDCollisionConstraint, class T, int d>
 	void TPBDRigidsEvolutionBase<FPBDRigidsEvolution, FPBDCollisionConstraint, T, d>::FlushExternalAccelerationQueue(FAccelerationStructure& Acceleration)
 	{
-		for (auto Itr : ExternalAccelerationQueue)
+		for (const FPendingSpatialData& PendingData : ExternalAccelerationQueue.PendingData)
 		{
-			ApplyParticlePendingData(Itr.Value, Acceleration, false);
+			ApplyParticlePendingData(PendingData, Acceleration, false);
 		}
-		ExternalAccelerationQueue.Empty();
+		ExternalAccelerationQueue.Reset();
 	}
 
 	template<class FPBDRigidsEvolution, class FPBDCollisionConstraint, class T, int d>
@@ -443,16 +432,16 @@ namespace Chaos
 			AsyncInternalAcceleration = TUniquePtr<FAccelerationStructure>(SpatialCollectionFactory->CreateEmptyCollection());
 			if (!bIsSingleThreaded)
 			{
-			ScratchExternalAcceleration = TUniquePtr<FAccelerationStructure>(SpatialCollectionFactory->CreateEmptyCollection());
-			AsyncExternalAcceleration = TUniquePtr<FAccelerationStructure>(SpatialCollectionFactory->CreateEmptyCollection());
+				ScratchExternalAcceleration = TUniquePtr<FAccelerationStructure>(SpatialCollectionFactory->CreateEmptyCollection());
+				AsyncExternalAcceleration = TUniquePtr<FAccelerationStructure>(SpatialCollectionFactory->CreateEmptyCollection());
 			}
 			FlushInternalAccelerationQueue();
 
 			if (!bIsSingleThreaded)
 			{
-			FlushExternalAccelerationQueue(*ScratchExternalAcceleration);
-			bExternalReady = true;
-		}
+				FlushExternalAccelerationQueue(*ScratchExternalAcceleration);
+				bExternalReady = true;
+			}
 		}
 
 		if (bBlock)
@@ -527,10 +516,10 @@ namespace Chaos
 	{
 		WaitOnAccelerationStructure();
 
-		ParticleToCacheInnerIdx.Empty();
-		AsyncAccelerationQueue.Empty();
-		InternalAccelerationQueue.Empty();
-		ExternalAccelerationQueue.Empty();
+		ParticleToCacheInnerIdx.Reset();
+		AsyncAccelerationQueue.Reset();
+		InternalAccelerationQueue.Reset();
+		ExternalAccelerationQueue.Reset();
 
 		AccelerationStructureTaskComplete = nullptr;
 		const auto& NonDisabled = Particles.GetNonDisabledView();
@@ -545,6 +534,7 @@ namespace Chaos
 	template<class FPBDRigidsEvolution, class FPBDCollisionConstraint, class T, int d>
 	void TPBDRigidsEvolutionBase<FPBDRigidsEvolution, FPBDCollisionConstraint, T, d>::Serialize(FChaosArchive& Ar)
 	{
+		ensure(false);	//disabled transient data serialization. Need to rethink
 		int32 DefaultBroadphaseType = ConfigSettings.BroadphaseType;
 
 		Ar.UsingCustomVersion(FExternalPhysicsCustomObjectVersion::GUID);
@@ -588,12 +578,12 @@ namespace Chaos
 				SpatialCollectionFactory->Serialize(InternalAcceleration, Ar);
 			}
 
-			if (Ar.CustomVer(FExternalPhysicsCustomObjectVersion::GUID) < FExternalPhysicsCustomObjectVersion::FlushEvolutionInternalAccelerationQueue)
+			/*if (Ar.CustomVer(FExternalPhysicsCustomObjectVersion::GUID) < FExternalPhysicsCustomObjectVersion::FlushEvolutionInternalAccelerationQueue)
 			{
 				SerializePendingMap(Ar, InternalAccelerationQueue);
 				SerializePendingMap(Ar, AsyncAccelerationQueue);
 				SerializePendingMap(Ar, ExternalAccelerationQueue);
-			}
+			}*/
 
 			ScratchExternalAcceleration = AsUniqueSpatialAccelerationChecked<FAccelerationStructure>(InternalAcceleration->Copy());
 		}
diff --git a/Engine/Source/Runtime/Experimental/Chaos/Public/Chaos/ISpatialAcceleration.h b/Engine/Source/Runtime/Experimental/Chaos/Public/Chaos/ISpatialAcceleration.h
--- a/Engine/Source/Runtime/Experimental/Chaos/Public/Chaos/ISpatialAcceleration.h
+++ b/Engine/Source/Runtime/Experimental/Chaos/Public/Chaos/ISpatialAcceleration.h
@@ -196,6 +196,12 @@ FORCEINLINE FUniqueIdx GetUniqueIdx(const int32 Payload)
 	return FUniqueIdx(Payload);
 }
 
+FORCEINLINE FUniqueIdx GetUniqueIdx(const FUniqueIdx Payload)
+{
+	ensure(Payload.IsValid());
+	return Payload;
+}
+
 
 template <typename TPayloadType, typename T>
 struct TPayloadBoundsElement
@@ -437,6 +443,15 @@ public:
 		Add(Key) = Value;
 	}
 
+	void RemoveChecked(const TKey& Key)
+	{
+		const int32 Idx = GetUniqueIdx(Key).Idx;
+		Entries[Idx] = FEntry();	//Mark as free, also resets default values for next use of value
+#if CHAOS_SERIALIZE_OUT
+		KeysToSerializeOut[Idx] = TKey();
+#endif
+	}
+
 	void Remove(const TKey& Key)
 	{
 		const int32 Idx = GetUniqueIdx(Key).Idx;
diff --git a/Engine/Source/Runtime/Experimental/Chaos/Public/Chaos/PBDRigidsEvolution.h b/Engine/Source/Runtime/Experimental/Chaos/Public/Chaos/PBDRigidsEvolution.h
--- a/Engine/Source/Runtime/Experimental/Chaos/Public/Chaos/PBDRigidsEvolution.h
+++ b/Engine/Source/Runtime/Experimental/Chaos/Public/Chaos/PBDRigidsEvolution.h
@@ -318,20 +318,20 @@ class TPBDRigidsEvolutionBase
 	template <bool bPersistent>
 	FORCEINLINE_DEBUGGABLE void DirtyParticle(TGeometryParticleHandleImp<T, d, bPersistent>& Particle)
 	{
-		FPendingSpatialData& SpatialData = InternalAccelerationQueue.FindOrAdd(Particle.Handle());
-		SpatialData.UpdateAccelerationHandle = TAccelerationStructureHandle<T,d>(Particle);
-		SpatialData.bUpdate = true;
-		SpatialData.UpdatedSpatialIdx = Particle.SpatialIdx();
+		//TODO: distinguish between new particles and dirty particles
+		const FUniqueIdx UniqueIdx = Particle.UniqueIdx();
+		FPendingSpatialData& SpatialData = InternalAccelerationQueue.FindOrAdd(UniqueIdx);
+		ensure(SpatialData.bDelete == false);
+		SpatialData.AccelerationHandle = TAccelerationStructureHandle<T,d>(Particle);
+		SpatialData.SpatialIdx = Particle.SpatialIdx();
 
-		auto& AsyncSpatialData = AsyncAccelerationQueue.FindOrAdd(Particle.Handle());
-		AsyncSpatialData.UpdateAccelerationHandle = TAccelerationStructureHandle<T, d>(Particle);
-		AsyncSpatialData.bUpdate = true;
-		AsyncSpatialData.UpdatedSpatialIdx = Particle.SpatialIdx();
+		auto& AsyncSpatialData = AsyncAccelerationQueue.FindOrAdd(UniqueIdx);
+		ensure(SpatialData.bDelete == false);
+		AsyncSpatialData = SpatialData;
 
-		auto& ExternalSpatialData = ExternalAccelerationQueue.FindOrAdd(Particle.Handle());
-		ExternalSpatialData.UpdateAccelerationHandle = TAccelerationStructureHandle<T, d>(Particle);
-		ExternalSpatialData.bUpdate = true;
-		ExternalSpatialData.UpdatedSpatialIdx = Particle.SpatialIdx();
+		auto& ExternalSpatialData = ExternalAccelerationQueue.FindOrAdd(UniqueIdx);
+		ensure(SpatialData.bDelete == false);
+		ExternalSpatialData = SpatialData;
 	}
 
 	void DestroyParticle(TGeometryParticleHandle<T, d>* Particle)
@@ -569,39 +569,23 @@ protected:
 	template <bool bPersistent>
 	FORCEINLINE_DEBUGGABLE void RemoveParticleFromAccelerationStructure(TGeometryParticleHandleImp<T, d, bPersistent>& ParticleHandle)
 	{
-		auto Particle = ParticleHandle.Handle();
-		FPendingSpatialData& AsyncSpatialData = AsyncAccelerationQueue.FindOrAdd(Particle);
+		//TODO: at the moment we don't distinguish between the first time a particle is created and when it's just moved
+		// If we had this distinction we could simply remove the entry for the async queue
+		const FUniqueIdx UniqueIdx = ParticleHandle.UniqueIdx();
+		FPendingSpatialData& SpatialData = AsyncAccelerationQueue.FindOrAdd(UniqueIdx);
 
-		if (!AsyncSpatialData.bDelete)
-		{
-			//There are three cases to consider:
-			//Simple single delete happens, in that case just use the index you see (the first and only index)
-			//Delete followed by any number deletes and updates and finally an update, in that case we must delete the first particle and add the final (so use first index for delete)
-			//Delete followed by multiple updates and or deletes and a final delete. In that case we still only delete the first particle since the final delete is not really needed (add will be cancelled)
-			AsyncSpatialData.DeletedSpatialIdx = ParticleHandle.SpatialIdx();
-
-			// We cannot overwrite this handle if delete is already pending. (If delete is pending, that means this is the third case, cancelling update is sufficient),
-			AsyncSpatialData.DeleteAccelerationHandle = TAccelerationStructureHandle<T, d>(ParticleHandle);
-		}
-
-		AsyncSpatialData.bUpdate = false;
-		AsyncSpatialData.bDelete = true;
+		SpatialData.bDelete = true;
+		SpatialData.SpatialIdx = ParticleHandle.SpatialIdx();
+		SpatialData.AccelerationHandle = TAccelerationStructureHandle<T, d>(ParticleHandle);
 
-		// Delete data should match async, and update is not set, so this operation should be safe.
-		ExternalAccelerationQueue.FindOrAdd(Particle) = AsyncSpatialData;
+		ExternalAccelerationQueue.FindOrAdd(UniqueIdx) = SpatialData;
 
-		InternalAccelerationQueue.Remove(Particle);
+		//Internal acceleration has all moves pending, so cancel them all now
+		InternalAccelerationQueue.Remove(UniqueIdx);
 
 		//remove particle immediately for intermediate structure
-		if (FixBadAccelerationStructureRemoval)
-		{
-			InternalAcceleration->RemoveElementFrom(TAccelerationStructureHandle<T, d>(ParticleHandle), ParticleHandle.SpatialIdx());	//even though we remove immediately, future adds are still pending
-		}
-		else
-		{
-			// This is wrong! Be afraid!
-			InternalAcceleration->RemoveElementFrom(AsyncSpatialData.DeleteAccelerationHandle, AsyncSpatialData.DeletedSpatialIdx);	//even though we remove immediately, future adds are still pending
-		}
+		//TODO: if we distinguished between first time adds we could avoid this. We could also make the RemoveElementFrom more strict and ensure when it fails
+		InternalAcceleration->RemoveElementFrom(SpatialData.AccelerationHandle, SpatialData.SpatialIdx);
 	}
 
 	void UpdateConstraintPositionBasedState(T Dt)
@@ -692,21 +676,17 @@ protected:
 	/** Used for updating intermediate spatial structures when they are finished */
 	struct FPendingSpatialData
 	{
-		TAccelerationStructureHandle<T, d> UpdateAccelerationHandle;
-		TAccelerationStructureHandle<T, d> DeleteAccelerationHandle;
-		FSpatialAccelerationIdx UpdatedSpatialIdx;
-		FSpatialAccelerationIdx DeletedSpatialIdx;	//need both updated and deleted in case memory is reused but a different idx is neede
-		bool bUpdate;
+		TAccelerationStructureHandle<T, d> AccelerationHandle;
+		FSpatialAccelerationIdx SpatialIdx;
 		bool bDelete;
 
 		FPendingSpatialData()
-			: bUpdate(false)
-			, bDelete(false)
+		: bDelete(false)
 		{}
 
 		void Serialize(FChaosArchive& Ar)
 		{
-			Ar.UsingCustomVersion(FExternalPhysicsCustomObjectVersion::GUID);
+			/*Ar.UsingCustomVersion(FExternalPhysicsCustomObjectVersion::GUID);
 			if (Ar.CustomVer(FExternalPhysicsCustomObjectVersion::GUID) >= FExternalPhysicsCustomObjectVersion::SerializeHashResult)
 			{
 				Ar << UpdateAccelerationHandle;
@@ -722,20 +702,73 @@ protected:
 			Ar << bDelete;
 
 			Ar << UpdatedSpatialIdx;
-			Ar << DeletedSpatialIdx;
+			Ar << DeletedSpatialIdx;*/
+			ensure(false);	//Serialization of transient data like this is currently broken. Need to reevaluate
+		}
+
+		FUniqueIdx UniqueIdx() const
+		{
+			return AccelerationHandle.UniqueIdx();
+		}
+	};
+
+	struct FPendingSpatialDataQueue
+	{
+		TArray<FPendingSpatialData> PendingData;
+		TArrayAsMap<FUniqueIdx, int32> ParticleToPendingData;
+
+		void Reset()
+		{
+			PendingData.Reset();
+			ParticleToPendingData.Reset();
+		}
+
+		int32 Num() const
+		{
+			return PendingData.Num();
+		}
+
+		FPendingSpatialData& FindOrAdd(const FUniqueIdx UniqueIdx)
+		{
+			if (int32* Existing = ParticleToPendingData.Find(UniqueIdx))
+			{
+				return PendingData[*Existing];
+			}
+			else
+			{
+				const int32 NewIdx = PendingData.AddDefaulted(1);
+				ParticleToPendingData.Add(UniqueIdx, NewIdx);
+				return PendingData[NewIdx];
+			}
+		}
+
+		void Remove(const FUniqueIdx UniqueIdx)
+		{
+			if (int32* Existing = ParticleToPendingData.Find(UniqueIdx))
+			{
+				const int32 SlotIdx = *Existing;
+				if (SlotIdx + 1 < PendingData.Num())
+				{
+					const FUniqueIdx LastElemUniqueIdx = PendingData.Last().UniqueIdx();
+					ParticleToPendingData.FindChecked(LastElemUniqueIdx) = SlotIdx;	//We're going to swap elements so the last element is now in the position of the element we removed
+				}
+
+				PendingData.RemoveAtSwap(SlotIdx);
+				ParticleToPendingData.RemoveChecked(UniqueIdx);
+			}
 		}
 	};
 
 	/** Pending operations for the internal acceleration structure */
-	TMap<TGeometryParticleHandle<T, d>*, FPendingSpatialData> InternalAccelerationQueue;
+	FPendingSpatialDataQueue InternalAccelerationQueue;
 
 	/** Pending operations for the acceleration structures being rebuilt asynchronously */
-	TMap<TGeometryParticleHandle<T, d>*, FPendingSpatialData> AsyncAccelerationQueue;
+	FPendingSpatialDataQueue AsyncAccelerationQueue;
 
 	/** Pending operations for the external acceleration structure*/
-	TMap<TGeometryParticleHandle<T, d>*, FPendingSpatialData> ExternalAccelerationQueue;
+	FPendingSpatialDataQueue ExternalAccelerationQueue;
 
-	void SerializePendingMap(FChaosArchive& Ar, TMap<TGeometryParticleHandle<T, d>*, FPendingSpatialData>& Map)
+	/*void SerializePendingMap(FChaosArchive& Ar, TMap<TGeometryParticleHandle<T, d>*, FPendingSpatialData>& Map)
 	{
 		TArray<TGeometryParticleHandle<T, d>*> Keys;
 		if (!Ar.IsLoading())
@@ -748,10 +781,11 @@ protected:
 			FPendingSpatialData& PendingData = Map.FindOrAdd(Key);
 			PendingData.Serialize(Ar);
 		}
-	}
+		//TODO: fix serialization
+	}*/
 
 	/** Used for async acceleration rebuild */
-	TMap<TGeometryParticleHandle<T, d>*, uint32> ParticleToCacheInnerIdx;
+	TArrayAsMap<FUniqueIdx, uint32> ParticleToCacheInnerIdx;
 
 	TMap<FSpatialAccelerationIdx, TUniquePtr<TSpatialAccelerationCache<T, d>>> SpatialAccelerationCache;
 
