Chaos - Joint solver using utilities rather than duped functions.

diff --git a/Engine/Source/Runtime/Experimental/Chaos/Private/Chaos/Joint/PBDJointSolverGaussSeidel.cpp b/Engine/Source/Runtime/Experimental/Chaos/Private/Chaos/Joint/PBDJointSolverGaussSeidel.cpp
--- a/Engine/Source/Runtime/Experimental/Chaos/Private/Chaos/Joint/PBDJointSolverGaussSeidel.cpp
+++ b/Engine/Source/Runtime/Experimental/Chaos/Private/Chaos/Joint/PBDJointSolverGaussSeidel.cpp
@@ -9,8 +9,7 @@
 #include "PBDJointSolverGaussSeidel.ispc.generated.h"
 #endif
 
-//#pragma optimize("", off)
-//PRAGMA_DISABLE_OPTIMIZATION_ACTUAL
+//PRAGMA_DISABLE_OPTIMIZATION
 
 #if !INTEL_ISPC
 const bool bChaos_Joint_ISPC_Enabled = false;
@@ -87,56 +86,6 @@ namespace Chaos
 	}
 
 
-	FReal GetTwistAngle(const FRotation3& InTwist)
-	{
-		FRotation3 Twist = InTwist.GetNormalized();
-		ensure(FMath::Abs(Twist.W) <= 1.0f);
-		FReal Angle = Twist.GetAngle();
-		if (Angle > PI)
-		{
-			Angle = Angle - (FReal)2 * PI;
-		}
-		if (Twist.X < 0.0f)
-		{
-			Angle = -Angle;
-		}
-		return Angle;
-	}
-
-
-	void GetTwistAxisAngle(
-		const FRotation3& R0,
-		const FRotation3& R1,
-		FVec3& Axis,
-		FReal& Angle)
-	{
-		FRotation3 R01Twist, R01Swing;
-		FPBDJointUtilities::DecomposeSwingTwistLocal(R0, R1, R01Swing, R01Twist);
-
-		Axis = R0 * FJointConstants::TwistAxis();
-		Angle = GetTwistAngle(R01Twist);
-	}
-
-
-	void GetConeAxisAngleLocal(
-		const FRotation3& R0,
-		const FRotation3& R1,
-		const FReal AngleTolerance,
-		FVec3& Axis,
-		FReal& Angle)
-	{
-		// Decompose rotation of body 1 relative to body 0 into swing and twist rotations, assuming twist is X axis
-		FRotation3 R01Twist, R01Swing;
-		FPBDJointUtilities::DecomposeSwingTwistLocal(R0, R1, R01Swing, R01Twist);
-
-		R01Swing.ToAxisAndAngleSafe(Axis, Angle, FJointConstants::Swing1Axis(), AngleTolerance);
-		if (Angle > PI)
-		{
-			Angle = Angle - (FReal)2 * PI;
-		}
-	}
-
-
 	FReal GetConeAngleError(
 		const FPBDJointSettings& JointSettings,
 		const FVec3& SwingAxisLocal,
@@ -173,75 +122,6 @@ namespace Chaos
 	}
 
 
-	void GetSwingAxisAngle(
-		const FRotation3& R0,
-		const FRotation3& R1,
-		const FReal AngleTolerance,
-		const EJointAngularConstraintIndex SwingConstraintIndex,
-		const EJointAngularAxisIndex SwingAxisIndex,
-		FVec3& Axis,
-		FReal& Angle)
-	{
-#if 0
-		const FMatrix33 RM0 = R0.ToMatrix();
-		const FMatrix33 RM1 = R1.ToMatrix();
-		const FVec3& Twist0 = RM0.GetAxis((int32)EJointAngularAxisIndex::Twist);
-		const FVec3& Twist1 = RM1.GetAxis((int32)EJointAngularAxisIndex::Twist);
-
-		Axis = RM0.GetAxis((int32)SwingAxisIndex);
-		FReal SinAngle = FVec3::DotProduct(FVec3::CrossProduct(Twist1, Twist0), Axis);
-		Angle = FMath::Asin(FMath::Clamp(SinAngle, (FReal)-1, (FReal)1));
-#elif 0
-		// Decompose rotation of body 1 relative to body 0 into swing and twist rotations, assuming twist is X axis
-		FRotation3 R01Twist, R01Swing;
-		FPBDJointUtilities::DecomposeSwingTwistLocal(R0, R1, R01Swing, R01Twist);
-		const FReal R01SwingYorZ = ((int32)SwingAxisIndex == 2) ? R01Swing.Z : R01Swing.Y;	// Can't index a quat :(
-		FReal Angle = 4.0f * FMath::Atan2(R01SwingYorZ, 1.0f + R01Swing.W);
-		if (R01Twist.X < 0)
-		{
-			if (Angle > 0)
-			{
-				Angle = (FReal)PI - Angle;
-			}
-			else
-			{
-				Angle = (FReal)-PI - Angle;
-			}
-		}
-		const FVec3& AxisLocal = (SwingConstraintIndex == EJointAngularConstraintIndex::Swing1) ? FJointConstants::Swing1Axis() : FJointConstants::Swing2Axis();
-		Axis = R0 * AxisLocal;
-#else
-		// Do something better than this!!
-		FRotation3 R01Twist, R01Swing;
-		FPBDJointUtilities::DecomposeSwingTwistLocal(R0, R1, R01Swing, R01Twist);
-		FVec3 TwistAxis01 = FJointConstants::TwistAxis();
-		const FVec3 TwistAxis = R0 * TwistAxis01;
-
-		const FRotation3 R1NoTwist = R1 * R01Twist.Inverse();
-		const FMatrix33 Axes0 = R0.ToMatrix();
-		const FMatrix33 Axes1 = R1NoTwist.ToMatrix();
-
-		Axis = Axes0.GetAxis((int32)SwingAxisIndex);
-		Angle = 0;
-
-		const EJointAngularAxisIndex OtherSwingAxis = (SwingAxisIndex == EJointAngularAxisIndex::Swing1) ? EJointAngularAxisIndex::Swing2 : EJointAngularAxisIndex::Swing1;
-		FVec3 SwingCross = FVec3::CrossProduct(Axes0.GetAxis((int32)OtherSwingAxis), Axes1.GetAxis((int32)OtherSwingAxis));
-		SwingCross = SwingCross - FVec3::DotProduct(TwistAxis, SwingCross) * TwistAxis;
-		const FReal SwingCrossLen = SwingCross.Size();
-		if (SwingCrossLen > KINDA_SMALL_NUMBER)
-		{
-			Axis = SwingCross / SwingCrossLen;
-
-			Angle = FMath::Asin(FMath::Clamp(SwingCrossLen, (FReal)0, (FReal)1));
-			const FReal SwingDot = FVec3::DotProduct(Axes0.GetAxis((int32)OtherSwingAxis), Axes1.GetAxis((int32)OtherSwingAxis));
-			if (SwingDot < (FReal)0)
-			{
-				Angle = (FReal)PI - Angle;
-			}
-		}
-#endif
-	}
-
 	//
 	//
 	//////////////////////////////////////////////////////////////////////////
@@ -1133,7 +1013,7 @@ namespace Chaos
 	{
 		FVec3 TwistAxis;
 		FReal TwistAngle;
-		GetTwistAxisAngle(Rs[0], Rs[1], TwistAxis, TwistAngle);
+		FPBDJointUtilities::GetTwistAxisAngle(Rs[0], Rs[1], TwistAxis, TwistAngle);
 
 		// Calculate the twist correction to apply to each body
 		FReal DTwistAngle = 0;
@@ -1165,7 +1045,7 @@ namespace Chaos
 	{
 		FVec3 TwistAxis;
 		FReal TwistAngle;
-		GetTwistAxisAngle(Rs[0], Rs[1], TwistAxis, TwistAngle);
+		FPBDJointUtilities::GetTwistAxisAngle(Rs[0], Rs[1], TwistAxis, TwistAngle);
 
 		// Calculate the twist correction to apply to each body
 		FReal DTwistAngle = 0;
@@ -1199,7 +1079,7 @@ namespace Chaos
 	{
 		FVec3 TwistAxis;
 		FReal TwistAngle;
-		GetTwistAxisAngle(Rs[0], Rs[1], TwistAxis, TwistAngle);
+		FPBDJointUtilities::GetTwistAxisAngle(Rs[0], Rs[1], TwistAxis, TwistAngle);
 
 		const FReal TwistAngleTarget = JointSettings.AngularDriveTargetAngles[(int32)EJointAngularConstraintIndex::Twist];
 		const FReal DTwistAngle = TwistAngle - TwistAngleTarget;
@@ -1234,7 +1114,7 @@ namespace Chaos
 		// Calculate swing angle and axis
 		FReal SwingAngle;
 		FVec3 SwingAxisLocal;
-		GetConeAxisAngleLocal(Rs[0], Rs[1], SolverSettings.SwingTwistAngleTolerance, SwingAxisLocal, SwingAngle);
+		FPBDJointUtilities::GetConeAxisAngleLocal(Rs[0], Rs[1], SolverSettings.SwingTwistAngleTolerance, SwingAxisLocal, SwingAngle);
 		const FVec3 SwingAxis = Rs[0] * SwingAxisLocal;
 
 		// Calculate swing angle error
@@ -1259,7 +1139,7 @@ namespace Chaos
 		// Calculate swing angle and axis
 		FReal SwingAngle;
 		FVec3 SwingAxisLocal;
-		GetConeAxisAngleLocal(Rs[0], Rs[1], SolverSettings.SwingTwistAngleTolerance, SwingAxisLocal, SwingAngle);
+		FPBDJointUtilities::GetConeAxisAngleLocal(Rs[0], Rs[1], SolverSettings.SwingTwistAngleTolerance, SwingAxisLocal, SwingAngle);
 		const FVec3 SwingAxis = Rs[0] * SwingAxisLocal;
 
 		// Calculate swing angle error
@@ -1286,7 +1166,7 @@ namespace Chaos
 		// Calculate swing angle and axis
 		FReal SwingAngle;
 		FVec3 SwingAxisLocal;
-		GetConeAxisAngleLocal(Rs[0], Rs[1], SolverSettings.SwingTwistAngleTolerance, SwingAxisLocal, SwingAngle);
+		FPBDJointUtilities::GetConeAxisAngleLocal(Rs[0], Rs[1], SolverSettings.SwingTwistAngleTolerance, SwingAxisLocal, SwingAngle);
 		const FVec3 SwingAxis = Rs[0] * SwingAxisLocal;
 
 		// Circular swing target (max of Swing1, Swing2 targets)
@@ -1316,7 +1196,7 @@ namespace Chaos
 		// Calculate swing angle and axis
 		FReal SwingAngle;
 		FVec3 SwingAxisLocal;
-		GetConeAxisAngleLocal(Rs[0], Rs[1], SolverSettings.SwingTwistAngleTolerance, SwingAxisLocal, SwingAngle);
+		FPBDJointUtilities::GetConeAxisAngleLocal(Rs[0], Rs[1], SolverSettings.SwingTwistAngleTolerance, SwingAxisLocal, SwingAngle);
 		const FVec3 SwingAxis = Rs[0] * SwingAxisLocal;
 
 		// Calculate swing angle error
@@ -1372,7 +1252,7 @@ namespace Chaos
 	{
 		FVec3 SwingAxis;
 		FReal SwingAngle;
-		GetSwingAxisAngle(Rs[0], Rs[1], SolverSettings.SwingTwistAngleTolerance, SwingConstraintIndex, SwingAxisIndex, SwingAxis, SwingAngle);
+		FPBDJointUtilities::GetSwingAxisAngle(Rs[0], Rs[1], SolverSettings.SwingTwistAngleTolerance, SwingConstraintIndex, SwingAxisIndex, SwingAxis, SwingAngle);
 
 		// Calculate swing error we need to correct
 		FReal DSwingAngle = 0;
@@ -1406,7 +1286,7 @@ namespace Chaos
 	{
 		FVec3 SwingAxis;
 		FReal SwingAngle;
-		GetSwingAxisAngle(Rs[0], Rs[1], SolverSettings.SwingTwistAngleTolerance, SwingConstraintIndex, SwingAxisIndex, SwingAxis, SwingAngle);
+		FPBDJointUtilities::GetSwingAxisAngle(Rs[0], Rs[1], SolverSettings.SwingTwistAngleTolerance, SwingConstraintIndex, SwingAxisIndex, SwingAxis, SwingAngle);
 
 		const FReal SwingAngleMax = JointSettings.AngularLimits[(int32)SwingConstraintIndex];
 
@@ -1443,7 +1323,7 @@ namespace Chaos
 	{
 		FVec3 SwingAxis;
 		FReal SwingAngle;
-		GetSwingAxisAngle(Rs[0], Rs[1], SolverSettings.SwingTwistAngleTolerance, SwingConstraintIndex, SwingAxisIndex, SwingAxis, SwingAngle);
+		FPBDJointUtilities::GetSwingAxisAngle(Rs[0], Rs[1], SolverSettings.SwingTwistAngleTolerance, SwingConstraintIndex, SwingAxisIndex, SwingAxis, SwingAngle);
 
 		const FReal SwingAngleTarget = JointSettings.AngularDriveTargetAngles[(int32)SwingConstraintIndex];
 		const FReal DSwingAngle = SwingAngle - SwingAngleTarget;
