Chaos - Get spring constraints working with MinEvolution (WIP)
- MinEvolution test (WIP)

diff --git a/Engine/Source/Programs/HeadlessChaos/Private/HeadlessChaosTestConstraintHandle.cpp b/Engine/Source/Programs/HeadlessChaos/Private/HeadlessChaosTestConstraintHandle.cpp
--- a/Engine/Source/Programs/HeadlessChaos/Private/HeadlessChaosTestConstraintHandle.cpp
+++ b/Engine/Source/Programs/HeadlessChaos/Private/HeadlessChaosTestConstraintHandle.cpp
@@ -152,13 +152,13 @@ namespace ChaosTest {
 		TEvolution Evolution(Particles);
 		TArray<TPBDRigidParticleHandle<FReal, 3>*> ParticleHandles = Evolution.CreateDynamicParticles(5);
 
-		TPBDRigidSpringConstraints<FReal, 3> Constraints;
-		TArray<TPBDRigidSpringConstraintHandle<FReal, 3>*> ConstraintsHandles =
+		FPBDRigidSpringConstraints Constraints;
+		TArray<FPBDRigidSpringConstraintHandle*> ConstraintsHandles =
 		{
-			Constraints.AddConstraint({ ParticleHandles[0], ParticleHandles[1] }, { { 0, 0, 0 }, { 0, 0, 0 } }),
-			Constraints.AddConstraint({ ParticleHandles[1], ParticleHandles[2] }, { { 0, 0, 0 }, { 0, 0, 0 } }),
-			Constraints.AddConstraint({ ParticleHandles[2], ParticleHandles[3] }, { { 0, 0, 0 }, { 0, 0, 0 } }),
-			Constraints.AddConstraint({ ParticleHandles[3], ParticleHandles[4] }, { { 0, 0, 0 }, { 0, 0, 0 } }),
+			Constraints.AddConstraint({ ParticleHandles[0], ParticleHandles[1] }, { { 0, 0, 0 }, { 0, 0, 0 } }, 1.0f, 0.0f, 0.0f),
+			Constraints.AddConstraint({ ParticleHandles[1], ParticleHandles[2] }, { { 0, 0, 0 }, { 0, 0, 0 } }, 1.0f, 0.0f, 0.0f),
+			Constraints.AddConstraint({ ParticleHandles[2], ParticleHandles[3] }, { { 0, 0, 0 }, { 0, 0, 0 } }, 1.0f, 0.0f, 0.0f),
+			Constraints.AddConstraint({ ParticleHandles[3], ParticleHandles[4] }, { { 0, 0, 0 }, { 0, 0, 0 } }, 1.0f, 0.0f, 0.0f),
 		};
 
 		CheckConstraintHandles(Constraints, ParticleHandles, ConstraintsHandles);
diff --git a/Engine/Source/Programs/HeadlessChaos/Private/HeadlessChaosTestJoints.cpp b/Engine/Source/Programs/HeadlessChaos/Private/HeadlessChaosTestJoints.cpp
--- a/Engine/Source/Programs/HeadlessChaos/Private/HeadlessChaosTestJoints.cpp
+++ b/Engine/Source/Programs/HeadlessChaos/Private/HeadlessChaosTestJoints.cpp
@@ -478,15 +478,15 @@ namespace ChaosTest {
 		Box2->P() = Box2->X();
 
 		TPBDRigidsEvolutionGBF<FReal, 3> Evolution(Particles);
-		TArray<TVector<TGeometryParticleHandle<FReal, 3>*,2>> Constraints = { TVector<TGeometryParticleHandle<FReal,3>*, 2>(StaticBox, Box2) };
-		TArray<FVec3> Points0 = { FVec3((FReal)100, (FReal)0, (FReal)1000) };
-		TArray<FVec3> Points1 = { FVec3((FReal)400, (FReal)0, (FReal)1000) };
+		TVector<TGeometryParticleHandle<FReal, 3>*,2> ConstrainedParticles = TVector<TGeometryParticleHandle<FReal,3>*, 2>(StaticBox, Box2);
+		TVector<FVec3, 2> Points = { FVec3((FReal)100, (FReal)0, (FReal)1000), FVec3((FReal)400, (FReal)0, (FReal)1000) };
 
 		Evolution.SetPhysicsMaterial(StaticBox, MakeSerializable(PhysicalMaterial));
 		Evolution.SetPhysicsMaterial(Box2, MakeSerializable(PhysicalMaterial));
 
-		auto JointConstraints = Chaos::TPBDRigidSpringConstraints<FReal, 3>(Points0, Points1, MoveTemp(Constraints), 1.f);
-		auto JointRule = Chaos::TPBDConstraintIslandRule<Chaos::TPBDRigidSpringConstraints<FReal, 3>>(JointConstraints);
+		auto JointConstraints = FPBDRigidSpringConstraints();
+		JointConstraints.AddConstraint(ConstrainedParticles, Points, 1.0f, 0.0f, (Points[0] - Points[1]).Size());
+		auto JointRule = Chaos::TPBDConstraintIslandRule<Chaos::FPBDRigidSpringConstraints>(JointConstraints);
 		Evolution.AddConstraintRule(&JointRule);
 
 		const FReal Dt = 0.01f;
@@ -494,7 +494,7 @@ namespace ChaosTest {
 		{
 			Evolution.AdvanceOneTimeStep(Dt);
 			Evolution.EndFrame(Dt);
-			EXPECT_LT(FMath::Abs((Box2->R().RotateVector(FVec3((FReal)-100, (FReal)0, (FReal)0)) + Box2->X() - Points0[0]).Size() - 300.f), 0.1);
+			EXPECT_LT(FMath::Abs((Box2->R().RotateVector(FVec3((FReal)-100, (FReal)0, (FReal)0)) + Box2->X() - Points[0]).Size() - 300.f), 0.1);
 		}
 	}
 
diff --git a/Engine/Source/Programs/HeadlessChaos/Private/HeadlessChaosTestMinEvolution.cpp b/Engine/Source/Programs/HeadlessChaos/Private/HeadlessChaosTestMinEvolution.cpp
new file mode 100644
--- /dev/null
+++ b/Engine/Source/Programs/HeadlessChaos/Private/HeadlessChaosTestMinEvolution.cpp
@@ -0,0 +1,102 @@
+// Copyright Epic Games, Inc. All Rights Reserved.
+
+#include "HeadlessChaos.h"
+
+#include "Chaos/Collision/CollisionDetector.h"
+#include "Chaos/Collision/CollisionReceiver.h"
+#include "Chaos/Collision/NarrowPhase.h"
+#include "Chaos/Collision/ParticlePairBroadPhase.h"
+#include "Chaos/Evolution/PBDMinEvolution.h"
+#include "Chaos/ParticleHandle.h"
+#include "Chaos/PBDCollisionConstraints.h"
+#include "Chaos/PBDConstraintRule.h"
+#include "Chaos/PBDRigidParticles.h"
+#include "Chaos/PBDRigidSpringConstraints.h"
+
+#include "ChaosLog.h"
+
+namespace ChaosTest
+{
+	using namespace Chaos;
+
+	// Check thar spring constraints work with MinEvolution
+	GTEST_TEST(MinEvolutionTests, TestSpringConstraints)
+	{
+		// @todo(ccaulfield): remove template parameters on collisions and other constraints
+		using FCollisionConstraints = TPBDCollisionConstraints<FReal, 3>;
+		using FCollisionDetector = TCollisionDetector<FParticlePairBroadPhase, FNarrowPhase, FSyncCollisionReceiver, FCollisionConstraints>;
+		using FRigidParticleSOAs = TPBDRigidsSOAs<FReal, 3>;
+		using FParticleHandle = TPBDRigidParticleHandle<FReal, 3>;
+		using FParticlePair = TVector<TGeometryParticleHandle<FReal, 3>*, 2>;
+
+		// Particles
+		FRigidParticleSOAs ParticlesContainer;
+
+		// @todo(ccaulfield): we shouldn't require collisions to use an evolution...
+		// Stuff needed for collisions
+		TArray<FParticlePair> ActivePotentiallyCollidingPairs;
+		TArrayCollectionArray<bool> CollidedParticles;
+		TArrayCollectionArray<Chaos::TSerializablePtr<Chaos::FChaosPhysicsMaterial>> ParticleMaterials;
+		TArrayCollectionArray<TUniquePtr<Chaos::FChaosPhysicsMaterial>> PerParticleMaterials;
+		FCollisionConstraints Collisions(ParticlesContainer, CollidedParticles, ParticleMaterials);
+		FParticlePairBroadPhase BroadPhase(ActivePotentiallyCollidingPairs, 0);
+		FCollisionDetector CollisionDetector(BroadPhase, Collisions);
+		TSimpleConstraintRule<FCollisionConstraints> CollisionsRule(1, Collisions);
+		// End collisions stuff
+
+		// Springs
+		FPBDRigidSpringConstraints Springs;
+		TSimpleConstraintRule<FPBDRigidSpringConstraints> SpringsRule(0, Springs);
+
+		// Evolution
+		// @todo(ccaulfield): this should start with some reasonable default iterations
+		FPBDMinEvolution Evolution(ParticlesContainer, CollisionDetector, 0);
+		Evolution.SetNumIterations(1);
+		Evolution.SetNumPushOutIterations(0);
+
+		Evolution.AddConstraintRule(&SpringsRule);
+		Evolution.SetGravity(FVec3(0));
+
+		FReal Dt = 1.0f / 30.0f;
+
+		// Add a couple dynamic particles connected by a spring
+		TArray<FParticleHandle*> Particles = ParticlesContainer.CreateDynamicParticles(2);
+
+		// Set up Particles
+		// @todo(ccaulfield) this needs to be easier
+		Particles[0]->X() = FVec3(-50, 0, 0);
+		Particles[0]->M() = 1.0f;
+		Particles[0]->I() = FMatrix33(100.0f, 100.0f, 100.0f);
+		Particles[0]->InvM() = 1.0f;
+		Particles[0]->InvI() = FMatrix33(1.0f / 100.0f, 1.0f / 100.0f, 1.0f / 100.0f);
+
+		Particles[1]->X() = FVec3(50, 0, 0);
+		Particles[1]->M() = 1.0f;
+		Particles[1]->I() = FMatrix33(100.0f, 100.0f, 100.0f);
+		Particles[1]->InvM() = 1.0f;
+		Particles[1]->InvI() = FMatrix33(1.0f / 100.0f, 1.0f / 100.0f, 1.0f / 100.0f);
+
+		// Spring connectors at particle centres
+		TArray<FVec3> Locations =
+		{
+			FVec3(-50, 0, 0),
+			FVec3(50, 0, 0)
+		};
+
+		// Create springs
+		FPBDRigidSpringConstraintHandle* Spring = Springs.AddConstraint({ Particles[0], Particles[1] }, { Locations[0], Locations[1] }, 0.1f, 0.0f, 60.0f);
+
+		for (int32 TimeIndex = 0; TimeIndex < 1000; ++TimeIndex)
+		{
+			//UE_LOG(LogChaos, Warning, TEXT("%d: %f %f %f - %f %f %f"), TimeIndex, Particles[0]->X().X, Particles[0]->X().Y, Particles[0]->X().Z, Particles[1]->X().X, Particles[1]->X().Y, Particles[1]->X().Z);
+			Evolution.Advance(Dt, 1);
+		}
+
+		// Particles should be separated by the spring's rest length
+		FVec3 P0 = Particles[0]->X();
+		FVec3 P1 = Particles[1]->X();
+		FReal Distance01 = (P0 - P1).Size();
+		EXPECT_NEAR(Distance01, Spring->GetRestLength(), 0.1f);
+	}
+
+}
\ No newline at end of file
diff --git a/Engine/Source/Runtime/Experimental/Chaos/Private/Chaos/PBDConstraintRule.cpp b/Engine/Source/Runtime/Experimental/Chaos/Private/Chaos/PBDConstraintRule.cpp
--- a/Engine/Source/Runtime/Experimental/Chaos/Private/Chaos/PBDConstraintRule.cpp
+++ b/Engine/Source/Runtime/Experimental/Chaos/Private/Chaos/PBDConstraintRule.cpp
@@ -54,16 +54,17 @@ namespace Chaos
 
 	template class TSimpleConstraintRule<TPBDCollisionConstraints<float, 3>>;
 	template class TSimpleConstraintRule<FPBDJointConstraints>;
+	template class TSimpleConstraintRule<FPBDRigidSpringConstraints>;
 
 	template class TPBDConstraintGraphRuleImpl<TPBDCollisionConstraints<float, 3>>;
 	template class TPBDConstraintGraphRuleImpl<FPBDJointConstraints>;
 	template class TPBDConstraintGraphRuleImpl<TPBDPositionConstraints<float, 3>>;
 	template class TPBDConstraintGraphRuleImpl<TPBDRigidDynamicSpringConstraints<float, 3>>;
-	template class TPBDConstraintGraphRuleImpl<TPBDRigidSpringConstraints<float, 3>>;
+	template class TPBDConstraintGraphRuleImpl<FPBDRigidSpringConstraints>;
 
 	template class TPBDConstraintColorRule<TPBDCollisionConstraints<float, 3>>;
 	template class TPBDConstraintIslandRule<FPBDJointConstraints>;
 	template class TPBDConstraintIslandRule<TPBDPositionConstraints<float, 3>>;
 	template class TPBDConstraintIslandRule<TPBDRigidDynamicSpringConstraints<float, 3>>;
-	template class TPBDConstraintIslandRule<TPBDRigidSpringConstraints<float, 3>>;
+	template class TPBDConstraintIslandRule<FPBDRigidSpringConstraints>;
 }
\ No newline at end of file
diff --git a/Engine/Source/Runtime/Experimental/Chaos/Private/Chaos/PBDRigidSpringConstraints.cpp b/Engine/Source/Runtime/Experimental/Chaos/Private/Chaos/PBDRigidSpringConstraints.cpp
--- a/Engine/Source/Runtime/Experimental/Chaos/Private/Chaos/PBDRigidSpringConstraints.cpp
+++ b/Engine/Source/Runtime/Experimental/Chaos/Private/Chaos/PBDRigidSpringConstraints.cpp
@@ -5,49 +5,134 @@
 
 namespace Chaos
 {
+	//
+	// Handle Impl
+	//
 
-	template<class T, int d>
-	void TPBDRigidSpringConstraints<T, d>::UpdateDistance(int32 ConstraintIndex, const TVector<T, d>& Location0, const TVector<T, d>& Location1)
+	const TVector<FVec3, 2>& FPBDRigidSpringConstraintHandle::GetConstraintPositions() const
 	{
-		const TVector<TGeometryParticleHandle<T, d>*, 2>& Constraint = Constraints[ConstraintIndex];
-		const TGeometryParticleHandle<T, d>* Particle0 = Constraint[0];
-		const TGeometryParticleHandle<T, d>* Particle1 = Constraint[1];
+		return ConstraintContainer->GetConstraintPositions(ConstraintIndex);
+
+	}
+
+	void FPBDRigidSpringConstraintHandle::SetConstraintPositions(const TVector<FVec3, 2>& ConstraintPositions)
+	{
+		ConstraintContainer->SetConstraintPositions(ConstraintIndex, ConstraintPositions);
+	}
+	
+	TVector<typename FPBDRigidSpringConstraintHandle::FGeometryParticleHandle*, 2> FPBDRigidSpringConstraintHandle::GetConstrainedParticles() const
+	{ 
+		return ConstraintContainer->GetConstrainedParticles(ConstraintIndex);
+	}
+
+	FReal FPBDRigidSpringConstraintHandle::GetRestLength() const
+	{
+		return ConstraintContainer->GetRestLength(ConstraintIndex);
+	}
+
+	void FPBDRigidSpringConstraintHandle::SetRestLength(const FReal SpringLength)
+	{
+		ConstraintContainer->SetRestLength(ConstraintIndex, SpringLength);
+	}
+
+	//
+	// Container Impl
+	//
+
+	FPBDRigidSpringConstraints::FPBDRigidSpringConstraints()
+	{}
+
+	FPBDRigidSpringConstraints::~FPBDRigidSpringConstraints()
+	{
+	}
+
+	typename FPBDRigidSpringConstraints::FConstraintContainerHandle* FPBDRigidSpringConstraints::AddConstraint(const FConstrainedParticlePair& InConstrainedParticles, const  TVector<FVec3, 2>& InLocations, FReal Stiffness, FReal Damping, FReal RestLength)
+	{
+		Handles.Add(HandleAllocator.AllocHandle(this, Handles.Num()));
+		int32 ConstraintIndex = Constraints.Add(InConstrainedParticles);
+
+		SpringSettings.Emplace(FSpringSettings({ Stiffness, Damping, RestLength }));
+
+		Distances.Add({});
+		UpdateDistance(ConstraintIndex, InLocations[0], InLocations[1]);
+
+		return Handles.Last();
+	}
+
+	void FPBDRigidSpringConstraints::RemoveConstraint(int ConstraintIndex)
+	{
+		FConstraintContainerHandle* ConstraintHandle = Handles[ConstraintIndex];
+		if (ConstraintHandle != nullptr)
+		{
+			// Release the handle for the freed constraint
+			HandleAllocator.FreeHandle(ConstraintHandle);
+			Handles[ConstraintIndex] = nullptr;
+		}
+
+		// Swap the last constraint into the gap to keep the array packed
+		Constraints.RemoveAtSwap(ConstraintIndex);
+		SpringSettings.RemoveAtSwap(ConstraintIndex);
+		Distances.RemoveAtSwap(ConstraintIndex);
+		Handles.RemoveAtSwap(ConstraintIndex);
+
+		// Update the handle for the constraint that was moved
+		if (ConstraintIndex < Handles.Num())
+		{
+			FConstraintHandle* Handle = Handles[ConstraintIndex];
+			SetConstraintIndex(Handle, ConstraintIndex);
+		}
+	}
+
+	void FPBDRigidSpringConstraints::RemoveConstraints(const TSet<TGeometryParticleHandle<FReal, 3>*>& RemovedParticles)
+	{
+	}
+
+	void FPBDRigidSpringConstraints::UpdateDistance(int32 ConstraintIndex, const FVec3& Location0, const FVec3& Location1)
+	{
+		const TVector<TGeometryParticleHandle<FReal, 3>*, 2>& Constraint = Constraints[ConstraintIndex];
+		const TGeometryParticleHandle<FReal, 3>* Particle0 = Constraint[0];
+		const TGeometryParticleHandle<FReal, 3>* Particle1 = Constraint[1];
 
 		Distances[ConstraintIndex][0] = Particle0->R().Inverse().RotateVector(Location0 - Particle0->X());
 		Distances[ConstraintIndex][1] = Particle1->R().Inverse().RotateVector(Location1 - Particle1->X());
-		SpringDistances[ConstraintIndex] = (Location0 - Location1).Size();
 	}
 
-	template<class T, int d>
-	TVector<T, d> TPBDRigidSpringConstraints<T, d>::GetDelta(int32 ConstraintIndex, const TVector<T, d>& WorldSpaceX1, const TVector<T, d>& WorldSpaceX2) const
+	FVec3 FPBDRigidSpringConstraints::GetDelta(int32 ConstraintIndex, const FVec3& WorldSpaceX1, const FVec3& WorldSpaceX2) const
 	{
-		const TVector<TGeometryParticleHandle<T, d>*, 2>& Constraint = Constraints[ConstraintIndex];
-		const TPBDRigidParticleHandle<T, d>* PBDRigid0 = Constraint[0]->CastToRigidParticle();
-		const TPBDRigidParticleHandle<T, d>* PBDRigid1 = Constraint[1]->CastToRigidParticle();
+		const TVector<TGeometryParticleHandle<FReal, 3>*, 2>& Constraint = Constraints[ConstraintIndex];
+		const TPBDRigidParticleHandle<FReal, 3>* PBDRigid0 = Constraint[0]->CastToRigidParticle();
+		const TPBDRigidParticleHandle<FReal, 3>* PBDRigid1 = Constraint[1]->CastToRigidParticle();
 		const bool bIsRigidDynamic0 = PBDRigid0 && PBDRigid0->ObjectState() == EObjectStateType::Dynamic;
 		const bool bIsRigidDynamic1 = PBDRigid1 && PBDRigid1->ObjectState() == EObjectStateType::Dynamic;
 
 		if (!bIsRigidDynamic0 && !bIsRigidDynamic1)
 		{
-			return TVector<T, d>(0);
+			return FVec3(0);
 		}
 
-		const TVector<T, d> Difference = WorldSpaceX2 - WorldSpaceX1;
+		const FVec3 Difference = WorldSpaceX2 - WorldSpaceX1;
 
 		const float Distance = Difference.Size();
 		check(Distance > 1e-7);
 
-		const TVector<T, d> Direction = Difference / Distance;
-		const TVector<T, d> Delta = (Distance - SpringDistances[ConstraintIndex]) * Direction;
-		const T InvM0 = (bIsRigidDynamic0) ? PBDRigid0->InvM() : (T)0;
-		const T InvM1 = (bIsRigidDynamic1) ? PBDRigid1->InvM() : (T)0;
-		const T CombinedMass = InvM0 + InvM1;
+		const FVec3 Direction = Difference / Distance;
+		const FVec3 Delta = (Distance - SpringSettings[ConstraintIndex].RestLength) * Direction;
+		const FReal InvM0 = (bIsRigidDynamic0) ? PBDRigid0->InvM() : (FReal)0;
+		const FReal InvM1 = (bIsRigidDynamic1) ? PBDRigid1->InvM() : (FReal)0;
+		const FReal CombinedMass = InvM0 + InvM1;
 
-		return Stiffness * Delta / CombinedMass;
+		return SpringSettings[ConstraintIndex].Stiffness * Delta / CombinedMass;
 	}
 
-	template<class T, int d>
-	void TPBDRigidSpringConstraints<T, d>::Apply(const T Dt, const TArray<FConstraintContainerHandle*>& InConstraintHandles, const int32 It, const int32 NumIts)
+	void FPBDRigidSpringConstraints::Apply(const FReal Dt, const int32 It, const int32 NumIts)
+	{
+		for (int32 ConstraintIndex = 0; ConstraintIndex < NumConstraints(); ++ConstraintIndex)
+		{
+			ApplySingle(Dt, ConstraintIndex);
+		}
+	}
+
+	void FPBDRigidSpringConstraints::Apply(const FReal Dt, const TArray<FConstraintContainerHandle*>& InConstraintHandles, const int32 It, const int32 NumIts)
 	{
 		for (FConstraintContainerHandle* ConstraintHandle : InConstraintHandles)
 		{
@@ -55,47 +140,43 @@ namespace Chaos
 		}
 	}
 
-
-	template<class T, int d>
-	void TPBDRigidSpringConstraints<T, d>::ApplySingle(const T Dt, int32 ConstraintIndex) const
+	void FPBDRigidSpringConstraints::ApplySingle(const FReal Dt, int32 ConstraintIndex) const
 	{
 		const FConstrainedParticlePair& Constraint = Constraints[ConstraintIndex];
 
-		TPBDRigidParticleHandle<T, d>* PBDRigid0 = Constraint[0]->CastToRigidParticle();
-		TPBDRigidParticleHandle<T, d>* PBDRigid1 = Constraint[1]->CastToRigidParticle();
+		TPBDRigidParticleHandle<FReal, 3>* PBDRigid0 = Constraint[0]->CastToRigidParticle();
+		TPBDRigidParticleHandle<FReal, 3>* PBDRigid1 = Constraint[1]->CastToRigidParticle();
 		const bool bIsRigidDynamic0 = PBDRigid0 && PBDRigid0->ObjectState() == EObjectStateType::Dynamic;
 		const bool bIsRigidDynamic1 = PBDRigid1 && PBDRigid1->ObjectState() == EObjectStateType::Dynamic;
 
-		check((bIsRigidDynamic0 && bIsRigidDynamic1 && PBDRigid0->Island() == PBDRigid1->Island()) || (!bIsRigidDynamic0 && bIsRigidDynamic1) || (bIsRigidDynamic0 && bIsRigidDynamic1));
+		check((bIsRigidDynamic0 && bIsRigidDynamic1) || (!bIsRigidDynamic0 && bIsRigidDynamic1) || (bIsRigidDynamic0 && bIsRigidDynamic1));
 
 		// @todo(ccaulfield): see if we can eliminate the need for all these ifs
-		const TRotation<T, d> & Q0 = bIsRigidDynamic0 ? PBDRigid0->Q() : Constraint[0]->R();
-		const TRotation<T, d> & Q1 = bIsRigidDynamic1 ? PBDRigid1->Q() : Constraint[1]->R();
-		const TVector<T, d> & P0 = bIsRigidDynamic0 ? PBDRigid0->P() : Constraint[0]->X();
-		const TVector<T, d> & P1 = bIsRigidDynamic1 ? PBDRigid1->P() : Constraint[1]->X();
+		const TRotation<FReal, 3> & Q0 = bIsRigidDynamic0 ? PBDRigid0->Q() : Constraint[0]->R();
+		const TRotation<FReal, 3> & Q1 = bIsRigidDynamic1 ? PBDRigid1->Q() : Constraint[1]->R();
+		const FVec3 & P0 = bIsRigidDynamic0 ? PBDRigid0->P() : Constraint[0]->X();
+		const FVec3 & P1 = bIsRigidDynamic1 ? PBDRigid1->P() : Constraint[1]->X();
 
-		const TVector<T, d> WorldSpaceX1 = Q0.RotateVector(Distances[ConstraintIndex][0]) + P0;
-		const TVector<T, d> WorldSpaceX2 = Q1.RotateVector(Distances[ConstraintIndex][1]) + P1;
-		const PMatrix<T, d, d> WorldSpaceInvI1 = bIsRigidDynamic0 ? Utilities::ComputeWorldSpaceInertia(Q0, PBDRigid0->InvI()) : PMatrix<T, d, d>(0);
-		const PMatrix<T, d, d> WorldSpaceInvI2 = bIsRigidDynamic1 ? Utilities::ComputeWorldSpaceInertia(Q1, PBDRigid1->InvI()) : PMatrix<T, d, d>(0);
-		const TVector<T, d> Delta = GetDelta(ConstraintIndex, WorldSpaceX1, WorldSpaceX2);
+		const FVec3 WorldSpaceX1 = Q0.RotateVector(Distances[ConstraintIndex][0]) + P0;
+		const FVec3 WorldSpaceX2 = Q1.RotateVector(Distances[ConstraintIndex][1]) + P1;
+		const FMatrix33 WorldSpaceInvI1 = bIsRigidDynamic0 ? Utilities::ComputeWorldSpaceInertia(Q0, PBDRigid0->InvI()) : FMatrix33(0);
+		const FMatrix33 WorldSpaceInvI2 = bIsRigidDynamic1 ? Utilities::ComputeWorldSpaceInertia(Q1, PBDRigid1->InvI()) : FMatrix33(0);
+		const FVec3 Delta = GetDelta(ConstraintIndex, WorldSpaceX1, WorldSpaceX2);
 
 		if (bIsRigidDynamic0)
 		{
-			const TVector<T, d> Radius = WorldSpaceX1 - PBDRigid0->P();
+			const FVec3 Radius = WorldSpaceX1 - PBDRigid0->P();
 			PBDRigid0->P() += PBDRigid0->InvM() * Delta;
-			PBDRigid0->Q() += TRotation<T, d>::FromElements(WorldSpaceInvI1 * TVector<T, d>::CrossProduct(Radius, Delta), 0.f) * PBDRigid0->Q() * T(0.5);
+			PBDRigid0->Q() += TRotation<FReal, 3>::FromElements(WorldSpaceInvI1 * FVec3::CrossProduct(Radius, Delta), 0.f) * PBDRigid0->Q() * FReal(0.5);
 			PBDRigid0->Q().Normalize();
 		}
 
 		if (bIsRigidDynamic1)
 		{
-			const TVector<T, d> Radius = WorldSpaceX2 - PBDRigid1->P();
+			const FVec3 Radius = WorldSpaceX2 - PBDRigid1->P();
 			PBDRigid1->P() -= PBDRigid1->InvM() * Delta;
-			PBDRigid1->Q() += TRotation<T, d>::FromElements(WorldSpaceInvI2 * TVector<T, d>::CrossProduct(Radius, -Delta), 0.f) * PBDRigid1->Q() * T(0.5);
+			PBDRigid1->Q() += TRotation<FReal, 3>::FromElements(WorldSpaceInvI2 * FVec3::CrossProduct(Radius, -Delta), 0.f) * PBDRigid1->Q() * FReal(0.5);
 			PBDRigid1->Q().Normalize();
 		}
 	}
-
-	template class TPBDRigidSpringConstraints<float, 3>;
 }
diff --git a/Engine/Source/Runtime/Experimental/Chaos/Public/Chaos/Evolution/PBDMinEvolution.h b/Engine/Source/Runtime/Experimental/Chaos/Public/Chaos/Evolution/PBDMinEvolution.h
--- a/Engine/Source/Runtime/Experimental/Chaos/Public/Chaos/Evolution/PBDMinEvolution.h
+++ b/Engine/Source/Runtime/Experimental/Chaos/Public/Chaos/Evolution/PBDMinEvolution.h
@@ -32,6 +32,8 @@ namespace Chaos
 	class CHAOS_API FPBDMinEvolution
 	{
 	public:
+		// @todo(ccaulfield): make it so that CollisionDetection is plugged in with a constraint rule...
+
 		using FCollisionConstraints = TPBDCollisionConstraints<FReal, 3>;
 		using FCollisionDetector = TCollisionDetector<FParticlePairBroadPhase, FNarrowPhase, FSyncCollisionReceiver, FCollisionConstraints>;
 		using FEvolutionCallback = TFunction<void()>;
diff --git a/Engine/Source/Runtime/Experimental/Chaos/Public/Chaos/PBDRigidSpringConstraints.h b/Engine/Source/Runtime/Experimental/Chaos/Public/Chaos/PBDRigidSpringConstraints.h
--- a/Engine/Source/Runtime/Experimental/Chaos/Public/Chaos/PBDRigidSpringConstraints.h
+++ b/Engine/Source/Runtime/Experimental/Chaos/Public/Chaos/PBDRigidSpringConstraints.h
@@ -8,65 +8,48 @@
 
 namespace Chaos
 {
-	template<class T, int d>
-	class TPBDRigidSpringConstraints;
+	class FPBDRigidSpringConstraints;
 
-	template<class T, int d>
-	class CHAOS_API TPBDRigidSpringConstraintHandle : public TContainerConstraintHandle<TPBDRigidSpringConstraints<T, d>>
+	class CHAOS_API FPBDRigidSpringConstraintHandle : public TContainerConstraintHandle<FPBDRigidSpringConstraints>
 	{
 	public:
-		using Base = TContainerConstraintHandle<TPBDRigidSpringConstraints<T, d>>;
-		using FConstraintContainer = TPBDRigidSpringConstraints<T, d>;
-		using FGeometryParticleHandle = TGeometryParticleHandle<T, d>;
+		using Base = TContainerConstraintHandle<FPBDRigidSpringConstraints>;
+		using FConstraintContainer = FPBDRigidSpringConstraints;
+		using FGeometryParticleHandle = TGeometryParticleHandle<FReal, 3>;
 
-		TPBDRigidSpringConstraintHandle() {}
-		TPBDRigidSpringConstraintHandle(FConstraintContainer* InConstraintContainer, int32 InConstraintIndex) : TContainerConstraintHandle<TPBDRigidSpringConstraints<T, d>>(InConstraintContainer, InConstraintIndex) {}
+		FPBDRigidSpringConstraintHandle()
+		{
+		}
+		
+		FPBDRigidSpringConstraintHandle(FConstraintContainer* InConstraintContainer, int32 InConstraintIndex) 
+		: TContainerConstraintHandle<FPBDRigidSpringConstraints>(InConstraintContainer, InConstraintIndex) 
+		{
+		}
 
-		const TVector<TVector<T, 3>, 2>& GetConstraintPositions() const;
-		void SetConstraintPositions(const TVector<TVector<T, 3>, 2>& ConstraintPositions);
-		TVector<FGeometryParticleHandle*, 2> GetConstrainedParticles() const { return ConstraintContainer->GetConstrainedParticles(ConstraintIndex); }
+		const TVector<FVec3, 2>& GetConstraintPositions() const;
+		void SetConstraintPositions(const TVector<FVec3, 2>& ConstraintPositions);
+		
+		TVector<FGeometryParticleHandle*, 2> GetConstrainedParticles() const;
 
-	protected:
-		using Base::ConstraintIndex;
-		using Base::ConstraintContainer;
+		// Get the rest length of the spring
+		FReal GetRestLength() const;
+		void SetRestLength(const FReal SpringLength);
 	};
 
 
-	template<class T, int d>
-	class TPBDRigidSpringConstraints : public FPBDConstraintContainer
+	class FPBDRigidSpringConstraints : public FPBDConstraintContainer
 	{
 	public:
+		// @todo(ccaulfield): an alternative AddConstraint which takes the constrain settings rather than assuming everything is in world-space rest pose
+
 		using Base = FPBDConstraintContainer;
-		using FReal = T;
-		static const int Dimensions = d;
-		using FConstraintContainerHandle = TPBDRigidSpringConstraintHandle<FReal, Dimensions>;
-		using FConstraintHandleAllocator = TConstraintHandleAllocator<TPBDRigidSpringConstraints<FReal, Dimensions>>;
-		using FConstrainedParticlePair = TVector<TGeometryParticleHandle<T, d>*, 2>;
+		using FConstraintContainerHandle = FPBDRigidSpringConstraintHandle;
+		using FConstraintHandleAllocator = TConstraintHandleAllocator<FPBDRigidSpringConstraints>;
+		using FConstrainedParticlePair = TVector<TGeometryParticleHandle<FReal, 3>*, 2>;
 		using FHandles = TArray<FConstraintContainerHandle*>;
 
-		TPBDRigidSpringConstraints(const T InStiffness = (T)1)
-			: Stiffness(InStiffness) 
-		{}
-
-		TPBDRigidSpringConstraints(const TArray<TVector<T, 3>>& Locations0, const TArray<TVector<T, 3>>& Locations1, TArray<FConstrainedParticlePair>&& InConstraints, const T InStiffness = (T)1)
-			: Constraints(MoveTemp(InConstraints)), Stiffness(InStiffness)
-		{
-			if (Constraints.Num() > 0)
-			{
-				Handles.Reserve(Constraints.Num());
-				Distances.Reserve(Constraints.Num());
-				SpringDistances.Reserve(Constraints.Num());
-				for (int32 ConstraintIndex = 0; ConstraintIndex < Constraints.Num(); ++ConstraintIndex)
-				{
-					Handles.Add(HandleAllocator.AllocHandle(this, ConstraintIndex));
-					Distances.Add({});
-					SpringDistances.Add({});
-					UpdateDistance(ConstraintIndex, Locations0[ConstraintIndex], Locations1[ConstraintIndex]);
-				}
-			}
-		}
-
-		virtual ~TPBDRigidSpringConstraints() {}
+		FPBDRigidSpringConstraints();
+		virtual ~FPBDRigidSpringConstraints();
 
 		//
 		// Constraint Container API
@@ -83,47 +66,19 @@ namespace Chaos
 		/**
 		 * Add a constraint initialized from current world-space particle positions.
 		 * You would use this method when your objects are already positioned in the world.
+		 *
+		 * \param InConstrainedParticles the two particles connected by the spring
+		 * \param InLocations the world-space locations of the spring connectors on each particle
 		 */
-		FConstraintContainerHandle* AddConstraint(const FConstrainedParticlePair& InConstrainedParticles, const  TVector<TVector<T, 3>, 2>& InLocations)
-		{
-			Handles.Add(HandleAllocator.AllocHandle(this, Handles.Num()));
-			int32 ConstraintIndex = Constraints.Add(InConstrainedParticles);
-			Distances.Add({});
-			SpringDistances.Add({});
-			UpdateDistance(ConstraintIndex, InLocations[0], InLocations[1]);
-			return Handles.Last();
-		}
+		FConstraintContainerHandle* AddConstraint(const FConstrainedParticlePair& InConstrainedParticles, const  TVector<FVec3, 2>& InLocations, FReal Stiffness, FReal Damping, FReal RestLength);
 
 		/**
 		 * Remove the specified constraint.
 		 */
-		void RemoveConstraint(int ConstraintIndex)
-		{
-			FConstraintContainerHandle* ConstraintHandle = Handles[ConstraintIndex];
-			if (ConstraintHandle != nullptr)
-			{
-				// Release the handle for the freed constraint
-				HandleAllocator.FreeHandle(ConstraintHandle);
-				Handles[ConstraintIndex] = nullptr;
-			}
-
-			// Swap the last constraint into the gap to keep the array packed
-			Constraints.RemoveAtSwap(ConstraintIndex);
-			Distances.RemoveAtSwap(ConstraintIndex);
-			Handles.RemoveAtSwap(ConstraintIndex);
-
-			// Update the handle for the constraint that was moved
-			if (ConstraintIndex < Handles.Num())
-			{
-				FConstraintHandle* Handle = Handles[ConstraintIndex];
-				SetConstraintIndex(Handle, ConstraintIndex);
-			}
-		}
+		void RemoveConstraint(int ConstraintIndex);
 
 		// @todo(ccaulfield): rename/remove  this
-		void RemoveConstraints(const TSet<TGeometryParticleHandle<T, d>*>& RemovedParticles)
-		{
-		}
+		void RemoveConstraints(const TSet<TGeometryParticleHandle<FReal, 3>*>& RemovedParticles);
 
 		//
 		// Constraint API
@@ -158,7 +113,7 @@ namespace Chaos
 		/**
 		 * Get the local-space constraint positions for each body.
 		 */
-		const TVector<TVector<T, 3>, 2>& GetConstraintPositions(int ConstraintIndex) const
+		const TVector<FVec3, 2>& GetConstraintPositions(int ConstraintIndex) const
 		{
 			return Distances[ConstraintIndex];
 		}
@@ -166,11 +121,27 @@ namespace Chaos
 		/**
 		 * Set the local-space constraint positions for each body.
 		 */
-		void SetConstraintPositions(int ConstraintIndex, const TVector<TVector<T, 3>, 2>& ConstraintPositions)
+		void SetConstraintPositions(int ConstraintIndex, const TVector<FVec3, 2>& ConstraintPositions)
 		{
 			Distances[ConstraintIndex] = ConstraintPositions;
 		}
 
+		/**
+		 * Get the rest length of the spring
+		 */
+		FReal GetRestLength(int32 ConstraintIndex) const
+		{
+			return SpringSettings[ConstraintIndex].RestLength;
+		}
+
+		/**
+		 * Set the rest length of the spring
+		 */
+		void SetRestLength(int32 ConstraintIndex, const FReal SpringLength)
+		{
+			SpringSettings[ConstraintIndex].RestLength = SpringLength;
+		}
+
 
 		//
 		// Island Rule API
@@ -184,13 +155,20 @@ namespace Chaos
 		{
 		}
 
-		void UpdatePositionBasedState(const T Dt)
+		void UpdatePositionBasedState(const FReal Dt)
+		{
+		}
+
+		void Apply(const FReal Dt, const int32 It, const int32 NumIts);
+
+		bool ApplyPushOut(const FReal Dt, const int32 It, const int32 NumIts)
 		{
+			return false;
 		}
 
-		void Apply(const T Dt, const TArray<FConstraintContainerHandle*>& InConstraintHandles, const int32 It, const int32 NumIts);
+		void Apply(const FReal Dt, const TArray<FConstraintContainerHandle*>& InConstraintHandles, const int32 It, const int32 NumIts);
 
-		bool ApplyPushOut(const T Dt, const TArray<FConstraintContainerHandle*>& InConstraintHandles, const int32 It, const int32 NumIts)
+		bool ApplyPushOut(const FReal Dt, const TArray<FConstraintContainerHandle*>& InConstraintHandles, const int32 It, const int32 NumIts)
 		{
 			return false;
 		}
@@ -200,16 +178,22 @@ namespace Chaos
 		using Base::SetConstraintIndex;
 
 	private:
-		void ApplySingle(const T Dt, int32 ConstraintIndex) const;
+		void ApplySingle(const FReal Dt, int32 ConstraintIndex) const;
 
-		void UpdateDistance(int32 ConstraintIndex, const TVector<T, d>& Location0, const TVector<T, d>& Location1);
+		void UpdateDistance(int32 ConstraintIndex, const FVec3& Location0, const FVec3& Location1);
 
-		TVector<T, d> GetDelta(int32 ConstraintIndex, const TVector<T, d>& WorldSpaceX1, const TVector<T, d>& WorldSpaceX2) const;
+		FVec3 GetDelta(int32 ConstraintIndex, const FVec3& WorldSpaceX1, const FVec3& WorldSpaceX2) const;
+
+		struct FSpringSettings
+		{
+			FReal Stiffness;
+			FReal Damping;
+			FReal RestLength;
+		};
 
 		TArray<FConstrainedParticlePair> Constraints;
-		TArray<TVector<TVector<T, 3>, 2>> Distances;
-		TArray<T> SpringDistances;
-		T Stiffness;
+		TArray<TVector<FVec3, 2>> Distances;
+		TArray<FSpringSettings> SpringSettings;
 
 		TArray<FConstraintContainerHandle*> Handles;
 		FConstraintHandleAllocator HandleAllocator;
