- Fixed static constexpr link error on Mac

diff --git a/Engine/Source/Runtime/Experimental/Chaos/Private/Chaos/CollisionConvexMesh.cpp b/Engine/Source/Runtime/Experimental/Chaos/Private/Chaos/CollisionConvexMesh.cpp
--- a/Engine/Source/Runtime/Experimental/Chaos/Private/Chaos/CollisionConvexMesh.cpp
+++ b/Engine/Source/Runtime/Experimental/Chaos/Private/Chaos/CollisionConvexMesh.cpp
@@ -5,9 +5,6 @@
 
 namespace Chaos
 {
-	constexpr float FConvexBuilder::TriQuadPrismInflation;
-
-
 	// CVars variables for controlling geometry complexity checking and simplification
 	int32 FConvexBuilder::PerformGeometryCheck = 0;
 
@@ -20,4 +17,4 @@ namespace Chaos
 	FAutoConsoleVariableRef CVarConvexGeometrySimplifyEnable(TEXT("p.Chaos.PerformGeometryReduction"), FConvexBuilder::PerformGeometryReduction, TEXT("Perform convex geometry simplification to increase performance in Chaos physics."));
 
 	FAutoConsoleVariableRef CVarConvexParticlesWarningThreshold(TEXT("p.Chaos.ConvexParticlesWarningThreshold"), FConvexBuilder::ParticlesThreshold, TEXT("Threshold beyond which we warn about collision geometry complexity."));
-}
\ No newline at end of file
+}
diff --git a/Engine/Source/Runtime/Experimental/Chaos/Public/Chaos/CollisionConvexMesh.h b/Engine/Source/Runtime/Experimental/Chaos/Public/Chaos/CollisionConvexMesh.h
--- a/Engine/Source/Runtime/Experimental/Chaos/Public/Chaos/CollisionConvexMesh.h
+++ b/Engine/Source/Runtime/Experimental/Chaos/Public/Chaos/CollisionConvexMesh.h
@@ -10,15 +10,16 @@
 
 namespace Chaos
 {
+	// When encountering a triangle or quad in hull generation (3-points or 4 coplanar points) we will instead generate
+	// a prism with a small thickness to emulate the desired result as a hull. Otherwise hull generation will fail on
+	// these cases. Verbose logging on LogChaos will point out when this path is taken for further scrutiny about
+	// the geometry
+static constexpr float TriQuadPrismInflation() { return 0.1f; }
+
 	class FConvexBuilder
 	{
 	public:
 
-		// When encountering a triangle or quad in hull generation (3-points or 4 coplanar points) we will instead generate
-		// a prism with a small thickness to emulate the desired result as a hull. Otherwise hull generation will fail on
-		// these cases. Verbose logging on LogChaos will point out when this path is taken for further scrutiny about
-		// the geometry
-		static constexpr float TriQuadPrismInflation = 0.1f;
 
 		static bool IsValidTriangle(const FVec3& A, const FVec3& B, const FVec3& C, FVec3& OutNormal)
 		{
@@ -108,9 +109,9 @@ namespace Chaos
 				//TODO_SQ_IMPLEMENTATION: should do proper cleanup to avoid this
 				if(ensureMsgf(bIsValidTriangle, TEXT("FConvexBuilder::Build(): Generated invalid triangle!")))
 				{
-					Inflate(InParticles, ModifiedParticles, PlanarNormal, TriQuadPrismInflation);
+					Inflate(InParticles, ModifiedParticles, PlanarNormal, TriQuadPrismInflation());
 					ParticlesToUse = &ModifiedParticles;
-					UE_LOG(LogChaos, Verbose, TEXT("Encountered a triangle in convex hull generation. Will prepare a prism of thickness %.5f in place of a triangle."), TriQuadPrismInflation);
+					UE_LOG(LogChaos, Verbose, TEXT("Encountered a triangle in convex hull generation. Will prepare a prism of thickness %.5f in place of a triangle."), TriQuadPrismInflation());
 				}
 				else
 				{
@@ -119,9 +120,9 @@ namespace Chaos
 			}
 			else if(IsPlanarShape(InParticles, PlanarNormal))
 			{
-				Inflate(InParticles, ModifiedParticles, PlanarNormal, TriQuadPrismInflation);
+				Inflate(InParticles, ModifiedParticles, PlanarNormal, TriQuadPrismInflation());
 				ParticlesToUse = &ModifiedParticles;
-				UE_LOG(LogChaos, Verbose, TEXT("Encountered a planar shape in convex hull generation. Will prepare a prism of thickness %.5f in place of a triangle."), TriQuadPrismInflation);
+				UE_LOG(LogChaos, Verbose, TEXT("Encountered a planar shape in convex hull generation. Will prepare a prism of thickness %.5f in place of a triangle."), TriQuadPrismInflation());
 			}
 
 			const int32 NumParticlesToUse = ParticlesToUse->Size();
