Fix EPA returning wrong normal direction when origin is just outside of minkowski volume

diff --git a/Engine/Source/Runtime/Experimental/Chaos/Private/Chaos/BoundingVolume.cpp b/Engine/Source/Runtime/Experimental/Chaos/Private/Chaos/BoundingVolume.cpp
--- a/Engine/Source/Runtime/Experimental/Chaos/Private/Chaos/BoundingVolume.cpp
+++ b/Engine/Source/Runtime/Experimental/Chaos/Private/Chaos/BoundingVolume.cpp
@@ -7,22 +7,28 @@
 int FBoundingVolumeCVars::FilterFarBodies = 0;
 
 FAutoConsoleVariableRef FBoundingVolumeCVars::CVarFilterFarBodies(
-	TEXT("p.RemoveFarBodiesFromBVH"),
-	FBoundingVolumeCVars::FilterFarBodies,
-	TEXT("Removes bodies far from the scene from the bvh\n")
-	TEXT("0: Kept, 1: Removed"),
-	ECVF_Default);
+    TEXT("p.RemoveFarBodiesFromBVH"),
+    FBoundingVolumeCVars::FilterFarBodies,
+    TEXT("Removes bodies far from the scene from the bvh\n")
+        TEXT("0: Kept, 1: Removed"),
+    ECVF_Default);
 
 namespace Chaos
 {
 	CHAOS_API int32 MaxDirtyElements = 10000;
 	FAutoConsoleVariableRef CVarMaxDirtyElements(
-		TEXT("p.MaxDirtyElements"),
-		MaxDirtyElements,
-		TEXT("The max number of dirty elements. This forces a flush which is very expensive"));
-
-	template <typename TPayloadType, typename T, int d>
-	ISpatialAcceleration<TPayloadType, T, d>* ISpatialAcceleration<TPayloadType, T, d>::SerializationFactory(FChaosArchive& Ar, ISpatialAcceleration<TPayloadType, T, d>* Accel)
+	    TEXT("p.MaxDirtyElements"),
+	    MaxDirtyElements,
+	    TEXT("The max number of dirty elements. This forces a flush which is very expensive"));
+
+	CHAOS_API int32 FixEPAWhenSimplexOutside = 1;
+	FAutoConsoleVariableRef CVarFixEPAWhenSimplexOutside(
+	    TEXT("p.FixEPAWhenSimplexOutside"),
+	    FixEPAWhenSimplexOutside,
+	    TEXT(""));
+
+	    template<typename TPayloadType, typename T, int d>
+	    ISpatialAcceleration<TPayloadType, T, d>* ISpatialAcceleration<TPayloadType, T, d>::SerializationFactory(FChaosArchive& Ar, ISpatialAcceleration<TPayloadType, T, d>* Accel)
 	{
 		if (Ar.CustomVer(FExternalPhysicsCustomObjectVersion::GUID) < FExternalPhysicsCustomObjectVersion::SerializeEvolutionGenericAcceleration)
 		{
@@ -42,8 +48,8 @@ namespace Chaos
 	}
 
 	template class CHAOS_API Chaos::ISpatialAcceleration<int32, float, 3>;
-	template class CHAOS_API Chaos::ISpatialAcceleration<TAccelerationStructureHandle<float,3>, float, 3>;
+	template class CHAOS_API Chaos::ISpatialAcceleration<TAccelerationStructureHandle<float, 3>, float, 3>;
 
-    template class CHAOS_API Chaos::TBoundingVolume<int32,float,3>;
-    template class CHAOS_API Chaos::TBoundingVolume<TAccelerationStructureHandle<float,3>,float,3>;
+	template class CHAOS_API Chaos::TBoundingVolume<int32, float, 3>;
+	template class CHAOS_API Chaos::TBoundingVolume<TAccelerationStructureHandle<float, 3>, float, 3>;
 }
diff --git a/Engine/Source/Runtime/Experimental/Chaos/Public/Chaos/EPA.h b/Engine/Source/Runtime/Experimental/Chaos/Public/Chaos/EPA.h
--- a/Engine/Source/Runtime/Experimental/Chaos/Public/Chaos/EPA.h
+++ b/Engine/Source/Runtime/Experimental/Chaos/Public/Chaos/EPA.h
@@ -298,6 +298,8 @@ void EPAComputeVisibilityBorder(TArray<TEPAEntry<T>>& Entries, int32 EntryIdx, c
 	}
 }
 
+extern CHAOS_API int32 FixEPAWhenSimplexOutside;
+
 template <typename T>
 void ComputeEPAResults(const TVec3<T>* VertsA, const TVec3<T>* VertsB, const TEPAEntry<T>& Entry, T& OutPenetration, TVec3<T>& OutDir, TVec3<T>& OutA, TVec3<T>& OutB)
 {
@@ -320,6 +322,12 @@ void ComputeEPAResults(const TVec3<T>* VertsA, const TVec3<T>* VertsB, const TEP
 		OutDir /= OutPenetration;
 	}
 
+	if (Entry.Distance < 0 && FixEPAWhenSimplexOutside)
+	{
+		//The origin is on the outside, so the direction is reversed
+		OutDir = -OutDir;
+	}
+
 	OutA = TVec3<T>(0);
 	OutB = TVec3<T>(0);
 
