Fix impulses.

diff --git a/Engine/Source/Runtime/Experimental/Chaos/Public/Chaos/ParticleHandle.h b/Engine/Source/Runtime/Experimental/Chaos/Public/Chaos/ParticleHandle.h
--- a/Engine/Source/Runtime/Experimental/Chaos/Public/Chaos/ParticleHandle.h
+++ b/Engine/Source/Runtime/Experimental/Chaos/Public/Chaos/ParticleHandle.h
@@ -2123,8 +2123,8 @@ public:
 		, MP(InParticle.P())
 		, MF(InParticle.F())
 		, MTorque(InParticle.Torque())
-		, MLinearImpulse(TVector<T, d>(0))
-		, MAngularImpulse(TVector<T, d>(0))
+		, MLinearImpulse(InParticle.LinearImpulse())
+		, MAngularImpulse(InParticle.AngularImpulse())
 		, MI(InParticle.I())
 		, MInvI(InParticle.InvI())
 		, MCollisionParticles(nullptr)
diff --git a/Engine/Source/Runtime/Experimental/ChaosSolvers/Private/PhysicsProxy/SingleParticlePhysicsProxy.cpp b/Engine/Source/Runtime/Experimental/ChaosSolvers/Private/PhysicsProxy/SingleParticlePhysicsProxy.cpp
--- a/Engine/Source/Runtime/Experimental/ChaosSolvers/Private/PhysicsProxy/SingleParticlePhysicsProxy.cpp
+++ b/Engine/Source/Runtime/Experimental/ChaosSolvers/Private/PhysicsProxy/SingleParticlePhysicsProxy.cpp
@@ -419,6 +419,8 @@ void FSingleParticlePhysicsProxy<Chaos::TPBDRigidParticle<float, 3>>::ClearAccum
 {
 	Particle->SetF(Chaos::TVector<float, 3>(0));
 	Particle->SetTorque(Chaos::TVector<float, 3>(0));
+	Particle->SetLinearImpulse(Chaos::TVector<float, 3>(0));
+	Particle->SetAngularImpulse(Chaos::TVector<float, 3>(0));
 	Particle->ClearDirtyFlags();
 }
 
