PushPhysicsState Optim
 -wrap around a cvar. Default = off.
 -Use a pool of particle data instead of alloc/delete them each frame
 -Do one task for each type of particles. The task is actually launched in the EndPhysics.
 -Shrink pool each (n) frame

diff --git a/Engine/Source/Runtime/Experimental/Chaos/Public/Chaos/ParticleHandle.h b/Engine/Source/Runtime/Experimental/Chaos/Public/Chaos/ParticleHandle.h
--- a/Engine/Source/Runtime/Experimental/Chaos/Public/Chaos/ParticleHandle.h
+++ b/Engine/Source/Runtime/Experimental/Chaos/Public/Chaos/ParticleHandle.h
@@ -1668,7 +1668,33 @@ public:
 		DebugName = NAME_None;
 #endif
 	}
-
+	
+	void Init(const TGeometryParticle<T, d>& InParticle)
+		{
+			Type = EParticleType::Static;
+			X = InParticle.X();
+			R = InParticle.R();
+			Geometry = InParticle.GeometrySharedLowLevel();
+			SpatialIdx = InParticle.SpatialIdx();
+			UniqueIdx = InParticle.UniqueIdx();
+			UserData = InParticle.UserData();
+			DirtyFlags = InParticle.DirtyFlags();
+#if CHAOS_CHECKED
+			DebugName = InParticle.DebugName();
+#endif
+			const TShapesArray<T, d>& Shapes = InParticle.ShapesArray();
+			ShapeCollisionDisableFlags.Empty(Shapes.Num());
+			CollisionTraceType.Empty(Shapes.Num());
+			ShapeSimData.Empty(Shapes.Num());
+			ShapeQueryData.Empty(Shapes.Num());
+			for (const TUniquePtr<TPerShapeData<T, d>>& ShapePtr : Shapes)
+			{
+				ShapeCollisionDisableFlags.Add(ShapePtr->bDisable);
+				CollisionTraceType.Add(ShapePtr->CollisionTraceType);
+				ShapeSimData.Add(ShapePtr->SimData);
+				ShapeQueryData.Add(ShapePtr->QueryData);
+			}
+	}
 	TVector<T, d> X;
 	TRotation<T, d> R;
 	TSharedPtr<FImplicitObject, ESPMode::ThreadSafe> Geometry;
@@ -1795,6 +1821,14 @@ public:
 		MRotationOfMass = TRotation<T, d>(FQuat(EForceInit::ForceInit));
 	}
 
+	void Init(const TKinematicGeometryParticle<T, d>& InParticle) {
+			Base::Init(InParticle);
+			MV = InParticle.V();
+			MW = InParticle.W();
+			MCenterOfMass = InParticle.CenterOfMass();
+			MRotationOfMass = InParticle.RotationOfMass();
+			Type = EParticleType::Kinematic;
+	}
 	TVector<T, d> MV;
 	TVector<T, d> MW;
 	TVector<T, d> MCenterOfMass;
@@ -2192,6 +2226,31 @@ public:
 		MToBeRemovedOnFracture = false;
 		MGravityEnabled = false;
 	}
+	void Init(const TPBDRigidParticle<T, d>& InParticle) {
+			Base::Init(InParticle);
+			MQ = InParticle.Q();
+			MPreV = InParticle.PreV();
+			MPreW = InParticle.PreW();
+			MP = InParticle.P();
+			MF = InParticle.F();
+			MTorque = InParticle.Torque();
+			MLinearImpulse = InParticle.LinearImpulse();
+			MAngularImpulse = InParticle.AngularImpulse();
+			MI = InParticle.I();
+			MInvI = InParticle.InvI();
+			MCollisionParticles = nullptr;
+			MM = InParticle.M();
+			MInvM = InParticle.InvM();
+			MLinearEtherDrag = InParticle.LinearEtherDrag();
+			MAngularEtherDrag = InParticle.AngularEtherDrag();
+			MIsland = InParticle.Island();
+			MCollisionGroup = InParticle.CollisionGroup();
+			MObjectState = InParticle.ObjectState();
+			MDisabled = InParticle.Disabled();
+			MToBeRemovedOnFracture = InParticle.ToBeRemovedOnFracture();
+			MGravityEnabled = InParticle.IsGravityEnabled();
+			Type = EParticleType::Rigid;
+		}
 };
 
 
diff --git a/Engine/Source/Runtime/Experimental/ChaosSolvers/Private/PBDRigidsSolver.cpp b/Engine/Source/Runtime/Experimental/ChaosSolvers/Private/PBDRigidsSolver.cpp
--- a/Engine/Source/Runtime/Experimental/ChaosSolvers/Private/PBDRigidsSolver.cpp
+++ b/Engine/Source/Runtime/Experimental/ChaosSolvers/Private/PBDRigidsSolver.cpp
@@ -30,12 +30,20 @@ FAutoConsoleVariableRef CVarChaosSolverDrawCollisions(TEXT("p.Chaos.Solver.Debug
 FAutoConsoleVariableRef CVarChaosSolverDrawBPBounds(TEXT("p.Chaos.Solver.DrawBPBounds"), ChaosSolverDrawBPBounds, TEXT("Draw bounding volumes inside the broadphase (0 = never; 1 = end of frame)."));
 #endif
 
+bool ChaosSolverUseParticlePool = true;
+FAutoConsoleVariableRef CVarChaosSolverUseParticlePool(TEXT("p.Chaos.Solver.UseParticlePool"), ChaosSolverUseParticlePool, TEXT("Whether or not to use dirty particle pool (Optim)"));
+
+int32 ChaosSolverParticlePoolNumFrameUntilShrink = 30;
+FAutoConsoleVariableRef CVarChaosSolverParticlePoolNumFrameUntilShrink(TEXT("p.Chaos.Solver.ParticlePoolNumFrameUntilShrink"), ChaosSolverParticlePoolNumFrameUntilShrink, TEXT("Num Frame until we can potentially shrink the pool"));
+
 float ChaosSolverCollisionDefaultIterationsCVar = 1;
 FAutoConsoleVariableRef CVarChaosSolverCollisionDefaultIterations(TEXT("p.ChaosSolverCollisionDefaultIterations"), ChaosSolverCollisionDefaultIterationsCVar, TEXT("Default collision iterations for the solver.[def:1]"));
 
 int32 ChaosSolverCleanupCommandsOnDestruction = 1;
 FAutoConsoleVariableRef CVarChaosSolverCleanupCommandsOnDestruction(TEXT("p.Chaos.Solver.CleanupCommandsOnDestruction"), ChaosSolverCleanupCommandsOnDestruction, TEXT("Whether or not to run internal command queue cleanup on solver destruction (0 = no cleanup, >0 = cleanup all commands)"));
 
+
+
 namespace Chaos
 {
 
@@ -553,6 +561,77 @@ namespace Chaos
 		}
 	}
 
+	template<typename ParticleEntry, typename ProxyEntry, SIZE_T PreAllocCount>
+	void FlushExec(FPBDRigidsSolver::TFramePool<ParticleEntry, ProxyEntry, PreAllocCount>& PoolParticles, Chaos::IDispatcher* Dispatcher, FPBDRigidsSolver * Solver)
+	{
+		Dispatcher->EnqueueCommandImmediate([Solver, &PoolParticles](Chaos::FPersistentPhysicsTask* PhysThread)
+		{
+			for (int32 i = 0; i < PoolParticles.GetEntryCount(); i++)
+			{
+				auto& Entry = PoolParticles.GetEntry(i);
+				// make sure the handle is still valid
+				if (auto* Handle = static_cast<Chaos::TGeometryParticleHandle<float, 3>*>(Entry.Proxy->GetHandle()))
+				{
+					Solver->GetEvolution()->DirtyParticle(*Handle);
+					Entry.Proxy->PushToPhysicsState(&Entry.Particle);
+				}
+				Entry.Particle.Reset();
+			}
+		});
+	}
+
+	template<typename ProxyType, typename ParticleEntry, typename ProxyEntry, SIZE_T PreAllocCount>
+	void PushPhysicsStateExec(FPBDRigidsSolver * Solver, ProxyType* Proxy, FPBDRigidsSolver::TFramePool<ParticleEntry, ProxyEntry, PreAllocCount>& PoolParticles, Chaos::IDispatcher* Dispatcher)
+	{
+		auto* RigidHandle = static_cast<Chaos::TGeometryParticleHandle<float, 3>*>(Proxy->GetHandle());
+		if (RigidHandle == nullptr)
+		{
+			return;
+		}
+
+		// get a new entry in the pool
+		auto& Entry = PoolParticles.GetNewEntry();
+		Entry.Particle.Init(*Proxy->GetParticle());
+		Entry.Proxy = Proxy;
+
+		Solver->RemoveDirtyProxy(Proxy);
+		Proxy->ClearAccumulatedData();
+	}
+
+	void FPBDRigidsSolver::PushPhysicsStatePooled(IDispatcher* Dispatcher)
+	{
+
+		ensure(Dispatcher != nullptr);
+
+		// reset the per frame pool
+		RigidParticlePool.ResetPool();
+		KinematicGeometryParticlePool.ResetPool();
+		GeometryParticlePool.ResetPool();
+
+		TArray< IPhysicsProxyBase *> DirtyProxiesArray = DirtyProxiesSet.Array();
+		for (auto & Proxy : DirtyProxiesArray)
+		{
+			switch (Proxy->GetType())
+			{
+			case EPhysicsProxyType::SingleRigidParticleType:
+				PushPhysicsStateExec(this, static_cast<FSingleParticlePhysicsProxy< Chaos::TPBDRigidParticle<float, 3> >*>(Proxy), RigidParticlePool, Dispatcher);
+				break;
+			case EPhysicsProxyType::SingleKinematicParticleType:
+				PushPhysicsStateExec(this, static_cast<FSingleParticlePhysicsProxy< Chaos::TKinematicGeometryParticle<float, 3> >*>(Proxy), KinematicGeometryParticlePool, Dispatcher);
+				break;
+			case EPhysicsProxyType::SingleGeometryParticleType:
+				PushPhysicsStateExec(this, static_cast<FSingleParticlePhysicsProxy< Chaos::TGeometryParticle<float, 3> >*>(Proxy), GeometryParticlePool, Dispatcher);
+				break;
+			default:
+				ensure("Unknown proxy type in physics solver.");
+			}
+		}
+
+		FlushExec(RigidParticlePool, Dispatcher, this);
+		FlushExec(KinematicGeometryParticlePool, Dispatcher, this);
+		FlushExec(GeometryParticlePool, Dispatcher, this);
+	}
+	
 	void PushPhysicsStateExec(FPBDRigidsSolver* Solver, FGeometryCollectionPhysicsProxy* Proxy, Chaos::IDispatcher* Dispatcher)
 	{
 		Proxy->NewData();
@@ -594,33 +673,40 @@ namespace Chaos
 
 	void FPBDRigidsSolver::PushPhysicsState(IDispatcher* Dispatcher)
 	{
-		if (DirtyProxiesSet.Num() == 0)
-			return;
-		TArray< IPhysicsProxyBase *> DirtyProxiesArray = DirtyProxiesSet.Array();
-		for (auto & Proxy : DirtyProxiesArray)
+		if (ChaosSolverUseParticlePool)
 		{
-			switch (Proxy->GetType())
+			PushPhysicsStatePooled(Dispatcher);
+		}
+		else
+		{
+			if (DirtyProxiesSet.Num() == 0)
+				return;
+			TArray< IPhysicsProxyBase *> DirtyProxiesArray = DirtyProxiesSet.Array();
+			for (auto & Proxy : DirtyProxiesArray)
 			{
-			//case EPhysicsProxyType::NoneType: // 0
-			//case EPhysicsProxyType::StaticMeshType: // 1
-			case EPhysicsProxyType::GeometryCollectionType: // 2
-				PushPhysicsStateExec(this, static_cast<FGeometryCollectionPhysicsProxy*>(Proxy), Dispatcher);
-				break;
-			case EPhysicsProxyType::FieldType: // 3
-				PushPhysicsStateExec(this, static_cast<FFieldSystemPhysicsProxy*>(Proxy), Dispatcher);
-				break;
-			//case EPhysicsProxyType::SkeletalMeshType: // 4
-			case EPhysicsProxyType::SingleRigidParticleType: // 7
-				PushPhysicsStateExec(this, static_cast<FSingleParticlePhysicsProxy< Chaos::TPBDRigidParticle<float, 3> >*>(Proxy), Dispatcher);
-				break;
-			case EPhysicsProxyType::SingleKinematicParticleType: // 6
-				PushPhysicsStateExec(this, static_cast<FSingleParticlePhysicsProxy< Chaos::TKinematicGeometryParticle<float, 3> >*>(Proxy), Dispatcher);
-				break;
-			case EPhysicsProxyType::SingleGeometryParticleType: // 5
-				PushPhysicsStateExec(this, static_cast<FSingleParticlePhysicsProxy< Chaos::TGeometryParticle<float, 3> >*>(Proxy), Dispatcher);
-				break;
-			default:
-				ensure("Unknown proxy type in physics solver.");
+				switch (Proxy->GetType())
+				{
+					//case EPhysicsProxyType::NoneType: // 0
+					//case EPhysicsProxyType::StaticMeshType: // 1
+				case EPhysicsProxyType::GeometryCollectionType: // 2
+					PushPhysicsStateExec(this, static_cast<FGeometryCollectionPhysicsProxy*>(Proxy), Dispatcher);
+					break;
+				case EPhysicsProxyType::FieldType: // 3
+					PushPhysicsStateExec(this, static_cast<FFieldSystemPhysicsProxy*>(Proxy), Dispatcher);
+					break;
+					//case EPhysicsProxyType::SkeletalMeshType: // 4
+				case EPhysicsProxyType::SingleRigidParticleType: // 7
+					PushPhysicsStateExec(this, static_cast<FSingleParticlePhysicsProxy< Chaos::TPBDRigidParticle<float, 3> >*>(Proxy), Dispatcher);
+					break;
+				case EPhysicsProxyType::SingleKinematicParticleType: // 6
+					PushPhysicsStateExec(this, static_cast<FSingleParticlePhysicsProxy< Chaos::TKinematicGeometryParticle<float, 3> >*>(Proxy), Dispatcher);
+					break;
+				case EPhysicsProxyType::SingleGeometryParticleType: // 5
+					PushPhysicsStateExec(this, static_cast<FSingleParticlePhysicsProxy< Chaos::TGeometryParticle<float, 3> >*>(Proxy), Dispatcher);
+					break;
+				default:
+					ensure("Unknown proxy type in physics solver.");
+				}
 			}
 		}
 	}
diff --git a/Engine/Source/Runtime/Experimental/ChaosSolvers/Public/PBDRigidsSolver.h b/Engine/Source/Runtime/Experimental/ChaosSolvers/Public/PBDRigidsSolver.h
--- a/Engine/Source/Runtime/Experimental/ChaosSolvers/Public/PBDRigidsSolver.h
+++ b/Engine/Source/Runtime/Experimental/ChaosSolvers/Public/PBDRigidsSolver.h
@@ -32,6 +32,12 @@ class FStaticMeshPhysicsProxy;
 class FGeometryCollectionPhysicsProxy;
 class FFieldSystemPhysicsProxy;
 
+#define PBDRIGID_PREALLOC_COUNT 1024
+#define KINEMATIC_GEOM_PREALLOC_COUNT 100
+#define GEOMETRY_PREALLOC_COUNT 100
+
+extern int32 ChaosSolverParticlePoolNumFrameUntilShrink;
+
 /**
 *
 */
@@ -221,6 +227,9 @@ namespace Chaos
 		/**/
 		void PushPhysicsState(IDispatcher* Dispatcher = nullptr);
 
+		/**/
+		void PushPhysicsStatePooled(IDispatcher* Dispatcher);
+
 		/**/
 		void BufferPhysicsResults();
 
@@ -381,6 +390,87 @@ namespace Chaos
 
 		template<ELockType>
 		friend struct TSolverQueryMaterialScope;
+
+	public:
+
+		template<typename ParticleEntry, typename ProxyEntry, SIZE_T PreAllocCount>
+		class TFramePool
+		{
+		public:
+
+			struct TPoolEntry
+			{
+				TPoolEntry()
+				{
+					Proxy = nullptr;
+				}
+
+				ParticleEntry Particle;
+				ProxyEntry    Proxy;
+			};
+
+			TFramePool()
+			{
+				EntryCount = 0;
+				MaxEntryCount = 0;
+				FrameCount = 0;
+				// Prealloc default object in the pool. 
+				// Must call DirtyParticleData.Init + DirtyParticleData.Reset each 
+				// time you use an entry in the pool.
+				Pool.AddDefaulted(PreAllocCount);
+			}
+
+			void ResetPool() 
+			{ 
+				// try to shrink each (n) frames
+				MaxEntryCount = (EntryCount >= MaxEntryCount) ? EntryCount : MaxEntryCount;
+				
+				if (FrameCount % ChaosSolverParticlePoolNumFrameUntilShrink == 0)
+				{
+					int32 NextLowerBound = (Pool.Num() / PreAllocCount) - 1;
+					NextLowerBound = (NextLowerBound >= 1) ? NextLowerBound : 1;
+					NextLowerBound *= PreAllocCount;
+					if (MaxEntryCount < NextLowerBound)
+					{
+						Pool.SetNum(NextLowerBound);
+					}
+					MaxEntryCount = 0;
+				}
+
+				FrameCount++;
+				EntryCount = 0;
+			}
+
+			int32 GetEntryCount()
+			{
+				return EntryCount;
+			}
+
+			TPoolEntry& GetEntry(int32 Index)
+			{
+				ensure(Index < Pool.Num());
+				return Pool[Index];
+			}
+
+			TPoolEntry& GetNewEntry()
+			{
+				if (EntryCount >= Pool.Num())
+				{
+					Pool.Add(TPoolEntry());
+				}
+				return Pool[EntryCount++];
+			}
+
+		private:
+			TArray<TPoolEntry> Pool;
+			int32 EntryCount;
+			int32 MaxEntryCount;
+			int32 FrameCount;
+		};
+
+		TFramePool<Chaos::TPBDRigidParticleData<float, 3>, FSingleParticlePhysicsProxy< Chaos::TPBDRigidParticle<float, 3> >*, PBDRIGID_PREALLOC_COUNT> RigidParticlePool;
+		TFramePool<Chaos::TKinematicGeometryParticleData<float, 3>, FSingleParticlePhysicsProxy< Chaos::TKinematicGeometryParticle<float, 3> >*, KINEMATIC_GEOM_PREALLOC_COUNT> KinematicGeometryParticlePool;
+		TFramePool<Chaos::TGeometryParticleData<float, 3>, FSingleParticlePhysicsProxy< Chaos::TGeometryParticle<float, 3> >*, GEOMETRY_PREALLOC_COUNT> GeometryParticlePool;
 	};
 
 	template<>
diff --git a/Engine/Source/Runtime/Experimental/ChaosSolvers/Public/PhysicsProxy/SingleParticlePhysicsProxy.h b/Engine/Source/Runtime/Experimental/ChaosSolvers/Public/PhysicsProxy/SingleParticlePhysicsProxy.h
--- a/Engine/Source/Runtime/Experimental/ChaosSolvers/Public/PhysicsProxy/SingleParticlePhysicsProxy.h
+++ b/Engine/Source/Runtime/Experimental/ChaosSolvers/Public/PhysicsProxy/SingleParticlePhysicsProxy.h
@@ -145,6 +145,11 @@ public:
 	/**/
 	void ClearEvents();
 
+	PARTICLE_TYPE* GetParticle()
+	{
+		return Particle;
+	}
+	
 private:
 	bool bInitialized;
 	TArray<int32> InitializedIndices;
