Chaos
- fix twist axis calculation
- fix joint math for one swing axis free

diff --git a/Engine/Source/Runtime/Experimental/Chaos/Private/Chaos/Joint/PBDJointSolverGaussSeidel.cpp b/Engine/Source/Runtime/Experimental/Chaos/Private/Chaos/Joint/PBDJointSolverGaussSeidel.cpp
--- a/Engine/Source/Runtime/Experimental/Chaos/Private/Chaos/Joint/PBDJointSolverGaussSeidel.cpp
+++ b/Engine/Source/Runtime/Experimental/Chaos/Private/Chaos/Joint/PBDJointSolverGaussSeidel.cpp
@@ -86,7 +86,7 @@ namespace Chaos
 	}
 
 
-	FReal GetConeAngleError(
+	FReal GetConeAngleLimit(
 		const FPBDJointSettings& JointSettings,
 		const FVec3& SwingAxisLocal,
 		const FReal SwingAngle)
@@ -108,17 +108,7 @@ namespace Chaos
 			SwingAngleMax = FMath::Sqrt(Swing1Limit * DotSwing1 * Swing1Limit * DotSwing1 + Swing2Limit * DotSwing2 * Swing2Limit * DotSwing2);
 		}
 
-		// Calculate swing error we need to correct
-		if (SwingAngle > SwingAngleMax)
-		{
-			return SwingAngle - SwingAngleMax;
-		}
-		else if (SwingAngle < -SwingAngleMax)
-		{
-			return SwingAngle + SwingAngleMax;
-		}
-
-		return 0;
+		return SwingAngleMax;
 	}
 
 
@@ -238,7 +228,7 @@ namespace Chaos
 		FJointSolverResult NetResult;
 
 		NetResult += ApplyPositionConstraints(Dt, SolverSettings, JointSettings);
-	
+
 		NetResult += ApplyRotationConstraints(Dt, SolverSettings, JointSettings);
 
 		return NetResult;
@@ -280,7 +270,7 @@ namespace Chaos
 		const FPBDJointSolverSettings& SolverSettings,
 		const FPBDJointSettings& JointSettings)
 	{
-		// @todo(ccaulfield): eliminate branches in solver loop by pre-building list of functions to call?
+		// Locked axes always use hard constraints. Limited axes use hard or soft depending on settings
 		FJointSolverResult NetResult;
 
 		EJointMotionType TwistMotion = JointSettings.AngularMotionTypes[(int32)EJointAngularConstraintIndex::Twist];
@@ -292,13 +282,13 @@ namespace Chaos
 		// Apply twist constraint
 		if (SolverSettings.bEnableTwistLimits)
 		{
-			if ((TwistMotion == EJointMotionType::Limited) && bTwistSoft)
+			if (TwistMotion == EJointMotionType::Limited)
 			{
-				NetResult += ApplyTwistConstraintSoft(Dt, SolverSettings, JointSettings);
+				NetResult += ApplyTwistConstraint(Dt, SolverSettings, JointSettings, bTwistSoft);
 			}
-			else if (TwistMotion != EJointMotionType::Free)
+			else if (TwistMotion == EJointMotionType::Locked)
 			{
-				NetResult += ApplyTwistConstraint(Dt, SolverSettings, JointSettings);
+				NetResult += ApplyTwistConstraint(Dt, SolverSettings, JointSettings, false);
 			}
 		}
 
@@ -307,34 +297,41 @@ namespace Chaos
 		{
 			if ((Swing1Motion == EJointMotionType::Limited) && (Swing2Motion == EJointMotionType::Limited))
 			{
-				if (bSwingSoft)
-				{
-					NetResult += ApplyConeConstraintSoft(Dt, SolverSettings, JointSettings);
-				}
-				else
-				{
-					NetResult += ApplyConeConstraint(Dt, SolverSettings, JointSettings);
-				}
+				NetResult += ApplyConeConstraint(Dt, SolverSettings, JointSettings, bSwingSoft);
 			}
-			else
+			else if ((Swing1Motion == EJointMotionType::Limited) && (Swing2Motion == EJointMotionType::Locked))
+			{
+				NetResult += ApplyDualConeSwingConstraint(Dt, SolverSettings, JointSettings, EJointAngularConstraintIndex::Swing2, EJointAngularAxisIndex::Swing2, false);
+				NetResult += ApplySwingConstraint(Dt, SolverSettings, JointSettings, EJointAngularConstraintIndex::Swing1, EJointAngularAxisIndex::Swing1, bSwingSoft);
+			}
+			else if ((Swing1Motion == EJointMotionType::Limited) && (Swing2Motion == EJointMotionType::Free))
+			{
+				NetResult += ApplyDualConeSwingConstraint(Dt, SolverSettings, JointSettings, EJointAngularConstraintIndex::Swing1, EJointAngularAxisIndex::Swing1, bSwingSoft);
+			}
+			else if ((Swing1Motion == EJointMotionType::Locked) && (Swing2Motion == EJointMotionType::Limited))
+			{
+				NetResult += ApplyDualConeSwingConstraint(Dt, SolverSettings, JointSettings, EJointAngularConstraintIndex::Swing1, EJointAngularAxisIndex::Swing1, false);
+				NetResult += ApplySwingConstraint(Dt, SolverSettings, JointSettings, EJointAngularConstraintIndex::Swing2, EJointAngularAxisIndex::Swing2, bSwingSoft);
+			}
+			else if ((Swing1Motion == EJointMotionType::Locked) && (Swing2Motion == EJointMotionType::Locked))
+			{
+				NetResult += ApplySwingConstraint(Dt, SolverSettings, JointSettings, EJointAngularConstraintIndex::Swing1, EJointAngularAxisIndex::Swing1, false);
+				NetResult += ApplySwingConstraint(Dt, SolverSettings, JointSettings, EJointAngularConstraintIndex::Swing2, EJointAngularAxisIndex::Swing2, false);
+			}
+			else if ((Swing1Motion == EJointMotionType::Locked) && (Swing2Motion == EJointMotionType::Free))
+			{
+				NetResult += ApplyDualConeSwingConstraint(Dt, SolverSettings, JointSettings, EJointAngularConstraintIndex::Swing1, EJointAngularAxisIndex::Swing1, false);
+			}
+			else if ((Swing1Motion == EJointMotionType::Free) && (Swing2Motion == EJointMotionType::Limited))
+			{
+				NetResult += ApplyDualConeSwingConstraint(Dt, SolverSettings, JointSettings, EJointAngularConstraintIndex::Swing2, EJointAngularAxisIndex::Swing2, bSwingSoft);
+			}
+			else if ((Swing1Motion == EJointMotionType::Free) && (Swing2Motion == EJointMotionType::Locked))
+			{
+				NetResult += ApplyDualConeSwingConstraint(Dt, SolverSettings, JointSettings, EJointAngularConstraintIndex::Swing2, EJointAngularAxisIndex::Swing2, false);
+			}
+			else if ((Swing1Motion == EJointMotionType::Free) && (Swing2Motion == EJointMotionType::Free))
 			{
-				if ((Swing1Motion == EJointMotionType::Limited) && bSwingSoft)
-				{
-					NetResult += ApplySwingConstraintSoft(Dt, SolverSettings, JointSettings, EJointAngularConstraintIndex::Swing1, EJointAngularAxisIndex::Swing1);
-				}
-				else if (Swing1Motion != EJointMotionType::Free)
-				{
-					NetResult += ApplySwingConstraint(Dt, SolverSettings, JointSettings, EJointAngularConstraintIndex::Swing1, EJointAngularAxisIndex::Swing1);
-				}
-				
-				if ((Swing2Motion == EJointMotionType::Limited) && bSwingSoft)
-				{
-					NetResult += ApplySwingConstraintSoft(Dt, SolverSettings, JointSettings, EJointAngularConstraintIndex::Swing2, EJointAngularAxisIndex::Swing2);
-				}
-				else if (Swing2Motion != EJointMotionType::Free)
-				{
-					NetResult += ApplySwingConstraint(Dt, SolverSettings, JointSettings, EJointAngularConstraintIndex::Swing2, EJointAngularAxisIndex::Swing2);
-				}
 			}
 		}
 
@@ -594,6 +591,8 @@ namespace Chaos
 		const FReal Stiffness,
 		const FVec3& DP)
 	{
+		//UE_LOG(LogChaosJoint, VeryVerbose, TEXT("      Apply DP%d %f %f %f"), BodyIndex, DP.X, DP.Y, DP.Z);
+
 		Ps[BodyIndex] += Stiffness * DP;
 	}
 
@@ -603,6 +602,9 @@ namespace Chaos
 		const FVec3& DP0,
 		const FVec3& DP1)
 	{
+		//UE_LOG(LogChaosJoint, VeryVerbose, TEXT("      Apply DP%d %f %f %f"), 0, DP0.X, DP0.Y, DP0.Z);
+		//UE_LOG(LogChaosJoint, VeryVerbose, TEXT("      Apply DP%d %f %f %f"), 1, DP1.X, DP1.Y, DP1.Z);
+
 		Ps[0] += Stiffness * DP0;
 		Ps[1] += Stiffness * DP1;
 	}
@@ -613,6 +615,8 @@ namespace Chaos
 		const FReal Stiffness,
 		const FVec3& DR)
 	{
+		//UE_LOG(LogChaosJoint, VeryVerbose, TEXT("      Apply DR%d %f %f %f"), BodyIndex, DR.X, DR.Y, DR.Z);
+
 		const FRotation3 DQ = (FRotation3::FromElements(Stiffness * DR, 0) * Qs[BodyIndex]) * (FReal)0.5;
 		Qs[BodyIndex] = (Qs[BodyIndex] + DQ).GetNormalized();
 	}
@@ -623,6 +627,9 @@ namespace Chaos
 		const FVec3& DR0,
 		const FVec3& DR1)
 	{
+		//UE_LOG(LogChaosJoint, VeryVerbose, TEXT("      Apply DR%d %f %f %f"), 0, DR0.X, DR0.Y, DR0.Z);
+		//UE_LOG(LogChaosJoint, VeryVerbose, TEXT("      Apply DR%d %f %f %f"), 1, DR1.X, DR1.Y, DR1.Z);
+
 		const FRotation3 DQ0 = (FRotation3::FromElements(Stiffness * DR0, 0) * Qs[0]) * (FReal)0.5;
 		const FRotation3 DQ1 = (FRotation3::FromElements(Stiffness * DR1, 0) * Qs[1]) * (FReal)0.5;
 		Qs[0] = (Qs[0] + DQ0).GetNormalized();
@@ -636,6 +643,9 @@ namespace Chaos
 		const FVec3& DP,
 		const FVec3& DR)
 	{
+		//UE_LOG(LogChaosJoint, VeryVerbose, TEXT("      Apply DP%d %f %f %f"), BodyIndex, DP.X, DP.Y, DP.Z);
+		//UE_LOG(LogChaosJoint, VeryVerbose, TEXT("      Apply DR%d %f %f %f"), BodyIndex, DR.X, DR.Y, DR.Z);
+
 		Ps[BodyIndex] += Stiffness * DP;
 		const FRotation3 DQ = (FRotation3::FromElements(Stiffness * DR, 0) * Qs[BodyIndex]) * (FReal)0.5;
 		Qs[BodyIndex] = (Qs[BodyIndex] + DQ).GetNormalized();
@@ -776,24 +786,6 @@ namespace Chaos
 		//const FVec3 DR1 = Axis * -(Angle * II1 / (II0 + II1));
 
 		ApplyRotationDelta(Stiffness, DR0, DR1);
-
-		// Correct the positional error that was introduced by the rotation correction. This is
-		// correct when the position dofs are all locked, but not otherwise (fixable).
-		// This significantly improves angular stiffness at lower iterations, although the same effect
-		// is achieved by increasing iterations.
-		// @todo(ccaulfield): this position correction needs to have components in direction of inactive position constraints removed
-		//if (AngularPositionCorrection > 0)
-		//{
-		//	const FVec3 PrevX0 = Xs[0];
-		//	const FVec3 PrevX1 = Xs[1];
-		//	UpdateDerivedState();
-
-		//	const FVec3 DX = (Xs[1] - Xs[0]) - (PrevX1 - PrevX0);
-		//	const FVec3 DP0 = DX * (InvMs[0] / (InvMs[0] + InvMs[1]));
-		//	const FVec3 DP1 = DX * (-InvMs[1] / (InvMs[0] + InvMs[1]));
-		//	ApplyPositionDelta(Stiffness * AngularPositionCorrection, DP0, DP1);
-		//}
-
 		UpdateDerivedState();
 	}
 
@@ -826,19 +818,6 @@ namespace Chaos
 
 			// Joint-space inverse mass
 			FReal II1 = FVec3::DotProduct(Axis, IA1);
-
-			// If we are correcting the position, we need to adjust the constraint effective mass using parallel axis theorem
-			// @todo(ccaulfield): the IMs here are constant per constraint...pre-build and cache it
-			//if (AngularPositionCorrection > 0)
-			//{
-			//	const FVec3 LinearAxis1 = FVec3::CrossProduct(Xs[DIndex] - Ps[DIndex], Axis);
-			//	const FReal LinearAxisLen1Sq = LinearAxis1.SizeSquared();
-			//	if (LinearAxisLen1Sq > KINDA_SMALL_NUMBER)
-			//	{
-			//		const FReal IM1 = InvMs[DIndex] / LinearAxisLen1Sq;
-			//		II1 = II1 * IM1 / (II1 + IM1);
-			//	}
-			//}
 			const FReal II = II1;
 
 			// Damping angular velocity
@@ -860,22 +839,6 @@ namespace Chaos
 
 			Lambda += DLambda;
 			ApplyRotationDelta(DIndex, 1.0f, DR1);
-
-			// Correct the positional error that was introduced by the rotation correction. This is
-			// correct when the position dofs are all locked, but not otherwise (fixable).
-			// This significantly improves angular stiffness at lower iterations, although the same effect
-			// is achieved by increasing iterations.
-			// @todo(ccaulfield): this position correction needs to have components in direction of inactive position constraints removed
-			//if (AngularPositionCorrection > 0)
-			//{
-			//	const FVec3 PrevX0 = Xs[KIndex];
-			//	const FVec3 PrevX1 = Xs[DIndex];
-			//	UpdateDerivedState(DIndex);
-
-			//	const FVec3 DP1 = (PrevX1 - PrevX0) - (Xs[DIndex] - Xs[KIndex]);
-			//	ApplyPositionDelta(DIndex, AngularPositionCorrection, DP1);
-			//}
-
 			UpdateDerivedState(DIndex);
 		}
 	}
@@ -910,26 +873,6 @@ namespace Chaos
 			// Joint-space inverse mass
 			FReal II0 = FVec3::DotProduct(Axis, IA0);
 			FReal II1 = FVec3::DotProduct(Axis, IA1);
-
-			// If we are correcting the position, we need to adjust the constraint effective mass using parallel axis theorem
-			// @todo(ccaulfield): the IMs here are constant per constraint...pre-build and cache it
-			//if (AngularPositionCorrection > 0)
-			//{
-			//	const FVec3 LinearAxis0 = FVec3::CrossProduct(Xs[0] - Ps[0], Axis);
-			//	const FReal LinearAxisLen0Sq = LinearAxis0.SizeSquared();
-			//	if (LinearAxisLen0Sq > KINDA_SMALL_NUMBER)
-			//	{
-			//		const FReal IM0 = InvMs[0] / LinearAxisLen0Sq;
-			//		II0 = II0 * IM0 / (II0 + IM0);
-			//	}
-			//	const FVec3 LinearAxis1 = FVec3::CrossProduct(Xs[1] - Ps[1], Axis);
-			//	const FReal LinearAxisLen1Sq = LinearAxis1.SizeSquared();
-			//	if (LinearAxisLen1Sq > KINDA_SMALL_NUMBER)
-			//	{
-			//		const FReal IM1 = InvMs[1] / LinearAxisLen1Sq;
-			//		II1 = II1 * IM1 / (II1 + IM1);
-			//	}
-			//}
 			const FReal II = (II0 + II1);
 
 			// Damping angular velocity
@@ -954,24 +897,6 @@ namespace Chaos
 
 			Lambda += DLambda;
 			ApplyRotationDelta(1.0f, DR0, DR1);
-
-			// Correct the positional error that was introduced by the rotation correction. This is
-			// correct when the position dofs are all locked, but not otherwise (fixable).
-			// This significantly improves angular stiffness at lower iterations, although the same effect
-			// is achieved by increasing iterations.
-			// @todo(ccaulfield): this position correction needs to have components in direction of inactive position constraints removed
-			//if (AngularPositionCorrection > 0)
-			//{
-			//	const FVec3 PrevX0 = Xs[0];
-			//	const FVec3 PrevX1 = Xs[1];
-			//	UpdateDerivedState();
-
-			//	const FVec3 DX = (Xs[1] - Xs[0]) - (PrevX1 - PrevX0);
-			//	const FVec3 DP0 = DX * (InvMs[0] / (InvMs[0] + InvMs[1]));
-			//	const FVec3 DP1 = DX * (-InvMs[1] / (InvMs[0] + InvMs[1]));
-			//	ApplyPositionDelta(AngularPositionCorrection, DP0, DP1);
-			//}
-
 			UpdateDerivedState();
 		}
 	}
@@ -1009,7 +934,8 @@ namespace Chaos
 	FJointSolverResult FJointSolverGaussSeidel::ApplyTwistConstraint(
 		const FReal Dt,
 		const FPBDJointSolverSettings& SolverSettings,
-		const FPBDJointSettings& JointSettings)
+		const FPBDJointSettings& JointSettings,
+		const bool bUseSoftLimit)
 	{
 		FVec3 TwistAxis;
 		FReal TwistAngle;
@@ -1017,7 +943,7 @@ namespace Chaos
 
 		// Calculate the twist correction to apply to each body
 		FReal DTwistAngle = 0;
-		FReal TwistAngleMax = JointSettings.AngularLimits[(int32)EJointAngularConstraintIndex::Twist];
+		FReal TwistAngleMax = JointSettings.AngularLimits[(int32)EJointAngularConstraintIndex::Twist] + AngleTolerance;
 		if (TwistAngle > TwistAngleMax)
 		{
 			DTwistAngle = TwistAngle - TwistAngleMax;
@@ -1027,45 +953,23 @@ namespace Chaos
 			DTwistAngle = TwistAngle + TwistAngleMax;
 		}
 
-		// Apply twist correction
-		if (FMath::Abs(DTwistAngle) > AngleTolerance)
-		{
-			FReal TwistStiffness = FPBDJointUtilities::GetTwistStiffness(SolverSettings, JointSettings);
-			ApplyRotationConstraint(TwistStiffness, TwistAxis, DTwistAngle);
-			return FJointSolverResult::MakeActive();
-		}
-		return FJointSolverResult::MakeSolved();
-	}
-
-
-	FJointSolverResult FJointSolverGaussSeidel::ApplyTwistConstraintSoft(
-		const FReal Dt,
-		const FPBDJointSolverSettings& SolverSettings,
-		const FPBDJointSettings& JointSettings)
-	{
-		FVec3 TwistAxis;
-		FReal TwistAngle;
-		FPBDJointUtilities::GetTwistAxisAngle(Rs[0], Rs[1], TwistAxis, TwistAngle);
-
-		// Calculate the twist correction to apply to each body
-		FReal DTwistAngle = 0;
-		FReal TwistAngleMax = JointSettings.AngularLimits[(int32)EJointAngularConstraintIndex::Twist];
-		if (TwistAngle > TwistAngleMax)
-		{
-			DTwistAngle = TwistAngle - TwistAngleMax;
-		}
-		else if (TwistAngle < -TwistAngleMax)
-		{
-			DTwistAngle = TwistAngle + TwistAngleMax;
-		}
+		//UE_LOG(LogChaosJoint, VeryVerbose, TEXT("    SoftTwist Angle %f [Limit %f]"), FMath::RadiansToDegrees(TwistAngle), FMath::RadiansToDegrees(TwistAngleMax));
 
 		// Apply twist correction
-		if (FMath::Abs(DTwistAngle) > AngleTolerance)
+		if (FMath::Abs(DTwistAngle) > 0)
 		{
-			const FReal TwistStiffness = FPBDJointUtilities::GetSoftTwistStiffness(SolverSettings, JointSettings);
-			const FReal TwistDamping = FPBDJointUtilities::GetSoftTwistDamping(SolverSettings, JointSettings);
-			const bool bAccelerationMode = FPBDJointUtilities::GetAngularSoftAccelerationMode(SolverSettings, JointSettings);
-			ApplyRotationConstraintSoft(Dt, TwistStiffness, TwistDamping, bAccelerationMode, TwistAxis, DTwistAngle, TwistSoftLambda);
+			if (bUseSoftLimit)
+			{
+				const FReal TwistStiffness = FPBDJointUtilities::GetSoftTwistStiffness(SolverSettings, JointSettings);
+				const FReal TwistDamping = FPBDJointUtilities::GetSoftTwistDamping(SolverSettings, JointSettings);
+				const bool bAccelerationMode = FPBDJointUtilities::GetAngularSoftAccelerationMode(SolverSettings, JointSettings);
+				ApplyRotationConstraintSoft(Dt, TwistStiffness, TwistDamping, bAccelerationMode, TwistAxis, DTwistAngle, TwistSoftLambda);
+			}
+			else
+			{
+				FReal TwistStiffness = FPBDJointUtilities::GetTwistStiffness(SolverSettings, JointSettings);
+				ApplyRotationConstraint(TwistStiffness, TwistAxis, DTwistAngle);
+			}
 			return FJointSolverResult::MakeActive();
 		}
 		return FJointSolverResult::MakeSolved();
@@ -1109,7 +1013,8 @@ namespace Chaos
 	FJointSolverResult FJointSolverGaussSeidel::ApplyConeConstraint(
 		const FReal Dt,
 		const FPBDJointSolverSettings& SolverSettings,
-		const FPBDJointSettings& JointSettings)
+		const FPBDJointSettings& JointSettings,
+		const bool bUseSoftLimit)
 	{
 		// Calculate swing angle and axis
 		FReal SwingAngle;
@@ -1118,40 +1023,34 @@ namespace Chaos
 		const FVec3 SwingAxis = Rs[0] * SwingAxisLocal;
 
 		// Calculate swing angle error
-		FReal DSwingAngle = GetConeAngleError(JointSettings, SwingAxisLocal, SwingAngle);
-
-		// Apply swing correction to each body
-		if (FMath::Abs(DSwingAngle) > AngleTolerance)
+		FReal SwingAngleMax = GetConeAngleLimit(JointSettings, SwingAxisLocal, SwingAngle) + AngleTolerance;
+		FReal DSwingAngle = 0.0f;
+		if (SwingAngle > SwingAngleMax)
 		{
-			FReal SwingStiffness = FPBDJointUtilities::GetSwingStiffness(SolverSettings, JointSettings);
-			ApplyRotationConstraint(SwingStiffness, SwingAxis, DSwingAngle);
-			return FJointSolverResult::MakeActive();
+			DSwingAngle = SwingAngle - SwingAngleMax;
+		}
+		else if (SwingAngle < -SwingAngleMax)
+		{
+			DSwingAngle = SwingAngle + SwingAngleMax;
 		}
-		return FJointSolverResult::MakeSolved();
-	}
-
-
-	FJointSolverResult FJointSolverGaussSeidel::ApplyConeConstraintSoft(
-		const FReal Dt,
-		const FPBDJointSolverSettings& SolverSettings,
-		const FPBDJointSettings& JointSettings)
-	{
-		// Calculate swing angle and axis
-		FReal SwingAngle;
-		FVec3 SwingAxisLocal;
-		FPBDJointUtilities::GetConeAxisAngleLocal(Rs[0], Rs[1], SolverSettings.SwingTwistAngleTolerance, SwingAxisLocal, SwingAngle);
-		const FVec3 SwingAxis = Rs[0] * SwingAxisLocal;
 
-		// Calculate swing angle error
-		const FReal DSwingAngle = GetConeAngleError(JointSettings, SwingAxisLocal, SwingAngle);
+		//UE_LOG(LogChaosJoint, VeryVerbose, TEXT("    Cone Angle %f [Limit %f]"), FMath::RadiansToDegrees(SwingAngle), FMath::RadiansToDegrees(SwingAngleMax));
 
 		// Apply swing correction to each body
-		if (FMath::Abs(DSwingAngle) > AngleTolerance)
+		if (FMath::Abs(DSwingAngle) > 0)
 		{
-			const FReal SoftSwingStiffness = FPBDJointUtilities::GetSoftSwingStiffness(SolverSettings, JointSettings);
-			const FReal SoftSwingDamping = FPBDJointUtilities::GetSoftSwingDamping(SolverSettings, JointSettings);
-			const bool bAccelerationMode = FPBDJointUtilities::GetAngularSoftAccelerationMode(SolverSettings, JointSettings);
-			ApplyRotationConstraintSoft(Dt, SoftSwingStiffness, SoftSwingDamping, bAccelerationMode, SwingAxis, DSwingAngle, SwingSoftLambda);
+			if (bUseSoftLimit)
+			{
+				const FReal SoftSwingStiffness = FPBDJointUtilities::GetSoftSwingStiffness(SolverSettings, JointSettings);
+				const FReal SoftSwingDamping = FPBDJointUtilities::GetSoftSwingDamping(SolverSettings, JointSettings);
+				const bool bAccelerationMode = FPBDJointUtilities::GetAngularSoftAccelerationMode(SolverSettings, JointSettings);
+				ApplyRotationConstraintSoft(Dt, SoftSwingStiffness, SoftSwingDamping, bAccelerationMode, SwingAxis, DSwingAngle, SwingSoftLambda);
+			}
+			else
+			{
+				FReal SwingStiffness = FPBDJointUtilities::GetSwingStiffness(SolverSettings, JointSettings);
+				ApplyRotationConstraint(SwingStiffness, SwingAxis, DSwingAngle);
+			}
 			return FJointSolverResult::MakeActive();
 		}
 		return FJointSolverResult::MakeSolved();
@@ -1193,70 +1092,25 @@ namespace Chaos
 		const FPBDJointSolverSettings& SolverSettings,
 		const FPBDJointSettings& JointSettings)
 	{
-		// Calculate swing angle and axis
-		FReal SwingAngle;
-		FVec3 SwingAxisLocal;
-		FPBDJointUtilities::GetConeAxisAngleLocal(Rs[0], Rs[1], SolverSettings.SwingTwistAngleTolerance, SwingAxisLocal, SwingAngle);
-		const FVec3 SwingAxis = Rs[0] * SwingAxisLocal;
-
-		// Calculate swing angle error
-		const FReal DSwingAngle = GetConeAngleError(JointSettings, SwingAxisLocal, SwingAngle);
-
-		// Apply rotation to body 1, set relative angular velocity to zero, if it would increase the error
-		if (FMath::Abs(DSwingAngle) > AngleTolerance)
-		{
-			const FReal AngularProjection = FPBDJointUtilities::GetAngularProjection(SolverSettings, JointSettings);
-
-			const FVec3 DR1 = -DSwingAngle * SwingAxis;
-			const FVec3 DP1 = -FVec3::CrossProduct(DR1, Xs[1] - Ps[1]);
-			FVec3 DV1 = FVec3(0);
-			FVec3 DW1 = FVec3(0);
-
-			// Remove relative angular velocity if it would increase the constraint error
-			const FReal WAxis = FVec3::DotProduct(Ws[1] - Ws[0], SwingAxis);
-			if (FMath::Sign(WAxis) == FMath::Sign(DSwingAngle))
-			{
-				DW1 = -WAxis * SwingAxis;
-			}
-
-			// Remove relative velocity if it would increase the constraint error
-			FReal DP1Len = DP1.Size();
-			if (DP1Len > KINDA_SMALL_NUMBER)
-			{
-				const FVec3 DP1Dir = DP1 / DP1Len;
-				const FVec3 V0 = Vs[0] + FVec3::CrossProduct(Ws[0], Xs[0] - Ps[0] + Ps[1] - Xs[1]);
-				const FVec3 V1 = Vs[1];
-				const FReal VAxis = FVec3::DotProduct(V1 - V0, DP1Dir);
-				if (VAxis < 0)
-				{
-					DV1 = -VAxis * DP1Dir;
-				}
-			}
-
-			ApplyRotationDelta(1, AngularProjection, DR1);
-			ApplyPositionDelta(1, AngularProjection, DP1);
-			ApplyVelocityDelta(1, AngularProjection, DV1, DW1);
-			UpdateDerivedState(1);
-			return FJointSolverResult::MakeActive();
-		}
 		return FJointSolverResult::MakeSolved();
 	}
 
 
-	FJointSolverResult FJointSolverGaussSeidel::ApplySwingConstraint(
+	FJointSolverResult FJointSolverGaussSeidel::ApplyDualConeSwingConstraint(
 		const FReal Dt,
 		const FPBDJointSolverSettings& SolverSettings,
 		const FPBDJointSettings& JointSettings,
 		const EJointAngularConstraintIndex SwingConstraintIndex,
-		const EJointAngularAxisIndex SwingAxisIndex)
+		const EJointAngularAxisIndex SwingAxisIndex,
+		const bool bUseSoftLimit)
 	{
 		FVec3 SwingAxis;
 		FReal SwingAngle;
-		FPBDJointUtilities::GetSwingAxisAngle(Rs[0], Rs[1], SolverSettings.SwingTwistAngleTolerance, SwingConstraintIndex, SwingAxisIndex, SwingAxis, SwingAngle);
+		FPBDJointUtilities::GetDualConeSwingAxisAngle(Rs[0], Rs[1], SwingConstraintIndex, SwingAxisIndex, SwingAxis, SwingAngle);
 
 		// Calculate swing error we need to correct
 		FReal DSwingAngle = 0;
-		const FReal SwingAngleMax = JointSettings.AngularLimits[(int32)SwingConstraintIndex];
+		const FReal SwingAngleMax = JointSettings.AngularLimits[(int32)SwingConstraintIndex] + AngleTolerance;
 		if (SwingAngle > SwingAngleMax)
 		{
 			DSwingAngle = SwingAngle - SwingAngleMax;
@@ -1266,32 +1120,44 @@ namespace Chaos
 			DSwingAngle = SwingAngle + SwingAngleMax;
 		}
 
+		//UE_LOG(LogChaosJoint, VeryVerbose, TEXT("    Swing%d Angle %f [Limit %f]"), (SwingConstraintIndex == EJointAngularConstraintIndex::Swing1) ? 1 : 2, FMath::RadiansToDegrees(SwingAngle), FMath::RadiansToDegrees(SwingAngleMax));
+
 		// Apply swing correction
-		if (FMath::Abs(DSwingAngle) > AngleTolerance)
+		if (FMath::Abs(DSwingAngle) > 0)
 		{
-			const FReal SwingStiffness = FPBDJointUtilities::GetSwingStiffness(SolverSettings, JointSettings);
-			ApplyRotationConstraint(SwingStiffness, SwingAxis, DSwingAngle);
+			if (bUseSoftLimit)
+			{
+				const FReal SoftSwingStiffness = FPBDJointUtilities::GetSoftSwingStiffness(SolverSettings, JointSettings);
+				const FReal SoftSwingDamping = FPBDJointUtilities::GetSoftSwingDamping(SolverSettings, JointSettings);
+				const bool bAccelerationMode = FPBDJointUtilities::GetAngularSoftAccelerationMode(SolverSettings, JointSettings);
+				ApplyRotationConstraintSoft(Dt, SoftSwingStiffness, SoftSwingDamping, bAccelerationMode, SwingAxis, DSwingAngle, SwingSoftLambda);
+			}
+			else
+			{
+				const FReal SwingStiffness = FPBDJointUtilities::GetSwingStiffness(SolverSettings, JointSettings);
+				ApplyRotationConstraint(SwingStiffness, SwingAxis, DSwingAngle);
+			}
 			return FJointSolverResult::MakeActive();
 		}
 		return FJointSolverResult::MakeSolved();
 	}
 
 
-	FJointSolverResult FJointSolverGaussSeidel::ApplySwingConstraintSoft(
+	FJointSolverResult FJointSolverGaussSeidel::ApplySwingConstraint(
 		const FReal Dt,
 		const FPBDJointSolverSettings& SolverSettings,
 		const FPBDJointSettings& JointSettings,
 		const EJointAngularConstraintIndex SwingConstraintIndex,
-		const EJointAngularAxisIndex SwingAxisIndex)
+		const EJointAngularAxisIndex SwingAxisIndex,
+		const bool bUseSoftLimit)
 	{
 		FVec3 SwingAxis;
 		FReal SwingAngle;
 		FPBDJointUtilities::GetSwingAxisAngle(Rs[0], Rs[1], SolverSettings.SwingTwistAngleTolerance, SwingConstraintIndex, SwingAxisIndex, SwingAxis, SwingAngle);
 
-		const FReal SwingAngleMax = JointSettings.AngularLimits[(int32)SwingConstraintIndex];
-
 		// Calculate swing error we need to correct
 		FReal DSwingAngle = 0;
+		const FReal SwingAngleMax = JointSettings.AngularLimits[(int32)SwingConstraintIndex] + AngleTolerance;
 		if (SwingAngle > SwingAngleMax)
 		{
 			DSwingAngle = SwingAngle - SwingAngleMax;
@@ -1301,13 +1167,23 @@ namespace Chaos
 			DSwingAngle = SwingAngle + SwingAngleMax;
 		}
 
+		//UE_LOG(LogChaosJoint, VeryVerbose, TEXT("    Swing%d Angle %f [Limit %f]"), (SwingConstraintIndex == EJointAngularConstraintIndex::Swing1) ? 1 : 2, FMath::RadiansToDegrees(SwingAngle), FMath::RadiansToDegrees(SwingAngleMax));
+
 		// Apply swing correction
-		if (FMath::Abs(DSwingAngle) > AngleTolerance)
+		if (FMath::Abs(DSwingAngle) > 0)
 		{
-			const FReal SoftSwingStiffness = FPBDJointUtilities::GetSoftSwingStiffness(SolverSettings, JointSettings);
-			const FReal SoftSwingDamping = FPBDJointUtilities::GetSoftSwingDamping(SolverSettings, JointSettings);
-			const bool bAccelerationMode = FPBDJointUtilities::GetAngularSoftAccelerationMode(SolverSettings, JointSettings);
-			ApplyRotationConstraintSoft(Dt, SoftSwingStiffness, SoftSwingDamping, bAccelerationMode, SwingAxis, DSwingAngle, SwingSoftLambda);
+			if (bUseSoftLimit)
+			{
+				const FReal SoftSwingStiffness = FPBDJointUtilities::GetSoftSwingStiffness(SolverSettings, JointSettings);
+				const FReal SoftSwingDamping = FPBDJointUtilities::GetSoftSwingDamping(SolverSettings, JointSettings);
+				const bool bAccelerationMode = FPBDJointUtilities::GetAngularSoftAccelerationMode(SolverSettings, JointSettings);
+				ApplyRotationConstraintSoft(Dt, SoftSwingStiffness, SoftSwingDamping, bAccelerationMode, SwingAxis, DSwingAngle, SwingSoftLambda);
+			}
+			else
+			{
+				const FReal SwingStiffness = FPBDJointUtilities::GetSwingStiffness(SolverSettings, JointSettings);
+				ApplyRotationConstraint(SwingStiffness, SwingAxis, DSwingAngle);
+			}
 			return FJointSolverResult::MakeActive();
 		}
 		return FJointSolverResult::MakeSolved();
@@ -1381,6 +1257,7 @@ namespace Chaos
 		return FJointSolverResult::MakeSolved();
 	}
 
+
 	// Kinematic-Dynamic bodies
 	FJointSolverResult FJointSolverGaussSeidel::ApplyPointPositionConstraintKD(
 		const int32 KIndex,
@@ -1393,6 +1270,9 @@ namespace Chaos
 
 		FReal LinearStiffness = FPBDJointUtilities::GetLinearStiffness(SolverSettings, JointSettings);
 		const FVec3 CX = Xs[DIndex] - Xs[KIndex];
+
+		//UE_LOG(LogChaosJoint, VeryVerbose, TEXT("    PointKD Delta %f [Limit %f]"), CX.Size(), PositionTolerance);
+
 		if (CX.SizeSquared() > PositionTolerance * PositionTolerance)
 		{
 			if (bChaos_Joint_ISPC_Enabled)
@@ -1421,6 +1301,7 @@ namespace Chaos
 		return FJointSolverResult::MakeSolved();
 	}
 
+
 	// Dynamic-Dynamic bodies
 	FJointSolverResult FJointSolverGaussSeidel::ApplyPointPositionConstraintDD(
 		const FReal Dt,
@@ -1432,6 +1313,9 @@ namespace Chaos
 
 		FReal LinearStiffness = FPBDJointUtilities::GetLinearStiffness(SolverSettings, JointSettings);
 		const FVec3 CX = Xs[1] - Xs[0];
+
+		//UE_LOG(LogChaosJoint, VeryVerbose, TEXT("    PointKD Delta %f [Limit %f]"), CX.Size(), PositionTolerance);
+
 		if (CX.SizeSquared() > PositionTolerance * PositionTolerance)
 		{
 			if (bChaos_Joint_ISPC_Enabled)
diff --git a/Engine/Source/Runtime/Experimental/Chaos/Private/Chaos/PBDJointConstraintUtilities.cpp b/Engine/Source/Runtime/Experimental/Chaos/Private/Chaos/PBDJointConstraintUtilities.cpp
--- a/Engine/Source/Runtime/Experimental/Chaos/Private/Chaos/PBDJointConstraintUtilities.cpp
+++ b/Engine/Source/Runtime/Experimental/Chaos/Private/Chaos/PBDJointConstraintUtilities.cpp
@@ -4,32 +4,23 @@
 #include "Chaos/ParticleHandle.h"
 #include "Chaos/Utilities.h"
 
-//#pragma optimize("", off)
+//PRAGMA_DISABLE_OPTIMIZATION
 
 namespace Chaos
 {
-	void ToSwingTwistX(const FRotation3& Q, FRotation3& OutSwing, FRotation3& OutTwist)
+	void FPBDJointUtilities::DecomposeSwingTwistLocal(const FRotation3& R0, const FRotation3& R1, FRotation3& R01Swing, FRotation3& R01Twist)
 	{
-		FReal S = FMath::Sqrt(Q.X * Q.X + Q.W * Q.W);
-		if (S > KINDA_SMALL_NUMBER)
-		{
-			FReal SInv = 1.0f / S;
-			//OutTwist = FRotation3::FromElements(Q.X * SInv, 0, 0, Q.W * SInv);
-			//OutSwing = FRotation3::FromElements(0, (Q.W * Q.Y + Q.X * Q.Z) * SInv, (Q.W * Q.Z - Q.X * Q.Y) * SInv, S);
-			OutTwist = FRotation3::FromElements(Q.X * SInv, 0, 0, Q.W * SInv);
-			OutSwing = FRotation3::FromElements(0, (Q.W * Q.Y - Q.X * Q.Z) * SInv, (Q.W * Q.Z + Q.X * Q.Y) * SInv, S);
-		}
-		else
-		{
-			OutTwist = FRotation3::FromIdentity();
-			OutSwing = Q;
-		}
+		const FRotation3 R01 = R0.Inverse() * R1;
+		R01.ToSwingTwistX(R01Swing, R01Twist);
 	}
 
-	void FPBDJointUtilities::DecomposeSwingTwistLocal(const FRotation3& R0, const FRotation3& R1, FRotation3& R01Swing, FRotation3& R01Twist)
+	void FPBDJointUtilities::GetSwingTwistAngles(const FRotation3& R0, const FRotation3& R1, FReal& TwistAngle, FReal& Swing1Angle, FReal& Swing2Angle)
 	{
-		const FRotation3 R01 = R0.Inverse() * R1;
-		ToSwingTwistX(R01, R01Swing, R01Twist);
+		FRotation3 R01Twist, R01Swing;
+		FPBDJointUtilities::DecomposeSwingTwistLocal(R0, R1, R01Swing, R01Twist);
+		TwistAngle = R01Twist.GetAngle();
+		Swing1Angle = 4.0f * FMath::Atan2(R01Swing.Z, 1.0f + R01Swing.W);
+		Swing2Angle = 4.0f * FMath::Atan2(R01Swing.Y, 1.0f + R01Swing.W);
 	}
 
 	FReal FPBDJointUtilities::GetTwistAngle(const FRotation3& InTwist)
@@ -58,7 +49,7 @@ namespace Chaos
 		FRotation3 R01Twist, R01Swing;
 		FPBDJointUtilities::DecomposeSwingTwistLocal(R0, R1, R01Swing, R01Twist);
 
-		Axis = R0 * FJointConstants::TwistAxis();
+		Axis = R1 * FJointConstants::TwistAxis();
 		Angle = FPBDJointUtilities::GetTwistAngle(R01Twist);
 	}
 
@@ -82,6 +73,22 @@ namespace Chaos
 	}
 
 
+	void FPBDJointUtilities::GetDualConeSwingAxisAngle(
+		const FRotation3& R0,
+		const FRotation3& R1,
+		const EJointAngularConstraintIndex SwingConstraintIndex,
+		const EJointAngularAxisIndex SwingAxisIndex,
+		FVec3& Axis,
+		FReal& Angle)
+	{
+		FVec3 Twist1 = R1 * FJointConstants::TwistAxis();
+		FVec3 Swing0 = R0 * ((SwingConstraintIndex == EJointAngularConstraintIndex::Swing1)? FJointConstants::Swing2Axis() : FJointConstants::Swing1Axis());
+		Axis = FVec3::CrossProduct(Swing0, Twist1);
+		FReal SwingTwistDot = FVec3::DotProduct(Swing0, Twist1);
+		Angle = FMath::Asin(FMath::Clamp(-SwingTwistDot, -1.0f, 1.0f));
+	}
+
+
 	void FPBDJointUtilities::GetSwingAxisAngle(
 		const FRotation3& R0,
 		const FRotation3& R1,
@@ -91,64 +98,13 @@ namespace Chaos
 		FVec3& Axis,
 		FReal& Angle)
 	{
-#if 0
-		const FMatrix33 RM0 = R0.ToMatrix();
-		const FMatrix33 RM1 = R1.ToMatrix();
-		const FVec3& Twist0 = RM0.GetAxis((int32)EJointAngularAxisIndex::Twist);
-		const FVec3& Twist1 = RM1.GetAxis((int32)EJointAngularAxisIndex::Twist);
-
-		Axis = RM0.GetAxis((int32)SwingAxisIndex);
-		FReal SinAngle = FVec3::DotProduct(FVec3::CrossProduct(Twist1, Twist0), Axis);
-		Angle = FMath::Asin(FMath::Clamp(SinAngle, (FReal)-1, (FReal)1));
-#elif 0
 		// Decompose rotation of body 1 relative to body 0 into swing and twist rotations, assuming twist is X axis
 		FRotation3 R01Twist, R01Swing;
 		FPBDJointUtilities::DecomposeSwingTwistLocal(R0, R1, R01Swing, R01Twist);
 		const FReal R01SwingYorZ = ((int32)SwingAxisIndex == 2) ? R01Swing.Z : R01Swing.Y;	// Can't index a quat :(
-		FReal Angle = 4.0f * FMath::Atan2(R01SwingYorZ, 1.0f + R01Swing.W);
-		if (R01Twist.X < 0)
-		{
-			if (Angle > 0)
-			{
-				Angle = (FReal)PI - Angle;
-			}
-			else
-			{
-				Angle = (FReal)-PI - Angle;
-			}
-		}
+		Angle = 4.0f * FMath::Atan2(R01SwingYorZ, 1.0f + R01Swing.W);
 		const FVec3& AxisLocal = (SwingConstraintIndex == EJointAngularConstraintIndex::Swing1) ? FJointConstants::Swing1Axis() : FJointConstants::Swing2Axis();
 		Axis = R0 * AxisLocal;
-#else
-		// Do something better than this!!
-		FRotation3 R01Twist, R01Swing;
-		FPBDJointUtilities::DecomposeSwingTwistLocal(R0, R1, R01Swing, R01Twist);
-		FVec3 TwistAxis01 = FJointConstants::TwistAxis();
-		const FVec3 TwistAxis = R0 * TwistAxis01;
-
-		const FRotation3 R1NoTwist = R1 * R01Twist.Inverse();
-		const FMatrix33 Axes0 = R0.ToMatrix();
-		const FMatrix33 Axes1 = R1NoTwist.ToMatrix();
-
-		Axis = Axes0.GetAxis((int32)SwingAxisIndex);
-		Angle = 0;
-
-		const EJointAngularAxisIndex OtherSwingAxis = (SwingAxisIndex == EJointAngularAxisIndex::Swing1) ? EJointAngularAxisIndex::Swing2 : EJointAngularAxisIndex::Swing1;
-		FVec3 SwingCross = FVec3::CrossProduct(Axes0.GetAxis((int32)OtherSwingAxis), Axes1.GetAxis((int32)OtherSwingAxis));
-		SwingCross = SwingCross - FVec3::DotProduct(TwistAxis, SwingCross) * TwistAxis;
-		const FReal SwingCrossLen = SwingCross.Size();
-		if (SwingCrossLen > KINDA_SMALL_NUMBER)
-		{
-			Axis = SwingCross / SwingCrossLen;
-
-			Angle = FMath::Asin(FMath::Clamp(SwingCrossLen, (FReal)0, (FReal)1));
-			const FReal SwingDot = FVec3::DotProduct(Axes0.GetAxis((int32)OtherSwingAxis), Axes1.GetAxis((int32)OtherSwingAxis));
-			if (SwingDot < (FReal)0)
-			{
-				Angle = (FReal)PI - Angle;
-			}
-		}
-#endif
 	}
 
 
diff --git a/Engine/Source/Runtime/Experimental/Chaos/Private/Chaos/PBDJointConstraints.cpp b/Engine/Source/Runtime/Experimental/Chaos/Private/Chaos/PBDJointConstraints.cpp
--- a/Engine/Source/Runtime/Experimental/Chaos/Private/Chaos/PBDJointConstraints.cpp
+++ b/Engine/Source/Runtime/Experimental/Chaos/Private/Chaos/PBDJointConstraints.cpp
@@ -696,6 +696,8 @@ namespace Chaos
 		FJointSolverResult NetResult;
 		for (int32 PairIt = 0; PairIt < NumPairIts; ++PairIt)
 		{
+			UE_LOG(LogChaosJoint, VeryVerbose, TEXT("  Pair Iteration %d / %d"), PairIt, NumIts);
+
 			NetResult += Solver.ApplyConstraints(Dt, Settings, JointSettings);
 			NetResult +=  Solver.ApplyDrives(Dt, Settings, JointSettings);
 
@@ -714,7 +716,7 @@ namespace Chaos
 	FJointSolverResult FPBDJointConstraints::ProjectPosition_GaussSiedel(const FReal Dt, const int32 ConstraintIndex, const int32 NumPairIts, const int32 It, const int32 NumIts)
 	{
 		const TVector<TGeometryParticleHandle<FReal, 3>*, 2>& Constraint = ConstraintParticles[ConstraintIndex];
-		UE_LOG(LogChaosJoint, VeryVerbose, TEXT("Solve Joint Constraint %d %s %s (dt = %f; it = %d / %d)"), ConstraintIndex, *Constraint[0]->ToString(), *Constraint[1]->ToString(), Dt, It, NumIts);
+		UE_LOG(LogChaosJoint, VeryVerbose, TEXT("Project Joint Constraint %d %s %s (dt = %f; it = %d / %d)"), ConstraintIndex, *Constraint[0]->ToString(), *Constraint[1]->ToString(), Dt, It, NumIts);
 
 		const FPBDJointSettings& JointSettings = ConstraintSettings[ConstraintIndex];
 		FJointSolverGaussSeidel& Solver = ConstraintSolvers[ConstraintIndex];
diff --git a/Engine/Source/Runtime/Experimental/Chaos/Public/Chaos/Joint/PBDJointSolverGaussSeidel.h b/Engine/Source/Runtime/Experimental/Chaos/Public/Chaos/Joint/PBDJointSolverGaussSeidel.h
--- a/Engine/Source/Runtime/Experimental/Chaos/Public/Chaos/Joint/PBDJointSolverGaussSeidel.h
+++ b/Engine/Source/Runtime/Experimental/Chaos/Public/Chaos/Joint/PBDJointSolverGaussSeidel.h
@@ -250,12 +250,8 @@ namespace Chaos
 		FJointSolverResult ApplyTwistConstraint(
 			const FReal Dt,
 			const FPBDJointSolverSettings& SolverSettings,
-			const FPBDJointSettings& JointSettings);
-
-		FJointSolverResult ApplyTwistConstraintSoft(
-			const FReal Dt,
-			const FPBDJointSolverSettings& SolverSettings,
-			const FPBDJointSettings& JointSettings);
+			const FPBDJointSettings& JointSettings,
+			const bool bUseSoftLimit);
 
 		FJointSolverResult ApplyTwistDrive(
 			const FReal Dt,
@@ -270,12 +266,8 @@ namespace Chaos
 		FJointSolverResult ApplyConeConstraint(
 			const FReal Dt,
 			const FPBDJointSolverSettings& SolverSettings,
-			const FPBDJointSettings& JointSettings);
-
-		FJointSolverResult ApplyConeConstraintSoft(
-			const FReal Dt,
-			const FPBDJointSolverSettings& SolverSettings,
-			const FPBDJointSettings& JointSettings);
+			const FPBDJointSettings& JointSettings,
+			const bool bUseSoftLimit);
 
 		FJointSolverResult ApplyConeDrive(
 			const FReal Dt,
@@ -287,19 +279,23 @@ namespace Chaos
 			const FPBDJointSolverSettings& SolverSettings,
 			const FPBDJointSettings& JointSettings);
 
-		FJointSolverResult ApplySwingConstraint(
+		// One Swing axis is free, and the other limited or locked. This applies the lock/limit: Body1 Twist axis is confined to space between two cones.
+		FJointSolverResult ApplyDualConeSwingConstraint(
 			const FReal Dt,
 			const FPBDJointSolverSettings& SolverSettings,
 			const FPBDJointSettings& JointSettings,
 			const EJointAngularConstraintIndex SwingConstraintIndex,
-			const EJointAngularAxisIndex SwingAxisIndex);
+			const EJointAngularAxisIndex SwingAxisIndex,
+			const bool bUseSoftLimit);
 
-		FJointSolverResult ApplySwingConstraintSoft(
+		// One swing axis is locked, the other limited or locked. This applies the Limited axis (ApplyDualConeSwingConstraint is used for the locked axis).
+		FJointSolverResult ApplySwingConstraint(
 			const FReal Dt,
 			const FPBDJointSolverSettings& SolverSettings,
 			const FPBDJointSettings& JointSettings,
 			const EJointAngularConstraintIndex SwingConstraintIndex,
-			const EJointAngularAxisIndex SwingAxisIndex);
+			const EJointAngularAxisIndex SwingAxisIndex,
+			const bool bUseSoftLimit);
 
 		FJointSolverResult ApplySwingDrive(
 			const FReal Dt,
diff --git a/Engine/Source/Runtime/Experimental/Chaos/Public/Chaos/PBDJointConstraintUtilities.h b/Engine/Source/Runtime/Experimental/Chaos/Public/Chaos/PBDJointConstraintUtilities.h
--- a/Engine/Source/Runtime/Experimental/Chaos/Public/Chaos/PBDJointConstraintUtilities.h
+++ b/Engine/Source/Runtime/Experimental/Chaos/Public/Chaos/PBDJointConstraintUtilities.h
@@ -18,6 +18,13 @@ namespace Chaos
 			FRotation3& R01Swing, 
 			FRotation3& R01Twist);
 		
+		static void GetSwingTwistAngles(
+			const FRotation3& R0, 
+			const FRotation3& R1, 
+			FReal& TwistAngle, 
+			FReal& Swing1Angle, 
+			FReal& Swing2Angle);
+
 		static FReal GetTwistAngle(
 			const FRotation3& InTwist);
 		
@@ -34,6 +41,14 @@ namespace Chaos
 			FVec3& AxisLocal,
 			FReal& Angle);
 
+		static void GetDualConeSwingAxisAngle(
+			const FRotation3& R0,
+			const FRotation3& R1,
+			const EJointAngularConstraintIndex SwingConstraintIndex,
+			const EJointAngularAxisIndex SwingAxisIndex,
+			FVec3& Axis,
+			FReal& Angle);
+
 		static void GetSwingAxisAngle(
 			const FRotation3& R0,
 			const FRotation3& R1,
