Implement poor man's ccd at high velocity which moves object back when a collision is detected to prevent going through terrain.

diff --git a/Engine/Source/Runtime/Experimental/Chaos/Private/Chaos/CollisionResolution.cpp b/Engine/Source/Runtime/Experimental/Chaos/Private/Chaos/CollisionResolution.cpp
--- a/Engine/Source/Runtime/Experimental/Chaos/Private/Chaos/CollisionResolution.cpp
+++ b/Engine/Source/Runtime/Experimental/Chaos/Private/Chaos/CollisionResolution.cpp
@@ -29,7 +29,6 @@ DECLARE_CYCLE_STAT(TEXT("Collisions::GJK"), STAT_Collisions_GJK, STATGROUP_Chaos
 #define SCOPE_CYCLE_COUNTER_GJK()
 #endif
 
-//#pragma optimize("", off)
 //PRAGMA_DISABLE_OPTIMIZATION_ACTUAL
 
 
@@ -110,6 +109,40 @@ namespace Chaos
 			return Contact;
 		}
 
+		template <typename T, int d, typename GeometryA, typename GeometryB>
+		TContactPoint<T> GJKContactPointSwept(const GeometryA& A, const TRigidTransform<T, d>& ATM, const GeometryB& B, const TRigidTransform<T, d>& BTM, const TVector<T, d>& Dir, const T Length, const T CullDistance, T& TOI)
+		{
+			TContactPoint<T> Contact;
+			const TRigidTransform<T, d> AToBTM = ATM.GetRelativeTransform(BTM);
+			const TVector<T, d> LocalDir = BTM.InverseTransformVectorNoScale(Dir);
+
+			T OutTime;
+			TVec3<T> Location, Normal;
+			int32 NumIterations = 0;
+
+			if (GJKRaycast2(B, A, AToBTM, LocalDir, Length, OutTime, Location, Normal, (T)0, true))
+			{
+				Contact.Location = BTM.TransformPosition(Location);
+				Contact.Normal = BTM.TransformVectorNoScale(Normal);
+
+				// If OutTime < 0, initial overlap. We subtract length to get the total penetration at at end of frame.
+				// Use Abs of Dot as normal is inconsistent with initial overlaps.
+				if (OutTime >= 0.0f)
+				{
+					FReal Dot = FMath::Abs(FVec3::DotProduct(Contact.Normal, -Dir));
+					Contact.Phi = (OutTime - Length) * Dot;
+					TOI = OutTime / Length;
+				}
+				else
+				{
+					Contact.Phi = OutTime;
+					TOI = 0.0f;
+				}
+			}
+
+			return Contact;
+		}
+
 
 		template <typename GeometryA, typename GeometryB, typename T, int d>
 		TContactPoint<T> GJKImplicitContactPoint(const FImplicitObject& A, const TRigidTransform<T, d>& ATransform, const GeometryB& B, const TRigidTransform<T, d>& BTransform, const T CullDistance)
@@ -153,6 +186,180 @@ namespace Chaos
 			return Contact;
 		}
 
+		template <typename GeometryA, typename GeometryB, typename T, int d>
+		TContactPoint<T> GJKImplicitSweptContactPoint(const FImplicitObject& A, const TRigidTransform<T, d>& AStartTransform, const GeometryB& B, const TRigidTransform<T, d>& BTransform, const TVector<T, d>& Dir, const T Length, const T CullDistance, T& TOI)
+		{
+			TContactPoint<T> Contact;
+			const TRigidTransform<T, d> AToBTM = AStartTransform.GetRelativeTransform(BTransform);
+			const TVector<T, d> LocalDir = BTransform.InverseTransformVectorNoScale(Dir);
+
+			T OutTime = FLT_MAX;
+			int32 FaceIndex = -1;
+			TVec3<T> Location, Normal;
+			if (const TImplicitObjectScaled<GeometryA>* ScaledConvexImplicit = A.template GetObject<const TImplicitObjectScaled<GeometryA> >())
+			{
+
+				if (B.SweepGeom(*ScaledConvexImplicit, AToBTM, LocalDir, Length, OutTime, Location, Normal, FaceIndex, 0.0f, true))
+				{
+					Contact.Location = BTransform.TransformPosition(Location);
+					Contact.Normal = BTransform.TransformVectorNoScale(Normal);
+
+					// If OutTime < 0, initial overlap. We subtract length to get the total penetration at at end of frame.
+					// Use Abs of Dot as normal is inconsistent with initial overlaps.
+					if (OutTime >= 0.0f)
+					{
+						FReal Dot = FMath::Abs(FVec3::DotProduct(Contact.Normal, -Dir));
+						Contact.Phi = (OutTime - Length) * Dot;
+						TOI = OutTime / Length;
+					}
+					else
+					{
+						Contact.Phi = OutTime;
+						TOI = 0.0f;
+					}
+				}
+			}
+			else if (const TImplicitObjectInstanced<GeometryA>* InstancedConvexImplicit = A.template GetObject<const TImplicitObjectInstanced<GeometryA> >())
+			{
+				if (const GeometryA * InstancedInnerObject = static_cast<const GeometryA*>(InstancedConvexImplicit->GetInstancedObject()))
+				{
+					if (B.SweepGeom(*InstancedInnerObject, AToBTM, LocalDir, Length, OutTime, Location, Normal, FaceIndex, 0.0f, true))
+					{
+						Contact.Location = BTransform.TransformPosition(Location);
+						Contact.Normal = BTransform.TransformVectorNoScale(Normal);
+
+						// If OutTime < 0, initial overlap. We subtract length to get the total penetration at at end of frame.
+						// Use Abs of Dot as normal is inconsistent with initial overlaps.
+						if (OutTime >= 0.0f)
+						{
+							FReal Dot = FMath::Abs(FVec3::DotProduct(Contact.Normal, -Dir));
+							Contact.Phi = (OutTime - Length) * Dot;
+							TOI = OutTime / Length;
+						}
+						else
+						{
+							Contact.Phi = OutTime;
+							TOI = 0.0f;
+						}
+					}
+				}
+			}
+			else if (const GeometryA* ConvexImplicit = A.template GetObject<const GeometryA>())
+			{
+				if (B.SweepGeom(*ConvexImplicit, AToBTM, LocalDir, Length, OutTime, Location, Normal, FaceIndex, 0.0f, true))
+				{
+					Contact.Location = BTransform.TransformPosition(Location);
+					Contact.Normal = BTransform.TransformVectorNoScale(Normal);
+
+					// If OutTime < 0, initial overlap. We subtract length to get the total penetration at at end of frame.
+					// Use Abs of Dot as normal is inconsistent with initial overlaps.
+					if (OutTime >= 0.0f)
+					{
+						FReal Dot = FMath::Abs(FVec3::DotProduct(Contact.Normal, -Dir));
+						Contact.Phi = (OutTime - Length) * Dot;
+						TOI = OutTime / Length;
+					}
+					else
+					{
+						Contact.Phi = OutTime;;
+						TOI = 0.0f;
+					}
+				}
+			}
+
+			return Contact;
+		}
+
+		// A is the implicit here, we want to return a contact point on B (trimesh)
+		template <typename GeometryA, typename T, int d>
+		TContactPoint<T> GJKImplicitScaledTriMeshSweptContactPoint(const FImplicitObject& A, const TRigidTransform<T, d>& AStartTransform, const TImplicitObjectScaled<FTriangleMeshImplicitObject>& B, const TRigidTransform<T, d>& BTransform, const TVector<T,d>& Dir, const T Length, const T CullDistance, T& TOI)
+		{
+			TContactPoint<T> Contact;
+			const TRigidTransform<T, d> AToBTM = AStartTransform.GetRelativeTransform(BTransform);
+			const TVector<T, d> LocalDir = BTransform.InverseTransformVectorNoScale(Dir);
+
+			if (!ensure(B.GetType(true) & ImplicitObjectType::TriangleMesh) || !ensure(!IsInstanced(B.GetType(true))))
+			{
+				return TContactPoint<T>();
+			}
+
+			T OutTime = FLT_MAX;
+			TVec3<T> Location, Normal;
+			int32 FaceIndex = -1;
+			if (const TImplicitObjectScaled<GeometryA>* ScaledConvexImplicit = A.template GetObject<const TImplicitObjectScaled<GeometryA> >())
+			{
+				if (B.LowLevelSweepGeom(*ScaledConvexImplicit, AToBTM, LocalDir, Length, OutTime, Location, Normal, FaceIndex, 0.0f, true))
+				{
+					Contact.Location = BTransform.TransformPosition(Location);
+					Contact.Normal = BTransform.TransformVectorNoScale(Normal);
+
+					// If OutTime < 0, initial overlap. We subtract length to get the total penetration at at end of frame.
+					// Use Abs of Dot as normal is inconsistent with initial overlaps.
+					if (OutTime >= 0.0f)
+					{
+						FReal Dot = FMath::Abs(FVec3::DotProduct(Contact.Normal, -Dir));
+						Contact.Phi = (OutTime - Length) * Dot;
+						TOI = OutTime / Length;
+					}
+					else
+					{
+						Contact.Phi = OutTime;
+						TOI = 0.0f;
+					}
+				}
+			}
+			else if (const TImplicitObjectInstanced<GeometryA>* InstancedConvexImplicit = A.template GetObject<const TImplicitObjectInstanced<GeometryA> >())
+			{
+				if (const GeometryA * InstancedInnerObject = static_cast<const GeometryA*>(InstancedConvexImplicit->GetInstancedObject()))
+				{
+					if (B.LowLevelSweepGeom(*InstancedInnerObject, AToBTM, LocalDir, Length, OutTime, Location, Normal, FaceIndex, 0.0f, true))
+					{
+						Contact.Location = BTransform.TransformPosition(Location);
+						Contact.Normal = BTransform.TransformVectorNoScale(Normal);
+
+						// If OutTime < 0, initial overlap. We subtract length to get the total penetration at at end of frame.
+						// Use Abs of Dot as normal is inconsistent with initial overlaps.
+						if (OutTime >= 0.0f)
+						{
+							FReal Dot = FMath::Abs(FVec3::DotProduct(Contact.Normal, -Dir));
+							Contact.Phi = (OutTime - Length) * Dot;
+							TOI = OutTime / Length;
+						}
+						else
+						{
+							Contact.Phi = OutTime;
+							TOI = 0.0f;
+						}
+					}
+				}
+			}
+			else if (const GeometryA* ConvexImplicit = A.template GetObject<const GeometryA>())
+			{
+				if (B.LowLevelSweepGeom(*ConvexImplicit, AToBTM, LocalDir, Length, OutTime, Location, Normal, FaceIndex, 0.0f, true))
+				{
+						Contact.Location = BTransform.TransformPosition(Location);
+						Contact.Normal = BTransform.TransformVectorNoScale(Normal);
+
+						// If OutTime < 0, initial overlap. We subtract length to get the total penetration at at end of frame.
+						// Use Abs of Dot as normal is inconsistent with initial overlaps.
+						if (OutTime >= 0.0f)
+						{
+							FReal Dot = FMath::Abs(FVec3::DotProduct(Contact.Normal, -Dir));
+							Contact.Phi = (OutTime - Length) * Dot;
+							TOI = OutTime / Length;
+						}
+						else
+						{
+							Contact.Phi = OutTime;
+							TOI = 0.0f;
+						}
+				}
+			}
+
+			return Contact;
+		}
+
+
 
 		// This is pretty unnecessary - all instanced shapes have the same implementation so we should be able to
 		// collapse this switch into a generic call. Maybe add a base class to TImplicitObjectInstanced.
@@ -213,6 +420,18 @@ namespace Chaos
 			return ContactPoint;
 		}
 
+		template<class T, int d>
+		TContactPoint<T> ConvexConvexContactPointSwept(const FImplicitObject& A, const TRigidTransform<T, d>& ATM, const FImplicitObject& B, const TRigidTransform<T, d>& BTM, const TVector<T, d>& Dir, const T Length, const T CullDistance, T& TOI)
+		{
+			return Utilities::CastHelper(A, ATM, [&](const auto& ADowncast, const TRigidTransform<T,d>& AFullTM)
+			{
+				return Utilities::CastHelper(B, BTM, [&](const auto& BDowncast, const TRigidTransform<T,d>& BFullTM)
+				{
+					return GJKContactPointSwept(ADowncast, AFullTM, BDowncast, BFullTM, Dir, Length, CullDistance, TOI);
+				});
+			});
+		}
+
 		template <typename T, int d>
 		void UpdateSingleShotManifold(TRigidBodyMultiPointContactConstraint<T, d>& Constraint, const TRigidTransform<T, d>& Transform0, const TRigidTransform<T, d>& Transform1, const T CullDistance)
 		{
@@ -1391,8 +1610,7 @@ namespace Chaos
 		{
 			UpdateContactPoint(Constraint.Manifold, ConvexTriangleMeshContactPoint(Convex0, Transform0, TriangleMesh1, Transform1, CullDistance));
 		}
-
-
+		
 		template <typename T, int d>
 		void UpdateConvexTriangleMeshManifold(TCollisionConstraintBase<T, d>&  Constraint, const TRigidTransform<T, d>& ATM, const TRigidTransform<T, d>& BTM, const T CullDistance)
 		{
@@ -2041,6 +2259,235 @@ namespace Chaos
 			ConstructConstraintsImpl(Particle0, Particle1, Implicit0, Implicit1, Transform0, Transform1, CullDistance, Context, NewConstraints);
 		}
 
+
+		template <typename GeometryA, typename GeometryB>
+		bool GetPairTOIHackImpl(const FImplicitObject& A, const FRigidTransform3& AStartTransform, const GeometryB& B, const FRigidTransform3& BTransform, const FVec3& Dir, const FReal Length, FReal& OutTOI, FVec3& OutNormal, FReal& OutPhi)
+		{
+			FReal CullDistance = 0.0f;
+
+			FReal TOI = 1.0f;
+			TContactPoint<FReal> Contact = GJKImplicitSweptContactPoint<GeometryA>(A, AStartTransform, B, BTransform, Dir, Length, CullDistance, TOI);
+			if (Contact.Phi < OutPhi)
+			{
+				OutPhi = Contact.Phi;
+				OutNormal = Contact.Normal;
+				OutTOI = TOI;
+				return true;
+			}
+			return false;
+		}
+
+		template <typename GeometryA>
+		bool GetPairTOIHackTriMeshImpl(const FImplicitObject& A, const FRigidTransform3& AStartTransform, const FImplicitObject* B, const FRigidTransform3& BTransform, const FVec3& Dir, const FReal Length, FReal& OutTOI, FVec3& OutNormal, FReal& OutPhi)
+		{
+			TContactPoint<FReal> Contact;
+			FReal TOI = 0.0f;
+			FReal CullDistance = 0.0f;
+
+			EImplicitObjectType TriMeshType = B->GetType();
+			if (const TImplicitObjectScaled<FTriangleMeshImplicitObject>* ScaledTriangleMesh = B->template GetObject<const TImplicitObjectScaled<FTriangleMeshImplicitObject>>())
+			{
+				Contact = GJKImplicitScaledTriMeshSweptContactPoint<GeometryA>(A, AStartTransform, *ScaledTriangleMesh, BTransform, Dir, Length, CullDistance, TOI);
+			}
+			else if (const FTriangleMeshImplicitObject* TriangleMesh = B->template GetObject<const FTriangleMeshImplicitObject>())
+			{
+				Contact = GJKImplicitSweptContactPoint<GeometryA>(A, AStartTransform, *TriangleMesh, BTransform, Dir, Length, CullDistance, TOI);
+			}
+
+			if (Contact.Phi < OutPhi)
+			{
+				OutPhi = Contact.Phi;
+				OutNormal = Contact.Normal;
+				OutTOI = TOI;
+				return true;
+			}
+			return false;
+		}
+
+		bool GetPairTOIHackConvexConvexImpl(const FImplicitObject& A, const FRigidTransform3& AStartTransform, const FImplicitObject& B, const FRigidTransform3& BTransform, const FVec3& Dir, const FReal Length, FReal& OutTOI, FVec3& OutNormal, FReal& OutPhi)
+		{
+			FReal CullDistance = 0.0f;
+			FReal TOI = 0.0f;
+
+			TContactPoint<FReal> Contact = ConvexConvexContactPointSwept(A, AStartTransform, B, BTransform, Dir, Length, CullDistance, TOI);
+			if (Contact.Phi < OutPhi)
+			{
+				OutPhi = Contact.Phi;
+				OutNormal = Contact.Normal;
+				OutTOI = TOI;
+				return true;
+			}
+			return false;
+		}
+
+
+
+		bool GetPairTOIHack(const TPBDRigidParticleHandle<FReal, 3>* Particle0, const TGeometryParticleHandle<FReal, 3>* Particle1, const FImplicitObject* Implicit0, const FImplicitObject* Implicit1, const FRigidTransform3& StartTransform0, const FRigidTransform3& Transform1, FReal& OutTOI, FVec3& OutNormal, FReal& OutPhi)
+		{
+			// @todo(chaos): We use GetInnerType here because TriMeshes are left with their "Instanced" wrapper, unlike all other instanced implicits. Should we strip the instance on Tri Mesh too?
+			EImplicitObjectType Implicit0Type = Implicit0 ? GetInnerType(Implicit0->GetType()) : ImplicitObjectType::Unknown;
+			EImplicitObjectType Implicit1Type = Implicit1 ? GetInnerType(Implicit1->GetType()) : ImplicitObjectType::Unknown;
+			bool bIsConvex0 = Implicit0 && Implicit0->IsConvex();
+			bool bIsConvex1 = Implicit1 && Implicit1->IsConvex();
+
+			FVec3 Dir = Particle0->P() - Particle0->X();
+			FReal Length = Dir.Size();
+			if (Length < KINDA_SMALL_NUMBER)
+			{
+				return false;
+			}
+			Dir /= Length;
+
+			if (Implicit0Type == TBox<FReal, 3>::StaticType() && Implicit1Type == THeightField<FReal>::StaticType())
+			{
+				return GetPairTOIHackImpl<TBox<FReal, 3>>(*Implicit0->template GetObject<TBox<FReal, 3>>(), StartTransform0, *Implicit1->template GetObject<THeightField<FReal>>(), Transform1, Dir, Length, OutTOI, OutNormal, OutPhi);
+			}
+			else if (Implicit0Type == TSphere<FReal, 3>::StaticType() && Implicit1Type == THeightField<FReal>::StaticType())
+			{
+				return GetPairTOIHackImpl<TSphere<FReal, 3>>(*Implicit0->template GetObject<TSphere<FReal, 3>>(), StartTransform0, *Implicit1->template GetObject<THeightField<FReal>>(), Transform1, Dir, Length, OutTOI, OutNormal, OutPhi);
+			}
+			else if (Implicit0Type == TCapsule<FReal>::StaticType() && Implicit1Type == THeightField<FReal>::StaticType())
+			{
+				return GetPairTOIHackImpl<TCapsule<FReal>>(*Implicit0->template GetObject<TCapsule<FReal>>(), StartTransform0, *Implicit1->template GetObject<THeightField<FReal>>(), Transform1, Dir, Length, OutTOI, OutNormal, OutPhi);
+			}
+			else if (Implicit0Type == TBox<FReal, 3>::StaticType() && Implicit1Type == FTriangleMeshImplicitObject::StaticType())
+			{
+				return GetPairTOIHackTriMeshImpl<TBox<FReal, 3>>(*Implicit0->template GetObject<TBox<FReal, 3>>(), StartTransform0, Implicit1, Transform1, Dir, Length, OutTOI, OutNormal, OutPhi);
+			}
+			else if (Implicit0Type == TSphere<FReal, 3>::StaticType() && Implicit1Type == FTriangleMeshImplicitObject::StaticType())
+			{
+				return GetPairTOIHackTriMeshImpl<TSphere<FReal, 3>>(*Implicit0->template GetObject<TSphere<FReal, 3>>(), StartTransform0, Implicit1, Transform1, Dir, Length, OutTOI, OutNormal, OutPhi);
+			}
+			if (Implicit0Type == TCapsule<FReal>::StaticType() && Implicit1Type == FTriangleMeshImplicitObject::StaticType())
+			{
+				return GetPairTOIHackTriMeshImpl<TCapsule<FReal>>(*Implicit0->template GetObject<TCapsule<FReal>>(), StartTransform0, Implicit1, Transform1, Dir, Length, OutTOI, OutNormal, OutPhi);
+			}
+			else if (bIsConvex0 && Implicit1Type == THeightField<FReal>::StaticType())
+			{
+				return GetPairTOIHackImpl<FConvex>(*Implicit0, StartTransform0, *Implicit1->template GetObject<THeightField<FReal>>(), Transform1, Dir, Length, OutTOI, OutNormal, OutPhi);
+			}
+			else if (bIsConvex0 && Implicit1Type == FTriangleMeshImplicitObject::StaticType())
+			{
+				return GetPairTOIHackTriMeshImpl<FConvex>(*Implicit0, StartTransform0, Implicit1, Transform1, Dir, Length, OutTOI, OutNormal, OutPhi);
+			}
+			else if (bIsConvex0 && bIsConvex1)
+			{
+				return GetPairTOIHackConvexConvexImpl(*Implicit0, StartTransform0, *Implicit1, Transform1, Dir, Length, OutTOI, OutNormal, OutPhi);
+			}
+
+			return false;
+		}
+
+
+		bool GetTOIHackImpl(const TPBDRigidParticleHandle<FReal, 3>* Particle0, const TGeometryParticleHandle<FReal, 3>* Particle1, const FImplicitObject* Implicit0, const FImplicitObject* Implicit1, const FRigidTransform3& Transform0, const FRigidTransform3& Transform1, FReal& OutTOI, FVec3& OutNormal, FReal& OutPhi)
+		{
+			EImplicitObjectType Implicit0Type = Implicit0 ? GetInnerType(Implicit0->GetType()) : ImplicitObjectType::Unknown;
+			EImplicitObjectType Implicit1Type = Implicit1 ? GetInnerType(Implicit1->GetType()) : ImplicitObjectType::Unknown;
+			
+			if (!Implicit0 || !Implicit1)
+			{
+				return false;
+			}
+
+			EImplicitObjectType Implicit0OuterType = Implicit0->GetType();
+			EImplicitObjectType Implicit1OuterType = Implicit1->GetType();
+
+			// Handle transform wrapper shape
+			if ((Implicit0OuterType == TImplicitObjectTransformed<FReal, 3>::StaticType()) && (Implicit1OuterType == TImplicitObjectTransformed<FReal, 3>::StaticType()))
+			{
+				const TImplicitObjectTransformed<FReal, 3>* TransformedImplicit0 = Implicit0->template GetObject<const TImplicitObjectTransformed<FReal, 3>>();
+				const TImplicitObjectTransformed<FReal, 3>* TransformedImplicit1 = Implicit1->template GetObject<const TImplicitObjectTransformed<FReal, 3>>();
+				TRigidTransform<FReal, 3> TransformedTransform0 = TransformedImplicit0->GetTransform() * Transform0;
+				TRigidTransform<FReal, 3> TransformedTransform1 = TransformedImplicit1->GetTransform() * Transform1;
+				return GetTOIHackImpl(Particle0, Particle1, TransformedImplicit0->GetTransformedObject(), TransformedImplicit1->GetTransformedObject(), TransformedTransform0, TransformedTransform1, OutTOI, OutNormal, OutPhi);
+			}
+			else if (Implicit0OuterType == TImplicitObjectTransformed<FReal, 3>::StaticType())
+			{
+				const TImplicitObjectTransformed<FReal, 3>* TransformedImplicit0 = Implicit0->template GetObject<const TImplicitObjectTransformed<FReal, 3>>();
+				TRigidTransform<FReal, 3> TransformedTransform0 = TransformedImplicit0->GetTransform() * Transform0;
+				return GetTOIHackImpl(Particle0, Particle1, TransformedImplicit0->GetTransformedObject(), Implicit1, TransformedTransform0, Transform1, OutTOI, OutNormal, OutPhi);
+			}
+			else if (Implicit1OuterType == TImplicitObjectTransformed<FReal, 3>::StaticType())
+			{
+				const TImplicitObjectTransformed<FReal, 3>* TransformedImplicit1 = Implicit1->template GetObject<const TImplicitObjectTransformed<FReal, 3>>();
+				TRigidTransform<FReal, 3> TransformedTransform1 = TransformedImplicit1->GetTransform() * Transform1;
+				return GetTOIHackImpl(Particle0, Particle1, Implicit0, TransformedImplicit1->GetTransformedObject(), Transform0, TransformedTransform1, OutTOI, OutNormal, OutPhi);
+			}
+			// Handle Instanced shapes
+			// NOTE: Tri Meshes are handled differently. We should probably do something about this...
+			if (((uint32)Implicit0OuterType & ImplicitObjectType::IsInstanced) || ((uint32)Implicit1OuterType & ImplicitObjectType::IsInstanced))
+			{
+				const FImplicitObject* InnerImplicit0 = Implicit0;
+				const FImplicitObject* InnerImplicit1 = Implicit1;
+				if ((uint32)Implicit0OuterType & ImplicitObjectType::IsInstanced)
+				{
+					InnerImplicit0 = GetInstancedImplicit(Implicit0);
+				}
+				if ((uint32)Implicit1OuterType & ImplicitObjectType::IsInstanced)
+				{
+					InnerImplicit1 = GetInstancedImplicit(Implicit1);
+				}
+				if (InnerImplicit0 && InnerImplicit1)
+				{
+					return GetTOIHackImpl(Particle0, Particle1, InnerImplicit0, InnerImplicit1, Transform0, Transform1, OutTOI, OutNormal, OutPhi);
+				}
+				else if (InnerImplicit0 && !InnerImplicit1)
+				{
+					return GetTOIHackImpl(Particle0, Particle1, InnerImplicit0, Implicit1, Transform0, Transform1, OutTOI, OutNormal, OutPhi);
+				}
+				else if (!InnerImplicit0 && InnerImplicit1)
+				{
+					return GetTOIHackImpl(Particle0, Particle1, Implicit0, InnerImplicit1, Transform0, Transform1, OutTOI, OutNormal, OutPhi);
+				}
+			}
+
+			// Handle Unions
+			if (Implicit0OuterType == FImplicitObjectUnion::StaticType())
+			{
+				bool bHit = false;
+				const FImplicitObjectUnion* Union0 = Implicit0->template GetObject<FImplicitObjectUnion>();
+				for (const auto& Child0 : Union0->GetObjects())
+				{
+					if (GetTOIHackImpl(Particle0, Particle1, Child0.Get(), Implicit1, Transform0, Transform1, OutTOI, OutNormal, OutPhi))
+					{
+						bHit = true;
+					}
+				}
+				return bHit;
+			}
+			if (Implicit1OuterType == FImplicitObjectUnion::StaticType())
+			{
+				bool bHit = false;
+				const FImplicitObjectUnion* Union1 = Implicit1->template GetObject<FImplicitObjectUnion>();
+				for (const auto& Child1 : Union1->GetObjects())
+				{
+					if (GetTOIHackImpl(Particle0, Particle1, Implicit0, Child1.Get(), Transform0, Transform1, OutTOI, OutNormal, OutPhi))
+					{
+						bHit = true;
+					}
+				}
+				return bHit;
+			}
+
+			// Check shape pair filtering if enable
+			if (!DoCollide(Implicit0Type, const_cast<TPBDRigidParticleHandle<FReal, 3>*>(Particle0)->GetImplicitShape(Implicit0), Implicit1Type, const_cast<TGeometryParticleHandle<FReal, 3>*>(Particle1)->GetImplicitShape(Implicit1)))
+			{
+				return false;
+			}
+
+			// If we get here, we have a pair of concrete shapes (i.e., no wrappers or containers)
+			// Create a constraint for the shape pair
+			return GetPairTOIHack(Particle0, Particle1, Implicit0, Implicit1, Transform0, Transform1, OutTOI, OutNormal, OutPhi);
+		}
+
+		bool GetTOIHack(const TPBDRigidParticleHandle<FReal, 3>* Particle0, const TGeometryParticleHandle<FReal, 3>* Particle1, FReal& OutTOI, FVec3& OutNormal, FReal& OutPhi)
+		{
+			OutPhi = 0.0f;	// Cull Distance
+
+			FRigidTransform3 Particle0StartTransform = FRigidTransform3(Particle0->X(), Particle0->Q());	// Start Pos, Final Rot
+			return GetTOIHackImpl(Particle0, Particle1, Particle0->Geometry().Get(), Particle1->Geometry().Get(), Particle0StartTransform, Collisions::GetTransform(Particle1), OutTOI, OutNormal, OutPhi);
+		}
+
 		template void UpdateBoxBoxConstraint<float, 3>(const TAABB<float, 3>& Box1, const TRigidTransform<float, 3>& Box1Transform, const TAABB<float, 3>& Box2, const TRigidTransform<float, 3>& Box2Transform, const float CullDistance, TRigidBodyPointContactConstraint<float, 3>& Constraint);
 		template void UpdateBoxBoxManifold<float, 3>(TCollisionConstraintBase<float, 3>&  Constraint, const TRigidTransform<float, 3>& ATM, const TRigidTransform<float, 3>& BTM, const float CullDistance);
 		template void ConstructBoxBoxConstraints<float, 3>(TGeometryParticleHandle<float, 3>* Particle0, TGeometryParticleHandle<float, 3>* Particle1, const FImplicitObject* Implicit0, const FImplicitObject* Implicit1, const TRigidTransform<float, 3>& Transform0, const TRigidTransform<float, 3>& Transform1, const float CullDistance, FCollisionConstraintsArray& NewConstraints);
diff --git a/Engine/Source/Runtime/Experimental/Chaos/Private/Chaos/PBDRigidsEvolutionGBF.cpp b/Engine/Source/Runtime/Experimental/Chaos/Private/Chaos/PBDRigidsEvolutionGBF.cpp
--- a/Engine/Source/Runtime/Experimental/Chaos/Private/Chaos/PBDRigidsEvolutionGBF.cpp
+++ b/Engine/Source/Runtime/Experimental/Chaos/Private/Chaos/PBDRigidsEvolutionGBF.cpp
@@ -21,6 +21,8 @@
 #include "Misc/Paths.h"
 #include "HAL/FileManager.h"
 
+//PRAGMA_DISABLE_OPTIMIZATION
+
 namespace Chaos
 {
 #if !UE_BUILD_SHIPPING
@@ -50,6 +52,12 @@ float BoundsThicknessVelocityMultiplier = 2.0f;
 FAutoConsoleVariableRef CVarBoundsThickness(TEXT("p.CollisionBoundsThickness"), BoundsThickness, TEXT("Collision inflation for speculative contact generation.[def:0.0]"));
 FAutoConsoleVariableRef CVarBoundsThicknessVelocityMultiplier(TEXT("p.CollisionBoundsVelocityInflation"), BoundsThicknessVelocityMultiplier, TEXT("Collision velocity inflation for speculatibe contact generation.[def:2.0]"));
 
+float HackCCD_EnableThreshold = 0.3f;
+FAutoConsoleVariableRef CVarHackCCDVelThreshold(TEXT("p.Chaos.CCD.EnableThreshold"), HackCCD_EnableThreshold, TEXT("If distance moved is greater than this times the minimum object dimension, use CCD"));
+
+float HackCCD_DepthThreshold = 0.05f;
+FAutoConsoleVariableRef CVarHackCCDDepthThreshold(TEXT("p.Chaos.CCD.DepthThreshold"), HackCCD_DepthThreshold, TEXT("When returning to TOI, leave this much contact depth (as a fraction of MinBounds)"));
+
 
 DECLARE_CYCLE_STAT(TEXT("TPBDRigidsEvolutionGBF::AdvanceOneTimeStep"), STAT_Evolution_AdvanceOneTimeStep, STATGROUP_Chaos);
 DECLARE_CYCLE_STAT(TEXT("TPBDRigidsEvolutionGBF::Integrate"), STAT_Evolution_Integrate, STATGROUP_Chaos);
@@ -129,6 +137,116 @@ void TPBDRigidsEvolutionGBF<T, d>::Advance(const T Dt, const T MaxStepDt, const
 	}
 }
 
+//
+//
+//
+// BEGIN TOI HACK STUFF
+//
+//
+//
+
+namespace Collisions
+{
+	extern bool GetTOIHack(const TPBDRigidParticleHandle<FReal, 3>* Particle0, const TGeometryParticleHandle<FReal, 3>* Particle1, FReal& OutTOI, FVec3& OutNormal, FReal& OutPhi);
+}
+
+void PostMoveToTOIHack(FReal Dt, TPBDRigidParticleHandle<FReal, 3>* Particle)
+{
+	// Update bounds, velocity etc
+	TPerParticlePBDUpdateFromDeltaPosition<FReal, 3> VelUpdateRule;
+	VelUpdateRule.Apply(Particle, Dt);
+}
+
+void MoveToTOIPairHack(FReal Dt, TPBDRigidParticleHandle<FReal, 3>* Particle1, const TGeometryParticleHandle<FReal, 3>* Particle2)
+{
+	FReal TOI = 0.0f;
+	FReal Phi = 0.0f;
+	FVec3 Normal = FVec3(0);
+	if (Collisions::GetTOIHack(Particle1, Particle2, TOI, Normal, Phi))
+	{
+		FReal Depth = -Phi;
+		FReal MinBounds = Particle1->Geometry()->BoundingBox().Extents().Min();
+		FReal MaxDepth = MinBounds * HackCCD_DepthThreshold;
+
+		//UE_LOG(LogChaos, Warning, TEXT("MoveTOIHAck: TOI %f; Depth %f"), TOI, Depth);
+
+		if ((Depth > MaxDepth) && (TOI > KINDA_SMALL_NUMBER) && (TOI < 1.0f))
+		{
+			// Move the particle to just after the TOI so we still have a collision to resolve but won't get tunneling
+			// The time after TOI we want is given by the ratio of max acceptable depth to depth at T = 1
+			FReal ExtraT = (1.0f - TOI) * MaxDepth / Depth;
+			FReal FinalT = TOI + ExtraT;
+			FVec3 FinalP = FVec3::Lerp(Particle1->X(), Particle1->P(), FinalT);
+
+			//UE_LOG(LogChaos, Warning, TEXT("MoveTOIHack: FinalTOI %f; Correction %f"), FinalT, (FinalP - Particle1->P()).Size());
+		
+			Particle1->SetP(FinalP);
+			PostMoveToTOIHack(Dt, Particle1);
+		}
+	}
+}
+
+
+void MoveToTOIHack(FReal Dt, TTransientPBDRigidParticleHandle<FReal, 3>& Particle, const ISpatialAcceleration<TAccelerationStructureHandle<FReal, 3>, FReal, 3>* SpatialAcceleration)
+{
+	if (const auto AABBTree = SpatialAcceleration->template As<TAABBTree<TAccelerationStructureHandle<FReal, 3>, TAABBTreeLeafArray<TAccelerationStructureHandle<FReal, 3>, FReal>, FReal>>())
+	{
+		MoveToTOIHackImpl(Dt, Particle, AABBTree);
+	}
+	else if (const auto BV = SpatialAcceleration->template As<TBoundingVolume<TAccelerationStructureHandle<FReal, 3>, FReal, 3>>())
+	{
+		MoveToTOIHackImpl(Dt, Particle, BV);
+	}
+	else if (const auto AABBTreeBV = SpatialAcceleration->template As<TAABBTree<TAccelerationStructureHandle<FReal, 3>, TBoundingVolume<TAccelerationStructureHandle<FReal, 3>, FReal, 3>, FReal>>())
+	{
+		MoveToTOIHackImpl(Dt, Particle, AABBTreeBV);
+	}
+	else if (const auto Collection = SpatialAcceleration->template As<ISpatialAccelerationCollection<TAccelerationStructureHandle<FReal, 3>, FReal, 3>>())
+	{
+		Collection->CallMoveToTOIHack(Dt, Particle);
+	}
+}
+
+bool RequiresCCDHack(FReal Dt, const TTransientPBDRigidParticleHandle<FReal, 3>& Particle)
+{
+	if (Particle.HasBounds() && (Particle.ObjectState() == EObjectStateType::Dynamic))
+	{
+		FReal Dist2 = Particle.V().SizeSquared() * Dt * Dt;
+		FReal MinBounds = Particle.Geometry()->BoundingBox().Extents().Min();
+		FReal MaxDepth = MinBounds * HackCCD_EnableThreshold;
+		if (Dist2 > MaxDepth * MaxDepth)
+		{
+			//UE_LOG(LogChaos, Warning, TEXT("MoveTOIHack: Enabled at DR = %f / %f"), FMath::Sqrt(Dist2), MaxDepth);
+
+			return true;
+		}
+	}
+	return false;
+}
+
+void CCDHack(const FReal Dt, TParticleView<TPBDRigidParticles<FReal, 3>>& ParticlesView, const ISpatialAcceleration<TAccelerationStructureHandle<FReal, 3>, FReal, 3>* SpatialAcceleration)
+{
+	if (HackCCD_EnableThreshold > 0)
+	{
+		for (auto& Particle : ParticlesView)
+		{
+			if (RequiresCCDHack(Dt, Particle))
+			{
+				MoveToTOIHack(Dt, Particle, SpatialAcceleration);
+			}
+		}
+	}
+}
+
+//
+//
+//
+// END TOI HACK STUFF
+//
+//
+//
+
+
 template <typename T, int d>
 void TPBDRigidsEvolutionGBF<T, d>::AdvanceOneTimeStep(const T Dt, const T StepFraction)
 {
@@ -163,6 +281,9 @@ void TPBDRigidsEvolutionGBF<T, d>::AdvanceOneTimeStep(const T Dt, const T StepFr
 	{
 		Base::ComputeIntermediateSpatialAcceleration();
 	}
+
+	CCDHack(Dt, Particles.GetActiveParticlesView(), InternalAcceleration.Get());
+
 	{
 		SCOPE_CYCLE_COUNTER(STAT_Evolution_DetectCollisions);
 		CollisionDetector.GetBroadPhase().SetSpatialAcceleration(InternalAcceleration.Get());
diff --git a/Engine/Source/Runtime/Experimental/Chaos/Public/Chaos/Collision/SpatialAccelerationBroadPhase.h b/Engine/Source/Runtime/Experimental/Chaos/Public/Chaos/Collision/SpatialAccelerationBroadPhase.h
--- a/Engine/Source/Runtime/Experimental/Chaos/Public/Chaos/Collision/SpatialAccelerationBroadPhase.h
+++ b/Engine/Source/Runtime/Experimental/Chaos/Public/Chaos/Collision/SpatialAccelerationBroadPhase.h
@@ -237,4 +237,43 @@ namespace Chaos
 		const TPBDRigidsSOAs<FReal, 3>& Particles;
 		const FAccelerationStructure* SpatialAcceleration;
 	};
+
+
+
+
+	extern void MoveToTOIPairHack(FReal Dt, TPBDRigidParticleHandle<FReal, 3>* Particle1, const TGeometryParticleHandle<FReal, 3>* Particle2);
+
+	template<typename T_SPATIALACCELERATION>
+	void MoveToTOIHackImpl(FReal Dt, TTransientPBDRigidParticleHandle<FReal, 3>& Particle1, const T_SPATIALACCELERATION* SpatialAcceleration)
+	{
+		TArray<TAccelerationStructureHandle<FReal, 3>> PotentialIntersections;
+		const TAABB<FReal, 3> Box1 = ComputeWorldSpaceBoundingBox(Particle1);
+		TSimOverlapVisitor<FReal, 3> OverlapVisitor(PotentialIntersections);
+		SpatialAcceleration->Overlap(Box1, OverlapVisitor);
+
+		const int32 NumPotentials = PotentialIntersections.Num();
+		for (int32 i = 0; i < NumPotentials; ++i)
+		{
+			auto& Particle2 = *PotentialIntersections[i].GetGeometryParticleHandle_PhysicsThread();
+
+			// We only care about static/kinematics
+			if (Particle2.CastToRigidParticle() && (Particle2.CastToRigidParticle()->ObjectState() == EObjectStateType::Dynamic))
+			{
+				continue;
+			}
+			if (!Particle1.Geometry() && !Particle2.Geometry())
+			{
+				continue;
+			}
+
+			const bool bBody2Bounded = HasBoundingBox(Particle2);
+			if (Particle1.Handle() == Particle2.Handle())
+			{
+				continue;
+			}
+
+			MoveToTOIPairHack(Dt, Particle1.Handle(), &Particle2);
+		}
+	}
+
 }
diff --git a/Engine/Source/Runtime/Experimental/Chaos/Public/Chaos/ISpatialAccelerationCollection.h b/Engine/Source/Runtime/Experimental/Chaos/Public/Chaos/ISpatialAccelerationCollection.h
--- a/Engine/Source/Runtime/Experimental/Chaos/Public/Chaos/ISpatialAccelerationCollection.h
+++ b/Engine/Source/Runtime/Experimental/Chaos/Public/Chaos/ISpatialAccelerationCollection.h
@@ -29,6 +29,7 @@ public:
 
 	/** This is kind of a hack to avoid virtuals. We simply route calls into templated functions */
 	virtual void PBDComputeConstraintsLowLevel(T Dt, FSpatialAccelerationBroadPhase& BroadPhase, FNarrowPhase& NarrowPhase, FAsyncCollisionReceiver& Receiver, CollisionStats::FStatData& StatData) const = 0;
+	virtual void CallMoveToTOIHack(FReal Dt, TTransientPBDRigidParticleHandle<FReal, 3>& Particle) const = 0;
 	virtual TArray<FSpatialAccelerationIdx> GetAllSpatialIndices() const = 0;
 
 	bool IsBucketActive(uint8 BucketIdx) const
diff --git a/Engine/Source/Runtime/Experimental/Chaos/Public/Chaos/SpatialAccelerationCollection.h b/Engine/Source/Runtime/Experimental/Chaos/Public/Chaos/SpatialAccelerationCollection.h
--- a/Engine/Source/Runtime/Experimental/Chaos/Public/Chaos/SpatialAccelerationCollection.h
+++ b/Engine/Source/Runtime/Experimental/Chaos/Public/Chaos/SpatialAccelerationCollection.h
@@ -306,6 +306,22 @@ typename TEnableIf<!TIsSame<typename SpatialAccelerationCollection::TPayloadType
 {
 }
 
+template<typename T_SPATIALACCELERATION>
+void MoveToTOIHackImpl(FReal Dt, TTransientPBDRigidParticleHandle<FReal, 3>& Particle1, const T_SPATIALACCELERATION* SpatialAcceleration);
+
+template <typename SpatialAccelerationCollection>
+typename TEnableIf<TIsSame<typename SpatialAccelerationCollection::TPayloadType, TAccelerationStructureHandle<FReal, 3>>::Value, void>::Type CallMoveToTOIHack_Helper(FReal Dt, const SpatialAccelerationCollection& Accel, TTransientPBDRigidParticleHandle<FReal, 3>& Particle)
+{
+	MoveToTOIHackImpl(Dt, Particle, &Accel);
+}
+
+template <typename SpatialAccelerationCollection>
+typename TEnableIf<!TIsSame<typename SpatialAccelerationCollection::TPayloadType, TAccelerationStructureHandle<FReal, 3>>::Value, void>::Type CallMoveToTOIHack_Helper(FReal Dt, const SpatialAccelerationCollection& Accel, TTransientPBDRigidParticleHandle<FReal, 3>& Particle)
+{
+}
+
+
+
 template <typename ... TSpatialAccelerationTypes>
 class CHAOS_API TSpatialAccelerationCollection : public
 	ISpatialAccelerationCollection<typename std::tuple_element<0, std::tuple< TSpatialAccelerationTypes...>>::type::PayloadType,
@@ -476,6 +492,11 @@ public:
 		PBDComputeConstraintsLowLevel_Helper(Dt, *this, BroadPhase, NarrowPhase, Receiver, StatData);
 	}
 
+	virtual void CallMoveToTOIHack(FReal Dt, TTransientPBDRigidParticleHandle<FReal, 3>& Particle) const override
+	{
+		CallMoveToTOIHack_Helper(Dt, *this, Particle);
+	}
+
 	virtual void Serialize(FChaosArchive& Ar)
 	{
 		//todo: let user serialize out bucket params
