Chaos: Heightfiled sweep fix where it could miss valid hits

diff --git a/Engine/Source/Runtime/Experimental/Chaos/Public/Chaos/HeightField.h b/Engine/Source/Runtime/Experimental/Chaos/Public/Chaos/HeightField.h
--- a/Engine/Source/Runtime/Experimental/Chaos/Public/Chaos/HeightField.h
+++ b/Engine/Source/Runtime/Experimental/Chaos/Public/Chaos/HeightField.h
@@ -476,8 +476,8 @@ namespace Chaos
 					}
 				}
 
-				const T TimeToEnter = FMath::Max(TimesToEnter[0],TimesToEnter[1]);
-				const T TimeToExit = FMath::Min(TimesToExit[0],TimesToExit[1]);
+				const T TimeToEnter = FMath::Max(FMath::Abs(TimesToEnter[0]),FMath::Abs(TimesToEnter[1]));
+				const T TimeToExit = FMath::Min(FMath::Abs(TimesToExit[0]),FMath::Abs(TimesToExit[1]));
 
 				if(TimeToExit < TimeToEnter)
 				{
@@ -507,7 +507,7 @@ namespace Chaos
 					}
 				}
 
-				const T MinTime = FMath::Min(Times[0],Times[1]);
+				const T MinTime = FMath::Min(FMath::Abs(Times[0]),FMath::Abs(Times[1]));
 				return MinTime;
 			}
 		};
@@ -554,3 +554,4 @@ namespace Chaos
 
 	};
 }
+
