Chaos : Sleeping

- Lowered the defaults for linear and angular sleeping velocities.

diff --git a/Engine/Source/Runtime/Experimental/Chaos/Private/Chaos/PBDConstraintGraph.cpp b/Engine/Source/Runtime/Experimental/Chaos/Private/Chaos/PBDConstraintGraph.cpp
--- a/Engine/Source/Runtime/Experimental/Chaos/Private/Chaos/PBDConstraintGraph.cpp
+++ b/Engine/Source/Runtime/Experimental/Chaos/Private/Chaos/PBDConstraintGraph.cpp
@@ -14,11 +14,14 @@
 
 using namespace Chaos;
 
-float ChaosSolverCollisionDefaultLinearSleepThresholdCVar = 1.f;
-FAutoConsoleVariableRef CVarChaosSolverCollisionDefaultLinearSleepThreshold(TEXT("p.ChaosSolverCollisionDefaultLinearSleepThreshold"), ChaosSolverCollisionDefaultLinearSleepThresholdCVar, TEXT("Default linear threshold for sleeping.[def:2]"));
+bool ChaosSolverCollisionDefaultUseMaterialSleepThresholdsCVar = true;
+FAutoConsoleVariableRef CVarChaosSolverCollisionDefaultUseMaterialSleepThresholds(TEXT("p.ChaosSolverCollisionDefaultUseMaterialSleepThresholds"), ChaosSolverCollisionDefaultUseMaterialSleepThresholdsCVar, TEXT("Enable material support for sleeping thresholds[def:true]"));
 
-float ChaosSolverCollisionDefaultAngularSleepThresholdCVar = 0.1f;
-FAutoConsoleVariableRef CVarChaosSolverCollisionDefaultAngularSleepThreshold(TEXT("p.ChaosSolverCollisionDefaultAngularSleepThreshold"), ChaosSolverCollisionDefaultAngularSleepThresholdCVar, TEXT("Default angular threshold for sleeping.[def:2]"));
+float ChaosSolverCollisionDefaultLinearSleepThresholdCVar = 0.001f; // .001 unit mass cm
+FAutoConsoleVariableRef CVarChaosSolverCollisionDefaultLinearSleepThreshold(TEXT("p.ChaosSolverCollisionDefaultLinearSleepThreshold"), ChaosSolverCollisionDefaultLinearSleepThresholdCVar, TEXT("Default linear threshold for sleeping.[def:0.001]"));
+
+float ChaosSolverCollisionDefaultAngularSleepThresholdCVar = 0.0087f;  //~1/2 unit mass degree
+FAutoConsoleVariableRef CVarChaosSolverCollisionDefaultAngularSleepThreshold(TEXT("p.ChaosSolverCollisionDefaultAngularSleepThreshold"), ChaosSolverCollisionDefaultAngularSleepThresholdCVar, TEXT("Default angular threshold for sleeping.[def:0.0087]"));
 
 FPBDConstraintGraph::FPBDConstraintGraph() : VisitToken(0)
 {
@@ -537,28 +540,44 @@ void FPBDConstraintGraph::ComputeIsland(const int32 InNode, const int32 Island,
 
 bool FPBDConstraintGraph::SleepInactive(const int32 Island, const TArrayCollectionArray<TSerializablePtr<FChaosPhysicsMaterial>>& PerParticleMaterialAttributes)
 {
-	FReal DefaultLinearSleepingThreshold = (FReal)ChaosSolverCollisionDefaultLinearSleepThresholdCVar;
-	FReal DefaultAngularSleepingThreshold = (FReal)ChaosSolverCollisionDefaultAngularSleepThresholdCVar;
+	FReal LinearSleepingThreshold = FLT_MAX;
+	FReal AngularSleepingThreshold = FLT_MAX;
+	FReal DefaultLinearSleepingThreshold = ChaosSolverCollisionDefaultLinearSleepThresholdCVar;
+	FReal DefaultAngularSleepingThreshold = ChaosSolverCollisionDefaultAngularSleepThresholdCVar;
+
 	// @todo(ccaulfield): should be able to eliminate this when island is already sleeping
 
 	const TArray<TGeometryParticleHandle<FReal, 3>*>& IslandParticles = GetIslandParticles(Island);
 	check(IslandParticles.Num());
 
-	if (IslandToParticles[Island].Num() == 1)
+	if (IslandParticles.Num() == 1)
 	{
-		if (TPBDRigidParticleHandle<FReal, 3>* PBDRigid = IslandToParticles[Island][0]->CastToRigidParticle())
+		if (TPBDRigidParticleHandle<FReal, 3>* PBDRigid = IslandParticles[0]->CastToRigidParticle())
 		{
-			if (PBDRigid->V().Size() < DefaultLinearSleepingThreshold && PBDRigid->W().Size() < DefaultAngularSleepingThreshold)
+			TSerializablePtr<FChaosPhysicsMaterial> PhysicsMaterial = PBDRigid->AuxilaryValue(PerParticleMaterialAttributes);
+			if (PhysicsMaterial && ChaosSolverCollisionDefaultUseMaterialSleepThresholdsCVar)
+			{
+				LinearSleepingThreshold = FMath::Min(LinearSleepingThreshold, PhysicsMaterial->SleepingLinearThreshold);
+				AngularSleepingThreshold = FMath::Min(AngularSleepingThreshold, PhysicsMaterial->SleepingAngularThreshold);
+			}
+			else
+			{
+				LinearSleepingThreshold = DefaultLinearSleepingThreshold;
+				AngularSleepingThreshold = DefaultAngularSleepingThreshold;
+			}
+
+			if (PBDRigid->V().Size() < LinearSleepingThreshold && PBDRigid->W().Size() < AngularSleepingThreshold)
 			{
 				PBDRigid->SetSleeping(true);
 				PBDRigid->V() = TVector<FReal, 3>(0);
 				PBDRigid->W() = TVector<FReal, 3>(0);
 				return true;
 			}
+
+			return false;
 		}
 	}
 
-
 	if (!IslandToData[Island].bIsIslandPersistant)
 	{
 		return false;
@@ -569,12 +588,12 @@ bool FPBDConstraintGraph::SleepInactive(const int32 Island, const TArrayCollecti
 	TVector<FReal, 3> V(0);
 	TVector<FReal, 3> W(0);
 	FReal M = 0;
-	FReal LinearSleepingThreshold = FLT_MAX;
-	FReal AngularSleepingThreshold = FLT_MAX;
+	LinearSleepingThreshold = FLT_MAX;
+	AngularSleepingThreshold = FLT_MAX;
 
 	int32 NumDynamicParticles = 0;
 
-	for (TGeometryParticleHandle<FReal, 3>* Particle : IslandToParticles[Island])
+	for (TGeometryParticleHandle<FReal, 3>* Particle : IslandParticles)
 	{
 		TPBDRigidParticleHandle<FReal, 3>* PBDRigid = Particle->CastToRigidParticle();
 		if (PBDRigid && (PBDRigid->ObjectState() == EObjectStateType::Dynamic))
@@ -605,10 +624,10 @@ bool FPBDConstraintGraph::SleepInactive(const int32 Island, const TArrayCollecti
 
 	V /= M;
 
-	for (const TGeometryParticleHandle<FReal, 3>* Particle: IslandParticles)
+	for (const TGeometryParticleHandle<FReal, 3>* Particle : IslandParticles)
 	{
 		const TPBDRigidParticleHandle<FReal, 3>* PBDRigid = Particle->CastToRigidParticle();
-		if(PBDRigid && PBDRigid->ObjectState() == EObjectStateType::Dynamic)
+		if (PBDRigid && PBDRigid->ObjectState() == EObjectStateType::Dynamic)
 		{
 			W += /*TVector<FReal, 3>::CrossProduct(PBDRigid->X() - X, PBDRigid->M() * PBDRigid->V()/ +*/ PBDRigid->W() * PBDRigid->M();
 		}
@@ -625,13 +644,13 @@ bool FPBDConstraintGraph::SleepInactive(const int32 Island, const TArrayCollecti
 			for (TGeometryParticleHandle<FReal, 3>* Particle : IslandParticles)
 			{
 				TPBDRigidParticleHandle<FReal, 3>* PBDRigid = Particle->CastToRigidParticle();
-				if(PBDRigid && PBDRigid->ObjectState() == EObjectStateType::Dynamic)
+				if (PBDRigid && PBDRigid->ObjectState() == EObjectStateType::Dynamic)
 				{
 					PBDRigid->SetSleeping(true);
 					PBDRigid->V() = TVector<FReal, 3>(0);
 					PBDRigid->W() = TVector<FReal, 3>(0);
 				}
-				
+
 			}
 			return true;
 		}
