function typo rename

diff --git a/Engine/Source/Runtime/Experimental/Chaos/Private/Chaos/PBDRigidsEvolution.cpp b/Engine/Source/Runtime/Experimental/Chaos/Private/Chaos/PBDRigidsEvolution.cpp
--- a/Engine/Source/Runtime/Experimental/Chaos/Private/Chaos/PBDRigidsEvolution.cpp
+++ b/Engine/Source/Runtime/Experimental/Chaos/Private/Chaos/PBDRigidsEvolution.cpp
@@ -280,7 +280,7 @@ namespace Chaos
 			}
 		}
 
-		AccelerationStructure->SetAllAsyncTrasksComplete(!IsTimeSlicingProgressing);
+		AccelerationStructure->SetAllAsyncTasksComplete(!IsTimeSlicingProgressing);
 
 		// If it's not progressing then it is finished so we can perform the final copy if required
 		if (!IsTimeSlicingProgressing)
@@ -454,11 +454,11 @@ namespace Chaos
 		if (AsyncComplete)
 		{
 			// only copy when the acceleration structures have completed time-slicing
-			if (AccelerationStructureTaskComplete && AsyncInternalAcceleration->IsAllAsyncTrasksComplete())
+			if (AccelerationStructureTaskComplete && AsyncInternalAcceleration->IsAllAsyncTasksComplete())
 			{
 				SCOPE_CYCLE_COUNTER(STAT_SwapAccelerationStructures);
 
-				check(AsyncInternalAcceleration->IsAllAsyncTrasksComplete());
+				check(AsyncInternalAcceleration->IsAllAsyncTasksComplete());
 
 				FlushAsyncAccelerationQueue();
 
diff --git a/Engine/Source/Runtime/Experimental/Chaos/Public/Chaos/ISpatialAccelerationCollection.h b/Engine/Source/Runtime/Experimental/Chaos/Public/Chaos/ISpatialAccelerationCollection.h
--- a/Engine/Source/Runtime/Experimental/Chaos/Public/Chaos/ISpatialAccelerationCollection.h
+++ b/Engine/Source/Runtime/Experimental/Chaos/Public/Chaos/ISpatialAccelerationCollection.h
@@ -20,7 +20,7 @@ public:
 	ISpatialAccelerationCollection()
 	: ISpatialAcceleration<TPayloadType, T, d>(StaticType)
 	, ActiveBucketsMask(0)
-	, AllAsyncTrasksComplete(true)
+	, AllAsyncTasksComplete(true)
 	{}
 	static constexpr ESpatialAcceleration StaticType = ESpatialAcceleration::Collection;
 	virtual FSpatialAccelerationIdx AddSubstructure(TUniquePtr<ISpatialAcceleration<TPayloadType, T, d>>&& Substructure, uint16 Bucket) = 0;
@@ -37,12 +37,12 @@ public:
 		return (1 << BucketIdx) & ActiveBucketsMask;
 	}
 
-	bool IsAllAsyncTrasksComplete() const { return AllAsyncTrasksComplete; }
-	void SetAllAsyncTrasksComplete(bool State) { AllAsyncTrasksComplete = State; }
+	bool IsAllAsyncTasksComplete() const { return AllAsyncTasksComplete; }
+	void SetAllAsyncTasksComplete(bool State) { AllAsyncTasksComplete = State; }
 
 protected:
 	uint8 ActiveBucketsMask;
-	bool AllAsyncTrasksComplete;
+	bool AllAsyncTasksComplete;
 };
 
 }
