Workaround to solve objects tunnelling through heightfields due to initial overlaps returning unexpected normals. In extreme cases this leaves us with a valid normal but incorrect MTD but so far seems to behave better then previously for terrain tunnelling.

diff --git a/Engine/Source/Runtime/Experimental/Chaos/Private/Chaos/HeightField.cpp b/Engine/Source/Runtime/Experimental/Chaos/Private/Chaos/HeightField.cpp
--- a/Engine/Source/Runtime/Experimental/Chaos/Private/Chaos/HeightField.cpp
+++ b/Engine/Source/Runtime/Experimental/Chaos/Private/Chaos/HeightField.cpp
@@ -259,6 +259,16 @@ namespace Chaos
 
 							if(Time <= 0) //initial overlap or MTD, so stop
 							{
+								// This is incorrect. To prevent objects pushing through the surface of the heightfield
+								// we adopt the triangle normal but this leaves us with an incorrect MTD from the GJK call
+								// above. #TODO possibly re-do GJK with a plane, or some geom vs.plane special case to solve
+								// both triangles as planes 
+								const TVector<T, 3> AB = B - A;
+								const TVector<T, 3> AC = C - A;
+
+								TVector<T, 3> TriNormal = TVector<T, 3>::CrossProduct(AB, AC);
+
+								OutNormal = TriNormal;
 								CurrentLength = 0;
 								return false;
 							}
