Fix Particle Initialized flag not being implemented in ProxyData Init/Reset. Fixes boats on cranes desyncing when crane is broken.

diff --git a/Engine/Source/Runtime/Experimental/Chaos/Public/Chaos/ParticleHandle.h b/Engine/Source/Runtime/Experimental/Chaos/Public/Chaos/ParticleHandle.h
--- a/Engine/Source/Runtime/Experimental/Chaos/Public/Chaos/ParticleHandle.h
+++ b/Engine/Source/Runtime/Experimental/Chaos/Public/Chaos/ParticleHandle.h
@@ -2225,6 +2225,7 @@ public:
 		MDisabled = false;
 		MToBeRemovedOnFracture = false;
 		MGravityEnabled = false;
+		MInitialized = false;
 	}
 	void Init(const TPBDRigidParticle<T, d>& InParticle) {
 			Base::Init(InParticle);
@@ -2249,6 +2250,7 @@ public:
 			MDisabled = InParticle.Disabled();
 			MToBeRemovedOnFracture = InParticle.ToBeRemovedOnFracture();
 			MGravityEnabled = InParticle.IsGravityEnabled();
+			MInitialized = InParticle.IsInitialized();
 			Type = EParticleType::Rigid;
 		}
 };
