Re-enable freciprocal-math and remove optimization on the faulty function that was causing regression. Using freciprocal was causing innacurate results.
rcpps is faster but less accurate (12 bits of precision), this can causes innacuracy, for example 1.0f / 1.0f = 0.9975f.

#rnx

diff --git a/Engine/Source/Runtime/Experimental/Chaos/Public/Chaos/UniformGrid.h b/Engine/Source/Runtime/Experimental/Chaos/Public/Chaos/UniformGrid.h
--- a/Engine/Source/Runtime/Experimental/Chaos/Public/Chaos/UniformGrid.h
+++ b/Engine/Source/Runtime/Experimental/Chaos/Public/Chaos/UniformGrid.h
@@ -71,6 +71,14 @@ class CHAOS_API TUniformGridBase
 	{
 		return MDx * Face.Second + MMinCorner + (TVector<T, d>(1) - TVector<T, d>::AxisVector(Face.First)) * (MDx / 2);
 	}
+
+#ifdef PLATFORM_COMPILER_CLANG
+	// Disable optimization (-ffast-math) since its currently causing regressions.
+	//		freciprocal-math:
+	//		x / y = x * rccps(y) 
+	//		rcpps is faster but less accurate (12 bits of precision), this can causes incorrect CellIdx
+	DISABLE_FUNCTION_OPTIMIZATION 
+#endif
 	TVector<int32, d> Cell(const TVector<T, d>& X) const
 	{
 		const TVector<T, d> Delta = X - MMinCorner;
