Chaos
- Fix normalization of swing axis
- Add joint mass conditioning
- Joint cleanup

diff --git a/Engine/Source/Runtime/Experimental/Chaos/Private/Chaos/Joint/PBDJointSolverGaussSeidel.cpp b/Engine/Source/Runtime/Experimental/Chaos/Private/Chaos/Joint/PBDJointSolverGaussSeidel.cpp
--- a/Engine/Source/Runtime/Experimental/Chaos/Private/Chaos/Joint/PBDJointSolverGaussSeidel.cpp
+++ b/Engine/Source/Runtime/Experimental/Chaos/Private/Chaos/Joint/PBDJointSolverGaussSeidel.cpp
@@ -167,12 +167,7 @@ namespace Chaos
 		InvMs[0] = JointSettings.ParentInvMassScale * InvM0;
 		InvMs[1] = InvM1;
 
-		// @todo(ccaulfield): mass conditioning
-		//FReal M0 = (JointSettings.ParentInvMassScale * InvM0 > KINDA_SMALL_NUMBER) ? 1.0f / (JointSettings.ParentInvMassScale * InvM0) : 0.0f;
-		//FReal M1 = (InvM1 > 0) ? 1.0f / InvM1 : 0.0f;
-		//FVec3 IL0 = (JointSettings.ParentInvMassScale * InvM0 > KINDA_SMALL_NUMBER) ? FVec3(1.0f / (JointSettings.ParentInvMassScale * InvIL0.X), 1.0f / (JointSettings.ParentInvMassScale * InvIL0.Y), 1.0f / (JointSettings.ParentInvMassScale * InvIL0.Z)) : FVec3(0);
-		//FVec3 IL1 = (InvM1 > 0) ? FVec3(1.0f / InvIL1.X, 1.0f / InvIL1.Y, 1.0f / InvIL1.Z) : FVec3(0);
-		//FPBDJointUtilities::GetConditionedInverseMass(M0, IL0, M1, IL1, InvMs[0], InvMs[1], InvILs[0], InvILs[1], SolverSettings.MinParentMassRatio, SolverSettings.MaxInertiaRatio);
+		FPBDJointUtilities::ConditionInverseMassAndInertia(InvMs[0], InvMs[1], InvILs[0], InvILs[1], SolverSettings.MinParentMassRatio, SolverSettings.MaxInertiaRatio);
 
 		PrevPs[0] = PrevP0;
 		PrevPs[1] = PrevP1;
@@ -301,34 +296,34 @@ namespace Chaos
 			}
 			else if ((Swing1Motion == EJointMotionType::Limited) && (Swing2Motion == EJointMotionType::Locked))
 			{
-				NetResult += ApplyDualConeSwingConstraint(Dt, SolverSettings, JointSettings, EJointAngularConstraintIndex::Swing2, EJointAngularAxisIndex::Swing2, false);
-				NetResult += ApplySwingConstraint(Dt, SolverSettings, JointSettings, EJointAngularConstraintIndex::Swing1, EJointAngularAxisIndex::Swing1, bSwingSoft);
+				NetResult += ApplyDualConeSwingConstraint(Dt, SolverSettings, JointSettings, EJointAngularConstraintIndex::Swing2, false);
+				NetResult += ApplySwingConstraint(Dt, SolverSettings, JointSettings, EJointAngularConstraintIndex::Swing1, bSwingSoft);
 			}
 			else if ((Swing1Motion == EJointMotionType::Limited) && (Swing2Motion == EJointMotionType::Free))
 			{
-				NetResult += ApplyDualConeSwingConstraint(Dt, SolverSettings, JointSettings, EJointAngularConstraintIndex::Swing1, EJointAngularAxisIndex::Swing1, bSwingSoft);
+				NetResult += ApplyDualConeSwingConstraint(Dt, SolverSettings, JointSettings, EJointAngularConstraintIndex::Swing1, bSwingSoft);
 			}
 			else if ((Swing1Motion == EJointMotionType::Locked) && (Swing2Motion == EJointMotionType::Limited))
 			{
-				NetResult += ApplyDualConeSwingConstraint(Dt, SolverSettings, JointSettings, EJointAngularConstraintIndex::Swing1, EJointAngularAxisIndex::Swing1, false);
-				NetResult += ApplySwingConstraint(Dt, SolverSettings, JointSettings, EJointAngularConstraintIndex::Swing2, EJointAngularAxisIndex::Swing2, bSwingSoft);
+				NetResult += ApplyDualConeSwingConstraint(Dt, SolverSettings, JointSettings, EJointAngularConstraintIndex::Swing1, false);
+				NetResult += ApplySwingConstraint(Dt, SolverSettings, JointSettings, EJointAngularConstraintIndex::Swing2, bSwingSoft);
 			}
 			else if ((Swing1Motion == EJointMotionType::Locked) && (Swing2Motion == EJointMotionType::Locked))
 			{
-				NetResult += ApplySwingConstraint(Dt, SolverSettings, JointSettings, EJointAngularConstraintIndex::Swing1, EJointAngularAxisIndex::Swing1, false);
-				NetResult += ApplySwingConstraint(Dt, SolverSettings, JointSettings, EJointAngularConstraintIndex::Swing2, EJointAngularAxisIndex::Swing2, false);
+				NetResult += ApplySwingConstraint(Dt, SolverSettings, JointSettings, EJointAngularConstraintIndex::Swing1, false);
+				NetResult += ApplySwingConstraint(Dt, SolverSettings, JointSettings, EJointAngularConstraintIndex::Swing2, false);
 			}
 			else if ((Swing1Motion == EJointMotionType::Locked) && (Swing2Motion == EJointMotionType::Free))
 			{
-				NetResult += ApplyDualConeSwingConstraint(Dt, SolverSettings, JointSettings, EJointAngularConstraintIndex::Swing1, EJointAngularAxisIndex::Swing1, false);
+				NetResult += ApplyDualConeSwingConstraint(Dt, SolverSettings, JointSettings, EJointAngularConstraintIndex::Swing1, false);
 			}
 			else if ((Swing1Motion == EJointMotionType::Free) && (Swing2Motion == EJointMotionType::Limited))
 			{
-				NetResult += ApplyDualConeSwingConstraint(Dt, SolverSettings, JointSettings, EJointAngularConstraintIndex::Swing2, EJointAngularAxisIndex::Swing2, bSwingSoft);
+				NetResult += ApplyDualConeSwingConstraint(Dt, SolverSettings, JointSettings, EJointAngularConstraintIndex::Swing2, bSwingSoft);
 			}
 			else if ((Swing1Motion == EJointMotionType::Free) && (Swing2Motion == EJointMotionType::Locked))
 			{
-				NetResult += ApplyDualConeSwingConstraint(Dt, SolverSettings, JointSettings, EJointAngularConstraintIndex::Swing2, EJointAngularAxisIndex::Swing2, false);
+				NetResult += ApplyDualConeSwingConstraint(Dt, SolverSettings, JointSettings, EJointAngularConstraintIndex::Swing2, false);
 			}
 			else if ((Swing1Motion == EJointMotionType::Free) && (Swing2Motion == EJointMotionType::Free))
 			{
@@ -375,11 +370,11 @@ namespace Chaos
 				}
 				else if (bSwingDriveEnabled && !bSwing1Locked)
 				{
-					NetResult += ApplySwingDrive(Dt, SolverSettings, JointSettings, EJointAngularConstraintIndex::Swing1, EJointAngularAxisIndex::Swing1);
+					NetResult += ApplySwingDrive(Dt, SolverSettings, JointSettings, EJointAngularConstraintIndex::Swing1);
 				}
 				else if (bSwingDriveEnabled && !bSwing2Locked)
 				{
-					NetResult += ApplySwingDrive(Dt, SolverSettings, JointSettings, EJointAngularConstraintIndex::Swing2, EJointAngularAxisIndex::Swing2);
+					NetResult += ApplySwingDrive(Dt, SolverSettings, JointSettings, EJointAngularConstraintIndex::Swing2);
 				}
 			}
 		}
@@ -419,11 +414,11 @@ namespace Chaos
 			{
 				if ((Swing1Motion == EJointMotionType::Locked) || ((Swing1Motion == EJointMotionType::Limited) && !bSwingSoft))
 				{
-					NetResult += ApplySwingProjection(Dt, SolverSettings, JointSettings, EJointAngularConstraintIndex::Swing1, EJointAngularAxisIndex::Swing1);
+					NetResult += ApplySwingProjection(Dt, SolverSettings, JointSettings, EJointAngularConstraintIndex::Swing1);
 				}
 				if ((Swing2Motion == EJointMotionType::Locked) || ((Swing2Motion == EJointMotionType::Limited) && !bSwingSoft))
 				{
-					NetResult += ApplySwingProjection(Dt, SolverSettings, JointSettings, EJointAngularConstraintIndex::Swing2, EJointAngularAxisIndex::Swing2);
+					NetResult += ApplySwingProjection(Dt, SolverSettings, JointSettings, EJointAngularConstraintIndex::Swing2);
 				}
 			}
 		}
@@ -1101,12 +1096,11 @@ namespace Chaos
 		const FPBDJointSolverSettings& SolverSettings,
 		const FPBDJointSettings& JointSettings,
 		const EJointAngularConstraintIndex SwingConstraintIndex,
-		const EJointAngularAxisIndex SwingAxisIndex,
 		const bool bUseSoftLimit)
 	{
 		FVec3 SwingAxis;
 		FReal SwingAngle;
-		FPBDJointUtilities::GetDualConeSwingAxisAngle(Rs[0], Rs[1], SwingConstraintIndex, SwingAxisIndex, SwingAxis, SwingAngle);
+		FPBDJointUtilities::GetDualConeSwingAxisAngle(Rs[0], Rs[1], SwingConstraintIndex, SwingAxis, SwingAngle);
 
 		// Calculate swing error we need to correct
 		FReal DSwingAngle = 0;
@@ -1148,12 +1142,11 @@ namespace Chaos
 		const FPBDJointSolverSettings& SolverSettings,
 		const FPBDJointSettings& JointSettings,
 		const EJointAngularConstraintIndex SwingConstraintIndex,
-		const EJointAngularAxisIndex SwingAxisIndex,
 		const bool bUseSoftLimit)
 	{
 		FVec3 SwingAxis;
 		FReal SwingAngle;
-		FPBDJointUtilities::GetSwingAxisAngle(Rs[0], Rs[1], SolverSettings.SwingTwistAngleTolerance, SwingConstraintIndex, SwingAxisIndex, SwingAxis, SwingAngle);
+		FPBDJointUtilities::GetSwingAxisAngle(Rs[0], Rs[1], SolverSettings.SwingTwistAngleTolerance, SwingConstraintIndex, SwingAxis, SwingAngle);
 
 		// Calculate swing error we need to correct
 		FReal DSwingAngle = 0;
@@ -1194,12 +1187,11 @@ namespace Chaos
 		const FReal Dt,
 		const FPBDJointSolverSettings& SolverSettings,
 		const FPBDJointSettings& JointSettings,
-		const EJointAngularConstraintIndex SwingConstraintIndex,
-		const EJointAngularAxisIndex SwingAxisIndex)
+		const EJointAngularConstraintIndex SwingConstraintIndex)
 	{
 		FVec3 SwingAxis;
 		FReal SwingAngle;
-		FPBDJointUtilities::GetSwingAxisAngle(Rs[0], Rs[1], SolverSettings.SwingTwistAngleTolerance, SwingConstraintIndex, SwingAxisIndex, SwingAxis, SwingAngle);
+		FPBDJointUtilities::GetSwingAxisAngle(Rs[0], Rs[1], SolverSettings.SwingTwistAngleTolerance, SwingConstraintIndex, SwingAxis, SwingAngle);
 
 		const FReal SwingAngleTarget = JointSettings.AngularDriveTargetAngles[(int32)SwingConstraintIndex];
 		const FReal DSwingAngle = SwingAngle - SwingAngleTarget;
@@ -1220,8 +1212,7 @@ namespace Chaos
 		const FReal Dt,
 		const FPBDJointSolverSettings& SolverSettings,
 		const FPBDJointSettings& JointSettings,
-		const EJointAngularConstraintIndex SwingConstraintIndex,
-		const EJointAngularAxisIndex SwingAxisIndex)
+		const EJointAngularConstraintIndex SwingConstraintIndex)
 	{
 		// @todo(ccaulfield): implement swing projection
 		return FJointSolverResult::MakeSolved();
@@ -1233,13 +1224,12 @@ namespace Chaos
 		const FPBDJointSolverSettings& SolverSettings,
 		const FPBDJointSettings& JointSettings)
 	{
-		// Calculate the rotation we need to apply to resolve the rotation delta
 		const FRotation3 TargetR1 = Rs[0] * JointSettings.AngularDrivePositionTarget;
-		const FRotation3 DR1 = TargetR1 * Rs[1].Inverse();
+		const FRotation3 DR = TargetR1 * Rs[1].Inverse();
 
 		FVec3 SLerpAxis;
 		FReal SLerpAngle;
-		if (DR1.ToAxisAndAngleSafe(SLerpAxis, SLerpAngle, FVec3(1, 0, 0)))
+		if (DR.ToAxisAndAngleSafe(SLerpAxis, SLerpAngle, FVec3(1, 0, 0)))
 		{
 			if (SLerpAngle > (FReal)PI)
 			{
diff --git a/Engine/Source/Runtime/Experimental/Chaos/Private/Chaos/PBDJointConstraintUtilities.cpp b/Engine/Source/Runtime/Experimental/Chaos/Private/Chaos/PBDJointConstraintUtilities.cpp
--- a/Engine/Source/Runtime/Experimental/Chaos/Private/Chaos/PBDJointConstraintUtilities.cpp
+++ b/Engine/Source/Runtime/Experimental/Chaos/Private/Chaos/PBDJointConstraintUtilities.cpp
@@ -77,15 +77,18 @@ namespace Chaos
 		const FRotation3& R0,
 		const FRotation3& R1,
 		const EJointAngularConstraintIndex SwingConstraintIndex,
-		const EJointAngularAxisIndex SwingAxisIndex,
 		FVec3& Axis,
 		FReal& Angle)
 	{
 		FVec3 Twist1 = R1 * FJointConstants::TwistAxis();
-		FVec3 Swing0 = R0 * ((SwingConstraintIndex == EJointAngularConstraintIndex::Swing1)? FJointConstants::Swing2Axis() : FJointConstants::Swing1Axis());
+		FVec3 Swing0 = R0 * FJointConstants::OtherSwingAxis(SwingConstraintIndex);
 		Axis = FVec3::CrossProduct(Swing0, Twist1);
-		FReal SwingTwistDot = FVec3::DotProduct(Swing0, Twist1);
-		Angle = FMath::Asin(FMath::Clamp(-SwingTwistDot, -1.0f, 1.0f));
+		Angle = 0.0f;
+		if (Utilities::NormalizeSafe(Axis, KINDA_SMALL_NUMBER))
+		{
+			FReal SwingTwistDot = FVec3::DotProduct(Swing0, Twist1);
+			Angle = FMath::Asin(FMath::Clamp(-SwingTwistDot, -1.0f, 1.0f));
+		}
 	}
 
 
@@ -94,14 +97,13 @@ namespace Chaos
 		const FRotation3& R1,
 		const FReal AngleTolerance,
 		const EJointAngularConstraintIndex SwingConstraintIndex,
-		const EJointAngularAxisIndex SwingAxisIndex,
 		FVec3& Axis,
 		FReal& Angle)
 	{
 		// Decompose rotation of body 1 relative to body 0 into swing and twist rotations, assuming twist is X axis
 		FRotation3 R01Twist, R01Swing;
 		FPBDJointUtilities::DecomposeSwingTwistLocal(R0, R1, R01Swing, R01Twist);
-		const FReal R01SwingYorZ = ((int32)SwingAxisIndex == 2) ? R01Swing.Z : R01Swing.Y;	// Can't index a quat :(
+		const FReal R01SwingYorZ = (FJointConstants::AxisIndex(SwingConstraintIndex) == 2) ? R01Swing.Z : R01Swing.Y;	// Can't index a quat :(
 		Angle = 4.0f * FMath::Atan2(R01SwingYorZ, 1.0f + R01Swing.W);
 		const FVec3& AxisLocal = (SwingConstraintIndex == EJointAngularConstraintIndex::Swing1) ? FJointConstants::Swing1Axis() : FJointConstants::Swing2Axis();
 		Axis = R0 * AxisLocal;
@@ -357,57 +359,49 @@ namespace Chaos
 	}
 
 	
-	void FPBDJointUtilities::GetConditionedInverseMass(
-		const FReal InMParent, 
-		const FVec3 InIParent, 
-		const FReal InMChild, 
-		const FVec3 InIChild, 
-		FReal& OutInvMParent, 
-		FReal& OutInvMChild, 
-		FVec3& OutInvIParent,
-		FVec3& OutInvIChild,
-		const FReal MinParentMassRatio, 
+	// @todo(ccaulfield): should also take into account the length of the joint connector to prevent over-rotation
+	void FPBDJointUtilities::ConditionInverseMassAndInertia(
+		FReal& InOutInvMParent,
+		FReal& InOutInvMChild,
+		FVec3& InOutInvIParent,
+		FVec3& InOutInvIChild,
+		const FReal MinParentMassRatio,
 		const FReal MaxInertiaRatio)
 	{
-		FReal MParent = ConditionParentMass(InMParent, InMChild, MinParentMassRatio);
-		FReal MChild = InMChild;
-
-		FVec3 IParent = ConditionInertia(InIParent, MaxInertiaRatio);
-		FVec3 IChild = ConditionInertia(InIChild, MaxInertiaRatio);
-		IParent = ConditionParentInertia(IParent, IChild, MinParentMassRatio);
+		FReal MParent = 0.0f;
+		FVec3 IParent = FVec3(0);
+		FReal MChild = 0.0f;
+		FVec3 IChild = FVec3(0);
 
-		OutInvMParent = 0;
-		OutInvIParent = FVec3(0, 0, 0);
-		if (MParent > 0)
+		// Set up inertia so that it is more uniform (reduce the maximum ratio of the inertia about each axis)
+		if (InOutInvMParent > 0)
 		{
-			OutInvMParent = (FReal)1 / MParent;
-			OutInvIParent = FVec3((FReal)1 / IParent.X, (FReal)1 / IParent.Y, (FReal)1 / IParent.Z);
+			MParent = 1.0f / InOutInvMParent;
+			IParent = ConditionInertia(FVec3(1.0f / InOutInvIParent.X, 1.0f / InOutInvIParent.Y, 1.0f / InOutInvIParent.Z), MaxInertiaRatio);
+		}
+		if (InOutInvMChild > 0)
+		{
+			MChild = 1.0f / InOutInvMChild;
+			IChild = ConditionInertia(FVec3(1.0f / InOutInvIChild.X, 1.0f / InOutInvIChild.Y, 1.0f / InOutInvIChild.Z), MaxInertiaRatio);
 		}
 
-		OutInvMChild = 0;
-		OutInvIChild = FVec3(0, 0, 0);
-		if (MChild > 0)
+		// Set up relative mass and inertia so that the parent cannot be much lighter than the child
+		if ((InOutInvMParent > 0) && (InOutInvMChild > 0))
 		{
-			OutInvMChild = (FReal)1 / MChild;
-			OutInvIChild = FVec3((FReal)1 / IChild.X, (FReal)1 / IChild.Y, (FReal)1 / IChild.Z);
+			MParent = ConditionParentMass(MParent, MChild, MinParentMassRatio);
+			IParent = ConditionParentInertia(IParent, IChild, MinParentMassRatio);
 		}
-	}
 
-	
-	void FPBDJointUtilities::GetConditionedInverseMass(
-		const FReal InM0,
-		const FVec3 InI0,
-		FReal& OutInvM0, 
-		FVec3& OutInvI0,
-		const FReal MaxInertiaRatio)
-	{
-		OutInvM0 = 0;
-		OutInvI0 = FVec3(0, 0, 0);
-		if (InM0 > 0)
+		// Map back to inverses
+		if (InOutInvMParent > 0)
+		{
+			InOutInvMParent = (FReal)1 / MParent;
+			InOutInvIParent = FVec3((FReal)1 / IParent.X, (FReal)1 / IParent.Y, (FReal)1 / IParent.Z);
+		}
+		if (InOutInvMChild > 0)
 		{
-			FVec3 I0 = ConditionInertia(InI0, MaxInertiaRatio);
-			OutInvM0 = (FReal)1 / InM0;
-			OutInvI0 = FVec3((FReal)1 / I0.X, (FReal)1 / I0.Y, (FReal)1 / I0.Z);
+			InOutInvMChild = (FReal)1 / MChild;
+			InOutInvIChild = FVec3((FReal)1 / IChild.X, (FReal)1 / IChild.Y, (FReal)1 / IChild.Z);
 		}
 	}
 
diff --git a/Engine/Source/Runtime/Experimental/Chaos/Public/Chaos/Joint/PBDJointSolverCholesky.h b/Engine/Source/Runtime/Experimental/Chaos/Public/Chaos/Joint/PBDJointSolverCholesky.h
--- a/Engine/Source/Runtime/Experimental/Chaos/Public/Chaos/Joint/PBDJointSolverCholesky.h
+++ b/Engine/Source/Runtime/Experimental/Chaos/Public/Chaos/Joint/PBDJointSolverCholesky.h
@@ -24,6 +24,17 @@ namespace Chaos
 	public:
 		static const int32 MaxConstrainedBodies = 2;
 
+		/**
+		 * The constraint-space axis about which each rotation constraint is applied
+		 */
+		enum class EJointAngularAxisIndex : int32
+		{
+			Twist = 0,		// Twist Axis = X
+			Swing2 = 1,		// Swing2 Axis = Y
+			Swing1 = 2,		// Swing1 Axis = Z
+		};
+
+
 		using FDenseMatrix66 = TDenseMatrix<6 * 6>;
 		using FDenseMatrix61 = TDenseMatrix<6 * 1>;
 
diff --git a/Engine/Source/Runtime/Experimental/Chaos/Public/Chaos/Joint/PBDJointSolverGaussSeidel.h b/Engine/Source/Runtime/Experimental/Chaos/Public/Chaos/Joint/PBDJointSolverGaussSeidel.h
--- a/Engine/Source/Runtime/Experimental/Chaos/Public/Chaos/Joint/PBDJointSolverGaussSeidel.h
+++ b/Engine/Source/Runtime/Experimental/Chaos/Public/Chaos/Joint/PBDJointSolverGaussSeidel.h
@@ -285,7 +285,6 @@ namespace Chaos
 			const FPBDJointSolverSettings& SolverSettings,
 			const FPBDJointSettings& JointSettings,
 			const EJointAngularConstraintIndex SwingConstraintIndex,
-			const EJointAngularAxisIndex SwingAxisIndex,
 			const bool bUseSoftLimit);
 
 		// One swing axis is locked, the other limited or locked. This applies the Limited axis (ApplyDualConeSwingConstraint is used for the locked axis).
@@ -294,22 +293,19 @@ namespace Chaos
 			const FPBDJointSolverSettings& SolverSettings,
 			const FPBDJointSettings& JointSettings,
 			const EJointAngularConstraintIndex SwingConstraintIndex,
-			const EJointAngularAxisIndex SwingAxisIndex,
 			const bool bUseSoftLimit);
 
 		FJointSolverResult ApplySwingDrive(
 			const FReal Dt,
 			const FPBDJointSolverSettings& SolverSettings,
 			const FPBDJointSettings& JointSettings,
-			const EJointAngularConstraintIndex SwingConstraintIndex,
-			const EJointAngularAxisIndex SwingAxisIndex);
+			const EJointAngularConstraintIndex SwingConstraintIndex);
 
 		FJointSolverResult ApplySwingProjection(
 			const FReal Dt,
 			const FPBDJointSolverSettings& SolverSettings,
 			const FPBDJointSettings& JointSettings,
-			const EJointAngularConstraintIndex SwingConstraintIndex,
-			const EJointAngularAxisIndex SwingAxisIndex);
+			const EJointAngularConstraintIndex SwingConstraintIndex);
 
 		FJointSolverResult ApplySLerpDrive(
 			const FReal Dt,
diff --git a/Engine/Source/Runtime/Experimental/Chaos/Public/Chaos/PBDJointConstraintTypes.h b/Engine/Source/Runtime/Experimental/Chaos/Public/Chaos/PBDJointConstraintTypes.h
--- a/Engine/Source/Runtime/Experimental/Chaos/Public/Chaos/PBDJointConstraintTypes.h
+++ b/Engine/Source/Runtime/Experimental/Chaos/Public/Chaos/PBDJointConstraintTypes.h
@@ -40,26 +40,57 @@ namespace Chaos
 		Swing1,
 	};
 
-	/**
-	 * The constraint-space axis about which each rotation constraint is applied
-	 */
-	enum class EJointAngularAxisIndex : int32
-	{
-		Twist = 0,		// Twist Axis = X
-		Swing2 = 1,		// Swing2 Axis = Y
-		Swing1 = 2,		// Swing1 Axis = Z
-	};
-
 	struct FJointConstants
 	{
 		/** The constraint-space twist axis (X Axis) */
-		static const FVec3 TwistAxis() { return FVec3(1, 0, 0); }
+		static inline const FVec3 TwistAxis() { return FVec3(1, 0, 0); }
 
 		/** The constraint-space Swing1 axis (Z Axis) */
-		static const FVec3 Swing1Axis() { return FVec3(0, 0, 1); }
+		static inline const FVec3 Swing1Axis() { return FVec3(0, 0, 1); }
 
 		/** The constraint-space Swing2 axis (Y Axis) */
-		static const FVec3 Swing2Axis() { return FVec3(0, 1, 0); }
+		static inline const FVec3 Swing2Axis() { return FVec3(0, 1, 0); }
+
+		/** Get the local-space axis for the specified constraint type. This will be one of the cardinal axes. */
+		static inline const FVec3 Axis(const EJointAngularConstraintIndex ConstraintIndex)
+		{
+			switch (ConstraintIndex)
+			{
+			case EJointAngularConstraintIndex::Twist:
+				return TwistAxis();
+			case EJointAngularConstraintIndex::Swing1:
+				return Swing1Axis();
+			case EJointAngularConstraintIndex::Swing2:
+				return Swing2Axis();
+			}
+		}
+
+		static inline const FVec3 SwingAxis(const EJointAngularConstraintIndex ConstraintIndex)
+		{
+			return (ConstraintIndex == EJointAngularConstraintIndex::Swing1) ? Swing1Axis() : Swing2Axis();
+		}
+
+		static inline const FVec3 OtherSwingAxis(const EJointAngularConstraintIndex ConstraintIndex)
+		{
+			return (ConstraintIndex == EJointAngularConstraintIndex::Swing1) ? Swing2Axis() : Swing1Axis();
+		}
+
+		/** Get the local-space axis index for the specified constraint type. This can be used to index the vectors of a transform matrix for example. */
+		static inline const int32 AxisIndex(const EJointAngularConstraintIndex ConstraintIndex)
+		{
+			if (ConstraintIndex == EJointAngularConstraintIndex::Twist)
+			{
+				return 0;	// X
+			}
+			else if (ConstraintIndex == EJointAngularConstraintIndex::Swing1)
+			{
+				return 2;	// Z
+			}
+			else
+			{
+				return 1;	// Y
+			}
+		}
 	};
 
 	class CHAOS_API FPBDJointSettings
diff --git a/Engine/Source/Runtime/Experimental/Chaos/Public/Chaos/PBDJointConstraintUtilities.h b/Engine/Source/Runtime/Experimental/Chaos/Public/Chaos/PBDJointConstraintUtilities.h
--- a/Engine/Source/Runtime/Experimental/Chaos/Public/Chaos/PBDJointConstraintUtilities.h
+++ b/Engine/Source/Runtime/Experimental/Chaos/Public/Chaos/PBDJointConstraintUtilities.h
@@ -45,7 +45,6 @@ namespace Chaos
 			const FRotation3& R0,
 			const FRotation3& R1,
 			const EJointAngularConstraintIndex SwingConstraintIndex,
-			const EJointAngularAxisIndex SwingAxisIndex,
 			FVec3& Axis,
 			FReal& Angle);
 
@@ -54,7 +53,6 @@ namespace Chaos
 			const FRotation3& R1,
 			const FReal AngleTolerance,
 			const EJointAngularConstraintIndex SwingConstraintIndex,
-			const EJointAngularAxisIndex SwingAxisIndex,
 			FVec3& Axis,
 			FReal& Angle);
 
@@ -88,26 +86,14 @@ namespace Chaos
 			const FReal MChild, 
 			const FReal MinRatio);
 
-		static CHAOS_API void GetConditionedInverseMass(
-			const float MParent,
-			const FVec3 IParent,
-			const float MChild,
-			const FVec3 IChild,
-			FReal& OutInvMParent,
-			FReal& OutInvMChild,
-			FVec3& OutInvIParent,
-			FVec3& OutInvIChild,
+		static CHAOS_API void ConditionInverseMassAndInertia(
+			FReal& InOutInvMParent,
+			FReal& InOutInvMChild,
+			FVec3& InOutInvIParent,
+			FVec3& InOutInvIChild,
 			const FReal MinParentMassRatio,
 			const FReal MaxInertiaRatio);
 
-		static CHAOS_API void GetConditionedInverseMass(
-			const float M,
-			const FVec3 I,
-			FReal& OutInvM0,
-			FVec3& OutInvI0,
-			const FReal MaxInertiaRatio);
-
-
 		static FReal GetLinearStiffness(
 			const FPBDJointSolverSettings& SolverSettings,
 			const FPBDJointSettings& JointSettings);
