Merging //UE4/Dev-Main (up to CL 9614276) to //UE4/Dev-RenderPlat-Staging

diff --git a/Engine/Plugins/Online/OnlineSubsystemUtils/Source/OnlineSubsystemUtils/Private/VoiceEngineImpl.cpp b/Engine/Plugins/Online/OnlineSubsystemUtils/Source/OnlineSubsystemUtils/Private/VoiceEngineImpl.cpp
--- a/Engine/Plugins/Online/OnlineSubsystemUtils/Source/OnlineSubsystemUtils/Private/VoiceEngineImpl.cpp
+++ b/Engine/Plugins/Online/OnlineSubsystemUtils/Source/OnlineSubsystemUtils/Private/VoiceEngineImpl.cpp
@@ -590,6 +590,12 @@ uint32 FVoiceEngineImpl::SubmitRemoteVoiceData(const FUniqueNetIdWrapper& Remote
 		}
 
 		QueuedData.VoipSynthComponent->SubmitPacket((float*)DecompressedVoiceBuffer.GetData(), BytesWritten, InSampleCount, EVoipStreamDataFormat::Int16);
+
+		// Try to start the VoipSynthComponent if it has been killed by the audio engine.
+		if (!QueuedData.VoipSynthComponent->IsPlaying())
+		{
+			QueuedData.VoipSynthComponent->Start();
+		}
 	}
 
 	return ONLINE_SUCCESS;
diff --git a/Engine/Plugins/Runtime/GameplayAbilities/Source/GameplayAbilities/Private/AbilitySystemComponent.cpp b/Engine/Plugins/Runtime/GameplayAbilities/Source/GameplayAbilities/Private/AbilitySystemComponent.cpp
--- a/Engine/Plugins/Runtime/GameplayAbilities/Source/GameplayAbilities/Private/AbilitySystemComponent.cpp
+++ b/Engine/Plugins/Runtime/GameplayAbilities/Source/GameplayAbilities/Private/AbilitySystemComponent.cpp
@@ -1356,6 +1356,11 @@ TArray<FActiveGameplayEffectHandle> UAbilitySystemComponent::GetActiveEffects(co
 	return ActiveGameplayEffects.GetActiveEffects(Query);
 }
 
+TArray<FActiveGameplayEffectHandle> UAbilitySystemComponent::GetActiveEffectsWithAllTags(FGameplayTagContainer Tags) const
+{
+	return GetActiveEffects(FGameplayEffectQuery::MakeQuery_MatchAllEffectTags(Tags));
+}
+
 float UAbilitySystemComponent::GetActiveEffectsEndTime(const FGameplayEffectQuery& Query) const
 {
 	TArray<AActor*> DummyInstigators;
diff --git a/Engine/Plugins/Runtime/GameplayAbilities/Source/GameplayAbilities/Public/AbilitySystemComponent.h b/Engine/Plugins/Runtime/GameplayAbilities/Source/GameplayAbilities/Public/AbilitySystemComponent.h
--- a/Engine/Plugins/Runtime/GameplayAbilities/Source/GameplayAbilities/Public/AbilitySystemComponent.h
+++ b/Engine/Plugins/Runtime/GameplayAbilities/Source/GameplayAbilities/Public/AbilitySystemComponent.h
@@ -713,6 +713,10 @@ class GAMEPLAYABILITIES_API UAbilitySystemComponent : public UGameplayTasksCompo
 	UFUNCTION(BlueprintCallable, BlueprintPure=false, Category = "GameplayEffects", meta=(DisplayName = "Get Activate Gameplay Effects for Query"))
 	TArray<FActiveGameplayEffectHandle> GetActiveEffects(const FGameplayEffectQuery& Query) const;
 
+	/** Returns list of active effects that have all of the passed in tags */
+	UFUNCTION(BlueprintCallable, BlueprintPure = false, Category = "GameplayEffects")
+	TArray<FActiveGameplayEffectHandle> GetActiveEffectsWithAllTags(FGameplayTagContainer Tags) const;
+
 	/** This will give the world time that all effects matching this query will be finished. If multiple effects match, it returns the one that returns last */
 	float GetActiveEffectsEndTime(const FGameplayEffectQuery& Query) const;
 	float GetActiveEffectsEndTimeWithInstigators(const FGameplayEffectQuery& Query, TArray<AActor*>& Instigators) const;
diff --git a/Engine/Shaders/Private/PathTracing/Light/PathTracingLightSampling.ush b/Engine/Shaders/Private/PathTracing/Light/PathTracingLightSampling.ush
--- a/Engine/Shaders/Private/PathTracing/Light/PathTracingLightSampling.ush
+++ b/Engine/Shaders/Private/PathTracing/Light/PathTracingLightSampling.ush
@@ -13,8 +13,10 @@ LightSampling.usf: Light sampling for next-event estimation
 #include "PathTracingSkyLight.ush"
 #include "PathTracingSpotLight.ush"
 
-void GenerateLightRay(RayDesc Ray, int LightId, float3 LightUV, out RayDesc OutLightRay)
+void GenerateLightRay(RayDesc Ray, int LightId, float3 LightUV, inout RayDesc OutLightRay)
 {
+	OutLightRay = (RayDesc)0; // #dxr_todo: Return false in the case of no matching light type and handle case from caller rather than blindly using a degenerate ray
+
 	if (IsPointLight(LightId))
 	{
 		PointLight_GenerateLightRay(Ray, LightId, LightUV, OutLightRay);
@@ -42,10 +44,11 @@ void EstimateLight(
 	int LightId,
 	RayDesc Ray,
 	FMaterialClosestHitPayload Payload,
-	out float OutIrradiance
+	inout float OutIrradiance
 )
 {
-	OutIrradiance = 0.0;
+	OutIrradiance = 0.0f;
+
 	if (IsPointLight(LightId))
 	{
 		PointLight_EstimateLight(LightId, Ray, Payload, OutIrradiance);
@@ -72,10 +75,12 @@ bool SampleLightSelection(
 	RayDesc Ray,
 	FMaterialClosestHitPayload Payload,
 	float RandSample,
-	out int OutLightId,
-	out float OutLightSelectionPdf
+	inout int OutLightId,
+	inout float OutLightSelectionPdf
 )
 {
+	OutLightSelectionPdf = 0.0;
+
 //#define UNIFORM_LIGHT_SELECTION
 #ifdef UNIFORM_LIGHT_SELECTION
 	OutLightId = RandSample * SceneLightsData.Count;
@@ -126,14 +131,15 @@ void PdfLightSelection(
 	RayDesc Ray,
 	FMaterialClosestHitPayload Payload,
 	uint LightId,
-	out float OutLightSelectionPdf
+	inout float OutLightSelectionPdf
 )
 {
+	OutLightSelectionPdf = 0.0;
+
 #ifdef UNIFORM_LIGHT_SELECTION
 	OutLightSelectionPdf = 1.0 / SceneLightsData.Count;
 	return;
 #else
-	OutLightSelectionPdf = 0.0;
 	if (SceneLightsData.Count == 0) return;
 
 	// Build irradiance estimate prefix sum
@@ -170,7 +176,7 @@ void PdfLight(
 	FMaterialClosestHitPayload Payload,
 	int LightId,
 	float3 LightUV,
-	out float OutPdf)
+	inout float OutPdf)
 {
 	// Light selection
 	// TODO: Use radiance estimate instead of uniform random selection
@@ -207,9 +213,11 @@ void EvalLight(
 	int LightId,
 	float3 LightUV,
 	RayDesc Ray,
-	out float3 OutRadiance
+	inout float3 OutRadiance
 )
 {
+	OutRadiance = float3(0.0f, 0.0f, 0.0f);
+
 	if (IsPointLight(LightId))
 	{
 		PointLight_EvalLight(LightId, LightUV, Ray, OutRadiance);
@@ -236,11 +244,11 @@ void SampleLight(
 	RayDesc Ray,
 	FMaterialClosestHitPayload Payload,
 	float4 RandSample,
-	out int OutLightId,
-	out float3 OutLightUV,
-	out float OutPdf)
+	inout int OutLightId,
+	inout float3 OutLightUV,
+	inout float OutPdf)
 {
-	OutPdf = 0.0;
+	OutPdf = 0.0f;
 
 	// Light selection
 	float LightSelectionPdf = 0.0;
@@ -250,7 +258,6 @@ void SampleLight(
 	}
 
 	// Light sampling
-	OutLightUV = 0.0;
 	float LightPdf = 0.0;
 	if (IsPointLight(OutLightId))
 	{
diff --git a/Engine/Shaders/Private/PathTracing/Light/PathTracingSpotLight.ush b/Engine/Shaders/Private/PathTracing/Light/PathTracingSpotLight.ush
--- a/Engine/Shaders/Private/PathTracing/Light/PathTracingSpotLight.ush
+++ b/Engine/Shaders/Private/PathTracing/Light/PathTracingSpotLight.ush
@@ -147,6 +147,10 @@ void SpotLight_EstimateLight(
 		OutIrradiance = LightPower * Falloff / LightDistanceSquared;
 		// #dxr_todo: UE-72534 apply inner core
 	}
+	else
+	{
+		OutIrradiance = 0.0f;
+	}
 }
 
 void SpotLight_SampleLight(
diff --git a/Engine/Shaders/Private/PathTracing/Material/PathTracingGlossy.ush b/Engine/Shaders/Private/PathTracing/Material/PathTracingGlossy.ush
--- a/Engine/Shaders/Private/PathTracing/Material/PathTracingGlossy.ush
+++ b/Engine/Shaders/Private/PathTracing/Material/PathTracingGlossy.ush
@@ -78,6 +78,10 @@ void Glossy_SampleMaterial(
 	out float3 OutThroughput,
 	out float OutPdf)
 {
+	OutDirection = float3(0, 0, 0);
+	OutThroughput = float3(0, 0, 0);
+	OutPdf = 0.0f;
+
 	float3 N_World = GetWorldNormal(Payload);
 
 	// Diffuse and Specular color are calculated in the payload unpacking, both depend on BaseColor and Metallic attributes
@@ -92,9 +96,6 @@ void Glossy_SampleMaterial(
 	const float NoV = V.z;
 	const float ProbSpecular = GetSpecularEventProbability(SpecularColor, NoV); 
 
-	OutPdf = 0.0f;
-	OutThroughput = float3(0, 0, 0);
-
 	if ((RandSample.z <= ProbSpecular)) 
 	{	
 		bool Flip = V.z < 0.0f;
diff --git a/Engine/Shaders/Private/RayTracing/RayTracingRectLightRGS.usf b/Engine/Shaders/Private/RayTracing/RayTracingRectLightRGS.usf
--- a/Engine/Shaders/Private/RayTracing/RayTracingRectLightRGS.usf
+++ b/Engine/Shaders/Private/RayTracing/RayTracingRectLightRGS.usf
@@ -30,6 +30,12 @@ bool GenerateRectLightImportanceSampledRay(
 	out float RayPdf
 )
 {
+	RayOrigin = float3(0.0f, 0.0f, 0.0f);
+	RayDirection = float3(0.0f, 0.0f, 0.0f);
+	RayTMin = 0.0f;
+	RayTMax = 0.0f;
+	RayPdf = 0.0f;
+
 	uint MipCount = log2(RectLight.MipTreeDimensions.x);
 	uint2 MipPixel = 0;
 	float MipPdf = 1.0;
@@ -126,7 +132,6 @@ bool GenerateRectLightImportanceSampledRay(
 	// Light-normal culling
 	if (dot(-LightDirection, LightNormal) <= 0.0)
 	{
-		RayPdf = 0.0;
 		return false;
 	}
 
diff --git a/Engine/Source/Runtime/AnimGraphRuntime/Private/BoneControllers/AnimNode_RigidBody.cpp b/Engine/Source/Runtime/AnimGraphRuntime/Private/BoneControllers/AnimNode_RigidBody.cpp
--- a/Engine/Source/Runtime/AnimGraphRuntime/Private/BoneControllers/AnimNode_RigidBody.cpp
+++ b/Engine/Source/Runtime/AnimGraphRuntime/Private/BoneControllers/AnimNode_RigidBody.cpp
@@ -23,7 +23,7 @@
 
 DEFINE_STAT(STAT_RigidBodyNodeInitTime);
 
-TAutoConsoleVariable<int32> CVarEnableRigidBodyNode(TEXT("p.RigidBodyNode"), 1, TEXT("Enables/disables rigid body node updates and evaluations"), ECVF_ReadOnly);
+TAutoConsoleVariable<int32> CVarEnableRigidBodyNode(TEXT("p.RigidBodyNode"), 1, TEXT("Enables/disables rigid body node updates and evaluations"), ECVF_Default);
 TAutoConsoleVariable<int32> CVarRigidBodyLODThreshold(TEXT("p.RigidBodyLODThreshold"), -1, TEXT("Max LOD that rigid body node is allowed to run on. Provides a global threshold that overrides per-node the LODThreshold property. -1 means no override."), ECVF_Scalability);
 
 int32 RBAN_MaxSubSteps = 4;
@@ -262,7 +262,7 @@ void FAnimNode_RigidBody::EvaluateSkeletalControl_AnyThread(FComponentSpacePoseC
 	const float DeltaSeconds = AccumulatedDeltaTime;
 	AccumulatedDeltaTime = 0.f;
 
-	if (CVarEnableRigidBodyNode.GetValueOnAnyThread() != 0 && PhysicsSimulation)	
+	if (bEnabled && PhysicsSimulation)	
 	{
 		const FBoneContainer& BoneContainer = Output.Pose.GetPose().GetBoneContainer();
 		const FTransform CompWorldSpaceTM = Output.AnimInstanceProxy->GetComponentTransform();
@@ -605,6 +605,9 @@ void FAnimNode_RigidBody::InitPhysics(const UAnimInstance* InAnimInstance)
 {
 	SCOPE_CYCLE_COUNTER(STAT_RigidBodyNodeInitTime);
 
+	delete PhysicsSimulation;
+	PhysicsSimulation = nullptr;
+
 	const USkeletalMeshComponent* SkeletalMeshComp = InAnimInstance->GetSkelMeshComponent();
 	const USkeletalMesh* SkeletalMeshAsset = SkeletalMeshComp->SkeletalMesh;
 
@@ -635,10 +638,10 @@ void FAnimNode_RigidBody::InitPhysics(const UAnimInstance* InAnimInstance)
 		AnimPhysicsMinDeltaTime = 0.f;
 		bSimulateAnimPhysicsAfterReset = false;
 	}
-
-	if(CVarEnableRigidBodyNode.GetValueOnAnyThread() != 0 && UsePhysicsAsset)
+	
+	bEnabled = UsePhysicsAsset && SkeletalMeshComp->GetAllowRigidBodyAnimNode() && CVarEnableRigidBodyNode.GetValueOnAnyThread() != 0;
+	if(bEnabled)
 	{
-		delete PhysicsSimulation;
 		PhysicsSimulation = new ImmediatePhysics::FSimulation();
 		const int32 NumBodies = UsePhysicsAsset->SkeletalBodySetups.Num();
 		Bodies.Empty(NumBodies);
@@ -915,7 +918,7 @@ DECLARE_CYCLE_STAT(TEXT("RigidBody_PreUpdate"), STAT_RigidBody_PreUpdate, STATGR
 void FAnimNode_RigidBody::PreUpdate(const UAnimInstance* InAnimInstance)
 {
 	// Don't update geometry if RBN is disabled
-	if(CVarEnableRigidBodyNode.GetValueOnAnyThread() == 0)
+	if(!bEnabled)
 	{
 		return;
 	}
@@ -975,7 +978,7 @@ void FAnimNode_RigidBody::UpdateInternal(const FAnimationUpdateContext& Context)
 {
 	DECLARE_SCOPE_HIERARCHICAL_COUNTER_ANIMNODE(UpdateInternal)
 	// Avoid this work if RBN is disabled, as the results would be discarded
-	if(CVarEnableRigidBodyNode.GetValueOnAnyThread() == 0)
+	if(!bEnabled)
 	{
 		return;
 	}
diff --git a/Engine/Source/Runtime/AnimGraphRuntime/Private/BoneControllers/AnimNode_RigidBody_Chaos.cpp b/Engine/Source/Runtime/AnimGraphRuntime/Private/BoneControllers/AnimNode_RigidBody_Chaos.cpp
--- a/Engine/Source/Runtime/AnimGraphRuntime/Private/BoneControllers/AnimNode_RigidBody_Chaos.cpp
+++ b/Engine/Source/Runtime/AnimGraphRuntime/Private/BoneControllers/AnimNode_RigidBody_Chaos.cpp
@@ -236,7 +236,7 @@ void FAnimNode_RigidBody_Chaos::EvaluateSkeletalControl_AnyThread(FComponentSpac
 	const float DeltaSeconds = AccumulatedDeltaTime;
 	AccumulatedDeltaTime = 0.f;
 
-	if (CVarEnableRigidBodyNode.GetValueOnAnyThread() != 0 && PhysicsSimulation)	
+	if (bEnabled && PhysicsSimulation)	
 	{
 		const FBoneContainer& BoneContainer = Output.Pose.GetPose().GetBoneContainer();
 		const FTransform CompWorldSpaceTM = Output.AnimInstanceProxy->GetComponentTransform();
@@ -572,6 +572,9 @@ void FAnimNode_RigidBody_Chaos::InitPhysics(const UAnimInstance* InAnimInstance)
 {
 	SCOPE_CYCLE_COUNTER(STAT_RigidBodyNodeInitTime);
 
+	delete PhysicsSimulation;
+	PhysicsSimulation = nullptr;
+
 	const USkeletalMeshComponent* SkeletalMeshComp = InAnimInstance->GetSkelMeshComponent();
 	const USkeletalMesh* SkeletalMeshAsset = SkeletalMeshComp->SkeletalMesh;
 
@@ -603,9 +606,9 @@ void FAnimNode_RigidBody_Chaos::InitPhysics(const UAnimInstance* InAnimInstance)
 		bSimulateAnimPhysicsAfterReset = false;
 	}
 
-	if(UsePhysicsAsset)
+	bEnabled = UsePhysicsAsset && SkeletalMeshComp->GetAllowRigidBodyAnimNode() && CVarEnableRigidBodyNode.GetValueOnAnyThread() != 0;
+	if (bEnabled)
 	{
-		delete PhysicsSimulation;
 		PhysicsSimulation = new ImmediatePhysics_Chaos::FSimulation();
 		const int32 NumBodies = UsePhysicsAsset->SkeletalBodySetups.Num();
 		Bodies.Empty(NumBodies);
@@ -876,7 +879,7 @@ DECLARE_CYCLE_STAT(TEXT("RigidBody_PreUpdate_Chaos"), STAT_RigidBody_Chaos_PreUp
 void FAnimNode_RigidBody_Chaos::PreUpdate(const UAnimInstance* InAnimInstance)
 {
 	// Don't update geometry if RBN is disabled
-	if(CVarEnableRigidBodyNode.GetValueOnAnyThread() == 0)
+	if(!bEnabled)
 	{
 		return;
 	}
@@ -935,7 +938,7 @@ DECLARE_CYCLE_STAT(TEXT("RigidBody_Update_Chaos"), STAT_RigidBody_Chaos_Update,
 void FAnimNode_RigidBody_Chaos::UpdateInternal(const FAnimationUpdateContext& Context)
 {
 	// Avoid this work if RBN is disabled, as the results would be discarded
-	if(CVarEnableRigidBodyNode.GetValueOnAnyThread() == 0)
+	if(!bEnabled)
 	{
 		return;
 	}
diff --git a/Engine/Source/Runtime/AnimGraphRuntime/Public/BoneControllers/AnimNode_RigidBody.h b/Engine/Source/Runtime/AnimGraphRuntime/Public/BoneControllers/AnimNode_RigidBody.h
--- a/Engine/Source/Runtime/AnimGraphRuntime/Public/BoneControllers/AnimNode_RigidBody.h
+++ b/Engine/Source/Runtime/AnimGraphRuntime/Public/BoneControllers/AnimNode_RigidBody.h
@@ -145,6 +145,7 @@ public:
 	uint8 bClampLinearTranslationLimitToRefPose : 1;
 
 private:
+	uint8 bEnabled : 1;
 	uint8 bSimulationStarted : 1;
 	uint8 bCheckForBodyTransformInit : 1;
 
diff --git a/Engine/Source/Runtime/AnimGraphRuntime/Public/BoneControllers/AnimNode_RigidBody_Chaos.h b/Engine/Source/Runtime/AnimGraphRuntime/Public/BoneControllers/AnimNode_RigidBody_Chaos.h
--- a/Engine/Source/Runtime/AnimGraphRuntime/Public/BoneControllers/AnimNode_RigidBody_Chaos.h
+++ b/Engine/Source/Runtime/AnimGraphRuntime/Public/BoneControllers/AnimNode_RigidBody_Chaos.h
@@ -155,6 +155,7 @@ private:
 
 	ETeleportType ResetSimulatedTeleportType;
 
+	uint8 bEnabled : 1;
 	uint8 bSimulationStarted : 1;
 	uint8 bCheckForBodyTransformInit : 1;
 
diff --git a/Engine/Source/Runtime/Apple/MetalRHI/Private/MetalRHIContext.cpp b/Engine/Source/Runtime/Apple/MetalRHI/Private/MetalRHIContext.cpp
--- a/Engine/Source/Runtime/Apple/MetalRHI/Private/MetalRHIContext.cpp
+++ b/Engine/Source/Runtime/Apple/MetalRHI/Private/MetalRHIContext.cpp
@@ -230,3 +230,8 @@ void FMetalRHICommandContext::RHIBeginComputePass(const TCHAR* InName)
 void FMetalRHICommandContext::RHIEndComputePass()
 {
 }
+
+void FMetalRHICommandContext::RHITransitionResources(FExclusiveDepthStencil DepthStencilMode, FRHITexture* DepthTexture)
+{
+	IRHICommandContext::RHITransitionResources(DepthStencilMode, DepthTexture);
+}
diff --git a/Engine/Source/Runtime/Apple/MetalRHI/Public/MetalRHIContext.h b/Engine/Source/Runtime/Apple/MetalRHI/Public/MetalRHIContext.h
--- a/Engine/Source/Runtime/Apple/MetalRHI/Public/MetalRHIContext.h
+++ b/Engine/Source/Runtime/Apple/MetalRHI/Public/MetalRHIContext.h
@@ -281,6 +281,11 @@ public:
 	 */
 	virtual void RHITransitionResources(EResourceTransitionAccess TransitionType, FRHITexture** InTextures, int32 NumTextures) final override;
 	
+	/**
+	 * Transition depth resources. Does not really do anything in metal.
+	 */
+	virtual void RHITransitionResources(FExclusiveDepthStencil DepthStencilMode, FRHITexture* DepthTexture) final override;
+	
 	/**
 	 * Compute queue will wait for the fence to be written before continuing.
 	 */
diff --git a/Engine/Source/Runtime/ApplicationCore/Private/IOS/IOSAppDelegate.cpp b/Engine/Source/Runtime/ApplicationCore/Private/IOS/IOSAppDelegate.cpp
--- a/Engine/Source/Runtime/ApplicationCore/Private/IOS/IOSAppDelegate.cpp
+++ b/Engine/Source/Runtime/ApplicationCore/Private/IOS/IOSAppDelegate.cpp
@@ -60,6 +60,7 @@ extern bool GShowSplashScreen;
 
 FIOSCoreDelegates::FOnOpenURL FIOSCoreDelegates::OnOpenURL;
 FIOSCoreDelegates::FOnWillResignActive FIOSCoreDelegates::OnWillResignActive;
+FIOSCoreDelegates::FOnDidBecomeActive FIOSCoreDelegates::OnDidBecomeActive;
 TArray<FIOSCoreDelegates::FFilterDelegateAndHandle> FIOSCoreDelegates::PushNotificationFilters;
 
 static uint GEnabledAudioFeatures[(uint8)EAudioFeature::NumFeatures];
@@ -1398,8 +1399,6 @@ extern EDeviceScreenOrientation ConvertFromUIInterfaceOrientation(UIInterfaceOri
 FCriticalSection RenderSuspend;
 - (void)applicationWillResignActive:(UIApplication *)application
 {
-	FIOSCoreDelegates::OnWillResignActive.Broadcast();
-	
     FIOSPlatformMisc::ResetBrightness();
     
     /*
@@ -1458,6 +1457,8 @@ FCriticalSection RenderSuspend;
             }, TStatId(), NULL, ENamedThreads::ActualRenderingThread);
         }
     }
+	
+	FIOSCoreDelegates::OnWillResignActive.Broadcast();
 }
 
 - (void)applicationDidEnterBackground:(UIApplication *)application
@@ -1506,7 +1507,9 @@ FCriticalSection RenderSuspend;
 extern double GCStartTime;
 - (void)applicationDidBecomeActive:(UIApplication *)application
 {
-    // make sure a GC will not timeout because it was started before entering background
+	FIOSCoreDelegates::OnDidBecomeActive.Broadcast();
+
+	// make sure a GC will not timeout because it was started before entering background
     GCStartTime = FPlatformTime::Seconds();
 	/*
 	 Restart any tasks that were paused (or not yet started) while the application was inactive. If the application was previously in the background, optionally refresh the user interface.
diff --git a/Engine/Source/Runtime/ApplicationCore/Public/IOS/IOSAppDelegate.h b/Engine/Source/Runtime/ApplicationCore/Public/IOS/IOSAppDelegate.h
--- a/Engine/Source/Runtime/ApplicationCore/Public/IOS/IOSAppDelegate.h
+++ b/Engine/Source/Runtime/ApplicationCore/Public/IOS/IOSAppDelegate.h
@@ -59,6 +59,10 @@ public:
 	DECLARE_MULTICAST_DELEGATE(FOnWillResignActive);
 	static FOnWillResignActive OnWillResignActive;
 
+	/** INTERNAL - called when becoming active - this is not thread-safe with the game thread or render thread as it is called from the app's Main thread */
+	DECLARE_MULTICAST_DELEGATE(FOnDidBecomeActive);
+	static FOnWillResignActive OnDidBecomeActive;
+	
 private:
 	struct FFilterDelegateAndHandle
 	{
diff --git a/Engine/Source/Runtime/Engine/Classes/Components/SkeletalMeshComponent.h b/Engine/Source/Runtime/Engine/Classes/Components/SkeletalMeshComponent.h
--- a/Engine/Source/Runtime/Engine/Classes/Components/SkeletalMeshComponent.h
+++ b/Engine/Source/Runtime/Engine/Classes/Components/SkeletalMeshComponent.h
@@ -473,6 +473,10 @@ public:
 	uint8 bDisableClothSimulation:1;
 
 private:
+	/** Disable rigid body animation nodes and play original animation without simulation */
+	UPROPERTY(EditAnywhere, Category=Clothing)
+	uint8 bDisableRigidBodyAnimNode:1;
+
 	/** Disable animation curves for this component. If this is set true, no curves will be processed */
 	UPROPERTY(EditAnywhere, AdvancedDisplay, Category = SkeletalMesh)
 	uint8 bAllowAnimCurveEvaluation : 1;
@@ -1055,6 +1059,15 @@ public:
 	UFUNCTION(BlueprintCallable, Category="Components|SkeletalMesh", meta=(UnsafeDuringActorConstruction="true"))
 	void UnbindClothFromMasterPoseComponent(bool bRestoreSimulationSpace = true);
 
+	/**
+	 * Sets whether or not to allow rigid body animation nodes for this component
+	 */
+	UFUNCTION(BlueprintCallable, Category = "Components|SkeletalMesh")
+	void SetAllowRigidBodyAnimNode(bool bInAllow, bool bReinitAnim = true);
+
+	UFUNCTION(BlueprintCallable, Category = "Components|SkeletalMesh")
+	bool GetAllowRigidBodyAnimNode() const { return !bDisableRigidBodyAnimNode; }
+
 	/**
 	* Sets whether or not to force tick component in order to update animation and refresh transform for this component
 	* This is supported only in the editor
diff --git a/Engine/Source/Runtime/Engine/Private/Components/SkeletalMeshComponent.cpp b/Engine/Source/Runtime/Engine/Private/Components/SkeletalMeshComponent.cpp
--- a/Engine/Source/Runtime/Engine/Private/Components/SkeletalMeshComponent.cpp
+++ b/Engine/Source/Runtime/Engine/Private/Components/SkeletalMeshComponent.cpp
@@ -3660,8 +3660,22 @@ void USkeletalMeshComponent::UnbindClothFromMasterPoseComponent(bool bRestoreSim
 		}
 
 		bBindClothToMasterComponent = false;
+	}
+}
+
+void USkeletalMeshComponent::SetAllowRigidBodyAnimNode(bool bInAllow, bool bReinitAnim)
+{
+	if(bDisableRigidBodyAnimNode == bInAllow)
+	{
+		bDisableRigidBodyAnimNode = !bInAllow;
+
+		if(bReinitAnim && bRegistered && SkeletalMesh)
+		{
+			// need to reinitialize rigid body nodes for new setting to take effect
+			InitializeAnimScriptInstance(true);
 		}
 	}
+}
 
 bool USkeletalMeshComponent::DoCustomNavigableGeometryExport(FNavigableGeometryExport& GeomExport) const
 {
diff --git a/Engine/Source/Runtime/Experimental/Chaos/Private/Chaos/Framework/Parallel.cpp b/Engine/Source/Runtime/Experimental/Chaos/Private/Chaos/Framework/Parallel.cpp
--- a/Engine/Source/Runtime/Experimental/Chaos/Private/Chaos/Framework/Parallel.cpp
+++ b/Engine/Source/Runtime/Experimental/Chaos/Private/Chaos/Framework/Parallel.cpp
@@ -5,12 +5,25 @@
 
 using namespace Chaos;
 
-static bool GNoParallelFor = false;
+namespace Chaos
+{
+#if !UE_BUILD_SHIPPING
+	bool bDisablePhysicsParallelFor = false;
+	CHAOS_API bool bDisableParticleParallelFor = false;
+	CHAOS_API bool bDisableCollisionParallelFor = false;
+
+	FAutoConsoleVariableRef CVarDisablePhysicsParallelFor(TEXT("p.Chaos.DisablePhysicsParallelFor"), bDisablePhysicsParallelFor, TEXT("Disable parallel execution in Chaos Evolution"));
+	FAutoConsoleVariableRef CVarDisableParticleParallelFor(TEXT("p.Chaos.DisableParticleParallelFor"), bDisableParticleParallelFor, TEXT("Disable parallel execution for Chaos Particles (Collisions, "));
+	FAutoConsoleVariableRef CVarDisableCollisionParallelFor(TEXT("p.Chaos.DisableCollisionParallelFor"), bDisableCollisionParallelFor, TEXT("Disable parallel execution for Chaos Collisions (also disabled by DisableParticleParallelFor)"));
+#else
+	const bool bDisablePhysicsParallelFor = false;
+#endif
+}
 
 void Chaos::PhysicsParallelFor(int32 InNum, TFunctionRef<void(int32)> InCallable, bool bForceSingleThreaded)
 {
 	// Passthrough for now, except with global flag to disable parallel
-	::ParallelFor(InNum, InCallable, GNoParallelFor || bForceSingleThreaded);
+	::ParallelFor(InNum, InCallable, bDisablePhysicsParallelFor || bForceSingleThreaded);
 }
 
 //class FRecursiveDivideTask
diff --git a/Engine/Source/Runtime/Experimental/Chaos/Private/Chaos/PBDCollisionConstraint.cpp b/Engine/Source/Runtime/Experimental/Chaos/Private/Chaos/PBDCollisionConstraint.cpp
--- a/Engine/Source/Runtime/Experimental/Chaos/Private/Chaos/PBDCollisionConstraint.cpp
+++ b/Engine/Source/Runtime/Experimental/Chaos/Private/Chaos/PBDCollisionConstraint.cpp
@@ -97,7 +97,7 @@ TPBDCollisionConstraint<T, d>::TPBDCollisionConstraint(const TPBDRigidsSOAs<T,d>
 {
 }
 
-DECLARE_CYCLE_STAT(TEXT("CollisionConstraint::Reset"), STAT_CollisionConstraintsReset, STATGROUP_Chaos);
+DECLARE_CYCLE_STAT(TEXT("TPBDCollisionConstraint::Reset"), STAT_CollisionConstraintsReset, STATGROUP_Chaos);
 
 template<typename T, int d>
 void TPBDCollisionConstraint<T, d>::Reset(/*const TPBDRigidParticles<T, d>& InParticles, const TArray<int32>& InIndices*/)
@@ -223,9 +223,6 @@ DEFINE_STAT(STAT_ComputeConstraintsNP);
 DEFINE_STAT(STAT_ComputeConstraintsBP);
 DEFINE_STAT(STAT_ComputeConstraintsSU);
 
-CHAOS_API int32 CollisionConstraintsForceSingleThreaded = 0;
-FAutoConsoleVariableRef CVarCollisionConstraintsForceSingleThreaded(TEXT("p.Chaos.Collision.ForceSingleThreaded"), CollisionConstraintsForceSingleThreaded, TEXT("CollisionConstraintsForceSingleThreaded"));
-
 template<typename T, int d>
 template <bool bGatherStats>
 void TPBDCollisionConstraint<T, d>::ComputeConstraints(const FAccelerationStructure& AccelerationStructure, T Dt)
@@ -424,13 +421,13 @@ void TPBDCollisionConstraint<T, d>::UpdateConstraintsHelper(/*const TPBDRigidPar
 #endif
 }
 
-DECLARE_CYCLE_STAT(TEXT("Reconcile Updated Constraints"), STAT_ReconcileConstraints2, STATGROUP_Chaos);
+DECLARE_CYCLE_STAT(TEXT("TPBDCollisionConstraint::ReconcileUpdatedConstraints"), STAT_ReconcileConstraints, STATGROUP_ChaosWide);
 template<typename T, int d>
 void TPBDCollisionConstraint<T, d>::UpdateConstraints(/*const TPBDRigidParticles<T, d>& InParticles, const TArray<int32>& InIndices,*/ T Dt, const TSet<TGeometryParticleHandle<T, d>*>& AddedParticles)
 {
 #if CHAOS_PARTICLEHANDLE_TODO
 	{
-		SCOPE_CYCLE_COUNTER(STAT_ReconcileConstraints2);
+		SCOPE_CYCLE_COUNTER(STAT_ReconcileConstraints);
 
 		// Updating post-clustering, we will have invalid constraints
 		int32 NumRemovedConstraints = 0;
@@ -721,17 +718,18 @@ void TPBDCollisionConstraint<T, d>::Apply(const T Dt, FRigidBodyContactConstrain
 	}
 }
 
-DECLARE_CYCLE_STAT(TEXT("Apply"), STAT_Apply2, STATGROUP_ChaosWide);
+DECLARE_CYCLE_STAT(TEXT("TPBDCollisionConstraint::Apply"), STAT_Apply, STATGROUP_Chaos);
 template<typename T, int d>
 void TPBDCollisionConstraint<T, d>::Apply(const T Dt, const TArray<FConstraintHandle*>& InConstraintHandles, const int32 It, const int32 NumIts)
 {
+	SCOPE_CYCLE_COUNTER(STAT_Apply);
 	if (bEnableVelocitySolve)
 	{
 		PhysicsParallelFor(InConstraintHandles.Num(), [&](int32 ConstraintHandleIndex) {
 			FConstraintHandle* ConstraintHandle = InConstraintHandles[ConstraintHandleIndex];
 			check(ConstraintHandle != nullptr);
 			Apply(Dt, Constraints[ConstraintHandle->GetConstraintIndex()]);
-			});
+			}, bDisableCollisionParallelFor);
 	}
 
 	if (PostApplyCallback != nullptr)
@@ -740,11 +738,9 @@ void TPBDCollisionConstraint<T, d>::Apply(const T Dt, const TArray<FConstraintHa
 	}
 }
 
-DECLARE_CYCLE_STAT(TEXT("ApplyPushOut"), STAT_ApplyPushOut2, STATGROUP_ChaosWide);
 template<typename T, int d>
 void TPBDCollisionConstraint<T, d>::ApplyPushOut(const T Dt, FRigidBodyContactConstraint& Constraint, const TSet<TGeometryParticleHandle<T,d>*>& IsTemporarilyStatic, int32 Iteration, int32 NumIterations, bool &NeedsAnotherIteration)
 {
-	SCOPE_CYCLE_COUNTER(STAT_ApplyPushOut2);
 	TGeometryParticleHandle<T, d>* Particle0 = Constraint.Particle;
 	TGeometryParticleHandle<T, d>* Particle1 = Constraint.Levelset;
 	TPBDRigidParticleHandle<T, d>* PBDRigid0 = Particle0->AsDynamic();
@@ -857,11 +853,11 @@ void TPBDCollisionConstraint<T, d>::ApplyPushOut(const T Dt, FRigidBodyContactCo
 	}
 }
 
-//DECLARE_CYCLE_STAT(TEXT("ApplyPushOut"), STAT_ApplyPushOut, STATGROUP_ChaosWide);
+DECLARE_CYCLE_STAT(TEXT("TPBDCollisionConstraint::ApplyPushOut"), STAT_ApplyPushOut, STATGROUP_Chaos);
 template<typename T, int d>
 bool TPBDCollisionConstraint<T, d>::ApplyPushOut(const T Dt, const TArray<FConstraintHandle*>& InConstraintHandles, const TSet<TGeometryParticleHandle<T,d>*>& IsTemporarilyStatic, int32 Iteration, int32 NumIterations)
 {
-	SCOPE_CYCLE_COUNTER(STAT_ApplyPushOut2);
+	SCOPE_CYCLE_COUNTER(STAT_ApplyPushOut);
 
 	bool NeedsAnotherIteration = false;
 
@@ -871,7 +867,7 @@ bool TPBDCollisionConstraint<T, d>::ApplyPushOut(const T Dt, const TArray<FConst
 			FConstraintHandle* ConstraintHandle = InConstraintHandles[ConstraintHandleIndex];
 			check(ConstraintHandle != nullptr);
 			ApplyPushOut(Dt, Constraints[ConstraintHandle->GetConstraintIndex()], IsTemporarilyStatic, Iteration, NumIterations, NeedsAnotherIteration);
-			});
+			}, bDisableCollisionParallelFor);
 	}
 
 	if (PostApplyPushOutCallback != nullptr)
@@ -962,12 +958,12 @@ bool SampleObjectNormalAverageHelper2(const TImplicitObject<T, d>& Object, const
 	return false;
 }
 
-DECLARE_CYCLE_STAT(TEXT("UpdateLevelsetPartial"), STAT_UpdateLevelsetPartial2, STATGROUP_ChaosWide);
-DECLARE_CYCLE_STAT(TEXT("UpdateLevelsetFindParticles"), STAT_UpdateLevelsetFindParticles2, STATGROUP_ChaosWide);
-DECLARE_CYCLE_STAT(TEXT("UpdateLevelsetBVHTraversal"), STAT_UpdateLevelsetBVHTraversal2, STATGROUP_ChaosWide);
-DECLARE_CYCLE_STAT(TEXT("UpdateLevelsetSignedDistance"), STAT_UpdateLevelsetSignedDistance2, STATGROUP_ChaosWide);
-DECLARE_CYCLE_STAT(TEXT("UpdateLevelsetAll"), STAT_UpdateLevelsetAll2, STATGROUP_ChaosWide);
-DECLARE_CYCLE_STAT(TEXT("SampleObject"), STAT_SampleObject2, STATGROUP_ChaosWide);
+DECLARE_CYCLE_STAT(TEXT("TPBDCollisionConstraint::UpdateLevelsetPartial"), STAT_UpdateLevelsetPartial, STATGROUP_ChaosWide);
+DECLARE_CYCLE_STAT(TEXT("TPBDCollisionConstraint::UpdateLevelsetFindParticles"), STAT_UpdateLevelsetFindParticles, STATGROUP_ChaosWide);
+DECLARE_CYCLE_STAT(TEXT("TPBDCollisionConstraint::UpdateLevelsetBVHTraversal"), STAT_UpdateLevelsetBVHTraversal, STATGROUP_ChaosWide);
+DECLARE_CYCLE_STAT(TEXT("TPBDCollisionConstraint::UpdateLevelsetSignedDistance"), STAT_UpdateLevelsetSignedDistance, STATGROUP_ChaosWide);
+DECLARE_CYCLE_STAT(TEXT("TPBDCollisionConstraint::UpdateLevelsetAll"), STAT_UpdateLevelsetAll, STATGROUP_ChaosWide);
+DECLARE_CYCLE_STAT(TEXT("TPBDCollisionConstraint::SampleObject"), STAT_SampleObject, STATGROUP_ChaosWide);
 
 int32 NormalAveraging2 = 1;
 FAutoConsoleVariableRef CVarNormalAveraging2(TEXT("p.NormalAveraging2"), NormalAveraging2, TEXT(""));
@@ -978,7 +974,7 @@ FAutoConsoleVariableRef CVarSampleMinParticlesForAcceleration2(TEXT("p.SampleMin
 template <ECollisionUpdateType UpdateType, typename T, int d>
 void SampleObject2(const TImplicitObject<T, d>& Object, const TRigidTransform<T, d>& ObjectTransform, const TBVHParticles<T, d>& SampleParticles, const TRigidTransform<T, d>& SampleParticlesTransform, T Thickness, TRigidBodyContactConstraint<T, d>& Constraint)
 {
-	SCOPE_CYCLE_COUNTER(STAT_SampleObject2);
+	SCOPE_CYCLE_COUNTER(STAT_SampleObject);
 	TRigidBodyContactConstraint<T, d> AvgConstraint;
 	AvgConstraint.Particle = Constraint.Particle;
 	AvgConstraint.Levelset = Constraint.Levelset;
@@ -993,16 +989,16 @@ void SampleObject2(const TImplicitObject<T, d>& Object, const TRigidTransform<T,
 	const TRigidTransform<T, d> & SampleToObjectTM = SampleParticlesTransform.GetRelativeTransform(ObjectTransform);
 	if (NumParticles > SampleMinParticlesForAcceleration2 && Object.HasBoundingBox())
 	{
-		SCOPE_CYCLE_COUNTER(STAT_UpdateLevelsetPartial2);
+		SCOPE_CYCLE_COUNTER(STAT_UpdateLevelsetPartial);
 		TBox<T, d> ImplicitBox = Object.BoundingBox().TransformedBox(ObjectTransform.GetRelativeTransform(SampleParticlesTransform));
 		ImplicitBox.Thicken(Thickness);
 		TArray<int32> PotentialParticles;
 		{
-			SCOPE_CYCLE_COUNTER(STAT_UpdateLevelsetFindParticles2);
+			SCOPE_CYCLE_COUNTER(STAT_UpdateLevelsetFindParticles);
 			PotentialParticles = SampleParticles.FindAllIntersections(ImplicitBox);
 		}
 		{
-			SCOPE_CYCLE_COUNTER(STAT_UpdateLevelsetSignedDistance2);
+			SCOPE_CYCLE_COUNTER(STAT_UpdateLevelsetSignedDistance);
 			for (int32 i : PotentialParticles)
 			{
 				if (NormalAveraging2 && UpdateType != ECollisionUpdateType::Any)	//if we just want one don't bother with normal
@@ -1026,7 +1022,7 @@ void SampleObject2(const TImplicitObject<T, d>& Object, const TRigidTransform<T,
 	}
 	else
 	{
-		SCOPE_CYCLE_COUNTER(STAT_UpdateLevelsetAll2);
+		SCOPE_CYCLE_COUNTER(STAT_UpdateLevelsetAll);
 		for (int32 i = 0; i < NumParticles; ++i)
 		{
 			if (NormalAveraging2 && UpdateType != ECollisionUpdateType::Any)	//if we just want one don't bother with normal
@@ -1083,7 +1079,7 @@ void SampleObject2(const TImplicitObject<T, d>& Object, const TRigidTransform<T,
 template<ECollisionUpdateType UpdateType>
 void SampleObject2(const TImplicitObject<float, 3>& Object, const TRigidTransform<float, 3>& ObjectTransform, const TBVHParticles<float, 3>& SampleParticles, const TRigidTransform<float, 3>& SampleParticlesTransform, float Thickness, TRigidBodyContactConstraint<float, 3>& Constraint)
 {
-	SCOPE_CYCLE_COUNTER(STAT_SampleObject2);
+	SCOPE_CYCLE_COUNTER(STAT_SampleObject);
 	TRigidBodyContactConstraint<float, 3> AvgConstraint;
 	AvgConstraint.Particle = Constraint.Particle;
 	AvgConstraint.Levelset = Constraint.Levelset;
@@ -1099,16 +1095,16 @@ void SampleObject2(const TImplicitObject<float, 3>& Object, const TRigidTransfor
 
 	if (NumParticles > SampleMinParticlesForAcceleration2 && Object.HasBoundingBox())
 	{
-		SCOPE_CYCLE_COUNTER(STAT_UpdateLevelsetPartial2);
+		SCOPE_CYCLE_COUNTER(STAT_UpdateLevelsetPartial);
 		TBox<float, 3> ImplicitBox = Object.BoundingBox().TransformedBox(ObjectTransform.GetRelativeTransform(SampleParticlesTransform));
 		ImplicitBox.Thicken(Thickness);
 		TArray<int32> PotentialParticles;
 		{
-			SCOPE_CYCLE_COUNTER(STAT_UpdateLevelsetFindParticles2);
+			SCOPE_CYCLE_COUNTER(STAT_UpdateLevelsetFindParticles);
 			PotentialParticles = SampleParticles.FindAllIntersections(ImplicitBox);
 		}
 		{
-			SCOPE_CYCLE_COUNTER(STAT_UpdateLevelsetSignedDistance2);
+			SCOPE_CYCLE_COUNTER(STAT_UpdateLevelsetSignedDistance);
 
 			if (Object.GetType(true) == ImplicitObjectType::LevelSet && PotentialParticles.Num() > 0)
 			{
@@ -1218,7 +1214,7 @@ void SampleObject2(const TImplicitObject<float, 3>& Object, const TRigidTransfor
 	}
 	else
 	{
-		SCOPE_CYCLE_COUNTER(STAT_UpdateLevelsetAll2);
+		SCOPE_CYCLE_COUNTER(STAT_UpdateLevelsetAll);
 		if (Object.GetType(true) == ImplicitObjectType::LevelSet && NumParticles > 0)
 		{
 			const TLevelSet<float, 3>* LevelSet = Object.GetObject<Chaos::TLevelSet<float, 3>>();
@@ -1554,11 +1550,11 @@ void UpdateCapsuleBoxConstraint(const TCapsule<T>& A, const TRigidTransform<T, d
 
 }
 
-DECLARE_CYCLE_STAT(TEXT("FindRelevantShapes2"), STAT_FindRelevantShapes2, STATGROUP_ChaosWide);
+DECLARE_CYCLE_STAT(TEXT("TPBDCollisionConstraint::FindRelevantShapes"), STAT_FindRelevantShapes, STATGROUP_ChaosWide);
 template <typename T, int d>
 TArray<Pair<const TImplicitObject<T, d>*, TRigidTransform<T, d>>> FindRelevantShapes2(const TImplicitObject<T,d>* ParticleObj, const TRigidTransform<T,d>& ParticlesTM, const TImplicitObject<T,d>& LevelsetObj, const TRigidTransform<T,d>& LevelsetTM, const T Thickness)
 {
-	SCOPE_CYCLE_COUNTER(STAT_FindRelevantShapes2);
+	SCOPE_CYCLE_COUNTER(STAT_FindRelevantShapes);
 	TArray<Pair<const TImplicitObject<T, d>*, TRigidTransform<T, d>>> RelevantShapes;
 	//find all levelset inner objects
 	if (ParticleObj)
@@ -1586,11 +1582,11 @@ TArray<Pair<const TImplicitObject<T, d>*, TRigidTransform<T, d>>> FindRelevantSh
 	return RelevantShapes;
 }
 
-DECLARE_CYCLE_STAT(TEXT("UpdateUnionUnionConstraint"), STAT_UpdateUnionUnionConstraint2, STATGROUP_ChaosWide);
+DECLARE_CYCLE_STAT(TEXT("TPBDCollisionConstraint::UpdateUnionUnionConstraint"), STAT_UpdateUnionUnionConstraint, STATGROUP_ChaosWide);
 template<ECollisionUpdateType UpdateType, typename T, int d>
 void UpdateUnionUnionConstraint(const T Thickness, TRigidBodyContactConstraint<T, d>& Constraint)
 {
-	SCOPE_CYCLE_COUNTER(STAT_UpdateUnionUnionConstraint2);
+	SCOPE_CYCLE_COUNTER(STAT_UpdateUnionUnionConstraint);
 
 	TGenericParticleHandle<T, d> Particle0 = Constraint.Particle;
 	TGenericParticleHandle<T, d> Particle1 = Constraint.Levelset;
@@ -1620,11 +1616,11 @@ void UpdateUnionUnionConstraint(const T Thickness, TRigidBodyContactConstraint<T
 	}
 }
 
-DECLARE_CYCLE_STAT(TEXT("UpdateSingleUnionConstraint"), STAT_UpdateSingleUnionConstraint2, STATGROUP_ChaosWide);
+DECLARE_CYCLE_STAT(TEXT("TPBDCollisionConstraint::UpdateSingleUnionConstraint"), STAT_UpdateSingleUnionConstraint, STATGROUP_ChaosWide);
 template<ECollisionUpdateType UpdateType, typename T, int d>
 void UpdateSingleUnionConstraint(const T Thickness, TRigidBodyContactConstraint<T, d>& Constraint)
 {
-	SCOPE_CYCLE_COUNTER(STAT_UpdateSingleUnionConstraint2);
+	SCOPE_CYCLE_COUNTER(STAT_UpdateSingleUnionConstraint);
 
 	TGenericParticleHandle<T, d> Particle0 = Constraint.Particle;
 	TGenericParticleHandle<T, d> Particle1 = Constraint.Levelset;
@@ -1644,12 +1640,12 @@ void UpdateSingleUnionConstraint(const T Thickness, TRigidBodyContactConstraint<
 	}
 }
 
-DECLARE_CYCLE_STAT(TEXT("UpdateLevelsetConstraint"), STAT_UpdateLevelsetConstraint2, STATGROUP_ChaosWide);
+DECLARE_CYCLE_STAT(TEXT("TPBDCollisionConstraint::UpdateLevelsetConstraint"), STAT_UpdateLevelsetConstraint, STATGROUP_ChaosWide);
 template<typename T, int d>
 template<ECollisionUpdateType UpdateType>
 void TPBDCollisionConstraint<T, d>::UpdateLevelsetConstraint(const T Thickness, FRigidBodyContactConstraint& Constraint)
 {
-	SCOPE_CYCLE_COUNTER(STAT_UpdateLevelsetConstraint2);
+	SCOPE_CYCLE_COUNTER(STAT_UpdateLevelsetConstraint);
 	
 	TGenericParticleHandle<T, d> Particle0 = Constraint.Particle;
 	TRigidTransform<T, d> ParticlesTM = TRigidTransform<T, d>(Particle0->P(), Particle0->Q());
@@ -1683,11 +1679,11 @@ void UpdateConvexConstraintsUsingCoreShapes(const TImplicitObject<T, d> & AObj,
 	GJKCoreShapeIntersection<T, 3>(AObj, ATM, BObj, BTM, Constraint.Location, Constraint.Phi, Constraint.Normal, Thickness);
 }
 
-DECLARE_CYCLE_STAT(TEXT("UpdateUnionLevelsetConstraint"), STAT_UpdateUnionLevelsetConstraint2, STATGROUP_ChaosWide);
+DECLARE_CYCLE_STAT(TEXT("TPBDCollisionConstraint::UpdateUnionLevelsetConstraint"), STAT_UpdateUnionLevelsetConstraint, STATGROUP_ChaosWide);
 template<ECollisionUpdateType UpdateType, typename T, int d>
 void UpdateUnionLevelsetConstraint(const T Thickness, TRigidBodyContactConstraint<T, d>& Constraint)
 {
-	SCOPE_CYCLE_COUNTER(STAT_UpdateUnionLevelsetConstraint2);
+	SCOPE_CYCLE_COUNTER(STAT_UpdateUnionLevelsetConstraint);
 
 	TGenericParticleHandle<T, d> Particle0 = Constraint.Particle;
 	TGenericParticleHandle<T, d> Particle1 = Constraint.Levelset;
@@ -1759,11 +1755,11 @@ void UpdateUnionLevelsetConstraint(const T Thickness, TRigidBodyContactConstrain
 }
 
 
-DECLARE_CYCLE_STAT(TEXT("UpdateLevelsetUnionConstraint"), STAT_UpdateLevelsetUnionConstraint2, STATGROUP_ChaosWide);
+DECLARE_CYCLE_STAT(TEXT("TPBDCollisionConstraint::UpdateLevelsetUnionConstraint"), STAT_UpdateLevelsetUnionConstraint, STATGROUP_ChaosWide);
 template<ECollisionUpdateType UpdateType, typename T, int d>
 void UpdateLevelsetUnionConstraint(const T Thickness, TRigidBodyContactConstraint<T, d>& Constraint)
 {
-	SCOPE_CYCLE_COUNTER(STAT_UpdateLevelsetUnionConstraint2);
+	SCOPE_CYCLE_COUNTER(STAT_UpdateLevelsetUnionConstraint);
 
 	TGenericParticleHandle<T, d> Particle0 = Constraint.Particle;
 	TGenericParticleHandle<T, d> Particle1 = Constraint.Levelset;
@@ -2156,13 +2152,13 @@ void UpdateConstraintImp2(const TImplicitObject<T, d>& ParticleObject, const TRi
 	}
 }
 
-DECLARE_CYCLE_STAT(TEXT("UpdateConstraint"), STAT_UpdateConstraint2, STATGROUP_ChaosWide);
+DECLARE_CYCLE_STAT(TEXT("TPBDCollisionConstraint::UpdateConstraint"), STAT_UpdateConstraint, STATGROUP_ChaosWide);
 
 template<typename T, int d>
 template<ECollisionUpdateType UpdateType>
 void TPBDCollisionConstraint<T, d>::UpdateConstraint(const T Thickness, FRigidBodyContactConstraint& Constraint)
 {
-	SCOPE_CYCLE_COUNTER(STAT_UpdateConstraint2);
+	SCOPE_CYCLE_COUNTER(STAT_UpdateConstraint);
 
 	Constraint.Phi = Thickness;
 	
diff --git a/Engine/Source/Runtime/Experimental/Chaos/Private/Chaos/PBDJointConstraints.cpp b/Engine/Source/Runtime/Experimental/Chaos/Private/Chaos/PBDJointConstraints.cpp
--- a/Engine/Source/Runtime/Experimental/Chaos/Private/Chaos/PBDJointConstraints.cpp
+++ b/Engine/Source/Runtime/Experimental/Chaos/Private/Chaos/PBDJointConstraints.cpp
@@ -6,12 +6,16 @@
 #include "Chaos/PBDJointConstraintUtilities.h"
 #include "Chaos/Utilities.h"
 #include "ChaosLog.h"
+#include "ChaosStats.h"
 
 #include "HAL/IConsoleManager.h"
 
 
 namespace Chaos
 {
+	DECLARE_CYCLE_STAT(TEXT("TPBDJointConstraints::Apply"), STAT_ApplyJointConstraints, STATGROUP_Chaos);
+
+
 	//
 	// Constraint Handle
 	//
@@ -317,6 +321,8 @@ namespace Chaos
 	template<class T, int d>
 	void TPBDJointConstraints<T, d>::Apply(const T Dt, const TArray<FConstraintHandle*>& InConstraintHandles, const int32 It, const int32 NumIts)
 	{
+		SCOPE_CYCLE_COUNTER(STAT_ApplyJointConstraints);
+
 		// @todo(ccaulfield): make sorting optional
 		// @todo(ccaulfield): handles should be sorted by level by the constraint rule/graph
 		// @todo(ccaulfield): the best sort order depends on whether we are freezing.
diff --git a/Engine/Source/Runtime/Experimental/Chaos/Private/Chaos/PBDRigidsEvolutionGBF.cpp b/Engine/Source/Runtime/Experimental/Chaos/Private/Chaos/PBDRigidsEvolutionGBF.cpp
--- a/Engine/Source/Runtime/Experimental/Chaos/Private/Chaos/PBDRigidsEvolutionGBF.cpp
+++ b/Engine/Source/Runtime/Experimental/Chaos/Private/Chaos/PBDRigidsEvolutionGBF.cpp
@@ -40,14 +40,17 @@ int DisableThreshold = 5;
 FAutoConsoleVariableRef CVarDisableThreshold(TEXT("p.DisableThreshold2"), DisableThreshold, TEXT("Disable threshold frames to transition to sleeping"));
 
 
-DECLARE_CYCLE_STAT(TEXT("TPBDRigidsEvolutionGBF::AdvanceOneTimeStep"), STAT_AdvanceOneTimeStep, STATGROUP_Chaos);
-DECLARE_CYCLE_STAT(TEXT("TPBDRigidsEvolutionGBF::Integrate"), STAT_Integrate, STATGROUP_Chaos);
-DECLARE_CYCLE_STAT(TEXT("TPBDRigidsEvolutionGBF::KinematicTargets"), STAT_KinematicTargets, STATGROUP_Chaos);
-DECLARE_CYCLE_STAT(TEXT("TPBDRigidsEvolutionGBF::UpdateConstraintPositionBasedState"), STAT_UpdateConstraintPositionBasedState, STATGROUP_Chaos);
-DECLARE_CYCLE_STAT(TEXT("TPBDRigidsEvolutionGBF::CreateConstraintGraph"), STAT_CreateConstraintGraph, STATGROUP_Chaos);
-DECLARE_CYCLE_STAT(TEXT("TPBDRigidsEvolutionGBF::CreateIslands"), STAT_CreateIslands, STATGROUP_Chaos);
-DECLARE_CYCLE_STAT(TEXT("TPBDRigidsEvolutionGBF::ParallelSolve"), STAT_ParallelSolve, STATGROUP_Chaos);
-DECLARE_CYCLE_STAT(TEXT("TPBDRigidsEvolutionGBF::DeactivateSleep"), STAT_DeactivateSleep, STATGROUP_Chaos);
+DECLARE_CYCLE_STAT(TEXT("TPBDRigidsEvolutionGBF::AdvanceOneTimeStep"), STAT_Evo_AdvanceOneTimeStep, STATGROUP_Chaos);
+DECLARE_CYCLE_STAT(TEXT("TPBDRigidsEvolutionGBF::Integrate"), STAT_Evo_Integrate, STATGROUP_Chaos);
+DECLARE_CYCLE_STAT(TEXT("TPBDRigidsEvolutionGBF::KinematicTargets"), STAT_Evo_KinematicTargets, STATGROUP_Chaos);
+DECLARE_CYCLE_STAT(TEXT("TPBDRigidsEvolutionGBF::ApplyConstraints"), STAT_Evo_ApplyConstraints, STATGROUP_Chaos);
+DECLARE_CYCLE_STAT(TEXT("TPBDRigidsEvolutionGBF::UpdateVelocities"), STAT_Evo_UpdateVelocites, STATGROUP_Chaos);
+DECLARE_CYCLE_STAT(TEXT("TPBDRigidsEvolutionGBF::ApplyPushOut"), STAT_Evo_ApplyPushOut, STATGROUP_Chaos);
+DECLARE_CYCLE_STAT(TEXT("TPBDRigidsEvolutionGBF::UpdateConstraintPositionBasedState"), STAT_Evo_UpdateConstraintPositionBasedState, STATGROUP_Chaos);
+DECLARE_CYCLE_STAT(TEXT("TPBDRigidsEvolutionGBF::CreateConstraintGraph"), STAT_Evo_CreateConstraintGraph, STATGROUP_Chaos);
+DECLARE_CYCLE_STAT(TEXT("TPBDRigidsEvolutionGBF::CreateIslands"), STAT_Evo_CreateIslands, STATGROUP_Chaos);
+DECLARE_CYCLE_STAT(TEXT("TPBDRigidsEvolutionGBF::ParallelSolve"), STAT_Evo_ParallelSolve, STATGROUP_Chaos);
+DECLARE_CYCLE_STAT(TEXT("TPBDRigidsEvolutionGBF::DeactivateSleep"), STAT_Evo_DeactivateSleep, STATGROUP_Chaos);
 
 int32 SerializeEvolution = 0;
 FAutoConsoleVariableRef CVarSerializeEvolution(TEXT("p.SerializeEvolution"), SerializeEvolution, TEXT(""));
@@ -87,7 +90,7 @@ void SerializeToDisk(TEvolution& Evolution)
 template <typename T, int d>
 void TPBDRigidsEvolutionGBF<T, d>::AdvanceOneTimeStep(T Dt)
 {
-	SCOPE_CYCLE_COUNTER(STAT_AdvanceOneTimeStep);
+	SCOPE_CYCLE_COUNTER(STAT_Evo_AdvanceOneTimeStep);
 
 #if !UE_BUILD_SHIPPING
 	if (SerializeEvolution)
@@ -97,12 +100,12 @@ void TPBDRigidsEvolutionGBF<T, d>::AdvanceOneTimeStep(T Dt)
 #endif
 
 	{
-		SCOPE_CYCLE_COUNTER(STAT_Integrate);
+		SCOPE_CYCLE_COUNTER(STAT_Evo_Integrate);
 		Integrate(Particles.GetNonDisabledDynamicView(), Dt);	//Question: should we use an awake view?
 	}
 
 	{
-		SCOPE_CYCLE_COUNTER(STAT_KinematicTargets);
+		SCOPE_CYCLE_COUNTER(STAT_Evo_KinematicTargets);
 		ApplyKinematicTargets(Dt);
 	}
 
@@ -117,15 +120,15 @@ void TPBDRigidsEvolutionGBF<T, d>::AdvanceOneTimeStep(T Dt)
 	}
 
 	{
-		SCOPE_CYCLE_COUNTER(STAT_UpdateConstraintPositionBasedState);
+		SCOPE_CYCLE_COUNTER(STAT_Evo_UpdateConstraintPositionBasedState);
 		UpdateConstraintPositionBasedState(Dt);
 	}
 	{
-		SCOPE_CYCLE_COUNTER(STAT_CreateConstraintGraph);
+		SCOPE_CYCLE_COUNTER(STAT_Evo_CreateConstraintGraph);
 		CreateConstraintGraph();
 	}
 	{
-		SCOPE_CYCLE_COUNTER(STAT_CreateIslands);
+		SCOPE_CYCLE_COUNTER(STAT_Evo_CreateIslands);
 		CreateIslands();
 	}
 
@@ -141,20 +144,29 @@ void TPBDRigidsEvolutionGBF<T, d>::AdvanceOneTimeStep(T Dt)
 	SleepedIslands.SetNum(ConstraintGraph.NumIslands());
 	if(Dt > 0)
 	{
-		SCOPE_CYCLE_COUNTER(STAT_ParallelSolve);
+		SCOPE_CYCLE_COUNTER(STAT_Evo_ParallelSolve);
 		PhysicsParallelFor(ConstraintGraph.NumIslands(), [&](int32 Island) {
 			const TArray<TGeometryParticleHandle<T, d>*>& IslandParticles = ConstraintGraph.GetIslandParticles(Island);
 
-			ApplyConstraints(Dt, Island);
+			{
+				SCOPE_CYCLE_COUNTER(STAT_Evo_ApplyConstraints);
+				ApplyConstraints(Dt, Island);
+			}
 
 			if (PostApplyCallback != nullptr)
 			{
 				PostApplyCallback(Island);
 			}
 
-			UpdateVelocities(Dt, Island);
+			{
+				SCOPE_CYCLE_COUNTER(STAT_Evo_UpdateVelocites);
+				UpdateVelocities(Dt, Island);
+			}
 
-			ApplyPushOut(Dt, Island);
+			{
+				SCOPE_CYCLE_COUNTER(STAT_Evo_ApplyPushOut);
+				ApplyPushOut(Dt, Island);
+			}
 
 			if (PostApplyPushOutCallback != nullptr)
 			{
@@ -198,7 +210,7 @@ void TPBDRigidsEvolutionGBF<T, d>::AdvanceOneTimeStep(T Dt)
 		});
 	}
 	{
-		SCOPE_CYCLE_COUNTER(STAT_DeactivateSleep);
+		SCOPE_CYCLE_COUNTER(STAT_Evo_DeactivateSleep);
 		for (int32 Island = 0; Island < ConstraintGraph.NumIslands(); ++Island)
 		{
 			if (SleepedIslands[Island])
diff --git a/Engine/Source/Runtime/Experimental/Chaos/Private/Chaos/ParticleIterator.cpp b/Engine/Source/Runtime/Experimental/Chaos/Private/Chaos/ParticleIterator.cpp
--- a/Engine/Source/Runtime/Experimental/Chaos/Private/Chaos/ParticleIterator.cpp
+++ b/Engine/Source/Runtime/Experimental/Chaos/Private/Chaos/ParticleIterator.cpp
@@ -4,6 +4,4 @@
 
 namespace Chaos
 {
-	CHAOS_API int32 ChaosParticleParallelFor = 1;
-	FAutoConsoleVariableRef CVarChaosParticleParallelFor(TEXT("p.ChaosParticleParallelFor"), ChaosParticleParallelFor, TEXT("ParticleParallelFor function style for chaos"));
 }
diff --git a/Engine/Source/Runtime/Experimental/Chaos/Public/Chaos/PBDCollisionConstraintImp.h b/Engine/Source/Runtime/Experimental/Chaos/Public/Chaos/PBDCollisionConstraintImp.h
--- a/Engine/Source/Runtime/Experimental/Chaos/Public/Chaos/PBDCollisionConstraintImp.h
+++ b/Engine/Source/Runtime/Experimental/Chaos/Public/Chaos/PBDCollisionConstraintImp.h
@@ -86,6 +86,12 @@ namespace Chaos
 
 #endif
 
+#if UE_BUILD_SHIPPING
+	const bool bDisableCollisionParallelFor = false;
+#else
+	CHAOS_API extern bool bDisableCollisionParallelFor;
+#endif
+
 	template <typename T, int d>
 	struct TSimOverlapVisitor
 	{
@@ -281,7 +287,7 @@ namespace Chaos
 				NarrowPhaseSkipped.Record(NPSkipped);
 				NarrowPhaseRejected.Record(RejectedNP);
 #endif
-			}, bGatherStats || CollisionConstraintsForceSingleThreaded);
+			}, bGatherStats || bDisableCollisionParallelFor);
 			
 			{
 				SCOPE_CYCLE_COUNTER(STAT_ComputeConstraintsSU);
diff --git a/Engine/Source/Runtime/Experimental/Chaos/Public/Chaos/ParticleIterator.h b/Engine/Source/Runtime/Experimental/Chaos/Public/Chaos/ParticleIterator.h
--- a/Engine/Source/Runtime/Experimental/Chaos/Public/Chaos/ParticleIterator.h
+++ b/Engine/Source/Runtime/Experimental/Chaos/Public/Chaos/ParticleIterator.h
@@ -12,7 +12,11 @@ DECLARE_CYCLE_STAT(TEXT("HandleViewParallelForImp"), STAT_HandleViewParallelForI
 
 namespace Chaos
 {
-CHAOS_API extern int32 ChaosParticleParallelFor;
+#if UE_BUILD_SHIPPING
+	const bool bDisableParticleParallelFor = false;
+#else
+	CHAOS_API extern bool bDisableParticleParallelFor;
+#endif
 
 template <typename TSOA>
 class TConstParticleView;
@@ -89,15 +93,13 @@ void ParticlesParallelFor(const TView& Particles, const Lambda& Func, bool bForc
 {
 	SCOPE_CYCLE_COUNTER(STAT_ParticlesParallelFor);
 
-	switch (ChaosParticleParallelFor)
+	if (!bForceSingleThreaded && !bDisableParticleParallelFor)
 	{
-		case 0:
-			Chaos::ParticlesSequentialFor(Particles, Func);
-			break;
-
-		case 1:
-			Chaos::ParticlesParallelForImp(Particles, Func);
-			break;
+		Chaos::ParticlesParallelForImp(Particles, Func);
+	}
+	else
+	{
+		Chaos::ParticlesSequentialFor(Particles, Func);
 	}
 }
 
diff --git a/Engine/Source/Runtime/Launch/Private/LaunchEngineLoop.cpp b/Engine/Source/Runtime/Launch/Private/LaunchEngineLoop.cpp
--- a/Engine/Source/Runtime/Launch/Private/LaunchEngineLoop.cpp
+++ b/Engine/Source/Runtime/Launch/Private/LaunchEngineLoop.cpp
@@ -1299,7 +1299,7 @@ int32 FEngineLoop::PreInitPreStartupScreen(const TCHAR* CmdLine)
 		ILauncherCheckModule::Get().RunLauncher(ELauncherAction::AppLaunch);
 		// We wish to exit
 		RequestEngineExit(TEXT("Run outside of launcher; restarting via launcher"));
-		return 0;
+		return 1;
 	}
 #endif
 
@@ -2481,7 +2481,12 @@ int32 FEngineLoop::PreInitPreStartupScreen(const TCHAR* CmdLine)
 
 int32 FEngineLoop::PreInitPostStartupScreen(const TCHAR* CmdLine)
 {
-	FScopedSlowTask SlowTask(50, NSLOCTEXT("EngineLoop", "EngineLoop_Initializing_PreInitPostStartupScreen", "PreInitPostStartupScreen..."));
+	if (IsEngineExitRequested())
+	{
+		return 0;
+	}
+
+	FScopedSlowTask SlowTask(50, NSLOCTEXT("EngineLoop", "EngineLoop_Initializing", "PreInitPostStartupScreen..."));
 	FScopeCycleCounter CycleCount_AfterStats(GET_STATID(STAT_FEngineLoop_PreInitPostStartupScreen_AfterStats));
 
 	// Restore PreInitContext
@@ -3429,7 +3434,7 @@ bool FEngineLoop::LoadStartupCoreModules()
 	}
 
 #if WITH_UNREAL_DEVELOPER_TOOLS
-	FModuleManager::Get().LoadModule("FunctionalTesting");
+		FModuleManager::Get().LoadModule("FunctionalTesting");
 #endif	//WITH_UNREAL_DEVELOPER_TOOLS
 
 	SlowTask.EnterProgressFrame(30);
@@ -3812,7 +3817,7 @@ void FEngineLoop::Exit()
 		}
 		if (bFlushOnExit)
 		{
-			FlushAsyncLoading();
+	FlushAsyncLoading();
 		}
 		else
 		{
@@ -5074,7 +5079,7 @@ bool FEngineLoop::AppInit( )
 		// Find the editor target
 		FString EditorTargetFileName;
 		for (const FTargetInfo& Target : FDesktopPlatformModule::Get()->GetTargetsForProject(FPaths::GetProjectFilePath()))
-		{
+	{
 			if (Target.Type == EBuildTargetType::Editor)
 			{
 				// Read the editor target receipt
