Merging //UE4/Dev-RenderPlat-Staging to Dev-VR (//UE4/Dev-VR)

Missed in a merge somehow?

diff --git a/Engine/Source/Runtime/Experimental/Chaos/Private/Chaos/ChaosDebugDraw.cpp b/Engine/Source/Runtime/Experimental/Chaos/Private/Chaos/ChaosDebugDraw.cpp
--- a/Engine/Source/Runtime/Experimental/Chaos/Private/Chaos/ChaosDebugDraw.cpp
+++ b/Engine/Source/Runtime/Experimental/Chaos/Private/Chaos/ChaosDebugDraw.cpp
@@ -54,13 +54,14 @@ namespace Chaos
 			case ImplicitObjectType::Sphere:
 			{
 				const TSphere<T, d>* Sphere = Shape->template GetObject<TSphere<T, d>>();
-				FDebugDrawQueue::GetInstance().DrawDebugSphere(ShapeTransform.GetLocation(), Sphere->GetRadius(), 20, Color, false, KINDA_SMALL_NUMBER, DrawPriority, LineThickness);
+				const TVector<T, d> P = ShapeTransform.TransformPosition(Sphere->GetCenter());
+				FDebugDrawQueue::GetInstance().DrawDebugSphere(P, Sphere->GetRadius(), 20, Color, false, KINDA_SMALL_NUMBER, DrawPriority, LineThickness);
 				break;
 			}
 			case ImplicitObjectType::Box:
 			{
 				const TBox<T, d>* Box = Shape->template GetObject<TBox<T, d>>();
-				const TVector<T, d> P = ShapeTransform.TransformPosition(Box->Center());
+				const TVector<T, d> P = ShapeTransform.TransformPosition(Box->GetCenter());
 				FDebugDrawQueue::GetInstance().DrawDebugBox(P, (T)0.5 * Box->Extents(), ShapeTransform.GetRotation(), Color, false, KINDA_SMALL_NUMBER, DrawPriority, LineThickness);
 				break;
 			}
@@ -114,7 +115,8 @@ namespace Chaos
 		template <typename T, int d>
 		void DrawParticleShapesImpl(const TRigidTransform<float, 3>& SpaceTransform, const TGeometryParticleHandle<T, d>* Particle, T ColorScale)
 		{
-			FColor Color = (ColorScale * FColor::White).ToFColor(false);
+			FColor ShapeColor = Particle->AsDynamic() ? FColor::Yellow : FColor::Orange;
+			FColor Color = (((T)0.5 * ColorScale) * ShapeColor).ToFColor(false);
 			TVector<T, d> P = SpaceTransform.TransformPosition(Particle->AsDynamic() ? Particle->AsDynamic()->P() : Particle->X());
 			TRotation<T, d> Q = SpaceTransform.GetRotation() * (Particle->AsDynamic() ? Particle->AsDynamic()->Q() : Particle->R());
 
diff --git a/Engine/Source/Runtime/Experimental/Chaos/Private/Chaos/CollisionConvexMesh.cpp b/Engine/Source/Runtime/Experimental/Chaos/Private/Chaos/CollisionConvexMesh.cpp
new file mode 100644
--- /dev/null
+++ b/Engine/Source/Runtime/Experimental/Chaos/Private/Chaos/CollisionConvexMesh.cpp
@@ -0,0 +1,23 @@
+// Copyright 1998-2019 Epic Games, Inc. All Rights Reserved.
+
+#include "Chaos/CollisionConvexMesh.h"
+#include "HAL/IConsoleManager.h"
+
+namespace Chaos
+{
+	// CVars variables for controlling geometry complexity checking and simplification
+	template <>
+	int32 TConvexBuilder<float>::PerformGeometryCheck = 0;
+
+	template <>
+	int32 TConvexBuilder<float>::PerformGeometryReduction = 0;
+
+	template <>
+	int32 TConvexBuilder<float>::ParticlesThreshold = 50;
+
+	FAutoConsoleVariableRef CVarConvexGeometryCheckEnable(TEXT("p.Chaos.ConvexGeometryCheckEnable"), TConvexBuilder<float>::PerformGeometryCheck, TEXT("Perform convex geometry complexity check for Chaos physics."));
+
+	FAutoConsoleVariableRef CVarConvexGeometrySimplifyEnable(TEXT("p.Chaos.PerformGeometryReduction"), TConvexBuilder<float>::PerformGeometryReduction, TEXT("Perform convex geometry simplification to increase performance in Chaos physics."));
+
+	FAutoConsoleVariableRef CVarConvexParticlesWarningThreshold(TEXT("p.Chaos.ConvexParticlesWarningThreshold"), TConvexBuilder<float>::ParticlesThreshold, TEXT("Threshold beyond which we warn about collision geometry complexity."));
+}
\ No newline at end of file
diff --git a/Engine/Source/Runtime/Experimental/Chaos/Private/Chaos/HeightField.cpp b/Engine/Source/Runtime/Experimental/Chaos/Private/Chaos/HeightField.cpp
--- a/Engine/Source/Runtime/Experimental/Chaos/Private/Chaos/HeightField.cpp
+++ b/Engine/Source/Runtime/Experimental/Chaos/Private/Chaos/HeightField.cpp
@@ -154,7 +154,7 @@ namespace Chaos
 			};
 
 			TVector<T, 3> Points[4];
-			GeomData->GetPoints(FullIndex, Points);
+			GeomData->GetPointsScaled(FullIndex, Points);
 
 			// Test both triangles that are in this cell, as we could hit both in any order
 			TestTriangle(FullIndex * 2, Points[0], Points[1], Points[2]);
@@ -250,7 +250,7 @@ namespace Chaos
 			};
 
 			TVector<T, 3> Points[4];
-			HfData->GetPoints(FullIndex, Points);
+			HfData->GetPointsScaled(FullIndex, Points);
 
 			TestTriangle(FullIndex * 2, Points[0], Points[1], Points[2]);
 			TestTriangle(FullIndex * 2 + 1, Points[2], Points[1], Points[3]);
@@ -273,62 +273,63 @@ namespace Chaos
 
 	};
 
-	template <typename T>
-	THeightField<T>::THeightField(TArray<T>&& Height, int32 NumRows, int32 NumCols, const TVector<T, 3>& InScale)
-		: TImplicitObject<T, 3>(EImplicitObject::HasBoundingBox, ImplicitObjectType::HeightField)
+	template<typename HeightfieldT, typename BufferType>
+	void BuildGeomData(TArrayView<BufferType> BufferView, int32 NumRows, int32 NumCols, const TVector<HeightfieldT, 3>& InScale, TUniqueFunction<HeightfieldT(const BufferType)> ToRealFunc, typename THeightField<HeightfieldT>::FDataType& OutData, TBox<HeightfieldT, 3>& OutBounds)
 	{
 		using FDataType = typename THeightField<HeightfieldT>::FDataType;
 
 		using RealType = typename FDataType::RealType;
 
 		const int32 NumCells = NumRows * NumCols;
-		ensure(Height.Num() == NumCells);
+		ensure(BufferView.Num() == NumCells);
 		ensure(NumRows > 1);
 		ensure(NumCols > 1);
 
 		// Populate data.
-		const int32 NumHeights = Height.Num();
-		GeomData.Heights.SetNum(NumHeights);
+		const int32 NumHeights = BufferView.Num();
+		OutData.Heights.SetNum(NumHeights);
 
-		GeomData.NumRows = NumRows;
-		GeomData.NumCols = NumCols;
-		GeomData.MinValue = Height[0];
-		GeomData.MaxValue = Height[0];
-		GeomData.Scale = InScale;
+		OutData.NumRows = NumRows;
+		OutData.NumCols = NumCols;
+		OutData.MinValue = ToRealFunc(BufferView[0]);
+		OutData.MaxValue = ToRealFunc(BufferView[0]);
+		OutData.Scale = InScale;
 
 		for(int32 HeightIndex = 1; HeightIndex < NumHeights; ++HeightIndex)
 		{
-			const typename FDataType::RealType CurrHeight = (typename FDataType::RealType)Height[HeightIndex];
+			const RealType CurrHeight = ToRealFunc(BufferView[HeightIndex]);
 
-			if(CurrHeight > GeomData.MaxValue)
+			if(CurrHeight > OutData.MaxValue)
 			{
-				GeomData.MaxValue = CurrHeight;
+				OutData.MaxValue = CurrHeight;
 			}
-			else if(CurrHeight < GeomData.MinValue)
+			else if(CurrHeight < OutData.MinValue)
 			{
-				GeomData.MinValue = CurrHeight;
+				OutData.MinValue = CurrHeight;
 			}
 		}
 
-		GeomData.Range = GeomData.MaxValue - GeomData.MinValue;
-		GeomData.HeightPerUnit = GeomData.Range / FDataType::StorageRange;
+		OutData.Range = OutData.MaxValue - OutData.MinValue;
+		OutData.HeightPerUnit = OutData.Range / FDataType::StorageRange;
 
 		for(int32 HeightIndex = 0; HeightIndex < NumHeights; ++HeightIndex)
 		{
-			GeomData.Heights[HeightIndex] = static_cast<typename FDataType::StorageType>((Height[HeightIndex] - GeomData.MinValue) / GeomData.HeightPerUnit);
+			OutData.Heights[HeightIndex] = static_cast<typename FDataType::StorageType>((ToRealFunc(BufferView[HeightIndex]) - OutData.MinValue) / OutData.HeightPerUnit);
 
 			int32 X = HeightIndex % (NumCols);
 			int32 Y = HeightIndex / (NumCols);
-			TVector<T, 3> Position(typename FDataType::RealType(X), typename FDataType::RealType(Y), GeomData.MinValue + GeomData.Heights[HeightIndex] * GeomData.HeightPerUnit);
+			TVector<HeightfieldT, 3> Position(RealType(X), RealType(Y), OutData.MinValue + OutData.Heights[HeightIndex] * OutData.HeightPerUnit);
 			if(HeightIndex == 0)
 			{
-				LocalBounds = TBox<T, 3>(Position * InScale, Position * InScale);
+				OutBounds = TBox<HeightfieldT, 3>(Position * InScale, Position * InScale);
 			}
 			else
 			{
-				LocalBounds.GrowToInclude(Position * InScale);
+				OutBounds.GrowToInclude(Position * InScale);
 			}
 		}
+		OutBounds.Thicken(KINDA_SMALL_NUMBER);
+	}
 
 	template<typename HeightfieldT, typename BufferType>
 	void EditGeomData(TArrayView<BufferType> BufferView, int32 InBeginRow, int32 InBeginCol, int32 NumRows, int32 NumCols, TUniqueFunction<HeightfieldT(const BufferType)> ToRealFunc, typename THeightField<HeightfieldT>::FDataType& OutData, TBox<HeightfieldT, 3>& OutBounds)
@@ -415,22 +416,43 @@ namespace Chaos
 	}
 
 
-		// Cache per-cell bounds
-		const int32 NumX = GeomData.NumCols - 1;
-		const int32 NumY = GeomData.NumRows - 1;
-		GeomData.CellBounds.SetNum(NumX * NumY);
-		for(int32 XIndex = 0; XIndex < NumX; ++XIndex)
 		{
-			for(int32 YIndex = 0; YIndex < NumY; ++YIndex)
+			// No resample, just push new heights into the data
+			for(int32 RowIdx = InBeginRow; RowIdx < EndRow; ++RowIdx)
 			{
-				const TVector<int32, 2> Cell(XIndex, YIndex);
-				TVector<T, 3> Min, Max;
-				CalcCellBounds3D(Cell, Min, Max);
-				GeomData.CellBounds[YIndex * NumX + XIndex] = TBox<T, 3>(Min, Max);
+				for(int32 ColIdx = InBeginCol; ColIdx < EndCol; ++ColIdx)
+				{
+					const int32 HeightIndex = RowIdx * NumCols + ColIdx;
+					const int32 NewSetIndex = (RowIdx - InBeginRow) * NumCols + (ColIdx - InBeginCol);
+					OutData.Heights[HeightIndex] = static_cast<typename FDataType::StorageType>((ToRealFunc(BufferView[NewSetIndex]) - OutData.MinValue) / OutData.HeightPerUnit);
+				}
 			}
 		}
 	}
 
+	template <typename T>
+	THeightField<T>::THeightField(TArray<T>&& Height, int32 NumRows, int32 NumCols, const TVector<T, 3>& InScale)
+		: TImplicitObject<T, 3>(EImplicitObject::HasBoundingBox, ImplicitObjectType::HeightField)
+	{
+		BuildGeomData<T, T>(MakeArrayView(Height), NumRows, NumCols, TVector<T, 3>(1), [](const T InVal) {return InVal; }, GeomData, LocalBounds);
+		CalcBounds();
+		SetScale(InScale);
+	}
+
+	template<typename T>
+	Chaos::THeightField<T>::THeightField(TArrayView<const uint16> InHeights, int32 InNumRows, int32 InNumCols, const TVector<T, 3>& InScale)
+		: TImplicitObject<T, 3>(EImplicitObject::HasBoundingBox, ImplicitObjectType::HeightField)
+	{
+		TUniqueFunction<T(const uint16)> ConversionFunc = [](const T InVal) -> float
+		{
+			return (float)((int32)InVal - 32768);
+		};
+
+		BuildGeomData<T, const uint16>(InHeights, InNumRows, InNumCols, TVector<T, 3>(1), MoveTemp(ConversionFunc), GeomData, LocalBounds);
+		CalcBounds();
+		SetScale(InScale);
+	}
+
 	template<typename T>
 	void Chaos::THeightField<T>::EditHeights(TArrayView<const uint16> InHeights, int32 InBeginRow, int32 InBeginCol, int32 InNumRows, int32 InNumCols)
 	{
@@ -499,6 +521,40 @@ namespace Chaos
 		return false;
 	}
 
+	template<typename T>
+	bool Chaos::THeightField<T>::GetCellBounds2DScaled(const TVector<int32, 2> InCoord, FBounds2D& OutBounds, const TVector<T, 2>& InInflate /*= {0}*/) const
+	{
+		if(FlatGrid.IsValid(InCoord))
+		{
+			const TBox<T, 3>& Bound = GeomData.CellBounds[InCoord[1] * (GeomData.NumCols - 1) + InCoord[0]];
+			OutBounds.Min = {Bound.Min()[0], Bound.Min()[1]};
+			OutBounds.Max = {Bound.Max()[0], Bound.Max()[1]};
+			OutBounds.Min -= InInflate;
+			OutBounds.Max += InInflate;
+			const TVector<float, 2> Scale2D = TVector<T, 2>(GeomData.Scale[0], GeomData.Scale[1]);
+			OutBounds.Min *= Scale2D;
+			OutBounds.Max *= Scale2D;
+			return true;
+		}
+
+		return false;
+	}
+
+	template<typename T>
+	bool Chaos::THeightField<T>::GetCellBounds3DScaled(const TVector<int32, 2> InCoord, TVector<T, 3>& OutMin, TVector<T, 3>& OutMax, const TVector<T, 3>& InInflate /*= {0}*/) const
+	{
+		if(FlatGrid.IsValid(InCoord))
+		{
+			const TBox<T, 3>& Bound = GeomData.CellBounds[InCoord[1] * (GeomData.NumCols - 1) + InCoord[0]];
+			OutMin = Bound.Min() - InInflate;
+			OutMax = Bound.Max() + InInflate;
+			OutMin *= GeomData.Scale;
+			OutMax *= GeomData.Scale;
+			return true;
+		}
+
+		return false;
+	}
 
 	template<typename T>
 	bool Chaos::THeightField<T>::CalcCellBounds3D(const TVector<int32, 2> InCoord, TVector<T, 3>& OutMin, TVector<T, 3>& OutMax, const TVector<T, 3>& InInflate /*= {0}*/) const
@@ -535,7 +591,7 @@ namespace Chaos
 		TVector<T, 2> ClippedFlatRayStart;
 		TVector<T, 2> ClippedFlatRayEnd;
 
-		if(FlattenedBounds.ClipLine(StartPoint, StartPoint + Dir * Length, ClippedFlatRayStart, ClippedFlatRayEnd))
+		if(GetFlatBounds().ClipLine(StartPoint, StartPoint + Dir * Length, ClippedFlatRayStart, ClippedFlatRayEnd))
 		{
 			// The line is now valid and is entirely enclosed by the bounds (and thus, the grid)
 			if(FMath::Abs((ClippedFlatRayEnd - ClippedFlatRayStart).SizeSquared()) < SMALL_NUMBER)
@@ -543,13 +599,16 @@ namespace Chaos
 				// This is a cast down the Z axis, handle in a simpler way as this is the common case
 				const int32 QueryX = (int32)(ClippedFlatRayStart[0] / GeomData.Scale[0]);
 				const int32 QueryY = (int32)(ClippedFlatRayStart[1] / GeomData.Scale[1]);
-				Visitor.VisitRaycast(QueryY * (GeomData.NumCols - 1) + QueryX, CurrentLength);
+				const TVector<int32, 2> Cell = FlatGrid.ClampIndex({QueryX, QueryY});
+
+				Visitor.VisitRaycast(Cell[1] * (GeomData.NumCols - 1) + Cell[0], CurrentLength);
 			}
 			else
 			{
 				// Rasterize the line over the grid
-				TVector<int32, 2> StartCell = FlatGrid.Cell(ClippedFlatRayStart);
-				TVector<int32, 2> EndCell = FlatGrid.Cell(ClippedFlatRayEnd);
+				const TVector<T, 2> Scale2D(GeomData.Scale[0], GeomData.Scale[1]);
+				TVector<int32, 2> StartCell = FlatGrid.Cell(ClippedFlatRayStart/Scale2D);
+				TVector<int32, 2> EndCell = FlatGrid.Cell(ClippedFlatRayEnd/Scale2D);
 
 				// Boundaries might push us one cell over
 				StartCell = FlatGrid.ClampIndex(StartCell);
@@ -643,7 +702,7 @@ namespace Chaos
 	{
 		T CurrentLength = Length;
 
-		FBounds2D InflatedBounds = FlattenedBounds;
+		FBounds2D InflatedBounds = GetFlatBounds();
 		InflatedBounds.Min -= InHalfExtents;
 		InflatedBounds.Max += InHalfExtents;
 
@@ -652,6 +711,7 @@ namespace Chaos
 		const TVector<T, 3> EndPoint = StartPoint + Dir * Length;
 		const TVector<T, 2> Start2D(StartPoint[0], StartPoint[1]);
 		const TVector<T, 2> End2D(EndPoint[0], EndPoint[1]);
+		const TVector<T, 2> Scale2D(GeomData.Scale[0], GeomData.Scale[1]);
 
 		TVector<T, 2> ClippedStart;
 		TVector<T, 2> ClippedEnd;
@@ -659,8 +719,8 @@ namespace Chaos
 		if(InflatedBounds.ClipLine(StartPoint, StartPoint + Dir * Length, ClippedStart, ClippedEnd))
 		{
 			// Rasterize the line over the grid
-			TVector<int32, 2> StartCell = FlatGrid.Cell(ClippedStart);
-			TVector<int32, 2> EndCell = FlatGrid.Cell(ClippedEnd);
+			TVector<int32, 2> StartCell = FlatGrid.Cell(ClippedStart / Scale2D);
+			TVector<int32, 2> EndCell = FlatGrid.Cell(ClippedEnd / Scale2D);
 
 			// Boundaries might push us one cell over
 			StartCell = FlatGrid.ClampIndex(StartCell);
@@ -778,7 +838,7 @@ namespace Chaos
 
 					// Check the current cell, if we hit its 3D bound we can move on to narrow phase
 					const TVector<int32, 2> Coord = CellCoord.Index;
-					if(GetCellBounds3D(Coord, Min, Max, HalfExtents3D) &&
+					if(GetCellBounds3DScaled(Coord, Min, Max, HalfExtents3D) &&
 						TBox<T, 3>::RaycastFast(Min, Max, StartPoint, Dir, InvDir, bParallel, CurrentLength, InvCurrentLength, ToI, HitPoint))
 					{
 						bool bContinue = Visitor.VisitSweep(CellCoord.Index[1] * (GeomData.NumCols - 1) + CellCoord.Index[0], CurrentLength);
@@ -872,10 +932,15 @@ namespace Chaos
 	{
 		OutInterssctions.Reset();
 
-		InFlatBounds.Min = FlattenedBounds.Clamp(InFlatBounds.Min);
-		InFlatBounds.Max = FlattenedBounds.Clamp(InFlatBounds.Max);
-		TVector<int32, 2> MinCell = FlatGrid.Cell(InFlatBounds.Min);
-		TVector<int32, 2> MaxCell = FlatGrid.Cell(InFlatBounds.Max);
+		const FBounds2D FlatBounds = GetFlatBounds();
+		const TVector<T, 2> Scale2D(GeomData.Scale[0], GeomData.Scale[1]);
+
+		InFlatBounds.Min = FlatBounds.Clamp(InFlatBounds.Min);
+		InFlatBounds.Max = FlatBounds.Clamp(InFlatBounds.Max);
+		TVector<int32, 2> MinCell = FlatGrid.Cell(InFlatBounds.Min / Scale2D);
+		TVector<int32, 2> MaxCell = FlatGrid.Cell(InFlatBounds.Max / Scale2D);
+		MinCell = FlatGrid.ClampIndex(MinCell);
+		MaxCell = FlatGrid.ClampIndex(MaxCell);
 
 		const int32 DeltaX = MaxCell[0] - MinCell[0];
 		const int32 DeltaY = MaxCell[1] - MinCell[1];
@@ -891,6 +956,15 @@ namespace Chaos
 		return OutInterssctions.Num() > 0;
 	}
 
+	template<typename T>
+	typename Chaos::THeightField<T>::FBounds2D Chaos::THeightField<T>::GetFlatBounds() const
+	{
+		FBounds2D Result;
+		Result.Min = TVector<T, 2>(CachedBounds.Min()[0], CachedBounds.Min()[1]);
+		Result.Max = TVector<T, 2>(CachedBounds.Max()[0], CachedBounds.Max()[1]);
+		return Result;
+	}
+
 	template <typename T>
 	bool THeightField<T>::Overlap(const TVector<T, 3>& Point, const T Thickness) const
 	{
@@ -1104,7 +1178,7 @@ namespace Chaos
 
 			TVector<T, 3> Points[4];
 
-			GeomData.GetPoints(FaceIndex, Points);
+			GeomData.GetPointsScaled(FaceIndex, Points);
 
 			TVector<T, 3> A;
 			TVector<T, 3> B;
@@ -1134,6 +1208,30 @@ namespace Chaos
 		return TVector<T, 3>(0, 0, 1);
 	}
 
+	template<typename T>
+	void THeightField<T>::CalcBounds()
+	{
+		// Flatten out the Z axis
+		FlattenedBounds = GetFlatBounds();
+
+		BuildQueryData();
+
+		// Cache per-cell bounds
+		const int32 NumX = GeomData.NumCols - 1;
+		const int32 NumY = GeomData.NumRows - 1;
+		GeomData.CellBounds.SetNum(NumX * NumY);
+		for(int32 XIndex = 0; XIndex < NumX; ++XIndex)
+		{
+			for(int32 YIndex = 0; YIndex < NumY; ++YIndex)
+			{
+				const TVector<int32, 2> Cell(XIndex, YIndex);
+				TVector<T, 3> Min, Max;
+				CalcCellBounds3D(Cell, Min, Max);
+				GeomData.CellBounds[YIndex * NumX + XIndex] = TBox<T, 3>(Min, Max);
+			}
+		}
+	}
+
 	template<typename T>
 	void THeightField<T>::BuildQueryData()
 	{
@@ -1150,3 +1248,4 @@ namespace Chaos
 }
 
 template class Chaos::THeightField<float>;
+
diff --git a/Engine/Source/Runtime/Experimental/Chaos/Private/Chaos/PBDCollisionConstraint.cpp b/Engine/Source/Runtime/Experimental/Chaos/Private/Chaos/PBDCollisionConstraint.cpp
--- a/Engine/Source/Runtime/Experimental/Chaos/Private/Chaos/PBDCollisionConstraint.cpp
+++ b/Engine/Source/Runtime/Experimental/Chaos/Private/Chaos/PBDCollisionConstraint.cpp
@@ -107,7 +107,7 @@ void TPBDCollisionConstraint<T, d>::Reset(/*const TPBDRigidParticles<T, d>& InPa
 	int32 Threshold = LifespanCounter - 1; // Maybe this should be solver time?
 	for (int32 Idx = Constraints.Num() - 1; Idx >= 0; Idx--)
 	{
-		if (Constraints[Idx].Lifespan < Threshold)
+		if ((Constraints[Idx].Lifespan < Threshold) || !bEnableCollisions)
 		{
 			RemoveConstraint(Idx);
 		}
@@ -148,6 +148,20 @@ void TPBDCollisionConstraint<T, d>::RemoveConstraints(const TSet<TGeometryPartic
 	Handles.Compact();
 }
 
+template<typename T, int d>
+void TPBDCollisionConstraint<T, d>::ApplyCollisionModifier(const TFunction<ECollisionModifierResult(FRigidBodyContactConstraint& Constraint)>& CollisionModifier)
+{
+	for (int ConstraintIndex = 0; ConstraintIndex < NumConstraints(); ++ConstraintIndex)
+	{
+		ECollisionModifierResult Result = CollisionModifier(Constraints[ConstraintIndex]);
+		if (Result == ECollisionModifierResult::Disabled)
+		{
+			RemoveConstraint(ConstraintIndex);
+			--ConstraintIndex;
+		}
+	}
+}
+
 template<typename T, int d>
 void TPBDCollisionConstraint<T, d>::SetPostComputeCallback(const TRigidBodyContactConstraintsPostComputeCallback<T, d>& Callback)
 {
@@ -246,7 +260,7 @@ void TPBDCollisionConstraint<T, d>::ComputeConstraints(const FAccelerationStruct
 	
 	if (PostComputeCallback != nullptr)
 	{
-		PostComputeCallback(Constraints);
+		PostComputeCallback();
 	}
 }
 
diff --git a/Engine/Source/Runtime/Experimental/Chaos/Private/Chaos/PBDJointConstraintUtilities.cpp b/Engine/Source/Runtime/Experimental/Chaos/Private/Chaos/PBDJointConstraintUtilities.cpp
--- a/Engine/Source/Runtime/Experimental/Chaos/Private/Chaos/PBDJointConstraintUtilities.cpp
+++ b/Engine/Source/Runtime/Experimental/Chaos/Private/Chaos/PBDJointConstraintUtilities.cpp
@@ -176,7 +176,7 @@ namespace Chaos
 		}
 		else
 		{
-			// Line Limits
+			// Line/Square/Cube Limits (no way to author square or cube limits, but would work if we wanted it)
 			TVector<T, d> CX = InCX;
 			if (Motion[0] != EJointMotionType::Locked)
 			{
@@ -531,7 +531,7 @@ namespace Chaos
 			// Map swing axis to ellipse and calculate limit for this swing axis
 			T DotSwing1 = FMath::Abs(TVector<T, d>::DotProduct(SwingAxis01, TJointConstants<T, d>::Swing1Axis()));
 			T DotSwing2 = FMath::Abs(TVector<T, d>::DotProduct(SwingAxis01, TJointConstants<T, d>::Swing2Axis()));
-			SwingAngleMax = FMath::Sqrt(Swing1Limit * DotSwing1 * Swing1Limit * DotSwing1 + Swing2Limit * DotSwing2 * Swing2Limit * DotSwing2);
+			SwingAngleMax = FMath::Sqrt(Swing1Limit * DotSwing2 * Swing1Limit * DotSwing2 + Swing2Limit * DotSwing1 * Swing2Limit * DotSwing1);
 		}
 
 		// Calculate swing error we need to correct
diff --git a/Engine/Source/Runtime/Experimental/Chaos/Private/Chaos/PBDJointConstraints.cpp b/Engine/Source/Runtime/Experimental/Chaos/Private/Chaos/PBDJointConstraints.cpp
--- a/Engine/Source/Runtime/Experimental/Chaos/Private/Chaos/PBDJointConstraints.cpp
+++ b/Engine/Source/Runtime/Experimental/Chaos/Private/Chaos/PBDJointConstraints.cpp
@@ -112,6 +112,7 @@ namespace Chaos
 		, PBDMaxInertiaRatio(5.0f)
 		, FreezeIterations(0)
 		, FrozenIterations(0)
+		, bEnableLinearLimits(1)
 		, bEnableTwistLimits(1)
 		, bEnableSwingLimits(1)
 		, bEnableDrives(1)
@@ -393,10 +394,17 @@ namespace Chaos
 			}
 		}
 
+		const TVector<EJointMotionType, d>& LinearMotion = JointSettings.Motion.LinearMotionTypes;
 		EJointMotionType TwistMotion = JointSettings.Motion.AngularMotionTypes[(int32)EJointAngularConstraintIndex::Twist];
 		EJointMotionType Swing1Motion = JointSettings.Motion.AngularMotionTypes[(int32)EJointAngularConstraintIndex::Swing1];
 		EJointMotionType Swing2Motion = JointSettings.Motion.AngularMotionTypes[(int32)EJointAngularConstraintIndex::Swing2];
 
+		// Disable a constraint if it has any linear limits?
+		if (!Settings.bEnableLinearLimits && ((LinearMotion[0] == EJointMotionType::Limited) || (LinearMotion[1] == EJointMotionType::Limited) || (LinearMotion[2] == EJointMotionType::Limited)))
+		{
+			return;
+		}
+
 		// Apply angular drives
 		if (Settings.bEnableDrives)
 		{
@@ -462,7 +470,6 @@ namespace Chaos
 		}
 
 		// Apply linear constraints
-		const TVector<EJointMotionType, d>& LinearMotion = JointSettings.Motion.LinearMotionTypes;
 		if ((LinearMotion[0] != EJointMotionType::Free) || (LinearMotion[1] != EJointMotionType::Free) || (LinearMotion[2] != EJointMotionType::Free))
 		{
 			TPBDJointUtilities<T, d>::ApplyJointPositionConstraint(Dt, Settings, JointSettings, Index0, Index1, P0, Q0, P1, Q1, InvM0, InvIL0, InvM1, InvIL1);
diff --git a/Engine/Source/Runtime/Experimental/Chaos/Public/Chaos/AABBTree.h b/Engine/Source/Runtime/Experimental/Chaos/Public/Chaos/AABBTree.h
--- a/Engine/Source/Runtime/Experimental/Chaos/Public/Chaos/AABBTree.h
+++ b/Engine/Source/Runtime/Experimental/Chaos/Public/Chaos/AABBTree.h
@@ -249,40 +249,87 @@ public:
 		return TUniquePtr<ISpatialAcceleration<TPayloadType, T, 3>>(new TAABBTree<TPayloadType, TLeafType, T>(*this));
 	}
 
-	virtual void Raycast(const TVector<T, 3>& Start, const TVector<T, 3>& Dir, const T OriginalLength, ISpatialVisitor<TPayloadType, T>& Visitor) const override
+	virtual void Raycast(const TVector<T, 3>& Start, const TVector<T, 3>& Dir, const T Length, ISpatialVisitor<TPayloadType, T>& Visitor) const override
 	{
 		TSpatialVisitor<TPayloadType, T> ProxyVisitor(Visitor);
-		Raycast(Start, Dir, OriginalLength, ProxyVisitor);
+		Raycast(Start, Dir, Length, ProxyVisitor);
 	}
 
 	template <typename SQVisitor>
-	void Raycast(const TVector<T, 3>& Start, const TVector<T, 3>& Dir, const T OriginalLength, SQVisitor& Visitor) const
+	void Raycast(const TVector<T, 3>& Start, const TVector<T, 3>& Dir, const T Length, SQVisitor& Visitor) const
 	{
-		QueryImp<EAABBQueryType::Raycast>(Start, Dir, OriginalLength, TVector<T,3>(), TBox<T,3>(), Visitor);
+		ensure(Length);
+		T CurLength = Length;
+		bool bParallel[3];
+		TVector<T, 3> InvDir;
+
+		T InvLength = 1 / Length;
+		for (int Axis = 0; Axis < 3; ++Axis)
+		{
+			bParallel[Axis] = Dir[Axis] == 0;
+			InvDir[Axis] = bParallel[Axis] ? 0 : 1 / Dir[Axis];
+		}
+
+		QueryImp<EAABBQueryType::Raycast>(Start, Dir, InvDir, bParallel, CurLength, InvLength, TVector<T,3>(), TBox<T,3>(), Visitor);
 	}
 
-	void Sweep(const TVector<T, 3>& Start, const TVector<T, 3>& Dir, T OriginalLength, const TVector<T, 3> QueryHalfExtents, ISpatialVisitor<TPayloadType, T>& Visitor) const override
+	template <typename SQVisitor>
+	bool RaycastFast(const TVector<T, 3>& Start, const TVector<T, 3>& Dir, const TVector<T,3>& InvDir, const bool* bParallel, T& CurLength, T& InvLength, SQVisitor& Visitor) const
+	{
+		return QueryImp<EAABBQueryType::Raycast>(Start, Dir, InvDir, bParallel, CurLength, InvLength, TVector<T, 3>(), TBox<T, 3>(), Visitor);
+	}
+
+	void Sweep(const TVector<T, 3>& Start, const TVector<T, 3>& Dir, const T Length, const TVector<T, 3> QueryHalfExtents, ISpatialVisitor<TPayloadType, T>& Visitor) const override
 	{
 		TSpatialVisitor<TPayloadType, T> ProxyVisitor(Visitor);
-		Sweep(Start, Dir, OriginalLength, QueryHalfExtents, ProxyVisitor);
+		Sweep(Start, Dir, Length, QueryHalfExtents, ProxyVisitor);
 	}
 
 	template <typename SQVisitor>
-	void Sweep(const TVector<T, 3>& Start, const TVector<T, 3>& Dir, const T OriginalLength, const TVector<T, 3> QueryHalfExtents, SQVisitor& Visitor) const
+	void Sweep(const TVector<T, 3>& Start, const TVector<T, 3>& Dir, const T Length, const TVector<T, 3> QueryHalfExtents, SQVisitor& Visitor) const
 	{
-		QueryImp<EAABBQueryType::Sweep>(Start, Dir, OriginalLength, QueryHalfExtents, TBox<T, 3>(), Visitor);
+		bool bParallel[3];
+		TVector<T, 3> InvDir;
+
+		ensure(Length);
+		T CurLength = Length;
+		T InvLength = 1 / Length;
+		for (int Axis = 0; Axis < 3; ++Axis)
+		{
+			bParallel[Axis] = Dir[Axis] == 0;
+			InvDir[Axis] = bParallel[Axis] ? 0 : 1 / Dir[Axis];
+		}
+
+		QueryImp<EAABBQueryType::Sweep>(Start, Dir, InvDir, bParallel, CurLength, InvLength, QueryHalfExtents, TBox<T, 3>(), Visitor);
+	}
+
+	template <typename SQVisitor>
+	bool SweepFast(const TVector<T, 3>& Start, const TVector<T, 3>& Dir, const TVector<T,3>& InvDir, const bool* bParallel, T& CurLength, T& InvLength, const TVector<T, 3> QueryHalfExtents, SQVisitor& Visitor) const
+	{
+		return QueryImp<EAABBQueryType::Sweep>(Start, Dir, InvDir, bParallel, CurLength, InvLength, QueryHalfExtents, TBox<T, 3>(), Visitor);
 	}
 
 	void Overlap(const TBox<T, 3>& QueryBounds, ISpatialVisitor<TPayloadType, T>& Visitor) const override
 	{
 		TSpatialVisitor<TPayloadType, T> ProxyVisitor(Visitor);
-		return Overlap(QueryBounds, ProxyVisitor);
+		Overlap(QueryBounds, ProxyVisitor);
 	}
 
 	template <typename SQVisitor>
 	void Overlap(const TBox<T,3>& QueryBounds, SQVisitor& Visitor) const
 	{
-		QueryImp<EAABBQueryType::Overlap>(TVector<T,3>(), TVector<T, 3>(), 0, TVector<T,3>(), QueryBounds, Visitor);
+		OverlapFast(QueryBounds, Visitor);
+	}
+
+	template <typename SQVisitor>
+	bool OverlapFast(const TBox<T, 3>& QueryBounds, SQVisitor& Visitor) const
+	{
+		//dummy variables to reuse templated path
+		T Length;
+		T InvLength;
+		TVector<T, 3> InvDir;
+		bool bParallel;
+		return QueryImp<EAABBQueryType::Overlap>(TVector<T, 3>(), TVector<T, 3>(), InvDir, &bParallel, Length, InvLength, TVector<T, 3>(), QueryBounds, Visitor);
 	}
 
 	virtual void RemoveElement(const TPayloadType& Payload)
@@ -408,25 +455,8 @@ private:
 	using FNode = TAABBTreeNode<T>;
 
 	template <EAABBQueryType Query, typename SQVisitor>
-	void QueryImp(const TVector<T, 3>& Start, const TVector<T, 3>& Dir, const T OriginalLength, const TVector<T, 3> QueryHalfExtents, const TBox<T,3>& QueryBounds, SQVisitor& Visitor) const
+	bool QueryImp(const TVector<T, 3>& Start, const TVector<T, 3>& Dir, const TVector<T,3>& InvDir, const bool* bParallel, T& CurrentLength, T& InvCurrentLength, const TVector<T, 3> QueryHalfExtents, const TBox<T,3>& QueryBounds, SQVisitor& Visitor) const
 	{
-		T CurrentLength = OriginalLength;
-
-		bool bParallel[3];
-		TVector<T, 3> InvDir;
-
-		T InvCurrentLength;
-
-		if (Query != EAABBQueryType::Overlap)
-		{
-			InvCurrentLength = 1 / CurrentLength;
-			for (int Axis = 0; Axis < 3; ++Axis)
-			{
-				bParallel[Axis] = Dir[Axis] == 0;
-				InvDir[Axis] = bParallel[Axis] ? 0 : 1 / Dir[Axis];
-			}
-		}
-
 		TVector<T, 3> TmpPosition;
 		T TOI = 0;
 
@@ -449,7 +479,7 @@ private:
 
 				if (!bContinue)
 				{
-					return;
+					return false;
 				}
 
 				if (Query != EAABBQueryType::Overlap)
@@ -479,7 +509,7 @@ private:
 				
 				if (!bContinue)
 				{
-					return;
+					return false;
 				}
 
 				if (Query != EAABBQueryType::Overlap)
@@ -500,9 +530,12 @@ private:
 		while (NodeStack.Num())
 		{
 			const FNodeQueueEntry NodeEntry = NodeStack.Pop(false);
-			if (NodeEntry.TOI > CurrentLength)
+			if (Query != EAABBQueryType::Overlap)
 			{
-				continue;
+				if (NodeEntry.TOI > CurrentLength)
+				{
+					continue;
+				}
 			}
 
 			const FNode& Node = Nodes[NodeEntry.NodeIdx];
@@ -513,19 +546,19 @@ private:
 				{
 					if (Leaf.OverlapFast(QueryBounds, Visitor) == false)
 					{
-						return;
+						return false;
 					}
 				}
 				else if (Query == EAABBQueryType::Sweep)
 				{
 					if (Leaf.SweepFast(Start, Dir, InvDir, bParallel, CurrentLength, InvCurrentLength, QueryHalfExtents, Visitor) == false)
 					{
-						return;
+						return false;
 					}
 				}
 				else if (Leaf.RaycastFast(Start, Dir, InvDir, bParallel, CurrentLength, InvCurrentLength, Visitor) == false)
 				{
-					return;
+					return false;
 				}
 			}
 			else
@@ -542,6 +575,8 @@ private:
 				}
 			}
 		}
+
+		return true;
 	}
 
 	template <typename TParticles>
diff --git a/Engine/Source/Runtime/Experimental/Chaos/Public/Chaos/BoundingVolume.h b/Engine/Source/Runtime/Experimental/Chaos/Public/Chaos/BoundingVolume.h
--- a/Engine/Source/Runtime/Experimental/Chaos/Public/Chaos/BoundingVolume.h
+++ b/Engine/Source/Runtime/Experimental/Chaos/Public/Chaos/BoundingVolume.h
@@ -256,38 +256,132 @@
 		return MGlobalPayloads;
 	}
 
-	virtual void Raycast(const TVector<T, d>& Start, const TVector<T, d>& Dir, const T OriginalLength, ISpatialVisitor<TPayloadType, T>& Visitor) const override
+	virtual void Raycast(const TVector<T, d>& Start, const TVector<T, d>& Dir, const T Length, ISpatialVisitor<TPayloadType, T>& Visitor) const override
 	{
 		TSpatialVisitor<TPayloadType, T> ProxyVisitor(Visitor);
-		Raycast(Start, Dir, OriginalLength, ProxyVisitor);
+		return Raycast(Start, Dir, Length, ProxyVisitor);
 	}
 
 	template <typename SQVisitor>
 	bool RaycastFast(const TVector<T,d>& Start, const TVector<T,d>& Dir, const TVector<T,d>& InvDir, const bool* bParallel, T& CurrentLength, T& InvCurrentLength, SQVisitor& Visitor) const
 	{
-		if (Raycast(Start, Dir, CurrentLength, Visitor))
-		{
-			InvCurrentLength = 1 / CurrentLength;
-			return true;
-		}
-		return false;
+		return RaycastImp(Start, Dir, InvDir, bParallel, CurrentLength, InvCurrentLength, Visitor);
 	}
 
 	template <typename SQVisitor, bool bPruneDuplicates = true>
-	bool Raycast(const TVector<T, d>& Start, const TVector<T, d>& Dir, const T OriginalLength, SQVisitor& Visitor) const
+	void Raycast(const TVector<T, d>& Start, const TVector<T, d>& Dir, const T Length, SQVisitor& Visitor) const
 	{
-		T CurrentLength = OriginalLength;
-
+		ensure(Length);
+		T CurLength = Length;
 		bool bParallel[d];
 		TVector<T, d> InvDir;
 
-		T InvCurrentLength = 1 / CurrentLength;
+		T InvLength = 1 / Length;
 		for (int Axis = 0; Axis < d; ++Axis)
 		{
 			bParallel[Axis] = Dir[Axis] == 0;
 			InvDir[Axis] = bParallel[Axis] ? 0 : 1 / Dir[Axis];
 		}
 
+		RaycastImp<SQVisitor, bPruneDuplicates>(Start, Dir, InvDir, bParallel, CurLength, InvLength, Visitor);
+	}
+
+	void Sweep(const TVector<T, d>& Start, const TVector<T, d>& Dir, const T Length, const TVector<T, d> QueryHalfExtents, ISpatialVisitor<TPayloadType, T>& Visitor) const override
+	{
+		TSpatialVisitor<TPayloadType, T> ProxyVisitor(Visitor);
+		Sweep(Start, Dir, Length, QueryHalfExtents, ProxyVisitor);
+	}
+
+	template <typename SQVisitor>
+	bool SweepFast(const TVector<T,d>& Start, const TVector<T,d>& Dir, const TVector<T,d>& InvDir, const bool* bParallel, T& CurrentLength, T& InvCurrentLength, const TVector<T,d>& QueryHalfExtents, SQVisitor& Visitor) const
+	{
+		return SweepImp(Start, Dir, InvDir, bParallel, CurrentLength, InvCurrentLength, QueryHalfExtents, Visitor);
+	}
+
+	template <typename SQVisitor, bool bPruneDuplicates = true>
+	void Sweep(const TVector<T, d>& Start, const TVector<T, d>& Dir, const T Length, const TVector<T, d> QueryHalfExtents, SQVisitor& Visitor) const
+	{
+		ensure(Length);
+		T CurLength = Length;
+		bool bParallel[d];
+		TVector<T, d> InvDir;
+
+		T InvLength = 1 / Length;
+		for (int Axis = 0; Axis < d; ++Axis)
+		{
+			bParallel[Axis] = Dir[Axis] == 0;
+			InvDir[Axis] = bParallel[Axis] ? 0 : 1 / Dir[Axis];
+		}
+
+		SweepImp<SQVisitor, bPruneDuplicates>(Start, Dir, InvDir, bParallel, CurLength, InvLength, QueryHalfExtents, Visitor);
+	}
+
+	virtual void Overlap(const TBox<T, d>& QueryBounds, ISpatialVisitor<TPayloadType, T>& Visitor) const override
+	{
+		TSpatialVisitor<TPayloadType, T> ProxyVisitor(Visitor);
+		Overlap(QueryBounds, ProxyVisitor);
+	}
+
+	template <typename SQVisitor>
+	bool OverlapFast(const TBox<T, d>& QueryBounds, SQVisitor& Visitor) const
+	{
+		return OverlapImp(QueryBounds, Visitor);
+	}
+
+	template <typename SQVisitor, bool bPruneDuplicates = true>
+	void Overlap(const TBox<T, d>& QueryBounds, SQVisitor& Visitor) const
+	{
+		OverlapImp<SQVisitor, bPruneDuplicates>(QueryBounds, Visitor);
+	}
+
+	virtual void Serialize(FChaosArchive& Ar) override
+	{
+		Ar.UsingCustomVersion(FExternalPhysicsCustomObjectVersion::GUID);
+		if (Ar.CustomVer(FExternalPhysicsCustomObjectVersion::GUID) < FExternalPhysicsCustomObjectVersion::GlobalElementsHaveBounds)
+		{
+			TArray<TPayloadType> TmpPayloads;
+			Ar << TmpPayloads;
+			MGlobalPayloads.Reserve(TmpPayloads.Num());
+			for (auto& Payload : TmpPayloads)
+			{
+				MGlobalPayloads.Add({ Payload, TBox<T,d>(TVector<T,d>(TNumericLimits<T>::Lowest()), TVector<T,d>(TNumericLimits<T>::Max())) });
+			}
+			MaxPayloadBounds = DefaultMaxPayloadBounds;
+		}
+		else
+		{
+			Ar << MGlobalPayloads;
+			Ar << MaxPayloadBounds;
+		}
+
+		Ar << MGrid;
+		Ar << MElements;
+		Ar << MDirtyElements;
+		Ar << bIsEmpty;
+
+		TArray<TPayloadType> Payloads;
+		if (!Ar.IsLoading())
+		{
+			MPayloadInfo.GenerateKeyArray(Payloads);
+		}
+		Ar << Payloads;
+
+		for (auto Payload : Payloads)
+		{
+			auto& Info = MPayloadInfo.FindOrAdd(Payload);
+			Ar << Info;
+		}
+
+	}
+
+private:
+
+	using FCellElement = TBVCellElement<TPayloadType, T, d>;
+	using FPayloadInfo = TBVPayloadInfo<T, d>;
+
+	template <typename SQVisitor, bool bPruneDuplicates = true>
+	bool RaycastImp(const TVector<T, d>& Start, const TVector<T, d>& Dir, const TVector<T, d>& InvDir, const bool* bParallel, T& CurrentLength, T& InvCurrentLength, SQVisitor& Visitor) const
+	{
 		TVector<T, d> TmpPosition;
 		T TOI;
 
@@ -298,11 +392,12 @@
 			{
 				TSpatialVisitorData<TPayloadType> VisitData(Elem.Payload, true, InstanceBounds);
 				const bool bContinue = Visitor.VisitRaycast(VisitData, CurrentLength);
+				InvCurrentLength = 1 / CurrentLength;
+
 				if (!bContinue)
 				{
 					return false;
 				}
-				InvCurrentLength = 1 / CurrentLength;
 			}
 		}
 
@@ -329,7 +424,7 @@
 		TBox<T, d> GlobalBounds(MGrid.MinCorner(), MGrid.MaxCorner());
 		TVector<T, d> NextStart;
 		TVector<int32, d> CellIdx;
-		bool bCellsLeft = MElements.Num() && TBox<T,d>::RaycastFast(GlobalBounds.Min(), GlobalBounds.Max(), Start, Dir, InvDir, bParallel, CurrentLength, InvCurrentLength, TOI, NextStart);
+		bool bCellsLeft = MElements.Num() && TBox<T, d>::RaycastFast(GlobalBounds.Min(), GlobalBounds.Max(), Start, Dir, InvDir, bParallel, CurrentLength, InvCurrentLength, TOI, NextStart);
 		if (bCellsLeft)
 		{
 			CellIdx = MGrid.Cell(NextStart);
@@ -346,7 +441,6 @@
 				{
 					if (bPruneDuplicates)
 					{
-
 						bool bSkip = false;
 						if (Elem.StartIdx[0] != Elem.EndIdx[0] || Elem.StartIdx[1] != Elem.EndIdx[1] || Elem.StartIdx[2] != Elem.EndIdx[2])
 						{
@@ -372,7 +466,7 @@
 						}
 					}
 					const auto& InstanceBounds = Elem.Bounds;
-					if (TBox<T,d>::RaycastFast(InstanceBounds.Min(), InstanceBounds.Max(), Start, Dir, InvDir, bParallel, CurrentLength, InvCurrentLength, TOI, TmpPosition))
+					if (TBox<T, d>::RaycastFast(InstanceBounds.Min(), InstanceBounds.Max(), Start, Dir, InvDir, bParallel, CurrentLength, InvCurrentLength, TOI, TmpPosition))
 					{
 						TSpatialVisitorData<TPayloadType> VisitData(Elem.Payload, true, InstanceBounds);
 						const bool bContinue = Visitor.VisitRaycast(VisitData, CurrentLength);
@@ -380,7 +474,7 @@
 						{
 							return false;
 						}
-						InvCurrentLength = 1 / CurrentLength;
+						InvCurrentLength = 1 / CurrentLength;	//todo: pass into visitor
 					}
 				}
 
@@ -434,43 +528,14 @@
 				NextStart = NextStart + Dir * BestTime;
 			} while (true);
 		}
-		
+
 		return true;
 	}
 
-	void Sweep(const TVector<T, d>& Start, const TVector<T, d>& Dir, T OriginalLength, const TVector<T, d> QueryHalfExtents, ISpatialVisitor<TPayloadType, T>& Visitor) const override
-	{
-		TSpatialVisitor<TPayloadType, T> ProxyVisitor(Visitor);
-		Sweep(Start, Dir, OriginalLength, QueryHalfExtents, ProxyVisitor);
-	}
-
-	template <typename SQVisitor>
-	bool SweepFast(const TVector<T,d>& Start, const TVector<T,d>& Dir, const TVector<T,d>& InvDir, const bool* bParallel, T& CurrentLength, T& InvCurrentLength, const TVector<T,d>& QueryHalfExtents, SQVisitor& Visitor) const
-	{
-		if (Sweep(Start, Dir, CurrentLength, QueryHalfExtents, Visitor))
-		{
-			InvCurrentLength = 1 / CurrentLength;
-			return true;
-		}
-		return false;
-	}
-
 	template <typename SQVisitor, bool bPruneDuplicates = true>
-	bool Sweep(const TVector<T, d>& Start, const TVector<T, d>& Dir, T OriginalLength, const TVector<T, d> QueryHalfExtents, SQVisitor& Visitor) const
+	bool SweepImp(const TVector<T, d>& Start, const TVector<T, d>& Dir, const TVector<T,d>& InvDir, const bool* bParallel, T& CurrentLength, T& InvCurrentLength, const TVector<T, d> QueryHalfExtents, SQVisitor& Visitor) const
 	{
-		T CurrentLength = OriginalLength;
-
-		T TOI = 0;	//not needed, but fixes compiler warning
-		bool bParallel[d];
-		TVector<T, d> InvDir;
-
-		T InvCurrentLength = 1 / CurrentLength;
-		for (int Axis = 0; Axis < d; ++Axis)
-		{
-			bParallel[Axis] = Dir[Axis] == 0;
-			InvDir[Axis] = bParallel[Axis] ? 0 : 1 / Dir[Axis];
-		}
-
+		T TOI = 0;
 		for (const auto& Elem : MGlobalPayloads)
 		{
 			const TBox<T, d>& InstanceBounds = Elem.Bounds;
@@ -513,7 +578,7 @@
 		}
 
 		TBox<T, d> GlobalBounds(MGrid.MinCorner() - QueryHalfExtents, MGrid.MaxCorner() + QueryHalfExtents);
-		
 
+
 		struct FCellIntersection
 		{
@@ -521,11 +586,11 @@
 			T TOI;
 		};
 
-		
+
 		TVector<T, d> HitPoint;
 		FGridSet IdxsSeen(MGrid.Counts());
 		FGridSet CellsVisited(MGrid.Counts());
-		const bool bInitialHit = TBox<T,d>::RaycastFast(GlobalBounds.Min(), GlobalBounds.Max(), Start, Dir, InvDir, bParallel, CurrentLength, InvCurrentLength, TOI, HitPoint);
+		const bool bInitialHit = TBox<T, d>::RaycastFast(GlobalBounds.Min(), GlobalBounds.Max(), Start, Dir, InvDir, bParallel, CurrentLength, InvCurrentLength, TOI, HitPoint);
 		if (bInitialHit)
 		{
 			//Flood fill from inflated cell so that we get all cells along the ray
@@ -585,7 +650,7 @@
 						const TVector<T, d> NeighborCenter = MGrid.Location(NeighborIdx);
 						const TVector<T, d> Min = NeighborCenter - QueryHalfExtents - HalfDx;
 						const TVector<T, d> Max = NeighborCenter + QueryHalfExtents + HalfDx;
-						if (TBox<T,d>::RaycastFast(Min, Max, Start, Dir, InvDir, bParallel, CurrentLength, InvCurrentLength, TOI, HitPoint))
+						if (TBox<T, d>::RaycastFast(Min, Max, Start, Dir, InvDir, bParallel, CurrentLength, InvCurrentLength, TOI, HitPoint))
 						{
 							IdxsQueue.Add({ NeighborIdx, TOI });	//should we sort by TOI?
 						}
@@ -607,7 +672,7 @@
 								{
 									for (int32 Z = Elem.StartIdx[2]; Z <= Elem.EndIdx[2]; ++Z)
 									{
-										if (CellsVisited.Contains(TVector<int32,3>(X,Y,Z)))
+										if (CellsVisited.Contains(TVector<int32, 3>(X, Y, Z)))
 										{
 											bSkip = true;
 											break;
@@ -626,7 +691,7 @@
 					const TBox<T, d>& InstanceBounds = Elem.Bounds;
 					const TVector<T, d> Min = InstanceBounds.Min() - QueryHalfExtents;
 					const TVector<T, d> Max = InstanceBounds.Max() + QueryHalfExtents;
-					if (TBox<T,d>::RaycastFast(Min, Max, Start, Dir, InvDir, bParallel, CurrentLength, InvCurrentLength, TOI, HitPoint))
+					if (TBox<T, d>::RaycastFast(Min, Max, Start, Dir, InvDir, bParallel, CurrentLength, InvCurrentLength, TOI, HitPoint))
 					{
 						TSpatialVisitorData<TPayloadType> VisitData(Elem.Payload, true, InstanceBounds);
 						const bool bContinue = Visitor.VisitSweep(VisitData, CurrentLength);
@@ -644,20 +709,8 @@
 		return true;
 	}
 
-	void Overlap(const TBox<T, d>& QueryBounds, ISpatialVisitor<TPayloadType, T>& Visitor) const override
-	{
-		TSpatialVisitor<TPayloadType, T> ProxyVisitor(Visitor);
-		Overlap(QueryBounds, ProxyVisitor);
-	}
-
-	template <typename SQVisitor>
-	bool OverlapFast(const TBox<T, d>& QueryBounds, SQVisitor& Visitor) const
-	{
-		return Overlap(QueryBounds, Visitor);
-	}
-
 	template <typename SQVisitor, bool bPruneDuplicates = true>
-	bool Overlap(const TBox<T, d>& QueryBounds, SQVisitor& Visitor) const
+	bool OverlapImp(const TBox<T, d>& QueryBounds, SQVisitor& Visitor) const
 	{
 		for (const auto& Elem : MGlobalPayloads)
 		{
@@ -730,51 +783,6 @@
 		return true;
 	}
 
-	virtual void Serialize(FChaosArchive& Ar) override
-	{
-		Ar.UsingCustomVersion(FExternalPhysicsCustomObjectVersion::GUID);
-		if (Ar.CustomVer(FExternalPhysicsCustomObjectVersion::GUID) < FExternalPhysicsCustomObjectVersion::GlobalElementsHaveBounds)
-		{
-			TArray<TPayloadType> TmpPayloads;
-			Ar << TmpPayloads;
-			MGlobalPayloads.Reserve(TmpPayloads.Num());
-			for (auto& Payload : TmpPayloads)
-			{
-				MGlobalPayloads.Add({ Payload, TBox<T,d>(TVector<T,d>(TNumericLimits<T>::Lowest()), TVector<T,d>(TNumericLimits<T>::Max())) });
-			}
-			MaxPayloadBounds = DefaultMaxPayloadBounds;
-		}
-		else
-		{
-			Ar << MGlobalPayloads;
-			Ar << MaxPayloadBounds;
-		}
-
-		Ar << MGrid;
-		Ar << MElements;
-		Ar << MDirtyElements;
-		Ar << bIsEmpty;
-
-		TArray<TPayloadType> Payloads;
-		if (!Ar.IsLoading())
-		{
-			MPayloadInfo.GenerateKeyArray(Payloads);
-		}
-		Ar << Payloads;
-
-		for (auto Payload : Payloads)
-		{
-			auto& Info = MPayloadInfo.FindOrAdd(Payload);
-			Ar << Info;
-		}
-
-	}
-
-private:
-
-	using FCellElement = TBVCellElement<TPayloadType, T, d>;
-	using FPayloadInfo = TBVPayloadInfo<T, d>;
-
 	template <typename ParticleView>
 	void GenerateTree(const ParticleView& Particles, const bool bUseVelocity, const T Dt, const int32 MaxCells)
 	{
diff --git a/Engine/Source/Runtime/Experimental/Chaos/Public/Chaos/DebugDrawQueue.h b/Engine/Source/Runtime/Experimental/Chaos/Public/Chaos/DebugDrawQueue.h
--- a/Engine/Source/Runtime/Experimental/Chaos/Public/Chaos/DebugDrawQueue.h
+++ b/Engine/Source/Runtime/Experimental/Chaos/Public/Chaos/DebugDrawQueue.h
@@ -119,6 +119,11 @@ struct CHAOS_API FLatentDrawCommand
 		Capsule,
 	} Type;
 
+	FLatentDrawCommand()
+		: TestBaseActor(nullptr)
+	{
+	}
+
 	static FLatentDrawCommand DrawPoint(const FVector& Position, const FColor& Color, bool bPersistentLines, float LifeTime, uint8 DepthPriority, float Thickness)
 	{
 		FLatentDrawCommand Command;
diff --git a/Engine/Source/Runtime/Experimental/Chaos/Public/Chaos/HeightField.h b/Engine/Source/Runtime/Experimental/Chaos/Public/Chaos/HeightField.h
--- a/Engine/Source/Runtime/Experimental/Chaos/Public/Chaos/HeightField.h
+++ b/Engine/Source/Runtime/Experimental/Chaos/Public/Chaos/HeightField.h
@@ -11,7 +11,7 @@
 #include "Templates/UnrealTypeTraits.h"
 #include "UniformGrid.h"
 #include "Utilities.h"
-PRAGMA_DISABLE_OPTIMIZATION
+#include "UObject/ExternalPhysicsCustomObjectVersion.h"
 
 namespace Chaos
 {
@@ -28,8 +28,10 @@ namespace Chaos
 		using TImplicitObject<T, 3>::GetTypeName;
 
 		THeightField(TArray<T>&& Height, int32 InNumRows, int32 InNumCols, const TVector<T,3>& InScale);
+		THeightField(TArrayView<const uint16> InHeights, int32 InNumRows, int32 InNumCols, const TVector<T, 3>& InScale);
 		THeightField(const THeightField& Other) = delete;
 		THeightField(THeightField&& Other) = default;
+
 		virtual ~THeightField() {}
 
 		/** Support for editing a subsection of the heightfield */
@@ -51,7 +53,8 @@ namespace Chaos
 
 		virtual const TBox<T, 3>& BoundingBox() const
 		{
-			return LocalBounds;
+			CachedBounds = TBox<T, 3>(LocalBounds.Min() * GeomData.Scale, LocalBounds.Max() * GeomData.Scale);
+			return CachedBounds;
 		}
 
 		virtual uint32 GetTypeHash() const override
@@ -79,12 +82,31 @@ namespace Chaos
 			
 			GeomData.Serialize(Ar);
 
+			Ar.UsingCustomVersion(FExternalPhysicsCustomObjectVersion::GUID);
+			if (Ar.CustomVer(FExternalPhysicsCustomObjectVersion::GUID) >= FExternalPhysicsCustomObjectVersion::HeightfieldData)
+			{
+				Ar << FlatGrid;
+				Ar << FlattenedBounds.Min;
+				Ar << FlattenedBounds.Max;
+				Ar << LocalBounds;
+			}
+			else
+			{
+				CalcBounds();
+			}
+			
+
 			if(Ar.IsLoading())
 			{
 				BuildQueryData();
 			}
 		}
 
+		void SetScale(const TVector<T, 3>& InScale)
+		{
+			GeomData.Scale = InScale;
+		}
+
 		template<typename InStorageType, typename InRealType>
 		struct FData
 		{
@@ -141,15 +163,20 @@ namespace Chaos
 				const int32 X = Index % (NumCols);
 				const int32 Y = Index / (NumCols);
 
-				TVector<typename FDataType::RealType, 3> P0 = {(typename FDataType::RealType)X, (typename FDataType::RealType)Y, H0};
-				TVector<typename FDataType::RealType, 3> P1 = {(typename FDataType::RealType)X + 1, (typename FDataType::RealType)Y, H1};
-				TVector<typename FDataType::RealType, 3> P2 = {(typename FDataType::RealType)X, (typename FDataType::RealType)Y + 1, H2};
-				TVector<typename FDataType::RealType, 3> P3 = {(typename FDataType::RealType)X + 1, (typename FDataType::RealType)Y + 1, H3};
-				
-				OutPts[0] = (P0 * Scale);
-				OutPts[1] = (P1 * Scale);
-				OutPts[2] = (P2 * Scale);
-				OutPts[3] = (P3 * Scale);
+				OutPts[0] = {(typename FDataType::RealType)X, (typename FDataType::RealType)Y, H0};
+				OutPts[1] = {(typename FDataType::RealType)X + 1, (typename FDataType::RealType)Y, H1};
+				OutPts[2] = {(typename FDataType::RealType)X, (typename FDataType::RealType)Y + 1, H2};
+				OutPts[3] = {(typename FDataType::RealType)X + 1, (typename FDataType::RealType)Y + 1, H3};
+			}
+
+			FORCEINLINE void GetPointsScaled(int32 Index, TVector<T, 3> OutPts[4]) const
+			{
+				GetPoints(Index, OutPts);
+
+				OutPts[0] *= Scale;
+				OutPts[1] *= Scale;
+				OutPts[2] *= Scale;
+				OutPts[3] *= Scale;
 			}
 
 			void Serialize(FChaosArchive& Ar)
@@ -172,6 +199,14 @@ namespace Chaos
 				Ar << MaxValue;
 				Ar << NumRows;
 				Ar << NumCols;
+
+				Ar.UsingCustomVersion(FExternalPhysicsCustomObjectVersion::GUID);
+				if (Ar.CustomVer(FExternalPhysicsCustomObjectVersion::GUID) >= FExternalPhysicsCustomObjectVersion::HeightfieldData)
+				{
+					Ar << Range;
+					Ar << HeightPerUnit;
+					Ar << CellBounds;
+				}
 			}
 		};
 
@@ -333,6 +368,8 @@ namespace Chaos
 		// Helpers for accessing bounds
 		bool GetCellBounds2D(const TVector<int32, 2> InCoord, FBounds2D& OutBounds, const TVector<T, 2>& InInflate = {0}) const;
 		bool GetCellBounds3D(const TVector<int32, 2> InCoord, TVector<T, 3>& OutMin, TVector<T, 3>& OutMax, const TVector<T, 3>& InInflate = TVector<T, 3>(0)) const;
+		bool GetCellBounds2DScaled(const TVector<int32, 2> InCoord, FBounds2D& OutBounds, const TVector<T, 2>& InInflate = {0}) const;
+		bool GetCellBounds3DScaled(const TVector<int32, 2> InCoord, TVector<T, 3>& OutMin, TVector<T, 3>& OutMax, const TVector<T, 3>& InInflate = TVector<T, 3>(0)) const;
 		bool CalcCellBounds3D(const TVector<int32, 2> InCoord, TVector<T, 3>& OutMin, TVector<T, 3>& OutMax, const TVector<T, 3>& InInflate = TVector<T, 3>(0)) const;
 
 		// Query functions - sweep, ray, overlap
@@ -341,13 +378,18 @@ namespace Chaos
 		bool GridCast(const TVector<T, 3>& StartPoint, const TVector<T, 3>& Dir, const T Length, THeightfieldRaycastVisitor<T>& Visitor) const;
 		bool GetGridIntersections(FBounds2D InFlatBounds, TArray<TVector<int32, 2>>& OutInterssctions) const;
 		
+		FBounds2D GetFlatBounds() const;
+
 		// Grid for queries, faster than bounding volumes for heightfields
 		TUniformGrid<T, 2> FlatGrid;
 		// Bounds in 2D of the whole heightfield, to clip queries against
 		FBounds2D FlattenedBounds;
 		// 3D bounds for the heightfield, for insertion to the scene structure
 		TBox<T, 3> LocalBounds;
+		// Cached when bounds are requested. Mutable to allow GetBounds to be logical const
+		mutable TBox<T, 3> CachedBounds;
 
+		void CalcBounds();
 		void BuildQueryData();
 		
 		// Needed for serialization
@@ -355,4 +397,3 @@ namespace Chaos
 		friend TImplicitObject<T, 3>;
 	};
 }
-PRAGMA_ENABLE_OPTIMIZATION
\ No newline at end of file
diff --git a/Engine/Source/Runtime/Experimental/Chaos/Public/Chaos/ISpatialAcceleration.h b/Engine/Source/Runtime/Experimental/Chaos/Public/Chaos/ISpatialAcceleration.h
--- a/Engine/Source/Runtime/Experimental/Chaos/Public/Chaos/ISpatialAcceleration.h
+++ b/Engine/Source/Runtime/Experimental/Chaos/Public/Chaos/ISpatialAcceleration.h
@@ -82,12 +82,6 @@ public:
 	virtual bool Sweep(const TSpatialVisitorData<TPayloadType>& Instance, T& CurLength) = 0;
 };
 
-enum class ESpatialAccelerationType
-{
-	Grid,
-	BVH
-};
-
 /**
  * Can be implemented by external, non-chaos systems to collect / render
  * debug information from spacial structures. When passed to the debug
@@ -113,9 +107,12 @@ enum ESpatialAcceleration
 	AABBTree,
 	AABBTreeBV,
 	Collection,
-	Unknown
+	Unknown,
+	//For custom types continue the enum after ESpatialAcceleration::Unknown
 };
 
+using SpatialAccelerationType = uint8;	//see ESpatialAcceleration. Projects can add their own custom types by using enum values higher than ESpatialAcceleration::Unknown
+
 
 template <typename TPayloadType, typename T>
 struct TPayloadBoundsElement
@@ -149,7 +146,7 @@ class CHAOS_API ISpatialAcceleration
 {
 public:
 
-	ISpatialAcceleration(ESpatialAcceleration InType = ESpatialAcceleration::Unknown)
+	ISpatialAcceleration(SpatialAccelerationType InType = ESpatialAcceleration::Unknown)
 		: Type(InType)
 	{
 	}
@@ -157,8 +154,8 @@ public:
 
 	virtual TArray<TPayloadType> FindAllIntersections(const TBox<T, d>& Box) const { check(false); return TArray<TPayloadType>(); }
 
-	virtual void Raycast(const TVector<T, d>& Start, const TVector<T, d>& Dir, const T OriginalLength, ISpatialVisitor<TPayloadType, T>& Visitor) const { check(false); }
-	virtual void Sweep(const TVector<T, d>& Start, const TVector<T, d>& Dir, T OriginalLength, const TVector<T, d> QueryHalfExtents, ISpatialVisitor<TPayloadType, T>& Visitor) const { check(false); }
+	virtual void Raycast(const TVector<T, d>& Start, const TVector<T, d>& Dir, const T Length, ISpatialVisitor<TPayloadType, T>& Visitor) const { check(false); }
+	virtual void Sweep(const TVector<T, d>& Start, const TVector<T, d>& Dir, const T Length, const TVector<T, d> QueryHalfExtents, ISpatialVisitor<TPayloadType, T>& Visitor) const { check(false);}
 	virtual void Overlap(const TBox<T, d>& QueryBounds, ISpatialVisitor<TPayloadType, T>& Visitor) const { check(false); }
 
 	virtual void RemoveElement(const TPayloadType& Payload)
@@ -227,7 +224,7 @@ public:
 	}
 
 private:
-	ESpatialAcceleration Type;
+	SpatialAccelerationType Type;
 };
 
 template <typename TBase, typename TDerived>
diff --git a/Engine/Source/Runtime/Experimental/Chaos/Public/Chaos/ImplicitObject.h b/Engine/Source/Runtime/Experimental/Chaos/Public/Chaos/ImplicitObject.h
--- a/Engine/Source/Runtime/Experimental/Chaos/Public/Chaos/ImplicitObject.h
+++ b/Engine/Source/Runtime/Experimental/Chaos/Public/Chaos/ImplicitObject.h
@@ -156,6 +156,11 @@ public:
 	void IgnoreAnalyticCollisions(const bool Ignore = true) { bIgnoreAnalyticCollisions = Ignore; }
 	bool GetIgnoreAnalyticCollisions() const { return bIgnoreAnalyticCollisions; }
 	void SetConvex(const bool Convex = true) { bIsConvex = Convex; }
+	virtual bool IsPerformanceWarning() const { return false; }
+	virtual FString PerformanceWarningAndSimplifaction() 
+	{
+		return FString::Printf(TEXT("ImplicitObject - No Performance String"));
+	};
 
 	Pair<TVector<T, d>, bool> FindDeepestIntersection(const TImplicitObject<T, d>* Other, const TBVHParticles<float, d>* Particles, const PMatrix<T, d, d>& OtherToLocalTransform, const T Thickness) const;
 	Pair<TVector<T, d>, bool> FindDeepestIntersection(const TImplicitObject<T, d>* Other, const TParticles<float, d>* Particles, const PMatrix<T, d, d>& OtherToLocalTransform, const T Thickness) const;
diff --git a/Engine/Source/Runtime/Experimental/Chaos/Public/Chaos/ImplicitObjectScaled.h b/Engine/Source/Runtime/Experimental/Chaos/Public/Chaos/ImplicitObjectScaled.h
--- a/Engine/Source/Runtime/Experimental/Chaos/Public/Chaos/ImplicitObjectScaled.h
+++ b/Engine/Source/Runtime/Experimental/Chaos/Public/Chaos/ImplicitObjectScaled.h
@@ -7,6 +7,7 @@
 #include "ChaosArchive.h"
 #include "Templates/ChooseClass.h"
 #include "Templates/EnableIf.h"
+#include "Math/NumericLimits.h"
 
 namespace Chaos
 {
@@ -109,7 +110,7 @@ public:
 				return true;
 			}
 		}
-		
+			
 		return false;
 	}
 
@@ -129,12 +130,13 @@ public:
 			const T UnscaledLength = Length * LengthScale;
 			const TVector<T, d> UnscaledDir = UnscaledDirDenorm * LengthScaleInv;
 
-		TVector<T, d> UnscaledPosition;
-		TVector<T, d> UnscaledNormal;
+			TVector<T, d> UnscaledPosition;
+			TVector<T, d> UnscaledNormal;
+			float UnscaledTime;
 
-		TUniquePtr<TImplicitObject<T, d>> HackBPtr(const_cast<TImplicitObject<T,d>*>(&B));	//todo: hack, need scaled object to except raw ptr similar to transformed implicit
-		TImplicitObjectScaled<T, d> ScaledB(MakeSerializable(HackBPtr), OwningScaled.MInvScale);
-		HackBPtr.Release();
+			TUniquePtr<TImplicitObject<T, d>> HackBPtr(const_cast<TImplicitObject<T,d>*>(&B));	//todo: hack, need scaled object to except raw ptr similar to transformed implicit
+			TImplicitObjectScaled<T, d> ScaledB(MakeSerializable(HackBPtr), OwningScaled.MInvScale);
+			HackBPtr.Release();
 
 			TRigidTransform<T, d> BToATMNoScale(BToATM.GetLocation() * OwningScaled.MInvScale, BToATM.GetRotation());
 			
@@ -320,4 +322,4 @@ private:
 template <typename T, int d>
 using TImplicitObjectScaledNonSerializable = TImplicitObjectScaled<T, d, false>;
 
-}
\ No newline at end of file
+}
diff --git a/Engine/Source/Runtime/Experimental/Chaos/Public/Chaos/PBDCollisionConstraint.h b/Engine/Source/Runtime/Experimental/Chaos/Public/Chaos/PBDCollisionConstraint.h
--- a/Engine/Source/Runtime/Experimental/Chaos/Public/Chaos/PBDCollisionConstraint.h
+++ b/Engine/Source/Runtime/Experimental/Chaos/Public/Chaos/PBDCollisionConstraint.h
@@ -45,6 +45,13 @@ enum class ECollisionUpdateType
 	Deepest	//find the deepest penetration. Compute location and normal
 };
 
+/** Return value of the collision modification callback */
+enum class ECollisionModifierResult
+{
+	Unchanged,	/** No change to the collision */
+	Modified,	/** Modified the collision, but want it to remain enabled */
+	Disabled,	/** Collision should be disabled */
+};
 
 template<class T, int d>
 class CHAOS_API TPBDCollisionConstraintHandle : public TContainerConstraintHandle<TPBDCollisionConstraint<T, d>>
@@ -65,7 +72,7 @@ protected:
 };
 
 template<typename T, int d>
-using TRigidBodyContactConstraintsPostComputeCallback = TFunction<void(TArray<TRigidBodyContactConstraint<T, d>>& Constraints)>;
+using TRigidBodyContactConstraintsPostComputeCallback = TFunction<void()>;
 
 template<typename T, int d>
 using TRigidBodyContactConstraintsPostApplyCallback = TFunction<void(const T Dt, const TArray<TPBDCollisionConstraintHandle<T, d>*>& InConstraintHandles)>;
@@ -166,7 +173,7 @@ public:
 
 	/**
 	 * Set the callback used just after contacts are generated at the start of a frame tick.
-	 * This can be used to modify or remove constraints from the system.
+	 * This can be used to modify or disable constraints (via ApplyCollisionModifier).
 	 */
 	void SetPostComputeCallback(const TRigidBodyContactConstraintsPostComputeCallback<T, d>& Callback);
 
diff --git a/Engine/Source/Runtime/Experimental/Chaos/Public/Chaos/PBDCollisionConstraintImp.h b/Engine/Source/Runtime/Experimental/Chaos/Public/Chaos/PBDCollisionConstraintImp.h
--- a/Engine/Source/Runtime/Experimental/Chaos/Public/Chaos/PBDCollisionConstraintImp.h
+++ b/Engine/Source/Runtime/Experimental/Chaos/Public/Chaos/PBDCollisionConstraintImp.h
@@ -294,7 +294,7 @@ namespace Chaos
 					{
 						FConstraintHandle* Handle = Handles[HandleID];
 						int32 Idx = Handle->GetConstraintIndex();
-						FVector Position = Constraint[Idx].Location;
+						FVector Position = Constraints[Idx].Location;
 
 						Queue.Dequeue(Constraints[Idx]);
 						Constraints[Idx].PreviousLocation = Position;
diff --git a/Engine/Source/Runtime/Experimental/Chaos/Public/Chaos/PBDJointConstraintTypes.h b/Engine/Source/Runtime/Experimental/Chaos/Public/Chaos/PBDJointConstraintTypes.h
--- a/Engine/Source/Runtime/Experimental/Chaos/Public/Chaos/PBDJointConstraintTypes.h
+++ b/Engine/Source/Runtime/Experimental/Chaos/Public/Chaos/PBDJointConstraintTypes.h
@@ -119,6 +119,7 @@ namespace Chaos
 		int32 FrozenIterations;
 
 		// @todo(ccaulfield): remove these TEMP overrides for testing
+		bool bEnableLinearLimits;
 		bool bEnableTwistLimits;
 		bool bEnableSwingLimits;
 		bool bEnableDrives;
diff --git a/Engine/Source/Runtime/Experimental/Chaos/Public/Chaos/PBDSphericalConstraint.h b/Engine/Source/Runtime/Experimental/Chaos/Public/Chaos/PBDSphericalConstraint.h
new file mode 100644
--- /dev/null
+++ b/Engine/Source/Runtime/Experimental/Chaos/Public/Chaos/PBDSphericalConstraint.h
@@ -0,0 +1,71 @@
+// Copyright 1998-2019 Epic Games, Inc. All Rights Reserved.
+#pragma once
+
+#include "ParticleRule.h"
+
+namespace Chaos
+{
+
+template<typename T, int d>
+class PBDSphericalConstraint : public TParticleRule<T, d>
+{
+public:
+	PBDSphericalConstraint(
+		const uint32 InFirstParticleIndex,
+		const uint32 InParticleCount,
+		bool bInside,
+		const TArray<TVector<T, d>>* const InSpherePositions,
+		const TArray<T>* const InSphereRadii,
+		const TArray <T>* const InSphereOffsetDistances = nullptr,
+		const TArray<TVector<T, d>>* const InSphereOffsetDirections = nullptr)
+		: FirstParticleIndex(InFirstParticleIndex), ParticleCount(InParticleCount), bConstraintInsideOfSphere(bInside), SpherePositions(InSpherePositions), SphereRadii(InSphereRadii), SphereOffsetDistances(InSphereOffsetDistances), SphereOffsetDirections(InSphereOffsetDirections)
+	{
+		check(SphereRadii);
+		check(SpherePositions);
+	}
+	virtual ~PBDSphericalConstraint() {}
+
+	inline void Apply(TPBDParticles<T, d>& InParticles, const T Dt) const override
+	{
+		const T CenterDeadZoneRadius = KINDA_SMALL_NUMBER; // We will not push the particle away in the dead zone
+
+		PhysicsParallelFor(ParticleCount, [&](int32 Index)
+		{
+			if (InParticles.InvM(FirstParticleIndex + Index) == 0)
+			{
+				return;
+			}
+			
+			const T Radius = (*SphereRadii)[Index];
+			TVector<T, d> Centre = (*SpherePositions)[Index + FirstParticleIndex];
+			if (SphereOffsetDistances && SphereOffsetDirections)
+			{
+				Centre -= (*SphereOffsetDistances)[Index] * (*SphereOffsetDirections)[Index + FirstParticleIndex];
+			}
+
+			const TVector<T, d> CenterToParticle = InParticles.P(FirstParticleIndex + Index) - Centre;
+			const T DistanceSquared = CenterToParticle.SizeSquared();
+			const bool bPushParticleToShpericalBoundery = ((DistanceSquared < Radius * Radius) ^ bConstraintInsideOfSphere) && DistanceSquared > CenterDeadZoneRadius * CenterDeadZoneRadius;
+
+			if (bPushParticleToShpericalBoundery)
+			{
+				const T Distance = sqrt(DistanceSquared);
+				const TVector<T, d> PositionOnSphere = (Radius / Distance) * (CenterToParticle);
+				InParticles.P(FirstParticleIndex + Index) = Centre + PositionOnSphere;
+			}
+		});
+	}
+
+private:	
+	const uint32 FirstParticleIndex; // TODO: Get rid of this member variable (this is specific to a temporary cloth instancing implementation)
+	const uint32 ParticleCount;
+	const bool bConstraintInsideOfSphere; // Should the points be constrained to the inside of their respective spheres or the outsides
+
+	const TArray<TVector<T, d>>* const SpherePositions;  // Positions of spheres without adding an offset vector // Start at index FirstParticleIndex
+	const TArray<T>* const SphereRadii; // Start at index 0
+
+	const TArray<T>* const SphereOffsetDistances;  // Optional Sphere position offsets // start at index 0
+	const TArray<TVector<T, d>>* const SphereOffsetDirections; //	Optional Sphere offset directions (TBD: This is really implementation specific to backstop constraints, consider creating a specialized constraint instead).
+};
+
+}
