Changed EPI init error reporting to reduce server ensures

diff --git a/Engine/Source/Runtime/Experimental/Chaos/Public/Chaos/EPA.h b/Engine/Source/Runtime/Experimental/Chaos/Public/Chaos/EPA.h
--- a/Engine/Source/Runtime/Experimental/Chaos/Public/Chaos/EPA.h
+++ b/Engine/Source/Runtime/Experimental/Chaos/Public/Chaos/EPA.h
@@ -2,6 +2,8 @@
 #pragma once
 #include "Chaos/Simplex.h"
 #include <queue>
+#include "ChaosCheck.h"
+#include "ChaosLog.h"
 
 namespace Chaos
 {
@@ -171,7 +173,7 @@ bool InitializeEPA(TArray<TVec3<T>>& VertsA, TArray<TVec3<T>>& VertsB, const Sup
 			TVec3<T> Dir = MinkowskiVert(VertsA.GetData(), VertsB.GetData(), 1) - MinkowskiVert(VertsA.GetData(), VertsB.GetData(), 0);
 
 			bValid = Dir.SizeSquared() > 1e-4;
-			if (ensure(bValid))	//two verts given should be distinct
+			if (CHAOS_ENSURE(bValid))	//two verts given should be distinct
 			{
 				//find most opposing axis
 				int32 BestAxis = 0;
@@ -209,7 +211,7 @@ bool InitializeEPA(TArray<TVec3<T>>& VertsA, TArray<TVec3<T>>& VertsB, const Sup
 		{
 			//triangle, add farthest point along normal
 			bValid = OutEntries[3].Initialize(VertsA.GetData(), VertsB.GetData(), 0, 2, 1, { 1,0,2 }, { 2,0,0 });	
-			if (ensure(bValid)) //input verts must form a valid triangle
+			if (CHAOS_ENSURE(bValid)) //input verts must form a valid triangle
 			{
 				const TEPAEntry<T>& Base = OutEntries[3];
 
@@ -233,11 +235,21 @@ bool InitializeEPA(TArray<TVec3<T>>& VertsA, TArray<TVec3<T>>& VertsB, const Sup
 			bValid &= OutEntries[1].Initialize(VertsA.GetData(), VertsB.GetData(), 0, 3, 2, { 2,0,3 }, { 2, 1, 0 });
 			bValid &= OutEntries[2].Initialize(VertsA.GetData(), VertsB.GetData(), 0, 1, 3, { 3,0, 1 }, { 2,2,0 });
 			bValid &= OutEntries[3].Initialize(VertsA.GetData(), VertsB.GetData(), 0, 2, 1, { 1,0,2 }, { 2,0,0 });
-			ensure(bValid);	//expect user to give us valid tetrahedron
+			
+			if(!bValid)
+			{
+				CHAOS_ENSURE(false);	//expect user to give us valid tetrahedron
+				UE_LOG(LogChaos, Log, TEXT("Invalid tetrahedron encountered in InitializeEPA"));
+			}
+
 			break;
 		}
 
-		default: ensure(false);
+		default: 
+		{
+			CHAOS_ENSURE(false);
+			break;
+		}
 	}
 
 	if (bValid)
