Chaos : DDC
- Update PhysicsMaterial to support SleepCounterThreshold

diff --git a/Engine/Source/Runtime/Core/Public/UObject/ExternalPhysicsCustomObjectVersion.h b/Engine/Source/Runtime/Core/Public/UObject/ExternalPhysicsCustomObjectVersion.h
--- a/Engine/Source/Runtime/Core/Public/UObject/ExternalPhysicsCustomObjectVersion.h
+++ b/Engine/Source/Runtime/Core/Public/UObject/ExternalPhysicsCustomObjectVersion.h
@@ -120,6 +120,9 @@ struct CORE_API FExternalPhysicsCustomObjectVersion
 		// Force rebuild of indices in BodySetup
 		ForceRebuildBodySetupIndices,
 
+		// Added serialization for the physics material sleep counter threshold
+		PhysicsMaterialSleepCounterThreshold,
+
 		// -----<new versions can be added above this line>-------------------------------------------------
 		VersionPlusOne,
 		LatestVersion = VersionPlusOne - 1
diff --git a/Engine/Source/Runtime/Engine/Private/PhysicsEngine/Experimental/ChaosDerivedData.cpp b/Engine/Source/Runtime/Engine/Private/PhysicsEngine/Experimental/ChaosDerivedData.cpp
--- a/Engine/Source/Runtime/Engine/Private/PhysicsEngine/Experimental/ChaosDerivedData.cpp
+++ b/Engine/Source/Runtime/Engine/Private/PhysicsEngine/Experimental/ChaosDerivedData.cpp
@@ -24,7 +24,7 @@ const TCHAR* FChaosDerivedDataCooker::GetPluginName() const
 
 const TCHAR* FChaosDerivedDataCooker::GetVersionString() const
 {
-	return TEXT("C247EE00C33148EA87C9F53715E83429");
+	return TEXT("99F5C9095ED541A5A8A8434EB38D8303");
 }
 
 FString FChaosDerivedDataCooker::GetPluginSpecificCacheKeySuffix() const
diff --git a/Engine/Source/Runtime/Experimental/Chaos/Public/Chaos/Defines.h b/Engine/Source/Runtime/Experimental/Chaos/Public/Chaos/Defines.h
--- a/Engine/Source/Runtime/Experimental/Chaos/Public/Chaos/Defines.h
+++ b/Engine/Source/Runtime/Experimental/Chaos/Public/Chaos/Defines.h
@@ -30,6 +30,7 @@ namespace Chaos
 		Chaos::FReal SleepingAngularThreshold;
 		Chaos::FReal DisabledLinearThreshold;
 		Chaos::FReal DisabledAngularThreshold;
+		int32 SleepCounterThreshold;
 		void* UserData;
 
 		FChaosPhysicsMaterial()
@@ -39,6 +40,7 @@ namespace Chaos
 			, SleepingAngularThreshold(1)
 			, DisabledLinearThreshold(0)
 			, DisabledAngularThreshold(0)
+			, SleepCounterThreshold(0)
 			, UserData(nullptr)
 		{
 		}
@@ -61,6 +63,11 @@ namespace Chaos
 		void Serialize(FArchive& Ar)
 		{
 			Ar << Friction << Restitution << SleepingLinearThreshold << SleepingAngularThreshold << DisabledLinearThreshold << DisabledAngularThreshold;
+
+			if (Ar.CustomVer(FExternalPhysicsCustomObjectVersion::GUID) >= FExternalPhysicsCustomObjectVersion::PhysicsMaterialSleepCounterThreshold)
+			{
+				Ar << SleepCounterThreshold;
+			}
 		}
 	};
 
