Chaos - RBAN alternative collision resolution

diff --git a/Engine/Source/Runtime/Engine/Private/PhysicsEngine/ImmediatePhysics/ImmediatePhysicsChaos/ImmediatePhysicsSimulation_Chaos.cpp b/Engine/Source/Runtime/Engine/Private/PhysicsEngine/ImmediatePhysics/ImmediatePhysicsChaos/ImmediatePhysicsSimulation_Chaos.cpp
--- a/Engine/Source/Runtime/Engine/Private/PhysicsEngine/ImmediatePhysics/ImmediatePhysicsChaos/ImmediatePhysicsSimulation_Chaos.cpp
+++ b/Engine/Source/Runtime/Engine/Private/PhysicsEngine/ImmediatePhysics/ImmediatePhysicsChaos/ImmediatePhysicsSimulation_Chaos.cpp
@@ -100,6 +100,16 @@ FAutoConsoleVariableRef CVarChaosImmPhysJointMinParentMassRatio(TEXT("p.Chaos.Im
 FAutoConsoleVariableRef CVarChaosImmPhysJointMaxInertiaRatio(TEXT("p.Chaos.ImmPhys.Joint.MaxInertiaRatio"), ChaosImmediate_Joint_MaxInertiaRatio, TEXT("6Dof joint MaxInertiaRatio (if > 0)"));
 FAutoConsoleVariableRef CVarChaosImmPhysJointAngularPositionCorrection(TEXT("p.Chaos.ImmPhys.Joint.AngularPositionCorrection"), ChaosImmediate_Joint_AngularPositionCorrection, TEXT("6Dof joint post-rotation constraint position correction amount [0-1]"));
 
+//
+// Even more temp that the above...
+//
+int32 ChaosImmediate_UsePositionSolver = 1;
+int32 ChaosImmediate_PositionSolverIts = 3;
+int32 ChaosImmediate_PositionSolverPairIts = 2;
+FAutoConsoleVariableRef CVarChaosImmPhysUsePositionSolver(TEXT("p.Chaos.ImmPhys.UsePositionSolver"), ChaosImmediate_UsePositionSolver, TEXT("Use position based collision solver for Immediate Physics (default true)"));
+FAutoConsoleVariableRef CVarChaosImmPhysPositionIts(TEXT("p.Chaos.ImmPhys.PositionSolverIts"), ChaosImmediate_PositionSolverIts, TEXT("Number of iterations to use in position solver mode (default 3)"));
+FAutoConsoleVariableRef CVarChaosImmPhysPositionPairIts(TEXT("p.Chaos.ImmPhys.PositionSolverPairIts"), ChaosImmediate_PositionSolverPairIts, TEXT("Number of pair-iterations to use in position solver mode (default 2)"));
+
 //
 // end remove when finished
 //
@@ -171,6 +181,7 @@ namespace ImmediatePhysics_Chaos
 
 		// RBAN collision customization
 		Collisions.DisableHandles();
+		Collisions.SetApplyType(ECollisionApplyType::Position);
 		CollisionDetector.GetContext().bFilteringEnabled = false;
 
 #if CHAOS_DEBUG_DRAW
@@ -577,13 +588,30 @@ namespace ImmediatePhysics_Chaos
 			BroadPhase.SetCullDustance(ChaosImmediate_Collision_CullDistance);
 			Evolution.SetBoundsExtension(ChaosImmediate_Evolution_BoundsExtension);
 
-			SetSolverIterations(
-				ChaosImmediate_Evolution_Iterations,
-				ChaosImmediate_Joint_PairIterations,
-				ChaosImmediate_Collision_PairIterations,
-				ChaosImmediate_Evolution_PushOutIterations,
-				ChaosImmediate_Joint_PushOutPairIterations,
-				ChaosImmediate_Collision_PushOutPairIterations);
+			if (ChaosImmediate_UsePositionSolver)
+			{
+				SetSolverIterations(
+					ChaosImmediate_PositionSolverIts,
+					ChaosImmediate_PositionSolverPairIts,
+					ChaosImmediate_PositionSolverPairIts,
+					0,
+					0,
+					0);
+
+				Collisions.SetApplyType(ECollisionApplyType::Position);
+			}
+			else
+			{
+				SetSolverIterations(
+					ChaosImmediate_Evolution_Iterations,
+					ChaosImmediate_Joint_PairIterations,
+					ChaosImmediate_Collision_PairIterations,
+					ChaosImmediate_Evolution_PushOutIterations,
+					ChaosImmediate_Joint_PushOutPairIterations,
+					ChaosImmediate_Collision_PushOutPairIterations);
+
+				Collisions.SetApplyType(ECollisionApplyType::Velocity);
+			}
 		}
 		UE_LOG(LogChaosJoint, Verbose, TEXT("Simulate Dt = %f Steps %d x %f"), DeltaTime, NumSteps, StepTime);
 
diff --git a/Engine/Source/Runtime/Experimental/Chaos/Private/Chaos/PBDCollisionConstraints.cpp b/Engine/Source/Runtime/Experimental/Chaos/Private/Chaos/PBDCollisionConstraints.cpp
--- a/Engine/Source/Runtime/Experimental/Chaos/Private/Chaos/PBDCollisionConstraints.cpp
+++ b/Engine/Source/Runtime/Experimental/Chaos/Private/Chaos/PBDCollisionConstraints.cpp
@@ -48,6 +48,9 @@ namespace Chaos
 	float CollisionFrictionOverride = -1.0f;
 	FAutoConsoleVariableRef CVarCollisionFrictionOverride(TEXT("p.CollisionFriction"), CollisionFrictionOverride, TEXT("Collision friction for all contacts if >= 0"));
 
+	float CollisionRestitutionOverride = -1.0f;
+	FAutoConsoleVariableRef CVarCollisionRestitutionOverride(TEXT("p.CollisionRestitution"), CollisionRestitutionOverride, TEXT("Collision restitution for all contacts if >= 0"));
+
 	CHAOS_API int32 EnableCollisions = 1;
 	FAutoConsoleVariableRef CVarEnableCollisions(TEXT("p.EnableCollisions"), EnableCollisions, TEXT("Enable/Disable collisions on the Chaos solver."));
 	
@@ -84,6 +87,7 @@ namespace Chaos
 		, bUseCCD(false)
 		, bEnableCollisions(true)
 		, bHandlesEnabled(true)
+		, ApplyType(ECollisionApplyType::Velocity)
 		, LifespanCounter(0)
 		, PostApplyCallback(nullptr)
 		, PostApplyPushOutCallback(nullptr)
@@ -156,6 +160,10 @@ namespace Chaos
 		{
 			Contact.Friction = CollisionFrictionOverride;
 		}
+		if (CollisionRestitutionOverride >= 0)
+		{
+			Contact.Restitution = CollisionRestitutionOverride;
+		}
 	}
 
 	template<typename T, int d>
@@ -413,7 +421,7 @@ namespace Chaos
 		if (MApplyPairIterations > 0)
 		{
 			const Collisions::TContactParticleParameters<T> ParticleParameters = { MCullDistance, MShapePadding, &MCollided };
-			const Collisions::TContactIterationParameters<T> IterationParameters = { Dt, Iterations, NumIterations, MApplyPairIterations, nullptr };
+			const Collisions::TContactIterationParameters<T> IterationParameters = { Dt, Iterations, NumIterations, MApplyPairIterations, ApplyType, nullptr };
 
 			for (FPointContactConstraint& Contact : PointConstraints)
 			{
@@ -442,7 +450,7 @@ namespace Chaos
 		if (MApplyPushOutPairIterations > 0)
 		{
 			const Collisions::TContactParticleParameters<T> ParticleParameters = { MCullDistance, MShapePadding, &MCollided };
-			const Collisions::TContactIterationParameters<T> IterationParameters = { Dt, Iterations, NumIterations, MApplyPushOutPairIterations, &bNeedsAnotherIteration };
+			const Collisions::TContactIterationParameters<T> IterationParameters = { Dt, Iterations, NumIterations, MApplyPushOutPairIterations, ECollisionApplyType::None, &bNeedsAnotherIteration };
 
 			for (FPointContactConstraint& Contact : PointConstraints)
 			{
@@ -476,7 +484,7 @@ namespace Chaos
 				check(ConstraintHandle != nullptr);
 
 				Collisions::TContactParticleParameters<T> ParticleParameters = { MCullDistance, MShapePadding, &MCollided };
-				Collisions::TContactIterationParameters<T> IterationParameters = { Dt, Iterations, NumIterations, MApplyPairIterations, nullptr };
+				Collisions::TContactIterationParameters<T> IterationParameters = { Dt, Iterations, NumIterations, MApplyPairIterations, ApplyType, nullptr };
 				Collisions::Apply(ConstraintHandle->GetContact(), IterationParameters, ParticleParameters);
 
 			}, bDisableCollisionParallelFor);
@@ -503,7 +511,7 @@ namespace Chaos
 				check(ConstraintHandle != nullptr);
 
 				Collisions::TContactParticleParameters<T> ParticleParameters = { MCullDistance, MShapePadding, &MCollided };
-				Collisions::TContactIterationParameters<T> IterationParameters = { Dt, Iteration, NumIterations, MApplyPushOutPairIterations, &bNeedsAnotherIteration };
+				Collisions::TContactIterationParameters<T> IterationParameters = { Dt, Iteration, NumIterations, MApplyPushOutPairIterations, ECollisionApplyType::None, &bNeedsAnotherIteration };
 				Collisions::ApplyPushOut(ConstraintHandle->GetContact(), IsTemporarilyStatic, IterationParameters, ParticleParameters);
 
 			}, bDisableCollisionParallelFor);
diff --git a/Engine/Source/Runtime/Experimental/Chaos/Private/Chaos/PBDCollisionConstraintsContact.cpp b/Engine/Source/Runtime/Experimental/Chaos/Private/Chaos/PBDCollisionConstraintsContact.cpp
--- a/Engine/Source/Runtime/Experimental/Chaos/Private/Chaos/PBDCollisionConstraintsContact.cpp
+++ b/Engine/Source/Runtime/Experimental/Chaos/Private/Chaos/PBDCollisionConstraintsContact.cpp
@@ -7,14 +7,17 @@
 #include "Chaos/Particle/ParticleUtilities.h"
 #include "Chaos/Utilities.h"
 
-//#pragma optimize("", off)
+//PRAGMA_DISABLE_OPTIMIZATION
 
 namespace Chaos
 {
 	namespace Collisions
 	{
 		int32 Chaos_Collision_EnergyClampEnabled = 1;
-		FAutoConsoleVariableRef CVarChaosImmPhysStepTime(TEXT("p.Chaos.Collision.EnergyClampEnabled"), Chaos_Collision_EnergyClampEnabled, TEXT("Whether to use energy clamping in collision apply step"));
+		FAutoConsoleVariableRef CVarChaosCollisionEnergyClampEnabled(TEXT("p.Chaos.Collision.EnergyClampEnabled"), Chaos_Collision_EnergyClampEnabled, TEXT("Whether to use energy clamping in collision apply step"));
+
+		int32 Chaos_Collision_ForceApplyType = 0;
+		FAutoConsoleVariableRef CVarChaosCollisionAlternativeApply(TEXT("p.Chaos.Collision.ForceApplyType"), Chaos_Collision_ForceApplyType, TEXT("Force Apply step to use Velocity(1) or Position(2) modes"));
 
 		extern void UpdateManifold(FRigidBodyMultiPointContactConstraint& Constraint, const FReal CullDistance, const FCollisionContext& Context)
 		{
@@ -209,6 +212,115 @@ namespace Chaos
 			return AccumulatedImpulse;
 		}
 
+		// A PBD collision penetration correction.
+		FVec3 ApplyContact2(FCollisionContact& Contact,
+			TGenericParticleHandle<FReal, 3> Particle0,
+			TGenericParticleHandle<FReal, 3> Particle1,
+			const TContactIterationParameters<FReal>& IterationParameters,
+			const TContactParticleParameters<FReal>& ParticleParameters)
+		{
+			FVec3 AccumulatedImpulse(0);
+
+			TPBDRigidParticleHandle<FReal, 3>* PBDRigid0 = Particle0->CastToRigidParticle();
+			TPBDRigidParticleHandle<FReal, 3>* PBDRigid1 = Particle1->CastToRigidParticle();
+
+			bool bIsRigidDynamic0 = PBDRigid0 && PBDRigid0->ObjectState() == EObjectStateType::Dynamic;
+			bool bIsRigidDynamic1 = PBDRigid1 && PBDRigid1->ObjectState() == EObjectStateType::Dynamic;
+
+			FVec3 X0 = FParticleUtilitiesXR::GetCoMWorldPosition(Particle0);
+			FVec3 X1 = FParticleUtilitiesXR::GetCoMWorldPosition(Particle1);
+			FVec3 P0 = FParticleUtilities::GetCoMWorldPosition(Particle0);
+			FVec3 P1 = FParticleUtilities::GetCoMWorldPosition(Particle1);
+			FRotation3 R0 = FParticleUtilitiesXR::GetCoMWorldRotation(Particle0);
+			FRotation3 R1 = FParticleUtilitiesXR::GetCoMWorldRotation(Particle1);
+			FRotation3 Q0 = FParticleUtilities::GetCoMWorldRotation(Particle0);
+			FRotation3 Q1 = FParticleUtilities::GetCoMWorldRotation(Particle1);
+			FVec3 VectorToPoint0 = Contact.Location - P0;
+			FVec3 VectorToPoint1 = Contact.Location - P1;
+
+			if (Contact.Phi < 0)
+			{
+				FReal InvM0 = bIsRigidDynamic0 ? PBDRigid0->InvM() : 0.0f;
+				FReal InvM1 = bIsRigidDynamic1 ? PBDRigid1->InvM() : 0.0f;
+				FMatrix33 InvI0 = bIsRigidDynamic0 ? Utilities::ComputeWorldSpaceInertia(Q0, PBDRigid0->InvI()) : FMatrix33(0);
+				FMatrix33 InvI1 = bIsRigidDynamic1 ? Utilities::ComputeWorldSpaceInertia(Q1, PBDRigid1->InvI()) : FMatrix33(0);
+				FMatrix33 ContactInvI =
+					(bIsRigidDynamic0 ? ComputeFactorMatrix3(VectorToPoint0, InvI0, InvM0) : FMatrix33(0)) +
+					(bIsRigidDynamic1 ? ComputeFactorMatrix3(VectorToPoint1, InvI1, InvM1) : FMatrix33(0));
+
+				// Calculate the normal correction
+				// @todo(chaos): support restitution in this mode (requires a small-impact cutoff threshold)
+				FReal Restitution = 0.0f;//Contact.Restitution;
+				FVec3 NormalError = Contact.Phi * Contact.Normal;
+				FReal NormalImpulseDenominator = FVec3::DotProduct(Contact.Normal, ContactInvI * Contact.Normal);
+				FVec3 NormalImpulseNumerator = -(1.0f + Restitution) * NormalError;
+				FVec3 NormalCorrection = NormalImpulseNumerator / NormalImpulseDenominator;
+
+				// Calculate the lateral correction
+				FVec3 LateralCorrection = FVec3(0);
+				if (Contact.Friction > 0)
+				{
+					// Get contact velocity
+					FVec3 V0 = FVec3::CalculateVelocity(X0, P0, IterationParameters.Dt);
+					FVec3 W0 = FRotation3::CalculateAngularVelocity(R0, Q0, IterationParameters.Dt);
+					FVec3 V1 = FVec3::CalculateVelocity(X1, P1, IterationParameters.Dt);
+					FVec3 W1 = FRotation3::CalculateAngularVelocity(R1, Q1, IterationParameters.Dt);
+					FVec3 CV0 = V0 + FVec3::CrossProduct(W0, VectorToPoint0);
+					FVec3 CV1 = V1 + FVec3::CrossProduct(W1, VectorToPoint1);
+					FVec3 CV = CV0 - CV1;
+
+					// Calculate lateral correction, clamped to the friction cone. Kinda.
+					FReal CVNormalMag = FVec3::DotProduct(CV, Contact.Normal);
+					if (CVNormalMag < 0.0f)
+					{
+						FVec3 CVLateral = CV - CVNormalMag * Contact.Normal;
+						FReal CVLateralMag = CVLateral.Size();
+						if (CVLateralMag > KINDA_SMALL_NUMBER)
+						{
+							FVec3 DirLateral = CVLateral / CVLateralMag;
+							FVec3 LateralImpulseNumerator = -CVLateral * IterationParameters.Dt;
+							FReal LateralImpulseDenominator = FVec3::DotProduct(DirLateral, ContactInvI * DirLateral);
+							LateralCorrection = LateralImpulseNumerator / LateralImpulseDenominator;
+							FReal LateralImpulseMag = LateralCorrection.Size();
+							FReal NormalImpulseMag = NormalCorrection.Size();
+							if (LateralImpulseMag > Contact.Friction * NormalImpulseMag)
+							{
+								LateralCorrection *= Contact.Friction * NormalImpulseMag / LateralImpulseMag;
+							}
+						}
+					}
+				}
+
+				// Net Correction
+				FVec3 DX = NormalCorrection + LateralCorrection;
+				
+				if (bIsRigidDynamic0)
+				{
+					FVec3 DP0 = InvM0 * DX;
+					FVec3 DR0 = Utilities::Multiply(InvI0, FVec3::CrossProduct(VectorToPoint0, DX));
+					PBDRigid0->V() += DP0 / IterationParameters.Dt;
+					PBDRigid0->W() += DR0 / IterationParameters.Dt;
+					P0 += DP0;
+					Q0 += FRotation3::FromElements(DR0, 0.f) * Q0 * IterationParameters.Dt * FReal(0.5);
+					Q0.Normalize();
+					FParticleUtilities::SetCoMWorldTransform(PBDRigid0, P0, Q0);
+				}
+				if (bIsRigidDynamic1)
+				{
+					FVec3 DP1 = InvM1 * -DX;
+					FVec3 DR1 = Utilities::Multiply(InvI1, FVec3::CrossProduct(VectorToPoint1, -DX));
+					PBDRigid1->V() += DP1 / IterationParameters.Dt;
+					PBDRigid1->W() += DR1 / IterationParameters.Dt;
+					P1 += DP1;
+					Q1 += FRotation3::FromElements(DR1, 0.f) * Q1 * IterationParameters.Dt * FReal(0.5);
+					Q1.Normalize();
+					FParticleUtilities::SetCoMWorldTransform(PBDRigid1, P1, Q1);
+				}
+			}
+			return AccumulatedImpulse;
+		}
+
+
 		template<typename T_CONSTRAINT>
 		void ApplyImpl(T_CONSTRAINT& Constraint, const FContactIterationParameters & IterationParameters, const FContactParticleParameters & ParticleParameters)
 		{
@@ -231,7 +343,6 @@ namespace Chaos
 					return;
 				}
 
-				// @todo(ccaulfield): CHAOS_PARTICLEHANDLE_TODO what's the best way to manage external per-particle data?
 				if (ParticleParameters.Collided)
 				{
 					Particle0->AuxilaryValue(*ParticleParameters.Collided) = true;
@@ -244,9 +355,21 @@ namespace Chaos
 				//   For example, and iterative constraint might have 4 penetrating points that need to be resolved. 
 				//
 
+				// What Apply algorithm should we use? Controlled by the solver, with forcable cvar override for now...
+				bool bUseVelocityMode = (IterationParameters.ApplyType == ECollisionApplyType::Velocity);
+				if (Chaos_Collision_ForceApplyType != 0)
+				{
+					bUseVelocityMode = (Chaos_Collision_ForceApplyType == (int32)ECollisionApplyType::Velocity);
+				}
 
-				Constraint.AccumulatedImpulse +=
-					ApplyContact(Constraint.Manifold, Particle0, Particle1, IterationParameters, ParticleParameters);
+				if (bUseVelocityMode)
+				{
+					Constraint.AccumulatedImpulse += ApplyContact(Constraint.Manifold, Particle0, Particle1, IterationParameters, ParticleParameters);
+				}
+				else
+				{
+					Constraint.AccumulatedImpulse += ApplyContact2(Constraint.Manifold, Particle0, Particle1, IterationParameters, ParticleParameters);
+				}
 			}
 		}
 
diff --git a/Engine/Source/Runtime/Experimental/Chaos/Public/Chaos/Collision/CollisionApplyType.h b/Engine/Source/Runtime/Experimental/Chaos/Public/Chaos/Collision/CollisionApplyType.h
new file mode 100644
--- /dev/null
+++ b/Engine/Source/Runtime/Experimental/Chaos/Public/Chaos/Collision/CollisionApplyType.h
@@ -0,0 +1,15 @@
+// Copyright Epic Games, Inc. All Rights Reserved.
+#pragma once
+
+namespace Chaos
+{
+	/**
+	 * What solver method to use for Collision Resolution
+	 */
+	enum ECollisionApplyType
+	{
+		None,
+		Velocity,	// Solve for velocity in the Apply step 
+		Position,	// Solve for position in the Apply step
+	};
+}
\ No newline at end of file
diff --git a/Engine/Source/Runtime/Experimental/Chaos/Public/Chaos/PBDCollisionConstraints.h b/Engine/Source/Runtime/Experimental/Chaos/Public/Chaos/PBDCollisionConstraints.h
--- a/Engine/Source/Runtime/Experimental/Chaos/Public/Chaos/PBDCollisionConstraints.h
+++ b/Engine/Source/Runtime/Experimental/Chaos/Public/Chaos/PBDCollisionConstraints.h
@@ -1,10 +1,11 @@
 // Copyright Epic Games, Inc. All Rights Reserved.
 #pragma once
 
-#include "Chaos/ConstraintHandle.h"
 #include "Chaos/CollisionResolutionTypes.h"
-#include "Chaos/PBDConstraintContainer.h"
 #include "Chaos/CollisionResolutionTypes.h"
+#include "Chaos/Collision/CollisionApplyType.h"
+#include "Chaos/ConstraintHandle.h"
+#include "Chaos/PBDConstraintContainer.h"
 #include "Framework/BufferedData.h"
 
 #include <memory>
@@ -98,6 +99,14 @@ public:
 	 */
 	void DisableHandles();
 
+	/**
+	 * Set the solver method to use in the Apply step
+	 */
+	void SetApplyType(ECollisionApplyType InApplyType)
+	{
+		ApplyType = InApplyType;
+	}
+
 	/**
 	*  Add the constraint to the container. 
 	*
@@ -315,6 +324,7 @@ private:
 	bool bUseCCD;
 	bool bEnableCollisions;
 	bool bHandlesEnabled;
+	ECollisionApplyType ApplyType;
 
 	int32 LifespanCounter;
 
diff --git a/Engine/Source/Runtime/Experimental/Chaos/Public/Chaos/PBDCollisionConstraintsContact.h b/Engine/Source/Runtime/Experimental/Chaos/Public/Chaos/PBDCollisionConstraintsContact.h
--- a/Engine/Source/Runtime/Experimental/Chaos/Public/Chaos/PBDCollisionConstraintsContact.h
+++ b/Engine/Source/Runtime/Experimental/Chaos/Public/Chaos/PBDCollisionConstraintsContact.h
@@ -4,6 +4,7 @@
 #include "Chaos/ExternalCollisionData.h"
 #include "Chaos/PBDCollisionConstraints.h"
 #include "Chaos/CollisionResolutionTypes.h"
+#include "Chaos/Collision/CollisionApplyType.h"
 
 namespace Chaos
 {
@@ -11,7 +12,6 @@ namespace Chaos
 
 	namespace Collisions
 	{
-
 		template<class T = float>
 		struct TContactParticleParameters {
 			T CullDistance;
@@ -26,6 +26,7 @@ namespace Chaos
 			const int32 Iteration;
 			const int32 NumIterations;
 			const int32 NumPairIterations;
+			const ECollisionApplyType ApplyType;	// @todo(chaos): a better way to customize the collision solver
 			bool* NeedsAnotherIteration;
 		};
 		using FContactIterationParameters = TContactIterationParameters<FReal>;
