Chaos - fix enum def (missing class)

diff --git a/Engine/Source/Runtime/Experimental/Chaos/Public/Chaos/Collision/CollisionApplyType.h b/Engine/Source/Runtime/Experimental/Chaos/Public/Chaos/Collision/CollisionApplyType.h
--- a/Engine/Source/Runtime/Experimental/Chaos/Public/Chaos/Collision/CollisionApplyType.h
+++ b/Engine/Source/Runtime/Experimental/Chaos/Public/Chaos/Collision/CollisionApplyType.h
@@ -6,7 +6,7 @@ namespace Chaos
 	/**
 	 * What solver method to use for Collision Resolution
 	 */
-	enum ECollisionApplyType
+	enum class ECollisionApplyType
 	{
 		None,
 		Velocity,	// Solve for velocity in the Apply step 
