Fixed an issue where the tri normal wasn't being normalized, so components moving via a movement component would sometimes resolve to crazy positions.

diff --git a/Engine/Source/Runtime/Experimental/Chaos/Private/Chaos/HeightField.cpp b/Engine/Source/Runtime/Experimental/Chaos/Private/Chaos/HeightField.cpp
--- a/Engine/Source/Runtime/Experimental/Chaos/Private/Chaos/HeightField.cpp
+++ b/Engine/Source/Runtime/Experimental/Chaos/Private/Chaos/HeightField.cpp
@@ -267,6 +267,7 @@ namespace Chaos
 								const TVector<T, 3> AC = C - A;
 
 								TVector<T, 3> TriNormal = TVector<T, 3>::CrossProduct(AB, AC);
+								TriNormal.SafeNormalize();
 
 								OutNormal = TriNormal;
 								CurrentLength = 0;
