Add missing include.

diff --git a/Engine/Source/Runtime/Experimental/Chaos/Public/Chaos/Defines.h b/Engine/Source/Runtime/Experimental/Chaos/Public/Chaos/Defines.h
--- a/Engine/Source/Runtime/Experimental/Chaos/Public/Chaos/Defines.h
+++ b/Engine/Source/Runtime/Experimental/Chaos/Public/Chaos/Defines.h
@@ -10,6 +10,7 @@
 #endif
 #include "Serializable.h"
 
+#include "UObject/ExternalPhysicsCustomObjectVersion.h"
 #include "Chaos/Core.h"
 
 #if COMPILE_WITHOUT_UNREAL_SUPPORT
