Chaos
- fix quaternion flip in joints that causes occasional pops
- add extra logging (compiled out of test and shipping)

diff --git a/Engine/Source/Runtime/Engine/Private/PhysicsEngine/ImmediatePhysics/ImmediatePhysicsChaos/ImmediatePhysicsSimulation_Chaos.cpp b/Engine/Source/Runtime/Engine/Private/PhysicsEngine/ImmediatePhysics/ImmediatePhysicsChaos/ImmediatePhysicsSimulation_Chaos.cpp
--- a/Engine/Source/Runtime/Engine/Private/PhysicsEngine/ImmediatePhysics/ImmediatePhysicsChaos/ImmediatePhysicsSimulation_Chaos.cpp
+++ b/Engine/Source/Runtime/Engine/Private/PhysicsEngine/ImmediatePhysics/ImmediatePhysicsChaos/ImmediatePhysicsSimulation_Chaos.cpp
@@ -7,6 +7,7 @@
 
 #include "Chaos/ChaosDebugDraw.h"
 #include "Chaos/DebugDrawQueue.h"
+#include "Chaos/Joint/ChaosJointLog.h"
 #include "Chaos/ParticleHandle.h"
 #include "Chaos/PBDCollisionConstraints.h"
 #include "Chaos/PBDJointConstraints.h"
@@ -18,7 +19,7 @@
 #include "PhysicsEngine/BodyInstance.h"
 #include "PhysicsEngine/ConstraintInstance.h"
 
-//#pragma optimize("", off)
+//PRAGMA_DISABLE_OPTIMIZATION
 
 DECLARE_CYCLE_STAT(TEXT("FSimulation::Simulate_Chaos"), STAT_ImmediateSimulate_Chaos, STATGROUP_ImmediatePhysics);
 
diff --git a/Engine/Source/Runtime/Experimental/Chaos/Private/Chaos/Joint/ChaosJointLog.cpp b/Engine/Source/Runtime/Experimental/Chaos/Private/Chaos/Joint/ChaosJointLog.cpp
new file mode 100644
--- /dev/null
+++ b/Engine/Source/Runtime/Experimental/Chaos/Private/Chaos/Joint/ChaosJointLog.cpp
@@ -0,0 +1,7 @@
+// Copyright Epic Games, Inc. All Rights Reserved.
+
+#include "Chaos/Joint/ChaosJointLog.h"
+#include "ChaosModule.h"
+#include "CoreMinimal.h"
+
+DEFINE_LOG_CATEGORY(LogChaosJoint);
diff --git a/Engine/Source/Runtime/Experimental/Chaos/Private/Chaos/Joint/PBDJointSolverCholesky.cpp b/Engine/Source/Runtime/Experimental/Chaos/Private/Chaos/Joint/PBDJointSolverCholesky.cpp
--- a/Engine/Source/Runtime/Experimental/Chaos/Private/Chaos/Joint/PBDJointSolverCholesky.cpp
+++ b/Engine/Source/Runtime/Experimental/Chaos/Private/Chaos/Joint/PBDJointSolverCholesky.cpp
@@ -1,9 +1,9 @@
 // Copyright Epic Games, Inc. All Rights Reserved.
 #include "Chaos/Joint/PBDJointSolverCholesky.h"
+#include "Chaos/Joint/ChaosJointLog.h"
 #include "Chaos/ParticleHandle.h"
 #include "Chaos/PBDJointConstraintUtilities.h"
 #include "Chaos/Utilities.h"
-#include "ChaosLog.h"
 #include "ChaosStats.h"
 
 //#pragma optimize("", off)
diff --git a/Engine/Source/Runtime/Experimental/Chaos/Private/Chaos/Joint/PBDJointSolverGaussSeidel.cpp b/Engine/Source/Runtime/Experimental/Chaos/Private/Chaos/Joint/PBDJointSolverGaussSeidel.cpp
--- a/Engine/Source/Runtime/Experimental/Chaos/Private/Chaos/Joint/PBDJointSolverGaussSeidel.cpp
+++ b/Engine/Source/Runtime/Experimental/Chaos/Private/Chaos/Joint/PBDJointSolverGaussSeidel.cpp
@@ -1,9 +1,9 @@
 // Copyright Epic Games, Inc. All Rights Reserved.
 #include "Chaos/Joint/PBDJointSolverGaussSeidel.h"
+#include "Chaos/Joint/ChaosJointLog.h"
 #include "Chaos/ParticleHandle.h"
 #include "Chaos/PBDJointConstraintUtilities.h"
 #include "Chaos/Utilities.h"
-#include "ChaosLog.h"
 #include "ChaosStats.h"
 #if INTEL_ISPC
 #include "PBDJointSolverGaussSeidel.ispc.generated.h"
@@ -135,12 +135,14 @@ namespace Chaos
 		Xs[1] = Ps[1] + Qs[1] * XLs[1].GetTranslation();
 		Rs[0] = Qs[0] * XLs[0].GetRotation();
 		Rs[1] = Qs[1] * XLs[1].GetRotation();
+		Rs[1].EnforceShortestArcWith(Rs[0]);
 	}
 
 	void FJointSolverGaussSeidel::UpdateDerivedState(const int32 BodyIndex)
 	{
 		Xs[BodyIndex] = Ps[BodyIndex] + Qs[BodyIndex] * XLs[BodyIndex].GetTranslation();
 		Rs[BodyIndex] = Qs[BodyIndex] * XLs[BodyIndex].GetRotation();
+		Rs[1].EnforceShortestArcWith(Rs[0]);
 	}
 
 
@@ -586,7 +588,7 @@ namespace Chaos
 		const FReal Stiffness,
 		const FVec3& DP)
 	{
-		//UE_LOG(LogChaosJoint, VeryVerbose, TEXT("      Apply DP%d %f %f %f"), BodyIndex, DP.X, DP.Y, DP.Z);
+		UE_LOG(LogChaosJoint, VeryVerbose, TEXT("      Apply DP%d %f %f %f"), BodyIndex, DP.X, DP.Y, DP.Z);
 
 		Ps[BodyIndex] += Stiffness * DP;
 	}
@@ -597,8 +599,8 @@ namespace Chaos
 		const FVec3& DP0,
 		const FVec3& DP1)
 	{
-		//UE_LOG(LogChaosJoint, VeryVerbose, TEXT("      Apply DP%d %f %f %f"), 0, DP0.X, DP0.Y, DP0.Z);
-		//UE_LOG(LogChaosJoint, VeryVerbose, TEXT("      Apply DP%d %f %f %f"), 1, DP1.X, DP1.Y, DP1.Z);
+		UE_LOG(LogChaosJoint, VeryVerbose, TEXT("      Apply DP%d %f %f %f"), 0, DP0.X, DP0.Y, DP0.Z);
+		UE_LOG(LogChaosJoint, VeryVerbose, TEXT("      Apply DP%d %f %f %f"), 1, DP1.X, DP1.Y, DP1.Z);
 
 		Ps[0] += Stiffness * DP0;
 		Ps[1] += Stiffness * DP1;
@@ -610,10 +612,11 @@ namespace Chaos
 		const FReal Stiffness,
 		const FVec3& DR)
 	{
-		//UE_LOG(LogChaosJoint, VeryVerbose, TEXT("      Apply DR%d %f %f %f"), BodyIndex, DR.X, DR.Y, DR.Z);
+		UE_LOG(LogChaosJoint, VeryVerbose, TEXT("      Apply DR%d %f %f %f"), BodyIndex, DR.X, DR.Y, DR.Z);
 
 		const FRotation3 DQ = (FRotation3::FromElements(Stiffness * DR, 0) * Qs[BodyIndex]) * (FReal)0.5;
 		Qs[BodyIndex] = (Qs[BodyIndex] + DQ).GetNormalized();
+		Qs[1].EnforceShortestArcWith(Qs[0]);
 	}
 
 
@@ -622,8 +625,8 @@ namespace Chaos
 		const FVec3& DR0,
 		const FVec3& DR1)
 	{
-		//UE_LOG(LogChaosJoint, VeryVerbose, TEXT("      Apply DR%d %f %f %f"), 0, DR0.X, DR0.Y, DR0.Z);
-		//UE_LOG(LogChaosJoint, VeryVerbose, TEXT("      Apply DR%d %f %f %f"), 1, DR1.X, DR1.Y, DR1.Z);
+		UE_LOG(LogChaosJoint, VeryVerbose, TEXT("      Apply DR%d %f %f %f"), 0, DR0.X, DR0.Y, DR0.Z);
+		UE_LOG(LogChaosJoint, VeryVerbose, TEXT("      Apply DR%d %f %f %f"), 1, DR1.X, DR1.Y, DR1.Z);
 
 		const FRotation3 DQ0 = (FRotation3::FromElements(Stiffness * DR0, 0) * Qs[0]) * (FReal)0.5;
 		const FRotation3 DQ1 = (FRotation3::FromElements(Stiffness * DR1, 0) * Qs[1]) * (FReal)0.5;
@@ -638,12 +641,13 @@ namespace Chaos
 		const FVec3& DP,
 		const FVec3& DR)
 	{
-		//UE_LOG(LogChaosJoint, VeryVerbose, TEXT("      Apply DP%d %f %f %f"), BodyIndex, DP.X, DP.Y, DP.Z);
-		//UE_LOG(LogChaosJoint, VeryVerbose, TEXT("      Apply DR%d %f %f %f"), BodyIndex, DR.X, DR.Y, DR.Z);
+		UE_LOG(LogChaosJoint, VeryVerbose, TEXT("      Apply DP%d %f %f %f"), BodyIndex, DP.X, DP.Y, DP.Z);
+		UE_LOG(LogChaosJoint, VeryVerbose, TEXT("      Apply DR%d %f %f %f"), BodyIndex, DR.X, DR.Y, DR.Z);
 
 		Ps[BodyIndex] += Stiffness * DP;
 		const FRotation3 DQ = (FRotation3::FromElements(Stiffness * DR, 0) * Qs[BodyIndex]) * (FReal)0.5;
 		Qs[BodyIndex] = (Qs[BodyIndex] + DQ).GetNormalized();
+		Qs[1].EnforceShortestArcWith(Qs[0]);
 	}
 
 
@@ -948,7 +952,7 @@ namespace Chaos
 			DTwistAngle = TwistAngle + TwistAngleMax;
 		}
 
-		//UE_LOG(LogChaosJoint, VeryVerbose, TEXT("    SoftTwist Angle %f [Limit %f]"), FMath::RadiansToDegrees(TwistAngle), FMath::RadiansToDegrees(TwistAngleMax));
+		UE_LOG(LogChaosJoint, VeryVerbose, TEXT("    SoftTwist Angle %f [Limit %f]"), FMath::RadiansToDegrees(TwistAngle), FMath::RadiansToDegrees(TwistAngleMax));
 
 		// Apply twist correction
 		if (FMath::Abs(DTwistAngle) > 0)
@@ -1029,7 +1033,7 @@ namespace Chaos
 			DSwingAngle = SwingAngle + SwingAngleMax;
 		}
 
-		//UE_LOG(LogChaosJoint, VeryVerbose, TEXT("    Cone Angle %f [Limit %f]"), FMath::RadiansToDegrees(SwingAngle), FMath::RadiansToDegrees(SwingAngleMax));
+		UE_LOG(LogChaosJoint, VeryVerbose, TEXT("    Cone Angle %f [Limit %f]"), FMath::RadiansToDegrees(SwingAngle), FMath::RadiansToDegrees(SwingAngleMax));
 
 		// Apply swing correction to each body
 		if (FMath::Abs(DSwingAngle) > 0)
@@ -1114,7 +1118,7 @@ namespace Chaos
 			DSwingAngle = SwingAngle + SwingAngleMax;
 		}
 
-		//UE_LOG(LogChaosJoint, VeryVerbose, TEXT("    Swing%d Angle %f [Limit %f]"), (SwingConstraintIndex == EJointAngularConstraintIndex::Swing1) ? 1 : 2, FMath::RadiansToDegrees(SwingAngle), FMath::RadiansToDegrees(SwingAngleMax));
+		UE_LOG(LogChaosJoint, VeryVerbose, TEXT("    DualConeSwing%d Angle %f [Limit %f]"), (SwingConstraintIndex == EJointAngularConstraintIndex::Swing1) ? 1 : 2, FMath::RadiansToDegrees(SwingAngle), FMath::RadiansToDegrees(SwingAngleMax));
 
 		// Apply swing correction
 		if (FMath::Abs(DSwingAngle) > 0)
@@ -1160,7 +1164,7 @@ namespace Chaos
 			DSwingAngle = SwingAngle + SwingAngleMax;
 		}
 
-		//UE_LOG(LogChaosJoint, VeryVerbose, TEXT("    Swing%d Angle %f [Limit %f]"), (SwingConstraintIndex == EJointAngularConstraintIndex::Swing1) ? 1 : 2, FMath::RadiansToDegrees(SwingAngle), FMath::RadiansToDegrees(SwingAngleMax));
+		UE_LOG(LogChaosJoint, VeryVerbose, TEXT("    Swing%d Angle %f [Limit %f]"), (SwingConstraintIndex == EJointAngularConstraintIndex::Swing1) ? 1 : 2, FMath::RadiansToDegrees(SwingAngle), FMath::RadiansToDegrees(SwingAngleMax));
 
 		// Apply swing correction
 		if (FMath::Abs(DSwingAngle) > 0)
@@ -1261,7 +1265,7 @@ namespace Chaos
 		FReal LinearStiffness = FPBDJointUtilities::GetLinearStiffness(SolverSettings, JointSettings);
 		const FVec3 CX = Xs[DIndex] - Xs[KIndex];
 
-		//UE_LOG(LogChaosJoint, VeryVerbose, TEXT("    PointKD Delta %f [Limit %f]"), CX.Size(), PositionTolerance);
+		UE_LOG(LogChaosJoint, VeryVerbose, TEXT("    PointKD Delta %f [Limit %f]"), CX.Size(), PositionTolerance);
 
 		if (CX.SizeSquared() > PositionTolerance * PositionTolerance)
 		{
@@ -1304,7 +1308,7 @@ namespace Chaos
 		FReal LinearStiffness = FPBDJointUtilities::GetLinearStiffness(SolverSettings, JointSettings);
 		const FVec3 CX = Xs[1] - Xs[0];
 
-		//UE_LOG(LogChaosJoint, VeryVerbose, TEXT("    PointKD Delta %f [Limit %f]"), CX.Size(), PositionTolerance);
+		UE_LOG(LogChaosJoint, VeryVerbose, TEXT("    PointKD Delta %f [Limit %f]"), CX.Size(), PositionTolerance);
 
 		if (CX.SizeSquared() > PositionTolerance * PositionTolerance)
 		{
diff --git a/Engine/Source/Runtime/Experimental/Chaos/Private/Chaos/Joint/PBDJointSolverGaussSeidel.ispc b/Engine/Source/Runtime/Experimental/Chaos/Private/Chaos/Joint/PBDJointSolverGaussSeidel.ispc
--- a/Engine/Source/Runtime/Experimental/Chaos/Private/Chaos/Joint/PBDJointSolverGaussSeidel.ispc
+++ b/Engine/Source/Runtime/Experimental/Chaos/Private/Chaos/Joint/PBDJointSolverGaussSeidel.ispc
@@ -110,14 +110,21 @@ static inline void UpdateDerivedState(uniform FJointSolverGaussSeidel * uniform
 {
 	M->Xs[Index] = M->Ps[Index] + VectorQuaternionRotateVector(M->Qs[Index], M->XLs[Index].Translation);
 	M->Rs[Index] = VectorQuaternionMultiply2(M->Qs[Index], M->XLs[Index].Rotation);
+
+	const uniform FVector4 Rs1Shortest = VectorQuaternionEnforceShortestArcWith(M->Rs[1], M->Rs[0]);
+	M->Rs[1] = Rs1Shortest;
 }
 
 static inline void UpdateDerivedState(uniform FJointSolverGaussSeidel * uniform M)
 {
 	M->Xs[0] = M->Ps[0] + VectorQuaternionRotateVector(M->Qs[0], M->XLs[0].Translation);
 	M->Xs[1] = M->Ps[1] + VectorQuaternionRotateVector(M->Qs[1], M->XLs[1].Translation);
-	M->Rs[0] = VectorQuaternionMultiply2(M->Qs[0], M->XLs[0].Rotation);
-	M->Rs[1] = VectorQuaternionMultiply2(M->Qs[1], M->XLs[1].Rotation);
+
+	const uniform FVector4 Rs0 = VectorQuaternionMultiply2(M->Qs[0], M->XLs[0].Rotation);
+	const uniform FVector4 Rs1 = VectorQuaternionMultiply2(M->Qs[1], M->XLs[1].Rotation);
+	const uniform FVector4 Rs1Shortest = VectorQuaternionEnforceShortestArcWith(Rs1, Rs0);
+	M->Rs[0] = Rs0;
+	M->Rs[1] = Rs1Shortest;
 }
 
 static inline void ApplyPositionDelta(
@@ -152,7 +159,10 @@ static inline void ApplyRotationDelta(
 
 	const uniform FVector4 QsNorm = VectorNormalizeQuaternion(Qs + DQ);
 
-	M->Qs[1] = QsNorm;
+	M->Qs[Index] = QsNorm;
+
+	const uniform FVector4 Qs1Shortest = VectorQuaternionEnforceShortestArcWith(M->Qs[1], M->Qs[0]);
+	M->Qs[1] = Qs1Shortest;
 }
 
 static inline void ApplyRotationDelta(
diff --git a/Engine/Source/Runtime/Experimental/Chaos/Private/Chaos/PBDJointConstraints.cpp b/Engine/Source/Runtime/Experimental/Chaos/Private/Chaos/PBDJointConstraints.cpp
--- a/Engine/Source/Runtime/Experimental/Chaos/Private/Chaos/PBDJointConstraints.cpp
+++ b/Engine/Source/Runtime/Experimental/Chaos/Private/Chaos/PBDJointConstraints.cpp
@@ -2,6 +2,7 @@
 #include "Chaos/PBDJointConstraints.h"
 #include "Chaos/ChaosDebugDraw.h"
 #include "Chaos/DebugDrawQueue.h"
+#include "Chaos/Joint/ChaosJointLog.h"
 #include "Chaos/Joint/PBDJointSolverGaussSeidel.h"
 #include "Chaos/Particle/ParticleUtilities.h"
 #include "Chaos/ParticleHandle.h"
@@ -12,7 +13,7 @@
 
 #include "HAL/IConsoleManager.h"
 
-//#pragma optimize("", off)
+//PRAGMA_DISABLE_OPTIMIZATION
 
 bool bChaos_Joint_EarlyOut_Enabled = true;
 FAutoConsoleVariableRef CVarChaosJointEarlyOutEnabled(TEXT("p.Chaos.Joint.EarlyOut"), bChaos_Joint_EarlyOut_Enabled, TEXT("Whether to iterating when joints report being solved"));
diff --git a/Engine/Source/Runtime/Experimental/Chaos/Private/ChaosLog.cpp b/Engine/Source/Runtime/Experimental/Chaos/Private/ChaosLog.cpp
--- a/Engine/Source/Runtime/Experimental/Chaos/Private/ChaosLog.cpp
+++ b/Engine/Source/Runtime/Experimental/Chaos/Private/ChaosLog.cpp
@@ -9,6 +9,4 @@ DEFINE_LOG_CATEGORY(LogChaosGeneral);
 DEFINE_LOG_CATEGORY(LogChaosThread);
 DEFINE_LOG_CATEGORY(LogChaosSimulation);
 DEFINE_LOG_CATEGORY(LogChaosDebug);
-DEFINE_LOG_CATEGORY(LogChaosJoint);
-DEFINE_LOG_CATEGORY(LogChaos6DJoint);
 DEFINE_LOG_CATEGORY(LogChaos);
diff --git a/Engine/Source/Runtime/Experimental/Chaos/Public/Chaos/Joint/ChaosJointLog.h b/Engine/Source/Runtime/Experimental/Chaos/Public/Chaos/Joint/ChaosJointLog.h
new file mode 100644
--- /dev/null
+++ b/Engine/Source/Runtime/Experimental/Chaos/Public/Chaos/Joint/ChaosJointLog.h
@@ -0,0 +1,11 @@
+// Copyright Epic Games, Inc. All Rights Reserved.
+
+#pragma once
+
+#include "CoreMinimal.h"
+
+#if UE_BUILD_SHIPPING || UE_BUILD_TEST
+CHAOS_API DECLARE_LOG_CATEGORY_EXTERN(LogChaosJoint, Log, Warning);
+#else
+CHAOS_API DECLARE_LOG_CATEGORY_EXTERN(LogChaosJoint, Log, All);
+#endif
diff --git a/Engine/Source/Runtime/Experimental/Chaos/Public/ChaosLog.h b/Engine/Source/Runtime/Experimental/Chaos/Public/ChaosLog.h
--- a/Engine/Source/Runtime/Experimental/Chaos/Public/ChaosLog.h
+++ b/Engine/Source/Runtime/Experimental/Chaos/Public/ChaosLog.h
@@ -8,7 +8,5 @@ CHAOS_API DECLARE_LOG_CATEGORY_EXTERN(LogChaosGeneral, Log, All);
 CHAOS_API DECLARE_LOG_CATEGORY_EXTERN(LogChaosThread, Log, All);
 CHAOS_API DECLARE_LOG_CATEGORY_EXTERN(LogChaosSimulation, Log, All);
 CHAOS_API DECLARE_LOG_CATEGORY_EXTERN(LogChaosDebug, Log, All);
-CHAOS_API DECLARE_LOG_CATEGORY_EXTERN(LogChaosJoint, Log, All);
-CHAOS_API DECLARE_LOG_CATEGORY_EXTERN(LogChaos6DJoint, Log, All);
 
 CHAOS_API DECLARE_LOG_CATEGORY_EXTERN(LogChaos, Log, All);
\ No newline at end of file
