Heavily optimize level streaming (Chaos AABBTree) in Reverb (~15x faster - Forgotten City sublevel make visible ~5min -> ~20sec).

The Chaos AABBTree implementation has a heuristic that reoptimizes the tree every X elements, where X defaults to 10k. The reoptimize tree pass is very expensive, as it does a ton of heap allocations and copies (due to resizing). In just the Forgotten City sub level group, ReoptimizeTree() is called ~118 times, as there are 1.2mil elements. Since the size of FElement is 64bytes, the heuristic is now set to 128MiB of pending data size, which handles 2mil elements before any reoptimize occurs.

In addition to the heuristic, this change also fixes a high number of heap allocations and copies from resize. For example, the original code, per leaf, would do a heap allocation in GatherElements() for MGlobalPayloads.Num(), but then it would reserve that num + MDirtyElements.Num(), which would result in a new larger allocation and then a copy of the old contents into the new allocation. Since OutElements is passed in by reference, I also added a global reserve sum that I compute and size that entire array for the entire collection of leaves, as well as a few other similar optimizations.

#[fyi] brian.karis, michal.valient, michael.lentine, devon.penney

diff --git a/Engine/Source/Runtime/Experimental/Chaos/Public/Chaos/AABBTree.h b/Engine/Source/Runtime/Experimental/Chaos/Public/Chaos/AABBTree.h
--- a/Engine/Source/Runtime/Experimental/Chaos/Public/Chaos/AABBTree.h
+++ b/Engine/Source/Runtime/Experimental/Chaos/Public/Chaos/AABBTree.h
@@ -129,6 +129,12 @@ struct TAABBTreeLeafArray : public TBoundsWrapperHelper<TPayloadType, T, bComput
 		OutElements.Append(Elems);
 	}
 
+	SIZE_T GetReserveCount() const
+	{
+		// Optimize for fewer memory allocations.
+		return Elems.Num();
+	}
+
 	template <typename TSQVisitor, typename TQueryFastData>
 	bool RaycastFast(const TVector<T,3>& Start, TQueryFastData& QueryFastData, TSQVisitor& Visitor) const
 	{
@@ -203,7 +209,6 @@ struct TAABBTreeLeafArray : public TBoundsWrapperHelper<TPayloadType, T, bComput
 		Ar << Elems;
 	}
 
-
 	TArray<TPayloadBoundsElement<TPayloadType, T>> Elems;
 	TAABB<T,3> Bounds;
 };
@@ -625,10 +630,19 @@ private:
 	void ReoptimizeTree()
 	{
 		TArray<FElement> AllElements;
+
+		SIZE_T ReserveCount = DirtyElements.Num() + GlobalPayloads.Num();
+		for (const auto& Leaf : Leaves)
+		{
+			ReserveCount += Leaf.GetReserveCount();
+		}
+
+		AllElements.Reserve(ReserveCount);
+
 		AllElements.Append(DirtyElements);
 		AllElements.Append(GlobalPayloads);
 
-		for(auto& Leaf : Leaves)
+		for (auto& Leaf : Leaves)
 		{
 			Leaf.GatherElements(AllElements);
 		}
diff --git a/Engine/Source/Runtime/Experimental/Chaos/Public/Chaos/BoundingVolume.h b/Engine/Source/Runtime/Experimental/Chaos/Public/Chaos/BoundingVolume.h
--- a/Engine/Source/Runtime/Experimental/Chaos/Public/Chaos/BoundingVolume.h
+++ b/Engine/Source/Runtime/Experimental/Chaos/Public/Chaos/BoundingVolume.h
@@ -346,9 +346,10 @@ public:
 
 	void GatherElements(TArray<TPayloadBoundsElement<TPayloadType, T>>& OutElements)
 	{
+		OutElements.Reserve(GetReserveCount());
 		OutElements.Append(MGlobalPayloads);
-		OutElements.Reserve(OutElements.Num() + MDirtyElements.Num());
-		for(const FCellElement& Elem : MDirtyElements)
+		
+		for (const FCellElement& Elem : MDirtyElements)
 		{
 			OutElements.Add(TPayloadBoundsElement<TPayloadType,T>{Elem.Payload,Elem.Bounds});
 		}
@@ -372,7 +373,14 @@ public:
 				}
 			}
 		}
+	}
 
+	SIZE_T GetReserveCount() const
+	{
+		// Optimize for fewer memory allocations.
+		const auto& Counts = MGrid.Counts(); 
+		const SIZE_T GridCount = Counts[0] * Counts[1] * Counts[2] * MElements.Num();
+		return MGlobalPayloads.Num() + MDirtyElements.Num() + GridCount;
 	}
 
 	TAABB<T, d> GetBounds() const
