Chaos - RBAN fix joints at degenerate angle (180 degree swing) if one swing axis is locked

diff --git a/Engine/Source/Runtime/Experimental/Chaos/Private/Chaos/Joint/PBDJointSolverGaussSeidel.cpp b/Engine/Source/Runtime/Experimental/Chaos/Private/Chaos/Joint/PBDJointSolverGaussSeidel.cpp
--- a/Engine/Source/Runtime/Experimental/Chaos/Private/Chaos/Joint/PBDJointSolverGaussSeidel.cpp
+++ b/Engine/Source/Runtime/Experimental/Chaos/Private/Chaos/Joint/PBDJointSolverGaussSeidel.cpp
@@ -20,6 +20,9 @@ bool bChaos_Joint_ISPC_Enabled = true;
 FAutoConsoleVariableRef CVarChaosJointISPCEnabled(TEXT("p.Chaos.Joint.ISPC"), bChaos_Joint_ISPC_Enabled, TEXT("Whether to use ISPC optimizations in the Joint Solver"));
 #endif
 
+float Chaos_Joint_DegenerateRotationLimit = -0.998f;	// Cos(176deg)
+FAutoConsoleVariableRef CVarChaosJointDegenerateRotationLimit(TEXT("p.Chaos.Joint.DegenerateRotationLimit"), Chaos_Joint_DegenerateRotationLimit, TEXT("Cosine of the swing angle that is considered degerenerate (default Cos(176deg))"));
+
 namespace Chaos
 {
 	void GetSphericalAxisDelta(
@@ -224,9 +227,25 @@ namespace Chaos
 	{
 		FJointSolverResult NetResult;
 
-		NetResult += ApplyPositionConstraints(Dt, SolverSettings, JointSettings);
+		bool bHasPositionConstraints =
+			(JointSettings.LinearMotionTypes[0] != EJointMotionType::Free)
+			|| (JointSettings.LinearMotionTypes[1] != EJointMotionType::Free)
+			|| (JointSettings.LinearMotionTypes[2] != EJointMotionType::Free);
+
+		bool bHasRotationConstraints =
+			(JointSettings.AngularMotionTypes[0] != EJointMotionType::Free)
+			|| (JointSettings.AngularMotionTypes[1] != EJointMotionType::Free)
+			|| (JointSettings.AngularMotionTypes[2] != EJointMotionType::Free);
 
-		NetResult += ApplyRotationConstraints(Dt, SolverSettings, JointSettings);
+		if (bHasPositionConstraints)
+		{
+			NetResult += ApplyPositionConstraints(Dt, SolverSettings, JointSettings);
+		}
+
+		if (bHasRotationConstraints)
+		{
+			NetResult += ApplyRotationConstraints(Dt, SolverSettings, JointSettings);
+		}
 
 		return NetResult;
 	}
@@ -239,9 +258,31 @@ namespace Chaos
 	{
 		FJointSolverResult NetResult;
 
-		NetResult += ApplyPositionDrives(Dt, SolverSettings, JointSettings);
+		bool bHasPositionDrives = 
+			JointSettings.bLinearPositionDriveEnabled[0]
+			|| JointSettings.bLinearPositionDriveEnabled[1]
+			|| JointSettings.bLinearPositionDriveEnabled[2]
+			|| JointSettings.bLinearVelocityDriveEnabled[0]
+			|| JointSettings.bLinearVelocityDriveEnabled[1]
+			|| JointSettings.bLinearVelocityDriveEnabled[2];
+
+		bool bHasRotationDrives = 
+			JointSettings.bAngularTwistPositionDriveEnabled
+			|| JointSettings.bAngularTwistVelocityDriveEnabled
+			|| JointSettings.bAngularSwingPositionDriveEnabled
+			|| JointSettings.bAngularSwingVelocityDriveEnabled
+			|| JointSettings.bAngularSLerpPositionDriveEnabled
+			|| JointSettings.bAngularSLerpVelocityDriveEnabled;
+
+		if (bHasPositionDrives)
+		{
+			NetResult += ApplyPositionDrives(Dt, SolverSettings, JointSettings);
+		}
 
-		NetResult += ApplyRotationDrives(Dt, SolverSettings, JointSettings);
+		if (bHasRotationDrives)
+		{
+			NetResult += ApplyRotationDrives(Dt, SolverSettings, JointSettings);
+		}
 
 		return NetResult;
 	}
@@ -276,8 +317,19 @@ namespace Chaos
 		bool bTwistSoft = JointSettings.bSoftTwistLimitsEnabled;
 		bool bSwingSoft = JointSettings.bSoftSwingLimitsEnabled;
 
+		// If the twist axes are opposing, we cannot decompose the orientation into swing and twist angles, so just give up
+		const FVec3 Twist0 = Rs[0] * FJointConstants::TwistAxis();
+		const FVec3 Twist1 = Rs[1] * FJointConstants::TwistAxis();
+		const FReal Twist01Dot = FVec3::DotProduct(Twist0, Twist1);
+		const bool bDegenerate = (Twist01Dot < Chaos_Joint_DegenerateRotationLimit);
+		if (bDegenerate)
+		{
+			UE_LOG(LogChaosJoint, VeryVerbose, TEXT(" Degenerate rotation at Swing %f deg"), FMath::RadiansToDegrees(FMath::Acos(Twist01Dot)));
+		}
+
 		// Apply twist constraint
-		if (SolverSettings.bEnableTwistLimits)
+		// NOTE: Cannot calculate twist angle at 180degree swing
+		if (SolverSettings.bEnableTwistLimits && !bDegenerate)
 		{
 			if (TwistMotion == EJointMotionType::Limited)
 			{
@@ -290,6 +342,7 @@ namespace Chaos
 		}
 
 		// Apply swing constraints
+		// NOTE: Cannot separate swing angles at 180degree swing (but we can still apply locks)
 		if (SolverSettings.bEnableSwingLimits)
 		{
 			if ((Swing1Motion == EJointMotionType::Limited) && (Swing2Motion == EJointMotionType::Limited))
@@ -298,34 +351,46 @@ namespace Chaos
 			}
 			else if ((Swing1Motion == EJointMotionType::Limited) && (Swing2Motion == EJointMotionType::Locked))
 			{
-				NetResult += ApplyDualConeSwingConstraint(Dt, SolverSettings, JointSettings, EJointAngularConstraintIndex::Swing2, false);
-				NetResult += ApplySwingConstraint(Dt, SolverSettings, JointSettings, EJointAngularConstraintIndex::Swing1, bSwingSoft);
+				NetResult += ApplyLockedSwingConstraint(Dt, SolverSettings, JointSettings, EJointAngularConstraintIndex::Swing2, false);
+				if (!bDegenerate)
+				{
+					NetResult += ApplySwingConstraint(Dt, SolverSettings, JointSettings, EJointAngularConstraintIndex::Swing1, bSwingSoft);
+				}
 			}
 			else if ((Swing1Motion == EJointMotionType::Limited) && (Swing2Motion == EJointMotionType::Free))
 			{
-				NetResult += ApplyDualConeSwingConstraint(Dt, SolverSettings, JointSettings, EJointAngularConstraintIndex::Swing1, bSwingSoft);
+				if (!bDegenerate)
+				{
+					NetResult += ApplyDualConeSwingConstraint(Dt, SolverSettings, JointSettings, EJointAngularConstraintIndex::Swing1, bSwingSoft);
+				}
 			}
 			else if ((Swing1Motion == EJointMotionType::Locked) && (Swing2Motion == EJointMotionType::Limited))
 			{
-				NetResult += ApplyDualConeSwingConstraint(Dt, SolverSettings, JointSettings, EJointAngularConstraintIndex::Swing1, false);
-				NetResult += ApplySwingConstraint(Dt, SolverSettings, JointSettings, EJointAngularConstraintIndex::Swing2, bSwingSoft);
+				NetResult += ApplyLockedSwingConstraint(Dt, SolverSettings, JointSettings, EJointAngularConstraintIndex::Swing1, false);
+				if (!bDegenerate)
+				{
+					NetResult += ApplySwingConstraint(Dt, SolverSettings, JointSettings, EJointAngularConstraintIndex::Swing2, bSwingSoft);
+				}
 			}
 			else if ((Swing1Motion == EJointMotionType::Locked) && (Swing2Motion == EJointMotionType::Locked))
 			{
-				NetResult += ApplySwingConstraint(Dt, SolverSettings, JointSettings, EJointAngularConstraintIndex::Swing1, false);
-				NetResult += ApplySwingConstraint(Dt, SolverSettings, JointSettings, EJointAngularConstraintIndex::Swing2, false);
+				NetResult += ApplyLockedSwingConstraint(Dt, SolverSettings, JointSettings, EJointAngularConstraintIndex::Swing1, false);
+				NetResult += ApplyLockedSwingConstraint(Dt, SolverSettings, JointSettings, EJointAngularConstraintIndex::Swing2, false);
 			}
 			else if ((Swing1Motion == EJointMotionType::Locked) && (Swing2Motion == EJointMotionType::Free))
 			{
-				NetResult += ApplyDualConeSwingConstraint(Dt, SolverSettings, JointSettings, EJointAngularConstraintIndex::Swing1, false);
+				NetResult += ApplyLockedSwingConstraint(Dt, SolverSettings, JointSettings, EJointAngularConstraintIndex::Swing1, false);
 			}
 			else if ((Swing1Motion == EJointMotionType::Free) && (Swing2Motion == EJointMotionType::Limited))
 			{
-				NetResult += ApplyDualConeSwingConstraint(Dt, SolverSettings, JointSettings, EJointAngularConstraintIndex::Swing2, bSwingSoft);
+				if (!bDegenerate)
+				{
+					NetResult += ApplyDualConeSwingConstraint(Dt, SolverSettings, JointSettings, EJointAngularConstraintIndex::Swing2, bSwingSoft);
+				}
 			}
 			else if ((Swing1Motion == EJointMotionType::Free) && (Swing2Motion == EJointMotionType::Locked))
 			{
-				NetResult += ApplyDualConeSwingConstraint(Dt, SolverSettings, JointSettings, EJointAngularConstraintIndex::Swing2, false);
+				NetResult += ApplyLockedSwingConstraint(Dt, SolverSettings, JointSettings, EJointAngularConstraintIndex::Swing2, false);
 			}
 			else if ((Swing1Motion == EJointMotionType::Free) && (Swing2Motion == EJointMotionType::Free))
 			{
@@ -588,7 +653,7 @@ namespace Chaos
 		const FReal Stiffness,
 		const FVec3& DP)
 	{
-		UE_LOG(LogChaosJoint, VeryVerbose, TEXT("      Apply DP%d %f %f %f"), BodyIndex, DP.X, DP.Y, DP.Z);
+		//UE_LOG(LogChaosJoint, VeryVerbose, TEXT("      Apply DP%d %f %f %f"), BodyIndex, DP.X, DP.Y, DP.Z);
 
 		Ps[BodyIndex] += Stiffness * DP;
 	}
@@ -599,8 +664,8 @@ namespace Chaos
 		const FVec3& DP0,
 		const FVec3& DP1)
 	{
-		UE_LOG(LogChaosJoint, VeryVerbose, TEXT("      Apply DP%d %f %f %f"), 0, DP0.X, DP0.Y, DP0.Z);
-		UE_LOG(LogChaosJoint, VeryVerbose, TEXT("      Apply DP%d %f %f %f"), 1, DP1.X, DP1.Y, DP1.Z);
+		//UE_LOG(LogChaosJoint, VeryVerbose, TEXT("      Apply DP%d %f %f %f"), 0, DP0.X, DP0.Y, DP0.Z);
+		//UE_LOG(LogChaosJoint, VeryVerbose, TEXT("      Apply DP%d %f %f %f"), 1, DP1.X, DP1.Y, DP1.Z);
 
 		Ps[0] += Stiffness * DP0;
 		Ps[1] += Stiffness * DP1;
@@ -612,7 +677,7 @@ namespace Chaos
 		const FReal Stiffness,
 		const FVec3& DR)
 	{
-		UE_LOG(LogChaosJoint, VeryVerbose, TEXT("      Apply DR%d %f %f %f"), BodyIndex, DR.X, DR.Y, DR.Z);
+		//UE_LOG(LogChaosJoint, VeryVerbose, TEXT("      Apply DR%d %f %f %f"), BodyIndex, DR.X, DR.Y, DR.Z);
 
 		const FRotation3 DQ = (FRotation3::FromElements(Stiffness * DR, 0) * Qs[BodyIndex]) * (FReal)0.5;
 		Qs[BodyIndex] = (Qs[BodyIndex] + DQ).GetNormalized();
@@ -625,8 +690,8 @@ namespace Chaos
 		const FVec3& DR0,
 		const FVec3& DR1)
 	{
-		UE_LOG(LogChaosJoint, VeryVerbose, TEXT("      Apply DR%d %f %f %f"), 0, DR0.X, DR0.Y, DR0.Z);
-		UE_LOG(LogChaosJoint, VeryVerbose, TEXT("      Apply DR%d %f %f %f"), 1, DR1.X, DR1.Y, DR1.Z);
+		//UE_LOG(LogChaosJoint, VeryVerbose, TEXT("      Apply DR%d %f %f %f"), 0, DR0.X, DR0.Y, DR0.Z);
+		//UE_LOG(LogChaosJoint, VeryVerbose, TEXT("      Apply DR%d %f %f %f"), 1, DR1.X, DR1.Y, DR1.Z);
 
 		const FRotation3 DQ0 = (FRotation3::FromElements(Stiffness * DR0, 0) * Qs[0]) * (FReal)0.5;
 		const FRotation3 DQ1 = (FRotation3::FromElements(Stiffness * DR1, 0) * Qs[1]) * (FReal)0.5;
@@ -641,8 +706,8 @@ namespace Chaos
 		const FVec3& DP,
 		const FVec3& DR)
 	{
-		UE_LOG(LogChaosJoint, VeryVerbose, TEXT("      Apply DP%d %f %f %f"), BodyIndex, DP.X, DP.Y, DP.Z);
-		UE_LOG(LogChaosJoint, VeryVerbose, TEXT("      Apply DR%d %f %f %f"), BodyIndex, DR.X, DR.Y, DR.Z);
+		//UE_LOG(LogChaosJoint, VeryVerbose, TEXT("      Apply DP%d %f %f %f"), BodyIndex, DP.X, DP.Y, DP.Z);
+		//UE_LOG(LogChaosJoint, VeryVerbose, TEXT("      Apply DR%d %f %f %f"), BodyIndex, DR.X, DR.Y, DR.Z);
 
 		Ps[BodyIndex] += Stiffness * DP;
 		const FRotation3 DQ = (FRotation3::FromElements(Stiffness * DR, 0) * Qs[BodyIndex]) * (FReal)0.5;
@@ -1095,6 +1160,43 @@ namespace Chaos
 	}
 
 
+	FJointSolverResult FJointSolverGaussSeidel::ApplyLockedSwingConstraint(
+		const FReal Dt,
+		const FPBDJointSolverSettings& SolverSettings,
+		const FPBDJointSettings& JointSettings,
+		const EJointAngularConstraintIndex SwingConstraintIndex,
+		const bool bUseSoftLimit)
+	{
+		// NOTE: SwingAxis is not normalized in this mode. It has length Sin(SwingAngle).
+		// Likewise, the SwingAngle is actually Sin(SwingAngle)
+		FVec3 SwingAxis;
+		FReal SwingAngle;
+		FPBDJointUtilities::GetLockedSwingAxisAngle(Rs[0], Rs[1], SwingConstraintIndex, SwingAxis, SwingAngle);
+
+		UE_LOG(LogChaosJoint, VeryVerbose, TEXT("    LockedSwing%d Angle %f [Tolerance %f]"), (SwingConstraintIndex == EJointAngularConstraintIndex::Swing1) ? 1 : 2, FMath::RadiansToDegrees(SwingAngle), FMath::RadiansToDegrees(AngleTolerance));
+
+		// Apply swing correction
+		FReal DSwingAngle = SwingAngle;
+		if (FMath::Abs(DSwingAngle) > AngleTolerance)
+		{
+			if (bUseSoftLimit)
+			{
+				const FReal SoftSwingStiffness = FPBDJointUtilities::GetSoftSwingStiffness(SolverSettings, JointSettings);
+				const FReal SoftSwingDamping = FPBDJointUtilities::GetSoftSwingDamping(SolverSettings, JointSettings);
+				const bool bAccelerationMode = FPBDJointUtilities::GetAngularSoftAccelerationMode(SolverSettings, JointSettings);
+				ApplyRotationConstraintSoft(Dt, SoftSwingStiffness, SoftSwingDamping, bAccelerationMode, SwingAxis, DSwingAngle, SwingSoftLambda);
+			}
+			else
+			{
+				const FReal SwingStiffness = FPBDJointUtilities::GetSwingStiffness(SolverSettings, JointSettings);
+				ApplyRotationConstraint(SwingStiffness, SwingAxis, DSwingAngle);
+			}
+			return FJointSolverResult::MakeActive();
+		}
+		return FJointSolverResult::MakeSolved();
+	}
+
+
 	FJointSolverResult FJointSolverGaussSeidel::ApplyDualConeSwingConstraint(
 		const FReal Dt,
 		const FPBDJointSolverSettings& SolverSettings,
diff --git a/Engine/Source/Runtime/Experimental/Chaos/Private/Chaos/PBDJointConstraintUtilities.cpp b/Engine/Source/Runtime/Experimental/Chaos/Private/Chaos/PBDJointConstraintUtilities.cpp
--- a/Engine/Source/Runtime/Experimental/Chaos/Private/Chaos/PBDJointConstraintUtilities.cpp
+++ b/Engine/Source/Runtime/Experimental/Chaos/Private/Chaos/PBDJointConstraintUtilities.cpp
@@ -73,6 +73,22 @@ namespace Chaos
 	}
 
 
+	void FPBDJointUtilities::GetLockedSwingAxisAngle(
+		const FRotation3& R0,
+		const FRotation3& R1,
+		const EJointAngularConstraintIndex SwingConstraintIndex,
+		FVec3& Axis,
+		FReal& Angle)
+	{
+		// NOTE: this differs from GetDualConeSwingAxisAngle in that it returns an axis with length Sin(SwingAngle)
+		// and an Angle that is actually Sin(SwingAngle). This allows it to be used when we get closer to degenerate twist angles.
+		FVec3 Twist1 = R1 * FJointConstants::TwistAxis();
+		FVec3 Swing0 = R0 * FJointConstants::OtherSwingAxis(SwingConstraintIndex);
+		Axis = FVec3::CrossProduct(Swing0, Twist1);
+		Angle = -FVec3::DotProduct(Swing0, Twist1);
+	}
+
+
 	void FPBDJointUtilities::GetDualConeSwingAxisAngle(
 		const FRotation3& R0,
 		const FRotation3& R1,
diff --git a/Engine/Source/Runtime/Experimental/Chaos/Private/Chaos/PBDJointConstraints.cpp b/Engine/Source/Runtime/Experimental/Chaos/Private/Chaos/PBDJointConstraints.cpp
--- a/Engine/Source/Runtime/Experimental/Chaos/Private/Chaos/PBDJointConstraints.cpp
+++ b/Engine/Source/Runtime/Experimental/Chaos/Private/Chaos/PBDJointConstraints.cpp
@@ -509,7 +509,7 @@ namespace Chaos
 			}
 		}
 
-		//UE_LOG(LogChaosJoint, Warning, TEXT("Apply %d / %d Active %d / %d"), It, NumIts, NetResult.GetNumActive(), NetResult.GetNumActive() + NetResult.GetNumSolved());
+		UE_LOG(LogChaosJoint, VeryVerbose, TEXT("Apply %d / %d Active %d / %d"), It, NumIts, NetResult.GetNumActive(), NetResult.GetNumActive() + NetResult.GetNumSolved());
 
 		if (PostApplyCallback != nullptr)
 		{
@@ -533,7 +533,7 @@ namespace Chaos
 			}
 		}
 
-		//UE_LOG(LogChaosJoint, Warning, TEXT("Push  %d / %d Active %d / %d"), It, NumIts, NetResult.GetNumActive(), NetResult.GetNumActive() + NetResult.GetNumSolved());
+		UE_LOG(LogChaosJoint, VeryVerbose, TEXT("PushOut  %d / %d Active %d / %d"), It, NumIts, NetResult.GetNumActive(), NetResult.GetNumActive() + NetResult.GetNumSolved());
 
 		if (PostProjectCallback != nullptr)
 		{
diff --git a/Engine/Source/Runtime/Experimental/Chaos/Public/Chaos/Joint/PBDJointSolverGaussSeidel.h b/Engine/Source/Runtime/Experimental/Chaos/Public/Chaos/Joint/PBDJointSolverGaussSeidel.h
--- a/Engine/Source/Runtime/Experimental/Chaos/Public/Chaos/Joint/PBDJointSolverGaussSeidel.h
+++ b/Engine/Source/Runtime/Experimental/Chaos/Public/Chaos/Joint/PBDJointSolverGaussSeidel.h
@@ -279,7 +279,15 @@ namespace Chaos
 			const FPBDJointSolverSettings& SolverSettings,
 			const FPBDJointSettings& JointSettings);
 
-		// One Swing axis is free, and the other limited or locked. This applies the lock/limit: Body1 Twist axis is confined to space between two cones.
+		// One Swing axis is free, and the other locked. This applies the lock: Body1 Twist axis is confined to a plane.
+		FJointSolverResult ApplyLockedSwingConstraint(
+			const FReal Dt,
+			const FPBDJointSolverSettings& SolverSettings,
+			const FPBDJointSettings& JointSettings,
+			const EJointAngularConstraintIndex SwingConstraintIndex,
+			const bool bUseSoftLimit);
+
+		// One Swing axis is free, and the other limited. This applies the limit: Body1 Twist axis is confined to space between two cones.
 		FJointSolverResult ApplyDualConeSwingConstraint(
 			const FReal Dt,
 			const FPBDJointSolverSettings& SolverSettings,
diff --git a/Engine/Source/Runtime/Experimental/Chaos/Public/Chaos/PBDJointConstraintUtilities.h b/Engine/Source/Runtime/Experimental/Chaos/Public/Chaos/PBDJointConstraintUtilities.h
--- a/Engine/Source/Runtime/Experimental/Chaos/Public/Chaos/PBDJointConstraintUtilities.h
+++ b/Engine/Source/Runtime/Experimental/Chaos/Public/Chaos/PBDJointConstraintUtilities.h
@@ -41,6 +41,13 @@ namespace Chaos
 			FVec3& AxisLocal,
 			FReal& Angle);
 
+		static void GetLockedSwingAxisAngle(
+			const FRotation3& R0,
+			const FRotation3& R1,
+			const EJointAngularConstraintIndex SwingConstraintIndex,
+			FVec3& Axis,
+			FReal& Angle);
+
 		static void GetDualConeSwingAxisAngle(
 			const FRotation3& R0,
 			const FRotation3& R1,
