* Increase NumIterations and NumPushoutIterations
* Fix iteration early outs for apply and apply pushout

diff --git a/Engine/Source/Runtime/Experimental/Chaos/Private/Chaos/Evolution/PBDMinEvolution.cpp b/Engine/Source/Runtime/Experimental/Chaos/Private/Chaos/Evolution/PBDMinEvolution.cpp
--- a/Engine/Source/Runtime/Experimental/Chaos/Private/Chaos/Evolution/PBDMinEvolution.cpp
+++ b/Engine/Source/Runtime/Experimental/Chaos/Private/Chaos/Evolution/PBDMinEvolution.cpp
@@ -299,9 +299,15 @@ namespace Chaos
 		// @todo(ccaulfield): track whether we are sufficiently solved and can early-out
 		for (int32 i = 0; i < NumApplyIterations; ++i)
 		{
+			bool bNeedsAnotherIteration = false;
 			for (FSimpleConstraintRule* ConstraintRule : PrioritizedConstraintRules)
 			{
-				ConstraintRule->ApplyConstraints(Dt, i, NumApplyIterations);
+				bNeedsAnotherIteration |= ConstraintRule->ApplyConstraints(Dt, i, NumApplyIterations);
+			}
+
+			if (!bNeedsAnotherIteration)
+			{
+				break;
 			}
 		}
 	}
diff --git a/Engine/Source/Runtime/Experimental/Chaos/Private/Chaos/PBDCollisionConstraints.cpp b/Engine/Source/Runtime/Experimental/Chaos/Private/Chaos/PBDCollisionConstraints.cpp
--- a/Engine/Source/Runtime/Experimental/Chaos/Private/Chaos/PBDCollisionConstraints.cpp
+++ b/Engine/Source/Runtime/Experimental/Chaos/Private/Chaos/PBDCollisionConstraints.cpp
@@ -414,14 +414,15 @@ namespace Chaos
 	}
 
 	template<typename T, int d>
-	void TPBDCollisionConstraints<T, d>::Apply(const T Dt, const int32 Iterations, const int32 NumIterations)
+	bool TPBDCollisionConstraints<T, d>::Apply(const T Dt, const int32 Iterations, const int32 NumIterations)
 	{
 		SCOPE_CYCLE_COUNTER(STAT_Collisions_Apply);
 
+		bool bNeedsAnotherIteration = false;
 		if (MApplyPairIterations > 0)
 		{
 			const Collisions::TContactParticleParameters<T> ParticleParameters = { MCullDistance, MShapePadding, &MCollided };
-			const Collisions::TContactIterationParameters<T> IterationParameters = { Dt, Iterations, NumIterations, MApplyPairIterations, ApplyType, nullptr };
+			const Collisions::TContactIterationParameters<T> IterationParameters = { Dt, Iterations, NumIterations, MApplyPairIterations, ApplyType, &bNeedsAnotherIteration };
 
 			for (FPointContactConstraint& Contact : PointConstraints)
 			{
@@ -438,6 +439,8 @@ namespace Chaos
 		{
 			PostApplyCallback(Dt, Handles);
 		}
+
+		return bNeedsAnotherIteration;
 	}
 
 	template<typename T, int d>
@@ -473,20 +476,28 @@ namespace Chaos
 
 
 	template<typename T, int d>
-	void TPBDCollisionConstraints<T, d>::Apply(const T Dt, const TArray<FConstraintContainerHandle*>& InConstraintHandles, const int32 Iterations, const int32 NumIterations)
+	bool TPBDCollisionConstraints<T, d>::Apply(const T Dt, const TArray<FConstraintContainerHandle*>& InConstraintHandles, const int32 Iterations, const int32 NumIterations)
 	{
 		SCOPE_CYCLE_COUNTER(STAT_Collisions_Apply);
 
+		TAtomic<bool> bNeedsAnotherIterationAtomic;
+		bNeedsAnotherIterationAtomic.Store(false);
 		if (MApplyPairIterations > 0)
 		{
 			PhysicsParallelFor(InConstraintHandles.Num(), [&](int32 ConstraintHandleIndex) {
 				FConstraintContainerHandle* ConstraintHandle = InConstraintHandles[ConstraintHandleIndex];
 				check(ConstraintHandle != nullptr);
 
+				bool bNeedsAnotherIteration = false;
 				Collisions::TContactParticleParameters<T> ParticleParameters = { MCullDistance, MShapePadding, &MCollided };
-				Collisions::TContactIterationParameters<T> IterationParameters = { Dt, Iterations, NumIterations, MApplyPairIterations, ApplyType, nullptr };
+				Collisions::TContactIterationParameters<T> IterationParameters = { Dt, Iterations, NumIterations, MApplyPairIterations, ApplyType, &bNeedsAnotherIteration };
 				Collisions::Apply(ConstraintHandle->GetContact(), IterationParameters, ParticleParameters);
 
+				if (bNeedsAnotherIteration)
+				{
+					bNeedsAnotherIterationAtomic.Store(true);
+				}
+
 			}, bDisableCollisionParallelFor);
 		}
 
@@ -494,6 +505,8 @@ namespace Chaos
 		{
 			PostApplyCallback(Dt, InConstraintHandles);
 		}
+
+		return bNeedsAnotherIterationAtomic.Load();
 	}
 
 
diff --git a/Engine/Source/Runtime/Experimental/Chaos/Private/Chaos/PBDCollisionConstraintsContact.cpp b/Engine/Source/Runtime/Experimental/Chaos/Private/Chaos/PBDCollisionConstraintsContact.cpp
--- a/Engine/Source/Runtime/Experimental/Chaos/Private/Chaos/PBDCollisionConstraintsContact.cpp
+++ b/Engine/Source/Runtime/Experimental/Chaos/Private/Chaos/PBDCollisionConstraintsContact.cpp
@@ -74,6 +74,8 @@ namespace Chaos
 
 			if (RelativeNormalVelocity < 0) // ignore separating constraints
 			{
+				*IterationParameters.NeedsAnotherIteration = true;
+
 				FMatrix33 WorldSpaceInvI1 = bIsRigidDynamic0 ? Utilities::ComputeWorldSpaceInertia(Q0, PBDRigid0->InvI()) : FMatrix33(0);
 				FMatrix33 WorldSpaceInvI2 = bIsRigidDynamic1 ? Utilities::ComputeWorldSpaceInertia(Q1, PBDRigid1->InvI()) : FMatrix33(0);
 				FMatrix33 Factor =
@@ -240,6 +242,8 @@ namespace Chaos
 
 			if (Contact.Phi < 0)
 			{
+				*IterationParameters.NeedsAnotherIteration = true;
+			
 				FReal InvM0 = bIsRigidDynamic0 ? PBDRigid0->InvM() : 0.0f;
 				FReal InvM1 = bIsRigidDynamic1 ? PBDRigid1->InvM() : 0.0f;
 				FMatrix33 InvI0 = bIsRigidDynamic0 ? Utilities::ComputeWorldSpaceInertia(Q0, PBDRigid0->InvI()) : FMatrix33(0);
@@ -428,11 +432,6 @@ namespace Chaos
 				return AccumulatedImpulse;
 			}
 
-			if (IterationParameters.NeedsAnotherIteration)
-			{
-				*IterationParameters.NeedsAnotherIteration = true;
-			}
-
 			FMatrix33 WorldSpaceInvI1 = bIsRigidDynamic0 ? Utilities::ComputeWorldSpaceInertia(Q0, PBDRigid0->InvI()) : FMatrix33(0);
 			FMatrix33 WorldSpaceInvI2 = bIsRigidDynamic1 ? Utilities::ComputeWorldSpaceInertia(Q1, PBDRigid1->InvI()) : FMatrix33(0);
 			FVec3 VectorToPoint1 = Contact.Location - P0;
@@ -450,6 +449,8 @@ namespace Chaos
 			const FReal RelativeVelocityDotNormal = FVec3::DotProduct(RelativeVelocity, Contact.Normal);
 			if (RelativeVelocityDotNormal < 0)
 			{
+				*IterationParameters.NeedsAnotherIteration = true;
+			
 				const FVec3 ImpulseNumerator = -FVec3::DotProduct(RelativeVelocity, Contact.Normal) * Contact.Normal * ScalingFactor;
 				const FVec3 FactorContactNormal = Factor * Contact.Normal;
 				FReal ImpulseDenominator = FVec3::DotProduct(Contact.Normal, FactorContactNormal);
diff --git a/Engine/Source/Runtime/Experimental/Chaos/Private/Chaos/PBDJointConstraints.cpp b/Engine/Source/Runtime/Experimental/Chaos/Private/Chaos/PBDJointConstraints.cpp
--- a/Engine/Source/Runtime/Experimental/Chaos/Private/Chaos/PBDJointConstraints.cpp
+++ b/Engine/Source/Runtime/Experimental/Chaos/Private/Chaos/PBDJointConstraints.cpp
@@ -491,7 +491,7 @@ namespace Chaos
 	//
 	//////////////////////////////////////////////////////////////////////////
 
-	void FPBDJointConstraints::Apply(const FReal Dt, const int32 It, const int32 NumIts)
+	bool FPBDJointConstraints::Apply(const FReal Dt, const int32 It, const int32 NumIts)
 	{
 		SCOPE_CYCLE_COUNTER(STAT_Joints_Apply);
 
@@ -515,6 +515,9 @@ namespace Chaos
 		{
 			PostApplyCallback(Dt, Handles);
 		}
+
+		// TODO: Return true only if more iteration are needed
+		return true;
 	}
 
 	bool FPBDJointConstraints::ApplyPushOut(const FReal Dt, const int32 It, const int32 NumIts)
@@ -555,7 +558,7 @@ namespace Chaos
 	//
 	//////////////////////////////////////////////////////////////////////////
 
-	void FPBDJointConstraints::Apply(const FReal Dt, const TArray<FConstraintContainerHandle*>& InConstraintHandles, const int32 It, const int32 NumIts)
+	bool FPBDJointConstraints::Apply(const FReal Dt, const TArray<FConstraintContainerHandle*>& InConstraintHandles, const int32 It, const int32 NumIts)
 	{
 		SCOPE_CYCLE_COUNTER(STAT_Joints_Apply);
 
@@ -590,6 +593,9 @@ namespace Chaos
 		{
 			PostApplyCallback(Dt, SortedConstraintHandles);
 		}
+
+		// TODO: Return true only if more iteration are needed
+		return true;
 	}
 
 	
diff --git a/Engine/Source/Runtime/Experimental/Chaos/Private/Chaos/PBDPositionConstraints.cpp b/Engine/Source/Runtime/Experimental/Chaos/Private/Chaos/PBDPositionConstraints.cpp
--- a/Engine/Source/Runtime/Experimental/Chaos/Private/Chaos/PBDPositionConstraints.cpp
+++ b/Engine/Source/Runtime/Experimental/Chaos/Private/Chaos/PBDPositionConstraints.cpp
@@ -7,12 +7,15 @@ namespace Chaos
 {
 
 	template<class T, int d>
-	void TPBDPositionConstraints<T, d>::Apply(const T Dt, const TArray<FConstraintContainerHandle*>& ConstraintHandles, const int32 It, const int32 NumIts) const
+	bool TPBDPositionConstraints<T, d>::Apply(const T Dt, const TArray<FConstraintContainerHandle*>& ConstraintHandles, const int32 It, const int32 NumIts) const
 	{
 		for (FConstraintContainerHandle* ConstraintHandle : ConstraintHandles)
 		{
 			ApplySingle(Dt, ConstraintHandle->GetConstraintIndex());
 		}
+
+		// TODO: Return true only if more iteration are needed
+		return true;
 	}
 
 
diff --git a/Engine/Source/Runtime/Experimental/Chaos/Private/Chaos/PBDRigidSpringConstraints.cpp b/Engine/Source/Runtime/Experimental/Chaos/Private/Chaos/PBDRigidSpringConstraints.cpp
--- a/Engine/Source/Runtime/Experimental/Chaos/Private/Chaos/PBDRigidSpringConstraints.cpp
+++ b/Engine/Source/Runtime/Experimental/Chaos/Private/Chaos/PBDRigidSpringConstraints.cpp
@@ -124,20 +124,22 @@ namespace Chaos
 		return SpringSettings[ConstraintIndex].Stiffness * Delta / CombinedMass;
 	}
 
-	void FPBDRigidSpringConstraints::Apply(const FReal Dt, const int32 It, const int32 NumIts)
+	bool FPBDRigidSpringConstraints::Apply(const FReal Dt, const int32 It, const int32 NumIts)
 	{
 		for (int32 ConstraintIndex = 0; ConstraintIndex < NumConstraints(); ++ConstraintIndex)
 		{
 			ApplySingle(Dt, ConstraintIndex);
 		}
+		return false;
 	}
 
-	void FPBDRigidSpringConstraints::Apply(const FReal Dt, const TArray<FConstraintContainerHandle*>& InConstraintHandles, const int32 It, const int32 NumIts)
+	bool FPBDRigidSpringConstraints::Apply(const FReal Dt, const TArray<FConstraintContainerHandle*>& InConstraintHandles, const int32 It, const int32 NumIts)
 	{
 		for (FConstraintContainerHandle* ConstraintHandle : InConstraintHandles)
 		{
 			ApplySingle(Dt, ConstraintHandle->GetConstraintIndex());
 		}
+		return false;
 	}
 
 	void FPBDRigidSpringConstraints::ApplySingle(const FReal Dt, int32 ConstraintIndex) const
diff --git a/Engine/Source/Runtime/Experimental/Chaos/Private/Chaos/PBDRigidsEvolution.cpp b/Engine/Source/Runtime/Experimental/Chaos/Private/Chaos/PBDRigidsEvolution.cpp
--- a/Engine/Source/Runtime/Experimental/Chaos/Private/Chaos/PBDRigidsEvolution.cpp
+++ b/Engine/Source/Runtime/Experimental/Chaos/Private/Chaos/PBDRigidsEvolution.cpp
@@ -7,6 +7,12 @@
 #include "Chaos/ParticleHandle.h"
 #include "Chaos/SpatialAccelerationCollection.h"
 
+int32 ChaosRigidsEvolutionApplyAllowEarlyOutCVar = 1;
+FAutoConsoleVariableRef CVarChaosRigidsEvolutionApplyAllowEarlyOut(TEXT("p.ChaosRigidsEvolutionApplyAllowEarlyOut"), ChaosRigidsEvolutionApplyAllowEarlyOutCVar, TEXT("Allow Chaos Rigids Evolution apply iterations to early out when resolved.[def:1]"));
+
+int32 ChaosRigidsEvolutionApplyPushoutAllowEarlyOutCVar = 1;
+FAutoConsoleVariableRef CVarChaosRigidsEvolutionApplyPushoutAllowEarlyOut(TEXT("p.ChaosRigidsEvolutionApplyPushoutAllowEarlyOut"), ChaosRigidsEvolutionApplyPushoutAllowEarlyOutCVar, TEXT("Allow Chaos Rigids Evolution apply-pushout iterations to early out when resolved.[def:1]"));
+
 namespace Chaos
 {
 	CHAOS_API int32 FixBadAccelerationStructureRemoval = 1;
diff --git a/Engine/Source/Runtime/Experimental/Chaos/Private/Chaos/PBDRigidsEvolutionGBF.cpp b/Engine/Source/Runtime/Experimental/Chaos/Private/Chaos/PBDRigidsEvolutionGBF.cpp
--- a/Engine/Source/Runtime/Experimental/Chaos/Private/Chaos/PBDRigidsEvolutionGBF.cpp
+++ b/Engine/Source/Runtime/Experimental/Chaos/Private/Chaos/PBDRigidsEvolutionGBF.cpp
@@ -424,10 +424,10 @@ void TPBDRigidsEvolutionGBF<T, d>::AdvanceOneTimeStep(const T Dt, const T StepFr
 }
 
 template <typename T, int d>
-TPBDRigidsEvolutionGBF<T, d>::TPBDRigidsEvolutionGBF(TPBDRigidsSOAs<T, d>& InParticles, int32 InNumIterations, bool InIsSingleThreaded)
-	: Base(InParticles, InNumIterations, 1, InIsSingleThreaded)
+TPBDRigidsEvolutionGBF<T, d>::TPBDRigidsEvolutionGBF(TPBDRigidsSOAs<T, d>& InParticles, int32 InNumIterations, int32 InNumPushoutIterations, bool InIsSingleThreaded)
+	: Base(InParticles, InNumIterations, InNumPushoutIterations, InIsSingleThreaded)
 	, CollisionConstraints(InParticles, Collided, PhysicsMaterials, DefaultNumPairIterations, DefaultNumPushOutPairIterations)
-	, CollisionRule(CollisionConstraints, DefaultNumPushOutIterations)
+	, CollisionRule(CollisionConstraints)
 	, BroadPhase(InParticles, BoundsThickness, BoundsThicknessVelocityMultiplier)
 	, CollisionDetector(BroadPhase, CollisionConstraints)
 	, PostIntegrateCallback(nullptr)
diff --git a/Engine/Source/Runtime/Experimental/Chaos/Public/Chaos/PBDCollisionConstraints.h b/Engine/Source/Runtime/Experimental/Chaos/Public/Chaos/PBDCollisionConstraints.h
--- a/Engine/Source/Runtime/Experimental/Chaos/Public/Chaos/PBDCollisionConstraints.h
+++ b/Engine/Source/Runtime/Experimental/Chaos/Public/Chaos/PBDCollisionConstraints.h
@@ -180,7 +180,7 @@ public:
 	// Simple Rule API
 	//
 
-	void Apply(const T Dt, const int32 It, const int32 NumIts);
+	bool Apply(const T Dt, const int32 It, const int32 NumIts);
 	bool ApplyPushOut(const T Dt, const int32 It, const int32 NumIts);
 
 	//
@@ -188,7 +188,7 @@ public:
 	//
 	// @todo(ccaulfield): this runs wide. The serial/parallel decision should be in the ConstraintRule
 
-	void Apply(const T Dt, const TArray<FConstraintContainerHandle*>& InConstraintHandles, const int32 It, const int32 NumIts);
+	bool Apply(const T Dt, const TArray<FConstraintContainerHandle*>& InConstraintHandles, const int32 It, const int32 NumIts);
 	bool ApplyPushOut(const T Dt, const TArray<FConstraintContainerHandle*>& InConstraintHandles, 
 		const TSet<const TGeometryParticleHandle<T,d>*>& IsTemporarilyStatic, int32 Iteration, int32 NumIterations);
 
diff --git a/Engine/Source/Runtime/Experimental/Chaos/Public/Chaos/PBDConstraintRule.h b/Engine/Source/Runtime/Experimental/Chaos/Public/Chaos/PBDConstraintRule.h
--- a/Engine/Source/Runtime/Experimental/Chaos/Public/Chaos/PBDConstraintRule.h
+++ b/Engine/Source/Runtime/Experimental/Chaos/Public/Chaos/PBDConstraintRule.h
@@ -50,7 +50,7 @@ namespace Chaos
 		FSimpleConstraintRule(int32 InPriority) : FConstraintRule(InPriority) {}
 
 		virtual void UpdatePositionBasedState(const FReal Dt) {}
-		virtual void ApplyConstraints(const FReal Dt, const int32 It, const int32 NumIts) {}
+		virtual bool ApplyConstraints(const FReal Dt, const int32 It, const int32 NumIts) { return false; }
 		virtual bool ApplyPushOut(const FReal Dt, const int32 It, const int32 NumIts) { return false; }
 	};
 
@@ -81,9 +81,9 @@ namespace Chaos
 			return Constraints.UpdatePositionBasedState(Dt);
 		}
 
-		virtual void ApplyConstraints(const FReal Dt, const int32 It, const int32 NumIts) override
+		virtual bool ApplyConstraints(const FReal Dt, const int32 It, const int32 NumIts) override
 		{
-			Constraints.Apply(Dt, It, NumIts);
+			return Constraints.Apply(Dt, It, NumIts);
 		}
 
 		virtual bool ApplyPushOut(const FReal Dt, const int32 It, const int32 NumIts) override
@@ -114,10 +114,10 @@ namespace Chaos
 		/** Called once per tick to allow constraint containers to create/alter their constraints based on particle position */
 		virtual void UpdatePositionBasedState(const FReal Dt) {}
 
-		/** Apply all corrections for constraints in the specified island */
-		virtual void ApplyConstraints(const FReal Dt, int32 Island, const int32 It, const int32 NumIts) {}
+		/** Apply all corrections for constraints in the specified island. Return true if more iterations are needed. */
+		virtual bool ApplyConstraints(const FReal Dt, int32 Island, const int32 It, const int32 NumIts) { return false; }
 
-		/** Apply push out for constraints in the specified island */
+		/** Apply push out for constraints in the specified island. Return true if more iterations are needed. */
 		virtual bool ApplyPushOut(const FReal Dt, int32 Island, const int32 It, const int32 NumIts) { return false; }
 
 		/** Add all constraints to the connectivity graph */
@@ -191,12 +191,13 @@ namespace Chaos
 		{
 		}
 
-		virtual void ApplyConstraints(const FReal Dt, int32 Island, const int32 It, const int32 NumIts) override
+		virtual bool ApplyConstraints(const FReal Dt, int32 Island, const int32 It, const int32 NumIts) override
 		{
 			if (IslandConstraintLists[Island].Num())
 			{
-				Constraints.Apply(Dt, GetIslandConstraints(Island), It, NumIts);
+				return Constraints.Apply(Dt, GetIslandConstraints(Island), It, NumIts);
 			}
+			return false;
 		}
 
 		virtual bool ApplyPushOut(const FReal Dt, int32 Island, const int32 It, const int32 NumIts) override
@@ -270,9 +271,8 @@ namespace Chaos
 	public:
 		using FConstraints = T_CONSTRAINTS;
 
-		TPBDConstraintColorRule(FConstraints& InConstraints, const int32 InPushOutIterations, int32 InPriority = 0)
+		TPBDConstraintColorRule(FConstraints& InConstraints, int32 InPriority = 0)
 			: TPBDConstraintGraphRuleImpl<T_CONSTRAINTS>(InConstraints, InPriority)
-			, PushOutIterations(InPushOutIterations)
 		{
 		}
 
@@ -282,11 +282,12 @@ namespace Chaos
 			Constraints.UpdatePositionBasedState(Dt);
 		}
 
-		virtual void ApplyConstraints(const FReal Dt, int32 Island, const int32 It, const int32 NumIts) override
+		virtual bool ApplyConstraints(const FReal Dt, int32 Island, const int32 It, const int32 NumIts) override
 		{
 			const typename FPBDConstraintColor::FLevelToColorToConstraintListMap& LevelToColorToConstraintListMap = GraphColor.GetIslandLevelToColorToConstraintListMap(Island);
 			int32 MaxColor = GraphColor.GetIslandMaxColor(Island);
 			int32 MaxLevel = GraphColor.GetIslandMaxLevel(Island);
+			bool bNeedsMoreIterations = false;
 			for (int32 Level = 0; Level <= MaxLevel; ++Level)
 			{
 				for (int32 Color = 0; Color <= MaxColor; ++Color)
@@ -294,10 +295,12 @@ namespace Chaos
 					if (LevelToColorToConstraintListMap[Level].Contains(Color) && LevelToColorToConstraintListMap[Level][Color].Num())
 					{
 						const TArray<typename FConstraints::FConstraintContainerHandle*>& ConstraintHandles = GetLevelColorConstraints(LevelToColorToConstraintListMap, Level, Color);
-						Constraints.Apply(Dt, ConstraintHandles, It, NumIts);
+						bNeedsMoreIterations |= Constraints.Apply(Dt, ConstraintHandles, It, NumIts);
 					}
 				}
 			}
+
+			return bNeedsMoreIterations;
 		}
 
 		virtual void RemoveConstraints(const TSet<TGeometryParticleHandle<FReal, 3>*>& InConstraints)
@@ -372,11 +375,6 @@ namespace Chaos
 			GraphColor.ComputeColor(Island, *ConstraintGraph, ContainerId);
 		}
 
-		void SetPushOutIterations(const int32 InPushOutIterations)
-		{
-			PushOutIterations = InPushOutIterations;
-		}
-
 		template<typename TVisitor>
 		void VisitIslandConstraints(const int32 Island, const TVisitor& Visitor) const
 		{
@@ -408,7 +406,6 @@ namespace Chaos
 		}
 
 		FPBDConstraintColor GraphColor;
-		int32 PushOutIterations;
 	};
 
 }
\ No newline at end of file
diff --git a/Engine/Source/Runtime/Experimental/Chaos/Public/Chaos/PBDJointConstraints.h b/Engine/Source/Runtime/Experimental/Chaos/Public/Chaos/PBDJointConstraints.h
--- a/Engine/Source/Runtime/Experimental/Chaos/Public/Chaos/PBDJointConstraints.h
+++ b/Engine/Source/Runtime/Experimental/Chaos/Public/Chaos/PBDJointConstraints.h
@@ -145,14 +145,14 @@ namespace Chaos
 		// Simple Rule API
 		//
 
-		void Apply(const FReal Dt, const int32 It, const int32 NumIts);
+		bool Apply(const FReal Dt, const int32 It, const int32 NumIts);
 		bool ApplyPushOut(const FReal Dt, const int32 It, const int32 NumIts);
 
 		//
 		// Island Rule API
 		//
 
-		void Apply(const FReal Dt, const TArray<FConstraintContainerHandle*>& InConstraintHandles, const int32 It, const int32 NumIts);
+		bool Apply(const FReal Dt, const TArray<FConstraintContainerHandle*>& InConstraintHandles, const int32 It, const int32 NumIts);
 		bool ApplyPushOut(const FReal Dt, const TArray<FConstraintContainerHandle*>& InConstraintHandles, const int32 It, const int32 NumIts);
 
 
diff --git a/Engine/Source/Runtime/Experimental/Chaos/Public/Chaos/PBDPositionConstraints.h b/Engine/Source/Runtime/Experimental/Chaos/Public/Chaos/PBDPositionConstraints.h
--- a/Engine/Source/Runtime/Experimental/Chaos/Public/Chaos/PBDPositionConstraints.h
+++ b/Engine/Source/Runtime/Experimental/Chaos/Public/Chaos/PBDPositionConstraints.h
@@ -175,7 +175,7 @@ namespace Chaos
 		{
 		}
 
-		void Apply(const T Dt, const TArray<FConstraintContainerHandle*>& ConstraintHandles, const int32 It, const int32 NumIts) const;
+		bool Apply(const T Dt, const TArray<FConstraintContainerHandle*>& ConstraintHandles, const int32 It, const int32 NumIts) const;
 
 		bool ApplyPushOut(const T Dt, const TArray<FConstraintContainerHandle*>& InConstraintIndices, const int32 It, const int32 NumIts) const
 		{
diff --git a/Engine/Source/Runtime/Experimental/Chaos/Public/Chaos/PBDRigidDynamicSpringConstraints.h b/Engine/Source/Runtime/Experimental/Chaos/Public/Chaos/PBDRigidDynamicSpringConstraints.h
--- a/Engine/Source/Runtime/Experimental/Chaos/Public/Chaos/PBDRigidDynamicSpringConstraints.h
+++ b/Engine/Source/Runtime/Experimental/Chaos/Public/Chaos/PBDRigidDynamicSpringConstraints.h
@@ -179,12 +179,15 @@ namespace Chaos
 
 		void UpdatePositionBasedState(const T Dt);
 
-		void Apply(const T Dt, const TArray<FConstraintContainerHandle*>& InConstraintHandles, const int32 It, const int32 NumIts)
+		bool Apply(const T Dt, const TArray<FConstraintContainerHandle*>& InConstraintHandles, const int32 It, const int32 NumIts)
 		{
 			for (FConstraintContainerHandle* ConstraintHandle : InConstraintHandles)
 			{
 				ApplySingle(Dt, ConstraintHandle->GetConstraintIndex());
 			}
+
+			// TODO: Return true only if more iteration are needed
+			return true;
 		}
 
 		bool ApplyPushOut(const T Dt, const TArray<FConstraintContainerHandle*>& InConstraintHandles, const int32 It, const int32 NumIts)
diff --git a/Engine/Source/Runtime/Experimental/Chaos/Public/Chaos/PBDRigidSpringConstraints.h b/Engine/Source/Runtime/Experimental/Chaos/Public/Chaos/PBDRigidSpringConstraints.h
--- a/Engine/Source/Runtime/Experimental/Chaos/Public/Chaos/PBDRigidSpringConstraints.h
+++ b/Engine/Source/Runtime/Experimental/Chaos/Public/Chaos/PBDRigidSpringConstraints.h
@@ -159,14 +159,14 @@ namespace Chaos
 		{
 		}
 
-		void Apply(const FReal Dt, const int32 It, const int32 NumIts);
+		bool Apply(const FReal Dt, const int32 It, const int32 NumIts);
 
 		bool ApplyPushOut(const FReal Dt, const int32 It, const int32 NumIts)
 		{
 			return false;
 		}
 
-		void Apply(const FReal Dt, const TArray<FConstraintContainerHandle*>& InConstraintHandles, const int32 It, const int32 NumIts);
+		bool Apply(const FReal Dt, const TArray<FConstraintContainerHandle*>& InConstraintHandles, const int32 It, const int32 NumIts);
 
 		bool ApplyPushOut(const FReal Dt, const TArray<FConstraintContainerHandle*>& InConstraintHandles, const int32 It, const int32 NumIts)
 		{
diff --git a/Engine/Source/Runtime/Experimental/Chaos/Public/Chaos/PBDRigidsEvolution.h b/Engine/Source/Runtime/Experimental/Chaos/Public/Chaos/PBDRigidsEvolution.h
--- a/Engine/Source/Runtime/Experimental/Chaos/Public/Chaos/PBDRigidsEvolution.h
+++ b/Engine/Source/Runtime/Experimental/Chaos/Public/Chaos/PBDRigidsEvolution.h
@@ -15,6 +15,10 @@
 #include "Chaos/SpatialAccelerationCollection.h"
 
 
+extern int32 ChaosRigidsEvolutionApplyAllowEarlyOutCVar;
+extern int32 ChaosRigidsEvolutionApplyPushoutAllowEarlyOutCVar;
+
+
 // Declaring so it can be friended for tests.
 namespace ChaosTest { void TestPendingSpatialDataHandlePointerConflict(); } 
 
@@ -457,9 +461,15 @@ class TPBDRigidsEvolutionBase
 		// @todo(ccaulfield): track whether we are sufficiently solved and can early-out
 		for (int i = 0; i < NumIterations; ++i)
 		{
+			bool bNeedsAnotherIteration = false;
 			for (FPBDConstraintGraphRule* ConstraintRule : PrioritizedConstraintRules)
 			{
-				ConstraintRule->ApplyConstraints(Dt, Island, i, NumIterations);
+				bNeedsAnotherIteration |= ConstraintRule->ApplyConstraints(Dt, Island, i, NumIterations);
+			}
+
+			if (ChaosRigidsEvolutionApplyAllowEarlyOutCVar && !bNeedsAnotherIteration)
+			{
+				break;
 			}
 		}
 	}
@@ -629,15 +639,17 @@ protected:
 	void ApplyPushOut(const T Dt, int32 Island)
 	{
 		bool bNeedsAnotherIteration = true;
-		for (int32 It = 0; bNeedsAnotherIteration && (It < NumPushOutIterations); ++It)
+		for (int32 It = 0; It < NumPushOutIterations; ++It)
 		{
 			bNeedsAnotherIteration = false;
 			for (FPBDConstraintGraphRule* ConstraintRule : PrioritizedConstraintRules)
 			{
-				if (ConstraintRule->ApplyPushOut(Dt, Island, It, NumPushOutIterations))
-				{
-					bNeedsAnotherIteration = true;
-				}
+				bNeedsAnotherIteration |= ConstraintRule->ApplyPushOut(Dt, Island, It, NumPushOutIterations);
+			}
+
+			if (ChaosRigidsEvolutionApplyPushoutAllowEarlyOutCVar && !bNeedsAnotherIteration)
+			{
+				break;
 			}
 		}
 	}
diff --git a/Engine/Source/Runtime/Experimental/Chaos/Public/Chaos/PBDRigidsEvolutionGBF.h b/Engine/Source/Runtime/Experimental/Chaos/Public/Chaos/PBDRigidsEvolutionGBF.h
--- a/Engine/Source/Runtime/Experimental/Chaos/Public/Chaos/PBDRigidsEvolutionGBF.h
+++ b/Engine/Source/Runtime/Experimental/Chaos/Public/Chaos/PBDRigidsEvolutionGBF.h
@@ -62,7 +62,7 @@ namespace Chaos
 		static constexpr int32 DefaultNumPushOutIterations = 5;
 		static constexpr int32 DefaultNumPushOutPairIterations = 2;
 
-		CHAOS_API TPBDRigidsEvolutionGBF(TPBDRigidsSOAs<T, d>& InParticles, int32 InNumIterations = DefaultNumIterations, bool InIsSingleThreaded = false);
+		CHAOS_API TPBDRigidsEvolutionGBF(TPBDRigidsSOAs<T, d>& InParticles, int32 InNumIterations = DefaultNumIterations, int32 InNumPushoutIterations = DefaultNumPushOutIterations, bool InIsSingleThreaded = false);
 		CHAOS_API ~TPBDRigidsEvolutionGBF() {}
 
 		void SetPostIntegrateCallback(const TPBDRigidsEvolutionCallback<T, d>& Cb)
diff --git a/Engine/Source/Runtime/Experimental/ChaosSolvers/Private/PBDRigidsSolver.cpp b/Engine/Source/Runtime/Experimental/ChaosSolvers/Private/PBDRigidsSolver.cpp
--- a/Engine/Source/Runtime/Experimental/ChaosSolvers/Private/PBDRigidsSolver.cpp
+++ b/Engine/Source/Runtime/Experimental/ChaosSolvers/Private/PBDRigidsSolver.cpp
@@ -36,9 +36,12 @@ FAutoConsoleVariableRef CVarChaosSolverUseParticlePool(TEXT("p.Chaos.Solver.UseP
 int32 ChaosSolverParticlePoolNumFrameUntilShrink = 30;
 FAutoConsoleVariableRef CVarChaosSolverParticlePoolNumFrameUntilShrink(TEXT("p.Chaos.Solver.ParticlePoolNumFrameUntilShrink"), ChaosSolverParticlePoolNumFrameUntilShrink, TEXT("Num Frame until we can potentially shrink the pool"));
 
-float ChaosSolverCollisionDefaultIterationsCVar = 1;
+int32 ChaosSolverCollisionDefaultIterationsCVar = 4;
 FAutoConsoleVariableRef CVarChaosSolverCollisionDefaultIterations(TEXT("p.ChaosSolverCollisionDefaultIterations"), ChaosSolverCollisionDefaultIterationsCVar, TEXT("Default collision iterations for the solver.[def:1]"));
 
+int32 ChaosSolverCollisionDefaultPushoutIterationsCVar = 4;
+FAutoConsoleVariableRef CVarChaosSolverCollisionDefaultPushoutIterations(TEXT("p.ChaosSolverCollisionDefaultPushoutIterations"), ChaosSolverCollisionDefaultPushoutIterationsCVar, TEXT("Default collision pushout iterations for the solver.[def:1]"));
+
 int32 ChaosSolverCleanupCommandsOnDestruction = 1;
 FAutoConsoleVariableRef CVarChaosSolverCleanupCommandsOnDestruction(TEXT("p.Chaos.Solver.CleanupCommandsOnDestruction"), ChaosSolverCleanupCommandsOnDestruction, TEXT("Whether or not to run internal command queue cleanup on solver destruction (0 = no cleanup, >0 = cleanup all commands)"));
 
@@ -161,7 +164,7 @@ namespace Chaos
 		, bHasFloor(true)
 		, bIsFloorAnalytic(false)
 		, FloorHeight(0.f)
-		, MEvolution(new FPBDRigidsEvolution(Particles, ChaosSolverCollisionDefaultIterationsCVar, BufferingModeIn == Chaos::EMultiBufferMode::Single))
+		, MEvolution(new FPBDRigidsEvolution(Particles, ChaosSolverCollisionDefaultIterationsCVar, ChaosSolverCollisionDefaultPushoutIterationsCVar, BufferingModeIn == Chaos::EMultiBufferMode::Single))
 		, MEventManager(new FEventManager(BufferingModeIn))
 		, MSolverEventFilters(new FSolverEventFilters())
 		, MActiveParticlesBuffer(new FActiveParticlesBuffer(BufferingModeIn))
@@ -483,7 +486,7 @@ namespace Chaos
 		CurrentFrame = 0;
 		MMaxDeltaTime = 1;
 		TimeStepMultiplier = 1;
-		MEvolution = TUniquePtr<FPBDRigidsEvolution>(new FPBDRigidsEvolution(Particles, ChaosSolverCollisionDefaultIterationsCVar, BufferMode == EMultiBufferMode::Single)); 
+		MEvolution = TUniquePtr<FPBDRigidsEvolution>(new FPBDRigidsEvolution(Particles, ChaosSolverCollisionDefaultIterationsCVar, ChaosSolverCollisionDefaultPushoutIterationsCVar, BufferMode == EMultiBufferMode::Single)); 
 
 		FEventDefaults::RegisterSystemEvents(*GetEventManager());
 	}
