Chaos: Overlap Sphere with scaled convex bug fix

diff --git a/Engine/Source/Runtime/Experimental/Chaos/Public/Chaos/GeometryQueries.h b/Engine/Source/Runtime/Experimental/Chaos/Public/Chaos/GeometryQueries.h
--- a/Engine/Source/Runtime/Experimental/Chaos/Public/Chaos/GeometryQueries.h
+++ b/Engine/Source/Runtime/Experimental/Chaos/Public/Chaos/GeometryQueries.h
@@ -42,15 +42,18 @@ namespace Chaos
 
 		const FRigidTransform3 BToATM = BTM.GetRelativeTransform(ATM);
 
-		if(BType == ImplicitObjectType::Sphere)
-		{
-			const FImplicitObject& BBase = static_cast<const FImplicitObject&>(B);
-			const TSphere<FReal, 3>& BSphere = static_cast<const TSphere<FReal, 3>&>(BBase);
-			const FVec3 PtInA = BToATM.TransformPositionNoScale(BSphere.GetCenter());
-			return A.Overlap(PtInA, Thickness + BSphere.GetRadius());
-		}
-		//todo: A is a sphere
-		else if (A.IsConvex())
+		// This specialization for sphere is bugged since the sphere radius is not inverse scaled, 
+		// nor can it be properly if testing against non-uniform scaled convexes
+		//if(BType == ImplicitObjectType::Sphere)
+		//{
+		//	const FImplicitObject& BBase = static_cast<const FImplicitObject&>(B);
+		//	const TSphere<FReal, 3>& BSphere = static_cast<const TSphere<FReal, 3>&>(BBase);
+		//	const FVec3 PtInA = BToATM.TransformPositionNoScale(BSphere.GetCenter());
+		//	return A.Overlap(PtInA, Thickness + BSphere.GetRadius());
+		//}
+		////todo: A is a sphere
+		//else 
+		if (A.IsConvex())
 		{
 			const FVec3 Offset = ATM.GetLocation() - BTM.GetLocation();
 			if (OutMTD)
diff --git a/Engine/Source/Runtime/Experimental/Chaos/Public/Chaos/ImplicitObjectScaled.h b/Engine/Source/Runtime/Experimental/Chaos/Public/Chaos/ImplicitObjectScaled.h
--- a/Engine/Source/Runtime/Experimental/Chaos/Public/Chaos/ImplicitObjectScaled.h
+++ b/Engine/Source/Runtime/Experimental/Chaos/Public/Chaos/ImplicitObjectScaled.h
@@ -458,7 +458,11 @@ public:
 	virtual bool Overlap(const TVector<T, d>& Point, const T Thickness) const override
 	{
 		const TVector<T, d> UnscaledPoint = MInvScale * Point;
-		return MObject->Overlap(UnscaledPoint, MInternalThickness + Thickness);
+
+		// TODO: consider alternative that handles thickness scaling properly in 3D, only works for uniform scaling right now
+		const T UnscaleThickness = MInvScale[0] * Thickness; 
+
+		return MObject->Overlap(UnscaledPoint, MInternalThickness + UnscaleThickness);
 	}
 
 	virtual Pair<TVector<T, d>, bool> FindClosestIntersectionImp(const TVector<T, d>& StartPoint, const TVector<T, d>& EndPoint, const T Thickness) const override
