Chaos
- ensure body bounds includes previous position
- this handles collision jitter caused by joint pulliing bodies back into collision

diff --git a/Engine/Source/Runtime/Experimental/Chaos/Private/Chaos/CollisionResolution.cpp b/Engine/Source/Runtime/Experimental/Chaos/Private/Chaos/CollisionResolution.cpp
--- a/Engine/Source/Runtime/Experimental/Chaos/Private/Chaos/CollisionResolution.cpp
+++ b/Engine/Source/Runtime/Experimental/Chaos/Private/Chaos/CollisionResolution.cpp
@@ -29,7 +29,7 @@ DECLARE_CYCLE_STAT(TEXT("Collisions::GJK"), STAT_Collisions_GJK, STATGROUP_Chaos
 #define SCOPE_CYCLE_COUNTER_GJK()
 #endif
 
-//PRAGMA_DISABLE_OPTIMIZATION_ACTUAL
+//PRAGMA_DISABLE_OPTIMIZATION
 
 
 // If GJKPenetration returns a phi of abs value < this number, we use PhiWithNormal to resample phi and normal.
diff --git a/Engine/Source/Runtime/Experimental/Chaos/Private/Chaos/Evolution/PBDMinEvolution.cpp b/Engine/Source/Runtime/Experimental/Chaos/Private/Chaos/Evolution/PBDMinEvolution.cpp
--- a/Engine/Source/Runtime/Experimental/Chaos/Private/Chaos/Evolution/PBDMinEvolution.cpp
+++ b/Engine/Source/Runtime/Experimental/Chaos/Private/Chaos/Evolution/PBDMinEvolution.cpp
@@ -17,6 +17,8 @@
 #include "Chaos/PerParticlePBDUpdateFromDeltaPosition.h"
 #include "ChaosStats.h"
 
+//PRAGMA_DISABLE_OPTIMIZATION
+
 namespace Chaos
 {
 	DECLARE_CYCLE_STAT(TEXT("MinEvolution::AdvanceOneTimeStep"), STAT_MinEvolution_AdvanceOneTimeStep, STATGROUP_Chaos);
@@ -163,8 +165,9 @@ namespace Chaos
 					TAABB<FReal, 3> WorldSpaceBounds = LocalBounds.TransformedAABB(FRigidTransform3(Particle.P(), Particle.Q()));
 					WorldSpaceBounds.ThickenSymmetrically(WorldSpaceBounds.Extents() * BoundsExtension);
 
-					//TAABB<FReal, 3> PrevWorldSpaceBounds = LocalBounds.TransformedAABB(FRigidTransform3(Particle.X(), Particle.R()));
-					//WorldSpaceBounds.GrowToInclude(PrevWorldSpaceBounds);
+					// Dynamic bodies may get pulled back into their old positions by joints - make sure we find collisions that may prevent this
+					// We could add the AABB at X/R here, but I'm avoiding another call to TransformedAABB. Hopefully this is good enough.
+					WorldSpaceBounds.GrowByVector(Particle.X() - Particle.P());
 
 					Particle.SetWorldSpaceInflatedBounds(WorldSpaceBounds);
 				}
diff --git a/Engine/Source/Runtime/Experimental/Chaos/Private/Chaos/PBDCollisionConstraints.cpp b/Engine/Source/Runtime/Experimental/Chaos/Private/Chaos/PBDCollisionConstraints.cpp
--- a/Engine/Source/Runtime/Experimental/Chaos/Private/Chaos/PBDCollisionConstraints.cpp
+++ b/Engine/Source/Runtime/Experimental/Chaos/Private/Chaos/PBDCollisionConstraints.cpp
@@ -29,8 +29,7 @@
 #include "PBDCollisionConstraints.ispc.generated.h"
 #endif
 
-//#pragma optimize("", off)
-//PRAGMA_DISABLE_OPTIMIZATION_ACTUAL
+//PRAGMA_DISABLE_OPTIMIZATION
 
 namespace Chaos
 {
diff --git a/Engine/Source/Runtime/Experimental/Chaos/Public/Chaos/AABB.h b/Engine/Source/Runtime/Experimental/Chaos/Public/Chaos/AABB.h
--- a/Engine/Source/Runtime/Experimental/Chaos/Public/Chaos/AABB.h
+++ b/Engine/Source/Runtime/Experimental/Chaos/Public/Chaos/AABB.h
@@ -424,6 +424,13 @@ namespace Chaos
 			MMax += AbsThickness;
 		}
 
+		/** Grow along a vector (as if swept by the vector's direction and magnitude) */
+		FORCEINLINE void GrowByVector(const TVector<T, d>& V)
+		{
+			MMin += V.ComponentwiseMin(TVector<T, d>(0));
+			MMax += V.ComponentwiseMax(TVector<T, d>(0));
+		}
+
 		FORCEINLINE TVector<T, d> Center() const { return (MMax - MMin) / (T)2 + MMin; }
 		FORCEINLINE TVector<T, d> GetCenter() const { return Center(); }
 		FORCEINLINE TVector<T, d> GetCenterOfMass() const { return GetCenter(); }
