Chaos - RBAN cleanup early out

diff --git a/Engine/Source/Runtime/Experimental/Chaos/Private/Chaos/Evolution/PBDMinEvolution.cpp b/Engine/Source/Runtime/Experimental/Chaos/Private/Chaos/Evolution/PBDMinEvolution.cpp
--- a/Engine/Source/Runtime/Experimental/Chaos/Private/Chaos/Evolution/PBDMinEvolution.cpp
+++ b/Engine/Source/Runtime/Experimental/Chaos/Private/Chaos/Evolution/PBDMinEvolution.cpp
@@ -299,7 +299,6 @@ namespace Chaos
 	{
 		SCOPE_CYCLE_COUNTER(STAT_MinEvolution_ApplyConstraints);
 
-		// @todo(ccaulfield): track whether we are sufficiently solved and can early-out
 		for (int32 i = 0; i < NumApplyIterations; ++i)
 		{
 			bool bNeedsAnotherIteration = false;
@@ -330,16 +329,17 @@ namespace Chaos
 	{
 		SCOPE_CYCLE_COUNTER(STAT_MinEvolution_ApplyPushOut);
 
-		bool bNeedsAnotherIteration = true;
-		for (int32 It = 0; bNeedsAnotherIteration && (It < NumApplyPushOutIterations); ++It)
+		for (int32 It = 0; It < NumApplyPushOutIterations; ++It)
 		{
-			bNeedsAnotherIteration = false;
+			bool bNeedsAnotherIteration = false;
 			for (FSimpleConstraintRule* ConstraintRule : PrioritizedConstraintRules)
 			{
-				if (ConstraintRule->ApplyPushOut(Dt, It, NumApplyPushOutIterations))
-				{
-					bNeedsAnotherIteration = true;
-				}
+				bNeedsAnotherIteration |= ConstraintRule->ApplyPushOut(Dt, It, NumApplyPushOutIterations);
+			}
+
+			if (!bNeedsAnotherIteration)
+			{
+				break;
 			}
 		}
 	}
diff --git a/Engine/Source/Runtime/Experimental/Chaos/Private/Chaos/PBDJointConstraints.cpp b/Engine/Source/Runtime/Experimental/Chaos/Private/Chaos/PBDJointConstraints.cpp
--- a/Engine/Source/Runtime/Experimental/Chaos/Private/Chaos/PBDJointConstraints.cpp
+++ b/Engine/Source/Runtime/Experimental/Chaos/Private/Chaos/PBDJointConstraints.cpp
@@ -509,24 +509,20 @@ namespace Chaos
 			}
 		}
 
-		UE_LOG(LogChaosJoint, VeryVerbose, TEXT("Apply %d / %d Active %d / %d"), It, NumIts, NetResult.GetNumActive(), NetResult.GetNumActive() + NetResult.GetNumSolved());
+		UE_LOG(LogChaosJoint, Verbose, TEXT("Apply Iteration: %d / %d; Active: %d / %d"), It, NumIts, NetResult.GetNumActive(), NetResult.GetNumActive() + NetResult.GetNumSolved());
 
 		if (PostApplyCallback != nullptr)
 		{
 			PostApplyCallback(Dt, Handles);
 		}
 
-		// TODO: Return true only if more iteration are needed
-		return true;
+		return (NetResult.GetNumActive() > 0);
 	}
 
 	bool FPBDJointConstraints::ApplyPushOut(const FReal Dt, const int32 It, const int32 NumIts)
 	{
 		SCOPE_CYCLE_COUNTER(STAT_Joints_ApplyPushOut);
 
-		// @todo(ccaulfield): track whether we are sufficiently solved
-		bool bNeedsAnotherIteration = true;
-
 		FJointSolverResult NetResult;
 		if (Settings.ApplyPushOutPairIterations > 0)
 		{
@@ -536,14 +532,14 @@ namespace Chaos
 			}
 		}
 
-		UE_LOG(LogChaosJoint, VeryVerbose, TEXT("PushOut  %d / %d Active %d / %d"), It, NumIts, NetResult.GetNumActive(), NetResult.GetNumActive() + NetResult.GetNumSolved());
+		UE_LOG(LogChaosJoint, Verbose, TEXT("PushOut Iteration: %d / %d; Active: %d / %d"), It, NumIts, NetResult.GetNumActive(), NetResult.GetNumActive() + NetResult.GetNumSolved());
 
 		if (PostProjectCallback != nullptr)
 		{
 			PostProjectCallback(Dt, Handles);
 		}
 
-		return bNeedsAnotherIteration;
+		return (NetResult.GetNumActive() > 0);
 	}
 
 	//////////////////////////////////////////////////////////////////////////
@@ -581,11 +577,12 @@ namespace Chaos
 		}
 
 
+		FJointSolverResult NetResult;
 		if (Settings.ApplyPairIterations > 0)
 		{
 			for (FConstraintContainerHandle* ConstraintHandle : SortedConstraintHandles)
 			{
-				SolvePosition_GaussSiedel(Dt, ConstraintHandle->GetConstraintIndex(), Settings.ApplyPairIterations, It, NumIts);
+				NetResult += SolvePosition_GaussSiedel(Dt, ConstraintHandle->GetConstraintIndex(), Settings.ApplyPairIterations, It, NumIts);
 			}
 		}
 
@@ -594,8 +591,7 @@ namespace Chaos
 			PostApplyCallback(Dt, SortedConstraintHandles);
 		}
 
-		// TODO: Return true only if more iteration are needed
-		return true;
+		return (NetResult.GetNumActive() > 0);
 	}
 
 	
@@ -603,9 +599,6 @@ namespace Chaos
 	{
 		SCOPE_CYCLE_COUNTER(STAT_Joints_ApplyPushOut);
 
-		// @todo(ccaulfield): track whether we are sufficiently solved
-		bool bNeedsAnotherIteration = true;
-
 		TArray<FConstraintContainerHandle*> SortedConstraintHandles = InConstraintHandles;
 		SortedConstraintHandles.Sort([](const FConstraintContainerHandle& L, const FConstraintContainerHandle& R)
 			{
@@ -613,11 +606,12 @@ namespace Chaos
 				return L.GetConstraintLevel() < R.GetConstraintLevel();
 			});
 
+		FJointSolverResult NetResult;
 		if (Settings.ApplyPushOutPairIterations > 0)
 		{
 			for (FConstraintContainerHandle* ConstraintHandle : SortedConstraintHandles)
 			{
-				ProjectPosition_GaussSiedel(Dt, ConstraintHandle->GetConstraintIndex(), Settings.ApplyPushOutPairIterations, It, NumIts);
+				NetResult += ProjectPosition_GaussSiedel(Dt, ConstraintHandle->GetConstraintIndex(), Settings.ApplyPushOutPairIterations, It, NumIts);
 			}
 		}
 
@@ -626,7 +620,7 @@ namespace Chaos
 			PostProjectCallback(Dt, SortedConstraintHandles);
 		}
 
-		return bNeedsAnotherIteration;
+		return (NetResult.GetNumActive() > 0);
 	}
 
 	//////////////////////////////////////////////////////////////////////////
