Chaos
- remove template parameters from Evolution classes
- remove EvolutionPGS

diff --git a/Engine/Source/Programs/HeadlessChaos/Private/HeadlessChaosApp.cpp b/Engine/Source/Programs/HeadlessChaos/Private/HeadlessChaosApp.cpp
--- a/Engine/Source/Programs/HeadlessChaos/Private/HeadlessChaosApp.cpp
+++ b/Engine/Source/Programs/HeadlessChaos/Private/HeadlessChaosApp.cpp
@@ -94,10 +94,7 @@ TEST(CollisionTests, PGS) {
 }
 
 TEST(ForceTests, Forces) {
-	ChaosTest::Gravity<Chaos::TPBDRigidsEvolutionGBF<float, 3>, float>();
-#if CHAOS_PARTICLEHANDLE_TODO
-	ChaosTest::Gravity<Chaos::TPBDRigidsEvolutionPGS<float, 3>, float>();
-#endif
+	ChaosTest::Gravity<Chaos::FPBDRigidsEvolutionGBF, float>();
 	SUCCEED();
 }
 
diff --git a/Engine/Source/Programs/HeadlessChaos/Private/HeadlessChaosTestBP.cpp b/Engine/Source/Programs/HeadlessChaos/Private/HeadlessChaosTestBP.cpp
--- a/Engine/Source/Programs/HeadlessChaos/Private/HeadlessChaosTestBP.cpp
+++ b/Engine/Source/Programs/HeadlessChaos/Private/HeadlessChaosTestBP.cpp
@@ -25,7 +25,7 @@ namespace ChaosTest
 			FChaosArchive Ar(*File);
 
 			TPBDRigidParticles<T, 3> Particles;
-			TPBDRigidsEvolutionGBF<T, 3> Evolution(MoveTemp(Particles));
+			FPBDRigidsEvolutionGBF Evolution(MoveTemp(Particles));
 			Evolution.SerializeForPerfTest(Ar);
 			Evolution.AdvanceOneTimeStep(Dt);
 			Evolution.EndFrame(Dt);
diff --git a/Engine/Source/Programs/HeadlessChaos/Private/HeadlessChaosTestBroadphase.cpp b/Engine/Source/Programs/HeadlessChaos/Private/HeadlessChaosTestBroadphase.cpp
--- a/Engine/Source/Programs/HeadlessChaos/Private/HeadlessChaosTestBroadphase.cpp
+++ b/Engine/Source/Programs/HeadlessChaos/Private/HeadlessChaosTestBroadphase.cpp
@@ -654,7 +654,7 @@ namespace ChaosTest
 		Particle->GTGeometryParticle() = TGeometryParticle<FReal, 3>::CreateParticle().Release();
 		Particle->X() = TVector<FReal, 3>(0, 0, 0);
 
-		TPBDRigidsEvolutionGBF<FReal, 3> Evolution(Particles);
+		FPBDRigidsEvolutionGBF Evolution(Particles);
 
 		// Flush spatial acceleration structures to put particle into structure.
 		Evolution.FlushSpatialAcceleration();
diff --git a/Engine/Source/Programs/HeadlessChaos/Private/HeadlessChaosTestClustering.cpp b/Engine/Source/Programs/HeadlessChaos/Private/HeadlessChaosTestClustering.cpp
--- a/Engine/Source/Programs/HeadlessChaos/Private/HeadlessChaosTestClustering.cpp
+++ b/Engine/Source/Programs/HeadlessChaos/Private/HeadlessChaosTestClustering.cpp
@@ -25,7 +25,7 @@ namespace ChaosTest {
 	void ImplicitCluster()
 	{
 		TPBDRigidsSOAs<T, 3> Particles;
-		TPBDRigidsEvolutionGBF<T, 3> Evolution(Particles);
+		FPBDRigidsEvolutionGBF Evolution(Particles);
 		auto& ClusteredParticles = Particles.GetClusteredParticles();
 		
 		uint32 FirstId = ClusteredParticles.Size();
@@ -66,7 +66,7 @@ namespace ChaosTest {
 	void FractureCluster()
 	{
 		TPBDRigidsSOAs<T, 3> Particles;
-		TPBDRigidsEvolutionGBF<T, 3> Evolution(Particles);
+		FPBDRigidsEvolutionGBF Evolution(Particles);
 		auto& ClusteredParticles = Particles.GetClusteredParticles();
 
 		//create a long row of boxes - the depth 0 cluster is the entire row, the depth 1 clusters 4 boxes each, the depth 2 clusters are 1 box each
@@ -181,7 +181,7 @@ namespace ChaosTest {
 	void PartialFractureCluster()
 	{
 		TPBDRigidsSOAs<T, 3> Particles;
-		TPBDRigidsEvolutionGBF<T, 3> Evolution(Particles);
+		FPBDRigidsEvolutionGBF Evolution(Particles);
 		auto& ClusteredParticles = Particles.GetClusteredParticles();
 
 		//create a long row of boxes - the depth 0 cluster is the entire row, the depth 1 clusters 4 boxes each, the depth 2 clusters are 1 box each
diff --git a/Engine/Source/Programs/HeadlessChaos/Private/HeadlessChaosTestConstraintHandle.cpp b/Engine/Source/Programs/HeadlessChaos/Private/HeadlessChaosTestConstraintHandle.cpp
--- a/Engine/Source/Programs/HeadlessChaos/Private/HeadlessChaosTestConstraintHandle.cpp
+++ b/Engine/Source/Programs/HeadlessChaos/Private/HeadlessChaosTestConstraintHandle.cpp
@@ -186,35 +186,35 @@ namespace ChaosTest {
 
 	TEST(ConstraintHandleTests, DISABLED_CollisionConstraintHandle)
 	{
-		CollisionConstraintHandles<Chaos::TPBDRigidsEvolutionGBF<FReal, 3>>();
+		CollisionConstraintHandles<Chaos::FPBDRigidsEvolutionGBF>();
 
 		SUCCEED();
 	}
 
 	TEST(ConstraintHandleTests, JointConstraintHandle)
 	{
-		JointConstraintHandles<Chaos::TPBDRigidsEvolutionGBF<FReal, 3>>();
+		JointConstraintHandles<Chaos::FPBDRigidsEvolutionGBF>();
 
 		SUCCEED();
 	}
 
 	TEST(ConstraintHandleTests, PositionConstraintHandles)
 	{
-		PositionConstraintHandles<Chaos::TPBDRigidsEvolutionGBF<FReal, 3>>();
+		PositionConstraintHandles<Chaos::FPBDRigidsEvolutionGBF>();
 
 		SUCCEED();
 	}
 
 	TEST(ConstraintHandleTests, RigidSpringConstraintHandles)
 	{
-		RigidSpringConstraintHandles<Chaos::TPBDRigidsEvolutionGBF<FReal, 3>>();
+		RigidSpringConstraintHandles<Chaos::FPBDRigidsEvolutionGBF>();
 
 		SUCCEED();
 	}
 
 	TEST(ConstraintHandleTests, RigidDynamicSpringConstraintHandles)
 	{
-		RigidDynamicSpringConstraintHandles<Chaos::TPBDRigidsEvolutionGBF<FReal, 3>>();
+		RigidDynamicSpringConstraintHandles<Chaos::FPBDRigidsEvolutionGBF>();
 
 		SUCCEED();
 	}
diff --git a/Engine/Source/Programs/HeadlessChaos/Private/HeadlessChaosTestConstraints.cpp b/Engine/Source/Programs/HeadlessChaos/Private/HeadlessChaosTestConstraints.cpp
--- a/Engine/Source/Programs/HeadlessChaos/Private/HeadlessChaosTestConstraints.cpp
+++ b/Engine/Source/Programs/HeadlessChaos/Private/HeadlessChaosTestConstraints.cpp
@@ -43,7 +43,7 @@ namespace ChaosTest {
 		}
 		{
 			TPBDRigidsSOAs<FReal, 3> Particles;
-			TPBDRigidsEvolutionGBF<FReal, 3> Evolution(Particles);
+			FPBDRigidsEvolutionGBF Evolution(Particles);
 			TArray<TPBDRigidParticleHandle<FReal, 3>*> Dynamics = Evolution.CreateDynamicParticles(1);
 			Evolution.GetGravityForces().SetEnabled(*Dynamics[0], false);
 
@@ -116,13 +116,9 @@ namespace ChaosTest {
 
 	TEST(ConstraintTests, Constraints)
 	{
-		ChaosTest::Position<Chaos::TPBDRigidsEvolutionGBF<float, 3>>();
-		ChaosTest::PositionAndJoint<Chaos::TPBDRigidsEvolutionGBF<float, 3>>();
+		ChaosTest::Position<Chaos::FPBDRigidsEvolutionGBF>();
+		ChaosTest::PositionAndJoint<Chaos::FPBDRigidsEvolutionGBF>();
 
-#if CHAOS_PARTICLEHANDLE_TODO
-		ChaosTest::Position<Chaos::TPBDRigidsEvolutionPGS<float, 3>, float>();
-		ChaosTest::PositionAndJoint<Chaos::TPBDRigidsEvolutionPGS<float, 3>>();
-#endif
 		SUCCEED();
 	}
 
diff --git a/Engine/Source/Programs/HeadlessChaos/Private/HeadlessChaosTestForces.cpp b/Engine/Source/Programs/HeadlessChaos/Private/HeadlessChaosTestForces.cpp
--- a/Engine/Source/Programs/HeadlessChaos/Private/HeadlessChaosTestForces.cpp
+++ b/Engine/Source/Programs/HeadlessChaos/Private/HeadlessChaosTestForces.cpp
@@ -23,7 +23,7 @@ namespace ChaosTest {
 		Evolution.AdvanceOneTimeStep(0.1);
 		EXPECT_LT(Dynamics[0]->X()[2], 0);
 	}
-	template void Gravity<Chaos::TPBDRigidsEvolutionGBF<float, 3>, float>();
+	template void Gravity<Chaos::FPBDRigidsEvolutionGBF, float>();
 
 
 }
\ No newline at end of file
diff --git a/Engine/Source/Programs/HeadlessChaos/Private/HeadlessChaosTestJoints.cpp b/Engine/Source/Programs/HeadlessChaos/Private/HeadlessChaosTestJoints.cpp
--- a/Engine/Source/Programs/HeadlessChaos/Private/HeadlessChaosTestJoints.cpp
+++ b/Engine/Source/Programs/HeadlessChaos/Private/HeadlessChaosTestJoints.cpp
@@ -477,7 +477,7 @@ namespace ChaosTest {
 		Box2->X() = FVec3((FReal)500, (FReal)0, (FReal)1000);
 		Box2->P() = Box2->X();
 
-		TPBDRigidsEvolutionGBF<FReal, 3> Evolution(Particles);
+		FPBDRigidsEvolutionGBF Evolution(Particles);
 		TVector<TGeometryParticleHandle<FReal, 3>*,2> ConstrainedParticles = TVector<TGeometryParticleHandle<FReal,3>*, 2>(StaticBox, Box2);
 		TVector<FVec3, 2> Points = { FVec3((FReal)100, (FReal)0, (FReal)1000), FVec3((FReal)400, (FReal)0, (FReal)1000) };
 
@@ -518,7 +518,7 @@ namespace ChaosTest {
 			Box2.X() = FVec3((FReal)500, (FReal)0, (FReal)1000);
 			Box2.P() = Box2.X();
 
-			TPBDRigidsEvolutionGBF<FReal, 3> Evolution(Particles);
+			FPBDRigidsEvolutionGBF Evolution(Particles);
 			TArray<TVector<TGeometryParticleHandle<FReal,3>*, 2>> Constraints = { TVector<TGeometryParticleHandle<FReal,3>*, 2>(&StaticBox, &Box2) };
 
 			Evolution.SetPhysicsMaterial(&StaticBox, MakeSerializable(PhysicalMaterial));
@@ -547,7 +547,7 @@ namespace ChaosTest {
 			Box2.X() = FVec3((FReal)500, (FReal)0, (FReal)1000);
 			Box2.P() = Box2.X();
 
-			TPBDRigidsEvolutionGBF<FReal, 3> Evolution(Particles);
+			FPBDRigidsEvolutionGBF Evolution(Particles);
 			TArray<TVector<TGeometryParticleHandle<FReal,3>*, 2>> Constraints = { TVector<TGeometryParticleHandle<FReal,3>*, 2>(&StaticBox, &Box2) };
 
 			Evolution.SetPhysicsMaterial(&StaticBox, MakeSerializable(PhysicalMaterial));
diff --git a/Engine/Source/Programs/HeadlessChaos/Private/HeadlessChaosTestSerialization.cpp b/Engine/Source/Programs/HeadlessChaos/Private/HeadlessChaosTestSerialization.cpp
--- a/Engine/Source/Programs/HeadlessChaos/Private/HeadlessChaosTestSerialization.cpp
+++ b/Engine/Source/Programs/HeadlessChaos/Private/HeadlessChaosTestSerialization.cpp
@@ -424,7 +424,7 @@ namespace ChaosTest
 			{
 				Chaos::FChaosArchive ChaosAr(*File);
 				TPBDRigidsSOAs<float, 3> Particles;
-				TPBDRigidsEvolutionGBF<float, 3> Evolution(Particles);
+				FPBDRigidsEvolutionGBF Evolution(Particles);
 
 				Evolution.Serialize(ChaosAr);
 				Evolution.AdvanceOneTimeStep(1 / 60.f);
diff --git a/Engine/Source/Programs/HeadlessChaos/Private/HeadlessChaosTestSim.cpp b/Engine/Source/Programs/HeadlessChaos/Private/HeadlessChaosTestSim.cpp
--- a/Engine/Source/Programs/HeadlessChaos/Private/HeadlessChaosTestSim.cpp
+++ b/Engine/Source/Programs/HeadlessChaos/Private/HeadlessChaosTestSim.cpp
@@ -16,7 +16,7 @@ namespace ChaosTest {
 	GTEST_TEST(SimTests, SphereSphereSimTest)
 	{
 		TPBDRigidsSOAs<FReal, 3> Particles;
-		TPBDRigidsEvolutionGBF<FReal, 3> Evolution(Particles);
+		FPBDRigidsEvolutionGBF Evolution(Particles);
 		auto Static = Evolution.CreateStaticParticles(1)[0];
 		auto Dynamic = Evolution.CreateDynamicParticles(1)[0];
 
@@ -45,7 +45,7 @@ namespace ChaosTest {
 	GTEST_TEST(SimTests, BoxBoxSimTest)
 	{
 		TPBDRigidsSOAs<FReal, 3> Particles;
-		TPBDRigidsEvolutionGBF<FReal, 3> Evolution(Particles);
+		FPBDRigidsEvolutionGBF Evolution(Particles);
 		auto Static = Evolution.CreateStaticParticles(1)[0];
 		auto Dynamic = Evolution.CreateDynamicParticles(1)[0];
 
@@ -83,7 +83,7 @@ namespace ChaosTest {
 	GTEST_TEST(SimTests, DISABLED_VeryLowInertiaSimTest)
 	{
 		TPBDRigidsSOAs<FReal, 3> Particles;
-		TPBDRigidsEvolutionGBF<FReal, 3> Evolution(Particles);
+		FPBDRigidsEvolutionGBF Evolution(Particles);
 		auto Static = Evolution.CreateStaticParticles(1)[0];
 		auto Dynamic = Evolution.CreateDynamicParticles(1)[0];
 
@@ -111,7 +111,7 @@ namespace ChaosTest {
 	GTEST_TEST(SimTests, SleepAndWakeSimTest)
 	{
 		TPBDRigidsSOAs<FReal, 3> Particles;
-		TPBDRigidsEvolutionGBF<FReal, 3> Evolution(Particles);
+		FPBDRigidsEvolutionGBF Evolution(Particles);
 		auto Static = Evolution.CreateStaticParticles(1)[0];
 		auto Dynamic1 = Evolution.CreateDynamicParticles(1)[0];
 		auto Dynamic2 = Evolution.CreateDynamicParticles(1)[0];
diff --git a/Engine/Source/Programs/HeadlessChaos/Private/HeadlessChaosTestUtility.cpp b/Engine/Source/Programs/HeadlessChaos/Private/HeadlessChaosTestUtility.cpp
--- a/Engine/Source/Programs/HeadlessChaos/Private/HeadlessChaosTestUtility.cpp
+++ b/Engine/Source/Programs/HeadlessChaos/Private/HeadlessChaosTestUtility.cpp
@@ -45,7 +45,7 @@ namespace ChaosTest {
 	template int32 AppendAnalyticSphere(TPBDRigidParticles<float, 3> &, float Scale);
 
 	template<class T>
-	TPBDRigidParticleHandle<T, 3>* AppendAnalyticSphere2(TPBDRigidsEvolutionGBF<T, 3>& Evolution, T Scale)
+	TPBDRigidParticleHandle<T, 3>* AppendAnalyticSphere2(FPBDRigidsEvolutionGBF& Evolution, T Scale)
 	{
 		TArray<TPBDRigidParticleHandle<T, 3>*> Particles = Evolution.CreateDynamicParticles(1);
 		TPBDRigidParticleHandle<T, 3>* Particle = Particles[0];
@@ -65,7 +65,7 @@ namespace ChaosTest {
 
 		return Particle;
 	}
-	template TPBDRigidParticleHandle<float, 3>* AppendAnalyticSphere2(TPBDRigidsEvolutionGBF<float, 3>& Evolution, float Scale);
+	template TPBDRigidParticleHandle<float, 3>* AppendAnalyticSphere2(FPBDRigidsEvolutionGBF& Evolution, float Scale);
 
 
 	template<class T>
@@ -116,22 +116,22 @@ namespace ChaosTest {
 	}
 
 	template<class T>
-	TKinematicGeometryParticleHandle<T, 3>* AppendKinematicAnalyticBox2(TPBDRigidsEvolutionGBF<T, 3>& Evolution, TVector<T, 3> Scale)
+	TKinematicGeometryParticleHandle<T, 3>* AppendKinematicAnalyticBox2(FPBDRigidsEvolutionGBF& Evolution, TVector<T, 3> Scale)
 	{
 		TArray<TKinematicGeometryParticleHandle<T, 3>*> Particles = Evolution.CreateKinematicParticles(1);
 		InitAnalyticBox2(Particles[0], Scale);
 		return Particles[0];
 	}
-	template TKinematicGeometryParticleHandle<float, 3>* AppendKinematicAnalyticBox2(TPBDRigidsEvolutionGBF<float, 3>& Evolution, TVector<float, 3> Scale);
+	template TKinematicGeometryParticleHandle<float, 3>* AppendKinematicAnalyticBox2(FPBDRigidsEvolutionGBF& Evolution, TVector<float, 3> Scale);
 
 	template<class T>
-	TPBDRigidParticleHandle<T, 3>* AppendDynamicAnalyticBox2(TPBDRigidsEvolutionGBF<T, 3>& Evolution, TVector<T, 3> Scale)
+	TPBDRigidParticleHandle<T, 3>* AppendDynamicAnalyticBox2(FPBDRigidsEvolutionGBF& Evolution, TVector<T, 3> Scale)
 	{
 		TArray<TPBDRigidParticleHandle<T, 3>*> Particles = Evolution.CreateDynamicParticles(1);
 		InitAnalyticBox2(Particles[0], Scale);
 		return Particles[0];
 	}
-	template TPBDRigidParticleHandle<float, 3>* AppendDynamicAnalyticBox2(TPBDRigidsEvolutionGBF<float, 3>& Evolution, TVector<float, 3> Scale);
+	template TPBDRigidParticleHandle<float, 3>* AppendDynamicAnalyticBox2(FPBDRigidsEvolutionGBF& Evolution, TVector<float, 3> Scale);
 
 
 	template<class T>
@@ -412,13 +412,13 @@ namespace ChaosTest {
 
 
 	template<class T>
-	TPBDRigidParticleHandle<T, 3>* AppendDynamicParticleBox(TPBDRigidsEvolutionGBF<T, 3>& Evolution, const TVector<T, 3>& Scale, TArray<TVector<int32, 3>>* OutElements)
+	TPBDRigidParticleHandle<T, 3>* AppendDynamicParticleBox(FPBDRigidsEvolutionGBF& Evolution, const TVector<T, 3>& Scale, TArray<TVector<int32, 3>>* OutElements)
 	{
 		TArray<TPBDRigidParticleHandle<T, 3>*> Particles = Evolution.CreateDynamicParticles(1);
 		InitDynamicParticleBox2(Particles[0], Scale, OutElements);
 		return Particles[0];
 	}
-	template TPBDRigidParticleHandle<float, 3>* AppendDynamicParticleBox(TPBDRigidsEvolutionGBF<float, 3>& Evolution, const TVector<float, 3>& Scale, TArray<Chaos::TVector<int32, 3>>* OutElements);
+	template TPBDRigidParticleHandle<float, 3>* AppendDynamicParticleBox(FPBDRigidsEvolutionGBF& Evolution, const TVector<float, 3>& Scale, TArray<Chaos::TVector<int32, 3>>* OutElements);
 
 	template<class T>
 	TPBDRigidParticleHandle<T, 3>* AppendDynamicParticleBox(TPBDRigidsSOAs<T, 3>& SOAs, const TVector<T, 3>& Scale, TArray<TVector<int32, 3>>* OutElements)
@@ -457,13 +457,13 @@ namespace ChaosTest {
 	template TPBDRigidParticleHandle<float, 3>* AppendDynamicParticleTaperedCylinder(TPBDRigidsSOAs<float, 3>& Evolution, const TVector<float, 3>& Scale, TArray<Chaos::TVector<int32, 3>>* OutElements);
 
 	template<class T>
-	TPBDRigidParticleHandle<T, 3>* AppendClusteredParticleBox(TPBDRigidsEvolutionGBF<T, 3>& Evolution, const TVector<T, 3>& Scale, TArray<TVector<int32, 3>>* OutElements)
+	TPBDRigidParticleHandle<T, 3>* AppendClusteredParticleBox(FPBDRigidsEvolutionGBF& Evolution, const TVector<T, 3>& Scale, TArray<TVector<int32, 3>>* OutElements)
 	{
 		auto Particles = Evolution.CreateClusteredParticles(1);
 		InitDynamicParticleBox2(Particles[0], Scale, OutElements);
 		return Particles[0];
 	}
-	template TPBDRigidParticleHandle<float, 3>* AppendClusteredParticleBox(TPBDRigidsEvolutionGBF<float, 3>& Evolution, const TVector<float, 3>& Scale, TArray<Chaos::TVector<int32, 3>>* OutElements);
+	template TPBDRigidParticleHandle<float, 3>* AppendClusteredParticleBox(FPBDRigidsEvolutionGBF& Evolution, const TVector<float, 3>& Scale, TArray<Chaos::TVector<int32, 3>>* OutElements);
 
 	template<class T>
 	TPBDRigidParticleHandle<T, 3>* AppendClusteredParticleBox(TPBDRigidsSOAs<T, 3>& SOAs, const TVector<T, 3>& Scale, TArray<TVector<int32, 3>>* OutElements)
@@ -555,7 +555,7 @@ namespace ChaosTest {
 	template int AppendStaticAnalyticFloor(TPBDRigidParticles<float, 3> &);
 
 	template<class T>
-	TKinematicGeometryParticleHandle<T, 3>* AppendStaticAnalyticFloor(TPBDRigidsEvolutionGBF<T, 3>& Evolution)
+	TKinematicGeometryParticleHandle<T, 3>* AppendStaticAnalyticFloor(FPBDRigidsEvolutionGBF& Evolution)
 	{
 		TArray<TKinematicGeometryParticleHandle<T, 3>*> Particles = Evolution.CreateKinematicParticles(1);
 		TKinematicGeometryParticleHandle<T, 3>* Particle = Particles[0];
@@ -570,7 +570,7 @@ namespace ChaosTest {
 
 		return Particle;
 	}
-	template TKinematicGeometryParticleHandle<float, 3>* AppendStaticAnalyticFloor(TPBDRigidsEvolutionGBF<float, 3>& Evolution);
+	template TKinematicGeometryParticleHandle<float, 3>* AppendStaticAnalyticFloor(FPBDRigidsEvolutionGBF& Evolution);
 
 	template<class T>
 	TKinematicGeometryParticleHandle<T, 3>* AppendStaticAnalyticFloor(TPBDRigidsSOAs<T, 3>& SOAs)
diff --git a/Engine/Source/Programs/HeadlessChaos/Public/HeadlessChaosTestConstraints.h b/Engine/Source/Programs/HeadlessChaos/Public/HeadlessChaosTestConstraints.h
--- a/Engine/Source/Programs/HeadlessChaos/Public/HeadlessChaosTestConstraints.h
+++ b/Engine/Source/Programs/HeadlessChaos/Public/HeadlessChaosTestConstraints.h
@@ -61,7 +61,7 @@ namespace ChaosTest
 
 		// Solver state
 		TPBDRigidsSOAs<FReal, 3> SOAs;
-		TPBDRigidsEvolutionGBF<FReal, 3> Evolution;
+		FPBDRigidsEvolutionGBF Evolution;
 		TUniquePtr<FChaosPhysicsMaterial> PhysicalMaterial;
 
 		TGeometryParticleHandle<FReal, 3>* GetParticle(const int32 Idx)
diff --git a/Engine/Source/Programs/HeadlessChaos/Public/HeadlessChaosTestUtility.h b/Engine/Source/Programs/HeadlessChaos/Public/HeadlessChaosTestUtility.h
--- a/Engine/Source/Programs/HeadlessChaos/Public/HeadlessChaosTestUtility.h
+++ b/Engine/Source/Programs/HeadlessChaos/Public/HeadlessChaosTestUtility.h
@@ -19,7 +19,7 @@
 
 namespace Chaos
 {
-	template<class T, int d> class TPBDRigidsEvolutionGBF;
+	class FPBDRigidsEvolutionGBF;
 	template <typename T, int d> class TPBDRigidsSOAs;
 }
 
@@ -45,21 +45,21 @@ namespace ChaosTest {
 	template<class T>
 	int32 AppendAnalyticSphere(TPBDRigidParticles<T, 3> & InParticles, T Scale = (T)1);
 	template<class T>
-	TPBDRigidParticleHandle<T, 3>* AppendAnalyticSphere2(TPBDRigidsEvolutionGBF<T, 3>& Evolution, T Scale = (T)1);
+	TPBDRigidParticleHandle<T, 3>* AppendAnalyticSphere2(FPBDRigidsEvolutionGBF& Evolution, T Scale = (T)1);
 
 	/**/
 	template<class T>
 	int32 AppendAnalyticBox(TPBDRigidParticles<T, 3>& InParticles, TVector<T, 3> Scale = TVector<T, 3>(1));
 	template<class T>
-	TKinematicGeometryParticleHandle<T, 3>* AppendKinematicAnalyticBox2(TPBDRigidsEvolutionGBF<T, 3>& Evolution, TVector<T, 3> Scale = TVector<T, 3>(1));
+	TKinematicGeometryParticleHandle<T, 3>* AppendKinematicAnalyticBox2(FPBDRigidsEvolutionGBF& Evolution, TVector<T, 3> Scale = TVector<T, 3>(1));
 	template<class T>
-	TPBDRigidParticleHandle<T, 3>* AppendDynamicAnalyticBox2(TPBDRigidsEvolutionGBF<T, 3>& Evolution, TVector<T, 3> Scale = TVector<T, 3>(1));
+	TPBDRigidParticleHandle<T, 3>* AppendDynamicAnalyticBox2(FPBDRigidsEvolutionGBF& Evolution, TVector<T, 3> Scale = TVector<T, 3>(1));
 
 	/**/
 	template<class T>
 	int32 AppendParticleBox(TPBDRigidParticles<T, 3>& InParticles, TVector<T, 3> Scale = TVector<T, 3>(1), TArray<TVector<int32, 3>>* OutElements = nullptr);
 	template<class T>
-	TPBDRigidParticleHandle<T, 3>* AppendDynamicParticleBox(TPBDRigidsEvolutionGBF<T, 3>& Evolution, const TVector<T, 3>& Scale = TVector<T, 3>(1), TArray<TVector<int32, 3>>* OutElements = nullptr);
+	TPBDRigidParticleHandle<T, 3>* AppendDynamicParticleBox(FPBDRigidsEvolutionGBF& Evolution, const TVector<T, 3>& Scale = TVector<T, 3>(1), TArray<TVector<int32, 3>>* OutElements = nullptr);
 
 	template<class T>
 	TPBDRigidParticleHandle<T, 3>* AppendDynamicParticleBox(TPBDRigidsSOAs<T, 3>& SOAs, const TVector<T, 3>& Scale = TVector<T, 3>(1), TArray<TVector<int32, 3>>* OutElements = nullptr);
@@ -77,7 +77,7 @@ namespace ChaosTest {
 	TGeometryParticleHandle<T, 3>* AppendStaticParticleBox(TPBDRigidsSOAs<T, 3>& SOAs, const TVector<T, 3>& Scale = TVector<T, 3>(1), TArray<TVector<int32, 3>>* OutElements = nullptr);
 
 	template<class T>
-	TPBDRigidParticleHandle<T, 3>* AppendClusteredParticleBox(TPBDRigidsEvolutionGBF<T, 3>& Evolution, const TVector<T, 3>& Scale = TVector<T, 3>(1), TArray<TVector<int32, 3>>* OutElements = nullptr);
+	TPBDRigidParticleHandle<T, 3>* AppendClusteredParticleBox(FPBDRigidsEvolutionGBF& Evolution, const TVector<T, 3>& Scale = TVector<T, 3>(1), TArray<TVector<int32, 3>>* OutElements = nullptr);
 
 	template<class T>
 	TPBDRigidParticleHandle<T, 3>* AppendClusteredParticleBox(TPBDRigidsSOAs<T, 3>& SOAs, const TVector<T, 3>& Scale = TVector<T, 3>(1), TArray<TVector<int32, 3>>* OutElements = nullptr);
@@ -86,7 +86,7 @@ namespace ChaosTest {
 	template<class T>
 	int32 AppendStaticAnalyticFloor(TPBDRigidParticles<T, 3>& InParticles);
 	template<class T>
-	TKinematicGeometryParticleHandle<T, 3>* AppendStaticAnalyticFloor(TPBDRigidsEvolutionGBF<T, 3>& Evolution);
+	TKinematicGeometryParticleHandle<T, 3>* AppendStaticAnalyticFloor(FPBDRigidsEvolutionGBF& Evolution);
 
 	template<class T>
 	TKinematicGeometryParticleHandle<T, 3>* AppendStaticAnalyticFloor(TPBDRigidsSOAs<T, 3>& SOAs);
diff --git a/Engine/Source/Runtime/Engine/Private/PhysicsEngine/PhysUtils.cpp b/Engine/Source/Runtime/Engine/Private/PhysicsEngine/PhysUtils.cpp
--- a/Engine/Source/Runtime/Engine/Private/PhysicsEngine/PhysUtils.cpp
+++ b/Engine/Source/Runtime/Engine/Private/PhysicsEngine/PhysUtils.cpp
@@ -497,7 +497,7 @@ bool FPhysicsInterface::ExecPhysCommands(const TCHAR* Cmd, FOutputDevice* Output
 		Chaos::FChaosArchive Ar(BaseAr);
 		FPhysScene* PhysScene = InWorld->GetPhysicsScene();
 		Chaos::FPhysicsSolver* Solver = PhysScene->GetSolver();
-		Chaos::TPBDRigidsEvolutionGBF<float, 3>* Evolution = Solver->GetEvolution();
+		Chaos::FPBDRigidsEvolutionGBF* Evolution = Solver->GetEvolution();
 		Evolution->Serialize(Ar);
 		TUniquePtr<Chaos::FChaosArchiveContext> ArchiveContext = Ar.StealContext();
 
diff --git a/Engine/Source/Runtime/Engine/Public/Physics/ImmediatePhysics/ImmediatePhysicsChaos/ImmediatePhysicsCore_Chaos.h b/Engine/Source/Runtime/Engine/Public/Physics/ImmediatePhysics/ImmediatePhysicsChaos/ImmediatePhysicsCore_Chaos.h
--- a/Engine/Source/Runtime/Engine/Public/Physics/ImmediatePhysics/ImmediatePhysicsChaos/ImmediatePhysicsCore_Chaos.h
+++ b/Engine/Source/Runtime/Engine/Public/Physics/ImmediatePhysics/ImmediatePhysicsChaos/ImmediatePhysicsCore_Chaos.h
@@ -13,11 +13,11 @@ namespace Chaos
 	class FImplicitObject;
 	class FPBDJointConstraintHandle;
 	class FPBDJointConstraints;
+	class FPBDRigidsEvolutionGBF;
 	template<typename T, int D> struct TKinematicGeometryParticleParameters;
 	template<typename T, int D> class TKinematicTarget;
 	template<typename P> class TPBDConstraintIslandRule;
 	template<typename T, int D> struct TPBDRigidParticleParameters;
-	template<typename T, int D> class TPBDRigidsEvolutionGBF;
 	template<typename T, int D> class TPBDRigidsSOAs;
 	template<typename T, int D> class TPerShapeData;
 }
diff --git a/Engine/Source/Runtime/Experimental/Chaos/Private/Chaos/PBDRigidClustering.cpp b/Engine/Source/Runtime/Experimental/Chaos/Private/Chaos/PBDRigidClustering.cpp
--- a/Engine/Source/Runtime/Experimental/Chaos/Private/Chaos/PBDRigidClustering.cpp
+++ b/Engine/Source/Runtime/Experimental/Chaos/Private/Chaos/PBDRigidClustering.cpp
@@ -8,7 +8,6 @@
 #include "Chaos/Levelset.h"
 #include "Chaos/MassProperties.h"
 #include "Chaos/PBDRigidsEvolution.h"
-#include "Chaos/PBDRigidsEvolutionPGS.h"
 #include "Chaos/PBDCollisionConstraints.h"
 #include "Chaos/PBDCollisionConstraintsPGS.h"
 #include "Chaos/Sphere.h"
@@ -114,7 +113,7 @@ namespace Chaos
 
 	DECLARE_CYCLE_STAT(TEXT("TPBDRigidClustering<>::RewindAndEvolve<BGF>()"), STAT_RewindAndEvolve_BGF, STATGROUP_Chaos);
 	template<class T, int d>
-	void RewindAndEvolve(TPBDRigidsEvolutionGBF<T, d>& Evolution, TPBDRigidClusteredParticles<T, d>& InParticles, const TSet<int32>& IslandsToRecollide, const TSet<uint32>& AllActivatedChildren, const T Dt, TPBDCollisionConstraints<T, d>& CollisionRule)
+	void RewindAndEvolve(FPBDRigidsEvolutionGBF& Evolution, TPBDRigidClusteredParticles<T, d>& InParticles, const TSet<int32>& IslandsToRecollide, const TSet<uint32>& AllActivatedChildren, const T Dt, TPBDCollisionConstraints<T, d>& CollisionRule)
 	{
 		SCOPE_CYCLE_COUNTER(STAT_RewindAndEvolve_BGF);
 		// Rewind active particles
@@ -188,76 +187,6 @@ namespace Chaos
 		}
 	}
 
-#if CHAOS_PARTICLEHANDLE_TODO
-	DECLARE_CYCLE_STAT(TEXT("TPBDRigidClustering<>::RewindAndEvolve<PSG>()"), STAT_RewindAndEvolve_PSG, STATGROUP_Chaos);
-	template<class T, int d>
-	void RewindAndEvolve(TPBDRigidsEvolutionPGS<T, d>& Evolution, TPBDRigidClusteredParticles<T, d>& InParticles, const TSet<int32>& IslandsToRecollide, const TSet<uint32>& AllActivatedChildren, const T Dt, TPBDCollisionConstraintPGS<T, d>& CollisionRule)
-	{
-		SCOPE_CYCLE_COUNTER(STAT_RewindAndEvolve_PSG);
-		// Rewind active particles
-		TArray<int32>& NonDisabledIndices = Evolution.GetNonDisabledIndices();
-		PhysicsParallelFor(IslandsToRecollide.Num(), [&](int32 Island) {
-			TArray<int32> ParticleIndices = Evolution.GetIslandParticles(Island); // copy
-			for (int32 ArrayIdx = ParticleIndices.Num() - 1; ArrayIdx >= 0; --ArrayIdx)
-			{
-				int32 Index = ParticleIndices[ArrayIdx];
-				if (InParticles.Sleeping(Index) || InParticles.Disabled(Index))
-				{
-					ParticleIndices.RemoveAtSwap(ArrayIdx);
-				}
-				else
-				{
-					InParticles.P(Index) = InParticles.X(Index);
-					InParticles.Q(Index) = InParticles.R(Index);
-					InParticles.V(Index) = InParticles.PreV(Index);
-					InParticles.W(Index) = InParticles.PreW(Index);
-				}
-			}
-			Evolution.IntegrateV(ParticleIndices, Dt);
-		});
-
-		TSet<uint32> AllIslandParticles;
-		for (const auto& Island : IslandsToRecollide)
-		{
-			const TArray<int32>& ParticleIndices = Evolution.GetIslandParticles(Island);
-			for (const auto& Index : ParticleIndices)
-			{
-				if (InParticles.Disabled(Index) == false) //HACK: cluster code is incorrectly adding disabled children
-				{
-					if (!AllIslandParticles.Contains(Index))
-					{
-						AllIslandParticles.Add(Index);
-						NonDisabledIndices.Add(Index);
-					}
-				}
-				else
-				{
-					//FPlatformMisc::DebugBreak();
-				}
-			}
-		}
-
-		// @todo(mlentine): We can precompute internal constraints which can filter some from the narrow phase tests but may not help much
-		CollisionRule.UpdateConstraints(InParticles, Dt, AllActivatedChildren, AllIslandParticles.Array());
-
-		Evolution.InitializeAccelerationStructures();
-
-		PhysicsParallelFor(Evolution.NumIslands(), [&](int32 Island) {
-			Evolution.UpdateAccelerationStructures(Island);
-			Evolution.ApplyConstraints(Dt, Island);
-		});
-		PhysicsParallelFor(Evolution.NumIslands(), [&](int32 Island) {
-			const TArray<int32>& ParticleIndices = Evolution.GetIslandParticles(Island);
-			Evolution.IntegrateX(ParticleIndices, Dt);
-		});
-
-		// @todo(mlentine): Need to enforce constraints
-		PhysicsParallelFor(Evolution.NumIslands(), [&](int32 Island) {
-			Evolution.ApplyPushOut(Dt, Island);
-		});
-	}
-#endif
-
 	DECLARE_CYCLE_STAT(TEXT("TPBDRigidClustering<>::UpdateClusterMassProperties()"), STAT_UpdateClusterMassProperties, STATGROUP_Chaos);
 	template<class T, int d>
 	void UpdateClusterMassProperties(TPBDRigidClusteredParticles<T, d>& Particles, const TArray<uint32>& Children, const uint32 NewIndex, const TRigidTransform<T, d>* ForceMassOrientation, const TArrayCollectionArray<TUniquePtr<TMultiChildProxyData<T, d>>>* MMultiChildProxyData, const TArrayCollectionArray<FMultiChildProxyId>* MMultiChildProxyId)
@@ -1925,10 +1854,7 @@ namespace Chaos
 } // namespace Chaos
 
 using namespace Chaos;
-template class CHAOS_API Chaos::TPBDRigidClustering<TPBDRigidsEvolutionGBF<float, 3>, TPBDCollisionConstraints<float, 3>, float, 3>;
-#if CHAOS_PARTICLEHANDLE_TODO
-template class CHAOS_API Chaos::TPBDRigidClustering<TPBDRigidsEvolutionPGS<float, 3>, TPBDCollisionConstraintPGS<float, 3>, float, 3>;
-#endif
+template class CHAOS_API Chaos::TPBDRigidClustering<FPBDRigidsEvolutionGBF, TPBDCollisionConstraints<float, 3>, float, 3>;
 template CHAOS_API void Chaos::UpdateClusterMassProperties<float, 3>(TPBDRigidClusteredParticles<float, 3>& Particles, const TArray<uint32>& Children, const uint32 ClusterIndex, const TRigidTransform<float, 3>* ForceMassOrientation,
 	const TArrayCollectionArray<TUniquePtr<TMultiChildProxyData<float, 3>>>* MMultiChildProxyData,
 	const TArrayCollectionArray<FMultiChildProxyId>* MMultiChildProxyId);
diff --git a/Engine/Source/Runtime/Experimental/Chaos/Private/Chaos/PBDRigidsEvolution.cpp b/Engine/Source/Runtime/Experimental/Chaos/Private/Chaos/PBDRigidsEvolution.cpp
--- a/Engine/Source/Runtime/Experimental/Chaos/Private/Chaos/PBDRigidsEvolution.cpp
+++ b/Engine/Source/Runtime/Experimental/Chaos/Private/Chaos/PBDRigidsEvolution.cpp
@@ -51,18 +51,17 @@ namespace Chaos
 	FAutoConsoleVariableRef CVarMaxPayloadSize(TEXT("p.MaxPayloadSize"), ConfigSettings.MaxPayloadSize, TEXT(""));
 	FAutoConsoleVariableRef CVarIterationsPerTimeSlice(TEXT("p.IterationsPerTimeSlice"), ConfigSettings.IterationsPerTimeSlice, TEXT(""));
 
-	template<typename T, int d>
-	struct TDefaultCollectionFactory : public ISpatialAccelerationCollectionFactory<T, d>
+	struct FDefaultCollectionFactory : public ISpatialAccelerationCollectionFactory
 	{
 		FAccelerationConfig Config;
 
-		using BVType = TBoundingVolume<TAccelerationStructureHandle<T, d>, T, d>;
-		using AABBTreeType = TAABBTree<TAccelerationStructureHandle<T, d>, TAABBTreeLeafArray<TAccelerationStructureHandle<T, d>, T>, T>;
-		using AABBTreeOfGridsType = TAABBTree<TAccelerationStructureHandle<T, d>, TBoundingVolume<TAccelerationStructureHandle<T, d>, T, d>, T>;
+		using BVType = TBoundingVolume<TAccelerationStructureHandle<FReal, 3>, FReal, 3>;
+		using AABBTreeType = TAABBTree<TAccelerationStructureHandle<FReal, 3>, TAABBTreeLeafArray<TAccelerationStructureHandle<FReal, 3>, FReal>, FReal>;
+		using AABBTreeOfGridsType = TAABBTree<TAccelerationStructureHandle<FReal, 3>, TBoundingVolume<TAccelerationStructureHandle<FReal, 3>, FReal, 3>, FReal>;
 
-		TUniquePtr<ISpatialAccelerationCollection<TAccelerationStructureHandle<T, d>, T, d>> CreateEmptyCollection() override
+		TUniquePtr<ISpatialAccelerationCollection<TAccelerationStructureHandle<FReal, 3>, FReal, 3>> CreateEmptyCollection() override
 		{
-			TConstParticleView<TSpatialAccelerationCache<T, d>> Empty;
+			TConstParticleView<FSpatialAccelerationCache> Empty;
 
 			const uint16 NumBuckets = ConfigSettings.BroadphaseType >= 3 ? 2 : 1;
 			auto Collection = new TSpatialAccelerationCollection<AABBTreeType, BVType, AABBTreeOfGridsType>();
@@ -72,7 +71,7 @@ namespace Chaos
 				Collection->AddSubstructure(CreateAccelerationPerBucket_Threaded(Empty, BucketIdx, true), BucketIdx);
 			}
 
-			return TUniquePtr<ISpatialAccelerationCollection<TAccelerationStructureHandle<T, d>, T, d>>(Collection);
+			return TUniquePtr<ISpatialAccelerationCollection<TAccelerationStructureHandle<FReal, 3>, FReal, 3>>(Collection);
 		}
 
 		virtual uint8 GetActiveBucketsMask() const
@@ -80,7 +79,7 @@ namespace Chaos
 			return ConfigSettings.BroadphaseType >= 3 ? 3 : 1;
 		}
 
-		virtual TUniquePtr<ISpatialAcceleration<TAccelerationStructureHandle<T, d>, T, d>> CreateAccelerationPerBucket_Threaded(const TConstParticleView<TSpatialAccelerationCache<T, d>>& Particles, uint16 BucketIdx, bool ForceFullBuild) override
+		virtual TUniquePtr<ISpatialAcceleration<TAccelerationStructureHandle<FReal, 3>, FReal, 3>> CreateAccelerationPerBucket_Threaded(const TConstParticleView<FSpatialAccelerationCache>& Particles, uint16 BucketIdx, bool ForceFullBuild) override
 		{
 			switch (BucketIdx)
 			{
@@ -112,7 +111,7 @@ namespace Chaos
 			}
 		}
 
-		virtual void Serialize(TUniquePtr<ISpatialAccelerationCollection<TAccelerationStructureHandle<T, d>, T, d>>& Ptr, FChaosArchive& Ar) override
+		virtual void Serialize(TUniquePtr<ISpatialAccelerationCollection<TAccelerationStructureHandle<FReal, 3>, FReal, 3>>& Ptr, FChaosArchive& Ar) override
 		{
 			if (Ar.IsLoading())
 			{
@@ -126,15 +125,13 @@ namespace Chaos
 		}
 	};
 
-	template<class FPBDRigidsEvolution, class FPBDCollisionConstraint, class T, int d>
-	TPBDRigidsEvolutionBase<FPBDRigidsEvolution, FPBDCollisionConstraint, T, d>::TPBDRigidsEvolutionBase(TPBDRigidsSOAs<T, d>& InParticles, int32 InNumIterations, int32 InNumPushOutIterations, bool InIsSingleThreaded)
+	FPBDRigidsEvolutionBase::FPBDRigidsEvolutionBase(TPBDRigidsSOAs<FReal, 3>& InParticles, int32 InNumIterations, int32 InNumPushOutIterations, bool InIsSingleThreaded)
 		: Particles(InParticles)
 		, bExternalReady(false)
 		, bIsSingleThreaded(InIsSingleThreaded)
-		, Clustering(static_cast<FPBDRigidsEvolution&>(*this), Particles.GetClusteredParticles())
 		, NumIterations(InNumIterations)
 		, NumPushOutIterations(InNumPushOutIterations)
-		, SpatialCollectionFactory(new TDefaultCollectionFactory<T, d>())
+		, SpatialCollectionFactory(new FDefaultCollectionFactory())
 	{
 		Particles.GetParticleHandles().AddArray(&PhysicsMaterials);
 		Particles.GetParticleHandles().AddArray(&PerParticlePhysicsMaterials);
@@ -149,8 +146,7 @@ namespace Chaos
 		ComputeIntermediateSpatialAcceleration();
 	}
 
-	template<class FPBDRigidsEvolution, class FPBDCollisionConstraint, class T, int d>
-	TPBDRigidsEvolutionBase<FPBDRigidsEvolution, FPBDCollisionConstraint, T, d>::~TPBDRigidsEvolutionBase()
+	FPBDRigidsEvolutionBase::~FPBDRigidsEvolutionBase()
 	{
 		Particles.GetParticleHandles().RemoveArray(&PhysicsMaterials);
 		Particles.GetParticleHandles().RemoveArray(&PerParticlePhysicsMaterials);
@@ -168,10 +164,9 @@ namespace Chaos
 	DECLARE_CYCLE_STAT(TEXT("CreateInitialAccelerationStructure"), STAT_CreateInitialAccelerationStructure, STATGROUP_Chaos);
 
 
-	template<class FPBDRigidsEvolution, class FPBDCollisionConstraint, class T, int d>
-	TPBDRigidsEvolutionBase<FPBDRigidsEvolution, FPBDCollisionConstraint, T, d>::FChaosAccelerationStructureTask::FChaosAccelerationStructureTask(
-		ISpatialAccelerationCollectionFactory<T, d>& InSpatialCollectionFactory
-		, const TMap<FSpatialAccelerationIdx, TUniquePtr<TSpatialAccelerationCache<T, d>>>& InSpatialAccelerationCache
+	FPBDRigidsEvolutionBase::FChaosAccelerationStructureTask::FChaosAccelerationStructureTask(
+		ISpatialAccelerationCollectionFactory& InSpatialCollectionFactory
+		, const TMap<FSpatialAccelerationIdx, TUniquePtr<FSpatialAccelerationCache>>& InSpatialAccelerationCache
 		, TUniquePtr<FAccelerationStructure>& InAccelerationStructure
 		, TUniquePtr<FAccelerationStructure>& InAccelerationStructureCopy
 		, bool InForceFullBuild
@@ -185,29 +180,25 @@ namespace Chaos
 	{
 	}
 
-	template<class FPBDRigidsEvolution, class FPBDCollisionConstraint, class T, int d>
-	TStatId TPBDRigidsEvolutionBase<FPBDRigidsEvolution, FPBDCollisionConstraint, T, d>::FChaosAccelerationStructureTask::GetStatId()
+	TStatId FPBDRigidsEvolutionBase::FChaosAccelerationStructureTask::GetStatId()
 	{
 		RETURN_QUICK_DECLARE_CYCLE_STAT(FChaosAccelerationStructureTask, STATGROUP_Chaos);
 	}
 
-	template<class FPBDRigidsEvolution, class FPBDCollisionConstraint, class T, int d>
-	ENamedThreads::Type TPBDRigidsEvolutionBase<FPBDRigidsEvolution, FPBDCollisionConstraint, T, d>::FChaosAccelerationStructureTask::GetDesiredThread()
+	ENamedThreads::Type FPBDRigidsEvolutionBase::FChaosAccelerationStructureTask::GetDesiredThread()
 	{
 		return ENamedThreads::AnyBackgroundThreadNormalTask;
 	}
 
-	template<class FPBDRigidsEvolution, class FPBDCollisionConstraint, class T, int d>
-	ESubsequentsMode::Type TPBDRigidsEvolutionBase<FPBDRigidsEvolution, FPBDCollisionConstraint, T, d>::FChaosAccelerationStructureTask::GetSubsequentsMode()
+	ESubsequentsMode::Type FPBDRigidsEvolutionBase::FChaosAccelerationStructureTask::GetSubsequentsMode()
 	{
 		return ESubsequentsMode::TrackSubsequents;
 	}
 
-	template<typename T, int d>
-	TUniquePtr<ISpatialAccelerationCollection<TAccelerationStructureHandle<T, d>, T, d>> CreateNewSpatialStructureFromSubStructure(TUniquePtr<ISpatialAcceleration<TAccelerationStructureHandle<T, d>, T, d>>&& Substructure)
+	TUniquePtr<ISpatialAccelerationCollection<TAccelerationStructureHandle<FReal, 3>, FReal, 3>> CreateNewSpatialStructureFromSubStructure(TUniquePtr<ISpatialAcceleration<TAccelerationStructureHandle<FReal, 3>, FReal, 3>>&& Substructure)
 	{
-		using BVType = TBoundingVolume<TAccelerationStructureHandle<T, d>, T, d>;
-		using AABBType = TAABBTree<TAccelerationStructureHandle<T, d>, TAABBTreeLeafArray<TAccelerationStructureHandle<T, d>, T>, T>;
+		using BVType = TBoundingVolume<TAccelerationStructureHandle<FReal, 3>, FReal, 3>;
+		using AABBType = TAABBTree<TAccelerationStructureHandle<FReal, 3>, TAABBTreeLeafArray<TAccelerationStructureHandle<FReal, 3>, FReal>, FReal>;
 
 		if (Substructure->template As<BVType>())
 		{
@@ -223,20 +214,19 @@ namespace Chaos
 		}
 		else
 		{
-			using AccelType = TAABBTree<TAccelerationStructureHandle<T, d>, TBoundingVolume<TAccelerationStructureHandle<T, d>, T, d>, T>;
+			using AccelType = TAABBTree<TAccelerationStructureHandle<FReal, 3>, TBoundingVolume<TAccelerationStructureHandle<FReal, 3>, FReal, 3>, FReal>;
 			auto Collection = MakeUnique<TSpatialAccelerationCollection<AccelType>>();
 			Collection->AddSubstructure(MoveTemp(Substructure), 0);
 			return Collection;
 		}
 	}
 
-	template<class FPBDRigidsEvolution, class FPBDCollisionConstraint, class T, int d>
-	void TPBDRigidsEvolutionBase<FPBDRigidsEvolution, FPBDCollisionConstraint, T, d>::FChaosAccelerationStructureTask::DoTask(ENamedThreads::Type CurrentThread, const FGraphEventRef& MyCompletionGraphEvent)
+	void FPBDRigidsEvolutionBase::FChaosAccelerationStructureTask::DoTask(ENamedThreads::Type CurrentThread, const FGraphEventRef& MyCompletionGraphEvent)
 	{
 		LLM_SCOPE(ELLMTag::ChaosAcceleration);
 
 		uint8 ActiveBucketsMask = SpatialCollectionFactory.GetActiveBucketsMask();
-		TArray<TSOAView<TSpatialAccelerationCache<T, d>>> ViewsPerBucket[8];
+		TArray<TSOAView<FSpatialAccelerationCache>> ViewsPerBucket[8];
 
 		bool IsTimeSlicingProgressing = false;
 
@@ -244,7 +234,7 @@ namespace Chaos
 		for (const auto& Itr : SpatialAccelerationCache)
 		{
 			const FSpatialAccelerationIdx SpatialIdx = Itr.Key;
-			const TSpatialAccelerationCache<T, d>& Cache = *Itr.Value;
+			const FSpatialAccelerationCache& Cache = *Itr.Value;
 			const uint8 BucketIdx = (1 << SpatialIdx.Bucket) & ActiveBucketsMask ? SpatialIdx.Bucket : 0;
 			if (AccelerationStructure->GetSubstructure(SpatialIdx) && !AccelerationStructure->GetSubstructure(SpatialIdx)->IsAsyncTimeSlicingComplete())
 			{
@@ -257,7 +247,7 @@ namespace Chaos
 			}
 			else
 			{
-				ViewsPerBucket[BucketIdx].Add(const_cast<TSpatialAccelerationCache<T, d>*>(&Cache));
+				ViewsPerBucket[BucketIdx].Add(const_cast<FSpatialAccelerationCache*>(&Cache));
 				if (AccelerationStructure->IsBucketActive(SpatialIdx.Bucket))
 				{
 					AccelerationStructure->RemoveSubstructure(SpatialIdx);
@@ -300,8 +290,7 @@ namespace Chaos
 		}
 	}
 
-	template<class FPBDRigidsEvolution, class FPBDCollisionConstraint, class T, int d>
-	void TPBDRigidsEvolutionBase<FPBDRigidsEvolution, FPBDCollisionConstraint, T, d>::ApplyParticlePendingData(const FPendingSpatialData& SpatialData, FAccelerationStructure& AccelerationStructure, bool bUpdateCache)
+	void FPBDRigidsEvolutionBase::ApplyParticlePendingData(const FPendingSpatialData& SpatialData, FAccelerationStructure& AccelerationStructure, bool bUpdateCache)
 	{
 		if (SpatialData.bDelete)
 		{
@@ -311,7 +300,7 @@ namespace Chaos
 			{
 				if (uint32* InnerIdxPtr = ParticleToCacheInnerIdx.Find(SpatialData.UniqueIdx()))
 				{
-					TSpatialAccelerationCache<T, d>& Cache = *SpatialAccelerationCache.FindChecked(SpatialData.SpatialIdx);	//can't delete from cache that doesn't exist
+					FSpatialAccelerationCache& Cache = *SpatialAccelerationCache.FindChecked(SpatialData.SpatialIdx);	//can't delete from cache that doesn't exist
 					const uint32 CacheInnerIdx = *InnerIdxPtr;
 					if (CacheInnerIdx + 1 < Cache.Size())	//will get swapped with last element, so update it
 					{
@@ -326,19 +315,19 @@ namespace Chaos
 		}
 		else
 		{
-			TGeometryParticleHandle<T, d>* UpdateParticle = SpatialData.AccelerationHandle.GetGeometryParticleHandle_PhysicsThread();
+			TGeometryParticleHandle<FReal, 3>* UpdateParticle = SpatialData.AccelerationHandle.GetGeometryParticleHandle_PhysicsThread();
 
 			AccelerationStructure.UpdateElementIn(UpdateParticle, UpdateParticle->WorldSpaceInflatedBounds(), UpdateParticle->HasBounds(), SpatialData.SpatialIdx);
 
 			if (bUpdateCache)
 			{
-				TUniquePtr<TSpatialAccelerationCache<T, d>>* CachePtrPtr = SpatialAccelerationCache.Find(SpatialData.SpatialIdx);
+				TUniquePtr<FSpatialAccelerationCache>* CachePtrPtr = SpatialAccelerationCache.Find(SpatialData.SpatialIdx);
 				if (CachePtrPtr == nullptr)
 				{
-					CachePtrPtr = &SpatialAccelerationCache.Add(SpatialData.SpatialIdx, TUniquePtr<TSpatialAccelerationCache<T, d>>(new TSpatialAccelerationCache<T, d>));
+					CachePtrPtr = &SpatialAccelerationCache.Add(SpatialData.SpatialIdx, TUniquePtr<FSpatialAccelerationCache>(new FSpatialAccelerationCache));
 				}
 
-				TSpatialAccelerationCache<T, d>& Cache = **CachePtrPtr;
+				FSpatialAccelerationCache& Cache = **CachePtrPtr;
 
 				//make sure in mapping
 				uint32 CacheInnerIdx;
@@ -361,8 +350,7 @@ namespace Chaos
 		}
 	}
 
-	template<class FPBDRigidsEvolution, class FPBDCollisionConstraint, class T, int d>
-	void TPBDRigidsEvolutionBase<FPBDRigidsEvolution, FPBDCollisionConstraint, T, d>::FlushInternalAccelerationQueue()
+	void FPBDRigidsEvolutionBase::FlushInternalAccelerationQueue()
 	{
 		for (const FPendingSpatialData& PendingData : InternalAccelerationQueue.PendingData)
 		{
@@ -371,8 +359,7 @@ namespace Chaos
 		InternalAccelerationQueue.Reset();
 	}
 
-	template<class FPBDRigidsEvolution, class FPBDCollisionConstraint, class T, int d>
-	void TPBDRigidsEvolutionBase<FPBDRigidsEvolution, FPBDCollisionConstraint, T, d>::FlushAsyncAccelerationQueue()
+	void FPBDRigidsEvolutionBase::FlushAsyncAccelerationQueue()
 	{
 		for (const FPendingSpatialData& PendingData : AsyncAccelerationQueue.PendingData)
 		{
@@ -399,8 +386,7 @@ namespace Chaos
 		ExternalAccelerationQueue.Reset();
 	}
 
-	template<class FPBDRigidsEvolution, class FPBDCollisionConstraint, class T, int d>
-	void TPBDRigidsEvolutionBase<FPBDRigidsEvolution, FPBDCollisionConstraint, T, d>::FlushExternalAccelerationQueue(FAccelerationStructure& Acceleration)
+	void FPBDRigidsEvolutionBase::FlushExternalAccelerationQueue(FAccelerationStructure& Acceleration)
 	{
 		for (const FPendingSpatialData& PendingData : ExternalAccelerationQueue.PendingData)
 		{
@@ -409,8 +395,7 @@ namespace Chaos
 		ExternalAccelerationQueue.Reset();
 	}
 
-	template<class FPBDRigidsEvolution, class FPBDCollisionConstraint, class T, int d>
-	void TPBDRigidsEvolutionBase<FPBDRigidsEvolution, FPBDCollisionConstraint, T, d>::WaitOnAccelerationStructure()
+	void FPBDRigidsEvolutionBase::WaitOnAccelerationStructure()
 	{
 		if (AccelerationStructureTaskComplete.GetReference())
 		{
@@ -421,8 +406,7 @@ namespace Chaos
 		}
 	}
 
-	template<class FPBDRigidsEvolution, class FPBDCollisionConstraint, class T, int d>
-	void TPBDRigidsEvolutionBase<FPBDRigidsEvolution, FPBDCollisionConstraint, T, d>::ComputeIntermediateSpatialAcceleration(bool bBlock)
+	void FPBDRigidsEvolutionBase::ComputeIntermediateSpatialAcceleration(bool bBlock)
 	{
 		LLM_SCOPE(ELLMTag::ChaosAcceleration);
 		SCOPE_CYCLE_COUNTER(STAT_ComputeIntermediateSpatialAcceleration);
@@ -478,7 +462,7 @@ namespace Chaos
 				bExternalReady = true;
 			}
 			
-			// we run the task for both starting a new accel structure as well as for the timeslicing
+			// we run the task for both starting a new accel structure as well as for the time-slicing
 			AccelerationStructureTaskComplete = TGraphTask<FChaosAccelerationStructureTask>::CreateTask().ConstructAndDispatchWhenReady(*SpatialCollectionFactory, SpatialAccelerationCache, AsyncInternalAcceleration, AsyncExternalAcceleration, ForceFullBuild, bIsSingleThreaded);
 
 		}
@@ -488,8 +472,7 @@ namespace Chaos
 		}
 	}
 
-	template<class FPBDRigidsEvolution, class FPBDCollisionConstraint, class T, int d>
-	void TPBDRigidsEvolutionBase<FPBDRigidsEvolution, FPBDCollisionConstraint, T, d>::UpdateExternalAccelerationStructure(TUniquePtr<ISpatialAccelerationCollection<TAccelerationStructureHandle<T, d>, T, d>>& StructToUpdate)
+	void FPBDRigidsEvolutionBase::UpdateExternalAccelerationStructure(TUniquePtr<ISpatialAccelerationCollection<TAccelerationStructureHandle<FReal, 3>, FReal, 3>>& StructToUpdate)
 	{
 		DECLARE_SCOPE_CYCLE_COUNTER(TEXT("CreateExternalAccelerationStructure"), STAT_CreateExternalAccelerationStructure, STATGROUP_Physics);
 		LLM_SCOPE(ELLMTag::ChaosAcceleration);
@@ -507,8 +490,7 @@ namespace Chaos
 		}
 	}
 
-	template<class FPBDRigidsEvolution, class FPBDCollisionConstraint, class T, int d>
-	void TPBDRigidsEvolutionBase<FPBDRigidsEvolution, FPBDCollisionConstraint, T, d>::FlushSpatialAcceleration()
+	void FPBDRigidsEvolutionBase::FlushSpatialAcceleration()
 	{
 		//force build acceleration structure with latest data
 		ComputeIntermediateSpatialAcceleration(true);
@@ -516,8 +498,7 @@ namespace Chaos
 		ComputeIntermediateSpatialAcceleration(true);
 	}
 
-	template<class FPBDRigidsEvolution, class FPBDCollisionConstraint, class T, int d>
-	void TPBDRigidsEvolutionBase<FPBDRigidsEvolution, FPBDCollisionConstraint, T, d>::RebuildSpatialAccelerationForPerfTest()
+	void FPBDRigidsEvolutionBase::RebuildSpatialAccelerationForPerfTest()
 	{
 		WaitOnAccelerationStructure();
 
@@ -536,8 +517,7 @@ namespace Chaos
 		FlushSpatialAcceleration();
 	}
 
-	template<class FPBDRigidsEvolution, class FPBDCollisionConstraint, class T, int d>
-	void TPBDRigidsEvolutionBase<FPBDRigidsEvolution, FPBDCollisionConstraint, T, d>::Serialize(FChaosArchive& Ar)
+	void FPBDRigidsEvolutionBase::Serialize(FChaosArchive& Ar)
 	{
 		ensure(false);	//disabled transient data serialization. Need to rethink
 		int32 DefaultBroadphaseType = ConfigSettings.BroadphaseType;
@@ -565,7 +545,7 @@ namespace Chaos
 			if (Ar.CustomVer(FExternalPhysicsCustomObjectVersion::GUID) < FExternalPhysicsCustomObjectVersion::SerializeMultiStructures)
 			{
 				//old path assumes single sub-structure
-				TUniquePtr<ISpatialAcceleration<TAccelerationStructureHandle<T, d>, T, d>> SubStructure;
+				TUniquePtr<ISpatialAcceleration<TAccelerationStructureHandle<FReal, 3>, FReal, 3>> SubStructure;
 				if (!Ar.IsLoading())
 				{
 					SubStructure = InternalAcceleration->RemoveSubstructure(FSpatialAccelerationIdx{ 0,0 });
@@ -605,12 +585,5 @@ namespace Chaos
 		}
 
 		ConfigSettings.BroadphaseType = DefaultBroadphaseType;
-
 	}
 }
-
-#ifdef __clang__
-template class CHAOS_API Chaos::TPBDRigidsEvolutionBase<Chaos::TPBDRigidsEvolutionGBF<float, 3>, Chaos::TPBDCollisionConstraints<float,3>, float, 3>;
-#else
-template class Chaos::TPBDRigidsEvolutionBase<Chaos::TPBDRigidsEvolutionGBF<float, 3>, Chaos::TPBDCollisionConstraints<float,3>, float, 3>;
-#endif
diff --git a/Engine/Source/Runtime/Experimental/Chaos/Private/Chaos/PBDRigidsEvolutionGBF.cpp b/Engine/Source/Runtime/Experimental/Chaos/Private/Chaos/PBDRigidsEvolutionGBF.cpp
--- a/Engine/Source/Runtime/Experimental/Chaos/Private/Chaos/PBDRigidsEvolutionGBF.cpp
+++ b/Engine/Source/Runtime/Experimental/Chaos/Private/Chaos/PBDRigidsEvolutionGBF.cpp
@@ -59,20 +59,20 @@ float HackCCD_DepthThreshold = 0.05f;
 FAutoConsoleVariableRef CVarHackCCDDepthThreshold(TEXT("p.Chaos.CCD.DepthThreshold"), HackCCD_DepthThreshold, TEXT("When returning to TOI, leave this much contact depth (as a fraction of MinBounds)"));
 
 
-DECLARE_CYCLE_STAT(TEXT("TPBDRigidsEvolutionGBF::AdvanceOneTimeStep"), STAT_Evolution_AdvanceOneTimeStep, STATGROUP_Chaos);
-DECLARE_CYCLE_STAT(TEXT("TPBDRigidsEvolutionGBF::Integrate"), STAT_Evolution_Integrate, STATGROUP_Chaos);
-DECLARE_CYCLE_STAT(TEXT("TPBDRigidsEvolutionGBF::KinematicTargets"), STAT_Evolution_KinematicTargets, STATGROUP_Chaos);
-DECLARE_CYCLE_STAT(TEXT("TPBDRigidsEvolutionGBF::PrepareConstraints"), STAT_Evolution_PrepareConstraints, STATGROUP_Chaos);
-DECLARE_CYCLE_STAT(TEXT("TPBDRigidsEvolutionGBF::UnprepareConstraints"), STAT_Evolution_UnprepareConstraints, STATGROUP_Chaos);
-DECLARE_CYCLE_STAT(TEXT("TPBDRigidsEvolutionGBF::ApplyConstraints"), STAT_Evolution_ApplyConstraints, STATGROUP_Chaos);
-DECLARE_CYCLE_STAT(TEXT("TPBDRigidsEvolutionGBF::UpdateVelocities"), STAT_Evolution_UpdateVelocites, STATGROUP_Chaos);
-DECLARE_CYCLE_STAT(TEXT("TPBDRigidsEvolutionGBF::ApplyPushOut"), STAT_Evolution_ApplyPushOut, STATGROUP_Chaos);
-DECLARE_CYCLE_STAT(TEXT("TPBDRigidsEvolutionGBF::DetectCollisions"), STAT_Evolution_DetectCollisions, STATGROUP_Chaos);
-DECLARE_CYCLE_STAT(TEXT("TPBDRigidsEvolutionGBF::UpdateConstraintPositionBasedState"), STAT_Evolution_UpdateConstraintPositionBasedState, STATGROUP_Chaos);
-DECLARE_CYCLE_STAT(TEXT("TPBDRigidsEvolutionGBF::CreateConstraintGraph"), STAT_Evolution_CreateConstraintGraph, STATGROUP_Chaos);
-DECLARE_CYCLE_STAT(TEXT("TPBDRigidsEvolutionGBF::CreateIslands"), STAT_Evolution_CreateIslands, STATGROUP_Chaos);
-DECLARE_CYCLE_STAT(TEXT("TPBDRigidsEvolutionGBF::ParallelSolve"), STAT_Evolution_ParallelSolve, STATGROUP_Chaos);
-DECLARE_CYCLE_STAT(TEXT("TPBDRigidsEvolutionGBF::DeactivateSleep"), STAT_Evolution_DeactivateSleep, STATGROUP_Chaos);
+DECLARE_CYCLE_STAT(TEXT("FPBDRigidsEvolutionGBF::AdvanceOneTimeStep"), STAT_Evolution_AdvanceOneTimeStep, STATGROUP_Chaos);
+DECLARE_CYCLE_STAT(TEXT("FPBDRigidsEvolutionGBF::Integrate"), STAT_Evolution_Integrate, STATGROUP_Chaos);
+DECLARE_CYCLE_STAT(TEXT("FPBDRigidsEvolutionGBF::KinematicTargets"), STAT_Evolution_KinematicTargets, STATGROUP_Chaos);
+DECLARE_CYCLE_STAT(TEXT("FPBDRigidsEvolutionGBF::PrepareConstraints"), STAT_Evolution_PrepareConstraints, STATGROUP_Chaos);
+DECLARE_CYCLE_STAT(TEXT("FPBDRigidsEvolutionGBF::UnprepareConstraints"), STAT_Evolution_UnprepareConstraints, STATGROUP_Chaos);
+DECLARE_CYCLE_STAT(TEXT("FPBDRigidsEvolutionGBF::ApplyConstraints"), STAT_Evolution_ApplyConstraints, STATGROUP_Chaos);
+DECLARE_CYCLE_STAT(TEXT("FPBDRigidsEvolutionGBF::UpdateVelocities"), STAT_Evolution_UpdateVelocites, STATGROUP_Chaos);
+DECLARE_CYCLE_STAT(TEXT("FPBDRigidsEvolutionGBF::ApplyPushOut"), STAT_Evolution_ApplyPushOut, STATGROUP_Chaos);
+DECLARE_CYCLE_STAT(TEXT("FPBDRigidsEvolutionGBF::DetectCollisions"), STAT_Evolution_DetectCollisions, STATGROUP_Chaos);
+DECLARE_CYCLE_STAT(TEXT("FPBDRigidsEvolutionGBF::UpdateConstraintPositionBasedState"), STAT_Evolution_UpdateConstraintPositionBasedState, STATGROUP_Chaos);
+DECLARE_CYCLE_STAT(TEXT("FPBDRigidsEvolutionGBF::CreateConstraintGraph"), STAT_Evolution_CreateConstraintGraph, STATGROUP_Chaos);
+DECLARE_CYCLE_STAT(TEXT("FPBDRigidsEvolutionGBF::CreateIslands"), STAT_Evolution_CreateIslands, STATGROUP_Chaos);
+DECLARE_CYCLE_STAT(TEXT("FPBDRigidsEvolutionGBF::ParallelSolve"), STAT_Evolution_ParallelSolve, STATGROUP_Chaos);
+DECLARE_CYCLE_STAT(TEXT("FPBDRigidsEvolutionGBF::DeactivateSleep"), STAT_Evolution_DeactivateSleep, STATGROUP_Chaos);
 
 int32 SerializeEvolution = 0;
 FAutoConsoleVariableRef CVarSerializeEvolution(TEXT("p.SerializeEvolution"), SerializeEvolution, TEXT(""));
@@ -109,18 +109,17 @@ void SerializeToDisk(TEvolution& Evolution)
 }
 #endif
 
-template <typename T, int d>
-void TPBDRigidsEvolutionGBF<T, d>::Advance(const T Dt, const T MaxStepDt, const int32 MaxSteps)
+void FPBDRigidsEvolutionGBF::Advance(const FReal Dt, const FReal MaxStepDt, const int32 MaxSteps)
 {
 	// Determine how many steps we would like to take
 	int32 NumSteps = FMath::CeilToInt(Dt / MaxStepDt);
 	if (NumSteps > 0)
 	{
 		// Determine the step time
-		const T StepDt = Dt / (T)NumSteps;
+		const FReal StepDt = Dt / (FReal)NumSteps;
 
 		// Limit the number of steps
-		// NOTE: This is after step time calculation so sim will appear to slow down for large Dt
+		// NOTE: This is after step time calculation so simulation will appear to slow down for large Dt
 		// but that is preferable to blowing up from a large timestep.
 		NumSteps = FMath::Clamp(NumSteps, 1, MaxSteps);
 
@@ -128,7 +127,7 @@ void TPBDRigidsEvolutionGBF<T, d>::Advance(const T Dt, const T MaxStepDt, const
 		{
 			// StepFraction: how much of the remaining time this step represents, used to interpolate kinematic targets
 			// E.g., for 4 steps this will be: 1/4, 1/3, 1/2, 1
-			const float StepFraction = (T)1 / (T)(NumSteps - Step);
+			const float StepFraction = (FReal)1 / (FReal)(NumSteps - Step);
 		
 			UE_LOG(LogChaos, Verbose, TEXT("Advance dt = %f [%d/%d]"), StepDt, Step + 1, NumSteps);
 
@@ -247,8 +246,7 @@ void CCDHack(const FReal Dt, TParticleView<TPBDRigidParticles<FReal, 3>>& Partic
 //
 
 
-template <typename T, int d>
-void TPBDRigidsEvolutionGBF<T, d>::AdvanceOneTimeStep(const T Dt, const T StepFraction)
+void FPBDRigidsEvolutionGBF::AdvanceOneTimeStep(const FReal Dt, const FReal StepFraction)
 {
 	SCOPE_CYCLE_COUNTER(STAT_Evolution_AdvanceOneTimeStep);
 
@@ -326,14 +324,14 @@ void TPBDRigidsEvolutionGBF<T, d>::AdvanceOneTimeStep(const T Dt, const T StepFr
 
 	TArray<bool> SleepedIslands;
 	SleepedIslands.SetNum(GetConstraintGraph().NumIslands());
-	TArray<TArray<TPBDRigidParticleHandle<T,d>*>> DisabledParticles;
+	TArray<TArray<TPBDRigidParticleHandle<FReal, 3>*>> DisabledParticles;
 	DisabledParticles.SetNum(GetConstraintGraph().NumIslands());
 	SleepedIslands.SetNum(GetConstraintGraph().NumIslands());
 	if(Dt > 0)
 	{
 		SCOPE_CYCLE_COUNTER(STAT_Evolution_ParallelSolve);
 		PhysicsParallelFor(GetConstraintGraph().NumIslands(), [&](int32 Island) {
-			const TArray<TGeometryParticleHandle<T, d>*>& IslandParticles = GetConstraintGraph().GetIslandParticles(Island);
+			const TArray<TGeometryParticleHandle<FReal, 3>*>& IslandParticles = GetConstraintGraph().GetIslandParticles(Island);
 
 			{
 				SCOPE_CYCLE_COUNTER(STAT_Evolution_ApplyConstraints);
@@ -423,9 +421,9 @@ void TPBDRigidsEvolutionGBF<T, d>::AdvanceOneTimeStep(const T Dt, const T StepFr
 	ParticleUpdatePosition(Particles.GetActiveParticlesView(), Dt);
 }
 
-template <typename T, int d>
-TPBDRigidsEvolutionGBF<T, d>::TPBDRigidsEvolutionGBF(TPBDRigidsSOAs<T, d>& InParticles, int32 InNumIterations, int32 InNumPushoutIterations, bool InIsSingleThreaded)
+FPBDRigidsEvolutionGBF::FPBDRigidsEvolutionGBF(TPBDRigidsSOAs<FReal, 3>& InParticles, int32 InNumIterations, int32 InNumPushoutIterations, bool InIsSingleThreaded)
 	: Base(InParticles, InNumIterations, InNumPushoutIterations, InIsSingleThreaded)
+	, Clustering(*this, Particles.GetClusteredParticles())
 	, CollisionConstraints(InParticles, Collided, PhysicsMaterials, DefaultNumPairIterations, DefaultNumPushOutPairIterations)
 	, CollisionRule(CollisionConstraints)
 	, BroadPhase(InParticles, BoundsThickness, BoundsThicknessVelocityMultiplier)
@@ -435,7 +433,7 @@ TPBDRigidsEvolutionGBF<T, d>::TPBDRigidsEvolutionGBF(TPBDRigidsSOAs<T, d>& InPar
 	, PostApplyCallback(nullptr)
 	, PostApplyPushOutCallback(nullptr)
 {
-	SetParticleUpdateVelocityFunction([PBDUpdateRule = TPerParticlePBDUpdateFromDeltaPosition<float, 3>(), this](const TArray<TGeometryParticleHandle<T, d>*>& ParticlesInput, const T Dt) {
+	SetParticleUpdateVelocityFunction([PBDUpdateRule = TPerParticlePBDUpdateFromDeltaPosition<float, 3>(), this](const TArray<TGeometryParticleHandle<FReal, 3>*>& ParticlesInput, const FReal Dt) {
 		ParticlesParallelFor(ParticlesInput, [&](auto& Particle, int32 Index) {
 			if (Particle->CastToRigidParticle() && Particle->ObjectState() == EObjectStateType::Dynamic)
 			{
@@ -444,7 +442,7 @@ TPBDRigidsEvolutionGBF<T, d>::TPBDRigidsEvolutionGBF(TPBDRigidsSOAs<T, d>& InPar
 		});
 	});
 
-	SetParticleUpdatePositionFunction([this](const TParticleView<TPBDRigidParticles<T, d>>& ParticlesInput, const T Dt)
+	SetParticleUpdatePositionFunction([this](const TParticleView<TPBDRigidParticles<FReal, 3>>& ParticlesInput, const FReal Dt)
 	{
 		ParticlesInput.ParallelFor([&](auto& Particle, int32 Index)
 		{
@@ -453,7 +451,7 @@ TPBDRigidsEvolutionGBF<T, d>::TPBDRigidsEvolutionGBF(TPBDRigidsSOAs<T, d>& InPar
 		});
 	});
 
-	AddForceFunction([this](TTransientPBDRigidParticleHandle<T, d>& HandleIn, const T Dt)
+	AddForceFunction([this](TTransientPBDRigidParticleHandle<FReal, 3>& HandleIn, const FReal Dt)
 	{
 		GravityForces.Apply(HandleIn, Dt);
 	});
@@ -461,13 +459,10 @@ TPBDRigidsEvolutionGBF<T, d>::TPBDRigidsEvolutionGBF(TPBDRigidsSOAs<T, d>& InPar
 	AddConstraintRule(&CollisionRule);
 }
 
-template <typename T, int d>
-void TPBDRigidsEvolutionGBF<T,d>::Serialize(FChaosArchive& Ar)
+void FPBDRigidsEvolutionGBF::Serialize(FChaosArchive& Ar)
 {
 	Base::Serialize(Ar);
 }
 
 }
 
-template class Chaos::TPBDRigidsEvolutionGBF<float, 3>;
-
diff --git a/Engine/Source/Runtime/Experimental/Chaos/Private/Chaos/PBDRigidsEvolutionPGS.cpp b/Engine/Source/Runtime/Experimental/Chaos/Private/Chaos/PBDRigidsEvolutionPGS.cpp
deleted file mode 100644
diff --git a/Engine/Source/Runtime/Experimental/Chaos/Public/Chaos/AABB.h b/Engine/Source/Runtime/Experimental/Chaos/Public/Chaos/AABB.h
--- a/Engine/Source/Runtime/Experimental/Chaos/Public/Chaos/AABB.h
+++ b/Engine/Source/Runtime/Experimental/Chaos/Public/Chaos/AABB.h
@@ -606,4 +606,6 @@ namespace Chaos
 			return SamplePoints;
 		}
 	};
+
+	using FAABB3 = TAABB<FReal, 3>;
 }
\ No newline at end of file
diff --git a/Engine/Source/Runtime/Experimental/Chaos/Public/Chaos/PBDRigidsEvolution.h b/Engine/Source/Runtime/Experimental/Chaos/Public/Chaos/PBDRigidsEvolution.h
--- a/Engine/Source/Runtime/Experimental/Chaos/Public/Chaos/PBDRigidsEvolution.h
+++ b/Engine/Source/Runtime/Experimental/Chaos/Public/Chaos/PBDRigidsEvolution.h
@@ -27,7 +27,6 @@ namespace Chaos
 
 extern CHAOS_API int32 FixBadAccelerationStructureRemoval;
 
-template<class T, int d> class TPBDRigidsEvolutionGBF;
 class FChaosArchive;
 
 template <typename TPayload, typename T, int d>
@@ -63,17 +62,15 @@ struct CHAOS_API FEvolutionStats
 	}
 };
 
-template <typename T, int d>
-struct TSpatialAccelerationCacheHandle;
+struct FSpatialAccelerationCacheHandle;
 
 /** The SOA cache used for a single acceleration structure */
-template <typename T, int d>
-class TSpatialAccelerationCache : public TArrayCollection
+class FSpatialAccelerationCache : public TArrayCollection
 {
 public:
-	using THandleType = TSpatialAccelerationCacheHandle<T, d>;
+	using THandleType = FSpatialAccelerationCacheHandle;
 
-	TSpatialAccelerationCache()
+	FSpatialAccelerationCache()
 	{
 		AddArray(&MHasBoundingBoxes);
 		AddArray(&MBounds);
@@ -84,8 +81,8 @@ public:
 #endif
 	}
 
-	TSpatialAccelerationCache(const TSpatialAccelerationCache<T, d>&) = delete;
-	TSpatialAccelerationCache(TSpatialAccelerationCache<T, d>&& Other)
+	FSpatialAccelerationCache(const FSpatialAccelerationCache&) = delete;
+	FSpatialAccelerationCache(FSpatialAccelerationCache&& Other)
 		: TArrayCollection()
 		, MHasBoundingBoxes(MoveTemp(Other.MHasBoundingBoxes))
 		, MBounds(MoveTemp(Other.MBounds))
@@ -102,7 +99,7 @@ public:
 #endif
 	}
 
-	TSpatialAccelerationCache& operator=(TSpatialAccelerationCache<T, d>&& Other)
+	FSpatialAccelerationCache& operator=(FSpatialAccelerationCache&& Other)
 	{
 		if (&Other != this)
 		{
@@ -139,11 +136,11 @@ public:
 	bool HasBounds(const int32 Idx) const { return MHasBoundingBoxes[Idx]; }
 	bool& HasBounds(const int32 Idx) { return MHasBoundingBoxes[Idx]; }
 
-	const TAABB<T,d>& Bounds(const int32 Idx) const { return MBounds[Idx]; }
-	TAABB<T, d>& Bounds(const int32 Idx) { return MBounds[Idx]; }
+	const TAABB<FReal, 3>& Bounds(const int32 Idx) const { return MBounds[Idx]; }
+	TAABB<FReal, 3>& Bounds(const int32 Idx) { return MBounds[Idx]; }
 
-	const TAccelerationStructureHandle<T, d>& Payload(const int32 Idx) const { return MPayloads[Idx]; }
-	TAccelerationStructureHandle<T, d>& Payload(const int32 Idx) { return MPayloads[Idx]; }
+	const TAccelerationStructureHandle<FReal, 3>& Payload(const int32 Idx) const { return MPayloads[Idx]; }
+	TAccelerationStructureHandle<FReal, 3>& Payload(const int32 Idx) { return MPayloads[Idx]; }
 
 private:
 	void IncrementDirtyValidation()
@@ -154,8 +151,8 @@ private:
 	}
 
 	TArrayCollectionArray<bool> MHasBoundingBoxes;
-	TArrayCollectionArray<TAABB<T, d>> MBounds;
-	TArrayCollectionArray<TAccelerationStructureHandle<T, d>> MPayloads;
+	TArrayCollectionArray<TAABB<FReal, 3>> MBounds;
+	TArrayCollectionArray<TAccelerationStructureHandle<FReal, 3>> MPayloads;
 
 #if PARTICLE_ITERATOR_RANGED_FOR_CHECK
 	int32 MDirtyValidationCount;
@@ -163,13 +160,12 @@ private:
 };
 
 /** The handle the acceleration structure uses to access the data (similar to particle handle) */
-template <typename T, int d>
-struct TSpatialAccelerationCacheHandle
+struct FSpatialAccelerationCacheHandle
 {
-	using THandleBase = TSpatialAccelerationCacheHandle<T, d>;
-	using TTransientHandle = TSpatialAccelerationCacheHandle<T, d>;
+	using THandleBase = FSpatialAccelerationCacheHandle;
+	using TTransientHandle = FSpatialAccelerationCacheHandle;
 
-	TSpatialAccelerationCacheHandle(TSpatialAccelerationCache<T, d>* InCache = nullptr, int32 InEntryIdx = INDEX_NONE)
+	FSpatialAccelerationCacheHandle(FSpatialAccelerationCache* InCache = nullptr, int32 InEntryIdx = INDEX_NONE)
 		: Cache(InCache)
 		, EntryIdx(InEntryIdx)
 	{}
@@ -185,15 +181,15 @@ struct TSpatialAccelerationCacheHandle
 		return Cache->HasBounds(EntryIdx);
 	}
 
-	const TAABB<T, d>& BoundingBox() const
+	const TAABB<FReal, 3>& BoundingBox() const
 	{
 		return Cache->Bounds(EntryIdx);
 	}
 
 	union
 	{
-		TSpatialAccelerationCache<T, d>* GeometryParticles;	//using same name as particles SOA for template reuse, should probably rethink this
-		TSpatialAccelerationCache<T, d>* Cache;
+		FSpatialAccelerationCache* GeometryParticles;	//using same name as particles SOA for template reuse, should probably rethink this
+		FSpatialAccelerationCache* Cache;
 	};
 
 	union
@@ -203,39 +199,37 @@ struct TSpatialAccelerationCacheHandle
 	};
 };
 
-template <typename T, int d>
 struct CHAOS_API ISpatialAccelerationCollectionFactory
 {
 	//Create an empty acceleration collection with the desired buckets. Chaos enqueues acceleration structure operations per bucket
-	virtual TUniquePtr<ISpatialAccelerationCollection<TAccelerationStructureHandle<T, d>, T, d>> CreateEmptyCollection() = 0;
+	virtual TUniquePtr<ISpatialAccelerationCollection<TAccelerationStructureHandle<FReal, 3>, FReal, 3>> CreateEmptyCollection() = 0;
 
 	//Chaos creates new acceleration structures per bucket. Factory can change underlying type at runtime as well as number of buckets to AB test
-	virtual TUniquePtr<ISpatialAcceleration<TAccelerationStructureHandle<T, d>, T, d>> CreateAccelerationPerBucket_Threaded(const TConstParticleView<TSpatialAccelerationCache<T, d>>& Particles, uint16 BucketIdx, bool ForceFullBuild) = 0;
+	virtual TUniquePtr<ISpatialAcceleration<TAccelerationStructureHandle<FReal, 3>, FReal, 3>> CreateAccelerationPerBucket_Threaded(const TConstParticleView<FSpatialAccelerationCache>& Particles, uint16 BucketIdx, bool ForceFullBuild) = 0;
 
 	//Mask indicating which bucket is active. Spatial indices in inactive buckets fallback to bucket 0. Bit 0 indicates bucket 0 is active, Bit 1 indicates bucket 1 is active, etc...
 	virtual uint8 GetActiveBucketsMask() const = 0;
 
 	//Serialize the collection in and out
-	virtual void Serialize(TUniquePtr<ISpatialAccelerationCollection<TAccelerationStructureHandle<T, d>, T, d>>& Ptr, FChaosArchive& Ar) = 0;
+	virtual void Serialize(TUniquePtr<ISpatialAccelerationCollection<TAccelerationStructureHandle<FReal, 3>, FReal, 3>>& Ptr, FChaosArchive& Ar) = 0;
 
 	virtual ~ISpatialAccelerationCollectionFactory() = default;
 };
 
-template<class FPBDRigidsEvolution, class FPBDCollisionConstraint, class T, int d>
-class TPBDRigidsEvolutionBase
+class FPBDRigidsEvolutionBase
 {
   public:
-	typedef TFunction<void(TTransientPBDRigidParticleHandle<T,d>& Particle, const T)> FForceRule;
-	typedef TFunction<void(const TArray<TGeometryParticleHandle<T, d>*>&, const T)> FUpdateVelocityRule;
-	typedef TFunction<void(const TParticleView<TPBDRigidParticles<T, d>>&, const T)> FUpdatePositionRule;
-	typedef TFunction<void(TPBDRigidParticles<T, d>&, const T, const T, const int32)> FKinematicUpdateRule;
+	typedef TFunction<void(TTransientPBDRigidParticleHandle<FReal, 3>& Particle, const FReal)> FForceRule;
+	typedef TFunction<void(const TArray<TGeometryParticleHandle<FReal, 3>*>&, const FReal)> FUpdateVelocityRule;
+	typedef TFunction<void(const TParticleView<TPBDRigidParticles<FReal, 3>>&, const FReal)> FUpdatePositionRule;
+	typedef TFunction<void(TPBDRigidParticles<FReal, 3>&, const FReal, const FReal, const int32)> FKinematicUpdateRule;
 
 	friend void ChaosTest::TestPendingSpatialDataHandlePointerConflict();
 
-	CHAOS_API TPBDRigidsEvolutionBase(TPBDRigidsSOAs<T, d>& InParticles, int32 InNumIterations = 1, int32 InNumPushOutIterations = 1, bool InIsSingleThreaded = false);
-	CHAOS_API virtual ~TPBDRigidsEvolutionBase();
+	CHAOS_API FPBDRigidsEvolutionBase(TPBDRigidsSOAs<FReal, 3>& InParticles, int32 InNumIterations = 1, int32 InNumPushOutIterations = 1, bool InIsSingleThreaded = false);
+	CHAOS_API virtual ~FPBDRigidsEvolutionBase();
 
-	CHAOS_API TArray<TGeometryParticleHandle<T, d>*> CreateStaticParticles(int32 NumParticles, const FUniqueIdx* ExistingIndices = nullptr, const TGeometryParticleParameters<T, d>& Params = TGeometryParticleParameters<T, d>())
+	CHAOS_API TArray<TGeometryParticleHandle<FReal, 3>*> CreateStaticParticles(int32 NumParticles, const FUniqueIdx* ExistingIndices = nullptr, const TGeometryParticleParameters<FReal, 3>& Params = TGeometryParticleParameters<FReal, 3>())
 	{
 		auto NewParticles = Particles.CreateStaticParticles(NumParticles, ExistingIndices, Params);
 		for (auto& Particle : NewParticles)
@@ -245,7 +239,7 @@ class TPBDRigidsEvolutionBase
 		return NewParticles;
 	}
 
-	CHAOS_API TArray<TKinematicGeometryParticleHandle<T, d>*> CreateKinematicParticles(int32 NumParticles, const FUniqueIdx* ExistingIndices = nullptr, const TKinematicGeometryParticleParameters<T, d>& Params = TKinematicGeometryParticleParameters<T, d>())
+	CHAOS_API TArray<TKinematicGeometryParticleHandle<FReal, 3>*> CreateKinematicParticles(int32 NumParticles, const FUniqueIdx* ExistingIndices = nullptr, const TKinematicGeometryParticleParameters<FReal, 3>& Params = TKinematicGeometryParticleParameters<FReal, 3>())
 	{
 		auto NewParticles = Particles.CreateKinematicParticles(NumParticles, ExistingIndices, Params);
 		for (auto& Particle : NewParticles)
@@ -255,7 +249,7 @@ class TPBDRigidsEvolutionBase
 		return NewParticles;
 	}
 
-	CHAOS_API TArray<TPBDRigidParticleHandle<T, d>*> CreateDynamicParticles(int32 NumParticles, const FUniqueIdx* ExistingIndices = nullptr, const TPBDRigidParticleParameters<T, d>& Params = TPBDRigidParticleParameters<T, d>())
+	CHAOS_API TArray<TPBDRigidParticleHandle<FReal, 3>*> CreateDynamicParticles(int32 NumParticles, const FUniqueIdx* ExistingIndices = nullptr, const TPBDRigidParticleParameters<FReal, 3>& Params = TPBDRigidParticleParameters<FReal, 3>())
 	{
 		auto NewParticles = Particles.CreateDynamicParticles(NumParticles, ExistingIndices, Params);
 		for (auto& Particle : NewParticles)
@@ -265,7 +259,7 @@ class TPBDRigidsEvolutionBase
 		return NewParticles;
 	}
 
-	CHAOS_API TArray<TPBDRigidClusteredParticleHandle<T, d>*> CreateClusteredParticles(int32 NumParticles,const FUniqueIdx* ExistingIndices = nullptr,  const TPBDRigidParticleParameters<T, d>& Params = TPBDRigidParticleParameters<T, d>())
+	CHAOS_API TArray<TPBDRigidClusteredParticleHandle<FReal, 3>*> CreateClusteredParticles(int32 NumParticles,const FUniqueIdx* ExistingIndices = nullptr,  const TPBDRigidParticleParameters<FReal, 3>& Params = TPBDRigidParticleParameters<FReal, 3>())
 	{
 		auto NewParticles = Particles.CreateClusteredParticles(NumParticles, ExistingIndices, Params);
 		for (auto& Particle : NewParticles)
@@ -280,11 +274,11 @@ class TPBDRigidsEvolutionBase
 	CHAOS_API void SetParticleUpdateVelocityFunction(FUpdateVelocityRule ParticleUpdate) { ParticleUpdateVelocity = ParticleUpdate; }
 	CHAOS_API void SetParticleUpdatePositionFunction(FUpdatePositionRule ParticleUpdate) { ParticleUpdatePosition = ParticleUpdate; }
 
-	CHAOS_API TGeometryParticleHandles<T, d>& GetParticleHandles() { return Particles.GetParticleHandles(); }
-	CHAOS_API const TGeometryParticleHandles<T, d>& GetParticleHandles() const { return Particles.GetParticleHandles(); }
+	CHAOS_API TGeometryParticleHandles<FReal, 3>& GetParticleHandles() { return Particles.GetParticleHandles(); }
+	CHAOS_API const TGeometryParticleHandles<FReal, 3>& GetParticleHandles() const { return Particles.GetParticleHandles(); }
 
-	CHAOS_API TPBDRigidsSOAs<T,d>& GetParticles() { return Particles; }
-	CHAOS_API const TPBDRigidsSOAs<T, d>& GetParticles() const { return Particles; }
+	CHAOS_API TPBDRigidsSOAs<FReal, 3>& GetParticles() { return Particles; }
+	CHAOS_API const TPBDRigidsSOAs<FReal, 3>& GetParticles() const { return Particles; }
 
 	CHAOS_API void AddConstraintRule(FPBDConstraintGraphRule* ConstraintRule)
 	{
@@ -303,30 +297,30 @@ class TPBDRigidsEvolutionBase
 		NumPushOutIterations = InNumIterations;
 	}
 
-	CHAOS_API void EnableParticle(TGeometryParticleHandle<T,d>* Particle, const TGeometryParticleHandle<T, d>* ParentParticle)
+	CHAOS_API void EnableParticle(TGeometryParticleHandle<FReal, 3>* Particle, const TGeometryParticleHandle<FReal, 3>* ParentParticle)
 	{
 		DirtyParticle(*Particle);
 		Particles.EnableParticle(Particle);
 		ConstraintGraph.EnableParticle(Particle, ParentParticle);
 	}
 
-	CHAOS_API void DisableParticle(TGeometryParticleHandle<T,d>* Particle)
+	CHAOS_API void DisableParticle(TGeometryParticleHandle<FReal, 3>* Particle)
 	{
 		RemoveParticleFromAccelerationStructure(*Particle);
 		Particles.DisableParticle(Particle);
 		ConstraintGraph.DisableParticle(Particle);
 
-		RemoveConstraints(TSet<TGeometryParticleHandle<T,d>*>({ Particle }));
+		RemoveConstraints(TSet<TGeometryParticleHandle<FReal, 3>*>({ Particle }));
 	}
 
 	template <bool bPersistent>
-	FORCEINLINE_DEBUGGABLE void DirtyParticle(TGeometryParticleHandleImp<T, d, bPersistent>& Particle)
+	FORCEINLINE_DEBUGGABLE void DirtyParticle(TGeometryParticleHandleImp<FReal, 3, bPersistent>& Particle)
 	{
 		//TODO: distinguish between new particles and dirty particles
 		const FUniqueIdx UniqueIdx = Particle.UniqueIdx();
 		FPendingSpatialData& SpatialData = InternalAccelerationQueue.FindOrAdd(UniqueIdx);
 		ensure(SpatialData.bDelete == false);
-		SpatialData.AccelerationHandle = TAccelerationStructureHandle<T,d>(Particle);
+		SpatialData.AccelerationHandle = TAccelerationStructureHandle<FReal, 3>(Particle);
 		SpatialData.SpatialIdx = Particle.SpatialIdx();
 
 		auto& AsyncSpatialData = AsyncAccelerationQueue.FindOrAdd(UniqueIdx);
@@ -338,29 +332,29 @@ class TPBDRigidsEvolutionBase
 		ExternalSpatialData = SpatialData;
 	}
 
-	void DestroyParticle(TGeometryParticleHandle<T, d>* Particle)
+	void DestroyParticle(TGeometryParticleHandle<FReal, 3>* Particle)
 	{
 		RemoveParticleFromAccelerationStructure(*Particle);
 		ConstraintGraph.RemoveParticle(Particle);
-		RemoveConstraints(TSet<TGeometryParticleHandle<T, d>*>({ Particle }));
+		RemoveConstraints(TSet<TGeometryParticleHandle<FReal, 3>*>({ Particle }));
 		Particles.DestroyParticle(Particle);
 	}
 
-	CHAOS_API void CreateParticle(TGeometryParticleHandle<T, d>* ParticleAdded)
+	CHAOS_API void CreateParticle(TGeometryParticleHandle<FReal, 3>* ParticleAdded)
 	{
 		ConstraintGraph.AddParticle(ParticleAdded);
 		DirtyParticle(*ParticleAdded);
 	}
 
-	CHAOS_API void SetParticleObjectState(TPBDRigidParticleHandle<T, d>* Particle, EObjectStateType ObjectState)
+	CHAOS_API void SetParticleObjectState(TPBDRigidParticleHandle<FReal, 3>* Particle, EObjectStateType ObjectState)
 	{
 		Particle->SetObjectStateLowLevel(ObjectState);
 		Particles.SetDynamicParticleSOA(Particle);
 	}
 
-	CHAOS_API void DisableParticles(const TSet<TGeometryParticleHandle<T,d>*>& InParticles)
+	CHAOS_API void DisableParticles(const TSet<TGeometryParticleHandle<FReal, 3>*>& InParticles)
 	{
-		for (TGeometryParticleHandle<T, d>* Particle : InParticles)
+		for (TGeometryParticleHandle<FReal, 3>* Particle : InParticles)
 		{
 			Particles.DisableParticle(Particle);
 			RemoveParticleFromAccelerationStructure(*Particle);
@@ -382,7 +376,7 @@ class TPBDRigidsEvolutionBase
 	}
 
 	// @todo(ccaulfield): Remove the uint version
-	CHAOS_API void RemoveConstraints(const TSet<TGeometryParticleHandle<T, d>*>& RemovedParticles)
+	CHAOS_API void RemoveConstraints(const TSet<TGeometryParticleHandle<FReal, 3>*>& RemovedParticles)
 	{
 		for (FPBDConstraintGraphRule* ConstraintRule : ConstraintRules)
 		{
@@ -394,23 +388,23 @@ class TPBDRigidsEvolutionBase
 	const auto& GetActiveClusteredArray() const { return Particles.GetActiveClusteredArray(); }
 	const auto& GetNonDisabledClusteredArray() const { return Particles.GetNonDisabledClusteredArray(); }
 
-	CHAOS_API TSerializablePtr<FChaosPhysicsMaterial> GetPhysicsMaterial(const TGeometryParticleHandle<T, d>* Particle) const { return Particle->AuxilaryValue(PhysicsMaterials); }
+	CHAOS_API TSerializablePtr<FChaosPhysicsMaterial> GetPhysicsMaterial(const TGeometryParticleHandle<FReal, 3>* Particle) const { return Particle->AuxilaryValue(PhysicsMaterials); }
 	
-	CHAOS_API const TUniquePtr<FChaosPhysicsMaterial> &GetPerParticlePhysicsMaterial(const TGeometryParticleHandle<T, d>* Particle) const { return Particle->AuxilaryValue(PerParticlePhysicsMaterials); }
+	CHAOS_API const TUniquePtr<FChaosPhysicsMaterial> &GetPerParticlePhysicsMaterial(const TGeometryParticleHandle<FReal, 3>* Particle) const { return Particle->AuxilaryValue(PerParticlePhysicsMaterials); }
 
-	CHAOS_API void SetPerParticlePhysicsMaterial(TGeometryParticleHandle<T, d>* Particle, TUniquePtr<FChaosPhysicsMaterial> &InMaterial)
+	CHAOS_API void SetPerParticlePhysicsMaterial(TGeometryParticleHandle<FReal, 3>* Particle, TUniquePtr<FChaosPhysicsMaterial> &InMaterial)
 	{
 		check(!Particle->AuxilaryValue(PerParticlePhysicsMaterials)); //shouldn't be setting non unique material if a unique one already exists
 		Particle->AuxilaryValue(PerParticlePhysicsMaterials) = MoveTemp(InMaterial);
 	}
 
-	CHAOS_API void SetPhysicsMaterial(TGeometryParticleHandle<T,d>* Particle, TSerializablePtr<FChaosPhysicsMaterial> InMaterial)
+	CHAOS_API void SetPhysicsMaterial(TGeometryParticleHandle<FReal, 3>* Particle, TSerializablePtr<FChaosPhysicsMaterial> InMaterial)
 	{
 		check(!Particle->AuxilaryValue(PerParticlePhysicsMaterials)); //shouldn't be setting non unique material if a unique one already exists
 		Particle->AuxilaryValue(PhysicsMaterials) = InMaterial;
 	}
 
-	CHAOS_API const TArray<TGeometryParticleHandle<T,d>*>& GetIslandParticles(const int32 Island) const { return ConstraintGraph.GetIslandParticles(Island); }
+	CHAOS_API const TArray<TGeometryParticleHandle<FReal, 3>*>& GetIslandParticles(const int32 Island) const { return ConstraintGraph.GetIslandParticles(Island); }
 	CHAOS_API int32 NumIslands() const { return ConstraintGraph.NumIslands(); }
 
 	void InitializeAccelerationStructures()
@@ -430,7 +424,7 @@ class TPBDRigidsEvolutionBase
 		}
 	}
 
-	void PrepareConstraints(const T Dt)
+	void PrepareConstraints(const FReal Dt)
 	{
 		for (FPBDConstraintGraphRule* ConstraintRule : ConstraintRules)
 		{
@@ -438,7 +432,7 @@ class TPBDRigidsEvolutionBase
 		}
 	}
 
-	void UnprepareConstraints(const T Dt)
+	void UnprepareConstraints(const FReal Dt)
 	{
 		for (FPBDConstraintGraphRule* ConstraintRule : ConstraintRules)
 		{
@@ -454,7 +448,7 @@ class TPBDRigidsEvolutionBase
 		}
 	}
 
-	void ApplyConstraints(const T Dt, int32 Island)
+	void ApplyConstraints(const FReal Dt, int32 Island)
 	{
 		UpdateAccelerationStructures(Island);
 
@@ -474,17 +468,17 @@ class TPBDRigidsEvolutionBase
 		}
 	}
 
-	void ApplyKinematicTargets(const T Dt, const T StepFraction)
+	void ApplyKinematicTargets(const FReal Dt, const FReal StepFraction)
 	{
-		check(StepFraction > (T)0);
-		check(StepFraction <= (T)1);
+		check(StepFraction > (FReal)0);
+		check(StepFraction <= (FReal)1);
 
 		// @todo(ccaulfield): optimize. Depending on the number of kinematics relative to the number that have 
 		// targets set, it may be faster to process a command list rather than iterate over them all each frame. 
-		const T MinDt = 1e-6f;
+		const FReal MinDt = 1e-6f;
 		for (auto& Particle : Particles.GetActiveKinematicParticlesView())
 		{
-			TKinematicTarget<T, d>& KinematicTarget = Particle.KinematicTarget();
+			TKinematicTarget<FReal, 3>& KinematicTarget = Particle.KinematicTarget();
 			switch (KinematicTarget.GetMode())
 			{
 			case EKinematicTargetMode::None:
@@ -494,8 +488,8 @@ class TPBDRigidsEvolutionBase
 			case EKinematicTargetMode::Zero:
 			{
 				// Reset velocity and then switch to do-nothing mode
-				Particle.V() = TVector<T, d>(0, 0, 0);
-				Particle.W() = TVector<T, d>(0, 0, 0);
+				Particle.V() = FVec3(0, 0, 0);
+				Particle.W() = FVec3(0, 0, 0);
 				KinematicTarget.SetMode(EKinematicTargetMode::None);
 				break;
 			}
@@ -504,9 +498,9 @@ class TPBDRigidsEvolutionBase
 			{
 				// Move to kinematic target and update velocities to match
 				// Target positions only need to be processed once, and we reset the velocity next frame (if no new target is set)
-				TVector<T, d> TargetPos;
-				TRotation<T, d> TargetRot;
-				if (FMath::IsNearlyEqual(StepFraction, (T)1, KINDA_SMALL_NUMBER))
+				FVec3 TargetPos;
+				FRotation3 TargetRot;
+				if (FMath::IsNearlyEqual(StepFraction, (FReal)1, KINDA_SMALL_NUMBER))
 				{
 					TargetPos = KinematicTarget.GetTarget().GetLocation();
 					TargetRot = KinematicTarget.GetTarget().GetRotation();
@@ -514,15 +508,15 @@ class TPBDRigidsEvolutionBase
 				}
 				else
 				{
-					TargetPos = TVector<T, d>::Lerp(Particle.X(), KinematicTarget.GetTarget().GetLocation(), StepFraction);
-					TargetRot = TRotation<T, d>::Slerp(Particle.R(), KinematicTarget.GetTarget().GetRotation(), StepFraction);
+					TargetPos = FVec3::Lerp(Particle.X(), KinematicTarget.GetTarget().GetLocation(), StepFraction);
+					TargetRot = FRotation3::Slerp(Particle.R(), KinematicTarget.GetTarget().GetRotation(), StepFraction);
 				}
 				if (Dt > MinDt)
 				{
-					TVector<float, 3> V = TVector<float, 3>::CalculateVelocity(Particle.X(), TargetPos, Dt);
+					FVec3 V = FVec3::CalculateVelocity(Particle.X(), TargetPos, Dt);
 					Particle.V() = V;
 
-					TVector<float, 3> W = TRotation<float, 3>::CalculateAngularVelocity(Particle.R(), TargetRot, Dt);
+					FVec3 W = FRotation3::CalculateAngularVelocity(Particle.R(), TargetRot, Dt);
 					Particle.W() = W;
 				}
 				Particle.X() = TargetPos;
@@ -534,7 +528,7 @@ class TPBDRigidsEvolutionBase
 			{
 				// Move based on velocity
 				Particle.X() = Particle.X() + Particle.V() * Dt;
-				Particle.R() = TRotation<T, d>::IntegrateRotationWithAngularVelocity(Particle.R(), Particle.W(), Dt);
+				Particle.R() = FRotation3::IntegrateRotationWithAngularVelocity(Particle.R(), Particle.W(), Dt);
 				break;
 			}
 			}
@@ -542,8 +536,8 @@ class TPBDRigidsEvolutionBase
 	}
 
 	/** Make a copy of the acceleration structure to allow for external modification. This is needed for supporting sync operations on SQ structure from game thread */
-	CHAOS_API void UpdateExternalAccelerationStructure(TUniquePtr<ISpatialAccelerationCollection<TAccelerationStructureHandle<T, d>, T, d>>& ExternalStructure);
-	ISpatialAccelerationCollection<TAccelerationStructureHandle<T, d>, T, d>* GetSpatialAcceleration() { return InternalAcceleration.Get(); }
+	CHAOS_API void UpdateExternalAccelerationStructure(TUniquePtr<ISpatialAccelerationCollection<TAccelerationStructureHandle<FReal, 3>, FReal, 3>>& ExternalStructure);
+	ISpatialAccelerationCollection<TAccelerationStructureHandle<FReal, 3>, FReal, 3>* GetSpatialAcceleration() { return InternalAcceleration.Get(); }
 
 	/** Perform a blocking flush of the spatial acceleration structure for situations where we aren't simulating but must have an up to date structure */
 	CHAOS_API void FlushSpatialAcceleration();
@@ -551,9 +545,6 @@ class TPBDRigidsEvolutionBase
 	/** Rebuilds the spatial acceleration from scratch. This should only be used for perf testing */
 	CHAOS_API void RebuildSpatialAccelerationForPerfTest();
 
-	const auto& GetRigidClustering() const { return Clustering; }
-	auto& GetRigidClustering() { return Clustering; }
-
 	CHAOS_API const FPBDConstraintGraph& GetConstraintGraph() const { return ConstraintGraph; }
 	CHAOS_API FPBDConstraintGraph& GetConstraintGraph() { return ConstraintGraph; }
 
@@ -577,7 +568,7 @@ protected:
 	}
 
 	template <bool bPersistent>
-	FORCEINLINE_DEBUGGABLE void RemoveParticleFromAccelerationStructure(TGeometryParticleHandleImp<T, d, bPersistent>& ParticleHandle)
+	FORCEINLINE_DEBUGGABLE void RemoveParticleFromAccelerationStructure(TGeometryParticleHandleImp<FReal, 3, bPersistent>& ParticleHandle)
 	{
 		//TODO: at the moment we don't distinguish between the first time a particle is created and when it's just moved
 		// If we had this distinction we could simply remove the entry for the async queue
@@ -586,7 +577,7 @@ protected:
 
 		SpatialData.bDelete = true;
 		SpatialData.SpatialIdx = ParticleHandle.SpatialIdx();
-		SpatialData.AccelerationHandle = TAccelerationStructureHandle<T, d>(ParticleHandle);
+		SpatialData.AccelerationHandle = TAccelerationStructureHandle<FReal, 3>(ParticleHandle);
 
 		ExternalAccelerationQueue.FindOrAdd(UniqueIdx) = SpatialData;
 
@@ -598,7 +589,7 @@ protected:
 		InternalAcceleration->RemoveElementFrom(SpatialData.AccelerationHandle, SpatialData.SpatialIdx);
 	}
 
-	void UpdateConstraintPositionBasedState(T Dt)
+	void UpdateConstraintPositionBasedState(FReal Dt)
 	{
 		for (FPBDConstraintGraphRule* ConstraintRule : ConstraintRules)
 		{
@@ -631,12 +622,12 @@ protected:
 		}
 	}
 	
-	void UpdateVelocities(const T Dt, int32 Island)
+	void UpdateVelocities(const FReal Dt, int32 Island)
 	{
 		ParticleUpdateVelocity(ConstraintGraph.GetIslandParticles(Island), Dt);
 	}
 
-	void ApplyPushOut(const T Dt, int32 Island)
+	void ApplyPushOut(const FReal Dt, int32 Island)
 	{
 		bool bNeedsAnotherIteration = true;
 		for (int32 It = 0; It < NumPushOutIterations; ++It)
@@ -654,7 +645,7 @@ protected:
 		}
 	}
 
-	using FAccelerationStructure = ISpatialAccelerationCollection<TAccelerationStructureHandle<T, d>, T, d>;
+	using FAccelerationStructure = ISpatialAccelerationCollection<TAccelerationStructureHandle<FReal, 3>, FReal, 3>;
 
 	void ComputeIntermediateSpatialAcceleration(bool bBlock = false);
 	void FlushInternalAccelerationQueue();
@@ -675,7 +666,7 @@ protected:
 	TArrayCollectionArray<int32> ParticleDisableCount;
 	TArrayCollectionArray<bool> Collided;
 
-	TPBDRigidsSOAs<T, d>& Particles;
+	TPBDRigidsSOAs<FReal, 3>& Particles;
 	TUniquePtr<FAccelerationStructure> InternalAcceleration;
 	TUniquePtr<FAccelerationStructure> AsyncInternalAcceleration;
 	TUniquePtr<FAccelerationStructure> AsyncExternalAcceleration;
@@ -683,12 +674,11 @@ protected:
 	bool bExternalReady;
 	bool bIsSingleThreaded;
 
-	TPBDRigidClustering<FPBDRigidsEvolution, FPBDCollisionConstraint, T, d> Clustering;
 
 	/** Used for updating intermediate spatial structures when they are finished */
 	struct FPendingSpatialData
 	{
-		TAccelerationStructureHandle<T, d> AccelerationHandle;
+		TAccelerationStructureHandle<FReal, 3> AccelerationHandle;
 		FSpatialAccelerationIdx SpatialIdx;
 		bool bDelete;
 
@@ -780,9 +770,9 @@ protected:
 	/** Pending operations for the external acceleration structure*/
 	FPendingSpatialDataQueue ExternalAccelerationQueue;
 
-	/*void SerializePendingMap(FChaosArchive& Ar, TMap<TGeometryParticleHandle<T, d>*, FPendingSpatialData>& Map)
+	/*void SerializePendingMap(FChaosArchive& Ar, TMap<TGeometryParticleHandle<FReal, 3>*, FPendingSpatialData>& Map)
 	{
-		TArray<TGeometryParticleHandle<T, d>*> Keys;
+		TArray<TGeometryParticleHandle<FReal, 3>*> Keys;
 		if (!Ar.IsLoading())
 		{
 			Map.GenerateKeyArray(Keys);
@@ -799,15 +789,15 @@ protected:
 	/** Used for async acceleration rebuild */
 	TArrayAsMap<FUniqueIdx, uint32> ParticleToCacheInnerIdx;
 
-	TMap<FSpatialAccelerationIdx, TUniquePtr<TSpatialAccelerationCache<T, d>>> SpatialAccelerationCache;
+	TMap<FSpatialAccelerationIdx, TUniquePtr<FSpatialAccelerationCache>> SpatialAccelerationCache;
 
 	FORCEINLINE_DEBUGGABLE void ApplyParticlePendingData(const FPendingSpatialData& PendingData, FAccelerationStructure& SpatialAcceleration, bool bUpdateCache);
 
 	class FChaosAccelerationStructureTask
 	{
 	public:
-		FChaosAccelerationStructureTask(ISpatialAccelerationCollectionFactory<T,d>& InSpatialCollectionFactory
-			, const TMap<FSpatialAccelerationIdx, TUniquePtr<TSpatialAccelerationCache<T,d>>>& InSpatialAccelerationCache
+		FChaosAccelerationStructureTask(ISpatialAccelerationCollectionFactory& InSpatialCollectionFactory
+			, const TMap<FSpatialAccelerationIdx, TUniquePtr<FSpatialAccelerationCache>>& InSpatialAccelerationCache
 			, TUniquePtr<FAccelerationStructure>& InAccelerationStructure
 			, TUniquePtr<FAccelerationStructure>& InAccelerationStructureCopy
 			, bool InForceFullBuild
@@ -817,8 +807,8 @@ protected:
 		static FORCEINLINE ESubsequentsMode::Type GetSubsequentsMode();
 		void DoTask(ENamedThreads::Type CurrentThread, const FGraphEventRef& MyCompletionGraphEvent);
 
-		ISpatialAccelerationCollectionFactory<T, d>& SpatialCollectionFactory;
-		const TMap<FSpatialAccelerationIdx, TUniquePtr<TSpatialAccelerationCache<T, d>>>& SpatialAccelerationCache;
+		ISpatialAccelerationCollectionFactory& SpatialCollectionFactory;
+		const TMap<FSpatialAccelerationIdx, TUniquePtr<FSpatialAccelerationCache>>& SpatialAccelerationCache;
 		TUniquePtr<FAccelerationStructure>& AccelerationStructure;
 		TUniquePtr<FAccelerationStructure>& AccelerationStructureCopy;
 		bool IsForceFullBuild;
@@ -828,14 +818,7 @@ protected:
 
 	int32 NumIterations;
 	int32 NumPushOutIterations;
-	TUniquePtr<ISpatialAccelerationCollectionFactory<T, d>> SpatialCollectionFactory;
+	TUniquePtr<ISpatialAccelerationCollectionFactory> SpatialCollectionFactory;
 };
 
 }
-
-// Only way to make this compile at the moment due to visibility attribute issues. TODO: Change this once a fix for this problem is applied.
-#if PLATFORM_MAC || PLATFORM_LINUX
-extern template class CHAOS_API Chaos::TPBDRigidsEvolutionBase<Chaos::TPBDRigidsEvolutionGBF<float, 3>, Chaos::TPBDCollisionConstraints<float,3>, float, 3>;
-#else
-extern template class Chaos::TPBDRigidsEvolutionBase<Chaos::TPBDRigidsEvolutionGBF<float, 3>, Chaos::TPBDCollisionConstraints<float,3>, float, 3>;
-#endif
diff --git a/Engine/Source/Runtime/Experimental/Chaos/Public/Chaos/PBDRigidsEvolutionGBF.h b/Engine/Source/Runtime/Experimental/Chaos/Public/Chaos/PBDRigidsEvolutionGBF.h
--- a/Engine/Source/Runtime/Experimental/Chaos/Public/Chaos/PBDRigidsEvolutionGBF.h
+++ b/Engine/Source/Runtime/Experimental/Chaos/Public/Chaos/PBDRigidsEvolutionGBF.h
@@ -26,77 +26,63 @@ namespace Chaos
 	CHAOS_API extern float HackLinearDrag;
 	CHAOS_API extern float HackAngularDrag;
 
-	template<typename T, int d>
-	class TPBDRigidsEvolutionGBF;
+	class FPBDRigidsEvolutionGBF;
 
-	template<typename T, int d>
-	using TPBDRigidsEvolutionCallback = TFunction<void()>;
+	using FPBDRigidsEvolutionCallback = TFunction<void()>;
 
-	template<typename T, int d>
-	using TPBDRigidsEvolutionIslandCallback = TFunction<void(int32 Island)>;
+	using FPBDRigidsEvolutionIslandCallback = TFunction<void(int32 Island)>;
 
-	template<typename T, int d>
-	class TPBDRigidsEvolutionGBF : public TPBDRigidsEvolutionBase<TPBDRigidsEvolutionGBF<T, d>, TPBDCollisionConstraints<T, d>, T, d>
+	class FPBDRigidsEvolutionGBF : public FPBDRigidsEvolutionBase
 	{
 	public:
-		using Base = TPBDRigidsEvolutionBase<TPBDRigidsEvolutionGBF<T, d>, TPBDCollisionConstraints<T, d>, T, d>;
-		using Base::AddConstraintRule;
-		using Base::AddForceFunction;
-		using Base::ApplyKinematicTargets;
-		using Base::Clustering;
-		using Base::ForceRules;
-		using Base::Particles;
-		using Base::ParticleUpdatePosition;
-		using Base::SetParticleUpdatePositionFunction;
-		using Base::SetParticleUpdateVelocityFunction;
-		using Base::GenerateUniqueIdx;
-		using typename Base::FForceRule;
-		using FGravityForces = TPerParticleGravity<T, d>;
-		using FCollisionConstraints = TPBDCollisionConstraints<T, d>;
+		using Base = FPBDRigidsEvolutionBase;
+
+		using FGravityForces = TPerParticleGravity<FReal, 3>;
+		using FCollisionConstraints = TPBDCollisionConstraints<FReal, 3>;
 		using FCollisionConstraintRule = TPBDConstraintColorRule<FCollisionConstraints>;
 		using FCollisionDetector = TCollisionDetector<FSpatialAccelerationBroadPhase, FNarrowPhase, FAsyncCollisionReceiver, FCollisionConstraints>;
-		using FExternalForces = TPerParticleExternalForces<T, d>;
+		using FExternalForces = TPerParticleExternalForces<FReal, 3>;
 
 		static constexpr int32 DefaultNumIterations = 1;
 		static constexpr int32 DefaultNumPairIterations = 1;
 		static constexpr int32 DefaultNumPushOutIterations = 5;
 		static constexpr int32 DefaultNumPushOutPairIterations = 2;
 
-		CHAOS_API TPBDRigidsEvolutionGBF(TPBDRigidsSOAs<T, d>& InParticles, int32 InNumIterations = DefaultNumIterations, int32 InNumPushoutIterations = DefaultNumPushOutIterations, bool InIsSingleThreaded = false);
-		CHAOS_API ~TPBDRigidsEvolutionGBF() {}
+		CHAOS_API FPBDRigidsEvolutionGBF(TPBDRigidsSOAs<FReal, 3>& InParticles, int32 InNumIterations = DefaultNumIterations, int32 InNumPushoutIterations = DefaultNumPushOutIterations, bool InIsSingleThreaded = false);
+		CHAOS_API ~FPBDRigidsEvolutionGBF() {}
 
-		void SetPostIntegrateCallback(const TPBDRigidsEvolutionCallback<T, d>& Cb)
+		void SetPostIntegrateCallback(const FPBDRigidsEvolutionCallback& Cb)
 		{
 			PostIntegrateCallback = Cb;
 		}
 
-		void SetPostDetectCollisionsCallback(const TPBDRigidsEvolutionCallback<T, d>& Cb)
+		void SetPostDetectCollisionsCallback(const FPBDRigidsEvolutionCallback& Cb)
 		{
 			PostDetectCollisionsCallback = Cb;
 		}
 
-		void SetCollisionModifierCallback(const TCollisionModifierCallback<T, d>& Cb)
+		void SetCollisionModifierCallback(const TCollisionModifierCallback<FReal, 3>& Cb)
 		{
 			CollisionModifierCallback = Cb;
 		}
 
-		void SetPreApplyCallback(const TPBDRigidsEvolutionCallback<T, d>& Cb)
+		void SetPreApplyCallback(const FPBDRigidsEvolutionCallback& Cb)
 		{
 			PreApplyCallback = Cb;
 		}
 
-		void SetPostApplyCallback(const TPBDRigidsEvolutionIslandCallback<T, d>& Cb)
+		void SetPostApplyCallback(const FPBDRigidsEvolutionIslandCallback& Cb)
 		{
 			PostApplyCallback = Cb;
 		}
 
-		void SetPostApplyPushOutCallback(const TPBDRigidsEvolutionIslandCallback<T, d>& Cb)
+		void SetPostApplyPushOutCallback(const FPBDRigidsEvolutionIslandCallback& Cb)
 		{
 			PostApplyPushOutCallback = Cb;
 		}
 
-		CHAOS_API void Advance(const T Dt, const T MaxStepDt, const int32 MaxSteps);
-		CHAOS_API void AdvanceOneTimeStep(const T dt, const T StepFraction = (T)1.0);
+		CHAOS_API void Advance(const FReal Dt, const FReal MaxStepDt, const int32 MaxSteps);
+		CHAOS_API void AdvanceOneTimeStep(const FReal dt, const FReal StepFraction = (FReal)1.0);
 
 		using Base::PrepareConstraints;
 		using Base::UnprepareConstraints;
@@ -115,26 +101,29 @@ namespace Chaos
 		FGravityForces& GetGravityForces() { return GravityForces; }
 		const FGravityForces& GetGravityForces() const { return GravityForces; }
 
-		CHAOS_API inline void EndFrame(T Dt)
+		const auto& GetRigidClustering() const { return Clustering; }
+		auto& GetRigidClustering() { return Clustering; }
+
+		CHAOS_API inline void EndFrame(FReal Dt)
 		{
 			Particles.GetNonDisabledDynamicView().ParallelFor([&](auto& Particle, int32 Index) {
-				Particle.F() = TVector<T, 3>(0);
-				Particle.Torque() = TVector<T, 3>(0);
+				Particle.F() = FVec3(0);
+				Particle.Torque() = FVec3(0);
 			});
 		}
 
 		template<typename TParticleView>
-		void Integrate(const TParticleView& InParticles, T Dt)
+		void Integrate(const TParticleView& InParticles, FReal Dt)
 		{
 			//SCOPE_CYCLE_COUNTER(STAT_Integrate);
 			CHAOS_SCOPED_TIMER(Integrate);
-			TPerParticleEulerStepVelocity<T, d> EulerStepVelocityRule;
-			TPerParticleAddImpulses<T, d> AddImpulsesRule;
-			TPerParticleEtherDrag<T, d> EtherDragRule;
-			TPerParticlePBDEulerStep<T, d> EulerStepRule;
+			TPerParticleEulerStepVelocity<FReal, 3> EulerStepVelocityRule;
+			TPerParticleAddImpulses<FReal, 3> AddImpulsesRule;
+			TPerParticleEtherDrag<FReal, 3> EtherDragRule;
+			TPerParticlePBDEulerStep<FReal, 3> EulerStepRule;
 
-			const T MaxAngularSpeedSq = HackMaxAngularVelocity * HackMaxAngularVelocity;
-			const T MaxSpeedSq = HackMaxVelocity * HackMaxVelocity;
+			const FReal MaxAngularSpeedSq = HackMaxAngularVelocity * HackMaxAngularVelocity;
+			const FReal MaxSpeedSq = HackMaxVelocity * HackMaxVelocity;
 			InParticles.ParallelFor([&](auto& GeomParticle, int32 Index) {
 				//question: can we enforce this at the API layer? Right now islands contain non dynamic which makes this hard
 				auto PBDParticle = GeomParticle.CastToRigidParticle();
@@ -156,7 +145,7 @@ namespace Chaos
 
 					if (HackMaxAngularVelocity >= 0.f)
 					{
-						const T AngularSpeedSq = Particle.W().SizeSquared();
+						const FReal AngularSpeedSq = Particle.W().SizeSquared();
 						if (AngularSpeedSq > MaxAngularSpeedSq)
 						{
 							Particle.W() = Particle.W() * (HackMaxAngularVelocity / FMath::Sqrt(AngularSpeedSq));
@@ -165,7 +154,7 @@ namespace Chaos
 
 					if (HackMaxVelocity >= 0.f)
 					{
-						const T SpeedSq = Particle.V().SizeSquared();
+						const FReal SpeedSq = Particle.V().SizeSquared();
 						if (SpeedSq > MaxSpeedSq)
 						{
 							Particle.V() = Particle.V() * (HackMaxVelocity / FMath::Sqrt(SpeedSq));
@@ -176,8 +165,8 @@ namespace Chaos
 
 					if (Particle.HasBounds())
 					{
-						const TAABB<T, d>& LocalBounds = Particle.LocalBounds();
-						TAABB<T, d> WorldSpaceBounds = LocalBounds.TransformedAABB(TRigidTransform<T, d>(Particle.P(), Particle.Q()));
+						const FAABB3& LocalBounds = Particle.LocalBounds();
+						FAABB3 WorldSpaceBounds = LocalBounds.TransformedAABB(FRigidTransform3(Particle.P(), Particle.Q()));
 						WorldSpaceBounds.ThickenSymmetrically(Particle.V());
 						Particle.SetWorldSpaceInflatedBounds(WorldSpaceBounds);
 					}
@@ -193,16 +182,7 @@ namespace Chaos
 		CHAOS_API void Serialize(FChaosArchive& Ar);
 
 	protected:
-		using Base::Collided;
-		using Base::ConstraintRules;
-		using Base::CreateConstraintGraph;
-		using Base::CreateIslands;
-		using Base::GetConstraintGraph;
-		using Base::InternalAcceleration;
-		using Base::ParticleDisableCount;
-		using Base::PhysicsMaterials;
-		using Base::UpdateConstraintPositionBasedState;
-		using Base::UpdateVelocities;
+		TPBDRigidClustering<FPBDRigidsEvolutionGBF, TPBDCollisionConstraints<FReal, 3>, FReal, 3> Clustering;
 
 		FGravityForces GravityForces;
 		FCollisionConstraints CollisionConstraints;
@@ -210,11 +190,11 @@ namespace Chaos
 		FSpatialAccelerationBroadPhase BroadPhase;
 		FCollisionDetector CollisionDetector;
 
-		TPBDRigidsEvolutionCallback<T, d> PostIntegrateCallback;
-		TPBDRigidsEvolutionCallback<T, d> PostDetectCollisionsCallback;
-		TCollisionModifierCallback<T, d> CollisionModifierCallback;
-		TPBDRigidsEvolutionCallback<T, d> PreApplyCallback;
-		TPBDRigidsEvolutionIslandCallback<T, d> PostApplyCallback;
-		TPBDRigidsEvolutionIslandCallback<T, d> PostApplyPushOutCallback;
+		FPBDRigidsEvolutionCallback PostIntegrateCallback;
+		FPBDRigidsEvolutionCallback PostDetectCollisionsCallback;
+		TCollisionModifierCallback<FReal, 3> CollisionModifierCallback;
+		FPBDRigidsEvolutionCallback PreApplyCallback;
+		FPBDRigidsEvolutionIslandCallback PostApplyCallback;
+		FPBDRigidsEvolutionIslandCallback PostApplyPushOutCallback;
 	};
 }
diff --git a/Engine/Source/Runtime/Experimental/Chaos/Public/Chaos/PBDRigidsEvolutionPGS.h b/Engine/Source/Runtime/Experimental/Chaos/Public/Chaos/PBDRigidsEvolutionPGS.h
deleted file mode 100644
diff --git a/Engine/Source/Runtime/Experimental/ChaosSolvers/Public/PBDRigidsSolver.h b/Engine/Source/Runtime/Experimental/ChaosSolvers/Public/PBDRigidsSolver.h
--- a/Engine/Source/Runtime/Experimental/ChaosSolvers/Public/PBDRigidsSolver.h
+++ b/Engine/Source/Runtime/Experimental/ChaosSolvers/Public/PBDRigidsSolver.h
@@ -94,7 +94,7 @@ namespace Chaos
 
 		typedef Chaos::TGeometryParticle<float, 3> FParticle;
 		typedef Chaos::TGeometryParticleHandle<float, 3> FHandle;
-		typedef Chaos::TPBDRigidsEvolutionGBF<float, 3> FPBDRigidsEvolution;
+		typedef Chaos::FPBDRigidsEvolutionGBF FPBDRigidsEvolution;
 		typedef Chaos::TPBDCollisionConstraints<float, 3> FPBDCollisionConstraints;
 
 		typedef FPBDCollisionConstraints FCollisionConstraints;
diff --git a/Engine/Source/Runtime/PhysicsCore/Private/PhysTestSerializer.cpp b/Engine/Source/Runtime/PhysicsCore/Private/PhysTestSerializer.cpp
--- a/Engine/Source/Runtime/PhysicsCore/Private/PhysTestSerializer.cpp
+++ b/Engine/Source/Runtime/PhysicsCore/Private/PhysTestSerializer.cpp
@@ -142,7 +142,7 @@ void FPhysTestSerializer::SetPhysicsData(physx::PxScene& Scene)
 #endif
 }
 
-void FPhysTestSerializer::SetPhysicsData(Chaos::TPBDRigidsEvolutionGBF<float,3>& Evolution)
+void FPhysTestSerializer::SetPhysicsData(Chaos::FPBDRigidsEvolutionGBF& Evolution)
 {
 	bDiskDataIsChaos = true;
 	Data.Empty();
@@ -330,11 +330,11 @@ void FPhysTestSerializer::CreateChaosData()
 			++Idx;
 		}
 
-		ChaosEvolution = MakeUnique<TPBDRigidsEvolutionGBF<float, 3>>(Particles);
+		ChaosEvolution = MakeUnique<FPBDRigidsEvolutionGBF>(Particles);
 	}
 	else
 	{
-		ChaosEvolution = MakeUnique<TPBDRigidsEvolutionGBF<float, 3>>(Particles);
+		ChaosEvolution = MakeUnique<FPBDRigidsEvolutionGBF>(Particles);
 
 		FMemoryReader Ar(Data);
 		FChaosArchive ChaosAr(Ar);
diff --git a/Engine/Source/Runtime/PhysicsCore/Private/SQCapture.cpp b/Engine/Source/Runtime/PhysicsCore/Private/SQCapture.cpp
--- a/Engine/Source/Runtime/PhysicsCore/Private/SQCapture.cpp
+++ b/Engine/Source/Runtime/PhysicsCore/Private/SQCapture.cpp
@@ -340,7 +340,7 @@ void FSQCapture::Serialize(Chaos::FChaosArchive& Ar)
 	}
 }
 
-void FSQCapture::StartCaptureChaosSweep(const Chaos::TPBDRigidsEvolutionGBF<float, 3>& Evolution, const Chaos::FImplicitObject& InQueryGeom, const FTransform& InStartTM, const FVector& InDir, float InDeltaMag, FHitFlags InOutputFlags, const FQueryFilterData& QueryFilter, const FCollisionFilterData& FilterData, ICollisionQueryFilterCallbackBase& Callback)
+void FSQCapture::StartCaptureChaosSweep(const Chaos::FPBDRigidsEvolutionGBF& Evolution, const Chaos::FImplicitObject& InQueryGeom, const FTransform& InStartTM, const FVector& InDir, float InDeltaMag, FHitFlags InOutputFlags, const FQueryFilterData& QueryFilter, const FCollisionFilterData& FilterData, ICollisionQueryFilterCallbackBase& Callback)
 {
 	if (IsInGameThread())
 	{
@@ -368,7 +368,7 @@ void FSQCapture::EndCaptureChaosSweep(const ChaosInterface::FSQHitBuffer<ChaosIn
 	}
 }
 
-void FSQCapture::StartCaptureChaosRaycast(const Chaos::TPBDRigidsEvolutionGBF<float,3>& Evolution, const FVector& InStartPoint, const FVector& InDir, float InDeltaMag, FHitFlags InOutputFlags, const FQueryFilterData& QueryFilter, const FCollisionFilterData& FilterData, ICollisionQueryFilterCallbackBase& Callback)
+void FSQCapture::StartCaptureChaosRaycast(const Chaos::FPBDRigidsEvolutionGBF& Evolution, const FVector& InStartPoint, const FVector& InDir, float InDeltaMag, FHitFlags InOutputFlags, const FQueryFilterData& QueryFilter, const FCollisionFilterData& FilterData, ICollisionQueryFilterCallbackBase& Callback)
 {
 	if (IsInGameThread())
 	{
@@ -394,7 +394,7 @@ void FSQCapture::EndCaptureChaosRaycast(const ChaosInterface::FSQHitBuffer<Chaos
 	}
 }
 
-void FSQCapture::StartCaptureChaosOverlap(const Chaos::TPBDRigidsEvolutionGBF<float, 3>& Evolution, const Chaos::FImplicitObject& InQueryGeom, const FTransform& InStartTM, const FQueryFilterData& QueryFilter, const FCollisionFilterData& FilterData, ICollisionQueryFilterCallbackBase& Callback)
+void FSQCapture::StartCaptureChaosOverlap(const Chaos::FPBDRigidsEvolutionGBF& Evolution, const Chaos::FImplicitObject& InQueryGeom, const FTransform& InStartTM, const FQueryFilterData& QueryFilter, const FCollisionFilterData& FilterData, ICollisionQueryFilterCallbackBase& Callback)
 {
 	if (IsInGameThread())
 	{
@@ -507,7 +507,7 @@ void FSQCapture::CapturePhysXFilterResults(const PxScene& TransientScene, const
 }
 
 #if WITH_PHYSX
-void FSQCapture::CaptureChaosFilterResults(const Chaos::TPBDRigidsEvolutionGBF<float,3>& TransientEvolution, const FCollisionFilterData& FilterData, ICollisionQueryFilterCallbackBase& Callback)
+void FSQCapture::CaptureChaosFilterResults(const Chaos::FPBDRigidsEvolutionGBF& TransientEvolution, const FCollisionFilterData& FilterData, ICollisionQueryFilterCallbackBase& Callback)
 {
 	using namespace Chaos;
 	const TPBDRigidsSOAs<float, 3>& Particles = TransientEvolution.GetParticles();
diff --git a/Engine/Source/Runtime/PhysicsCore/Public/PhysTestSerializer.h b/Engine/Source/Runtime/PhysicsCore/Public/PhysTestSerializer.h
--- a/Engine/Source/Runtime/PhysicsCore/Public/PhysTestSerializer.h
+++ b/Engine/Source/Runtime/PhysicsCore/Public/PhysTestSerializer.h
@@ -33,8 +33,7 @@ namespace physx
 
 namespace Chaos
 {
-	template <typename, int>
-	class TPBDRigidsEvolutionGBF;
+	class FPBDRigidsEvolutionGBF;
 
 	class FChaosArchive;
 }
@@ -52,7 +51,7 @@ public:
 	void Serialize(const TCHAR* FilePrefix);
 
 	//Set the data from an external source. This will obliterate any existing data. Make sure you are not holding on to old internal data as it will go away
-	void SetPhysicsData(Chaos::TPBDRigidsEvolutionGBF<float, 3>& ChaosEvolution);
+	void SetPhysicsData(Chaos::FPBDRigidsEvolutionGBF& ChaosEvolution);
 	void SetPhysicsData(physx::PxScene& Scene);
 
 	const Chaos::FChaosArchiveContext* GetChaosContext() const
@@ -85,7 +84,7 @@ public:
 		return SQCapture.Get();
 	}
 
-	Chaos::TPBDRigidsEvolutionGBF<float, 3>* GetChaosData()
+	Chaos::FPBDRigidsEvolutionGBF* GetChaosData()
 	{
 		if (!bChaosDataReady)
 		{
@@ -130,7 +129,7 @@ private:
 
 	TUniquePtr<FSQCapture> SQCapture;
 
-	TUniquePtr<Chaos::TPBDRigidsEvolutionGBF<float, 3>> ChaosEvolution;
+	TUniquePtr<Chaos::FPBDRigidsEvolutionGBF> ChaosEvolution;
 	Chaos::TPBDRigidsSOAs<float, 3> Particles;
 	TArray <TUniquePtr<Chaos::TGeometryParticle<float, 3>>> GTParticles;
 
diff --git a/Engine/Source/Runtime/PhysicsCore/Public/SQCapture.h b/Engine/Source/Runtime/PhysicsCore/Public/SQCapture.h
--- a/Engine/Source/Runtime/PhysicsCore/Public/SQCapture.h
+++ b/Engine/Source/Runtime/PhysicsCore/Public/SQCapture.h
@@ -24,8 +24,7 @@ namespace ChaosInterface
 
 namespace Chaos
 {
-	template <typename T, int d>
-	class TPBDRigidsEvolutionGBF;
+	class FPBDRigidsEvolutionGBF;
 
 	class FChaosArchive;
 }
@@ -55,13 +54,13 @@ struct PHYSICSCORE_API FSQCapture
 	void EndCapturePhysXOverlap(const PxHitCallback<PxOverlapHit>& Results);
 #endif
 
-	void StartCaptureChaosSweep(const Chaos::TPBDRigidsEvolutionGBF<float, 3>& Evolution, const Chaos::FImplicitObject& InQueryGeom, const FTransform& InStartTM, const FVector& InDir, float InDeltaMag, FHitFlags InOutputFlags, const FQueryFilterData& QueryFilter, const FCollisionFilterData& FilterData, ICollisionQueryFilterCallbackBase& Callback);
+	void StartCaptureChaosSweep(const Chaos::FPBDRigidsEvolutionGBF& Evolution, const Chaos::FImplicitObject& InQueryGeom, const FTransform& InStartTM, const FVector& InDir, float InDeltaMag, FHitFlags InOutputFlags, const FQueryFilterData& QueryFilter, const FCollisionFilterData& FilterData, ICollisionQueryFilterCallbackBase& Callback);
 	void EndCaptureChaosSweep(const ChaosInterface::FSQHitBuffer<ChaosInterface::FSweepHit>& Results);
 
-	void StartCaptureChaosRaycast(const Chaos::TPBDRigidsEvolutionGBF<float,3>& Evolution, const FVector& InStartPoint, const FVector& InDir, float InDeltaMag, FHitFlags InOutputFlags, const FQueryFilterData& QueryFilter, const FCollisionFilterData& FilterData, ICollisionQueryFilterCallbackBase& Callback);
+	void StartCaptureChaosRaycast(const Chaos::FPBDRigidsEvolutionGBF& Evolution, const FVector& InStartPoint, const FVector& InDir, float InDeltaMag, FHitFlags InOutputFlags, const FQueryFilterData& QueryFilter, const FCollisionFilterData& FilterData, ICollisionQueryFilterCallbackBase& Callback);
 	void EndCaptureChaosRaycast(const ChaosInterface::FSQHitBuffer<ChaosInterface::FRaycastHit>& Results);
 
-	void StartCaptureChaosOverlap(const Chaos::TPBDRigidsEvolutionGBF<float, 3>& Evolution, const Chaos::FImplicitObject& InQueryGeom, const FTransform& InStartTM, const FQueryFilterData& QueryFilter, const FCollisionFilterData& FilterData, ICollisionQueryFilterCallbackBase& Callback);
+	void StartCaptureChaosOverlap(const Chaos::FPBDRigidsEvolutionGBF& Evolution, const Chaos::FImplicitObject& InQueryGeom, const FTransform& InStartTM, const FQueryFilterData& QueryFilter, const FCollisionFilterData& FilterData, ICollisionQueryFilterCallbackBase& Callback);
 	void EndCaptureChaosOverlap(const ChaosInterface::FSQHitBuffer<ChaosInterface::FOverlapHit>& Results);
 
 	ECollisionQueryHitType GetFilterResult(const Chaos::TPerShapeData<float,3>* Shape, const Chaos::TGeometryParticle<float,3>* Actor) const;
@@ -158,7 +157,7 @@ private:
 
 	FPhysTestSerializer& PhysSerializer;
 
-	void CaptureChaosFilterResults(const Chaos::TPBDRigidsEvolutionGBF<float,3>& Evolution, const FCollisionFilterData& FilterData, ICollisionQueryFilterCallbackBase& Callback);
+	void CaptureChaosFilterResults(const Chaos::FPBDRigidsEvolutionGBF& Evolution, const FCollisionFilterData& FilterData, ICollisionQueryFilterCallbackBase& Callback);
 
 	TMap<Chaos::TGeometryParticle<float,3>*, TArray<TPair<Chaos::TPerShapeData<float,3>*, ECollisionQueryHitType>>> ChaosActorToShapeHitsArray;
 
