Allow access to debugname through TGenericParticleHandleHandleImp

diff --git a/Engine/Source/Runtime/Experimental/Chaos/Public/Chaos/ParticleHandle.h b/Engine/Source/Runtime/Experimental/Chaos/Public/Chaos/ParticleHandle.h
--- a/Engine/Source/Runtime/Experimental/Chaos/Public/Chaos/ParticleHandle.h
+++ b/Engine/Source/Runtime/Experimental/Chaos/Public/Chaos/ParticleHandle.h
@@ -1149,6 +1149,13 @@ public:
 		return TRotation<T, d>::FromIdentity();
 	}
 
+#if CHAOS_CHECKED
+	const FName& DebugName() const
+	{
+		return Handle->DebugName();
+	}
+#endif
+
 	int32 Island() const
 	{
 		if (Handle->CastToRigidParticle() && Handle->ObjectState() == EObjectStateType::Dynamic)
