Fix build.

diff --git a/Engine/Source/Runtime/Experimental/Chaos/Private/Chaos/MassProperties.cpp b/Engine/Source/Runtime/Experimental/Chaos/Private/Chaos/MassProperties.cpp
--- a/Engine/Source/Runtime/Experimental/Chaos/Private/Chaos/MassProperties.cpp
+++ b/Engine/Source/Runtime/Experimental/Chaos/Private/Chaos/MassProperties.cpp
@@ -73,7 +73,6 @@ namespace Chaos
 		
 		return FinalRotation;
 	}
-	template TRotation<float, 3> TransformToLocalSpace(PMatrix<float, 3, 3>& Inertia);
 
 	template<typename T, int d, typename TSurfaces>
 	void CalculateVolumeAndCenterOfMass(const TParticles<T, d>& Vertices, const TSurfaces& Surfaces, T& OutVolume, TVector<T, d>& OutCenterOfMass)
