Exposed IsHole for heightfields and added handling to navigation export for holes.

diff --git a/Engine/Source/Runtime/Experimental/Chaos/Private/Chaos/HeightField.cpp b/Engine/Source/Runtime/Experimental/Chaos/Private/Chaos/HeightField.cpp
--- a/Engine/Source/Runtime/Experimental/Chaos/Private/Chaos/HeightField.cpp
+++ b/Engine/Source/Runtime/Experimental/Chaos/Private/Chaos/HeightField.cpp
@@ -567,6 +567,38 @@ namespace Chaos
 		return GetHeight(Index);
 	}
 
+	template<typename T>
+	uint8 Chaos::THeightField<T>::GetMaterialIndex(int32 InIndex) const
+	{
+		if(CHAOS_ENSURE(InIndex >= 0 && InIndex < GeomData.MaterialIndices.Num()))
+		{
+			return GeomData.MaterialIndices[InIndex];
+		}
+
+		return TNumericLimits<uint8>::Max();
+	}
+
+	template<typename T>
+	uint8 Chaos::THeightField<T>::GetMaterialIndex(int32 InX, int32 InY) const
+	{
+		const int32 Index = InY * GeomData.NumCols + InX;
+		return GetMaterialIndex(Index);
+	}
+
+	template<typename T>
+	bool Chaos::THeightField<T>::IsHole(int32 InIndex) const
+	{
+		return GetMaterialIndex(InIndex) == TNumericLimits<uint8>::Max();
+	}
+
+	template<typename T>
+	bool Chaos::THeightField<T>::IsHole(int32 InCellX, int32 InCellY) const
+	{
+		// Convert to single cell index
+		const int32 Index = InCellY * (GeomData.NumCols - 1) + InCellX;
+		return IsHole(Index);
+	}
+
 	template<typename T>
 	T Chaos::THeightField<T>::GetHeightAt(const TVector<T, 2>& InGridLocationLocal) const
 	{
diff --git a/Engine/Source/Runtime/Experimental/Chaos/Public/Chaos/HeightField.h b/Engine/Source/Runtime/Experimental/Chaos/Public/Chaos/HeightField.h
--- a/Engine/Source/Runtime/Experimental/Chaos/Public/Chaos/HeightField.h
+++ b/Engine/Source/Runtime/Experimental/Chaos/Public/Chaos/HeightField.h
@@ -41,6 +41,10 @@ namespace Chaos
 		void EditHeights(TArrayView<const uint16> InHeights, int32 InBeginRow, int32 InBeginCol, int32 InNumRows, int32 InNumCols);
 		T GetHeight(int32 InIndex) const;
 		T GetHeight(int32 InX, int32 InY) const;
+		uint8 GetMaterialIndex(int32 InIndex) const;
+		uint8 GetMaterialIndex(int32 InX, int32 InY) const;
+		bool IsHole(int32 InIndex) const;
+		bool IsHole(int32 InCellX, int32 InCellY) const;
 		T GetHeightAt(const TVector<T, 2>& InGridLocationLocal) const;
 
 		int32 GetNumRows() const { return GeomData.NumRows; }
diff --git a/Engine/Source/Runtime/NavigationSystem/Private/NavMesh/RecastNavMeshGenerator.cpp b/Engine/Source/Runtime/NavigationSystem/Private/NavMesh/RecastNavMeshGenerator.cpp
--- a/Engine/Source/Runtime/NavigationSystem/Private/NavMesh/RecastNavMeshGenerator.cpp
+++ b/Engine/Source/Runtime/NavigationSystem/Private/NavMesh/RecastNavMeshGenerator.cpp
@@ -750,9 +750,7 @@ void ExportChaosHeightField(const Chaos::THeightField<float>* const HeightField,
 	{
 		for(int32 X = 0; X < NumCols - 1; X++)
 		{
-			// #PHYSTODO Hole support for chaos heightfields
-			const bool bIsHole = false;
-			if(bIsHole)
+			if(HeightField->IsHole(X, Y))
 			{
 				continue;
 			}
