Make ensures CHAOS_ENSURE to avoid server hitch.

diff --git a/Engine/Source/Runtime/Experimental/Chaos/Public/Chaos/ImplicitObjectScaled.h b/Engine/Source/Runtime/Experimental/Chaos/Public/Chaos/ImplicitObjectScaled.h
--- a/Engine/Source/Runtime/Experimental/Chaos/Public/Chaos/ImplicitObjectScaled.h
+++ b/Engine/Source/Runtime/Experimental/Chaos/Public/Chaos/ImplicitObjectScaled.h
@@ -515,8 +515,8 @@ public:
 	{
 		if (T UnscaledMargin = MObject->GetMargin())
 		{
-			ensure(FMath::IsNearlyEqual(MScale[0], MScale[1], KINDA_SMALL_NUMBER));
-			ensure(FMath::IsNearlyEqual(MScale[1], MScale[2], KINDA_SMALL_NUMBER));
+			CHAOS_ENSURE(FMath::IsNearlyEqual(MScale[0], MScale[1], KINDA_SMALL_NUMBER));
+			CHAOS_ENSURE(FMath::IsNearlyEqual(MScale[1], MScale[2], KINDA_SMALL_NUMBER));
 
 			return UnscaledMargin * FMath::Abs(MScale[0]);
 		}
