Chaos - joint fix multiple locked axis math

diff --git a/Engine/Source/Runtime/Experimental/Chaos/Private/Chaos/Joint/PBDJointSolverGaussSeidel.cpp b/Engine/Source/Runtime/Experimental/Chaos/Private/Chaos/Joint/PBDJointSolverGaussSeidel.cpp
--- a/Engine/Source/Runtime/Experimental/Chaos/Private/Chaos/Joint/PBDJointSolverGaussSeidel.cpp
+++ b/Engine/Source/Runtime/Experimental/Chaos/Private/Chaos/Joint/PBDJointSolverGaussSeidel.cpp
@@ -337,7 +337,10 @@ namespace Chaos
 			}
 			else if (TwistMotion == EJointMotionType::Locked)
 			{
-				NetResult += ApplyTwistConstraint(Dt, SolverSettings, JointSettings, false);
+				// Covered below
+			}
+			else if (TwistMotion == EJointMotionType::Free)
+			{
 			}
 		}
 
@@ -351,7 +354,7 @@ namespace Chaos
 			}
 			else if ((Swing1Motion == EJointMotionType::Limited) && (Swing2Motion == EJointMotionType::Locked))
 			{
-				NetResult += ApplyLockedSwingConstraint(Dt, SolverSettings, JointSettings, EJointAngularConstraintIndex::Swing2, false);
+				NetResult += ApplySingleLockedSwingConstraint(Dt, SolverSettings, JointSettings, EJointAngularConstraintIndex::Swing2, false);
 				if (!bDegenerate)
 				{
 					NetResult += ApplySwingConstraint(Dt, SolverSettings, JointSettings, EJointAngularConstraintIndex::Swing1, bSwingSoft);
@@ -366,7 +369,7 @@ namespace Chaos
 			}
 			else if ((Swing1Motion == EJointMotionType::Locked) && (Swing2Motion == EJointMotionType::Limited))
 			{
-				NetResult += ApplyLockedSwingConstraint(Dt, SolverSettings, JointSettings, EJointAngularConstraintIndex::Swing1, false);
+				NetResult += ApplySingleLockedSwingConstraint(Dt, SolverSettings, JointSettings, EJointAngularConstraintIndex::Swing1, false);
 				if (!bDegenerate)
 				{
 					NetResult += ApplySwingConstraint(Dt, SolverSettings, JointSettings, EJointAngularConstraintIndex::Swing2, bSwingSoft);
@@ -374,12 +377,11 @@ namespace Chaos
 			}
 			else if ((Swing1Motion == EJointMotionType::Locked) && (Swing2Motion == EJointMotionType::Locked))
 			{
-				NetResult += ApplyLockedSwingConstraint(Dt, SolverSettings, JointSettings, EJointAngularConstraintIndex::Swing1, false);
-				NetResult += ApplyLockedSwingConstraint(Dt, SolverSettings, JointSettings, EJointAngularConstraintIndex::Swing2, false);
+				// Covered below
 			}
 			else if ((Swing1Motion == EJointMotionType::Locked) && (Swing2Motion == EJointMotionType::Free))
 			{
-				NetResult += ApplyLockedSwingConstraint(Dt, SolverSettings, JointSettings, EJointAngularConstraintIndex::Swing1, false);
+				NetResult += ApplySingleLockedSwingConstraint(Dt, SolverSettings, JointSettings, EJointAngularConstraintIndex::Swing1, false);
 			}
 			else if ((Swing1Motion == EJointMotionType::Free) && (Swing2Motion == EJointMotionType::Limited))
 			{
@@ -390,13 +392,24 @@ namespace Chaos
 			}
 			else if ((Swing1Motion == EJointMotionType::Free) && (Swing2Motion == EJointMotionType::Locked))
 			{
-				NetResult += ApplyLockedSwingConstraint(Dt, SolverSettings, JointSettings, EJointAngularConstraintIndex::Swing2, false);
+				NetResult += ApplySingleLockedSwingConstraint(Dt, SolverSettings, JointSettings, EJointAngularConstraintIndex::Swing2, false);
 			}
 			else if ((Swing1Motion == EJointMotionType::Free) && (Swing2Motion == EJointMotionType::Free))
 			{
 			}
 		}
 
+		// Note: single-swing locks are already handled above so we only need to do something here if both are locked
+		bool bLockedTwist = SolverSettings.bEnableTwistLimits 
+			&& (JointSettings.AngularMotionTypes[(int32)EJointAngularConstraintIndex::Twist] == EJointMotionType::Locked);
+		bool bLockedSwing = SolverSettings.bEnableSwingLimits 
+			&& (JointSettings.AngularMotionTypes[(int32)EJointAngularConstraintIndex::Swing1] == EJointMotionType::Locked) 
+			&& (JointSettings.AngularMotionTypes[(int32)EJointAngularConstraintIndex::Swing2] == EJointMotionType::Locked);
+		if (bLockedTwist || bLockedSwing)
+		{
+			NetResult += ApplyLockedRotationConstraints(Dt, SolverSettings, JointSettings, bLockedTwist, bLockedSwing);
+		}
+
 		return NetResult;
 	}
 
@@ -994,6 +1007,37 @@ namespace Chaos
 	//
 	//
 
+	FJointSolverResult FJointSolverGaussSeidel::ApplyLockedRotationConstraints(
+		const FReal Dt,
+		const FPBDJointSolverSettings& SolverSettings,
+		const FPBDJointSettings& JointSettings,
+		const bool bApplyTwist,
+		const bool bApplySwing)
+	{
+		FVec3 Axis0, Axis1, Axis2;
+		FPBDJointUtilities::GetLockedAxes(Rs[0], Rs[1], Axis0, Axis1, Axis2);
+
+		const FRotation3 R01 = Rs[0].Inverse() * Rs[1];
+
+		FJointSolverResult NetResult;
+		if (bApplyTwist)
+		{
+			FReal TwistStiffness = FPBDJointUtilities::GetTwistStiffness(SolverSettings, JointSettings);
+			ApplyRotationConstraint(TwistStiffness, Axis0, R01.X);
+			NetResult += FJointSolverResult::MakeActive();
+		}
+
+		if (bApplySwing)
+		{
+			FReal SwingStiffness = FPBDJointUtilities::GetSwingStiffness(SolverSettings, JointSettings);
+			ApplyRotationConstraint(SwingStiffness, Axis1, R01.Y);
+			ApplyRotationConstraint(SwingStiffness, Axis2, R01.Z);
+			NetResult += FJointSolverResult::MakeActive();
+			NetResult += FJointSolverResult::MakeActive();
+		}
+
+		return NetResult;
+	}
 
 	FJointSolverResult FJointSolverGaussSeidel::ApplyTwistConstraint(
 		const FReal Dt,
@@ -1160,7 +1204,7 @@ namespace Chaos
 	}
 
 
-	FJointSolverResult FJointSolverGaussSeidel::ApplyLockedSwingConstraint(
+	FJointSolverResult FJointSolverGaussSeidel::ApplySingleLockedSwingConstraint(
 		const FReal Dt,
 		const FPBDJointSolverSettings& SolverSettings,
 		const FPBDJointSettings& JointSettings,
diff --git a/Engine/Source/Runtime/Experimental/Chaos/Private/Chaos/PBDJointConstraintUtilities.cpp b/Engine/Source/Runtime/Experimental/Chaos/Private/Chaos/PBDJointConstraintUtilities.cpp
--- a/Engine/Source/Runtime/Experimental/Chaos/Private/Chaos/PBDJointConstraintUtilities.cpp
+++ b/Engine/Source/Runtime/Experimental/Chaos/Private/Chaos/PBDJointConstraintUtilities.cpp
@@ -126,6 +126,33 @@ namespace Chaos
 	}
 
 
+	void FPBDJointUtilities::GetLockedAxes(const FRotation3& R0, const FRotation3& R1, FVec3& Axis0, FVec3& Axis1, FVec3& Axis2)
+	{
+		const FReal W0 = R0.W;
+		const FReal W1 = R1.W;
+		const FVec3 V0 = FVec3(R0.X, R0.Y, R0.Z);
+		const FVec3 V1 = FVec3(R1.X, R1.Y, R1.Z);
+
+		const FVec3 C = V1 * W0 + V0 * W1;
+		const FReal D0 = W0 * W1;
+		const FReal D1 = FVec3::DotProduct(V0, V1);
+		const FReal D = D0 - D1;
+
+		Axis0 = 0.5f * (V0 * V1.X + V1 * V0.X + FVec3(D, C.Z, -C.Y));
+		Axis1 = 0.5f * (V0 * V1.Y + V1 * V0.Y + FVec3(-C.Z, D, C.X));
+		Axis2 = 0.5f * (V0 * V1.Z + V1 * V0.Z + FVec3(C.Y, -C.X, D));
+
+		// Handle degenerate case of 180 deg swing
+		if (FMath::Abs(D0 + D1) < SMALL_NUMBER)
+		{
+			const FReal Epsilon = SMALL_NUMBER;
+			Axis0.X += Epsilon;
+			Axis1.Y += Epsilon;
+			Axis2.Z += Epsilon;
+		}
+	}
+	
+	
 	// @todo(ccaulfield): separate linear soft and stiff
 	FReal FPBDJointUtilities::GetLinearStiffness(
 		const FPBDJointSolverSettings& SolverSettings,
diff --git a/Engine/Source/Runtime/Experimental/Chaos/Public/Chaos/Joint/PBDJointSolverGaussSeidel.h b/Engine/Source/Runtime/Experimental/Chaos/Public/Chaos/Joint/PBDJointSolverGaussSeidel.h
--- a/Engine/Source/Runtime/Experimental/Chaos/Public/Chaos/Joint/PBDJointSolverGaussSeidel.h
+++ b/Engine/Source/Runtime/Experimental/Chaos/Public/Chaos/Joint/PBDJointSolverGaussSeidel.h
@@ -247,6 +247,13 @@ namespace Chaos
 			const FReal Angle,
 			FReal& Lambda);
 
+		FJointSolverResult ApplyLockedRotationConstraints(
+			const FReal Dt,
+			const FPBDJointSolverSettings& SolverSettings,
+			const FPBDJointSettings& JointSettings,
+			const bool bApplyTwist,
+			const bool bApplySwing);
+
 		FJointSolverResult ApplyTwistConstraint(
 			const FReal Dt,
 			const FPBDJointSolverSettings& SolverSettings,
@@ -280,7 +287,7 @@ namespace Chaos
 			const FPBDJointSettings& JointSettings);
 
 		// One Swing axis is free, and the other locked. This applies the lock: Body1 Twist axis is confined to a plane.
-		FJointSolverResult ApplyLockedSwingConstraint(
+		FJointSolverResult ApplySingleLockedSwingConstraint(
 			const FReal Dt,
 			const FPBDJointSolverSettings& SolverSettings,
 			const FPBDJointSettings& JointSettings,
diff --git a/Engine/Source/Runtime/Experimental/Chaos/Public/Chaos/PBDJointConstraintUtilities.h b/Engine/Source/Runtime/Experimental/Chaos/Public/Chaos/PBDJointConstraintUtilities.h
--- a/Engine/Source/Runtime/Experimental/Chaos/Public/Chaos/PBDJointConstraintUtilities.h
+++ b/Engine/Source/Runtime/Experimental/Chaos/Public/Chaos/PBDJointConstraintUtilities.h
@@ -63,6 +63,13 @@ namespace Chaos
 			FVec3& Axis,
 			FReal& Angle);
 
+		static void GetLockedAxes(
+			const FRotation3& R0, 
+			const FRotation3& R1, 
+			FVec3& Axis0, 
+			FVec3& Axis1, 
+			FVec3& Axis2);
+
 		/**
 		 * Increase the lower inertia components to ensure that the maximum ratio between any pair of elements is MaxRatio.
 		 *
