Chaos - More inertia fixes

diff --git a/Engine/Source/Runtime/Engine/Private/PhysicsEngine/Experimental/PhysInterface_Chaos.cpp b/Engine/Source/Runtime/Engine/Private/PhysicsEngine/Experimental/PhysInterface_Chaos.cpp
--- a/Engine/Source/Runtime/Engine/Private/PhysicsEngine/Experimental/PhysInterface_Chaos.cpp
+++ b/Engine/Source/Runtime/Engine/Private/PhysicsEngine/Experimental/PhysInterface_Chaos.cpp
@@ -818,7 +818,7 @@ void FPhysInterface_Chaos::AddAngularVelocityInRadians_AssumesLocked(const FPhys
 		Chaos::TPBDRigidParticle<float, 3>* Rigid = InActorReference->CastToRigidParticle();
 		if (ensure(Rigid))
 		{
-			const Chaos::FMatrix33 WorldI = Chaos::Utilities::ComputeWorldSpaceInertia(Rigid->R(), Rigid->I());
+			const Chaos::FMatrix33 WorldI = Chaos::FParticleUtilitiesXR::GetWorldInertia(Rigid);
 			AddAngularImpulseInRadians_AssumesLocked(InActorReference, WorldI * InAngularVelocityDeltaRad);
 		}
 	}
diff --git a/Engine/Source/Runtime/Engine/Private/PhysicsEngine/Experimental/PhysScene_Chaos.cpp b/Engine/Source/Runtime/Engine/Private/PhysicsEngine/Experimental/PhysScene_Chaos.cpp
--- a/Engine/Source/Runtime/Engine/Private/PhysicsEngine/Experimental/PhysScene_Chaos.cpp
+++ b/Engine/Source/Runtime/Engine/Private/PhysicsEngine/Experimental/PhysScene_Chaos.cpp
@@ -1527,7 +1527,7 @@ void FPhysScene_ChaosInterface::AddTorque_AssumesLocked(FBodyInstance* BodyInsta
 				const Chaos::TVector<float, 3> CurrentTorque = Rigid->Torque();
 				if (bAccelChange)
 				{
-					Rigid->SetTorque(CurrentTorque + (Utilities::ComputeWorldSpaceInertia(Rigid->R() * Rigid->RotationOfMass(), Rigid->I()) * Torque));
+					Rigid->SetTorque(CurrentTorque + (FParticleUtilitiesXR::GetWorldInertia(Rigid) * Torque));
 				}
 				else
 				{
diff --git a/Engine/Source/Runtime/Engine/Private/PhysicsEngine/ImmediatePhysics/ImmediatePhysicsChaos/ImmediatePhysicsActorHandle_Chaos.cpp b/Engine/Source/Runtime/Engine/Private/PhysicsEngine/ImmediatePhysics/ImmediatePhysicsChaos/ImmediatePhysicsActorHandle_Chaos.cpp
--- a/Engine/Source/Runtime/Engine/Private/PhysicsEngine/ImmediatePhysics/ImmediatePhysicsChaos/ImmediatePhysicsActorHandle_Chaos.cpp
+++ b/Engine/Source/Runtime/Engine/Private/PhysicsEngine/ImmediatePhysics/ImmediatePhysicsChaos/ImmediatePhysicsActorHandle_Chaos.cpp
@@ -177,7 +177,7 @@ namespace ImmediatePhysics_Chaos
 		return true;
 	}
 
-	bool CreateGeometry(FBodyInstance* BodyInstance, const FVector& Scale, float& OutMass, Chaos::TVector<float, 3>& OutInertia, Chaos::TRigidTransform<float, 3>& OutCoMTransform, TUniquePtr<Chaos::FImplicitObject>& OutGeom, TArray<TUniquePtr<Chaos::TPerShapeData<float, 3>>>& OutShapes)
+	bool CreateGeometry(FBodyInstance* BodyInstance, EActorType ActorType, const FVector& Scale, float& OutMass, Chaos::TVector<float, 3>& OutInertia, Chaos::TRigidTransform<float, 3>& OutCoMTransform, TUniquePtr<Chaos::FImplicitObject>& OutGeom, TArray<TUniquePtr<Chaos::TPerShapeData<float, 3>>>& OutShapes)
 	{
 		using namespace Chaos;
 
@@ -189,23 +189,28 @@ namespace ImmediatePhysics_Chaos
 		UBodySetup* BodySetup = BodyInstance->BodySetup.Get();
 
 #if WITH_CHAOS && !PHYSICS_INTERFACE_PHYSX
-		TMassProperties<float, 3> MassProperties;
-		CalculateMassProperties(Scale, FTransform::Identity, BodySetup->AggGeom, MassProperties);
-		float Density = 1.e-3f;	// 1g/cm3	@todo(ccaulfield): should come from material
-		if (BodyInstance->bOverrideMass)
+		OutMass = 0.0f;
+		OutInertia = FVector::ZeroVector;
+		OutCoMTransform = FTransform::Identity;
+		if (ActorType == EActorType::DynamicActor)
 		{
-			Density = BodyInstance->GetMassOverride() / MassProperties.Volume;
+			TMassProperties<float, 3> MassProperties;
+			CalculateMassProperties(Scale, FTransform::Identity, BodySetup->AggGeom, MassProperties);
+			float Density = 1.e-3f;	// 1g/cm3	@todo(ccaulfield): should come from material
+			if (BodyInstance->bOverrideMass)
+			{
+				Density = BodyInstance->GetMassOverride() / MassProperties.Volume;
+			}
+			OutMass = Density * BodyInstance->MassScale * MassProperties.Volume;
+			OutInertia = Utilities::ScaleInertia(Density * TVector<float, 3>(MassProperties.InertiaTensor.M[0][0], MassProperties.InertiaTensor.M[1][1], MassProperties.InertiaTensor.M[2][2]), BodyInstance->InertiaTensorScale, true);	// bScaleMass true to match legacy, but not correct
+			OutCoMTransform = FTransform(MassProperties.RotationOfMass, MassProperties.CenterOfMass + BodyInstance->COMNudge);
 		}
-		OutMass = Density * BodyInstance->MassScale * MassProperties.Volume;
-		OutInertia = Utilities::ScaleInertia(Density * TVector<float, 3>(MassProperties.InertiaTensor.M[0][0], MassProperties.InertiaTensor.M[1][1], MassProperties.InertiaTensor.M[2][2]), BodyInstance->InertiaTensorScale, true);	// bScaleMass true to match legacy, but not correct
-		OutCoMTransform = FTransform(MassProperties.RotationOfMass, MassProperties.CenterOfMass + BodyInstance->COMNudge);
 #else
 		OutMass = BodyInstance->GetBodyMass();
 		OutInertia = BodyInstance->GetBodyInertiaTensor();
 		OutCoMTransform = BodyInstance->GetMassSpaceLocal();
 #endif
 
-
 		FBodyCollisionData BodyCollisionData;
 		BodyInstance->BuildBodyFilterData(BodyCollisionData.CollisionFilterData);
 		BodyInstance->BuildBodyCollisionFlags(BodyCollisionData.CollisionFlags, BodyInstance->GetCollisionEnabled(), BodyInstance->BodySetup->GetCollisionTraceFlag() == CTF_UseComplexAsSimple);
@@ -271,7 +276,7 @@ namespace ImmediatePhysics_Chaos
 		float Mass = 0;
 		FVec3 Inertia = FVec3::OneVector;
 		FRigidTransform3 CoMTransform = FRigidTransform3::Identity;
-		if (CreateGeometry(BodyInstance, FVector::OneVector, Mass, Inertia, CoMTransform, Geometry, Shapes))
+		if (CreateGeometry(BodyInstance, ActorType, FVector::OneVector, Mass, Inertia, CoMTransform, Geometry, Shapes))
 		{
 			switch (ActorType)
 			{
diff --git a/Engine/Source/Runtime/Experimental/Chaos/Private/Chaos/Evolution/PBDMinEvolution.cpp b/Engine/Source/Runtime/Experimental/Chaos/Private/Chaos/Evolution/PBDMinEvolution.cpp
--- a/Engine/Source/Runtime/Experimental/Chaos/Private/Chaos/Evolution/PBDMinEvolution.cpp
+++ b/Engine/Source/Runtime/Experimental/Chaos/Private/Chaos/Evolution/PBDMinEvolution.cpp
@@ -136,8 +136,11 @@ namespace Chaos
 				Particle.PreV() = Particle.V();
 				Particle.PreW() = Particle.W();
 
+				const FVec3 XCoM = FParticleUtilitiesXR::GetCoMWorldPosition(&Particle);
+				const FRotation3 RCoM = FParticleUtilitiesXR::GetCoMWorldRotation(&Particle);
+
 				// Calculate new velocities from forces, torques and drag
-				const FMatrix33 WorldInvI = Utilities::ComputeWorldSpaceInertia(Particle.R(), Particle.InvI());
+				const FMatrix33 WorldInvI = Utilities::ComputeWorldSpaceInertia(RCoM, Particle.InvI());
 				const FVec3 DV = Particle.InvM() * (Particle.F() * Dt + Particle.LinearImpulse());
 				const FVec3 DW = WorldInvI * (Particle.Torque() * Dt + Particle.AngularImpulse());
 				const FReal LinearDrag = FMath::Max(FReal(0), FReal(1) - (Particle.LinearEtherDrag() * Dt));
@@ -145,12 +148,10 @@ namespace Chaos
 				const FVec3 V = LinearDrag * (Particle.V() + DV);
 				const FVec3 W = AngularDrag * (Particle.W() + DW);
 
-				FVec3 PCoM = FParticleUtilitiesXR::GetCoMWorldPosition(&Particle);
-				FRotation3 QCoM = FParticleUtilitiesXR::GetCoMWorldRotation(&Particle);
-				PCoM = PCoM + V * Dt;
-				QCoM = FRotation3::IntegrateRotationWithAngularVelocity(QCoM, W, Dt);
+				const FVec3 PCoM = XCoM + V * Dt;
+				const FRotation3 QCoM = FRotation3::IntegrateRotationWithAngularVelocity(RCoM, W, Dt);
 
-				// Update particle state (forces are zeroed at the end of the frame)
+				// Update particle state (forces are not zeroed until the end of the frame)
 				FParticleUtilitiesPQ::SetCoMWorldTransform(&Particle, PCoM, QCoM);
 				Particle.V() = V;
 				Particle.W() = W;
diff --git a/Engine/Source/Runtime/Experimental/Chaos/Private/Chaos/PBDRigidDynamicSpringConstraints.cpp b/Engine/Source/Runtime/Experimental/Chaos/Private/Chaos/PBDRigidDynamicSpringConstraints.cpp
--- a/Engine/Source/Runtime/Experimental/Chaos/Private/Chaos/PBDRigidDynamicSpringConstraints.cpp
+++ b/Engine/Source/Runtime/Experimental/Chaos/Private/Chaos/PBDRigidDynamicSpringConstraints.cpp
@@ -1,6 +1,7 @@
 // Copyright Epic Games, Inc. All Rights Reserved.
 
 #include "Chaos/PBDRigidDynamicSpringConstraints.h"
+#include "Chaos/Particle/ParticleUtilities.h"
 #include "Chaos/Utilities.h"
 
 using namespace Chaos;
@@ -121,46 +122,44 @@ TVector<T, d> TPBDRigidDynamicSpringConstraints<T, d>::GetDelta(const TVector<T,
 template<class T, int d>
 void TPBDRigidDynamicSpringConstraints<T, d>::ApplySingle(const T Dt, int32 ConstraintIndex) const
 {
-	TGeometryParticleHandle<T, d>* Static0 = Constraints[ConstraintIndex][0];
-	TGeometryParticleHandle<T, d>* Static1 = Constraints[ConstraintIndex][1];
-	TPBDRigidParticleHandle<T, d>* PBDRigid0 = Static0->CastToRigidParticle();
-	TPBDRigidParticleHandle<T, d>* PBDRigid1 = Static1->CastToRigidParticle();
-	const bool bIsRigidDynamic0 = PBDRigid0 && PBDRigid0->ObjectState() == EObjectStateType::Dynamic;
-	const bool bIsRigidDynamic1 = PBDRigid1 && PBDRigid1->ObjectState() == EObjectStateType::Dynamic;
-
+	TGenericParticleHandle<T, d> Particle0 = Constraints[ConstraintIndex][0];
+	TGenericParticleHandle<T, d> Particle1 = Constraints[ConstraintIndex][1];
+	const bool bIsRigidDynamic0 = Particle0->IsDynamic();
+	const bool bIsRigidDynamic1 = Particle1->IsDynamic();
 	check(bIsRigidDynamic0 || bIsRigidDynamic1);
-	check(!bIsRigidDynamic0|| !bIsRigidDynamic1|| (PBDRigid0->Island() == PBDRigid1->Island()));
 
-	TRotation<T, d>& Q0 = bIsRigidDynamic0 ? PBDRigid0->Q() : Static0->R();
-	TRotation<T, d>& Q1 = bIsRigidDynamic1 ? PBDRigid1->Q() : Static1->R();
-	TVector<T, d>& P0 = bIsRigidDynamic0 ? PBDRigid0->P() : Static0->X();
-	TVector<T, d>& P1 = bIsRigidDynamic1 ? PBDRigid1->P() : Static1->X();
+	TRotation<T, d> Q0 = FParticleUtilities::GetCoMWorldRotation(Particle0);
+	TRotation<T, d> Q1 = FParticleUtilities::GetCoMWorldRotation(Particle1);
+	TVector<T, d> P0 = FParticleUtilities::GetCoMWorldPosition(Particle0);
+	TVector<T, d> P1 = FParticleUtilities::GetCoMWorldPosition(Particle1);
 
 	const int32 NumSprings = SpringDistances[ConstraintIndex].Num();
-	const PMatrix<T, d, d> WorldSpaceInvI1 = bIsRigidDynamic0 ? Utilities::ComputeWorldSpaceInertia(Q0, PBDRigid0->InvI()) : PMatrix<T, d, d>(0);
-	const PMatrix<T, d, d> WorldSpaceInvI2 = bIsRigidDynamic1 ? Utilities::ComputeWorldSpaceInertia(Q1, PBDRigid1->InvI()) : PMatrix<T, d, d>(0);;
+	const PMatrix<T, d, d> WorldSpaceInvI1 = bIsRigidDynamic0 ? Utilities::ComputeWorldSpaceInertia(Q0, Particle0->InvI()) : PMatrix<T, d, d>(0);
+	const PMatrix<T, d, d> WorldSpaceInvI2 = bIsRigidDynamic1 ? Utilities::ComputeWorldSpaceInertia(Q1, Particle1->InvI()) : PMatrix<T, d, d>(0);;
 	for (int32 SpringIndex = 0; SpringIndex < NumSprings; ++SpringIndex)
 	{
 		const TVector<T, d>& Distance0 = Distances[ConstraintIndex][SpringIndex][0];
 		const TVector<T, d>& Distance1 = Distances[ConstraintIndex][SpringIndex][1];
-		const TVector<T, d> WorldSpaceX1 = Q0.RotateVector(Distance0) + P0;
-		const TVector<T, d> WorldSpaceX2 = Q1.RotateVector(Distance1) + P1;
+		const TVector<T, d> WorldSpaceX1 = Particle0->Q().RotateVector(Distance0) + Particle0->P();
+		const TVector<T, d> WorldSpaceX2 = Particle1->Q().RotateVector(Distance1) + Particle1->P();
 		const TVector<T, d> Delta = GetDelta(WorldSpaceX1, WorldSpaceX2, ConstraintIndex, SpringIndex);
 
 		if (bIsRigidDynamic0)
 		{
 			const TVector<T, d> Radius = WorldSpaceX1 - P0;
-			P0 += PBDRigid0->InvM() * Delta;
+			P0 += Particle0->InvM() * Delta;
 			Q0 += TRotation<T, d>::FromElements(WorldSpaceInvI1 * TVector<T, d>::CrossProduct(Radius, Delta), 0.f) * Q0 * T(0.5);
 			Q0.Normalize();
+			FParticleUtilities::SetCoMWorldTransform(Particle0, P0, Q0);
 		}
 
 		if (bIsRigidDynamic1)
 		{
 			const TVector<T, d> Radius = WorldSpaceX2 - P1;
-			P1 -= PBDRigid1->InvM() * Delta;
+			P1 -= Particle1->InvM() * Delta;
 			Q1 += TRotation<T, d>::FromElements(WorldSpaceInvI2 * TVector<T, d>::CrossProduct(Radius, -Delta), 0.f) * Q1 * T(0.5);
 			Q1.Normalize();
+			FParticleUtilities::SetCoMWorldTransform(Particle1, P1, Q1);
 		}
 	}
 }
diff --git a/Engine/Source/Runtime/Experimental/Chaos/Private/Chaos/PBDRigidSpringConstraints.cpp b/Engine/Source/Runtime/Experimental/Chaos/Private/Chaos/PBDRigidSpringConstraints.cpp
--- a/Engine/Source/Runtime/Experimental/Chaos/Private/Chaos/PBDRigidSpringConstraints.cpp
+++ b/Engine/Source/Runtime/Experimental/Chaos/Private/Chaos/PBDRigidSpringConstraints.cpp
@@ -1,6 +1,7 @@
 // Copyright Epic Games, Inc. All Rights Reserved.
 
 #include "Chaos/PBDRigidSpringConstraints.h"
+#include "Chaos/Particle/ParticleUtilities.h"
 #include "Chaos/Utilities.h"
 
 namespace Chaos
@@ -146,39 +147,39 @@ namespace Chaos
 	{
 		const FConstrainedParticlePair& Constraint = Constraints[ConstraintIndex];
 
-		TPBDRigidParticleHandle<FReal, 3>* PBDRigid0 = Constraint[0]->CastToRigidParticle();
-		TPBDRigidParticleHandle<FReal, 3>* PBDRigid1 = Constraint[1]->CastToRigidParticle();
-		const bool bIsRigidDynamic0 = PBDRigid0 && PBDRigid0->ObjectState() == EObjectStateType::Dynamic;
-		const bool bIsRigidDynamic1 = PBDRigid1 && PBDRigid1->ObjectState() == EObjectStateType::Dynamic;
+		TGenericParticleHandle<FReal, 3> Particle0 = Constraints[ConstraintIndex][0];
+		TGenericParticleHandle<FReal, 3> Particle1 = Constraints[ConstraintIndex][1];
+		const bool bIsRigidDynamic0 = Particle0->IsDynamic();
+		const bool bIsRigidDynamic1 = Particle1->IsDynamic();
 
 		check((bIsRigidDynamic0 && bIsRigidDynamic1) || (!bIsRigidDynamic0 && bIsRigidDynamic1) || (bIsRigidDynamic0 && bIsRigidDynamic1));
 
-		// @todo(ccaulfield): see if we can eliminate the need for all these ifs
-		const TRotation<FReal, 3> & Q0 = bIsRigidDynamic0 ? PBDRigid0->Q() : Constraint[0]->R();
-		const TRotation<FReal, 3> & Q1 = bIsRigidDynamic1 ? PBDRigid1->Q() : Constraint[1]->R();
-		const FVec3 & P0 = bIsRigidDynamic0 ? PBDRigid0->P() : Constraint[0]->X();
-		const FVec3 & P1 = bIsRigidDynamic1 ? PBDRigid1->P() : Constraint[1]->X();
-
-		const FVec3 WorldSpaceX1 = Q0.RotateVector(Distances[ConstraintIndex][0]) + P0;
-		const FVec3 WorldSpaceX2 = Q1.RotateVector(Distances[ConstraintIndex][1]) + P1;
-		const FMatrix33 WorldSpaceInvI1 = bIsRigidDynamic0 ? Utilities::ComputeWorldSpaceInertia(Q0, PBDRigid0->InvI()) : FMatrix33(0);
-		const FMatrix33 WorldSpaceInvI2 = bIsRigidDynamic1 ? Utilities::ComputeWorldSpaceInertia(Q1, PBDRigid1->InvI()) : FMatrix33(0);
+		const FVec3 WorldSpaceX1 = Particle0->Q().RotateVector(Distances[ConstraintIndex][0]) + Particle0->P();
+		const FVec3 WorldSpaceX2 = Particle1->Q().RotateVector(Distances[ConstraintIndex][1]) + Particle1->P();
 		const FVec3 Delta = GetDelta(ConstraintIndex, WorldSpaceX1, WorldSpaceX2);
 
 		if (bIsRigidDynamic0)
 		{
-			const FVec3 Radius = WorldSpaceX1 - PBDRigid0->P();
-			PBDRigid0->P() += PBDRigid0->InvM() * Delta;
-			PBDRigid0->Q() += TRotation<FReal, 3>::FromElements(WorldSpaceInvI1 * FVec3::CrossProduct(Radius, Delta), 0.f) * PBDRigid0->Q() * FReal(0.5);
-			PBDRigid0->Q().Normalize();
+			FRotation3 Q0 = FParticleUtilities::GetCoMWorldRotation(Particle0);
+			FVec3 P0 = FParticleUtilities::GetCoMWorldPosition(Particle0);
+			const FMatrix33 WorldSpaceInvI1 = Utilities::ComputeWorldSpaceInertia(Q0, Particle0->InvI());
+			const FVec3 Radius = WorldSpaceX1 - P0;
+			P0 += Particle0->InvM() * Delta;
+			Q0 += TRotation<FReal, 3>::FromElements(WorldSpaceInvI1 * FVec3::CrossProduct(Radius, Delta), 0.f) * Q0 * FReal(0.5);
+			Q0.Normalize();
+			FParticleUtilities::SetCoMWorldTransform(Particle0, P0, Q0);
 		}
 
 		if (bIsRigidDynamic1)
 		{
-			const FVec3 Radius = WorldSpaceX2 - PBDRigid1->P();
-			PBDRigid1->P() -= PBDRigid1->InvM() * Delta;
-			PBDRigid1->Q() += TRotation<FReal, 3>::FromElements(WorldSpaceInvI2 * FVec3::CrossProduct(Radius, -Delta), 0.f) * PBDRigid1->Q() * FReal(0.5);
-			PBDRigid1->Q().Normalize();
+			FRotation3 Q1 = FParticleUtilities::GetCoMWorldRotation(Particle1);
+			FVec3 P1 = FParticleUtilities::GetCoMWorldPosition(Particle1);
+			const FMatrix33 WorldSpaceInvI2 = Utilities::ComputeWorldSpaceInertia(Q1, Particle1->InvI());
+			const FVec3 Radius = WorldSpaceX2 - P1;
+			P1 -= Particle1->InvM() * Delta;
+			Q1 += TRotation<FReal, 3>::FromElements(WorldSpaceInvI2 * FVec3::CrossProduct(Radius, -Delta), 0.f) * Q1 * FReal(0.5);
+			Q1.Normalize();
+			FParticleUtilities::SetCoMWorldTransform(Particle1, P1, Q1);
 		}
 	}
 }
diff --git a/Engine/Source/Runtime/Experimental/Chaos/Public/Chaos/Particle/ParticleUtilities.h b/Engine/Source/Runtime/Experimental/Chaos/Public/Chaos/Particle/ParticleUtilities.h
--- a/Engine/Source/Runtime/Experimental/Chaos/Public/Chaos/Particle/ParticleUtilities.h
+++ b/Engine/Source/Runtime/Experimental/Chaos/Public/Chaos/Particle/ParticleUtilities.h
@@ -2,6 +2,7 @@
 #pragma once
 
 #include "Chaos/ParticleHandle.h"
+#include "Chaos/Utilities.h"
 
 #define CHAOS_PARTICLE_ACTORTRANSFORM 1
 
@@ -72,6 +73,18 @@ namespace Chaos
 			return ActorWorldT;
 		}
 
+		template<typename T_PARTICLEHANDLE>
+		static inline FMatrix33 GetWorldInertia(T_PARTICLEHANDLE Particle)
+		{
+			return Utilities::ComputeWorldSpaceInertia(TSpatialAccessor::GetRotation(Particle) * Particle->RotationOfMass(), Particle->I());
+		}
+
+		template<typename T_PARTICLEHANDLE>
+		static inline FMatrix33 GetWorldInvInertia(T_PARTICLEHANDLE Particle)
+		{
+			return Utilities::ComputeWorldSpaceInertia(TSpatialAccessor::GetRotation(Particle) * Particle->RotationOfMass(), Particle->InvI());
+		}
+
 		/**
 		 * Convert an particle position into a com-local-space position
 		 */
@@ -198,6 +211,18 @@ namespace Chaos
 			return FRigidTransform3(CoMP, CoMQ);
 		}
 
+		template<typename T_PARTICLEHANDLE>
+		static inline FMatrix33 GetWorldInertia(T_PARTICLEHANDLE Particle)
+		{
+			return Utilities::ComputeWorldSpaceInertia(TSpatialAccessor::Rotation(Particle), Particle->I());
+		}
+
+		template<typename T_PARTICLEHANDLE>
+		static inline FMatrix33 GetWorldInvInertia(T_PARTICLEHANDLE Particle)
+		{
+			return Utilities::ComputeWorldSpaceInertia(TSpatialAccessor::Rotation(Particle), Particle->InvI());
+		}
+
 		template<typename T_PARTICLEHANDLE>
 		static inline const FVec3& ParticleLocalToCoMLocal(T_PARTICLEHANDLE Particle, const FVec3& P)
 		{
diff --git a/Engine/Source/Runtime/Experimental/Chaos/Public/Chaos/PerParticleAddImpulses.h b/Engine/Source/Runtime/Experimental/Chaos/Public/Chaos/PerParticleAddImpulses.h
--- a/Engine/Source/Runtime/Experimental/Chaos/Public/Chaos/PerParticleAddImpulses.h
+++ b/Engine/Source/Runtime/Experimental/Chaos/Public/Chaos/PerParticleAddImpulses.h
@@ -2,6 +2,7 @@
 #pragma once
 
 #include "Chaos/PerParticleRule.h"
+#include "Chaos/Particle/ParticleUtilities.h"
 #include "Chaos/Utilities.h"
 
 namespace Chaos
@@ -38,7 +39,12 @@ namespace Chaos
 			//       Just using W += InvI * (Torque - W x (I * W)) * dt is not correct, since Torque
 			//		 and W are in an inertial frame.
 			//
+#if CHAOS_PARTICLE_ACTORTRANSFORM
+			const FMatrix33 WorldInvI = Utilities::ComputeWorldSpaceInertia(InParticles.R(Index) * InParticles.RotationOfMass(Index), InParticles.InvI(Index));
+#else
 			const FMatrix33 WorldInvI = Utilities::ComputeWorldSpaceInertia(InParticles.R(Index), InParticles.InvI(Index));
+#endif
+
 			InParticles.W(Index) += WorldInvI * InParticles.AngularImpulse(Index);
 			InParticles.LinearImpulse(Index) = TVector<T, 3>(0);
 			InParticles.AngularImpulse(Index) = TVector<T, 3>(0);
@@ -47,7 +53,11 @@ namespace Chaos
 		inline void Apply(TTransientPBDRigidParticleHandle<T, d>& Particle, const T Dt) const override //-V762
 		{
 			Particle.V() += Particle.LinearImpulse() * Particle.InvM();
+#if CHAOS_PARTICLE_ACTORTRANSFORM
+			const FMatrix33 WorldInvI = Utilities::ComputeWorldSpaceInertia(Particle.R() * Particle.RotationOfMass(), Particle.InvI());
+#else
 			const FMatrix33 WorldInvI = Utilities::ComputeWorldSpaceInertia(Particle.R(), Particle.InvI());
+#endif
 			Particle.W() += WorldInvI * Particle.AngularImpulse();
 			Particle.LinearImpulse() = TVector<T, 3>(0);
 			Particle.AngularImpulse() = TVector<T, 3>(0);
diff --git a/Engine/Source/Runtime/Experimental/Chaos/Public/Chaos/PerParticleEulerStepVelocity.h b/Engine/Source/Runtime/Experimental/Chaos/Public/Chaos/PerParticleEulerStepVelocity.h
--- a/Engine/Source/Runtime/Experimental/Chaos/Public/Chaos/PerParticleEulerStepVelocity.h
+++ b/Engine/Source/Runtime/Experimental/Chaos/Public/Chaos/PerParticleEulerStepVelocity.h
@@ -2,6 +2,7 @@
 #pragma once
 
 #include "Chaos/PerParticleRule.h"
+#include "Chaos/Particle/ParticleUtilities.h"
 #include "Chaos/Utilities.h"
 
 namespace Chaos
@@ -39,14 +40,22 @@ class TPerParticleEulerStepVelocity : public TPerParticleRule<T, d>
 		//       Just using W += InvI * (Torque - W x (I * W)) * dt is not correct, since Torque
 		//		 and W are in an inertial frame.
 		//
+#if CHAOS_PARTICLE_ACTORTRANSFORM
 		const FMatrix33 WorldInvI = Utilities::ComputeWorldSpaceInertia(InParticles.R(Index) * InParticles.RotationOfMass(Index), InParticles.InvI(Index));
+#else
+		const FMatrix33 WorldInvI = Utilities::ComputeWorldSpaceInertia(InParticles.R(Index), InParticles.InvI(Index));
+#endif
 		InParticles.W(Index) += WorldInvI * InParticles.Torque(Index) * Dt;
 	}
 	
 	inline void Apply(TTransientPBDRigidParticleHandle<T, d>& Particle, const T Dt) const override //-V762
 	{
 		Particle.V() += Particle.F() * Particle.InvM() * Dt;
+#if CHAOS_PARTICLE_ACTORTRANSFORM
 		const FMatrix33 WorldInvI = Utilities::ComputeWorldSpaceInertia(Particle.R() * Particle.RotationOfMass(), Particle.InvI());
+#else
+		const FMatrix33 WorldInvI = Utilities::ComputeWorldSpaceInertia(Particle.R(Index), Particle.InvI());
+#endif
 		Particle.W() += WorldInvI * Particle.Torque() * Dt;
 	}
 };
diff --git a/Engine/Source/Runtime/Experimental/Chaos/Public/Chaos/Utilities.h b/Engine/Source/Runtime/Experimental/Chaos/Public/Chaos/Utilities.h
--- a/Engine/Source/Runtime/Experimental/Chaos/Public/Chaos/Utilities.h
+++ b/Engine/Source/Runtime/Experimental/Chaos/Public/Chaos/Utilities.h
@@ -230,18 +230,18 @@ namespace Chaos
 		}
 
 		/**
-		 * Calculate the world-space inertia (or inverse inertia) for a body with rotation "Q" and local-space inertia/inverse-inertia "I".
+		 * Calculate the world-space inertia (or inverse inertia) for a body with center-of-mass rotation "CoMRotation" and local-space inertia/inverse-inertia "I".
 		 */
-		static FMatrix33 ComputeWorldSpaceInertia(const FRotation3& Q, const FMatrix33& I)
+		static FMatrix33 ComputeWorldSpaceInertia(const FRotation3& CoMRotation, const FMatrix33& I)
 		{
-			FMatrix33 QM = Q.ToMatrix();
+			FMatrix33 QM = CoMRotation.ToMatrix();
 			return MultiplyAB(QM, MultiplyABt(I, QM));
 		}
 
-		static FMatrix33 ComputeWorldSpaceInertia(const FRotation3& Q, const FVec3& I)
+		static FMatrix33 ComputeWorldSpaceInertia(const FRotation3& CoMRotation, const FVec3& I)
 		{
 			// @todo(ccaulfield): optimize ComputeWorldSpaceInertia
-			return ComputeWorldSpaceInertia(Q, FMatrix33(I.X, I.Y, I.Z));
+			return ComputeWorldSpaceInertia(CoMRotation, FMatrix33(I.X, I.Y, I.Z));
 		}
 
 		/**
