Fix ClosestBInA not transformed correctly in GJKRaycast2

diff --git a/Engine/Source/Runtime/Experimental/Chaos/Public/Chaos/GJK.h b/Engine/Source/Runtime/Experimental/Chaos/Public/Chaos/GJK.h
--- a/Engine/Source/Runtime/Experimental/Chaos/Public/Chaos/GJK.h
+++ b/Engine/Source/Runtime/Experimental/Chaos/Public/Chaos/GJK.h
@@ -586,7 +586,7 @@ namespace Chaos
 					ClosestB = Bs[0];
 				}
 				
-				const TVec3<T> ClosestBInA = StartTM.TransformPosition(ClosestB);
+				const TVec3<T> ClosestBInA = StartPoint + ClosestB;
 				const T InGJKPreDist = FMath::Sqrt(InGJKPreDist2);
 				OutNormal = (ClosestBInA - ClosestA).GetUnsafeNormal();	//question: should we just use InGJKPreDist2?
 				const T Penetration = FMath::Clamp<T>(ThicknessA + ThicknessB - InGJKPreDist, 0, TNumericLimits<T>::Max());
