Chaos Acceleration structures bug fix when time slicing: Make sure that we don't copy an incomplete acceleration struction to the GT when time slicing

diff --git a/Engine/Source/Runtime/Experimental/Chaos/Private/Chaos/PBDRigidsEvolution.cpp b/Engine/Source/Runtime/Experimental/Chaos/Private/Chaos/PBDRigidsEvolution.cpp
--- a/Engine/Source/Runtime/Experimental/Chaos/Private/Chaos/PBDRigidsEvolution.cpp
+++ b/Engine/Source/Runtime/Experimental/Chaos/Private/Chaos/PBDRigidsEvolution.cpp
@@ -243,7 +243,8 @@ namespace Chaos
 				AccelerationStructure->GetSubstructure(SpatialIdx)->ProgressAsyncTimeSlicing(IsForceFullBuild);
 
 				// is it still progressing or now complete
-				IsTimeSlicingProgressing = !AccelerationStructure->GetSubstructure(SpatialIdx)->IsAsyncTimeSlicingComplete();
+				if (!AccelerationStructure->GetSubstructure(SpatialIdx)->IsAsyncTimeSlicingComplete())
+					IsTimeSlicingProgressing = true;
 			}
 			else
 			{
