Remove cvar from speculative fix

diff --git a/Engine/Source/Runtime/Experimental/Chaos/Private/Chaos/BoundingVolume.cpp b/Engine/Source/Runtime/Experimental/Chaos/Private/Chaos/BoundingVolume.cpp
--- a/Engine/Source/Runtime/Experimental/Chaos/Private/Chaos/BoundingVolume.cpp
+++ b/Engine/Source/Runtime/Experimental/Chaos/Private/Chaos/BoundingVolume.cpp
@@ -21,12 +21,6 @@ namespace Chaos
 	    MaxDirtyElements,
 	    TEXT("The max number of dirty elements. This forces a flush which is very expensive"));
 
-	CHAOS_API int32 FixEPAWhenSimplexOutside = 1;
-	FAutoConsoleVariableRef CVarFixEPAWhenSimplexOutside(
-	    TEXT("p.FixEPAWhenSimplexOutside"),
-	    FixEPAWhenSimplexOutside,
-	    TEXT(""));
-
 	    template<typename TPayloadType, typename T, int d>
 	    ISpatialAcceleration<TPayloadType, T, d>* ISpatialAcceleration<TPayloadType, T, d>::SerializationFactory(FChaosArchive& Ar, ISpatialAcceleration<TPayloadType, T, d>* Accel)
 	{
diff --git a/Engine/Source/Runtime/Experimental/Chaos/Public/Chaos/EPA.h b/Engine/Source/Runtime/Experimental/Chaos/Public/Chaos/EPA.h
--- a/Engine/Source/Runtime/Experimental/Chaos/Public/Chaos/EPA.h
+++ b/Engine/Source/Runtime/Experimental/Chaos/Public/Chaos/EPA.h
@@ -310,8 +310,6 @@ void EPAComputeVisibilityBorder(TArray<TEPAEntry<T>>& Entries, int32 EntryIdx, c
 	}
 }
 
-extern CHAOS_API int32 FixEPAWhenSimplexOutside;
-
 template <typename T>
 void ComputeEPAResults(const TVec3<T>* VertsA, const TVec3<T>* VertsB, const TEPAEntry<T>& Entry, T& OutPenetration, TVec3<T>& OutDir, TVec3<T>& OutA, TVec3<T>& OutB)
 {
@@ -334,7 +332,7 @@ void ComputeEPAResults(const TVec3<T>* VertsA, const TVec3<T>* VertsB, const TEP
 		OutDir /= OutPenetration;
 	}
 
-	if (Entry.Distance < 0 && FixEPAWhenSimplexOutside)
+	if (Entry.Distance < 0)
 	{
 		//The origin is on the outside, so the direction is reversed
 		OutDir = -OutDir;
