Reorganize solver substepping, and hook it up to configuration variables.

diff --git a/Engine/Source/Runtime/Engine/Private/PhysicsEngine/Experimental/PhysScene_Chaos.cpp b/Engine/Source/Runtime/Engine/Private/PhysicsEngine/Experimental/PhysScene_Chaos.cpp
--- a/Engine/Source/Runtime/Engine/Private/PhysicsEngine/Experimental/PhysScene_Chaos.cpp
+++ b/Engine/Source/Runtime/Engine/Private/PhysicsEngine/Experimental/PhysScene_Chaos.cpp
@@ -1794,10 +1794,16 @@ void FPhysScene_ChaosInterface::ApplyWorldOffset(FVector InOffset)
 	check(InOffset.Size() == 0);
 }
 
-void FPhysScene_ChaosInterface::SetUpForFrame(const FVector* NewGrav, float InDeltaSeconds /*= 0.0f*/, float InMaxPhysicsDeltaTime /*= 0.0f*/)
+void FPhysScene_ChaosInterface::SetUpForFrame(const FVector* NewGrav, float InDeltaSeconds /*= 0.0f*/, float InMaxPhysicsDeltaTime /*= 0.0f*/, float InMaxSubstepDeltaTime /*= 0.0f*/, int32 InMaxSubsteps)
 {
 	SetGravity(*NewGrav);
-	MDeltaTime = InDeltaSeconds < InMaxPhysicsDeltaTime ? InDeltaSeconds : InMaxPhysicsDeltaTime;
+	MDeltaTime = InMaxPhysicsDeltaTime > 0.f ? FMath::Min(InDeltaSeconds, InMaxPhysicsDeltaTime) : InDeltaSeconds;
+
+	if (Chaos::FPhysicsSolver* Solver = GetSolver())
+	{
+		Solver->SetMaxDeltaTime(InMaxSubstepDeltaTime);
+		Solver->SetMaxSubSteps(InMaxSubsteps);
+	}
 }
 
 void FPhysScene_ChaosInterface::StartFrame()
diff --git a/Engine/Source/Runtime/Engine/Private/PhysicsEngine/PhysLevel.cpp b/Engine/Source/Runtime/Engine/Private/PhysicsEngine/PhysLevel.cpp
--- a/Engine/Source/Runtime/Engine/Private/PhysicsEngine/PhysLevel.cpp
+++ b/Engine/Source/Runtime/Engine/Private/PhysicsEngine/PhysLevel.cpp
@@ -157,7 +157,7 @@ void UWorld::SetupPhysicsTickFunctions(float DeltaSeconds)
 	FVector DefaultGravity( 0.f, 0.f, GetGravityZ() );
 
 	static const auto CVar_MaxPhysicsDeltaTime = IConsoleManager::Get().FindTConsoleVariableDataFloat(TEXT("p.MaxPhysicsDeltaTime"));
-	PhysScene->SetUpForFrame(&DefaultGravity, DeltaSeconds, UPhysicsSettings::Get()->MaxPhysicsDeltaTime);
+	PhysScene->SetUpForFrame(&DefaultGravity, DeltaSeconds, UPhysicsSettings::Get()->MaxPhysicsDeltaTime, UPhysicsSettings::Get()->MaxSubstepDeltaTime, UPhysicsSettings::Get()->MaxSubsteps);
 }
 
 void UWorld::StartPhysicsSim()
diff --git a/Engine/Source/Runtime/Engine/Public/Physics/Experimental/PhysScene_Chaos.h b/Engine/Source/Runtime/Engine/Public/Physics/Experimental/PhysScene_Chaos.h
--- a/Engine/Source/Runtime/Engine/Public/Physics/Experimental/PhysScene_Chaos.h
+++ b/Engine/Source/Runtime/Engine/Public/Physics/Experimental/PhysScene_Chaos.h
@@ -389,7 +389,7 @@ public:
 
 	ENGINE_API static bool SupportsOriginShifting();
 	void ApplyWorldOffset(FVector InOffset);
-	ENGINE_API void SetUpForFrame(const FVector* NewGrav, float InDeltaSeconds = 0.0f, float InMaxPhysicsDeltaTime = 0.0f);
+	ENGINE_API void SetUpForFrame(const FVector* NewGrav, float InDeltaSeconds = 0.0f, float InMaxPhysicsDeltaTime = 0.0f, float InMaxSubstepDeltaTime = 0.f, int32 InMaxSubsteps = 1);
 	ENGINE_API void StartFrame();
 	ENGINE_API void EndFrame(ULineBatchComponent* InLineBatcher);
 	ENGINE_API void WaitPhysScenes();
diff --git a/Engine/Source/Runtime/Experimental/ChaosSolvers/Private/PBDRigidsSolver.cpp b/Engine/Source/Runtime/Experimental/ChaosSolvers/Private/PBDRigidsSolver.cpp
--- a/Engine/Source/Runtime/Experimental/ChaosSolvers/Private/PBDRigidsSolver.cpp
+++ b/Engine/Source/Runtime/Experimental/ChaosSolvers/Private/PBDRigidsSolver.cpp
@@ -20,7 +20,7 @@
 #include "EventsData.h"
 
 
-DEFINE_LOG_CATEGORY_STATIC(LogPBDRigidsSolverSolver, Log, All);
+DEFINE_LOG_CATEGORY_STATIC(LogPBDRigidsSolver, Log, All);
 
 // DebugDraw CVars
 #if CHAOS_DEBUG_DRAW
@@ -60,13 +60,13 @@ namespace Chaos
 			: MSolver(Scene)
 			, MDeltaTime(DeltaTime)
 		{
-			UE_LOG(LogPBDRigidsSolverSolver, Verbose, TEXT("AdvanceOneTimeStepTask::AdvanceOneTimeStepTask()"));
+			UE_LOG(LogPBDRigidsSolver, Verbose, TEXT("AdvanceOneTimeStepTask::AdvanceOneTimeStepTask()"));
 		}
 
 		void DoWork()
 		{
 			LLM_SCOPE(ELLMTag::Chaos);
-			UE_LOG(LogPBDRigidsSolverSolver, Verbose, TEXT("AdvanceOneTimeStepTask::DoWork()"));
+			UE_LOG(LogPBDRigidsSolver, Verbose, TEXT("AdvanceOneTimeStepTask::DoWork()"));
 
 			{
 				SCOPE_CYCLE_COUNTER(STAT_UpdateParams);
@@ -89,30 +89,48 @@ namespace Chaos
 
 			{
 				SCOPE_CYCLE_COUNTER(STAT_EvolutionAndKinematicUpdate);
-				while (MDeltaTime > MSolver->GetMaxDeltaTime())
+
+				// This outer loop can potentially cause the system to lose energy over integration
+				// in a couple of different cases.
+				//
+				// * If we have a timestep that's smaller than MinDeltaTime, then we just won't step.
+				//   Yes, we'll lose some teeny amount of energy, but we'll avoid 1/dt issues.
+				//
+				// * If we have used all of our substeps but still have time remaining, then some
+				//   energy will be lost.
+				const float MinDeltaTime = MSolver->GetMinDeltaTime();
+				const float MaxDeltaTime = MSolver->GetMaxDeltaTime();
+				int32 StepsRemaining = MSolver->GetMaxSubSteps();
+				float TimeRemaining = MDeltaTime;
+				while (StepsRemaining > 0 && TimeRemaining > MinDeltaTime)
 				{
+					--StepsRemaining;
+					const float DeltaTime = MaxDeltaTime > 0.f ? FMath::Min(TimeRemaining, MaxDeltaTime) : TimeRemaining;
+					TimeRemaining -= DeltaTime;
+
 					Chaos::TArrayCollectionArray<FVector> Forces, Torques;
-					//MSolver->ForceUpdateCallback(MSolver->GetSolverTime());
 					for (FFieldSystemPhysicsProxy* Obj : MSolver->GetFieldSystemPhysicsProxies())
 					{
 						auto& GeomCollectionParticles = MSolver->GetEvolution()->GetParticles().GetGeometryCollectionParticles();
 						Obj->FieldForcesUpdateCallback(MSolver, GeomCollectionParticles, Forces, Torques, MSolver->GetSolverTime());
 					}
-					//MSolver->GetEvolution()->ReconcileIslands();
-					//MSolver->KinematicUpdateCallback(MSolver->GetMaxDeltaTime(), MSolver->GetSolverTime());
-					MSolver->GetEvolution()->AdvanceOneTimeStep(MSolver->GetMaxDeltaTime());
-					MDeltaTime -= MSolver->GetMaxDeltaTime();
+					MSolver->GetEvolution()->AdvanceOneTimeStep(DeltaTime);
 				}
-				//MSolver->ForceUpdateCallback(MSolver->GetSolverTime());
-				Chaos::TArrayCollectionArray<FVector> Forces, Torques;
-				for (FFieldSystemPhysicsProxy* Obj : MSolver->GetFieldSystemPhysicsProxies())
+
+#if CHAOS_CHECKED
+				// If time remains, then log why we have lost energy over the timestep.
+				if (TimeRemaining > 0.f)
 				{
-					auto& GeomCollectionParticles = MSolver->GetEvolution()->GetParticles().GetGeometryCollectionParticles();
-					Obj->FieldForcesUpdateCallback(MSolver, GeomCollectionParticles, Forces, Torques, MSolver->GetSolverTime());
+					if (StepsRemaining == 0)
+					{
+						UE_LOG(LogPBDRigidsSolver, Warning, TEXT("AdvanceOneTimeStepTask::DoWork() - Energy lost over %fs due to too many substeps over large timestep"), TimeRemaining);
+					}
+					else
+					{
+						UE_LOG(LogPBDRigidsSolver, Warning, TEXT("AdvanceOneTimeStepTask::DoWork() - Energy lost over %fs due to small timestep remainder"), TimeRemaining);
+					}
 				}
-				//MSolver->GetEvolution()->ReconcileIslands();
-				//MSolver->KinematicUpdateCallback(MDeltaTime, MSolver->GetSolverTime());
-				MSolver->GetEvolution()->AdvanceOneTimeStep(MDeltaTime);
+#endif
 			}
 
 			{
@@ -159,7 +177,8 @@ namespace Chaos
 		, MTime(0.0)
 		, MLastDt(0.0)
 		, MMaxDeltaTime(0.0)
-		, TimeStepMultiplier(1.0)
+		, MMinDeltaTime(1.e-10f)
+		, MMaxSubSteps(1)
 		, bEnabled(false)
 		, bHasFloor(true)
 		, bIsFloorAnalytic(false)
@@ -170,7 +189,7 @@ namespace Chaos
 		, MActiveParticlesBuffer(new FActiveParticlesBuffer(BufferingModeIn))
 		, MCurrentLock(new FCriticalSection())
 	{
-		UE_LOG(LogPBDRigidsSolverSolver, Verbose, TEXT("PBDRigidsSolver::PBDRigidsSolver()"));
+		UE_LOG(LogPBDRigidsSolver, Verbose, TEXT("PBDRigidsSolver::PBDRigidsSolver()"));
 		Reset();
 	}
 
@@ -198,7 +217,7 @@ namespace Chaos
 	{
 		LLM_SCOPE(ELLMTag::Chaos);
 
-		UE_LOG(LogPBDRigidsSolverSolver, Verbose, TEXT("FPBDRigidsSolver::RegisterObject()"));
+		UE_LOG(LogPBDRigidsSolver, Verbose, TEXT("FPBDRigidsSolver::RegisterObject()"));
 
 		// Make sure this particle doesn't already have a proxy
 		checkSlow(GTParticle->Proxy == nullptr);
@@ -252,7 +271,7 @@ namespace Chaos
 		//enqueue onto physics thread for finalizing registration
 		FChaosSolversModule::GetModule()->GetDispatcher()->EnqueueCommandImmediate(this, [GTParticle, InParticleType, ProxyBase, ProxyData](FPBDRigidsSolver* Solver)
 		{
-			UE_LOG(LogPBDRigidsSolverSolver, Verbose, TEXT("FPBDRigidsSolver::RegisterObject() ~ Dequeue"));
+			UE_LOG(LogPBDRigidsSolver, Verbose, TEXT("FPBDRigidsSolver::RegisterObject() ~ Dequeue"));
 
 			// Create a handle for the new particle
 			TGeometryParticleHandle<float, 3>* Handle;
@@ -298,7 +317,7 @@ namespace Chaos
 
 	void FPBDRigidsSolver::UnregisterObject(TGeometryParticle<float, 3>* GTParticle)
 	{
-		UE_LOG(LogPBDRigidsSolverSolver, Verbose, TEXT("FPBDRigidsSolver::UnregisterObject()"));
+		UE_LOG(LogPBDRigidsSolver, Verbose, TEXT("FPBDRigidsSolver::UnregisterObject()"));
 
 		// Get the proxy associated with this particle
 		IPhysicsProxyBase* InProxy = GTParticle->Proxy;
@@ -334,7 +353,7 @@ namespace Chaos
 		// Enqueue a command to remove the particle and delete the proxy
 		FChaosSolversModule::GetModule()->GetDispatcher()->EnqueueCommandImmediate(this, [InProxy, InParticleType](FPBDRigidsSolver* Solver)
 		{
-			UE_LOG(LogPBDRigidsSolverSolver, Verbose, TEXT("FPBDRigidsSolver::UnregisterObject() ~ Dequeue"));
+			UE_LOG(LogPBDRigidsSolver, Verbose, TEXT("FPBDRigidsSolver::UnregisterObject() ~ Dequeue"));
 
 				// Generally need to remove stale events for particles that no longer exist
 				Solver->GetEventManager()->ClearEvents<FCollisionEventData>([InProxy]
@@ -401,7 +420,7 @@ namespace Chaos
 
 	void FPBDRigidsSolver::RegisterObject(FGeometryCollectionPhysicsProxy* InProxy)
 	{
-		UE_LOG(LogPBDRigidsSolverSolver, Verbose, TEXT("FPBDRigidsSolver::RegisterObject(FGeometryCollectionPhysicsProxy*)"));
+		UE_LOG(LogPBDRigidsSolver, Verbose, TEXT("FPBDRigidsSolver::RegisterObject(FGeometryCollectionPhysicsProxy*)"));
 		GeometryCollectionPhysicsProxies.AddUnique(InProxy);
 		InProxy->SetSolver(this);
 		InProxy->Initialize();
@@ -412,7 +431,7 @@ namespace Chaos
 		FChaosSolversModule::GetModule()->GetDispatcher()->EnqueueCommandImmediate(this,
 			[this, InProxy](FPBDRigidsSolver* Solver)
 		{
-			UE_LOG(LogPBDRigidsSolverSolver, Verbose, 
+			UE_LOG(LogPBDRigidsSolver, Verbose, 
 				TEXT("FPBDRigidsSolver::RegisterObject(FGeometryCollectionPhysicsProxy*)"));
 			InProxy->InitializeBodiesPT(Solver, Particles);
 			InProxy->PushToPhysicsState(nullptr); // Uses buffer on proxy for transport, not ProxyData.
@@ -427,7 +446,7 @@ namespace Chaos
 
 	void FPBDRigidsSolver::RegisterObject(FFieldSystemPhysicsProxy* InProxy)
 	{
-		UE_LOG(LogPBDRigidsSolverSolver, Verbose, TEXT("FPBDRigidsSolver::RegisterObject(FFieldSystemPhysicsProxy*)"));
+		UE_LOG(LogPBDRigidsSolver, Verbose, TEXT("FPBDRigidsSolver::RegisterObject(FFieldSystemPhysicsProxy*)"));
 		FieldSystemPhysicsProxies.AddUnique(InProxy);
 		InProxy->SetSolver(this);
 		InProxy->Initialize();
@@ -437,7 +456,7 @@ namespace Chaos
 			this,
 			[this, InProxy, ProxyData](FPBDRigidsSolver* Solver)
 			{
-				UE_LOG(LogPBDRigidsSolverSolver, Verbose,
+				UE_LOG(LogPBDRigidsSolver, Verbose,
 					TEXT("FPBDRigidsSolver::RegisterObject(FFieldSystemPhysicsProxy*)"));
 				//InProxy->ActivateBodies();
 				InProxy->PushToPhysicsState(ProxyData);
@@ -478,14 +497,15 @@ namespace Chaos
 
 	void FPBDRigidsSolver::Reset()
 	{
-		UE_LOG(LogPBDRigidsSolverSolver, Verbose, TEXT("PBDRigidsSolver::Reset()"));
+		UE_LOG(LogPBDRigidsSolver, Verbose, TEXT("PBDRigidsSolver::Reset()"));
 
 		MTime = 0;
 		MLastDt = 0.0f;
 		bEnabled = false;
 		CurrentFrame = 0;
-		MMaxDeltaTime = 1;
-		TimeStepMultiplier = 1;
+		MMaxDeltaTime = 1.f;
+		MMinDeltaTime = 1.e-10f;
+		MMaxSubSteps = 1;
 		MEvolution = TUniquePtr<FPBDRigidsEvolution>(new FPBDRigidsEvolution(Particles, ChaosSolverCollisionDefaultIterationsCVar, ChaosSolverCollisionDefaultPushoutIterationsCVar, BufferMode == EMultiBufferMode::Single)); 
 
 		FEventDefaults::RegisterSystemEvents(*GetEventManager());
@@ -499,17 +519,11 @@ namespace Chaos
 
 	void FPBDRigidsSolver::AdvanceSolverBy(float DeltaTime)
 	{
-		UE_LOG(LogPBDRigidsSolverSolver, Verbose, TEXT("PBDRigidsSolver::Tick(%3.5f)"), DeltaTime);
+		UE_LOG(LogPBDRigidsSolver, Verbose, TEXT("PBDRigidsSolver::Tick(%3.5f)"), DeltaTime);
 		if (bEnabled)
 		{
 			MLastDt = DeltaTime;
-
-			int32 NumTimeSteps = (int32)(1.f*TimeStepMultiplier);
-			float dt = FMath::Min(DeltaTime, float(5.f / 30.f)) / (float)NumTimeSteps;
-			for (int i = 0; i < NumTimeSteps; i++)
-			{
-				AdvanceOneTimeStepTask(this, DeltaTime).DoWork();
-			}
+			AdvanceOneTimeStepTask(this, DeltaTime).DoWork();
 		}
 
 	}
diff --git a/Engine/Source/Runtime/Experimental/ChaosSolvers/Public/PBDRigidsSolver.h b/Engine/Source/Runtime/Experimental/ChaosSolvers/Public/PBDRigidsSolver.h
--- a/Engine/Source/Runtime/Experimental/ChaosSolvers/Public/PBDRigidsSolver.h
+++ b/Engine/Source/Runtime/Experimental/ChaosSolvers/Public/PBDRigidsSolver.h
@@ -249,8 +249,12 @@ namespace Chaos
 		const float GetSolverTime() const { return MTime; }
 
 		/**/
-		const float GetLastDt() const { return MLastDt; }
-		const float GetMaxDeltaTime() const { return MMaxDeltaTime; }
+		void SetMaxDeltaTime(const float InMaxDeltaTime) { MMaxDeltaTime = InMaxDeltaTime; }
+		float GetLastDt() const { return MLastDt; }
+		float GetMaxDeltaTime() const { return MMaxDeltaTime; }
+		float GetMinDeltaTime() const { return MMinDeltaTime; }
+		void SetMaxSubSteps(const int32 InMaxSubSteps) { MMaxSubSteps = InMaxSubSteps; }
+		int32 GetMaxSubSteps() const { return MMaxSubSteps; }
 
 		/**/
 		void SetGenerateCollisionData(bool bDoGenerate) { GetEventFilters()->SetGenerateCollisionEvents(bDoGenerate); }
@@ -345,8 +349,8 @@ namespace Chaos
 		float MTime;
 		float MLastDt;
 		float MMaxDeltaTime;
-		float TimeStepMultiplier;
-
+		float MMinDeltaTime;
+		int32 MMaxSubSteps;
 		bool bEnabled;
 		bool bHasFloor;
 		bool bIsFloorAnalytic;
