Fix capsule raycasts when trace starts within cylinder, but isn't parallel.
The issue was that previously the ray would still cast against the cylinder rather than the caps even when its starting point was within the cylinder, unless it was parallel to the axis within some threshold.

diff --git a/Engine/Source/Runtime/Experimental/Chaos/Public/Chaos/Capsule.h b/Engine/Source/Runtime/Experimental/Chaos/Public/Chaos/Capsule.h
--- a/Engine/Source/Runtime/Experimental/Chaos/Public/Chaos/Capsule.h
+++ b/Engine/Source/Runtime/Experimental/Chaos/Public/Chaos/Capsule.h
@@ -173,10 +173,10 @@ namespace Chaos
 			constexpr T Epsilon = 1e-4;
 			bool bCheckCaps = false;
 
-			if (A < Epsilon)
+			if (C <= 0.f)
 			{
-				//Parallel and inside cylinder so check caps
-				bCheckCaps = C <= 0;
+				// Inside cylinder so check caps
+				bCheckCaps = true;
 			}
 			else
 			{
