Low level work for only copying data that actually changed

diff --git a/Engine/Source/Programs/HeadlessChaos/Private/HeadlessChaosTestDirtyParticle.cpp b/Engine/Source/Programs/HeadlessChaos/Private/HeadlessChaosTestDirtyParticle.cpp
new file mode 100644
--- /dev/null
+++ b/Engine/Source/Programs/HeadlessChaos/Private/HeadlessChaosTestDirtyParticle.cpp
@@ -0,0 +1,78 @@
+// Copyright Epic Games, Inc. All Rights Reserved.
+
+#include "HeadlessChaos.h"
+#include "HeadlessChaosTestUtility.h"
+
+#include "Chaos/PBDRigidsEvolutionGBF.h"
+#include "Chaos/Box.h"
+#include "Chaos/Sphere.h"
+#include "Chaos/Utilities.h"
+
+#include "Modules/ModuleManager.h"
+
+namespace ChaosTest
+{
+
+	using namespace Chaos;
+
+	GTEST_TEST(DirtyParticleTests,Basic)
+	{
+		FDirtyPropertiesManager DirtyPropertiesPool;
+		FDirtyIdx Idx = DirtyPropertiesPool.WriteX(FVec3(1,1,1));
+		FDirtyIdx Idx2 = DirtyPropertiesPool.WriteX(FVec3(2,1,1));
+		FDirtyIdx Idx3 = DirtyPropertiesPool.WriteInvM(3);
+
+
+		EXPECT_EQ(Idx.Property, EParticleProperty::X);
+		EXPECT_EQ(Idx.Entry, 0);
+		EXPECT_EQ(DirtyPropertiesPool.ReadX(Idx),FVec3(1,1,1));
+
+
+		EXPECT_EQ(Idx2.Property, EParticleProperty::X);
+		EXPECT_EQ(Idx2.Entry, 1);
+		EXPECT_EQ(DirtyPropertiesPool.ReadX(Idx2),FVec3(2,1,1));
+
+
+		EXPECT_EQ(Idx3.Property, EParticleProperty::InvM);
+		EXPECT_EQ(Idx3.Entry, 0);
+		EXPECT_EQ(DirtyPropertiesPool.ReadInvM(Idx3),3);
+
+
+		//release second x proprty
+		DirtyPropertiesPool.FreeProperty(Idx2);
+
+		//get another x property, should be same idx as Idx2 because we reuse indices
+		FDirtyIdx Idx4 = DirtyPropertiesPool.WriteX(FVec3(4,1,1));
+		EXPECT_EQ(Idx4,Idx2);
+
+		EXPECT_EQ(DirtyPropertiesPool.ReadX(Idx4),FVec3(4,1,1));
+
+		//using higher level api
+		//pretend we have a particle with V and Mass as dirty
+
+		//create a dirty properties instance to push to other thread
+		FDirtyProperties DirtyParticle;
+
+		FDirtyIdx VIdx0 = DirtyPropertiesPool.WriteV(FVec3(5,1,1));
+		DirtyParticle.DirtyProperty(VIdx0);
+		DirtyParticle.DirtyProperty(DirtyPropertiesPool.WriteM(6));
+
+		//Other thread can see the dirty properties and read them
+		EXPECT_TRUE(DirtyParticle.IsDirty(EParticleFlags::V));
+		EXPECT_TRUE(DirtyParticle.IsDirty(EParticleFlags::M));
+		//Other values are clean
+		EXPECT_FALSE(DirtyParticle.IsDirty(EParticleFlags::X));
+
+		//read the values
+		EXPECT_EQ(DirtyParticle.ReadV(DirtyPropertiesPool),FVec3(5,1,1));
+		EXPECT_EQ(DirtyParticle.ReadM(DirtyPropertiesPool),6);
+
+		DirtyParticle.Clean(DirtyPropertiesPool);
+
+		//verify that we really cleaned the resource
+		FDirtyIdx VIdx1 = DirtyPropertiesPool.WriteV(FVec3(5,1,1));
+		EXPECT_EQ(VIdx1,VIdx0);
+
+	}
+}
+
diff --git a/Engine/Source/Runtime/Experimental/Chaos/Public/Chaos/ParticleDirtyFlags.h b/Engine/Source/Runtime/Experimental/Chaos/Public/Chaos/ParticleDirtyFlags.h
--- a/Engine/Source/Runtime/Experimental/Chaos/Public/Chaos/ParticleDirtyFlags.h
+++ b/Engine/Source/Runtime/Experimental/Chaos/Public/Chaos/ParticleDirtyFlags.h
@@ -4,46 +4,29 @@
 
 namespace Chaos
 {
+
+using ChaosGeometrySp = TSharedPtr<FImplicitObject,ESPMode::ThreadSafe>;
+
+#define PARTICLE_PROPERTY(PropName, Type) PropName,
+	enum EParticleProperty : uint32
+	{
+#include "ParticleProperties.h"
+		NumProperties
+	};
+
+#undef PARTICLE_PROPERTY
+	
 	// There is a dirty flag for every user-settable particle property.
 	// Dirty property values will get copied from game to physics thread buffers,
 	// but clean property values will get overridden with physics thread results.
-	enum class EParticleFlags : int32
+#define PARTICLE_PROPERTY(PropName, Type) PropName = 1 << (uint32)EParticleProperty::PropName,
+
+	enum class EParticleFlags : uint32
 	{
-		X						= 1 << 0,
-		R						= 1 << 1,
-		V						= 1 << 2,
-		W						= 1 << 3,
-		CenterOfMass			= 1 << 4,
-		RotationOfMass			= 1 << 5,
-		CollisionGroup			= 1 << 6,
-		Disabled				= 1 << 7,
-		PreV					= 1 << 8,
-		PreW					= 1 << 9,
-		P						= 1 << 10,
-		Q						= 1 << 11,
-		F						= 1 << 12,
-		Torque					= 1 << 13,
-		I						= 1 << 14,
-		InvI					= 1 << 15,
-		M						= 1 << 16,
-		InvM					= 1 << 17,
-		LinearEtherDrag			= 1 << 18,
-		AngularEtherDrag		= 1 << 19,
-		ObjectState				= 1 << 20,
-		Geometry				= 1 << 21,
-		LinearImpulse			= 1 << 22,
-		AngularImpulse			= 1 << 23,
-		GravityEnabled			= 1 << 24,
-		SpatialIdx				= 1 << 25,
-		UniqueIdx				= 1 << 26,
-		ShapeDisableCollision	= 1 << 27,
-		CollisionTraceType		= 1 << 28,
-		ShapeSimData            = 1 << 29,
-		UserData				= 1 << 30
-#if CHAOS_CHECKED
-		, DebugName				= 1 << 31
-#endif
+		#include "ParticleProperties.h"
+		DummyFlag
 	};
+#undef PARTICLE_PROPERTY
 
 	class FParticleDirtyFlags
 	{
@@ -81,7 +64,156 @@ namespace Chaos
 			Bits = 0;
 		}
 
+		bool IsClean() const
+		{
+			return Bits == 0;
+		}
+
 	private:
 		int32 Bits;
 	};
+
+	struct FDirtyIdx
+	{
+		EParticleProperty Property : 5;
+		uint32 Entry : 27;
+
+		static_assert((uint32)EParticleProperty::NumProperties <= 31,"Cannot have more than 31 properties");
+
+		EParticleFlags GetPropertyFlag() const
+		{
+			return (EParticleFlags)(1 << (uint32)Property);
+		}
+
+		bool operator==(const FDirtyIdx& Rhs) const
+		{
+			return ((const uint32&)*this) == ((const uint32&)Rhs);
+		}
+	};
+
+	template <typename T>
+	class TDirtyElementPool
+	{
+	public:
+		const T& Read(int32 Idx) const
+		{
+			return Elements[Idx];
+		}
+
+		void Free(int32 Idx)
+		{
+			Elements[Idx].~T();
+			FreeIndices.Add(Idx);
+		}
+
+		int32 Write(const T& Element)
+		{
+			const int32 Idx = GetFree();
+			Elements[Idx] = Element;
+			return Idx;
+		}
+
+		int32 Write(T&& Element)
+		{
+			const int32 Idx = GetFree();
+			Elements[Idx] = MoveTemp(Element);
+			return Idx;
+		}
+
+	private:
+
+		int32 GetFree()
+		{
+			return FreeIndices.Num() ? FreeIndices.Pop(/*bAllowShrinking=*/false) : Elements.AddUninitialized(1);
+		}
+
+		TArray<T> Elements;
+		TArray<int32> FreeIndices;
+	};
+
+	class FDirtyPropertiesManager
+	{
+	public:
+#define PARTICLE_PROPERTY(PropName, Type) FDirtyIdx Write##PropName(Type const & Val){ FDirtyIdx Result; Result.Property = EParticleProperty::PropName; Result.Entry = Dirty##PropName##Pool.Write(Val); return Result; };
+#include "ParticleProperties.h"
+#undef PARTICLE_PROPERTY
+
+#define PARTICLE_PROPERTY(PropName, Type) Type const & Read##PropName(const FDirtyIdx Idx) const { ensure(Idx.Property == EParticleProperty::PropName); return Dirty##PropName##Pool.Read(Idx.Entry);};
+#include "ParticleProperties.h"
+#undef PARTICLE_PROPERTY
+
+	void FreeProperty(FDirtyIdx Idx)
+	{
+		switch((uint32)Idx.Property)
+		{
+#define PARTICLE_PROPERTY(PropName, Type) case EParticleProperty::PropName: { Dirty##PropName##Pool.Free(Idx.Entry); break; }
+#include "ParticleProperties.h"
+#undef PARTICLE_PROPERTY
+		default: check(false);
+
+		}
+	}
+	private:
+
+		//creates a TDirtyElementPool for each property
+#define PARTICLE_PROPERTY(Property, Type) TDirtyElementPool<Type> Dirty##Property##Pool;
+#include "ParticleProperties.h"
+#undef PARTICLE_PROPERTY
+	};
+
+	class FDirtyProperties
+	{
+	public:
+		bool IsDirty(const EParticleFlags CheckBits) const
+		{
+			return Flags.IsDirty(CheckBits);
+		}
+
+		bool IsDirty(const int32 CheckBits) const
+		{
+			return Flags.IsDirty(CheckBits);
+		}
+
+		void DirtyProperty(FDirtyIdx DirtyIdx)
+		{
+			Flags.MarkDirty(DirtyIdx.GetPropertyFlag());
+			DirtyProperties.Add(DirtyIdx);
+		}
+
+		~FDirtyProperties()
+		{
+			//Make sure to call Clean before killing this guy
+			ensure(Flags.IsClean());
+		}
+
+		void Clean(FDirtyPropertiesManager& DirtyPropertyManager)
+		{
+			for(const FDirtyIdx Idx : DirtyProperties)
+			{
+				DirtyPropertyManager.FreeProperty(Idx);
+			}
+
+			Flags.Clear();
+		}
+
+#define PARTICLE_PROPERTY(PropName, Type) Type const& Read##PropName(const FDirtyPropertiesManager& DirtyPropertyManager) const\
+{\
+	for(const FDirtyIdx Idx : DirtyProperties)\
+	{\
+		if(Idx.Property == EParticleProperty::PropName)\
+		{\
+			return DirtyPropertyManager.Read##PropName(Idx);\
+		}\
+	}\
+	check(false);\
+	static Type Error;\
+	return Error;\
+}
+#include "ParticleProperties.h"
+#undef PARTICLE_PROPERTY
+
+	private:
+		TArray<FDirtyIdx,TInlineAllocator<4>> DirtyProperties;
+		FParticleDirtyFlags Flags;
+	};
 }
diff --git a/Engine/Source/Runtime/Experimental/Chaos/Public/Chaos/ParticleProperties.h b/Engine/Source/Runtime/Experimental/Chaos/Public/Chaos/ParticleProperties.h
new file mode 100644
--- /dev/null
+++ b/Engine/Source/Runtime/Experimental/Chaos/Public/Chaos/ParticleProperties.h
@@ -0,0 +1,39 @@
+#undef PARTICLE_PROPERTY_CHECKED
+#if CHAOS_CHECKED
+#define PARTICLE_PROPERTY_CHECKED(x, Type) PARTICLE_PROPERTY(x, Type)
+#else
+#define PARTICLE_PROPERTY_CHECKED(x, Type)
+#endif
+
+PARTICLE_PROPERTY(X,FVec3)
+PARTICLE_PROPERTY(R,FRotation3)
+PARTICLE_PROPERTY(V,FVec3)
+PARTICLE_PROPERTY(W,FVec3)
+PARTICLE_PROPERTY(CenterOfMass,FVec3)
+PARTICLE_PROPERTY(RotationOfMass,FRotation3)
+PARTICLE_PROPERTY(CollisionGroup,int32)
+PARTICLE_PROPERTY(Disabled,bool)
+PARTICLE_PROPERTY(PreV,FVec3)
+PARTICLE_PROPERTY(PreW,FVec3)
+PARTICLE_PROPERTY(P,FVec3)
+PARTICLE_PROPERTY(Q,FRotation3)
+PARTICLE_PROPERTY(F,FVec3)
+PARTICLE_PROPERTY(Torque,FVec3)
+PARTICLE_PROPERTY(I,FMatrix33)
+PARTICLE_PROPERTY(InvI,FMatrix33)
+PARTICLE_PROPERTY(M,FReal)
+PARTICLE_PROPERTY(InvM,FReal)
+PARTICLE_PROPERTY(LinearEtherDrag,FReal)
+PARTICLE_PROPERTY(AngularEtherDrag,FReal)
+PARTICLE_PROPERTY(ObjectState,EObjectStateType)
+PARTICLE_PROPERTY(Geometry, ChaosGeometrySp)
+PARTICLE_PROPERTY(LinearImpulse, FVec3)
+PARTICLE_PROPERTY(AngularImpulse, FVec3)
+PARTICLE_PROPERTY(GravityEnabled, bool)
+PARTICLE_PROPERTY(SpatialIdx, FSpatialAccelerationIdx)
+PARTICLE_PROPERTY(UniqueIdx, FUniqueIdx)
+PARTICLE_PROPERTY(ShapeDisableCollision, int32)
+PARTICLE_PROPERTY(CollisionTraceType, EChaosCollisionTraceFlag)
+PARTICLE_PROPERTY(ShapeSimData, TArray<FCollisionFilterData>)
+PARTICLE_PROPERTY(UserData, void*)
+PARTICLE_PROPERTY_CHECKED(DebugName, FName)
\ No newline at end of file
