Fix CIS

diff --git a/Engine/Source/Runtime/Experimental/Chaos/Public/Chaos/ParticleProperties.h b/Engine/Source/Runtime/Experimental/Chaos/Public/Chaos/ParticleProperties.h
--- a/Engine/Source/Runtime/Experimental/Chaos/Public/Chaos/ParticleProperties.h
+++ b/Engine/Source/Runtime/Experimental/Chaos/Public/Chaos/ParticleProperties.h
@@ -1,3 +1,5 @@
+// Copyright Epic Games, Inc. All Rights Reserved.
+
 #undef PARTICLE_PROPERTY_CHECKED
 #if CHAOS_CHECKED
 #define PARTICLE_PROPERTY_CHECKED(x, Type) PARTICLE_PROPERTY(x, Type)
