Chaos - RBAN fix collision weirdness

diff --git a/Engine/Source/Runtime/Engine/Private/PhysicsEngine/ImmediatePhysics/ImmediatePhysicsChaos/ImmediatePhysicsSimulation_Chaos.cpp b/Engine/Source/Runtime/Engine/Private/PhysicsEngine/ImmediatePhysics/ImmediatePhysicsChaos/ImmediatePhysicsSimulation_Chaos.cpp
--- a/Engine/Source/Runtime/Engine/Private/PhysicsEngine/ImmediatePhysics/ImmediatePhysicsChaos/ImmediatePhysicsSimulation_Chaos.cpp
+++ b/Engine/Source/Runtime/Engine/Private/PhysicsEngine/ImmediatePhysics/ImmediatePhysicsChaos/ImmediatePhysicsSimulation_Chaos.cpp
@@ -189,8 +189,8 @@ namespace ImmediatePhysics_Chaos
 		Evolution.SetPostIntegrateCallback(
 			[this]()
 			{
-				DebugDrawDynamicParticles(3, 3, FColor(32, 32, 0));
-				DebugDrawConstraints(3, 3, 0.3f);
+				DebugDrawDynamicParticles(3, 4, FColor(32, 32, 0));
+				DebugDrawConstraints(3, 4, 0.3f);
 			});
 		Evolution.SetPostApplyCallback(
 			[this]()
@@ -637,7 +637,7 @@ namespace ImmediatePhysics_Chaos
 		Evolution.Advance(StepTime, NumSteps);
 
 		DebugDrawKinematicParticles(1, 4, FColor(128, 0, 0));
-		DebugDrawDynamicParticles(1, 4, FColor(255, 255, 0));
+		DebugDrawDynamicParticles(1, 3, FColor(255, 255, 0));
 		DebugDrawConstraints(1, 2, 1.0f);
 	}
 
diff --git a/Engine/Source/Runtime/Experimental/Chaos/Private/Chaos/PBDCollisionConstraintsContact.cpp b/Engine/Source/Runtime/Experimental/Chaos/Private/Chaos/PBDCollisionConstraintsContact.cpp
--- a/Engine/Source/Runtime/Experimental/Chaos/Private/Chaos/PBDCollisionConstraintsContact.cpp
+++ b/Engine/Source/Runtime/Experimental/Chaos/Private/Chaos/PBDCollisionConstraintsContact.cpp
@@ -305,7 +305,7 @@ namespace Chaos
 					PBDRigid0->V() += DP0 / IterationParameters.Dt;
 					PBDRigid0->W() += DR0 / IterationParameters.Dt;
 					P0 += DP0;
-					Q0 += FRotation3::FromElements(DR0, 0.f) * Q0 * IterationParameters.Dt * FReal(0.5);
+					Q0 += FRotation3::FromElements(DR0, 0.f) * Q0 * FReal(0.5);
 					Q0.Normalize();
 					FParticleUtilities::SetCoMWorldTransform(PBDRigid0, P0, Q0);
 				}
@@ -316,7 +316,7 @@ namespace Chaos
 					PBDRigid1->V() += DP1 / IterationParameters.Dt;
 					PBDRigid1->W() += DR1 / IterationParameters.Dt;
 					P1 += DP1;
-					Q1 += FRotation3::FromElements(DR1, 0.f) * Q1 * IterationParameters.Dt * FReal(0.5);
+					Q1 += FRotation3::FromElements(DR1, 0.f) * Q1 * FReal(0.5);
 					Q1.Normalize();
 					FParticleUtilities::SetCoMWorldTransform(PBDRigid1, P1, Q1);
 				}
