HeadlessChaos : Sleeping

- Fix for sleeping unit test failures

diff --git a/Engine/Source/Programs/HeadlessChaos/Private/HeadlessChaosTestGraph.cpp b/Engine/Source/Programs/HeadlessChaos/Private/HeadlessChaosTestGraph.cpp
--- a/Engine/Source/Programs/HeadlessChaos/Private/HeadlessChaosTestGraph.cpp
+++ b/Engine/Source/Programs/HeadlessChaos/Private/HeadlessChaosTestGraph.cpp
@@ -573,6 +573,7 @@ namespace ChaosTest {
 		PhysicalMaterial->SleepingAngularThreshold = 10;
 		PhysicalMaterial->DisabledLinearThreshold = 0;
 		PhysicalMaterial->DisabledAngularThreshold = 0;
+		PhysicalMaterial->SleepCounterThreshold = 5;
 
 		// Create some dynamic particles
 		int32 NumParticles = 6;
@@ -580,7 +581,7 @@ namespace ChaosTest {
 		TArray<TPBDRigidParticleHandle<FReal, 3>*> DynParticles = SOAs.CreateDynamicParticles(NumParticles);
 		TArray<TGeometryParticleHandle<FReal, 3>*> Particles;
 		TArrayCollectionArray<TSerializablePtr<FChaosPhysicsMaterial>> PhysicsMaterials;
-		SOAs.GetParticleHandles().AddArray(&PhysicsMaterials);
+ 		SOAs.GetParticleHandles().AddArray(&PhysicsMaterials);
 
 		for (int32 Idx = 0; Idx < NumParticles; ++Idx)
 		{
@@ -611,7 +612,7 @@ namespace ChaosTest {
 			Constraints.AddConstraint(ConstrainedParticleIndices);
 		}
 		FPBDConstraintGraph Graph;
-		for (int32 LoopIndex = 0; LoopIndex < 5 + FPBDConstraintGraph::SleepCountThreshold; ++LoopIndex)
+		for (int32 LoopIndex = 0; LoopIndex < 5 + PhysicalMaterial->SleepCounterThreshold; ++LoopIndex)
 		{
 			Graph.InitializeGraph(SOAs.GetNonDisabledView());
 
@@ -657,7 +658,7 @@ namespace ChaosTest {
 			EXPECT_TRUE(ContainsHelper(SOAs.GetActiveParticlesView(), Particles[1]));
 			EXPECT_TRUE(ContainsHelper(SOAs.GetActiveParticlesView(), Particles[2]));
 			// Particles 3-5 should sleep when we hit the frame count threshold and then stay asleep
-			bool bSomeShouldSleep = (LoopIndex > FPBDConstraintGraph::SleepCountThreshold);
+			bool bSomeShouldSleep = (LoopIndex >= PhysicalMaterial->SleepCounterThreshold );
 			EXPECT_EQ(Particles[3]->Sleeping(), bSomeShouldSleep);
 			EXPECT_EQ(Particles[4]->Sleeping(), bSomeShouldSleep);
 			EXPECT_EQ(Particles[5]->Sleeping(), bSomeShouldSleep);
diff --git a/Engine/Source/Runtime/Experimental/Chaos/Private/Chaos/PBDConstraintGraph.cpp b/Engine/Source/Runtime/Experimental/Chaos/Private/Chaos/PBDConstraintGraph.cpp
--- a/Engine/Source/Runtime/Experimental/Chaos/Private/Chaos/PBDConstraintGraph.cpp
+++ b/Engine/Source/Runtime/Experimental/Chaos/Private/Chaos/PBDConstraintGraph.cpp
@@ -17,6 +17,9 @@ using namespace Chaos;
 bool ChaosSolverCollisionDefaultUseMaterialSleepThresholdsCVar = true;
 FAutoConsoleVariableRef CVarChaosSolverCollisionDefaultUseMaterialSleepThresholds(TEXT("p.ChaosSolverCollisionDefaultUseMaterialSleepThresholds"), ChaosSolverCollisionDefaultUseMaterialSleepThresholdsCVar, TEXT("Enable material support for sleeping thresholds[def:true]"));
 
+int32 ChaosSolverCollisionDefaultSleepCounterThresholdCVar = 0; 
+FAutoConsoleVariableRef CVarChaosSolverCollisionDefaultSleepCounterThreshold(TEXT("p.ChaosSolverCollisionDefaultSleepCounterThreshold"), ChaosSolverCollisionDefaultSleepCounterThresholdCVar, TEXT("Default counter threshold for sleeping.[def:0]"));
+
 float ChaosSolverCollisionDefaultLinearSleepThresholdCVar = 0.001f; // .001 unit mass cm
 FAutoConsoleVariableRef CVarChaosSolverCollisionDefaultLinearSleepThreshold(TEXT("p.ChaosSolverCollisionDefaultLinearSleepThreshold"), ChaosSolverCollisionDefaultLinearSleepThresholdCVar, TEXT("Default linear threshold for sleeping.[def:0.001]"));
 
@@ -537,47 +540,15 @@ void FPBDConstraintGraph::ComputeIsland(const int32 InNode, const int32 Island,
 	}
 }
 
-
 bool FPBDConstraintGraph::SleepInactive(const int32 Island, const TArrayCollectionArray<TSerializablePtr<FChaosPhysicsMaterial>>& PerParticleMaterialAttributes)
 {
 	FReal LinearSleepingThreshold = FLT_MAX;
 	FReal AngularSleepingThreshold = FLT_MAX;
-	FReal DefaultLinearSleepingThreshold = ChaosSolverCollisionDefaultLinearSleepThresholdCVar;
-	FReal DefaultAngularSleepingThreshold = ChaosSolverCollisionDefaultAngularSleepThresholdCVar;
-
-	// @todo(ccaulfield): should be able to eliminate this when island is already sleeping
+	int32 SleepCounterThreshold = 0;
 
 	const TArray<TGeometryParticleHandle<FReal, 3>*>& IslandParticles = GetIslandParticles(Island);
 	check(IslandParticles.Num());
 
-	if (IslandParticles.Num() == 1)
-	{
-		if (TPBDRigidParticleHandle<FReal, 3>* PBDRigid = IslandParticles[0]->CastToRigidParticle())
-		{
-			TSerializablePtr<FChaosPhysicsMaterial> PhysicsMaterial = PBDRigid->AuxilaryValue(PerParticleMaterialAttributes);
-			if (PhysicsMaterial && ChaosSolverCollisionDefaultUseMaterialSleepThresholdsCVar)
-			{
-				LinearSleepingThreshold = FMath::Min(LinearSleepingThreshold, PhysicsMaterial->SleepingLinearThreshold);
-				AngularSleepingThreshold = FMath::Min(AngularSleepingThreshold, PhysicsMaterial->SleepingAngularThreshold);
-			}
-			else
-			{
-				LinearSleepingThreshold = DefaultLinearSleepingThreshold;
-				AngularSleepingThreshold = DefaultAngularSleepingThreshold;
-			}
-
-			if (PBDRigid->V().Size() < LinearSleepingThreshold && PBDRigid->W().Size() < AngularSleepingThreshold)
-			{
-				PBDRigid->SetSleeping(true);
-				PBDRigid->V() = TVector<FReal, 3>(0);
-				PBDRigid->W() = TVector<FReal, 3>(0);
-				return true;
-			}
-
-			return false;
-		}
-	}
-
 	if (!IslandToData[Island].bIsIslandPersistant)
 	{
 		return false;
@@ -588,35 +559,37 @@ bool FPBDConstraintGraph::SleepInactive(const int32 Island, const TArrayCollecti
 	TVector<FReal, 3> V(0);
 	TVector<FReal, 3> W(0);
 	FReal M = 0;
-	LinearSleepingThreshold = FLT_MAX;
-	AngularSleepingThreshold = FLT_MAX;
 
 	int32 NumDynamicParticles = 0;
 
 	for (TGeometryParticleHandle<FReal, 3>* Particle : IslandParticles)
 	{
-		TPBDRigidParticleHandle<FReal, 3>* PBDRigid = Particle->CastToRigidParticle();
-		if (PBDRigid && (PBDRigid->ObjectState() == EObjectStateType::Dynamic))
+		if (TPBDRigidParticleHandle<FReal, 3>* PBDRigid = Particle->CastToRigidParticle())
 		{
-			NumDynamicParticles++;
+			if (PBDRigid->ObjectState() == EObjectStateType::Dynamic)
+			{
+				NumDynamicParticles++;
 
-			M += PBDRigid->M();
-			V += PBDRigid->V() * PBDRigid->M();
+				M += PBDRigid->M();
+				V += PBDRigid->V() * PBDRigid->M();
 
-			if (TSerializablePtr<FChaosPhysicsMaterial> PhysicsMaterial = Particle->AuxilaryValue(PerParticleMaterialAttributes))
-			{
-				LinearSleepingThreshold = FMath::Min(LinearSleepingThreshold, PhysicsMaterial->SleepingLinearThreshold);
-				AngularSleepingThreshold = FMath::Min(AngularSleepingThreshold, PhysicsMaterial->SleepingAngularThreshold);
-			}
-			else
-			{
-				LinearSleepingThreshold = DefaultLinearSleepingThreshold;
-				AngularSleepingThreshold = DefaultAngularSleepingThreshold;
+				if (TSerializablePtr<FChaosPhysicsMaterial> PhysicsMaterial = Particle->AuxilaryValue(PerParticleMaterialAttributes))
+				{
+					LinearSleepingThreshold = FMath::Min(LinearSleepingThreshold, PhysicsMaterial->SleepingLinearThreshold);
+					AngularSleepingThreshold = FMath::Min(AngularSleepingThreshold, PhysicsMaterial->SleepingAngularThreshold);
+					SleepCounterThreshold = FMath::Max(SleepCounterThreshold, PhysicsMaterial->SleepCounterThreshold);
+				}
+				else
+				{
+					LinearSleepingThreshold = FMath::Min(LinearSleepingThreshold, ChaosSolverCollisionDefaultLinearSleepThresholdCVar);
+					AngularSleepingThreshold = FMath::Min(AngularSleepingThreshold, ChaosSolverCollisionDefaultAngularSleepThresholdCVar);
+					SleepCounterThreshold = FMath::Max(SleepCounterThreshold, ChaosSolverCollisionDefaultSleepCounterThresholdCVar);
+				}
 			}
 		}
 	}
 
-	if (NumDynamicParticles == 0)
+	if (NumDynamicParticles == 0 || FMath::IsNearlyZero(M))
 	{
 		// prevent divide by zero - all particles much be sleeping/disabled already
 		return false;
@@ -629,7 +602,7 @@ bool FPBDConstraintGraph::SleepInactive(const int32 Island, const TArrayCollecti
 		const TPBDRigidParticleHandle<FReal, 3>* PBDRigid = Particle->CastToRigidParticle();
 		if (PBDRigid && PBDRigid->ObjectState() == EObjectStateType::Dynamic)
 		{
-			W += /*TVector<FReal, 3>::CrossProduct(PBDRigid->X() - X, PBDRigid->M() * PBDRigid->V()/ +*/ PBDRigid->W() * PBDRigid->M();
+			W += PBDRigid->W() * PBDRigid->M();
 		}
 	}
 
@@ -639,18 +612,19 @@ bool FPBDConstraintGraph::SleepInactive(const int32 Island, const TArrayCollecti
 	const FReal WSize = W.SizeSquared();
 	if (VSize < LinearSleepingThreshold && WSize < AngularSleepingThreshold)
 	{
-		if (IslandSleepCount > SleepCountThreshold)
+		if (IslandSleepCount >= SleepCounterThreshold)
 		{
 			for (TGeometryParticleHandle<FReal, 3>* Particle : IslandParticles)
 			{
-				TPBDRigidParticleHandle<FReal, 3>* PBDRigid = Particle->CastToRigidParticle();
-				if (PBDRigid && PBDRigid->ObjectState() == EObjectStateType::Dynamic)
+				if (TPBDRigidParticleHandle<FReal, 3>* PBDRigid = Particle->CastToRigidParticle())
 				{
-					PBDRigid->SetSleeping(true);
-					PBDRigid->V() = TVector<FReal, 3>(0);
-					PBDRigid->W() = TVector<FReal, 3>(0);
+					if (PBDRigid->ObjectState() == EObjectStateType::Dynamic)
+					{
+						PBDRigid->SetSleeping(true);
+						PBDRigid->V() = TVector<FReal, 3>(0);
+						PBDRigid->W() = TVector<FReal, 3>(0);
+					}
 				}
-
 			}
 			return true;
 		}
diff --git a/Engine/Source/Runtime/Experimental/Chaos/Public/Chaos/Defines.h b/Engine/Source/Runtime/Experimental/Chaos/Public/Chaos/Defines.h
--- a/Engine/Source/Runtime/Experimental/Chaos/Public/Chaos/Defines.h
+++ b/Engine/Source/Runtime/Experimental/Chaos/Public/Chaos/Defines.h
@@ -63,6 +63,8 @@ namespace Chaos
 
 		void Serialize(FArchive& Ar)
 		{
+			Ar.UsingCustomVersion(FExternalPhysicsCustomObjectVersion::GUID);
+
 			Ar << Friction << Restitution << SleepingLinearThreshold << SleepingAngularThreshold << DisabledLinearThreshold << DisabledAngularThreshold;
 
 			if (Ar.CustomVer(FExternalPhysicsCustomObjectVersion::GUID) >= FExternalPhysicsCustomObjectVersion::PhysicsMaterialSleepCounterThreshold)
diff --git a/Engine/Source/Runtime/Experimental/Chaos/Public/Chaos/PBDConstraintGraph.h b/Engine/Source/Runtime/Experimental/Chaos/Public/Chaos/PBDConstraintGraph.h
--- a/Engine/Source/Runtime/Experimental/Chaos/Public/Chaos/PBDConstraintGraph.h
+++ b/Engine/Source/Runtime/Experimental/Chaos/Public/Chaos/PBDConstraintGraph.h
@@ -42,9 +42,6 @@ namespace Chaos
 	public:
 		friend class FPBDConstraintColor;
 
-		/** Number of sleep threshold passes before sleep mode is enabled */
-		static const int32 SleepCountThreshold = 5;
-
 		/** Information required to map a graph edge back to its constraint */
 		struct FConstraintData
 		{
