Changes to particle dirty properties

diff --git a/Engine/Source/Programs/HeadlessChaos/Private/HeadlessChaosTestDirtyParticle.cpp b/Engine/Source/Programs/HeadlessChaos/Private/HeadlessChaosTestDirtyParticle.cpp
--- a/Engine/Source/Programs/HeadlessChaos/Private/HeadlessChaosTestDirtyParticle.cpp
+++ b/Engine/Source/Programs/HeadlessChaos/Private/HeadlessChaosTestDirtyParticle.cpp
@@ -18,61 +18,48 @@ namespace ChaosTest
 	GTEST_TEST(DirtyParticleTests,Basic)
 	{
 		FDirtyPropertiesManager DirtyPropertiesPool;
-		FDirtyIdx Idx = DirtyPropertiesPool.WriteX(FVec3(1,1,1));
-		FDirtyIdx Idx2 = DirtyPropertiesPool.WriteX(FVec3(2,1,1));
-		FDirtyIdx Idx3 = DirtyPropertiesPool.WriteInvM(3);
-
-
-		EXPECT_EQ(Idx.Property, EParticleProperty::X);
-		EXPECT_EQ(Idx.Entry, 0);
-		EXPECT_EQ(DirtyPropertiesPool.ReadX(Idx),FVec3(1,1,1));
-
-
-		EXPECT_EQ(Idx2.Property, EParticleProperty::X);
-		EXPECT_EQ(Idx2.Entry, 1);
-		EXPECT_EQ(DirtyPropertiesPool.ReadX(Idx2),FVec3(2,1,1));
-
-
-		EXPECT_EQ(Idx3.Property, EParticleProperty::InvM);
-		EXPECT_EQ(Idx3.Entry, 0);
-		EXPECT_EQ(DirtyPropertiesPool.ReadInvM(Idx3),3);
-
-
-		//release second x proprty
-		DirtyPropertiesPool.FreeProperty(Idx2);
-
-		//get another x property, should be same idx as Idx2 because we reuse indices
-		FDirtyIdx Idx4 = DirtyPropertiesPool.WriteX(FVec3(4,1,1));
-		EXPECT_EQ(Idx4,Idx2);
-
-		EXPECT_EQ(DirtyPropertiesPool.ReadX(Idx4),FVec3(4,1,1));
-
-		//using higher level api
-		//pretend we have a particle with V and Mass as dirty
-
-		//create a dirty properties instance to push to other thread
-		FDirtyProperties DirtyParticle;
-
-		FDirtyIdx VIdx0 = DirtyPropertiesPool.WriteV(FVec3(5,1,1));
-		DirtyParticle.DirtyProperty(VIdx0);
-		DirtyParticle.DirtyProperty(DirtyPropertiesPool.WriteM(6));
-
-		//Other thread can see the dirty properties and read them
-		EXPECT_TRUE(DirtyParticle.IsDirty(EParticleFlags::V));
-		EXPECT_TRUE(DirtyParticle.IsDirty(EParticleFlags::M));
-		//Other values are clean
-		EXPECT_FALSE(DirtyParticle.IsDirty(EParticleFlags::X));
-
-		//read the values
-		EXPECT_EQ(DirtyParticle.ReadV(DirtyPropertiesPool),FVec3(5,1,1));
-		EXPECT_EQ(DirtyParticle.ReadM(DirtyPropertiesPool),6);
-
-		DirtyParticle.Clean(DirtyPropertiesPool);
-
-		//verify that we really cleaned the resource
-		FDirtyIdx VIdx1 = DirtyPropertiesPool.WriteV(FVec3(5,1,1));
-		EXPECT_EQ(VIdx1,VIdx0);
-
+		FDirtyProperties DirtyParticle1;
+		DirtyParticle1.WriteX(DirtyPropertiesPool, FVec3(1,1,1));
+		DirtyParticle1.WriteInvM(DirtyPropertiesPool, 3);
+
+		FDirtyProperties DirtyParticle2;
+		DirtyParticle2.WriteX(DirtyPropertiesPool, FVec3(2,1,1));
+
+		//properties are dirty
+		EXPECT_TRUE(DirtyParticle1.IsDirty(EParticleFlags::X));
+		EXPECT_TRUE(DirtyParticle1.IsDirty(EParticleFlags::InvM));
+
+		EXPECT_TRUE(DirtyParticle2.IsDirty(EParticleFlags::X));
+
+		//untouched properties are clean
+		EXPECT_FALSE(DirtyParticle2.IsDirty(EParticleFlags::InvM));
+
+		//values were saved
+		EXPECT_EQ(DirtyParticle1.ReadX(DirtyPropertiesPool), FVec3(1,1,1));
+		EXPECT_EQ(DirtyParticle1.ReadInvM(DirtyPropertiesPool), 3);
+		EXPECT_EQ(DirtyParticle2.ReadX(DirtyPropertiesPool), FVec3(2,1,1));
+		
+		EXPECT_EQ(DirtyParticle1.PopX(DirtyPropertiesPool),FVec3(1,1,1));
+		EXPECT_FALSE(DirtyParticle1.IsDirty(EParticleFlags::X));
+
+		//make sure we are not leaking shared ptrs
+		TSharedPtr<FImplicitObject,ESPMode::ThreadSafe> Ptr(new TSphere<FReal,3>(FVec3(0), 0));
+		DirtyParticle1.WriteGeometry(DirtyPropertiesPool, Ptr);
+		TWeakPtr<FImplicitObject,ESPMode::ThreadSafe> WeakPtr(Ptr);
+		Ptr = nullptr;
+
+		EXPECT_TRUE(WeakPtr.IsValid());	//still around because dirty pool is holding on to it
+		{
+			TSharedPtr<FImplicitObject,ESPMode::ThreadSafe> Geom = DirtyParticle1.PopGeometry(DirtyPropertiesPool);
+			EXPECT_TRUE(WeakPtr.IsValid());	//Popped the geometry but still holding on to it
+		}
+
+		EXPECT_FALSE(WeakPtr.IsValid());	//Finished with popped geometry so shared ptr goes away
+		
+
+		//If we haven't popped everything we call this
+		DirtyParticle1.Clean(DirtyPropertiesPool);
+		DirtyParticle2.Clean(DirtyPropertiesPool);
 	}
 }
 
diff --git a/Engine/Source/Runtime/Experimental/Chaos/Public/Chaos/ParticleDirtyFlags.h b/Engine/Source/Runtime/Experimental/Chaos/Public/Chaos/ParticleDirtyFlags.h
--- a/Engine/Source/Runtime/Experimental/Chaos/Public/Chaos/ParticleDirtyFlags.h
+++ b/Engine/Source/Runtime/Experimental/Chaos/Public/Chaos/ParticleDirtyFlags.h
@@ -75,19 +75,11 @@ using ChaosGeometrySp = TSharedPtr<FImplicitObject,ESPMode::ThreadSafe>;
 
 	struct FDirtyIdx
 	{
-		EParticleProperty Property : 5;
-		uint32 Entry : 27;
-
-		static_assert((uint32)EParticleProperty::NumProperties <= 31,"Cannot have more than 31 properties");
-
-		EParticleFlags GetPropertyFlag() const
-		{
-			return (EParticleFlags)(1 << (uint32)Property);
-		}
+		uint32 Entry;
 
 		bool operator==(const FDirtyIdx& Rhs) const
 		{
-			return ((const uint32&)*this) == ((const uint32&)Rhs);
+			return Entry == Rhs.Entry;
 		}
 	};
 
@@ -106,6 +98,15 @@ using ChaosGeometrySp = TSharedPtr<FImplicitObject,ESPMode::ThreadSafe>;
 			FreeIndices.Add(Idx);
 		}
 
+		T Pop(int32 Idx)
+		{
+			FreeIndices.Add(Idx);
+			T Result;
+			Swap(Result,Elements[Idx]);
+			Elements[Idx].~T();
+			return Result;
+		}
+
 		int32 Write(const T& Element)
 		{
 			const int32 Idx = GetFree();
@@ -113,18 +114,26 @@ using ChaosGeometrySp = TSharedPtr<FImplicitObject,ESPMode::ThreadSafe>;
 			return Idx;
 		}
 
-		int32 Write(T&& Element)
+		void Update(int32 Entry, const T& Element)
 		{
-			const int32 Idx = GetFree();
-			Elements[Idx] = MoveTemp(Element);
-			return Idx;
+			Elements[Entry] = Element;
 		}
 
 	private:
 
 		int32 GetFree()
 		{
-			return FreeIndices.Num() ? FreeIndices.Pop(/*bAllowShrinking=*/false) : Elements.AddUninitialized(1);
+			//todo: can we avoid default constructors? maybe if std::is_trivially_copyable
+			if(FreeIndices.Num())
+			{
+				int32 NewIdx = FreeIndices.Pop();
+				Elements[NewIdx] = T();
+				return NewIdx;
+			}
+			else
+			{
+				return Elements.AddDefaulted(1);
+			}
 		}
 
 		TArray<T> Elements;
@@ -134,25 +143,12 @@ using ChaosGeometrySp = TSharedPtr<FImplicitObject,ESPMode::ThreadSafe>;
 	class FDirtyPropertiesManager
 	{
 	public:
-#define PARTICLE_PROPERTY(PropName, Type) FDirtyIdx Write##PropName(Type const & Val){ FDirtyIdx Result; Result.Property = EParticleProperty::PropName; Result.Entry = Dirty##PropName##Pool.Write(Val); return Result; };
-#include "ParticleProperties.h"
-#undef PARTICLE_PROPERTY
-
-#define PARTICLE_PROPERTY(PropName, Type) Type const & Read##PropName(const FDirtyIdx Idx) const { ensure(Idx.Property == EParticleProperty::PropName); return Dirty##PropName##Pool.Read(Idx.Entry);};
-#include "ParticleProperties.h"
-#undef PARTICLE_PROPERTY
+#define PARTICLE_PROPERTY(PropName, Type) \
+const TDirtyElementPool<Type>& Get##PropName##Pool() const { return Dirty##PropName##Pool; }\
+TDirtyElementPool<Type>& Get##PropName##Pool() { return Dirty##PropName##Pool; }
 
-	void FreeProperty(FDirtyIdx Idx)
-	{
-		switch((uint32)Idx.Property)
-		{
-#define PARTICLE_PROPERTY(PropName, Type) case EParticleProperty::PropName: { Dirty##PropName##Pool.Free(Idx.Entry); break; }
 #include "ParticleProperties.h"
 #undef PARTICLE_PROPERTY
-		default: check(false);
-
-		}
-	}
 	private:
 
 		//creates a TDirtyElementPool for each property
@@ -174,12 +170,6 @@ using ChaosGeometrySp = TSharedPtr<FImplicitObject,ESPMode::ThreadSafe>;
 			return Flags.IsDirty(CheckBits);
 		}
 
-		void DirtyProperty(FDirtyIdx DirtyIdx)
-		{
-			Flags.MarkDirty(DirtyIdx.GetPropertyFlag());
-			DirtyProperties.Add(DirtyIdx);
-		}
-
 		~FDirtyProperties()
 		{
 			//Make sure to call Clean before killing this guy
@@ -188,32 +178,54 @@ using ChaosGeometrySp = TSharedPtr<FImplicitObject,ESPMode::ThreadSafe>;
 
 		void Clean(FDirtyPropertiesManager& DirtyPropertyManager)
 		{
-			for(const FDirtyIdx Idx : DirtyProperties)
+			if(!Flags.IsClean())
 			{
-				DirtyPropertyManager.FreeProperty(Idx);
+#define PARTICLE_PROPERTY(PropName, Type) if(Flags.IsDirty(EParticleFlags::PropName)){ Pop##PropName(DirtyPropertyManager); }
+#include "ParticleProperties.h"
+#undef PARTICLE_PROPERTY
 			}
+		}
+
+#define PARTICLE_PROPERTY(PropName, Type)\
+Type const& Read##PropName(const FDirtyPropertiesManager& DirtyPropertiesManager) const { return ReadImp(PropName, DirtyPropertiesManager.Get##PropName##Pool()); }\
+Type Pop##PropName(FDirtyPropertiesManager& DirtyPropertiesManager) { return PopImp(PropName, EParticleFlags::PropName, DirtyPropertiesManager.Get##PropName##Pool()); }\
+void Write##PropName(FDirtyPropertiesManager& DirtyPropertiesManager, Type const& Val) { WriteImp(PropName, EParticleFlags::PropName, Val, DirtyPropertiesManager.Get##PropName##Pool()); }
 
-			Flags.Clear();
-		}
-
-#define PARTICLE_PROPERTY(PropName, Type) Type const& Read##PropName(const FDirtyPropertiesManager& DirtyPropertyManager) const\
-{\
-	for(const FDirtyIdx Idx : DirtyProperties)\
-	{\
-		if(Idx.Property == EParticleProperty::PropName)\
-		{\
-			return DirtyPropertyManager.Read##PropName(Idx);\
-		}\
-	}\
-	check(false);\
-	static Type Error;\
-	return Error;\
-}
 #include "ParticleProperties.h"
 #undef PARTICLE_PROPERTY
 
 	private:
-		TArray<FDirtyIdx,TInlineAllocator<4>> DirtyProperties;
 		FParticleDirtyFlags Flags;
+#define PARTICLE_PROPERTY(PropName, Type) FDirtyIdx PropName;
+#include "ParticleProperties.h"
+#undef PARTICLE_PROPERTY
+
+		template <typename T>
+		void WriteImp(FDirtyIdx& Prop, EParticleFlags Flag, const T& Val, TDirtyElementPool<T>& Pool)
+		{
+			if(Flags.IsDirty(Flag))
+			{
+				Pool.Update(Prop.Entry,Val);
+			}
+			else
+			{
+				Prop.Entry = Pool.Write(Val);
+				Flags.MarkDirty(Flag);
+			}
+		}
+
+		template <typename T>
+		const T& ReadImp(const FDirtyIdx Prop, const TDirtyElementPool<T>& Pool) const
+		{
+			return Pool.Read(Prop.Entry);
+		}
+
+		template <typename T>
+		T PopImp(const FDirtyIdx Prop, EParticleFlags Flag, TDirtyElementPool<T>& Pool)
+		{
+			ensure(Flags.IsDirty(Flag));
+			Flags.MarkClean(Flag);
+			return Pool.Pop(Prop.Entry);
+		}
 	};
 }
