Implement CCD w/ Apply step at time of impact and at end of frame.

diff --git a/Engine/Source/Runtime/Experimental/Chaos/Private/Chaos/CollisionResolution.cpp b/Engine/Source/Runtime/Experimental/Chaos/Private/Chaos/CollisionResolution.cpp
--- a/Engine/Source/Runtime/Experimental/Chaos/Private/Chaos/CollisionResolution.cpp
+++ b/Engine/Source/Runtime/Experimental/Chaos/Private/Chaos/CollisionResolution.cpp
@@ -31,6 +31,11 @@ DECLARE_CYCLE_STAT(TEXT("Collisions::GJK"), STAT_Collisions_GJK, STATGROUP_Chaos
 
 //PRAGMA_DISABLE_OPTIMIZATION
 
+float CCDEnableThresholdBoundsScale = 0.4f;
+FAutoConsoleVariableRef  CVarCCDEnableThresholdBoundsScale(TEXT("p.Chaos.CCD.EnableThresholdBoundsScale"), CCDEnableThresholdBoundsScale , TEXT("CCD is used when object position is changing > smallest bound's extent * BoundsScale. 0 will always Use CCD. Values < 0 disables CCD."));
+
+float CCDAllowedDepthBoundsScale = 0.05f;
+FAutoConsoleVariableRef CVarCCDAllowedDepthBoundsScale(TEXT("p.Chaos.CCD.AllowedDepthBoundsScale"), CCDAllowedDepthBoundsScale, TEXT("When rolling back to TOI, allow (smallest bound's extent) * AllowedDepthBoundsScale, instead of rolling back to exact TOI w/ penetration = 0."));
 
 // If GJKPenetration returns a phi of abs value < this number, we use PhiWithNormal to resample phi and normal.
 // We have observed bad normals coming from GJKPenetration when barely in contact.
@@ -69,6 +74,104 @@ namespace Chaos
 			, Phi(TNumericLimits<T>::Max()) {}
 		};
 
+		// Determines if body should use CCD. If using CCD, computes Dir and Length of sweep.
+		template <typename T, int d>
+		bool UseCCD(const TGeometryParticleHandle<T, d>* SweptParticle, const TGeometryParticleHandle<T,d>* OtherParticle, const FImplicitObject* Implicit, TVector<T,d>& Dir, FReal& Length)
+		{
+			if (OtherParticle->ObjectState() != EObjectStateType::Static)
+			{
+				return false;
+			}
+
+			auto* RigidParticle= SweptParticle->CastToRigidParticle();
+			if (RigidParticle)
+			{
+				Dir = RigidParticle->P() - RigidParticle->X();
+				Length = Dir.Size();
+
+				if (!Implicit->HasBoundingBox())
+				{
+					return false;
+				}
+
+				T MinBoundsAxis = Implicit->BoundingBox().Extents().Min();
+				T LengthCCDThreshold = MinBoundsAxis * CCDEnableThresholdBoundsScale;
+
+				if (Length > 0.0f && CCDEnableThresholdBoundsScale >= 0.0f && Length >= LengthCCDThreshold)
+				{
+					Dir /= Length;
+					return true;
+				}
+			}
+
+			return false;
+		}
+
+		void ComputeSweptContactPhiAndTOIHelper(const FVec3& ContactNormal, const FVec3& Dir, const FReal& Length, const FReal& HitTime, FReal& OutTOI, FReal& OutPhi)
+		{
+			if (HitTime >= 0.0f)
+			{
+				// We subtract length to get the total penetration at at end of frame.
+				// Project penetration vector onto geometry normal for correct phi.
+				FReal Dot = FMath::Abs(FVec3::DotProduct(ContactNormal, -Dir));
+				OutPhi = (HitTime - Length) * Dot; 
+				
+				// TOI is between [0,1], used to compute particle position
+				OutTOI = HitTime / Length;
+			}
+			else
+			{
+				// Initial overlap case:
+				// TOI = 0 as we are overlapping at X.
+				// OutTime is penetration value of MTD.
+				OutPhi = HitTime;
+				OutTOI = 0.0f;
+			}
+		}
+
+		// Initializes TimeOfImpact and adjusts Phi for TOI.
+		template <typename T, int d>
+		void SetSweptConstraintTOI(TGeometryParticleHandle<T, d>* Particle, const FReal TOI, const FReal Length, const FVec3& Dir, FRigidBodySweptPointContactConstraint& SweptConstraint)
+		{
+			if (SweptConstraint.Manifold.Phi > 0.0f)
+			{
+				return;
+			}
+
+			TPBDRigidParticleHandle<T, d>* RigidParticle = Particle->CastToRigidParticle();
+			if (!CHAOS_ENSURE(RigidParticle))
+			{
+				return;
+			}
+
+			const FReal MinBounds = Particle->Geometry()->BoundingBox().Extents().Min();
+			const FReal Depth = -SweptConstraint.Manifold.Phi;
+			const FReal MaxDepth = MinBounds * CCDAllowedDepthBoundsScale;
+
+			if (TOI <= 0.0f)
+			{
+				// Initial overlap. Phi is correct already.
+				SweptConstraint.TimeOfImpact = 0.0f;
+			}
+			else if (TOI > 0.0f && TOI <= 1.0f)
+			{
+				if (Depth <= MaxDepth)
+				{
+					// Although we may have an earlier TOI, depth at TOI = 1 is within tolerance, so just use TOI = 1.
+					SweptConstraint.TimeOfImpact = 1.0f;
+				}
+				else if (Depth > MaxDepth)
+				{
+					// Compute ExtraT that will take us from TOI and 0 phi to slightly past TOI w/ phi at maximum allowed depth.
+					const FReal DirDotNormal = FMath::Abs(FVec3::DotProduct(Dir, SweptConstraint.Manifold.Normal));
+					const FReal ExtraT = MaxDepth / (DirDotNormal * Length);
+
+					SweptConstraint.TimeOfImpact = TOI + ExtraT;
+					SweptConstraint.Manifold.Phi = -MaxDepth;
+				}
+			}
+		}
+
 
 		template <typename T>
 		void UpdateContactPoint(TCollisionContact<T, 3>& Manifold, const TContactPoint<T>& NewContactPoint)
@@ -124,20 +227,7 @@ namespace Chaos
 			{
 				Contact.Location = BTM.TransformPosition(Location);
 				Contact.Normal = BTM.TransformVectorNoScale(Normal);
-
-				// If OutTime < 0, initial overlap. We subtract length to get the total penetration at at end of frame.
-				// Use Abs of Dot as normal is inconsistent with initial overlaps.
-				if (OutTime >= 0.0f)
-				{
-					FReal Dot = FMath::Abs(FVec3::DotProduct(Contact.Normal, -Dir));
-					Contact.Phi = (OutTime - Length) * Dot;
-					TOI = OutTime / Length;
-				}
-				else
-				{
-					Contact.Phi = OutTime;
-					TOI = 0.0f;
-				}
+				ComputeSweptContactPhiAndTOIHelper(Contact.Normal, Dir, Length, OutTime, TOI, Contact.Phi);
 			}
 
 			return Contact;
@@ -196,76 +286,16 @@ namespace Chaos
 			T OutTime = FLT_MAX;
 			int32 FaceIndex = -1;
 			TVec3<T> Location, Normal;
-			if (const TImplicitObjectScaled<GeometryA>* ScaledConvexImplicit = A.template GetObject<const TImplicitObjectScaled<GeometryA> >())
-			{
-
-				if (B.SweepGeom(*ScaledConvexImplicit, AToBTM, LocalDir, Length, OutTime, Location, Normal, FaceIndex, 0.0f, true))
-				{
-					Contact.Location = BTransform.TransformPosition(Location);
-					Contact.Normal = BTransform.TransformVectorNoScale(Normal);
-
-					// If OutTime < 0, initial overlap. We subtract length to get the total penetration at at end of frame.
-					// Use Abs of Dot as normal is inconsistent with initial overlaps.
-					if (OutTime >= 0.0f)
-					{
-						FReal Dot = FMath::Abs(FVec3::DotProduct(Contact.Normal, -Dir));
-						Contact.Phi = (OutTime - Length) * Dot;
-						TOI = OutTime / Length;
-					}
-					else
-					{
-						Contact.Phi = OutTime;
-						TOI = 0.0f;
-					}
-				}
-			}
-			else if (const TImplicitObjectInstanced<GeometryA>* InstancedConvexImplicit = A.template GetObject<const TImplicitObjectInstanced<GeometryA> >())
-			{
-				if (const GeometryA * InstancedInnerObject = static_cast<const GeometryA*>(InstancedConvexImplicit->GetInstancedObject()))
-				{
-					if (B.SweepGeom(*InstancedInnerObject, AToBTM, LocalDir, Length, OutTime, Location, Normal, FaceIndex, 0.0f, true))
-					{
-						Contact.Location = BTransform.TransformPosition(Location);
-						Contact.Normal = BTransform.TransformVectorNoScale(Normal);
 
-						// If OutTime < 0, initial overlap. We subtract length to get the total penetration at at end of frame.
-						// Use Abs of Dot as normal is inconsistent with initial overlaps.
-						if (OutTime >= 0.0f)
-						{
-							FReal Dot = FMath::Abs(FVec3::DotProduct(Contact.Normal, -Dir));
-							Contact.Phi = (OutTime - Length) * Dot;
-							TOI = OutTime / Length;
-						}
-						else
-						{
-							Contact.Phi = OutTime;
-							TOI = 0.0f;
-						}
-					}
-				}
-			}
-			else if (const GeometryA* ConvexImplicit = A.template GetObject<const GeometryA>())
+			Utilities::CastHelper(A, AStartTransform, [&](const auto& ADowncast, const TRigidTransform<T,d>& AFullTM)
 			{
-				if (B.SweepGeom(*ConvexImplicit, AToBTM, LocalDir, Length, OutTime, Location, Normal, FaceIndex, 0.0f, true))
+				if (B.SweepGeom(ADowncast, AToBTM, LocalDir, Length, OutTime, Location, Normal, FaceIndex, 0.0f, true))
 				{
 					Contact.Location = BTransform.TransformPosition(Location);
 					Contact.Normal = BTransform.TransformVectorNoScale(Normal);
-
-					// If OutTime < 0, initial overlap. We subtract length to get the total penetration at at end of frame.
-					// Use Abs of Dot as normal is inconsistent with initial overlaps.
-					if (OutTime >= 0.0f)
-					{
-						FReal Dot = FMath::Abs(FVec3::DotProduct(Contact.Normal, -Dir));
-						Contact.Phi = (OutTime - Length) * Dot;
-						TOI = OutTime / Length;
-					}
-					else
-					{
-						Contact.Phi = OutTime;;
-						TOI = 0.0f;
-					}
+					ComputeSweptContactPhiAndTOIHelper(Contact.Normal, Dir, Length, OutTime, TOI, Contact.Phi);
 				}
-			}
+			});
 
 			return Contact;
 		}
@@ -286,75 +316,16 @@ namespace Chaos
 			T OutTime = FLT_MAX;
 			TVec3<T> Location, Normal;
 			int32 FaceIndex = -1;
-			if (const TImplicitObjectScaled<GeometryA>* ScaledConvexImplicit = A.template GetObject<const TImplicitObjectScaled<GeometryA> >())
+
+			Utilities::CastHelper(A, AStartTransform, [&](const auto& ADowncast, const TRigidTransform<T,d>& AFullTM)
 			{
-				if (B.LowLevelSweepGeom(*ScaledConvexImplicit, AToBTM, LocalDir, Length, OutTime, Location, Normal, FaceIndex, 0.0f, true))
+				if(B.LowLevelSweepGeom(ADowncast, AToBTM, LocalDir, Length, OutTime, Location, Normal, FaceIndex, 0.0f, true))
 				{
 					Contact.Location = BTransform.TransformPosition(Location);
 					Contact.Normal = BTransform.TransformVectorNoScale(Normal);
-
-					// If OutTime < 0, initial overlap. We subtract length to get the total penetration at at end of frame.
-					// Use Abs of Dot as normal is inconsistent with initial overlaps.
-					if (OutTime >= 0.0f)
-					{
-						FReal Dot = FMath::Abs(FVec3::DotProduct(Contact.Normal, -Dir));
-						Contact.Phi = (OutTime - Length) * Dot;
-						TOI = OutTime / Length;
-					}
-					else
-					{
-						Contact.Phi = OutTime;
-						TOI = 0.0f;
-					}
+					ComputeSweptContactPhiAndTOIHelper(Contact.Normal, Dir, Length, OutTime, TOI, Contact.Phi);
 				}
-			}
-			else if (const TImplicitObjectInstanced<GeometryA>* InstancedConvexImplicit = A.template GetObject<const TImplicitObjectInstanced<GeometryA> >())
-			{
-				if (const GeometryA * InstancedInnerObject = static_cast<const GeometryA*>(InstancedConvexImplicit->GetInstancedObject()))
-				{
-					if (B.LowLevelSweepGeom(*InstancedInnerObject, AToBTM, LocalDir, Length, OutTime, Location, Normal, FaceIndex, 0.0f, true))
-					{
-						Contact.Location = BTransform.TransformPosition(Location);
-						Contact.Normal = BTransform.TransformVectorNoScale(Normal);
-
-						// If OutTime < 0, initial overlap. We subtract length to get the total penetration at at end of frame.
-						// Use Abs of Dot as normal is inconsistent with initial overlaps.
-						if (OutTime >= 0.0f)
-						{
-							FReal Dot = FMath::Abs(FVec3::DotProduct(Contact.Normal, -Dir));
-							Contact.Phi = (OutTime - Length) * Dot;
-							TOI = OutTime / Length;
-						}
-						else
-						{
-							Contact.Phi = OutTime;
-							TOI = 0.0f;
-						}
-					}
-				}
-			}
-			else if (const GeometryA* ConvexImplicit = A.template GetObject<const GeometryA>())
-			{
-				if (B.LowLevelSweepGeom(*ConvexImplicit, AToBTM, LocalDir, Length, OutTime, Location, Normal, FaceIndex, 0.0f, true))
-				{
-						Contact.Location = BTransform.TransformPosition(Location);
-						Contact.Normal = BTransform.TransformVectorNoScale(Normal);
-
-						// If OutTime < 0, initial overlap. We subtract length to get the total penetration at at end of frame.
-						// Use Abs of Dot as normal is inconsistent with initial overlaps.
-						if (OutTime >= 0.0f)
-						{
-							FReal Dot = FMath::Abs(FVec3::DotProduct(Contact.Normal, -Dir));
-							Contact.Phi = (OutTime - Length) * Dot;
-							TOI = OutTime / Length;
-						}
-						else
-						{
-							Contact.Phi = OutTime;
-							TOI = 0.0f;
-						}
-				}
-			}
+			});
 
 			return Contact;
 		}
@@ -1403,6 +1374,14 @@ namespace Chaos
 			UpdateContactPoint(Constraint.Manifold, CapsuleHeightFieldContactPoint(A, ATransform, B, BTransform, CullDistance));
 		}
 
+		template <typename T, int d>
+		void UpdateCapsuleHeightFieldConstraintSwept(TGeometryParticleHandle<T, d>* Particle0, const TCapsule<T>& A, const TRigidTransform<T, d>& ATransform, const THeightField<T>& B, const TRigidTransform<T, d>& BTransform, const TVector<T, d>& Dir, const T Length, const T CullDistance, TRigidBodySweptPointContactConstraint<T, d>& Constraint)
+		{
+			T TOI = 1.0f;
+			UpdateContactPoint(Constraint.Manifold, GJKImplicitSweptContactPoint<TCapsule<float> >(A, ATransform, B, BTransform, Dir, Length, CullDistance, TOI));
+			SetSweptConstraintTOI(Particle0, TOI, Length, Dir, Constraint);
+		}
+
 		template<class T, int d>
 		void UpdateCapsuleHeightFieldManifold(TCollisionConstraintBase<T, d>&  Constraint, const TRigidTransform<T, d>& ATM, const TRigidTransform<T, d>& BTM, const T CullDistance)
 		{
@@ -1426,6 +1405,24 @@ namespace Chaos
 			}
 		}
 
+		template<typename T, int d>
+		void ConstructCapsuleHeightFieldConstraintsSwept(TGeometryParticleHandle<T, d>* Particle0, TGeometryParticleHandle<T, d>* Particle1, const FImplicitObject* Implicit0, const FImplicitObject* Implicit1, const TRigidTransform<T, d>& Transform0, const TRigidTransform<T, d>& Transform1, const T CullDistance, const TVector<T,d>& Dir, FReal Length, FCollisionConstraintsArray& NewConstraints)
+		{
+			const TCapsule<T> * Object0 = Implicit0->template GetObject<const TCapsule<T> >();
+			const THeightField<T> * Object1 = Implicit1->template GetObject<const THeightField<T> >();
+			if (ensure(Object0 && Object1))
+			{
+				const TRigidTransform<T, 3> TransformX0(Particle0->X(), Particle0->R());
+				TRigidTransform<T, d> ParticleImplicit0TM = Transform0.GetRelativeTransform(Collisions::GetTransform(Particle0));
+				TRigidTransform<T, d> ParticleImplicit1TM = Transform1.GetRelativeTransform(Collisions::GetTransform(Particle1));
+
+				FRigidBodySweptPointContactConstraint Constraint = FRigidBodySweptPointContactConstraint(Particle0, Implicit0, ParticleImplicit0TM, Particle1, Implicit1, ParticleImplicit1TM, EContactShapesType::CapsuleHeightField);
+				UpdateCapsuleHeightFieldConstraintSwept(Particle0, *Object0, TransformX0, *Object1, Transform1, Dir, Length, CullDistance, Constraint);
+				NewConstraints.TryAdd(CullDistance, Constraint);
+			}
+		}
+
+
 		//
 		// Capsule-TriangleMesh
 		//
@@ -1436,12 +1433,38 @@ namespace Chaos
 			return GJKImplicitContactPoint< TCapsule<T> >(A, ATransform, B, BTransform, CullDistance);
 		}
 
+		template <typename TriMeshType, typename T, int d>
+		TContactPoint<T> CapsuleTriangleMeshSweptContactPoint(const TCapsule<T>& A, const TRigidTransform<T, d>& ATransform, const TriMeshType& B, const TRigidTransform<T, d>& BStartTransform, const TVector<T, d>& Dir, const T Length, const T CullDistance, T& TOI)
+		{
+			EImplicitObjectType TriMeshType = B.GetType();
+
+			if (const TImplicitObjectScaled<FTriangleMeshImplicitObject>* ScaledTriangleMesh = B.template GetObject<const TImplicitObjectScaled<FTriangleMeshImplicitObject>>())
+			{
+				return GJKImplicitScaledTriMeshSweptContactPoint<TCapsule<T>>(A, ATransform, *ScaledTriangleMesh, BStartTransform, Dir, Length, CullDistance, TOI);
+			}
+			else if (const FTriangleMeshImplicitObject* TriangleMesh = B.template GetObject<const FTriangleMeshImplicitObject>())
+			{
+				return GJKImplicitSweptContactPoint<TCapsule<T>>(A, ATransform, *TriangleMesh, BStartTransform, Dir, Length, CullDistance, TOI);
+			}
+
+			ensure(false);
+			return TContactPoint<T>();
+		}
+
+
 		template <typename TriMeshType, typename T, int d>
 		void UpdateCapsuleTriangleMeshConstraint(const TCapsule<T>& Capsule0, const TRigidTransform<T, d>& Transform0, const TriMeshType& TriangleMesh1, const TRigidTransform<T, d>& Transform1, const T CullDistance, TRigidBodyPointContactConstraint<T, d>& Constraint)
 		{
 			UpdateContactPoint(Constraint.Manifold, CapsuleTriangleMeshContactPoint(Capsule0, Transform0, TriangleMesh1, Transform1, CullDistance));
 		}
 
+		template <typename TriMeshType, typename T, int d>
+		void UpdateCapsuleTriangleMeshConstraintSwept(TGeometryParticleHandle<T, d>* Particle0, const TCapsule<T>& Capsule0, const TRigidTransform<T, d>& Transform0, const TriMeshType& TriangleMesh1, const TRigidTransform<T, d>& Transform1, const TVector<T, d>& Dir, const T Length, const T CullDistance, TRigidBodySweptPointContactConstraint<T, d>& Constraint)
+		{
+			T TOI = 1.0f;
+			UpdateContactPoint(Constraint.Manifold, CapsuleTriangleMeshSweptContactPoint(Capsule0, Transform0, TriangleMesh1, Transform1, Dir, Length, CullDistance, TOI));
+			SetSweptConstraintTOI(Particle0, TOI, Length, Dir, Constraint);
+		}
 
 		template <typename T, int d>
 		void UpdateCapsuleTriangleMeshManifold(TCollisionConstraintBase<T, d>&  Constraint, const TRigidTransform<T, d>& ATM, const TRigidTransform<T, d>& BTM, const T CullDistance)
@@ -1478,6 +1501,35 @@ namespace Chaos
 			}
 		}
 
+		template<typename T, int d>
+		void ConstructCapsuleTriangleMeshConstraintsSwept(TGeometryParticleHandle<T, d>* Particle0, TGeometryParticleHandle<T, d>* Particle1, const FImplicitObject* Implicit0, const FImplicitObject* Implicit1, const TRigidTransform<T, d>& Transform0, const TRigidTransform<T, d>& Transform1, const T CullDistance, const TVector<T,d>& Dir, FReal Length, FCollisionConstraintsArray& NewConstraints)
+		{
+			const TCapsule<T> * Object0 = Implicit0->template GetObject<const TCapsule<T> >();
+			if (ensure(Object0))
+			{
+				TRigidTransform<T, 3> TransformX0(Particle0->X(), Particle0->R());
+				TRigidTransform<T, d> ParticleImplicit0TM = Transform0.GetRelativeTransform(Collisions::GetTransform(Particle0));
+				TRigidTransform<T, d> ParticleImplicit1TM = Transform1.GetRelativeTransform(Collisions::GetTransform(Particle1));
+
+				if (const TImplicitObjectScaled<FTriangleMeshImplicitObject>* ScaledTriangleMesh = Implicit1->template GetObject<const TImplicitObjectScaled<FTriangleMeshImplicitObject>>())
+				{
+					FRigidBodySweptPointContactConstraint Constraint = FRigidBodySweptPointContactConstraint(Particle0, Implicit0, ParticleImplicit0TM, Particle1, Implicit1, ParticleImplicit1TM, EContactShapesType::CapsuleTriMesh);
+					UpdateCapsuleTriangleMeshConstraintSwept(Particle0, *Object0, TransformX0, *ScaledTriangleMesh, Transform1, Dir, Length, CullDistance, Constraint);
+					NewConstraints.TryAdd(CullDistance, Constraint);
+				}
+				else if (const FTriangleMeshImplicitObject* TriangleMesh = Implicit1->template GetObject<const FTriangleMeshImplicitObject>())
+				{
+					FRigidBodySweptPointContactConstraint Constraint = FRigidBodySweptPointContactConstraint(Particle0, Implicit0, ParticleImplicit0TM, Particle1, Implicit1, ParticleImplicit1TM, EContactShapesType::CapsuleTriMesh);
+					UpdateCapsuleTriangleMeshConstraintSwept(Particle0, *Object0, TransformX0, *TriangleMesh, Transform1, Dir, Length, CullDistance, Constraint);
+					NewConstraints.TryAdd(CullDistance, Constraint);
+				}
+				else
+				{
+					ensure(false);
+				}
+			}
+		}
+
 		//
 		// Convex - Convex
 		//
@@ -1491,6 +1543,10 @@ namespace Chaos
 			{
 				ContactPoint = ConvexConvexContactPoint(Implicit0, Transform0, Implicit1, Transform1, CullDistance);
 			}
+			else if (ConstraintBase.GetType() == FRigidBodySweptPointContactConstraint::StaticType())
+			{
+				ContactPoint = ConvexConvexContactPoint(Implicit0, Transform0, Implicit1, Transform1, CullDistance);
+			}
 			else if (ConstraintBase.GetType() == FRigidBodyMultiPointContactConstraint::StaticType())
 			{
 				TRigidBodyMultiPointContactConstraint<T, d>& Constraint = *ConstraintBase.template As<TRigidBodyMultiPointContactConstraint<T, d>>();
@@ -1520,6 +1576,14 @@ namespace Chaos
 			UpdateContactPoint(ConstraintBase.Manifold, ContactPoint);
 		}
 
+		template<class T, int d>
+		void UpdateConvexConvexConstraintSwept(TGeometryParticleHandle<T, d>* Particle0, const FImplicitObject& Implicit0, const TRigidTransform<T, d>& Transform0, const FImplicitObject& Implicit1, const TRigidTransform<T, d>& Transform1, const TVector<T, d>& Dir, const T Length, const T CullDistance, TRigidBodySweptPointContactConstraint<T, d>& Constraint)
+		{
+			T TOI = 1.0f;
+			UpdateContactPoint(Constraint.Manifold, ConvexConvexContactPointSwept(Implicit0, Transform0, Implicit1, Transform1, Dir, Length, CullDistance, TOI));
+			SetSweptConstraintTOI(Particle0, TOI, Length, Dir, Constraint);
+		}
+
 		template<class T, int d>
 		void UpdateConvexConvexManifold(TCollisionConstraintBase<T, d>&  ConstraintBase, const TRigidTransform<T, d>& Transform0, const TRigidTransform<T, d>& Transform1, const T CullDistance)
 		{
@@ -1558,6 +1622,18 @@ namespace Chaos
 			}
 		}
 
+		template<class T, int d>
+		void ConstructConvexConvexConstraintsSwept(TGeometryParticleHandle<T, d>* Particle0, TGeometryParticleHandle<T, d>* Particle1, const FImplicitObject* Implicit0, const FImplicitObject* Implicit1, const TRigidTransform<T, d>& Transform0, const TRigidTransform<T, d>& Transform1, const T CullDistance, const TVector<T,d>& Dir, FReal Length, FCollisionConstraintsArray& NewConstraints)
+		{
+			TRigidTransform<T, 3> TransformX0(Particle0->X(), Particle0->R());
+			TRigidTransform<T, d> ParticleImplicit0TM = Transform0.GetRelativeTransform(Collisions::GetTransform(Particle0));
+			TRigidTransform<T, d> ParticleImplicit1TM = Transform1.GetRelativeTransform(Collisions::GetTransform(Particle1));
+
+			FRigidBodySweptPointContactConstraint Constraint = FRigidBodySweptPointContactConstraint(Particle0, Implicit0, ParticleImplicit0TM, Particle1, Implicit1, ParticleImplicit1TM, EContactShapesType::ConvexConvex);
+			UpdateConvexConvexConstraintSwept(Particle0, *Implicit0, TransformX0, *Implicit1, Transform1, Dir, Length, CullDistance, Constraint);
+			NewConstraints.TryAdd(CullDistance, Constraint);
+		}
+
 		//
 		// Convex - HeightField
 		//
@@ -1574,6 +1650,14 @@ namespace Chaos
 			UpdateContactPoint(Constraint.Manifold, ConvexHeightFieldContactPoint(A, ATransform, B, BTransform, CullDistance));
 		}
 
+		template <typename T, int d>
+		void UpdateConvexHeightFieldConstraintSwept(TGeometryParticleHandle<T, d>* Particle0, const FImplicitObject& A, const TRigidTransform<T, d>& ATransform, const THeightField<T>& B, const TRigidTransform<T, d>& BTransform, const TVector<T, d>& Dir, const T Length, const T CullDistance, TRigidBodySweptPointContactConstraint<T, d>& Constraint)
+		{
+			T TOI = 1.0f;
+			UpdateContactPoint(Constraint.Manifold, GJKImplicitSweptContactPoint< FConvex >(A, ATransform, B, BTransform, Dir, Length, CullDistance, TOI));
+			SetSweptConstraintTOI(Particle0, TOI, Length, Dir, Constraint);
+		}
+
 		template<class T, int d>
 		void UpdateConvexHeightFieldManifold(TCollisionConstraintBase<T, d>&  Constraint, const TRigidTransform<T, d>& ATM, const TRigidTransform<T, d>& BTM, const T CullDistance)
 		{
@@ -1595,6 +1679,22 @@ namespace Chaos
 			}
 		}
 
+		template<typename T, int d>
+		void ConstructConvexHeightFieldConstraintsSwept(TGeometryParticleHandle<T, d>* Particle0, TGeometryParticleHandle<T, d>* Particle1, const FImplicitObject* Implicit0, const FImplicitObject* Implicit1, const TRigidTransform<T, d>& Transform0, const TRigidTransform<T, d>& Transform1, const T CullDistance, const TVector<T, d>& Dir, FReal Length, FCollisionConstraintsArray& NewConstraints)
+		{
+			const THeightField<T> * Object1 = Implicit1->template GetObject<const THeightField<T> >();
+			if (ensure(Implicit0->IsConvex() && Object1))
+			{
+				TRigidTransform<T, 3> TransformX0(Particle0->X(), Particle0->R());
+				TRigidTransform<T, d> ParticleImplicit0TM = Transform0.GetRelativeTransform(Collisions::GetTransform(Particle0));
+				TRigidTransform<T, d> ParticleImplicit1TM = Transform1.GetRelativeTransform(Collisions::GetTransform(Particle1));
+
+				FRigidBodySweptPointContactConstraint Constraint = FRigidBodySweptPointContactConstraint(Particle0, Implicit0, ParticleImplicit0TM, Particle1, Implicit1, ParticleImplicit1TM, EContactShapesType::ConvexHeightField);
+				UpdateConvexHeightFieldConstraintSwept(Particle0, *Implicit0, TransformX0, *Object1, Transform1, Dir, Length, CullDistance, Constraint);
+				NewConstraints.TryAdd(CullDistance, Constraint);
+			}
+		}
+
 		//
 		// Convex-TriangleMesh
 		//
@@ -1605,12 +1705,38 @@ namespace Chaos
 			return GJKImplicitContactPoint< FConvex >(A, ATransform, B, BTransform, CullDistance);
 		}
 
+		template <typename TriMeshType, typename T, int d>
+		TContactPoint<T> ConvexTriangleMeshSweptContactPoint(const FImplicitObject& A, const TRigidTransform<T, d>& ATransform, const TriMeshType& B, const TRigidTransform<T, d>& BStartTransform, const TVector<T, d>& Dir, const T Length, const T CullDistance, T& TOI)
+		{
+			EImplicitObjectType TriMeshType = B.GetType();
+			if (const TImplicitObjectScaled<FTriangleMeshImplicitObject>* ScaledTriangleMesh = B.template GetObject<const TImplicitObjectScaled<FTriangleMeshImplicitObject>>())
+			{
+				return GJKImplicitScaledTriMeshSweptContactPoint<FConvex>(A, ATransform, *ScaledTriangleMesh, BStartTransform, Dir, Length, CullDistance, TOI);
+			}
+			else if(const FTriangleMeshImplicitObject* TriangleMesh = B.template GetObject<const FTriangleMeshImplicitObject>())
+			{
+				return GJKImplicitSweptContactPoint<FConvex>(A, ATransform, *TriangleMesh, BStartTransform, Dir, Length, CullDistance, TOI);
+			}
+
+			ensure(false);
+			return TContactPoint<T>();
+		}
+		
 		template <typename TriMeshType, typename T, int d>
 		void UpdateConvexTriangleMeshConstraint(const FImplicitObject& Convex0, const TRigidTransform<T, d>& Transform0, const TriMeshType& TriangleMesh1, const TRigidTransform<T, d>& Transform1, const T CullDistance, TRigidBodyPointContactConstraint<T, d>& Constraint)
 		{
 			UpdateContactPoint(Constraint.Manifold, ConvexTriangleMeshContactPoint(Convex0, Transform0, TriangleMesh1, Transform1, CullDistance));
 		}
-		
+
+		// Sweeps convex against trimesh
+		template <typename TriMeshType, typename T, int d>
+		void UpdateConvexTriangleMeshConstraintSwept(TGeometryParticleHandle<T,d>* Particle0, const FImplicitObject& Convex0, const TRigidTransform<T, d>& Transform0, const TriMeshType& TriangleMesh1, const TRigidTransform<T, d>& Transform1, const TVector<T, d>& Dir, const T Length, const T CullDistance, TRigidBodySweptPointContactConstraint<T, d>& Constraint)
+		{
+			T TOI = 1.0f;
+			UpdateContactPoint(Constraint.Manifold, ConvexTriangleMeshSweptContactPoint(Convex0, Transform0, TriangleMesh1, Transform1, Dir, Length, CullDistance, TOI));
+			SetSweptConstraintTOI(Particle0, TOI, Length, Dir, Constraint);
+		}
+
 		template <typename T, int d>
 		void UpdateConvexTriangleMeshManifold(TCollisionConstraintBase<T, d>&  Constraint, const TRigidTransform<T, d>& ATM, const TRigidTransform<T, d>& BTM, const T CullDistance)
 		{
@@ -1645,7 +1771,32 @@ namespace Chaos
 			}
 		}
 
-
+		template<typename T, int d>
+		void ConstructConvexTriangleMeshConstraintsSwept(TGeometryParticleHandle<T, d>* Particle0, TGeometryParticleHandle<T, d>* Particle1, const FImplicitObject* Implicit0, const FImplicitObject* Implicit1, const TRigidTransform<T, d>& Transform0, const TRigidTransform<T, d>& Transform1, const T CullDistance, const TVector<T, d>& Dir, FReal Length, FCollisionConstraintsArray& NewConstraints)
+		{
+			TRigidTransform<T, d> ParticleImplicit0TM = Transform0.GetRelativeTransform(Collisions::GetTransform(Particle0));
+			TRigidTransform<T, d> ParticleImplicit1TM = Transform1.GetRelativeTransform(Collisions::GetTransform(Particle1));
+			if (ensure(Implicit0->IsConvex()))
+			{
+				TRigidTransform<T, 3> TransformX0(Particle0->X(), Particle0->R());
+				if (const TImplicitObjectScaled<FTriangleMeshImplicitObject>* ScaledTriangleMesh = Implicit1->template GetObject<const TImplicitObjectScaled<FTriangleMeshImplicitObject>>())
+				{
+					FRigidBodySweptPointContactConstraint Constraint = FRigidBodySweptPointContactConstraint(Particle0, Implicit0, ParticleImplicit0TM, Particle1, Implicit1, ParticleImplicit1TM, EContactShapesType::ConvexTriMesh);
+					UpdateConvexTriangleMeshConstraintSwept(Particle0, *Implicit0, TransformX0, *ScaledTriangleMesh, Transform1, Dir, Length, CullDistance, Constraint);
+					NewConstraints.TryAdd(CullDistance, Constraint);
+				}
+				else if (const FTriangleMeshImplicitObject* TriangleMesh = Implicit1->template GetObject<const FTriangleMeshImplicitObject>())
+				{
+					FRigidBodySweptPointContactConstraint Constraint = FRigidBodySweptPointContactConstraint(Particle0, Implicit0, ParticleImplicit0TM, Particle1, Implicit1, ParticleImplicit1TM, EContactShapesType::ConvexTriMesh);
+					UpdateConvexTriangleMeshConstraintSwept(Particle0, *Implicit0, TransformX0, *TriangleMesh, Transform1, Dir, Length, CullDistance, Constraint);
+					NewConstraints.TryAdd(CullDistance, Constraint);
+				}
+				else
+				{
+					ensure(false);
+				}
+			}
+		}
 
 
 		//
@@ -1880,6 +2031,10 @@ namespace Chaos
 			bool bIsConvex0 = Implicit0 && Implicit0->IsConvex();
 			bool bIsConvex1 = Implicit1 && Implicit1->IsConvex();
 
+			T LengthCCD = 0.0f;
+			TVector<T, d> DirCCD(0.0f);
+			bool bUseCCD = UseCCD(Particle0, Particle1, Implicit0, DirCCD, LengthCCD);
+
 			if (Implicit0Type == TBox<FReal, 3>::StaticType() && Implicit1Type == TBox<FReal, 3>::StaticType())
 			{
 				ConstructBoxBoxConstraints(Particle0, Particle1, Implicit0, Implicit1, Transform0, Transform1, CullDistance, NewConstraints);
@@ -1950,10 +2105,22 @@ namespace Chaos
 			}
 			else if (Implicit0Type == TCapsule<FReal>::StaticType() && Implicit1Type == THeightField<FReal>::StaticType())
 			{
+				if (bUseCCD)
+				{
+					ConstructCapsuleHeightFieldConstraintsSwept(Particle0, Particle1, Implicit0, Implicit1, Transform0, Transform1, CullDistance, DirCCD, LengthCCD, NewConstraints);
+					return;
+				}
+
 				ConstructCapsuleHeightFieldConstraints(Particle0, Particle1, Implicit0, Implicit1, Transform0, Transform1, CullDistance, NewConstraints);
 			}
 			else if (Implicit0Type == THeightField<FReal>::StaticType() && Implicit1Type == TCapsule<FReal>::StaticType())
 			{
+				if (bUseCCD)
+				{
+					ConstructCapsuleHeightFieldConstraintsSwept(Particle1, Particle0, Implicit1, Implicit0, Transform1, Transform0, CullDistance, DirCCD, LengthCCD, NewConstraints);
+					return;
+				}
+
 				ConstructCapsuleHeightFieldConstraints(Particle1, Particle0, Implicit1, Implicit0, Transform1, Transform0, CullDistance, NewConstraints);
 			}
 			else if (Implicit0Type == TBox<FReal, 3>::StaticType() && Implicit1Type == FTriangleMeshImplicitObject::StaticType())
@@ -1974,30 +2141,72 @@ namespace Chaos
 			}
 			else if (Implicit0Type == TCapsule<FReal>::StaticType() && Implicit1Type == FTriangleMeshImplicitObject::StaticType())
 			{
+				if (bUseCCD)
+				{
+					ConstructCapsuleTriangleMeshConstraintsSwept(Particle0, Particle1, Implicit0, Implicit1, Transform0, Transform1, CullDistance, DirCCD, LengthCCD, NewConstraints);
+					return;
+				}
+
 				ConstructCapsuleTriangleMeshConstraints(Particle0, Particle1, Implicit0, Implicit1, Transform0, Transform1, CullDistance, NewConstraints);
 			}
 			else if (Implicit0Type == FTriangleMeshImplicitObject::StaticType() && Implicit1Type == TCapsule<FReal>::StaticType())
 			{
+				if (bUseCCD)
+				{
+					ConstructCapsuleTriangleMeshConstraintsSwept(Particle0, Particle1, Implicit0, Implicit1, Transform0, Transform1, CullDistance, DirCCD, LengthCCD, NewConstraints);
+					return;
+				}
+
 				ConstructCapsuleTriangleMeshConstraints(Particle1, Particle0, Implicit1, Implicit0, Transform1, Transform0, CullDistance, NewConstraints);
 			}
 			else if (bIsConvex0 && Implicit1Type == THeightField<FReal>::StaticType())
 			{
+				if (bUseCCD)
+				{
+					ConstructConvexHeightFieldConstraintsSwept(Particle0, Particle1, Implicit0, Implicit1, Transform0, Transform1, CullDistance, DirCCD, LengthCCD, NewConstraints);
+					return;
+				}
+
 				ConstructConvexHeightFieldConstraints(Particle0, Particle1, Implicit0, Implicit1, Transform0, Transform1, CullDistance, NewConstraints);
 			}
 			else if (Implicit0Type == THeightField<FReal>::StaticType() && bIsConvex1)
 			{
+				if (bUseCCD)
+				{
+					ConstructConvexHeightFieldConstraintsSwept(Particle1, Particle0, Implicit1, Implicit0, Transform1, Transform0, CullDistance, DirCCD, LengthCCD, NewConstraints);
+					return;
+				}
+
 				ConstructConvexHeightFieldConstraints(Particle1, Particle0, Implicit1, Implicit0, Transform1, Transform0, CullDistance, NewConstraints);
 			}
 			else if (bIsConvex0 && Implicit1Type == FTriangleMeshImplicitObject::StaticType())
 			{
+				if (bUseCCD)
+				{
+					ConstructConvexTriangleMeshConstraintsSwept(Particle0, Particle1, Implicit0, Implicit1, Transform0, Transform1, CullDistance, DirCCD, LengthCCD, NewConstraints);
+					return;
+				}
+
 				ConstructConvexTriangleMeshConstraints(Particle0, Particle1, Implicit0, Implicit1, Transform0, Transform1, CullDistance, NewConstraints);
 			}
 			else if (Implicit0Type == FTriangleMeshImplicitObject::StaticType() && bIsConvex1)
 			{
+				if (bUseCCD)
+				{
+					ConstructConvexTriangleMeshConstraintsSwept(Particle1, Particle0, Implicit1, Implicit0, Transform1, Transform0, CullDistance, DirCCD, LengthCCD, NewConstraints);
+					return;
+				}
+
 				ConstructConvexTriangleMeshConstraints(Particle1, Particle0, Implicit1, Implicit0, Transform1, Transform0, CullDistance, NewConstraints);
 			}
 			else if (bIsConvex0 && bIsConvex1)
 			{
+				if (bUseCCD)
+				{
+					ConstructConvexConvexConstraintsSwept(Particle0, Particle1, Implicit0, Implicit1, Transform0, Transform1, CullDistance, DirCCD, LengthCCD, NewConstraints);
+					return;
+				}
+
 				ConstructConvexConvexConstraints(Particle0, Particle1, Implicit0, Implicit1, Transform0, Transform1, CullDistance, NewConstraints);
 			}
 			else
diff --git a/Engine/Source/Runtime/Experimental/Chaos/Private/Chaos/PBDCollisionConstraints.cpp b/Engine/Source/Runtime/Experimental/Chaos/Private/Chaos/PBDCollisionConstraints.cpp
--- a/Engine/Source/Runtime/Experimental/Chaos/Private/Chaos/PBDCollisionConstraints.cpp
+++ b/Engine/Source/Runtime/Experimental/Chaos/Private/Chaos/PBDCollisionConstraints.cpp
@@ -189,6 +189,33 @@ namespace Chaos
 		UpdateConstraintMaterialProperties(PointConstraints[Idx]);
 	}
 
+	template<typename T, int d>
+	void TPBDCollisionConstraints<T, d>::AddConstraint(const TRigidBodySweptPointContactConstraint<FReal, 3>& InConstraint)
+	{
+		int32 Idx = SweptPointConstraints.Add(InConstraint);
+
+		if (bHandlesEnabled)
+		{
+			FConstraintContainerHandle* Handle = HandleAllocator.template AllocHandle< TRigidBodySweptPointContactConstraint<T, d> >(this, Idx);
+			Handle->GetContact().Timestamp = -INT_MAX; // force point constraints to be deleted.
+
+			SweptPointConstraints[Idx].ConstraintHandle = Handle;
+
+			if(ensure(Handle != nullptr))
+			{			
+				Handles.Add(Handle);
+
+#if CHAOS_COLLISION_PERSISTENCE_ENABLED
+				check(!Manifolds.Contains(Handle->GetKey()));
+				Manifolds.Add(Handle->GetKey(), Handle);
+#endif
+			}
+		}
+
+		UpdateConstraintMaterialProperties(SweptPointConstraints[Idx]);
+	}
+
+
 	template<typename T, int d>
 	void TPBDCollisionConstraints<T, d>::AddConstraint(const TRigidBodyMultiPointContactConstraint<FReal, 3>& InConstraint)
 	{
@@ -243,6 +270,7 @@ namespace Chaos
 			HandleAllocator.FreeHandle(Handle);
 		}
 		PointConstraints.Reset();
+		SweptPointConstraints.Reset();
 		IterativeConstraints.Reset();
 		Handles.Reset();
 #endif
@@ -312,6 +340,22 @@ namespace Chaos
 			}
 
 		}
+		else if (ConstraintType == FCollisionConstraintBase::FType::SinglePointSwept)
+		{
+#if CHAOS_COLLISION_PERSISTENCE_ENABLED
+			if (Idx < SweptPointConstraints.Num() - 1)
+			{
+				// update the handle
+				FConstraintContainerHandleKey Key = FPBDCollisionConstraintHandle::MakeKey(&SweptPointConstraints.Last());
+				Manifolds[Key]->SetConstraintIndex(Idx, ConstraintType);
+			}
+#endif
+			SweptPointConstraints.RemoveAtSwap(Idx);
+			if (bHandlesEnabled && (Idx < SweptPointConstraints.Num()))
+			{
+				SweptPointConstraints[Idx].ConstraintHandle->SetConstraintIndex(Idx, FCollisionConstraintBase::FType::SinglePointSwept);
+			}
+		}
 		else if (ConstraintType == FCollisionConstraintBase::FType::MultiPoint)
 		{
 #if CHAOS_COLLISION_PERSISTENCE_ENABLED
@@ -342,7 +386,7 @@ namespace Chaos
 			Manifolds.Remove(KeyToRemove);
 #endif
 			Handles.Remove(Handle);
-			check(Handles.Num() == PointConstraints.Num() + IterativeConstraints.Num());
+			check(Handles.Num() == PointConstraints.Num() + SweptPointConstraints.Num() + IterativeConstraints.Num());
 
 			HandleAllocator.FreeHandle(Handle);
 		}
@@ -428,6 +472,11 @@ namespace Chaos
 				Collisions::ApplySinglePoint(Contact, IterationParameters, ParticleParameters);
 			}
 
+			for (FSweptPointContactConstraint& Contact : SweptPointConstraints)
+			{
+				Collisions::Apply(Contact, IterationParameters, ParticleParameters);
+			}
+
 			for (FMultiPointContactConstraint& Contact : IterativeConstraints)
 			{
 				Collisions::ApplyMultiPoint(Contact, IterationParameters, ParticleParameters);
@@ -459,6 +508,11 @@ namespace Chaos
 				Collisions::ApplyPushOutSinglePoint(Contact, TempStatic, IterationParameters, ParticleParameters);
 			}
 
+			for (FSweptPointContactConstraint& Contact : SweptPointConstraints)
+			{
+				Collisions::ApplyPushOut(Contact, TempStatic, IterationParameters, ParticleParameters);
+			}
+
 			for (FMultiPointContactConstraint& Contact : IterativeConstraints)
 			{
 				Collisions::ApplyPushOutMultiPoint(Contact, TempStatic, IterationParameters, ParticleParameters);
diff --git a/Engine/Source/Runtime/Experimental/Chaos/Private/Chaos/PBDCollisionConstraintsContact.cpp b/Engine/Source/Runtime/Experimental/Chaos/Private/Chaos/PBDCollisionConstraintsContact.cpp
--- a/Engine/Source/Runtime/Experimental/Chaos/Private/Chaos/PBDCollisionConstraintsContact.cpp
+++ b/Engine/Source/Runtime/Experimental/Chaos/Private/Chaos/PBDCollisionConstraintsContact.cpp
@@ -36,6 +36,12 @@ namespace Chaos
 			UpdateConstraintFromGeometry<ECollisionUpdateType::Deepest>(Constraint, Transform0, Transform1, CullDistance);
 		}
 
+		void Update(FRigidBodySweptPointContactConstraint& Constraint, const FReal CullDistance)
+		{
+			// Update as a point constraint (base class).
+			Update(*Constraint.As<FRigidBodyPointContactConstraint>(), CullDistance);
+		}
+
 		void Update(FRigidBodyMultiPointContactConstraint& Constraint, const FReal CullDistance)
 		{
 			const FRigidTransform3 Transform0 = GetTransform(Constraint.Particle[0]);
@@ -376,6 +382,36 @@ namespace Chaos
 				}
 			}
 		}
+		
+		void ApplySweptImpl(FRigidBodySweptPointContactConstraint& Constraint, const FContactIterationParameters & IterationParameters, const FContactParticleParameters & ParticleParameters)
+		{
+			TGenericParticleHandle<FReal, 3> Particle0 = TGenericParticleHandle<FReal, 3>(Constraint.Particle[0]);
+			TGenericParticleHandle<FReal, 3> Particle1 = TGenericParticleHandle<FReal, 3>(Constraint.Particle[1]);
+
+			if (IterationParameters.Iteration > 0 || Constraint.TimeOfImpact == 1)
+			{
+				// If not on first iteration, or at TOI = 1 (normal constraint) we don't want to split timestep at TOI.
+				ApplyImpl(Constraint, IterationParameters, ParticleParameters);
+				return;
+			}
+
+			// Rebuild iteration params with partial dt, and non-zero iteration count to force update of constraint.
+			// P may have changed due to other constraints, so at TOI our manifold needs updating.
+			const FReal PartialDT = Constraint.TimeOfImpact * IterationParameters.Dt;
+			const FReal RemainingDT = (1 - Constraint.TimeOfImpact) * IterationParameters.Dt;
+			const int32 FakeIteration = FMath::Max(IterationParameters.Iteration, 1); // Force Apply to update constraint, as other constraints could've changed P
+			const FContactIterationParameters IterationParametersPartialDT{ PartialDT, FakeIteration, IterationParameters.NumIterations, IterationParameters.NumPairIterations, IterationParameters.ApplyType, IterationParameters.NeedsAnotherIteration };
+			const FContactIterationParameters IterationParametersRemainingDT{ RemainingDT, FakeIteration, IterationParameters.NumIterations, IterationParameters.NumPairIterations, IterationParameters.ApplyType, IterationParameters.NeedsAnotherIteration };
+
+			// Rewind P to TOI and Apply
+			Particle0->P() = FMath::Lerp(Particle0->X(), Particle0->P(), Constraint.TimeOfImpact);
+			ApplyImpl(Constraint, IterationParametersPartialDT, ParticleParameters);
+
+			// Advance P to end of frame from TOI, and Apply
+			Particle0->P() = Particle0->P() + Particle0->V() * RemainingDT;
+			ApplyImpl(Constraint, IterationParametersRemainingDT, ParticleParameters);
+		}
+
 
 		void Apply(FCollisionConstraintBase& Constraint, const FContactIterationParameters & IterationParameters, const FContactParticleParameters & ParticleParameters)
 		{
@@ -383,6 +419,10 @@ namespace Chaos
 			{
 				ApplyImpl(*Constraint.As<FRigidBodyPointContactConstraint>(), IterationParameters, ParticleParameters);
 			}
+			else if (Constraint.GetType() == FCollisionConstraintBase::FType::SinglePointSwept)
+			{
+				ApplySweptImpl(*Constraint.As<FRigidBodySweptPointContactConstraint>(), IterationParameters, ParticleParameters);
+			}
 			else if (Constraint.GetType() == FCollisionConstraintBase::FType::MultiPoint)
 			{
 				ApplyImpl(*Constraint.As<FRigidBodyMultiPointContactConstraint>(), IterationParameters, ParticleParameters);
@@ -533,6 +573,10 @@ namespace Chaos
 			{
 				ApplyPushOutImpl<FRigidBodyPointContactConstraint>(*Constraint.As<FRigidBodyPointContactConstraint>(), IsTemporarilyStatic, IterationParameters, ParticleParameters);
 			}
+			else if (Constraint.GetType() == FCollisionConstraintBase::FType::SinglePointSwept)
+			{
+				ApplyPushOutImpl(*Constraint.As<FRigidBodySweptPointContactConstraint>(), IsTemporarilyStatic, IterationParameters, ParticleParameters);
+			}
 			else if (Constraint.GetType() == FCollisionConstraintBase::FType::MultiPoint)
 			{
 				ApplyPushOutImpl<FRigidBodyMultiPointContactConstraint>(*Constraint.As<FRigidBodyMultiPointContactConstraint>(), IsTemporarilyStatic, IterationParameters, ParticleParameters);
diff --git a/Engine/Source/Runtime/Experimental/Chaos/Public/Chaos/Collision/CollisionReceiver.h b/Engine/Source/Runtime/Experimental/Chaos/Public/Chaos/Collision/CollisionReceiver.h
--- a/Engine/Source/Runtime/Experimental/Chaos/Public/Chaos/Collision/CollisionReceiver.h
+++ b/Engine/Source/Runtime/Experimental/Chaos/Public/Chaos/Collision/CollisionReceiver.h
@@ -30,6 +30,11 @@ namespace Chaos
 			{
 				SinglePointQueue.Enqueue(Constraint);
 			}
+			for (const TRigidBodySweptPointContactConstraint<FReal, 3>& Constraint : Constraints.SinglePointSweptConstraints)
+			{
+				SingleSweptPointQueue.Enqueue(Constraint);
+			}
+
 			for (const TRigidBodyMultiPointContactConstraint<FReal, 3>& Constraint : Constraints.MultiPointConstraints)
 			{
 				MultiPointQueue.Enqueue(Constraint);
@@ -54,6 +59,15 @@ namespace Chaos
 					CollisionConstraints.AddConstraint(SinglePointConstraint);
 				}
 			}
+			TRigidBodySweptPointContactConstraint<FReal, 3> SinglePointSweptConstraint;
+			while (SingleSweptPointQueue.Dequeue(SinglePointSweptConstraint))
+			{
+				if (!CollisionConstraints.Contains(&SinglePointSweptConstraint))
+				{
+					CollisionConstraints.AddConstraint(SinglePointSweptConstraint);
+				}
+			}
+
 			TRigidBodyMultiPointContactConstraint<FReal, 3> MultiPointConstraint;
 			while (MultiPointQueue.Dequeue(MultiPointConstraint))
 			{
@@ -67,6 +81,7 @@ namespace Chaos
 	private:
 		//todo(ocohen): use per thread buffer instead, need better support than ParallelFor for this
 		TQueue<TRigidBodyPointContactConstraint<FReal, 3>, EQueueMode::Mpsc> SinglePointQueue;
+		TQueue<TRigidBodySweptPointContactConstraint<FReal, 3>, EQueueMode::Mpsc> SingleSweptPointQueue;
 		TQueue<TRigidBodyMultiPointContactConstraint<FReal, 3>, EQueueMode::Mpsc> MultiPointQueue;
 		FCollisionConstraints& CollisionConstraints;
 	};
diff --git a/Engine/Source/Runtime/Experimental/Chaos/Public/Chaos/CollisionResolutionTypes.h b/Engine/Source/Runtime/Experimental/Chaos/Public/Chaos/CollisionResolutionTypes.h
--- a/Engine/Source/Runtime/Experimental/Chaos/Public/Chaos/CollisionResolutionTypes.h
+++ b/Engine/Source/Runtime/Experimental/Chaos/Public/Chaos/CollisionResolutionTypes.h
@@ -115,6 +115,7 @@ namespace Chaos
 		{
 			None = 0,     // default value also indicates a invalid constraint
 			SinglePoint,  //   TRigidBodyPointContactConstraint
+			SinglePointSwept, // TRigidBodySweptPointContactConstraint
 			MultiPoint    //   TRigidBodyMultiPointContactConstraint
 		};
 
@@ -284,6 +285,33 @@ namespace Chaos
 	};
 	typedef TRigidBodyMultiPointContactConstraint<float, 3> FRigidBodyMultiPointContactConstraint;
 
+	/*
+	*
+	*/
+	template<class T, int d>
+	class TRigidBodySweptPointContactConstraint : public TRigidBodyPointContactConstraint<T,d>
+	{
+	public:
+		using Base = TRigidBodyPointContactConstraint<T, d>;
+		using FGeometryParticleHandle = TGeometryParticleHandle<T, d>;
+		using Base::Particle;
+
+		TRigidBodySweptPointContactConstraint() : Base(Base::FType::SinglePointSwept) {}
+		TRigidBodySweptPointContactConstraint(
+			FGeometryParticleHandle* Particle0, const FImplicitObject* Implicit0, const TRigidTransform<T, d>& Transform0,
+			FGeometryParticleHandle* Particle1, const FImplicitObject* Implicit1, const TRigidTransform<T, d>& Transform1,
+			EContactShapesType ShapesType)
+			: Base(Particle0, Implicit0, Transform0, Particle1, Implicit1, Transform1, Base::FType::SinglePointSwept, ShapesType)
+			, TimeOfImpact(0) {}
+
+		// Value in range [0,1] used to interpolate P between [X,P] that we will rollback to when solving at time of impact.
+		T TimeOfImpact;
+
+		static typename Base::FType StaticType() { return Base::FType::SinglePointSwept; };
+	};
+	typedef TRigidBodySweptPointContactConstraint<float, 3> FRigidBodySweptPointContactConstraint;
+
+
 
 	//
 	//
@@ -351,6 +379,11 @@ namespace Chaos
 			{
 				return ConstraintContainer->PointConstraints[ConstraintIndex];
 			}
+			else if (GetType() == FConstraintBase::FType::SinglePointSwept)
+			{
+				return ConstraintContainer->SweptPointConstraints[ConstraintIndex];
+			}
+
 			check(GetType() == FConstraintBase::FType::MultiPoint);
 			return ConstraintContainer->IterativeConstraints[ConstraintIndex];
 		}
@@ -361,6 +394,11 @@ namespace Chaos
 			{
 				return ConstraintContainer->PointConstraints[ConstraintIndex];
 			}
+			else if (GetType() == FConstraintBase::FType::SinglePointSwept)
+			{
+				return ConstraintContainer->SweptPointConstraints[ConstraintIndex];
+			}
+
 			check(GetType() == FConstraintBase::FType::MultiPoint);
 			return ConstraintContainer->IterativeConstraints[ConstraintIndex];
 		}
@@ -370,6 +408,9 @@ namespace Chaos
 		const TRigidBodyPointContactConstraint<T, d>& GetPointContact() const { check(GetType() == FConstraintBase::FType::SinglePoint); return ConstraintContainer->PointConstraints[ConstraintIndex]; }
 		TRigidBodyPointContactConstraint<T, d>& GetPointContact() { check(GetType() == FConstraintBase::FType::SinglePoint); return ConstraintContainer->PointConstraints[ConstraintIndex]; }
 
+		const TRigidBodySweptPointContactConstraint<T, d>& GetSweptPointContact() const { check(GetType() == FConstraintBase::FType::SinglePointSwept); return ConstraintContainer->SweptPointConstraints[ConstraintIndex]; }
+		TRigidBodySweptPointContactConstraint<T, d>& GetSweptPointContact() { check(GetType() == FConstraintBase::FType::SinglePointSwept); return ConstraintContainer->SweptPointConstraints[ConstraintIndex]; }
+
 		const TRigidBodyMultiPointContactConstraint<T, d>& GetMultiPointContact() const { check(GetType() == FConstraintBase::FType::MultiPoint); return ConstraintContainer->IterativeConstraints[ConstraintIndex]; }
 		TRigidBodyMultiPointContactConstraint<T, d>& GetMultiPointContact() { check(GetType() == FConstraintBase::FType::MultiPoint); return ConstraintContainer->IterativeConstraints[ConstraintIndex]; }
 
@@ -417,13 +458,15 @@ namespace Chaos
 		static const int32 InlineMaxConstraints = 8;
 
 		TArray<FRigidBodyPointContactConstraint, TInlineAllocator<InlineMaxConstraints>> SinglePointConstraints;
+		TArray<FRigidBodySweptPointContactConstraint, TInlineAllocator<InlineMaxConstraints>> SinglePointSweptConstraints;
 		TArray<FRigidBodyMultiPointContactConstraint, TInlineAllocator<InlineMaxConstraints>> MultiPointConstraints;
 
-		int32 Num() const { return SinglePointConstraints.Num() + MultiPointConstraints.Num(); }
+		int32 Num() const { return SinglePointConstraints.Num() + SinglePointSweptConstraints.Num() + MultiPointConstraints.Num(); }
 
 		void Empty()
 		{
 			SinglePointConstraints.Empty();
+			SinglePointSweptConstraints.Empty();
 			MultiPointConstraints.Empty();
 		}
 
@@ -437,6 +480,16 @@ namespace Chaos
 			return nullptr;
 		}
 
+		TRigidBodySweptPointContactConstraint<FReal, 3>* TryAdd(FReal MaxPhi, const TRigidBodySweptPointContactConstraint<FReal, 3>& C)
+		{
+			if (C.GetPhi() < MaxPhi)
+			{
+				int32 ConstraintIndex = SinglePointSweptConstraints.Add(C);
+				return &SinglePointSweptConstraints[ConstraintIndex];
+			}
+			return nullptr;
+		}
+
 		TRigidBodyMultiPointContactConstraint<FReal, 3>* TryAdd(FReal MaxPhi, const TRigidBodyMultiPointContactConstraint<FReal, 3>& C)
 		{
 			if (C.GetPhi() < MaxPhi)
diff --git a/Engine/Source/Runtime/Experimental/Chaos/Public/Chaos/PBDCollisionConstraints.h b/Engine/Source/Runtime/Experimental/Chaos/Public/Chaos/PBDCollisionConstraints.h
--- a/Engine/Source/Runtime/Experimental/Chaos/Public/Chaos/PBDCollisionConstraints.h
+++ b/Engine/Source/Runtime/Experimental/Chaos/Public/Chaos/PBDCollisionConstraints.h
@@ -77,6 +77,7 @@ public:
 	using FConstraintContainerHandle = TPBDCollisionConstraintHandle<T, d>;
 	using FConstraintBase = TCollisionConstraintBase<T, d>;
  	using FPointContactConstraint = TRigidBodyPointContactConstraint<T, d>;
+ 	using FSweptPointContactConstraint = TRigidBodySweptPointContactConstraint<T, d>;
 	using FMultiPointContactConstraint = TRigidBodyMultiPointContactConstraint<T, d>;
 	using FConstraintHandleAllocator = TConstraintHandleAllocator<TPBDCollisionConstraints<T, d>>;
 	using FConstraintContainerHandleKey = typename TPBDCollisionConstraintHandle<T, d>::FHandleKey;
@@ -117,6 +118,7 @@ public:
 	*  then delete the InConstraint. 
 	*/
 	void AddConstraint(const TRigidBodyPointContactConstraint<FReal, 3>& InConstraint);
+	void AddConstraint(const TRigidBodySweptPointContactConstraint<FReal, 3>& InConstraint);
 	void AddConstraint(const TRigidBodyMultiPointContactConstraint<FReal, 3>& InConstraint);
 
 	/**
@@ -269,7 +271,7 @@ public:
 
 	int32 NumConstraints() const
 	{
-		return PointConstraints.Num() + IterativeConstraints.Num();
+		return PointConstraints.Num() + SweptPointConstraints.Num() + IterativeConstraints.Num();
 	}
 
 	FHandles& GetConstraintHandles()
@@ -306,6 +308,7 @@ private:
 	const TPBDRigidsSOAs<T,d>& Particles;
 
 	TArray<FPointContactConstraint> PointConstraints;
+	TArray<FSweptPointContactConstraint> SweptPointConstraints;
 	TArray<FMultiPointContactConstraint> IterativeConstraints;
 
 #if CHAOS_COLLISION_PERSISTENCE_ENABLED
