fix render bounds for geometry collection
from //UE4/Private-Reverb

diff --git a/Engine/Source/Runtime/Experimental/GeometryCollectionEngine/Private/GeometryCollection/GeometryCollectionComponent.cpp b/Engine/Source/Runtime/Experimental/GeometryCollectionEngine/Private/GeometryCollection/GeometryCollectionComponent.cpp
--- a/Engine/Source/Runtime/Experimental/GeometryCollectionEngine/Private/GeometryCollection/GeometryCollectionComponent.cpp
+++ b/Engine/Source/Runtime/Experimental/GeometryCollectionEngine/Private/GeometryCollection/GeometryCollectionComponent.cpp
@@ -262,9 +262,16 @@ FBoxSphereBounds UGeometryCollectionComponent::CalcBounds(const FTransform& Loca
 {	
 	SCOPE_CYCLE_COUNTER(STAT_GCCUpdateBounds);
 
-	// Don't use bounds calculated in the physics object if we are doing sequencer cache playback
-	// because we are overriding the transforms in the GC
-	if (!CachePlayback && WorldBounds.GetSphere().W > 1e-5)
+	// #todo(dmp): hack to make bounds calculation work when we don't have valid physics proxy data.  This will
+	// force bounds calculation.
+	FChaosSolversModule* Module = FChaosSolversModule::GetModule();
+	Module->LockResultsRead();
+
+	const FGeometryCollectionResults* Results = PhysicsProxy ? PhysicsProxy->GetConsumerResultsGT() : nullptr;
+
+	const int32 NumTransforms = Results ? Results->GlobalTransforms.Num() : 0;
+
+	if (!CachePlayback && WorldBounds.GetSphere().W > 1e-5 && NumTransforms > 0)
 	{
 		return WorldBounds;
 	} 
