Fix CIS - TriMeshType shadows template arg.

diff --git a/Engine/Source/Runtime/Experimental/Chaos/Private/Chaos/CollisionResolution.cpp b/Engine/Source/Runtime/Experimental/Chaos/Private/Chaos/CollisionResolution.cpp
--- a/Engine/Source/Runtime/Experimental/Chaos/Private/Chaos/CollisionResolution.cpp
+++ b/Engine/Source/Runtime/Experimental/Chaos/Private/Chaos/CollisionResolution.cpp
@@ -1436,8 +1436,6 @@ namespace Chaos
 		template <typename TriMeshType, typename T, int d>
 		TContactPoint<T> CapsuleTriangleMeshSweptContactPoint(const TCapsule<T>& A, const TRigidTransform<T, d>& ATransform, const TriMeshType& B, const TRigidTransform<T, d>& BStartTransform, const TVector<T, d>& Dir, const T Length, const T CullDistance, T& TOI)
 		{
-			EImplicitObjectType TriMeshType = B.GetType();
-
 			if (const TImplicitObjectScaled<FTriangleMeshImplicitObject>* ScaledTriangleMesh = B.template GetObject<const TImplicitObjectScaled<FTriangleMeshImplicitObject>>())
 			{
 				return GJKImplicitScaledTriMeshSweptContactPoint<TCapsule<T>>(A, ATransform, *ScaledTriangleMesh, BStartTransform, Dir, Length, CullDistance, TOI);
@@ -1708,7 +1706,6 @@ namespace Chaos
 		template <typename TriMeshType, typename T, int d>
 		TContactPoint<T> ConvexTriangleMeshSweptContactPoint(const FImplicitObject& A, const TRigidTransform<T, d>& ATransform, const TriMeshType& B, const TRigidTransform<T, d>& BStartTransform, const TVector<T, d>& Dir, const T Length, const T CullDistance, T& TOI)
 		{
-			EImplicitObjectType TriMeshType = B.GetType();
 			if (const TImplicitObjectScaled<FTriangleMeshImplicitObject>* ScaledTriangleMesh = B.template GetObject<const TImplicitObjectScaled<FTriangleMeshImplicitObject>>())
 			{
 				return GJKImplicitScaledTriMeshSweptContactPoint<FConvex>(A, ATransform, *ScaledTriangleMesh, BStartTransform, Dir, Length, CullDistance, TOI);
@@ -2493,7 +2490,6 @@ namespace Chaos
 			FReal TOI = 0.0f;
 			FReal CullDistance = 0.0f;
 
-			EImplicitObjectType TriMeshType = B->GetType();
 			if (const TImplicitObjectScaled<FTriangleMeshImplicitObject>* ScaledTriangleMesh = B->template GetObject<const TImplicitObjectScaled<FTriangleMeshImplicitObject>>())
 			{
 				Contact = GJKImplicitScaledTriMeshSweptContactPoint<GeometryA>(A, AStartTransform, *ScaledTriangleMesh, BTransform, Dir, Length, CullDistance, TOI);
