Fix Mac template warnings

diff --git a/Engine/Source/Runtime/Experimental/Chaos/Public/Chaos/PBDConstraintRule.h b/Engine/Source/Runtime/Experimental/Chaos/Public/Chaos/PBDConstraintRule.h
--- a/Engine/Source/Runtime/Experimental/Chaos/Public/Chaos/PBDConstraintRule.h
+++ b/Engine/Source/Runtime/Experimental/Chaos/Public/Chaos/PBDConstraintRule.h
@@ -408,4 +408,12 @@ namespace Chaos
 		FPBDConstraintColor GraphColor;
 	};
 
-}
\ No newline at end of file
+}
+
+// Only way to make this compile at the moment due to visibility attribute issues. TODO: Change this once a fix for this problem is applied.
+#include "Chaos/PBDCollisionConstraints.h"
+#if PLATFORM_MAC || PLATFORM_LINUX
+extern template class CHAOS_API Chaos::TPBDConstraintColorRule<Chaos::TPBDCollisionConstraints<float, 3>>;
+#else
+template class Chaos::TPBDConstraintColorRule<Chaos::TPBDCollisionConstraints<float, 3>>;
+#endif
