Fix CIS - monolithics dupe template instantiation in PBDConstraintRule.cpp

diff --git a/Engine/Source/Runtime/Experimental/Chaos/Private/Chaos/PBDConstraintRule.cpp b/Engine/Source/Runtime/Experimental/Chaos/Private/Chaos/PBDConstraintRule.cpp
--- a/Engine/Source/Runtime/Experimental/Chaos/Private/Chaos/PBDConstraintRule.cpp
+++ b/Engine/Source/Runtime/Experimental/Chaos/Private/Chaos/PBDConstraintRule.cpp
@@ -62,7 +62,7 @@ namespace Chaos
 	template class TPBDConstraintGraphRuleImpl<TPBDRigidDynamicSpringConstraints<float, 3>>;
 	template class TPBDConstraintGraphRuleImpl<FPBDRigidSpringConstraints>;
 
-	template class TPBDConstraintColorRule<TPBDCollisionConstraints<float, 3>>;
+	//template class TPBDConstraintColorRule<TPBDCollisionConstraints<float, 3>>;
 	template class TPBDConstraintIslandRule<FPBDJointConstraints>;
 	template class TPBDConstraintIslandRule<TPBDPositionConstraints<float, 3>>;
 	template class TPBDConstraintIslandRule<TPBDRigidDynamicSpringConstraints<float, 3>>;
