CIS Fix - missing symbol in monolithics

diff --git a/Engine/Source/Runtime/Experimental/Chaos/Private/Chaos/PBDConstraintRule.cpp b/Engine/Source/Runtime/Experimental/Chaos/Private/Chaos/PBDConstraintRule.cpp
--- a/Engine/Source/Runtime/Experimental/Chaos/Private/Chaos/PBDConstraintRule.cpp
+++ b/Engine/Source/Runtime/Experimental/Chaos/Private/Chaos/PBDConstraintRule.cpp
@@ -62,7 +62,7 @@ namespace Chaos
 	template class TPBDConstraintGraphRuleImpl<TPBDRigidDynamicSpringConstraints<float, 3>>;
 	template class TPBDConstraintGraphRuleImpl<FPBDRigidSpringConstraints>;
 
-	//template class TPBDConstraintColorRule<TPBDCollisionConstraints<float, 3>>;
+	template class TPBDConstraintColorRule<TPBDCollisionConstraints<float, 3>>;
 	template class TPBDConstraintIslandRule<FPBDJointConstraints>;
 	template class TPBDConstraintIslandRule<TPBDPositionConstraints<float, 3>>;
 	template class TPBDConstraintIslandRule<TPBDRigidDynamicSpringConstraints<float, 3>>;
diff --git a/Engine/Source/Runtime/Experimental/Chaos/Public/Chaos/PBDConstraintRule.h b/Engine/Source/Runtime/Experimental/Chaos/Public/Chaos/PBDConstraintRule.h
--- a/Engine/Source/Runtime/Experimental/Chaos/Public/Chaos/PBDConstraintRule.h
+++ b/Engine/Source/Runtime/Experimental/Chaos/Public/Chaos/PBDConstraintRule.h
@@ -414,6 +414,4 @@ namespace Chaos
 #include "Chaos/PBDCollisionConstraints.h"
 #if PLATFORM_MAC || PLATFORM_LINUX
 extern template class CHAOS_API Chaos::TPBDConstraintColorRule<Chaos::TPBDCollisionConstraints<float, 3>>;
-#else
-template class Chaos::TPBDConstraintColorRule<Chaos::TPBDCollisionConstraints<float, 3>>;
 #endif
