Fix include after code block warning

diff --git a/Engine/Source/Runtime/Experimental/Chaos/Public/Chaos/PBDConstraintRule.h b/Engine/Source/Runtime/Experimental/Chaos/Public/Chaos/PBDConstraintRule.h
--- a/Engine/Source/Runtime/Experimental/Chaos/Public/Chaos/PBDConstraintRule.h
+++ b/Engine/Source/Runtime/Experimental/Chaos/Public/Chaos/PBDConstraintRule.h
@@ -7,6 +7,10 @@
 #include "Chaos/PBDConstraintColor.h"
 #include "Chaos/PBDConstraintGraph.h"
 
+// Only used for Mac/Linux forward decl at the bottom - should be removed as per the comment there (we get a warning if the include is also at the bottom)
+#include "Chaos/PBDCollisionConstraints.h"
+
+
 #define USE_SHOCK_PROPOGATION 1
 
 namespace Chaos
@@ -411,7 +415,6 @@ namespace Chaos
 }
 
 // Only way to make this compile at the moment due to visibility attribute issues. TODO: Change this once a fix for this problem is applied.
-#include "Chaos/PBDCollisionConstraints.h"
 #if PLATFORM_MAC || PLATFORM_LINUX
 extern template class CHAOS_API Chaos::TPBDConstraintColorRule<Chaos::TPBDCollisionConstraints<float, 3>>;
 #endif
