More work on particle properties, also fix static analysis

diff --git a/Engine/Source/Programs/HeadlessChaos/Private/HeadlessChaosTestDirtyParticle.cpp b/Engine/Source/Programs/HeadlessChaos/Private/HeadlessChaosTestDirtyParticle.cpp
--- a/Engine/Source/Programs/HeadlessChaos/Private/HeadlessChaosTestDirtyParticle.cpp
+++ b/Engine/Source/Programs/HeadlessChaos/Private/HeadlessChaosTestDirtyParticle.cpp
@@ -17,7 +17,116 @@ namespace ChaosTest
 
 	GTEST_TEST(DirtyParticleTests,Basic)
 	{
-		FDirtyPropertiesManager DirtyPropertiesPool;
+#if 0
+		struct FSimpleParticleRemote
+		{
+			FSimpleParticleRemote() : Manager(nullptr){}
+
+			FDirtyPropertiesManager* Manager;
+			TRemoteParticleProperty<FVec3> X;
+			TRemoteParticleProperty<FReal> InvM;
+			TRemoteParticleProperty<TSharedPtr<FImplicitObject,ESPMode::ThreadSafe>> Geometry;
+
+			const FVec3& GetX(FDirtyPropertiesManager& Manager) const { return X.Read(*Manager); }
+			const FReal GetInvM(FDirtyPropertiesManager& Manager) const { return InvM.Read(*Manager); }
+			const TSharedPtr<FImplicitObject,ESPMode::ThreadSafe>& GetGeometry(FDirtyPropertiesManager& Manager) const { return Geometry.Read(*Manager); }
+			
+			void Clear()
+			{
+				if(Manager)
+				{
+					X.Clear(*Manager);
+					InvM.Clear(*Manager);
+					Geometry.Clear(*Manager);
+				}
+			}
+
+			~FSimpleParticleRemote()
+			{
+				Clear();
+			}
+		};
+
+		struct FSimpleParticle
+		{
+			FSimpleParticle() : Manager(nullptr){}
+			const FVec3& GetX() const{ return X.Read(); }
+			void SetX(const FVec3& InX){ X.Write(InX,Remote.X,ToFlush,Remote.Dirty, Remote.Manager); }
+
+			const FReal GetInvM() const { return InvM.Read(); }
+			void SetInvM(FReal InInvM){ InvM.Write(InInvM,Remote.InvM, ToFlush, Remote.Dirty, Remote.Manager); }
+
+			const TSharedPtr<FImplicitObject,ESPMode::ThreadSafe>& GetGeometry() const { return Geometry.Read(); }
+			void SetGeometry(const TSharedPtr<FImplicitObject,ESPMode::ThreadSafe>& InGeom){ Geometry.Write(InGeom, Remote.Geometry, ToFlush, Remote.Dirty, Remote.Manager); }
+
+			void Flush()
+			{
+				ensure(Remote.Manager);
+				if(ToFlush.IsDirty())
+				{
+					if(ToFlush.IsDirty(EParticleFlags::X))
+					{
+						X.Flush(Remote.X,ToFlush,*Remote.Manager);
+					}
+
+					if(ToFlush.IsDirty(EParticleFlags::InvM))
+					{
+						InvM.Flush(Remote.InvM,ToFlush, *Remote.Manager);
+					}
+
+					if(ToFlush.IsDirty(EParticleFlags::Geometry))
+					{
+						Geometry.Flush(Remote.Geometry,ToFlush,*Remote.Manager);
+					}
+				}
+
+				ensure(ToFlush.IsClean());
+			}
+
+			bool IsDirty() const { return Remote.Dirty.IsDirty(); }
+
+			void StealRemote(FSimpleParticleRemote& Dst)
+			{
+				Flush();
+				Swap(Dst,Remote);
+			}
+
+			TParticleProperty<FVec3,EParticleFlags::X> X;
+			TParticleProperty<FReal,EParticleFlags::InvM> InvM;
+			TParticleProperty<TSharedPtr<FImplicitObject, ESPMode::ThreadSafe>,EParticleFlags::Geometry> Geometry;
+			FParticleDirtyFlags ToFlush;
+
+			FSimpleParticleRemote Remote;
+		};
+
+		FSimpleParticle P1;
+
+		P1.SetX(FVec3(1,1,1));
+		P1.SetInvM(3);
+
+		EXPECT_EQ(P1.GetX(),FVec3(1,1,1));
+		EXPECT_EQ(P1.GetInvM(),3);
+
+		EXPECT_TRUE(P1.ToFlush.IsDirty(EParticleFlags::X));
+		EXPECT_TRUE(P1.ToFlush.IsDirty(EParticleFlags::InvM));
+		EXPECT_FALSE(P1.ToFlush.IsDirty(EParticleFlags::Geometry));
+
+		FDirtyPropertiesManager Manager;
+		P1.Remote.Manager = &Manager;
+
+		FSimpleParticleRemote Remote;
+		P1.StealRemote(Remote);
+
+		EXPECT_TRUE(Remote.Dirty.IsDirty(EParticleFlags::X));
+		EXPECT_TRUE(Remote.Dirty.IsDirty(EParticleFlags::X));
+		EXPECT_FALSE(Remote.Dirty.IsDirty(EParticleFlags::Geometry));
+
+		EXPECT_EQ(Remote.GetX(Manager),FVec3(1,1,1));
+		EXPECT_EQ(Remote.GetInvM(Manager),3);
+
+		Remote.Clear();
+
+		/*
 		FDirtyProperties DirtyParticle1;
 		DirtyParticle1.WriteX(DirtyPropertiesPool, FVec3(1,1,1));
 		DirtyParticle1.WriteInvM(DirtyPropertiesPool, 3);
@@ -59,7 +168,8 @@ namespace ChaosTest
 
 		//If we haven't popped everything we call this
 		DirtyParticle1.Clean(DirtyPropertiesPool);
-		DirtyParticle2.Clean(DirtyPropertiesPool);
+		DirtyParticle2.Clean(DirtyPropertiesPool);*/
+#endif
 	}
 }
 
diff --git a/Engine/Source/Runtime/Experimental/Chaos/Public/Chaos/ParticleDirtyFlags.h b/Engine/Source/Runtime/Experimental/Chaos/Public/Chaos/ParticleDirtyFlags.h
--- a/Engine/Source/Runtime/Experimental/Chaos/Public/Chaos/ParticleDirtyFlags.h
+++ b/Engine/Source/Runtime/Experimental/Chaos/Public/Chaos/ParticleDirtyFlags.h
@@ -2,6 +2,10 @@
 
 #pragma once
 
+#include "Chaos/Core.h"
+
+class FName;
+
 namespace Chaos
 {
 
@@ -10,11 +14,35 @@ using ChaosGeometrySp = TSharedPtr<FImplicitObject,ESPMode::ThreadSafe>;
 #define PARTICLE_PROPERTY(PropName, Type) PropName,
 	enum EParticleProperty : uint32
 	{
-#include "ParticleProperties.h"
+#include "ParticleProperties.inl"
 		NumProperties
 	};
 
 #undef PARTICLE_PROPERTY
+
+#define PARTICLE_PROPERTY_TYPE(TypeName, Type) TypeName,
+	enum EParticlePropertyType: uint32
+	{
+#include "ParticlePropertiesTypes.inl"
+		NumTypes
+	};
+
+#undef PARTICLE_PROPERTY_TYPE
+
+template <typename T>
+struct TParticlePropertyTrait
+{
+};
+
+#define PARTICLE_PROPERTY_TYPE(TypeName, Type) \
+template <>\
+struct TParticlePropertyTrait<Type>\
+{\
+	static constexpr uint32 PoolIdx = (uint32)EParticlePropertyType::TypeName;\
+};
+
+#include "ParticlePropertiesTypes.inl"
+#undef PARTICLE_PROPERTY_TYPE
 	
 	// There is a dirty flag for every user-settable particle property.
 	// Dirty property values will get copied from game to physics thread buffers,
@@ -23,7 +51,7 @@ using ChaosGeometrySp = TSharedPtr<FImplicitObject,ESPMode::ThreadSafe>;
 
 	enum class EParticleFlags : uint32
 	{
-		#include "ParticleProperties.h"
+		#include "ParticleProperties.inl"
 		DummyFlag
 	};
 #undef PARTICLE_PROPERTY
@@ -75,17 +103,14 @@ using ChaosGeometrySp = TSharedPtr<FImplicitObject,ESPMode::ThreadSafe>;
 
 	struct FDirtyIdx
 	{
-		uint32 Entry;
-
-		bool operator==(const FDirtyIdx& Rhs) const
-		{
-			return Entry == Rhs.Entry;
-		}
+		uint32 bHasEntry : 1;
+		uint32 Entry : 31;
 	};
 
 	template <typename T>
 	class TDirtyElementPool
 	{
+		static_assert(sizeof(TParticlePropertyTrait<T>::PoolIdx),"Property type must be registed. Is it in ParticlePropertiesTypes?");
 	public:
 		const T& Read(int32 Idx) const
 		{
@@ -143,89 +168,133 @@ using ChaosGeometrySp = TSharedPtr<FImplicitObject,ESPMode::ThreadSafe>;
 	class FDirtyPropertiesManager
 	{
 	public:
-#define PARTICLE_PROPERTY(PropName, Type) \
-const TDirtyElementPool<Type>& Get##PropName##Pool() const { return Dirty##PropName##Pool; }\
-TDirtyElementPool<Type>& Get##PropName##Pool() { return Dirty##PropName##Pool; }
+		template <typename T>
+		TDirtyElementPool<T>& GetPool()
+		{
+			switch(TParticlePropertyTrait<T>::PoolIdx)
+			{
+#define PARTICLE_PROPERTY_TYPE(TypeName, Type) case EParticlePropertyType::TypeName: return (TDirtyElementPool<T>&)TypeName##Pool;
+#include "ParticlePropertiesTypes.inl"
+#undef PARTICLE_PROPERTY_TYPE
+			default: check(false);
+			}
+			
+			static TDirtyElementPool<T> ErrorPool;
+			return ErrorPool;
+		}
+
+		template <typename T>
+		const TDirtyElementPool<T>& GetPool() const
+		{
+			switch(TParticlePropertyTrait<T>::PoolIdx)
+			{
+#define PARTICLE_PROPERTY_TYPE(TypeName, Type) case EParticlePropertyType::TypeName: return (const TDirtyElementPool<T>&)TypeName##Pool;
+#include "ParticlePropertiesTypes.inl"
+#undef PARTICLE_PROPERTY_TYPE
+			default: check(false);
+			}
+
+			static TDirtyElementPool<T> ErrorPool;
+			return ErrorPool;
+		}
 
-#include "ParticleProperties.h"
-#undef PARTICLE_PROPERTY
 	private:
 
-		//creates a TDirtyElementPool for each property
-#define PARTICLE_PROPERTY(Property, Type) TDirtyElementPool<Type> Dirty##Property##Pool;
-#include "ParticleProperties.h"
-#undef PARTICLE_PROPERTY
+#define PARTICLE_PROPERTY_TYPE(TypeName, Type) TDirtyElementPool<Type> TypeName##Pool;
+#include "ParticlePropertiesTypes.inl"
+#undef PARTICLE_PROPERTY_TYPE
 	};
 
-	class FDirtyProperties
+
+template <typename T>
+class TRemoteParticleProperty
+{
+public:
+	TRemoteParticleProperty()
 	{
-	public:
-		bool IsDirty(const EParticleFlags CheckBits) const
-		{
-			return Flags.IsDirty(CheckBits);
-		}
+		Idx.bHasEntry = false;
+	}
 
-		bool IsDirty(const int32 CheckBits) const
+	TRemoteParticleProperty(const TRemoteParticleProperty<T>& Rhs) = delete;
+	TRemoteParticleProperty(TRemoteParticleProperty<T>&& Rhs)
+	: Idx(Rhs.Idx)
+	{
+		Rhs.bHasEntry = false;
+	}
+
+	~TRemoteParticleProperty()
+	{
+		ensure(!Idx.bHasEntry);	//leaking, make sure to call Pop
+	}
+
+	const T& Read(const FDirtyPropertiesManager& Manager) const
+	{
+		ensure(Idx.bHasEntry);
+		return Manager.GetPool<T>().Read(Idx.Entry);
+	}
+
+	T Pop(FDirtyPropertiesManager& Manager)
+	{
+		ensure(Idx.bHasEntry);
+		Idx.bHasEntry = false;
+		return Manager.GetPool<T>().Pop(Idx.Entry);
+	}
+
+	void Clear(FDirtyPropertiesManager& Manager)
+	{
+		if(Idx.bHasEntry)
 		{
-			return Flags.IsDirty(CheckBits);
+			Idx.bHasEntry = false;
+			Manager.GetPool<T>().Pop(Idx.Entry);
 		}
+	}
 
-		~FDirtyProperties()
+	void Write(FDirtyPropertiesManager& Manager, const T& Val)
+	{
+		if(Idx.bHasEntry)
 		{
-			//Make sure to call Clean before killing this guy
-			ensure(Flags.IsClean());
+			Manager.GetPool<T>().Update(Idx.Entry,Val);
 		}
-
-		void Clean(FDirtyPropertiesManager& DirtyPropertyManager)
+		else
 		{
-			if(!Flags.IsClean())
-			{
-#define PARTICLE_PROPERTY(PropName, Type) if(Flags.IsDirty(EParticleFlags::PropName)){ Pop##PropName(DirtyPropertyManager); }
-#include "ParticleProperties.h"
-#undef PARTICLE_PROPERTY
-			}
+			Idx.Entry = Manager.GetPool<T>().Write(Val);
+			Idx.bHasEntry = true;
 		}
+	}
+private:
+	FDirtyIdx Idx;
 
-#define PARTICLE_PROPERTY(PropName, Type)\
-Type const& Read##PropName(const FDirtyPropertiesManager& DirtyPropertiesManager) const { return ReadImp(PropName, DirtyPropertiesManager.Get##PropName##Pool()); }\
-Type Pop##PropName(FDirtyPropertiesManager& DirtyPropertiesManager) { return PopImp(PropName, EParticleFlags::PropName, DirtyPropertiesManager.Get##PropName##Pool()); }\
-void Write##PropName(FDirtyPropertiesManager& DirtyPropertiesManager, Type const& Val) { WriteImp(PropName, EParticleFlags::PropName, Val, DirtyPropertiesManager.Get##PropName##Pool()); }
-
-#include "ParticleProperties.h"
-#undef PARTICLE_PROPERTY
-
-	private:
-		FParticleDirtyFlags Flags;
-#define PARTICLE_PROPERTY(PropName, Type) FDirtyIdx PropName;
-#include "ParticleProperties.h"
-#undef PARTICLE_PROPERTY
+	TRemoteParticleProperty<T>& operator=(const TRemoteParticleProperty<T>& Rhs){}
+};
 
-		template <typename T>
-		void WriteImp(FDirtyIdx& Prop, EParticleFlags Flag, const T& Val, TDirtyElementPool<T>& Pool)
+template <typename T, EParticleFlags Flag>
+class TParticleProperty
+{
+public:
+	const T& Read() const { return Property; }
+	void Write(const T& Val, TRemoteParticleProperty<T>& Remote, FParticleDirtyFlags& PendingFlush, FParticleDirtyFlags& Dirty, FDirtyPropertiesManager* Manager)
+	{
+		Property = Val;
+		if(Manager)
 		{
-			if(Flags.IsDirty(Flag))
-			{
-				Pool.Update(Prop.Entry,Val);
-			}
-			else
-			{
-				Prop.Entry = Pool.Write(Val);
-				Flags.MarkDirty(Flag);
-			}
+			Remote.Write(*Manager, Val);
+			PendingFlush.MarkClean(Flag);
 		}
-
-		template <typename T>
-		const T& ReadImp(const FDirtyIdx Prop, const TDirtyElementPool<T>& Pool) const
+		else
 		{
-			return Pool.Read(Prop.Entry);
+			PendingFlush.MarkDirty(Flag);
 		}
 
-		template <typename T>
-		T PopImp(const FDirtyIdx Prop, EParticleFlags Flag, TDirtyElementPool<T>& Pool)
-		{
-			ensure(Flags.IsDirty(Flag));
-			Flags.MarkClean(Flag);
-			return Pool.Pop(Prop.Entry);
-		}
-	};
+		Dirty.MarkDirty(Flag);
+	}
+
+	void Flush(TRemoteParticleProperty<T>& Remote,FParticleDirtyFlags& PendingFlush,FDirtyPropertiesManager& Manager)
+	{
+		Remote.Write(Manager, Property);
+		PendingFlush.MarkClean(Flag);
+	}
+
+private:
+	T Property;
+};
 }
diff --git a/Engine/Source/Runtime/Experimental/Chaos/Public/Chaos/ParticleProperties.h b/Engine/Source/Runtime/Experimental/Chaos/Public/Chaos/ParticleProperties.h
deleted file mode 100644
diff --git a/Engine/Source/Runtime/Experimental/Chaos/Public/Chaos/ParticleProperties.inl b/Engine/Source/Runtime/Experimental/Chaos/Public/Chaos/ParticleProperties.inl
new file mode 100644
--- /dev/null
+++ b/Engine/Source/Runtime/Experimental/Chaos/Public/Chaos/ParticleProperties.inl
@@ -0,0 +1,41 @@
+// Copyright Epic Games, Inc. All Rights Reserved.
+
+#undef PARTICLE_PROPERTY_CHECKED
+#if CHAOS_CHECKED
+#define PARTICLE_PROPERTY_CHECKED(x, Type) PARTICLE_PROPERTY(x, Type)
+#else
+#define PARTICLE_PROPERTY_CHECKED(x, Type)
+#endif
+
+PARTICLE_PROPERTY(X,FVec3)
+PARTICLE_PROPERTY(R,FRotation3)
+PARTICLE_PROPERTY(V,FVec3)
+PARTICLE_PROPERTY(W,FVec3)
+PARTICLE_PROPERTY(CenterOfMass,FVec3)
+PARTICLE_PROPERTY(RotationOfMass,FRotation3)
+PARTICLE_PROPERTY(CollisionGroup,int32)
+PARTICLE_PROPERTY(Disabled,bool)
+PARTICLE_PROPERTY(PreV,FVec3)
+PARTICLE_PROPERTY(PreW,FVec3)
+PARTICLE_PROPERTY(P,FVec3)
+PARTICLE_PROPERTY(Q,FRotation3)
+PARTICLE_PROPERTY(F,FVec3)
+PARTICLE_PROPERTY(Torque,FVec3)
+PARTICLE_PROPERTY(I,FMatrix33)
+PARTICLE_PROPERTY(InvI,FMatrix33)
+PARTICLE_PROPERTY(M,FReal)
+PARTICLE_PROPERTY(InvM,FReal)
+PARTICLE_PROPERTY(LinearEtherDrag,FReal)
+PARTICLE_PROPERTY(AngularEtherDrag,FReal)
+PARTICLE_PROPERTY(ObjectState,EObjectStateType)
+PARTICLE_PROPERTY(Geometry, ChaosGeometrySp)
+PARTICLE_PROPERTY(LinearImpulse, FVec3)
+PARTICLE_PROPERTY(AngularImpulse, FVec3)
+PARTICLE_PROPERTY(GravityEnabled, bool)
+PARTICLE_PROPERTY(SpatialIdx, FSpatialAccelerationIdx)
+PARTICLE_PROPERTY(UniqueIdx, FUniqueIdx)
+PARTICLE_PROPERTY(ShapeDisableCollision, int32)
+PARTICLE_PROPERTY(CollisionTraceType, EChaosCollisionTraceFlag)
+PARTICLE_PROPERTY(ShapeSimData, TArray<FCollisionFilterData>)
+PARTICLE_PROPERTY(UserData, void*)
+PARTICLE_PROPERTY_CHECKED(DebugName, FName)
\ No newline at end of file
diff --git a/Engine/Source/Runtime/Experimental/Chaos/Public/Chaos/ParticlePropertiesTypes.inl b/Engine/Source/Runtime/Experimental/Chaos/Public/Chaos/ParticlePropertiesTypes.inl
new file mode 100644
--- /dev/null
+++ b/Engine/Source/Runtime/Experimental/Chaos/Public/Chaos/ParticlePropertiesTypes.inl
@@ -0,0 +1,16 @@
+// Copyright Epic Games, Inc. All Rights Reserved.
+
+PARTICLE_PROPERTY_TYPE(Vec3_type, FVec3)
+PARTICLE_PROPERTY_TYPE(Rotation3_type, FRotation3)
+PARTICLE_PROPERTY_TYPE(int32_type, int32)
+PARTICLE_PROPERTY_TYPE(bool_type, bool)
+PARTICLE_PROPERTY_TYPE(Matrix33_type, FMatrix33)
+PARTICLE_PROPERTY_TYPE(Real_type, FReal)
+PARTICLE_PROPERTY_TYPE(ObjectStateType_type, EObjectStateType)
+PARTICLE_PROPERTY_TYPE(ChaosGeometrySp_type, ChaosGeometrySp)
+PARTICLE_PROPERTY_TYPE(SpatialAccelerationIdx_type, FSpatialAccelerationIdx)
+PARTICLE_PROPERTY_TYPE(UniqueIdx_type, FUniqueIdx)
+PARTICLE_PROPERTY_TYPE(ChaosCollisionTraceFlag_type, EChaosCollisionTraceFlag)
+PARTICLE_PROPERTY_TYPE(TArray_FCollisionFilterData, TArray<FCollisionFilterData>)
+PARTICLE_PROPERTY_TYPE(VoidPtr_type, void*)
+PARTICLE_PROPERTY_TYPE(Name_type, FName)
\ No newline at end of file
