Fix compilation of DirtyFlags.h when there are 32 particle properties (ie, when CHAOS_CHECKED is true)

diff --git a/Engine/Source/Runtime/Experimental/Chaos/Public/Chaos/ParticleDirtyFlags.h b/Engine/Source/Runtime/Experimental/Chaos/Public/Chaos/ParticleDirtyFlags.h
--- a/Engine/Source/Runtime/Experimental/Chaos/Public/Chaos/ParticleDirtyFlags.h
+++ b/Engine/Source/Runtime/Experimental/Chaos/Public/Chaos/ParticleDirtyFlags.h
@@ -47,7 +47,7 @@ struct TParticlePropertyTrait<Type>\
 	// There is a dirty flag for every user-settable particle property.
 	// Dirty property values will get copied from game to physics thread buffers,
 	// but clean property values will get overridden with physics thread results.
-#define PARTICLE_PROPERTY(PropName, Type) PropName = 1 << (uint32)EParticleProperty::PropName,
+#define PARTICLE_PROPERTY(PropName, Type) PropName = (uint32)1 << (uint32)EParticleProperty::PropName,
 
 	enum class EParticleFlags : uint32
 	{
