Move single physics proxy into Chaos module

diff --git a/Engine/Source/Programs/HeadlessChaos/Private/HeadlessChaosTestSolverCommandList.cpp b/Engine/Source/Programs/HeadlessChaos/Private/HeadlessChaosTestSolverCommandList.cpp
--- a/Engine/Source/Programs/HeadlessChaos/Private/HeadlessChaosTestSolverCommandList.cpp
+++ b/Engine/Source/Programs/HeadlessChaos/Private/HeadlessChaosTestSolverCommandList.cpp
@@ -6,7 +6,7 @@
 #include "Chaos/ParticleHandle.h"
 #include "Chaos/ErrorReporter.h"
 #include "Chaos/Sphere.h"
-#include "PhysicsProxy/SingleParticlePhysicsProxy.h"
+#include "Chaos/SingleParticlePhysicsProxy.h"
 #include "Chaos/Utilities.h"
 #include "PBDRigidsSolver.h"
 #include "ChaosSolversModule.h"
diff --git a/Engine/Source/Programs/HeadlessChaos/Private/HeadlessChaosTestSolverProxies.cpp b/Engine/Source/Programs/HeadlessChaos/Private/HeadlessChaosTestSolverProxies.cpp
--- a/Engine/Source/Programs/HeadlessChaos/Private/HeadlessChaosTestSolverProxies.cpp
+++ b/Engine/Source/Programs/HeadlessChaos/Private/HeadlessChaosTestSolverProxies.cpp
@@ -7,7 +7,7 @@
 
 #include "Chaos/ParticleHandle.h"
 #include "Chaos/ErrorReporter.h"
-#include "PhysicsProxy/SingleParticlePhysicsProxy.h"
+#include "Chaos/SingleParticlePhysicsProxy.h"
 #include "PhysicsProxy/GeometryCollectionPhysicsProxy.h"
 #include "Chaos/Utilities.h"
 #include "PBDRigidsSolver.h"
diff --git a/Engine/Source/Runtime/Engine/Private/PhysicsEngine/Experimental/PhysInterface_Chaos.cpp b/Engine/Source/Runtime/Engine/Private/PhysicsEngine/Experimental/PhysInterface_Chaos.cpp
--- a/Engine/Source/Runtime/Engine/Private/PhysicsEngine/Experimental/PhysInterface_Chaos.cpp
+++ b/Engine/Source/Runtime/Engine/Private/PhysicsEngine/Experimental/PhysInterface_Chaos.cpp
@@ -93,7 +93,7 @@ Chaos::FChaosPhysicsMaterial* GetMaterialFromInternalFaceIndex(const FPhysicsSha
 {
 	if(Shape.Materials.Num() > 0 && Actor.Proxy)
 	{
-		Chaos::FPBDRigidsSolver* Solver = Actor.Proxy->GetSolver();
+		Chaos::FPBDRigidsSolver* Solver = static_cast<Chaos::FPBDRigidsSolver*>(Actor.Proxy->GetSolver());
 
 		if(ensure(Solver))
 		{
@@ -123,7 +123,7 @@ Chaos::FChaosPhysicsMaterial* GetMaterialFromInternalFaceIndexAndHitLocation(con
 
 		if (Shape.Materials.Num() > 0 && Actor.Proxy)
 		{
-			Chaos::FPBDRigidsSolver* Solver = Actor.Proxy->GetSolver();
+			Chaos::FPBDRigidsSolver* Solver = static_cast<Chaos::FPBDRigidsSolver*>(Actor.Proxy->GetSolver());
 
 			if (ensure(Solver))
 			{
@@ -455,7 +455,7 @@ FPhysScene* FPhysInterface_Chaos::GetCurrentScene(const FPhysicsActorHandle& InH
 
 	if (IPhysicsProxyBase* Proxy = InHandle->Proxy)
 	{
-		Chaos::FPhysicsSolver* Solver = Proxy->GetSolver();
+		Chaos::FPhysicsSolver* Solver = static_cast<Chaos::FPhysicsSolver*>(Proxy->GetSolver());
 		return static_cast<FPhysScene*>(Solver ? Solver->PhysSceneHack : nullptr);
 	}
 	return nullptr;
diff --git a/Engine/Source/Runtime/Engine/Private/PhysicsEngine/Experimental/PhysScene_Chaos.cpp b/Engine/Source/Runtime/Engine/Private/PhysicsEngine/Experimental/PhysScene_Chaos.cpp
--- a/Engine/Source/Runtime/Engine/Private/PhysicsEngine/Experimental/PhysScene_Chaos.cpp
+++ b/Engine/Source/Runtime/Engine/Private/PhysicsEngine/Experimental/PhysScene_Chaos.cpp
@@ -26,7 +26,7 @@
 
 #include "PhysicsProxy/FieldSystemPhysicsProxy.h"
 #include "PhysicsProxy/GeometryCollectionPhysicsProxy.h"
-#include "PhysicsProxy/SingleParticlePhysicsProxy.h"
+#include "Chaos/SingleParticlePhysicsProxy.h"
 #include "PhysicsProxy/SkeletalMeshPhysicsProxy.h"
 #include "PhysicsProxy/StaticMeshPhysicsProxy.h"
 #include "Chaos/UniformGrid.h"
@@ -819,7 +819,7 @@ void FPhysScene_Chaos::RemoveObject(FGeometryParticlePhysicsProxy* InObject)
 
 void FPhysScene_Chaos::RemoveObject(FGeometryCollectionPhysicsProxy* InObject)
 {
-	Chaos::FPhysicsSolver* Solver = InObject->GetSolver();
+	Chaos::FPhysicsSolver* Solver = static_cast<Chaos::FPhysicsSolver*>(InObject->GetSolver());
 	if(Solver && !Solver->UnregisterObject(InObject))
 	{
 		UE_LOG(LogChaos, Warning, TEXT("Attempted to remove an object that wasn't found in its solver's gamethread storage - it's likely the solver has been mistakenly changed."));
@@ -830,7 +830,7 @@ void FPhysScene_Chaos::RemoveObject(FGeometryCollectionPhysicsProxy* InObject)
 
 void FPhysScene_Chaos::RemoveObject(FFieldSystemPhysicsProxy* InObject)
 {
-	Chaos::FPhysicsSolver* CurrSceneSolver = InObject->GetSolver();
+	Chaos::FPhysicsSolver* CurrSceneSolver = static_cast<Chaos::FPhysicsSolver*>(InObject->GetSolver());
 	if(CurrSceneSolver && !CurrSceneSolver->UnregisterObject(InObject))
 	{
 		UE_LOG(LogChaos, Warning, TEXT("Attempted to remove an object that wasn't found in its solver's gamethread storage - it's likely the solver has been mistakenly changed."));
diff --git a/Engine/Source/Runtime/Engine/Public/Physics/Experimental/PhysScene_Chaos.h b/Engine/Source/Runtime/Engine/Public/Physics/Experimental/PhysScene_Chaos.h
--- a/Engine/Source/Runtime/Engine/Public/Physics/Experimental/PhysScene_Chaos.h
+++ b/Engine/Source/Runtime/Engine/Public/Physics/Experimental/PhysScene_Chaos.h
@@ -14,7 +14,7 @@
 #include "Chaos/Transform.h"
 #include "Chaos/ArrayCollectionArray.h"
 #include "Chaos/Declares.h"
-#include "PhysicsProxy/SingleParticlePhysicsProxyFwd.h"
+#include "Chaos/SingleParticlePhysicsProxyFwd.h"
 #include "Framework/Threading.h"
 #include "Chaos/PBDCollisionConstraints.h"
 
diff --git a/Engine/Source/Runtime/Experimental/Chaos/Private/Chaos/Framework/PhysicsSolverBase.cpp b/Engine/Source/Runtime/Experimental/Chaos/Private/Chaos/Framework/PhysicsSolverBase.cpp
--- a/Engine/Source/Runtime/Experimental/Chaos/Private/Chaos/Framework/PhysicsSolverBase.cpp
+++ b/Engine/Source/Runtime/Experimental/Chaos/Private/Chaos/Framework/PhysicsSolverBase.cpp
@@ -2,6 +2,7 @@
 
 #include "Chaos/Framework/PhysicsSolverBase.h"
 #include "Chaos/Framework/PhysicsProxyBase.h"
+#include "Chaos/PBDRigidsEvolutionGBF.h"
 
 namespace Chaos
 {	
@@ -10,4 +11,15 @@ namespace Chaos
 		BufferMode = InBufferMode;
 	}
 
+	FPhysicsSolverBase::FPhysicsSolverBase(const EMultiBufferMode BufferingModeIn)
+	: BufferMode(BufferingModeIn)
+	{}
+
+	FPhysicsSolverBase::~FPhysicsSolverBase() = default;
+	FPhysicsSolverBase::FPhysicsSolverBase(FPhysicsSolverBase&&) = default;
+
+	void FPhysicsSolverBase::SetEvolution(TUniquePtr<FPBDRigidsEvolutionGBF>&& Evolution)
+	{
+		MEvolution = MoveTemp(Evolution);
+	}
 }
diff --git a/Engine/Source/Runtime/Experimental/Chaos/Private/Chaos/ParticleHandle.cpp b/Engine/Source/Runtime/Experimental/Chaos/Private/Chaos/ParticleHandle.cpp
--- a/Engine/Source/Runtime/Experimental/Chaos/Private/Chaos/ParticleHandle.cpp
+++ b/Engine/Source/Runtime/Experimental/Chaos/Private/Chaos/ParticleHandle.cpp
@@ -86,7 +86,7 @@ namespace Chaos
 
 			if (Proxy)
 			{
-				if (FPhysicsSolverBase* PhysicsSolverBase = Proxy->GetSolver<FPhysicsSolverBase>())
+				if (FPhysicsSolverBase* PhysicsSolverBase = Proxy->GetSolver())
 				{
 					PhysicsSolverBase->AddDirtyProxy(Proxy);
 				}
diff --git a/Engine/Source/Runtime/Experimental/Chaos/Private/Chaos/SingleParticlePhysicsProxy.cpp b/Engine/Source/Runtime/Experimental/Chaos/Private/Chaos/SingleParticlePhysicsProxy.cpp
new file mode 100644
--- /dev/null
+++ b/Engine/Source/Runtime/Experimental/Chaos/Private/Chaos/SingleParticlePhysicsProxy.cpp
@@ -0,0 +1,487 @@
+// Copyright Epic Games, Inc. All Rights Reserved.
+
+#include "Chaos/SingleParticlePhysicsProxy.h"
+
+#include "ChaosStats.h"
+#include "Chaos/Box.h"
+#include "Chaos/Capsule.h"
+#include "Chaos/Sphere.h"
+#include "Chaos/ErrorReporter.h"
+#include "Chaos/Serializable.h"
+#include "Chaos/PerParticleGravity.h"
+#include "Chaos/ParticleHandle.h"
+#include "Chaos/Framework/MultiBufferResource.h"
+#include "Chaos/PBDRigidsEvolutionGBF.h"
+
+template< class PARTICLE_TYPE >
+FSingleParticlePhysicsProxy<PARTICLE_TYPE>::FSingleParticlePhysicsProxy(PARTICLE_TYPE* InParticle, FParticleHandle* InHandle, UObject* InOwner, FInitialState InInitialState)
+	: Base(InOwner)
+	, bInitialized(false)
+	, InitialState(InInitialState)
+	, Particle(InParticle)
+	, Handle(InHandle)
+{
+	BufferedData = Chaos::FMultiBufferFactory< typename PARTICLE_TYPE::FData >::CreateBuffer(Chaos::EMultiBufferMode::Double);
+}
+
+template< class PARTICLE_TYPE >
+FSingleParticlePhysicsProxy<PARTICLE_TYPE>::~FSingleParticlePhysicsProxy()
+{
+}
+
+template< class PARTICLE_TYPE >
+const FInitialState& 
+FSingleParticlePhysicsProxy<PARTICLE_TYPE>::GetInitialState() const 
+{ 
+	return InitialState; 
+}
+
+
+//
+// TGeometryParticle<float, 3> template specialization 
+//
+
+template< >
+CHAOS_API void FSingleParticlePhysicsProxy<Chaos::TGeometryParticle<float, 3>>::PushToPhysicsState(const Chaos::FParticleData* InData)
+{
+	// move the copied game thread data into the handle
+	if (auto* RigidHandle = static_cast<Chaos::TGeometryParticleHandle<float, 3>*>(Handle))
+	{
+		typedef Chaos::TKinematicGeometryParticle<float, 3> PARTICLE_TYPE;
+		const auto* Data = static_cast<const typename PARTICLE_TYPE::FData*>(InData);
+
+		RigidHandle->SetX(Data->X);
+		RigidHandle->SetR(Data->R);
+		RigidHandle->SetSharedGeometry(Data->Geometry);
+		if (Data->Geometry && Data->Geometry->HasBoundingBox())
+		{
+			RigidHandle->SetHasBounds(true);
+			RigidHandle->SetLocalBounds(Data->Geometry->BoundingBox());
+			RigidHandle->SetWorldSpaceInflatedBounds(Data->Geometry->BoundingBox().TransformedAABB(Chaos::TRigidTransform<float, 3>(Data->X, Data->R)));
+		}
+		RigidHandle->SetSpatialIdx(Data->SpatialIdx);	//todo: this needs to only happen once during initialization
+		RigidHandle->SetUniqueIdx(Data->UniqueIdx);
+		RigidHandle->SetUserData(Data->UserData);
+#if CHAOS_CHECKED
+		RigidHandle->SetDebugName(Data->DebugName);
+#endif
+
+		if (Data->DirtyFlags.IsDirty(Chaos::EParticleFlags::ShapeDisableCollision))
+		{
+			int32 CurrShape = 0;
+			for (const TUniquePtr<Chaos::TPerShapeData<Chaos::FReal, 3>>& Shape : RigidHandle->ShapesArray())
+			{
+				Shape->bDisable = Data->ShapeCollisionDisableFlags[CurrShape++];
+			}
+		}
+
+		// @todo(chaos) : Fix - The initial update is not flagging the dirty flag, so this state is not updated. 
+		//if (Data->DirtyFlags.IsDirty(Chaos::EParticleFlags::CollisionTraceType))
+		{
+			int32 CurrShape = 0;
+			for (const TUniquePtr<Chaos::TPerShapeData<Chaos::FReal, 3>>& Shape : RigidHandle->ShapesArray())
+			{
+				Shape->CollisionTraceType = Data->CollisionTraceType[CurrShape++];
+			}
+		}
+
+		// @todo(chaos) : Fix - The initial update is not flagging the dirty flag, so this state is not updated. 
+		//if (Data->DirtyFlags.IsDirty(Chaos::EParticleFlags::ShapeSimData))
+		{
+			int32 CurrShape = 0;
+			for (const TUniquePtr<Chaos::TPerShapeData<Chaos::FReal, 3>>& Shape : RigidHandle->ShapesArray())
+			{
+				Shape->SimData = Data->ShapeSimData[CurrShape];
+				Shape->QueryData = Data->ShapeQueryData[CurrShape++];
+			}
+		}
+
+	}
+}
+
+
+
+template< >
+CHAOS_API void FSingleParticlePhysicsProxy<Chaos::TGeometryParticle<float, 3>>::ClearAccumulatedData()
+{
+	Particle->ClearDirtyFlags();
+}
+
+
+template< >
+CHAOS_API void FSingleParticlePhysicsProxy<Chaos::TGeometryParticle<float, 3>>::BufferPhysicsResults()
+{
+	// Move simulation results into the double buffer.
+}
+
+
+template< >
+CHAOS_API void FSingleParticlePhysicsProxy<Chaos::TGeometryParticle<float, 3>>::PullFromPhysicsState()
+{
+	// Move buffered data into the TGeometryParticle
+}
+
+template< >
+CHAOS_API bool FSingleParticlePhysicsProxy<Chaos::TGeometryParticle<float, 3>>::IsDirty()
+{
+	return Particle->IsDirty();
+}
+
+template< >
+CHAOS_API EPhysicsProxyType FSingleParticlePhysicsProxy<Chaos::TGeometryParticle<float, 3>>::ConcreteType()
+{
+	return EPhysicsProxyType::SingleGeometryParticleType;
+}
+
+template< >
+CHAOS_API bool FSingleParticlePhysicsProxy<Chaos::TGeometryParticle<float, 3>>::HasAwakeEvent() const
+{
+	return false;
+}
+
+template< >
+CHAOS_API void FSingleParticlePhysicsProxy<Chaos::TGeometryParticle<float, 3>>::ClearEvents()
+{
+
+}
+//
+// TKinematicGeometryParticle template specialization 
+//
+
+template< >
+CHAOS_API void FSingleParticlePhysicsProxy<Chaos::TKinematicGeometryParticle<float, 3>>::PushToPhysicsState(const Chaos::FParticleData* InData)
+{
+	// move the copied game thread data into the handle
+	if (auto* RigidHandle = static_cast<Chaos::TKinematicGeometryParticleHandle<float, 3>*>(Handle))
+	{
+		typedef Chaos::TKinematicGeometryParticle<float, 3> PARTICLE_TYPE;
+		const auto* Data = static_cast<const typename PARTICLE_TYPE::FData*>(InData);
+
+		RigidHandle->SetX(Data->X);
+		RigidHandle->SetR(Data->R);
+		RigidHandle->SetSharedGeometry(Data->Geometry);
+		RigidHandle->SetV(Data->MV);
+		RigidHandle->SetW(Data->MW);
+		RigidHandle->SetCenterOfMass(Data->MCenterOfMass);
+		RigidHandle->SetRotationOfMass(Data->MRotationOfMass);
+		RigidHandle->SetSpatialIdx(Data->SpatialIdx);	//todo: this needs to only happen once during initialization
+		RigidHandle->SetUniqueIdx(Data->UniqueIdx);
+		RigidHandle->SetUserData(Data->UserData);
+#if CHAOS_CHECKED
+		RigidHandle->SetDebugName(Data->DebugName);
+#endif
+		
+		if (Data->Geometry && Data->Geometry->HasBoundingBox())
+		{
+			RigidHandle->SetHasBounds(true);
+			RigidHandle->SetLocalBounds(Data->Geometry->BoundingBox());
+			Chaos::TAABB<float, 3> WorldSpaceBox = Data->Geometry->BoundingBox().TransformedAABB(Chaos::TRigidTransform<float, 3>(Data->X, Data->R));
+			WorldSpaceBox.ThickenSymmetrically(Data->MV);
+			RigidHandle->SetWorldSpaceInflatedBounds(WorldSpaceBox);
+		}
+
+		if (Data->DirtyFlags.IsDirty(Chaos::EParticleFlags::ShapeDisableCollision))
+		{
+			int32 CurrShape = 0;
+			for (const TUniquePtr<Chaos::TPerShapeData<Chaos::FReal, 3>>& Shape : RigidHandle->ShapesArray())
+			{
+				Shape->bDisable = Data->ShapeCollisionDisableFlags[CurrShape++];
+			}
+		}
+		if (Data->DirtyFlags.IsDirty(Chaos::EParticleFlags::CollisionTraceType))
+		{
+			int32 CurrShape = 0;
+			for (const TUniquePtr<Chaos::TPerShapeData<Chaos::FReal, 3>>& Shape : RigidHandle->ShapesArray())
+			{
+				Shape->CollisionTraceType = Data->CollisionTraceType[CurrShape++];
+			}
+		}
+		//if (Data->DirtyFlags.IsDirty(Chaos::EParticleFlags::ShapeSimData))
+		{
+			int32 CurrShape = 0;
+			for (const TUniquePtr<Chaos::TPerShapeData<Chaos::FReal, 3>>& Shape : RigidHandle->ShapesArray())
+			{
+				Shape->SimData = Data->ShapeSimData[CurrShape];
+				Shape->QueryData = Data->ShapeQueryData[CurrShape++];
+			}
+		}
+
+
+	}
+}
+
+template< >
+CHAOS_API void FSingleParticlePhysicsProxy<Chaos::TKinematicGeometryParticle<float, 3>>::ClearAccumulatedData()
+{
+	Particle->ClearDirtyFlags();
+}
+
+template< >
+CHAOS_API void FSingleParticlePhysicsProxy<Chaos::TKinematicGeometryParticle<float, 3>>::BufferPhysicsResults()
+{
+	// Move simulation results into the double buffer.
+}
+
+template< >
+CHAOS_API void FSingleParticlePhysicsProxy<Chaos::TKinematicGeometryParticle<float, 3>>::PullFromPhysicsState()
+{
+	// Move buffered data into the TKinematicGeometryParticle
+}
+
+template< >
+CHAOS_API bool FSingleParticlePhysicsProxy<Chaos::TKinematicGeometryParticle<float, 3>>::IsDirty()
+{
+	return Particle->IsDirty();
+}
+
+template< >
+CHAOS_API EPhysicsProxyType FSingleParticlePhysicsProxy<Chaos::TKinematicGeometryParticle<float, 3>>::ConcreteType()
+{
+	return EPhysicsProxyType::SingleKinematicParticleType;
+}
+
+template< >
+CHAOS_API bool FSingleParticlePhysicsProxy<Chaos::TKinematicGeometryParticle<float, 3>>::HasAwakeEvent() const
+{
+	return false;
+}
+
+template< >
+CHAOS_API void FSingleParticlePhysicsProxy<Chaos::TKinematicGeometryParticle<float, 3>>::ClearEvents()
+{
+
+}
+
+//
+// TPBDRigidParticle template specialization 
+//
+
+template< >
+CHAOS_API void FSingleParticlePhysicsProxy<Chaos::TPBDRigidParticle<float, 3>>::PushToPhysicsState(const Chaos::FParticleData* InData)
+{
+	// move the copied game thread data into the handle
+	if (auto* RigidHandle = static_cast<Chaos::TPBDRigidParticleHandle<float, 3>*>(Handle))
+	{
+		typedef Chaos::TPBDRigidParticle<float, 3> PARTICLE_TYPE;
+		const auto* Data = static_cast<const typename PARTICLE_TYPE::FData*>(InData);
+
+		//
+		// TODO: Split all these setters into separate functions so that
+		//       we don't have to check for each dirty type. We don't
+		//       *actually* have to check each dirty type, but some of
+		//       these set operations require extra actions. This avoids
+		//       performing them for every set.
+		//
+
+		// TODO: Since this can't change, it doesn't have a dirty flag.
+		// This should be moved to initialization.
+		RigidHandle->SetSpatialIdx(Data->SpatialIdx);	//todo: this needs to only happen once during initialization
+		RigidHandle->SetUniqueIdx(Data->UniqueIdx);	//todo: this needs to only happen once during initialization
+		RigidHandle->SetUserData(Data->UserData);
+#if CHAOS_CHECKED
+		RigidHandle->SetDebugName(Data->DebugName);
+#endif
+
+		//  DynamicProperties are anything that would wake up a sleeping
+		//  particle. For example, if we set position on a sleeping particle
+		//  we will need to wake it up
+		bool bDynamicPropertyUpdated = false;
+
+		if (Data->DirtyFlags.IsDirty(Chaos::EParticleFlags::X))
+		{
+			RigidHandle->SetX(Data->X);
+			RigidHandle->SetP(Data->X);
+			bDynamicPropertyUpdated = true;
+		}
+		if (Data->DirtyFlags.IsDirty(Chaos::EParticleFlags::R))
+		{
+			RigidHandle->SetR(Data->R);
+			RigidHandle->SetQ(Data->R);
+			bDynamicPropertyUpdated = true;
+		}
+		if (Data->DirtyFlags.IsDirty(Chaos::EParticleFlags::Geometry))
+		{
+			RigidHandle->SetSharedGeometry(Data->Geometry);
+		}
+		if (Data->DirtyFlags.IsDirty(Chaos::EParticleFlags::V))
+		{
+			RigidHandle->SetV(Data->MV);
+			bDynamicPropertyUpdated = true;
+		}
+		if (Data->DirtyFlags.IsDirty(Chaos::EParticleFlags::W))
+		{
+			RigidHandle->SetW(Data->MW);
+			bDynamicPropertyUpdated = true;
+		}
+
+		RigidHandle->SetCenterOfMass(Data->MCenterOfMass);
+		RigidHandle->SetRotationOfMass(Data->MRotationOfMass);
+		RigidHandle->SetM(Data->MM);
+		RigidHandle->SetInvM(Data->MInvM);
+		RigidHandle->SetI(Data->MI);
+		RigidHandle->SetInvI(Data->MInvI);
+		RigidHandle->SetLinearEtherDrag(Data->MLinearEtherDrag);
+		RigidHandle->SetAngularEtherDrag(Data->MAngularEtherDrag);
+
+		if (Data->DirtyFlags.IsDirty(Chaos::EParticleFlags::F))
+		{
+			RigidHandle->SetF(Data->MF);
+			bDynamicPropertyUpdated = true;
+		}
+		if (Data->DirtyFlags.IsDirty(Chaos::EParticleFlags::Torque))
+		{
+			RigidHandle->SetTorque(Data->MTorque);
+			bDynamicPropertyUpdated = true;
+		}
+
+		if (Data->DirtyFlags.IsDirty(Chaos::EParticleFlags::LinearImpulse))
+		{
+			RigidHandle->SetLinearImpulse(Data->MLinearImpulse);
+			bDynamicPropertyUpdated = true;
+		}
+
+		if (Data->DirtyFlags.IsDirty(Chaos::EParticleFlags::AngularImpulse))
+		{
+			RigidHandle->SetAngularImpulse(Data->MAngularImpulse);
+			bDynamicPropertyUpdated = true;
+		}
+
+		if (Data->DirtyFlags.IsDirty(Chaos::EParticleFlags::ObjectState))
+		{
+			GetSolver()->GetEvolution()->SetParticleObjectState(RigidHandle, Data->MObjectState);
+		}
+		if (Data->DirtyFlags.IsDirty(Chaos::EParticleFlags::GravityEnabled))
+		{
+			GetSolver()->GetEvolution()->GetGravityForces().SetEnabled(*RigidHandle, Data->MGravityEnabled);
+		}
+
+		if (Data->DirtyFlags.IsDirty((int32)Chaos::EParticleFlags::X | (int32)Chaos::EParticleFlags::R | (int32)Chaos::EParticleFlags::V | (int32)Chaos::EParticleFlags::Geometry))
+		{
+			if (Data->Geometry && Data->Geometry->HasBoundingBox())
+			{
+				RigidHandle->SetHasBounds(true);
+				RigidHandle->SetLocalBounds(Data->Geometry->BoundingBox());
+				Chaos::TAABB<float, 3> WorldSpaceBox = Data->Geometry->BoundingBox().TransformedAABB(Chaos::TRigidTransform<float, 3>(Data->X, Data->R));
+				WorldSpaceBox.ThickenSymmetrically(Data->MV);
+				RigidHandle->SetWorldSpaceInflatedBounds(WorldSpaceBox);
+			}
+			else
+			{
+				//todo: compute bounds based on sample points
+			}
+		}
+
+		if (Data->DirtyFlags.IsDirty(Chaos::EParticleFlags::ShapeDisableCollision))
+		{
+			int32 CurrShape = 0;
+			for (const TUniquePtr<Chaos::TPerShapeData<Chaos::FReal, 3>>& Shape : RigidHandle->ShapesArray())
+			{
+				Shape->bDisable = Data->ShapeCollisionDisableFlags[CurrShape++];
+			}
+		}
+		if (Data->DirtyFlags.IsDirty(Chaos::EParticleFlags::CollisionTraceType))
+		{
+			int32 CurrShape = 0;
+			for (const TUniquePtr<Chaos::TPerShapeData<Chaos::FReal, 3>>& Shape : RigidHandle->ShapesArray())
+			{
+				Shape->CollisionTraceType = Data->CollisionTraceType[CurrShape++];
+			}
+		}
+		//if (Data->DirtyFlags.IsDirty(Chaos::EParticleFlags::ShapeSimData))
+		{
+			int32 CurrShape = 0;
+			for (const TUniquePtr<Chaos::TPerShapeData<Chaos::FReal, 3>>& Shape : RigidHandle->ShapesArray())
+			{
+				Shape->SimData = Data->ShapeSimData[CurrShape];
+				Shape->QueryData = Data->ShapeQueryData[CurrShape++];
+			}
+		}
+
+		if (bDynamicPropertyUpdated)
+		{
+			if (bInitialized && RigidHandle->ObjectState() == Chaos::EObjectStateType::Sleeping)
+			{
+				GetSolver()->GetEvolution()->SetParticleObjectState(RigidHandle, Chaos::EObjectStateType::Dynamic);
+			}
+		}
+		
+		if(Data->MInitialized)
+		{
+			// wait for the first pass with nothing updated to claim its initialized
+			bInitialized = true;
+		}
+	}
+
+}
+
+template< >
+CHAOS_API void FSingleParticlePhysicsProxy<Chaos::TPBDRigidParticle<float, 3>>::ClearAccumulatedData()
+{
+	Particle->SetF(Chaos::TVector<float, 3>(0));
+	Particle->SetTorque(Chaos::TVector<float, 3>(0));
+	Particle->SetLinearImpulse(Chaos::TVector<float, 3>(0));
+	Particle->SetAngularImpulse(Chaos::TVector<float, 3>(0));
+	Particle->ClearDirtyFlags();
+}
+
+template< >
+CHAOS_API void FSingleParticlePhysicsProxy<Chaos::TPBDRigidParticle<float, 3>>::BufferPhysicsResults()
+{
+	// Move simulation results into the double buffer.
+	if (Chaos::TPBDRigidParticleHandle<float, 3>* RigidHandle = static_cast<Chaos::TPBDRigidParticleHandle<float, 3>*>(Handle))
+	{
+		Chaos::TPBDRigidParticleData<float, 3>* Buffer = BufferedData->AccessProducerBuffer();
+		Buffer->X = RigidHandle->X();
+		Buffer->R = RigidHandle->R();
+		Buffer->MV = RigidHandle->V();
+		Buffer->MW = RigidHandle->W();
+		Buffer->MObjectState = RigidHandle->ObjectState();
+	}
+}
+
+template< >
+CHAOS_API void FSingleParticlePhysicsProxy<Chaos::TPBDRigidParticle<float, 3>>::PullFromPhysicsState()
+{
+	// Move buffered data into the TPBDRigidParticle without triggering invalidation of the physics state.
+	if (Particle)
+	{
+		const Chaos::TPBDRigidParticleData<float, 3>* Buffer = BufferedData->GetConsumerBuffer();
+		Particle->SetX(Buffer->X, false);
+		Particle->SetR(Buffer->R, false);
+		Particle->SetV(Buffer->MV, false);
+		Particle->SetW(Buffer->MW, false);
+		Particle->UpdateShapeBounds();
+		if (!Particle->IsDirty(Chaos::EParticleFlags::ObjectState))
+		{
+			Particle->SetObjectState(Buffer->MObjectState, true);
+		}
+	}
+}
+
+template< >
+CHAOS_API bool FSingleParticlePhysicsProxy<Chaos::TPBDRigidParticle<float, 3>>::IsDirty()
+{
+	return Particle->IsDirty();
+}
+
+template< >
+CHAOS_API bool FSingleParticlePhysicsProxy< Chaos::TPBDRigidParticle<float, 3> >::HasAwakeEvent() const
+{
+	return Particle->HasAwakeEvent();
+}
+
+template< >
+CHAOS_API void FSingleParticlePhysicsProxy< Chaos::TPBDRigidParticle<float, 3> >::ClearEvents()
+{
+	Particle->ClearEvents();
+}
+
+template< >
+CHAOS_API EPhysicsProxyType  FSingleParticlePhysicsProxy<Chaos::TPBDRigidParticle<float, 3>>::ConcreteType() 
+{ 
+	return EPhysicsProxyType::SingleRigidParticleType;
+}
+
+template class CHAOS_API FSingleParticlePhysicsProxy< Chaos::TGeometryParticle<float, 3> >;
+template class CHAOS_API FSingleParticlePhysicsProxy< Chaos::TKinematicGeometryParticle<float, 3> >;
+template class CHAOS_API FSingleParticlePhysicsProxy< Chaos::TPBDRigidParticle<float,3> >;
diff --git a/Engine/Source/Runtime/Experimental/Chaos/Public/Chaos/Framework/PhysicsProxyBase.h b/Engine/Source/Runtime/Experimental/Chaos/Public/Chaos/Framework/PhysicsProxyBase.h
--- a/Engine/Source/Runtime/Experimental/Chaos/Public/Chaos/Framework/PhysicsProxyBase.h
+++ b/Engine/Source/Runtime/Experimental/Chaos/Public/Chaos/Framework/PhysicsProxyBase.h
@@ -18,6 +18,11 @@ enum class EPhysicsProxyType
 	SingleRigidParticleType = 7
 };
 
+namespace Chaos
+{
+	class FPhysicsSolverBase;
+}
+
 class IPhysicsProxyBase
 {
 public:
@@ -28,8 +33,7 @@ public:
 
 	virtual UObject* GetOwner() const = 0;
 
-	template< class SOLVER_TYPE = Chaos::FPhysicsSolver>
-	SOLVER_TYPE* GetSolver() const { return static_cast<SOLVER_TYPE*>(Solver); }
+	Chaos::FPhysicsSolverBase* GetSolver() const { return Solver; }
 
 	//Should this be in the public API? probably not
 	template< class SOLVER_TYPE = Chaos::FPhysicsSolver>
@@ -45,9 +49,9 @@ protected:
 	// but no one can delete using a IPhysicsProxyBase*
 	virtual ~IPhysicsProxyBase()
 	{
-		if (GetSolver<Chaos::FPhysicsSolverBase>())
+		if (Solver)
 		{
-			GetSolver<Chaos::FPhysicsSolverBase>()->RemoveDirtyProxy(this);
+			Solver->RemoveDirtyProxy(this);
 		}
 	}
 
diff --git a/Engine/Source/Runtime/Experimental/Chaos/Public/Chaos/Framework/PhysicsSolverBase.h b/Engine/Source/Runtime/Experimental/Chaos/Public/Chaos/Framework/PhysicsSolverBase.h
--- a/Engine/Source/Runtime/Experimental/Chaos/Public/Chaos/Framework/PhysicsSolverBase.h
+++ b/Engine/Source/Runtime/Experimental/Chaos/Public/Chaos/Framework/PhysicsSolverBase.h
@@ -9,14 +9,17 @@ class IPhysicsProxyBase; // WTF - not in Chaos?
 
 namespace Chaos
 {
+	class FPBDRigidsEvolutionGBF;
+	using FPBDRigidsEvolution = FPBDRigidsEvolutionGBF;
 
 	class CHAOS_API FPhysicsSolverBase
 	{
 	public:
 
-		FPhysicsSolverBase(const EMultiBufferMode BufferingModeIn)
-			: BufferMode(BufferingModeIn)
-		{}
+		FPhysicsSolverBase(const EMultiBufferMode BufferingModeIn);
+		~FPhysicsSolverBase();
+		FPhysicsSolverBase(const FPhysicsSolverBase&) = delete;
+		FPhysicsSolverBase(FPhysicsSolverBase&&);
 		
 		void ChangeBufferMode(EMultiBufferMode InBufferMode);
 
@@ -29,6 +32,9 @@ namespace Chaos
 			DirtyProxiesSet.Remove(ProxyBaseIn);
 		}
 
+		FPBDRigidsEvolution* GetEvolution() { return MEvolution.Get(); }
+		FPBDRigidsEvolution* GetEvolution() const { return MEvolution.Get(); }
+
 #if CHAOS_CHECKED
 		void SetDebugName(const FName& Name)
 		{
@@ -50,5 +56,9 @@ namespace Chaos
 #if CHAOS_CHECKED
 			FName DebugName;
 #endif
+			//TODO: Make more extensible
+			TUniquePtr<FPBDRigidsEvolution> MEvolution;
+
+			void SetEvolution(TUniquePtr<FPBDRigidsEvolution>&& Evolution);
 	};
 }
diff --git a/Engine/Source/Runtime/Experimental/Chaos/Public/Chaos/SingleParticlePhysicsProxy.h b/Engine/Source/Runtime/Experimental/Chaos/Public/Chaos/SingleParticlePhysicsProxy.h
new file mode 100644
--- /dev/null
+++ b/Engine/Source/Runtime/Experimental/Chaos/Public/Chaos/SingleParticlePhysicsProxy.h
@@ -0,0 +1,246 @@
+// Copyright Epic Games, Inc. All Rights Reserved.
+
+#pragma once
+
+#include "Chaos/ArrayCollectionArray.h"
+#include "Chaos/Framework/MultiBufferResource.h"
+#include "Chaos/Framework/PhysicsProxy.h"
+#include "Chaos/PBDPositionConstraints.h"
+#include "Chaos/ParticleHandle.h"
+
+namespace Chaos
+{
+	template<typename T, int d>
+	class TGeometryParticle;
+}
+
+class FInitialState
+{
+public:
+	FInitialState()
+	    : Mass(0.f)
+	    , InvMass(0.f)
+	    , InertiaTensor(1.f)
+	{}
+
+	FInitialState(float MassIn, float InvMassIn, FVector InertiaTensorIn)
+	    : Mass(MassIn)
+	    , InvMass(InvMassIn)
+	    , InertiaTensor(InertiaTensorIn)
+	{}
+
+	float GetMass() const { return Mass; }
+	float GetInverseMass() const { return InvMass; }
+	FVector GetInertiaTensor() const { return InertiaTensor; }
+
+private:
+	float Mass;
+	float InvMass;
+	FVector InertiaTensor;
+};
+
+/**
+ * \p PARTICLE_TYPE is one of:
+ *		\c Chaos::TGeometryParticle<float,3>
+ *		\c Chaos::TKinematicGeometryParticle<float,3>
+ *		\c Chaos::TPBDRigidParticle<float,3>
+ */
+template<class PARTICLE_TYPE>
+class FSingleParticlePhysicsProxy : public TPhysicsProxy<FSingleParticlePhysicsProxy<PARTICLE_TYPE>, typename PARTICLE_TYPE::FData>
+{
+	typedef TPhysicsProxy<FSingleParticlePhysicsProxy<PARTICLE_TYPE>, typename PARTICLE_TYPE::FData> Base;
+
+public:
+	using FParticleHandle = typename PARTICLE_TYPE::FHandle;
+	using FStorageData = typename PARTICLE_TYPE::FData;
+	using FParticlesType = typename Base::FParticlesType;
+
+	FSingleParticlePhysicsProxy() = delete;
+	FSingleParticlePhysicsProxy(PARTICLE_TYPE* InParticle, FParticleHandle* InHandle, UObject* InOwner = nullptr, FInitialState InitialState = FInitialState());
+	virtual ~FSingleParticlePhysicsProxy();
+
+	// DELETE MOST OF ME
+	void Initialize() {}
+	bool IsSimulating() const { return true; }
+	void ParameterUpdateCallback(FParticlesType& InParticles, const float InTime) {}
+	void UpdateKinematicBodiesCallback(const FParticlesType& InParticles, const float InDt, const float InTime, FKinematicProxy& InKinematicProxy) {}
+	void BindParticleCallbackMapping(Chaos::TArrayCollectionArray<PhysicsProxyWrapper>& PhysicsProxyReverseMap, Chaos::TArrayCollectionArray<int32>& ParticleIDReverseMap) {}
+	void StartFrameCallback(const float InDt, const float InTime) {}
+	void EndFrameCallback(const float InDt) {}
+	void CreateRigidBodyCallback(FParticlesType& InOutParticles) {}
+	void DisableCollisionsCallback(TSet<TTuple<int32, int32>>& InPairs) {}
+	void AddForceCallback(FParticlesType& InParticles, const float InDt, const int32 InIndex) {}
+	void FieldForcesUpdateCallback(Chaos::FPhysicsSolver* InSolver, FParticlesType& Particles, Chaos::TArrayCollectionArray<FVector>& Force, Chaos::TArrayCollectionArray<FVector>& Torque, const float Time) {}
+	void BufferCommand(Chaos::FPhysicsSolver* InSolver, const FFieldSystemCommand& InCommand) {}
+	void SyncBeforeDestroy() {}
+	void OnRemoveFromScene() {}
+	// END DELETE ME
+
+	/**/
+	const FInitialState& GetInitialState() const;
+
+	FParticleHandle* GetHandle()
+	{
+		return Handle;
+	}
+
+	virtual void* GetHandleUnsafe() const override
+	{
+		return Handle;
+	}
+
+	void SetHandle(FParticleHandle* InHandle)
+	{
+		Handle = InHandle;
+	}
+
+	void* GetUserData() const
+	{
+		auto GameThreadHandle = Handle->GTGeometryParticle();
+		return GameThreadHandle ? GameThreadHandle->UserData() : nullptr;
+	}
+
+	Chaos::TRigidTransform<float, 3> GetTransform()
+	{
+		return Chaos::TRigidTransform<float, 3>(Handle->X(), Handle->R());
+	}
+
+	/** 
+	 * Creates a copy of sim state and returns it as \c TGeometryParticleData, 
+	 * \c TKinemticGeometr yParticleData, or \c TPBDRigidParticleData. 
+	 */
+	Chaos::FParticleData* NewData()
+	{
+		if (Particle)
+			return Particle->NewData();
+		return nullptr;
+	}
+
+	/**/
+	EPhysicsProxyType ConcreteType() { return EPhysicsProxyType::NoneType; }
+
+	// Threading API
+
+	/**/
+	void FlipBuffer() { BufferedData->FlipProducer(); }
+
+	/**/
+	void PushToPhysicsState(const Chaos::FParticleData*);
+
+	/**/
+	void ClearAccumulatedData();
+
+	/**/
+	void BufferPhysicsResults();
+
+	/**/
+	void PullFromPhysicsState();
+
+	/**/
+	bool IsDirty();
+
+	/**/
+	bool HasAwakeEvent() const;
+
+	/**/
+	void ClearEvents();
+
+	PARTICLE_TYPE* GetParticle()
+	{
+		return Particle;
+	}
+	
+private:
+	bool bInitialized;
+	TArray<int32> InitializedIndices;
+
+private:
+	FInitialState InitialState;
+
+	PARTICLE_TYPE* Particle;
+	FParticleHandle* Handle;
+	TUniquePtr<Chaos::IBufferResource<FStorageData>> BufferedData;
+};
+
+
+// TGeometryParticle specialization prototypes
+template< >
+CHAOS_API void FSingleParticlePhysicsProxy<Chaos::TGeometryParticle<float, 3>>::ClearAccumulatedData();
+
+template< >
+CHAOS_API void FSingleParticlePhysicsProxy<Chaos::TGeometryParticle<float, 3>>::BufferPhysicsResults();
+
+template< >
+CHAOS_API void FSingleParticlePhysicsProxy<Chaos::TGeometryParticle<float, 3>>::PullFromPhysicsState();
+
+template< >
+CHAOS_API bool FSingleParticlePhysicsProxy<Chaos::TGeometryParticle<float, 3>>::IsDirty();
+
+template< >
+CHAOS_API void FSingleParticlePhysicsProxy<Chaos::TGeometryParticle<float, 3>>::PushToPhysicsState(const Chaos::FParticleData* InData);
+
+template< >
+EPhysicsProxyType FSingleParticlePhysicsProxy<Chaos::TGeometryParticle<float, 3>>::ConcreteType();
+
+template< >
+CHAOS_API bool FSingleParticlePhysicsProxy<Chaos::TGeometryParticle<float, 3>>::HasAwakeEvent() const;
+
+template< >
+CHAOS_API void FSingleParticlePhysicsProxy<Chaos::TGeometryParticle<float, 3>>::ClearEvents();
+
+// TKinematicGeometryParticle specialization prototypes
+
+template< >
+CHAOS_API void FSingleParticlePhysicsProxy<Chaos::TKinematicGeometryParticle<float, 3>>::PushToPhysicsState(const Chaos::FParticleData* InData);
+
+template< >
+CHAOS_API void FSingleParticlePhysicsProxy<Chaos::TKinematicGeometryParticle<float, 3>>::ClearAccumulatedData();
+
+template< >
+CHAOS_API void FSingleParticlePhysicsProxy<Chaos::TKinematicGeometryParticle<float, 3>>::BufferPhysicsResults();
+
+template< >
+CHAOS_API void FSingleParticlePhysicsProxy<Chaos::TKinematicGeometryParticle<float, 3>>::PullFromPhysicsState();
+
+template< >
+CHAOS_API bool FSingleParticlePhysicsProxy<Chaos::TKinematicGeometryParticle<float, 3>>::IsDirty();
+
+template< >
+CHAOS_API EPhysicsProxyType FSingleParticlePhysicsProxy<Chaos::TKinematicGeometryParticle<float, 3>>::ConcreteType();
+
+template< >
+CHAOS_API bool FSingleParticlePhysicsProxy<Chaos::TKinematicGeometryParticle<float, 3>>::HasAwakeEvent() const;
+
+template< >
+CHAOS_API void FSingleParticlePhysicsProxy<Chaos::TKinematicGeometryParticle<float, 3>>::ClearEvents();
+
+// TPBDRigidParticles specialization prototypes
+
+template< >
+CHAOS_API void FSingleParticlePhysicsProxy<Chaos::TPBDRigidParticle<float, 3>>::ClearAccumulatedData();
+
+template< >
+CHAOS_API void FSingleParticlePhysicsProxy<Chaos::TPBDRigidParticle<float, 3>>::BufferPhysicsResults();
+
+template< >
+CHAOS_API void FSingleParticlePhysicsProxy<Chaos::TPBDRigidParticle<float, 3>>::PullFromPhysicsState();
+
+template< >
+CHAOS_API bool FSingleParticlePhysicsProxy<Chaos::TPBDRigidParticle<float, 3>>::IsDirty();
+
+template< >
+CHAOS_API void FSingleParticlePhysicsProxy<Chaos::TPBDRigidParticle<float, 3>>::PushToPhysicsState(const Chaos::FParticleData* InData);
+
+template< >
+CHAOS_API EPhysicsProxyType FSingleParticlePhysicsProxy<Chaos::TPBDRigidParticle<float, 3>>::ConcreteType();
+
+template< >
+CHAOS_API bool FSingleParticlePhysicsProxy<Chaos::TPBDRigidParticle<float, 3>>::HasAwakeEvent() const;
+
+template< >
+CHAOS_API void FSingleParticlePhysicsProxy<Chaos::TPBDRigidParticle<float, 3>>::ClearEvents();
+
+
+extern template class FSingleParticlePhysicsProxy< Chaos::TGeometryParticle<float, 3> >;
+extern template class FSingleParticlePhysicsProxy< Chaos::TKinematicGeometryParticle<float, 3> >;
+extern template class FSingleParticlePhysicsProxy< Chaos::TPBDRigidParticle<float,3> >;
diff --git a/Engine/Source/Runtime/Experimental/Chaos/Public/Chaos/SingleParticlePhysicsProxyFwd.h b/Engine/Source/Runtime/Experimental/Chaos/Public/Chaos/SingleParticlePhysicsProxyFwd.h
new file mode 100644
--- /dev/null
+++ b/Engine/Source/Runtime/Experimental/Chaos/Public/Chaos/SingleParticlePhysicsProxyFwd.h
@@ -0,0 +1,22 @@
+// Copyright Epic Games, Inc. All Rights Reserved.
+
+#pragma once
+
+namespace Chaos
+{
+	template <typename T, int d>
+	class TGeometryParticle;
+
+	template <typename T, int d>
+	class TKinematicGeometryParticle;
+
+	template <typename T, int d>
+	class TPBDRigidParticle;
+}
+
+template <typename T>
+class FSingleParticlePhysicsProxy;
+
+typedef FSingleParticlePhysicsProxy<Chaos::TGeometryParticle<float, 3>> FGeometryParticlePhysicsProxy;
+typedef FSingleParticlePhysicsProxy< Chaos::TKinematicGeometryParticle<float, 3> > FKinematicGeometryParticlePhysicsProxy;
+typedef FSingleParticlePhysicsProxy< Chaos::TPBDRigidParticle<float, 3> > FRigidParticlePhysicsProxy;
\ No newline at end of file
diff --git a/Engine/Source/Runtime/Experimental/ChaosSolvers/Private/EventDefaults.cpp b/Engine/Source/Runtime/Experimental/ChaosSolvers/Private/EventDefaults.cpp
--- a/Engine/Source/Runtime/Experimental/ChaosSolvers/Private/EventDefaults.cpp
+++ b/Engine/Source/Runtime/Experimental/ChaosSolvers/Private/EventDefaults.cpp
@@ -2,7 +2,7 @@
 
 #include "EventDefaults.h"
 #include "EventsData.h"
-#include "PhysicsProxy/SingleParticlePhysicsProxy.h"
+#include "Chaos/SingleParticlePhysicsProxy.h"
 #include "Chaos/PBDRigidsEvolution.h"
 #include "Chaos/PBDRigidClustering.h"
 #include "Chaos/PBDCollisionConstraints.h"
diff --git a/Engine/Source/Runtime/Experimental/ChaosSolvers/Private/Framework/PersistentTask.cpp b/Engine/Source/Runtime/Experimental/ChaosSolvers/Private/Framework/PersistentTask.cpp
--- a/Engine/Source/Runtime/Experimental/ChaosSolvers/Private/Framework/PersistentTask.cpp
+++ b/Engine/Source/Runtime/Experimental/ChaosSolvers/Private/Framework/PersistentTask.cpp
@@ -15,7 +15,7 @@
 
 #include "PhysicsProxy/SkeletalMeshPhysicsProxy.h"
 #include "PhysicsProxy/StaticMeshPhysicsProxy.h"
-#include "PhysicsProxy/SingleParticlePhysicsProxy.h"
+#include "Chaos/SingleParticlePhysicsProxy.h"
 #include "PhysicsProxy/GeometryCollectionPhysicsProxy.h"
 #include "PhysicsProxy/FieldSystemPhysicsProxy.h"
 
diff --git a/Engine/Source/Runtime/Experimental/ChaosSolvers/Private/PBDRigidsSolver.cpp b/Engine/Source/Runtime/Experimental/ChaosSolvers/Private/PBDRigidsSolver.cpp
--- a/Engine/Source/Runtime/Experimental/ChaosSolvers/Private/PBDRigidsSolver.cpp
+++ b/Engine/Source/Runtime/Experimental/ChaosSolvers/Private/PBDRigidsSolver.cpp
@@ -11,7 +11,7 @@
 #include "ChaosSolversModule.h"
 #include "HAL/FileManager.h"
 #include "Misc/ScopeLock.h"
-#include "PhysicsProxy/SingleParticlePhysicsProxy.h"
+#include "Chaos/SingleParticlePhysicsProxy.h"
 #include "PhysicsProxy/SkeletalMeshPhysicsProxy.h"
 #include "PhysicsProxy/StaticMeshPhysicsProxy.h"
 #include "PhysicsProxy/GeometryCollectionPhysicsProxy.h"
@@ -183,7 +183,6 @@ namespace Chaos
 		, bHasFloor(true)
 		, bIsFloorAnalytic(false)
 		, FloorHeight(0.f)
-		, MEvolution(new FPBDRigidsEvolution(Particles, ChaosSolverCollisionDefaultIterationsCVar, ChaosSolverCollisionDefaultPushoutIterationsCVar, BufferingModeIn == Chaos::EMultiBufferMode::Single))
 		, MEventManager(new FEventManager(BufferingModeIn))
 		, MSolverEventFilters(new FSolverEventFilters())
 		, MActiveParticlesBuffer(new FActiveParticlesBuffer(BufferingModeIn))
@@ -506,7 +505,7 @@ namespace Chaos
 		MMaxDeltaTime = 1.f;
 		MMinDeltaTime = 1.e-10f;
 		MMaxSubSteps = 1;
-		MEvolution = TUniquePtr<FPBDRigidsEvolution>(new FPBDRigidsEvolution(Particles, ChaosSolverCollisionDefaultIterationsCVar, ChaosSolverCollisionDefaultPushoutIterationsCVar, BufferMode == EMultiBufferMode::Single)); 
+		SetEvolution(TUniquePtr<FPBDRigidsEvolution>(new FPBDRigidsEvolution(Particles, ChaosSolverCollisionDefaultIterationsCVar, ChaosSolverCollisionDefaultPushoutIterationsCVar, BufferMode == EMultiBufferMode::Single))); 
 
 		FEventDefaults::RegisterSystemEvents(*GetEventManager());
 	}
diff --git a/Engine/Source/Runtime/Experimental/ChaosSolvers/Private/PhysicsProxy/GeometryCollectionPhysicsProxy.cpp b/Engine/Source/Runtime/Experimental/ChaosSolvers/Private/PhysicsProxy/GeometryCollectionPhysicsProxy.cpp
--- a/Engine/Source/Runtime/Experimental/ChaosSolvers/Private/PhysicsProxy/GeometryCollectionPhysicsProxy.cpp
+++ b/Engine/Source/Runtime/Experimental/ChaosSolvers/Private/PhysicsProxy/GeometryCollectionPhysicsProxy.cpp
@@ -1946,7 +1946,7 @@ void FGeometryCollectionPhysicsProxy::ProcessCommands(FParticlesType& Particles,
 	if (Commands.Num())
 	{
 		TArray<ContextIndex> IndicesArray;
-		Chaos::FPhysicsSolver* CurrentSolver = GetSolver();
+		Chaos::FPhysicsSolver* CurrentSolver = static_cast<Chaos::FPhysicsSolver*>(GetSolver());
 
 		for (int32 CommandIndex = Commands.Num()-1; 0<=CommandIndex; CommandIndex--)
 		{
@@ -2053,7 +2053,7 @@ void FGeometryCollectionPhysicsProxy::ProcessCommands(FParticlesType& Particles,
 	// Process Particle-Particle commands
 	if (Commands.Num())
 	{
-		Chaos::FPhysicsSolver* CurrentSolver = GetSolver();
+		Chaos::FPhysicsSolver* CurrentSolver = static_cast<Chaos::FPhysicsSolver*>(GetSolver());
 
 		//  Generate a Index mapping between the rigid body indices and 
 		//  the particle indices. This allows the geometry collection to
diff --git a/Engine/Source/Runtime/Experimental/ChaosSolvers/Private/PhysicsProxy/SingleParticlePhysicsProxy.cpp b/Engine/Source/Runtime/Experimental/ChaosSolvers/Private/PhysicsProxy/SingleParticlePhysicsProxy.cpp
deleted file mode 100644
diff --git a/Engine/Source/Runtime/Experimental/ChaosSolvers/Private/PhysicsProxy/SkeletalMeshPhysicsProxy.cpp b/Engine/Source/Runtime/Experimental/ChaosSolvers/Private/PhysicsProxy/SkeletalMeshPhysicsProxy.cpp
--- a/Engine/Source/Runtime/Experimental/ChaosSolvers/Private/PhysicsProxy/SkeletalMeshPhysicsProxy.cpp
+++ b/Engine/Source/Runtime/Experimental/ChaosSolvers/Private/PhysicsProxy/SkeletalMeshPhysicsProxy.cpp
@@ -380,7 +380,7 @@ void FSkeletalMeshPhysicsProxy::SyncBeforeDestroy()
 void FSkeletalMeshPhysicsProxy::OnRemoveFromScene()
 {
 	// Disable the particle we added
-	Chaos::FPhysicsSolver* CurrSolver = GetSolver();
+	Chaos::FPhysicsSolver* CurrSolver = static_cast<Chaos::FPhysicsSolver*>(GetSolver());
 	if(CurrSolver && RigidBodyIds.Num())
 	{
 #if TODO_REIMPLEMENT_GET_RIGID_PARTICLES
diff --git a/Engine/Source/Runtime/Experimental/ChaosSolvers/Private/PhysicsProxy/StaticMeshPhysicsObject.cpp b/Engine/Source/Runtime/Experimental/ChaosSolvers/Private/PhysicsProxy/StaticMeshPhysicsObject.cpp
--- a/Engine/Source/Runtime/Experimental/ChaosSolvers/Private/PhysicsProxy/StaticMeshPhysicsObject.cpp
+++ b/Engine/Source/Runtime/Experimental/ChaosSolvers/Private/PhysicsProxy/StaticMeshPhysicsObject.cpp
@@ -326,7 +326,7 @@ void FStaticMeshPhysicsProxy::AddForceCallback(FParticlesType& InParticles, cons
 void FStaticMeshPhysicsProxy::OnRemoveFromScene()
 {
 	// Disable the particle we added
-	Chaos::FPhysicsSolver* CurrSolver = GetSolver();
+	Chaos::FPhysicsSolver* CurrSolver = static_cast<Chaos::FPhysicsSolver*>(GetSolver());
 
 	if(CurrSolver && RigidBodyId != INDEX_NONE)
 	{
diff --git a/Engine/Source/Runtime/Experimental/ChaosSolvers/Public/PBDRigidsSolver.h b/Engine/Source/Runtime/Experimental/ChaosSolvers/Public/PBDRigidsSolver.h
--- a/Engine/Source/Runtime/Experimental/ChaosSolvers/Public/PBDRigidsSolver.h
+++ b/Engine/Source/Runtime/Experimental/ChaosSolvers/Public/PBDRigidsSolver.h
@@ -21,7 +21,7 @@
 #include "Framework/Dispatcher.h"
 #include "Field/FieldSystem.h"
 #include "PBDRigidActiveParticlesBuffer.h"
-#include "PhysicsProxy/SingleParticlePhysicsProxyFwd.h"
+#include "Chaos/SingleParticlePhysicsProxyFwd.h"
 #include "SolverEventFilters.h"
 
 class FPhysicsSolverAdvanceTask;
@@ -273,9 +273,6 @@ namespace Chaos
 		void SetIsFloorAnalytic(bool bIsAnalytic) { bIsFloorAnalytic = bIsAnalytic; }
 		void SetFloorHeight(float Height) { FloorHeight = Height; }
 
-		/**/
-		FPBDRigidsEvolution* GetEvolution() { return MEvolution.Get(); }
-		FPBDRigidsEvolution* GetEvolution() const { return MEvolution.Get(); }
 
 		FParticlesType& GetParticles() { return Particles; }
 		const FParticlesType& GetParticles() const { return Particles; }
@@ -357,7 +354,6 @@ namespace Chaos
 		float FloorHeight;
 
 		FParticlesType Particles;
-		TUniquePtr<FPBDRigidsEvolution> MEvolution;
 		TUniquePtr<FEventManager> MEventManager;
 		TUniquePtr<FSolverEventFilters> MSolverEventFilters;
 		TUniquePtr<FActiveParticlesBuffer> MActiveParticlesBuffer;
diff --git a/Engine/Source/Runtime/Experimental/ChaosSolvers/Public/PhysicsProxy/PhysicsProxies.h b/Engine/Source/Runtime/Experimental/ChaosSolvers/Public/PhysicsProxy/PhysicsProxies.h
--- a/Engine/Source/Runtime/Experimental/ChaosSolvers/Public/PhysicsProxy/PhysicsProxies.h
+++ b/Engine/Source/Runtime/Experimental/ChaosSolvers/Public/PhysicsProxy/PhysicsProxies.h
@@ -10,7 +10,7 @@
  * the currently supported set of objects.
  */
 
-#include "SingleParticlePhysicsProxy.h"
+#include "Chaos/SingleParticlePhysicsProxy.h"
 #include "FieldSystemPhysicsProxy.h"
 #include "GeometryCollectionPhysicsProxy.h"
 #include "SkeletalMeshPhysicsProxy.h"
diff --git a/Engine/Source/Runtime/Experimental/ChaosSolvers/Public/PhysicsProxy/SingleParticlePhysicsProxy.h b/Engine/Source/Runtime/Experimental/ChaosSolvers/Public/PhysicsProxy/SingleParticlePhysicsProxy.h
deleted file mode 100644
diff --git a/Engine/Source/Runtime/Experimental/ChaosSolvers/Public/PhysicsProxy/SingleParticlePhysicsProxyFwd.h b/Engine/Source/Runtime/Experimental/ChaosSolvers/Public/PhysicsProxy/SingleParticlePhysicsProxyFwd.h
deleted file mode 100644
diff --git a/Engine/Source/Runtime/Experimental/GeometryCollectionEngine/Private/GeometryCollection/GeometryCollectionDebugDrawComponent.cpp b/Engine/Source/Runtime/Experimental/GeometryCollectionEngine/Private/GeometryCollection/GeometryCollectionDebugDrawComponent.cpp
--- a/Engine/Source/Runtime/Experimental/GeometryCollectionEngine/Private/GeometryCollection/GeometryCollectionDebugDrawComponent.cpp
+++ b/Engine/Source/Runtime/Experimental/GeometryCollectionEngine/Private/GeometryCollection/GeometryCollectionDebugDrawComponent.cpp
@@ -737,7 +737,7 @@ void UGeometryCollectionDebugDrawComponent::DebugDrawChaosTick()
 	const FGeometryCollectionPhysicsProxy* const PhysicsProxy = GeometryCollectionComponent->GetPhysicsProxy();
 	if (PhysicsProxy)
 	{
-		ParticlesData.Sync(PhysicsProxy->GetSolver(), RigidBodyIds);
+		ParticlesData.Sync(static_cast<Chaos::FPhysicsSolver*>(PhysicsProxy->GetSolver()), RigidBodyIds);
 	}
 
 	// Visualize single rigid body
