Fix Hack check.

diff --git a/Engine/Source/Runtime/Experimental/Chaos/Private/Chaos/PBDRigidsEvolutionGBF.cpp b/Engine/Source/Runtime/Experimental/Chaos/Private/Chaos/PBDRigidsEvolutionGBF.cpp
--- a/Engine/Source/Runtime/Experimental/Chaos/Private/Chaos/PBDRigidsEvolutionGBF.cpp
+++ b/Engine/Source/Runtime/Experimental/Chaos/Private/Chaos/PBDRigidsEvolutionGBF.cpp
@@ -139,7 +139,7 @@ void FPBDRigidsEvolutionGBF::Advance(const FReal Dt, const FReal MaxStepDt, cons
 //
 //
 //
-// BEGIN TOI HACK STUFF
+// BEGIN TOI
 //
 //
 //
@@ -240,7 +240,7 @@ void CCDHack(const FReal Dt, TParticleView<TPBDRigidParticles<FReal, 3>>& Partic
 //
 //
 //
-// END TOI HACK STUFF
+// END TOI
 //
 //
 //
