Merge 11257801 to Release-4.25.

Get clustering building with handles - preliminary checkin.

diff --git a/Engine/Source/Runtime/Engine/Private/PhysicsEngine/Experimental/PhysScene_Chaos.cpp b/Engine/Source/Runtime/Engine/Private/PhysicsEngine/Experimental/PhysScene_Chaos.cpp
--- a/Engine/Source/Runtime/Engine/Private/PhysicsEngine/Experimental/PhysScene_Chaos.cpp
+++ b/Engine/Source/Runtime/Engine/Private/PhysicsEngine/Experimental/PhysScene_Chaos.cpp
@@ -2117,6 +2117,8 @@ void FPhysScene_ChaosInterface::SyncBodies(Chaos::FPhysicsSolver* Solver)
 
 	Chaos::FPBDRigidActiveParticlesBufferAccessor Accessor(Solver->GetActiveParticlesBuffer());
 
+	TSet<FGeometryCollectionPhysicsProxy*> GCProxies;
+
 	const Chaos::FPBDRigidActiveParticlesBufferOut* ActiveParticleBuffer = Accessor.GetSolverOutData();
 	for (Chaos::TGeometryParticle<float, 3>* ActiveParticle : ActiveParticleBuffer->ActiveGameThreadParticles)
 	{
@@ -2160,10 +2162,26 @@ void FPhysScene_ChaosInterface::SyncBodies(Chaos::FPhysicsSolver* Solver)
 			else if(ProxyBase->GetType() == EPhysicsProxyType::GeometryCollectionType)
 			{
 				FGeometryCollectionPhysicsProxy* Proxy = static_cast<FGeometryCollectionPhysicsProxy*>(ProxyBase);
-				Proxy->PullFromPhysicsState();
+				GCProxies.Add(Proxy);
 			}
 		}
 	}
+	for (IPhysicsProxyBase* ProxyBase : ActiveParticleBuffer->PhysicsParticleProxies) 
+	{
+		if(ProxyBase->GetType() == EPhysicsProxyType::GeometryCollectionType)
+		{
+			FGeometryCollectionPhysicsProxy* Proxy = static_cast<FGeometryCollectionPhysicsProxy*>(ProxyBase);
+			GCProxies.Add(Proxy);
+		}
+		else
+		{
+			ensure(false); // Unhandled physics only particle proxy!
+		}
+	}
+	for (auto* GCProxy : GCProxies)
+	{
+		GCProxy->PullFromPhysicsState();
+	}
 	for (const FPhysScenePendingComponentTransform_Chaos& ComponentTransform : PendingTransforms)
 	{
 		if (ComponentTransform.OwningComp != nullptr)
diff --git a/Engine/Source/Runtime/Experimental/Chaos/Private/Chaos/ImplicitObjectUnion.cpp b/Engine/Source/Runtime/Experimental/Chaos/Private/Chaos/ImplicitObjectUnion.cpp
--- a/Engine/Source/Runtime/Experimental/Chaos/Private/Chaos/ImplicitObjectUnion.cpp
+++ b/Engine/Source/Runtime/Experimental/Chaos/Private/Chaos/ImplicitObjectUnion.cpp
@@ -41,6 +41,8 @@ FChaosArchive& operator<<(FChaosArchive& Ar, FLargeImplicitObjectUnionData& Larg
 	return Ar;
 }
 
+FImplicitObjectUnion::FImplicitObjectUnion() : FImplicitObject(EImplicitObject::HasBoundingBox, ImplicitObjectType::Union){}
+
 FImplicitObjectUnion::FImplicitObjectUnion(TArray<TUniquePtr<FImplicitObject>>&& Objects)
 	: FImplicitObject(EImplicitObject::HasBoundingBox, ImplicitObjectType::Union)
 	, MObjects(MoveTemp(Objects))
@@ -93,68 +95,6 @@ void FImplicitObjectUnion::FindAllIntersectingObjects(TArray < Pair<const FImpli
 	}
 }
 
-TArray<int32> FImplicitObjectUnionClustered::FindAllIntersectingChildren(const TAABB<FReal,3>& LocalBounds) const
-{
-	TArray<int32> IntersectingChildren;
-	if (LargeUnionData) //todo: make this work when hierarchy is not built
-	{
-		IntersectingChildren = LargeUnionData->Hierarchy.FindAllIntersections(LocalBounds);
-		for (int32 i = IntersectingChildren.Num() - 1; i >= 0; --i)
-		{
-			const int32 Idx = IntersectingChildren[i];
-			if (Idx < MOriginalParticleLookupHack.Num())
-			{
-				IntersectingChildren[i] = MOriginalParticleLookupHack[Idx];
-			}
-			else
-			{
-				IntersectingChildren.RemoveAtSwap(i);
-			}
-		}
-		/*for (int32& Idx : IntersectingChildren)
-		{
-			Idx = MOriginalParticleLookupHack[Idx];
-		}*/
-	}
-	else
-	{
-		IntersectingChildren = MOriginalParticleLookupHack;
-	}
-
-	return IntersectingChildren;
-}
-
-TArray<int32> FImplicitObjectUnionClustered::FindAllIntersectingChildren(const TSpatialRay<FReal,3>& LocalRay) const
-{
-	TArray<int32> IntersectingChildren;
-	if (LargeUnionData) //todo: make this work when hierarchy is not built
-	{
-		IntersectingChildren = LargeUnionData->Hierarchy.FindAllIntersections(LocalRay);
-		for (int32 i = IntersectingChildren.Num() - 1; i >= 0; --i)
-		{
-			const int32 Idx = IntersectingChildren[i];
-			if (Idx < MOriginalParticleLookupHack.Num())
-			{
-				IntersectingChildren[i] = MOriginalParticleLookupHack[Idx];
-			}
-			else
-			{
-				IntersectingChildren.RemoveAtSwap(i);
-			}
-		}
-		/*for (int32& Idx : IntersectingChildren)
-		{
-			Idx = MOriginalParticleLookupHack[Idx];
-		}*/
-	}
-	else
-	{
-		IntersectingChildren = MOriginalParticleLookupHack;
-	}
-
-	return IntersectingChildren;
-}
-
 void FImplicitObjectUnion::CacheAllImplicitObjects()
 {
 	TArray < Pair<TSerializablePtr<FImplicitObject>,FRigidTransform3>> SubObjects;
@@ -199,38 +139,71 @@ void FImplicitObjectUnion::Serialize(FChaosArchive& Ar)
 	}
 }
 
-FImplicitObjectUnion::FImplicitObjectUnion() : FImplicitObject(EImplicitObject::HasBoundingBox, ImplicitObjectType::Union){}
 
+FImplicitObjectUnionClustered::FImplicitObjectUnionClustered()
+	: FImplicitObjectUnion()
+{
+	Type = ImplicitObjectType::UnionClustered;
+}
 
-FImplicitObjectUnionClustered::FImplicitObjectUnionClustered(TArray<TUniquePtr<FImplicitObject>>&& Objects, const TArray<int32>& OriginalParticleLookupHack)
-	: FImplicitObjectUnion(MoveTemp(Objects))
+FImplicitObjectUnionClustered::FImplicitObjectUnionClustered(
+	TArray<TUniquePtr<FImplicitObject>>&& Objects, 
+	const TArray<TPBDRigidParticleHandle<FReal, 3>*>& OriginalParticleLookupHack)
+    : FImplicitObjectUnion(MoveTemp(Objects))
 	, MOriginalParticleLookupHack(OriginalParticleLookupHack)
 {
 	Type = ImplicitObjectType::UnionClustered;
-	for (int32 i = 0; i < MObjects.Num(); ++i)
+	check(MOriginalParticleLookupHack.Num() == 0 || MOriginalParticleLookupHack.Num() == MObjects.Num());
+	MCollisionParticleLookupHack.Reserve(FMath::Min(MOriginalParticleLookupHack.Num(), Objects.Num()));
+	for (int32 i = 0; MOriginalParticleLookupHack.Num() > 0 && i < MObjects.Num(); ++i)
 	{
-		check(MOriginalParticleLookupHack.Num() == 0 || MOriginalParticleLookupHack.Num() == MObjects.Num());
-		if (MOriginalParticleLookupHack.Num() > 0)
+		// This whole part sucks, only needed because of how we get union 
+		// children. Need to refactor and enforce no unions of unions.
+		if (const TImplicitObjectTransformed<FReal, 3>* Transformed = 
+			MObjects[i]->template GetObject<const TImplicitObjectTransformed<FReal, 3>>())
 		{
-			//this whole part sucks, only needed because of how we get union children. Need to refactor and enforce no unions of unions
-			if (const TImplicitObjectTransformed<FReal,3>* Transformed = MObjects[i]->template GetObject<const TImplicitObjectTransformed<FReal,3>>())
-			{
-				MCollisionParticleLookupHack.Add(Transformed->GetTransformedObject(), MOriginalParticleLookupHack[i]);
-			}
-			else
-			{
-				ensure(false);	//shouldn't be here
-			}
+			// Map const TImplicitObject<T,d>* to int32, where the latter
+			// was the RigidBodyId
+			MCollisionParticleLookupHack.Add(
+				Transformed->GetTransformedObject(), MOriginalParticleLookupHack[i]);
+		}
+		else
+		{
+			ensure(false);	//shouldn't be here
 		}
 	}
 }
 
 FImplicitObjectUnionClustered::FImplicitObjectUnionClustered(FImplicitObjectUnionClustered&& Other)
-: FImplicitObjectUnion(MoveTemp(Other))
-, MOriginalParticleLookupHack(MoveTemp(MOriginalParticleLookupHack))
-, MCollisionParticleLookupHack(MoveTemp(MCollisionParticleLookupHack))
+	: FImplicitObjectUnion(MoveTemp(Other))
+	, MOriginalParticleLookupHack(MoveTemp(MOriginalParticleLookupHack))
+	, MCollisionParticleLookupHack(MoveTemp(MCollisionParticleLookupHack))
 {
 	Type = ImplicitObjectType::UnionClustered;
 }
 
+TArray<TPBDRigidParticleHandle<FReal, 3>*>
+FImplicitObjectUnionClustered::FindAllIntersectingChildren(const TAABB<FReal, 3>& LocalBounds) const
+{
+	TArray<TPBDRigidParticleHandle<FReal, 3>*> IntersectingChildren;
+	if (LargeUnionData) //todo: make this work when hierarchy is not built
+	{
+		TArray<int32> IntersectingIndices = LargeUnionData->Hierarchy.FindAllIntersections(LocalBounds);
+		IntersectingChildren.Reserve(IntersectingIndices.Num());
+		for (const int32 Idx : IntersectingIndices)
+		{
+			if (MOriginalParticleLookupHack.IsValidIndex(Idx))
+			{
+				IntersectingChildren.Add(MOriginalParticleLookupHack[Idx]);
+			}
+		}
+	}
+	else
+	{
+		IntersectingChildren = MOriginalParticleLookupHack;
+	}
+	return IntersectingChildren;
 }
+
+
+} // namespace Chaos
diff --git a/Engine/Source/Runtime/Experimental/Chaos/Private/Chaos/PBDConstraintGraph.cpp b/Engine/Source/Runtime/Experimental/Chaos/Private/Chaos/PBDConstraintGraph.cpp
--- a/Engine/Source/Runtime/Experimental/Chaos/Private/Chaos/PBDConstraintGraph.cpp
+++ b/Engine/Source/Runtime/Experimental/Chaos/Private/Chaos/PBDConstraintGraph.cpp
@@ -36,6 +36,29 @@ FPBDConstraintGraph::FPBDConstraintGraph(const TParticleView<TGeometryParticles<
 	InitializeGraph(Particles);
 }
 
+int32 FPBDConstraintGraph::ReserveParticles(const int32 Num)
+{
+	const int32 NumFree = FreeIndexList.Num();
+	int32 NumToAdd = NumFree - Num; // Can be negative
+	if (NumToAdd > 0)
+	{
+		return 0;
+	}
+	NumToAdd = -NumToAdd; // Flip sign
+	const int32 NumNodes = Nodes.Num();
+	Nodes.SetNum(NumNodes + NumToAdd);
+	FreeIndexList.Reserve(NumToAdd);
+	for (int32 Counter = 0; Counter < NumToAdd; Counter++)
+	{
+		FreeIndexList.Push(NumNodes + Counter);
+	}
+
+	ParticleToNodeIndex.Reserve(ParticleToNodeIndex.Num() + NumToAdd);
+	Visited.Reserve(Visited.Num() + NumToAdd);
+
+	return NumToAdd;
+}
+
 /**
  * Bill added this.
  * Adds new Node to Nodes array when a new particle is created
@@ -43,16 +66,24 @@ FPBDConstraintGraph::FPBDConstraintGraph(const TParticleView<TGeometryParticles<
 
 void FPBDConstraintGraph::ParticleAdd(TGeometryParticleHandle<FReal, 3>* AddedParticle)
 {
-	int32 NewNodeIndex = GetNextNodeIndex();
-
-	FGraphNode& Node = Nodes[NewNodeIndex];
-	ensure(Node.Edges.Num() == 0);
-	ensure(Node.Island == INDEX_NONE);
-	ensure(ParticleToNodeIndex.Find(AddedParticle) == nullptr);
-
-	Node.Particle = AddedParticle->Handle();
-	ParticleToNodeIndex.Add(Node.Particle, NewNodeIndex);
-	Visited.Add(0);
+	if (ensure(!ParticleToNodeIndex.Contains(AddedParticle)))
+	{
+		const int32 NewNodeIndex = GetNextNodeIndex();
+		FGraphNode& Node = Nodes[NewNodeIndex];
+		ensure(Node.Edges.Num() == 0);
+		ensure(Node.Island == INDEX_NONE);
+
+		Node.Particle = AddedParticle->Handle();
+		check(Node.Particle);
+		ParticleToNodeIndex.Add(Node.Particle, NewNodeIndex);
+
+		const int32 NewMinNum = NewNodeIndex + 1;
+		const int32 NumToAdd = NewMinNum - Visited.Num();
+		if (NumToAdd > 0)
+			Visited.AddZeroed(NumToAdd);
+		else
+			Visited[NewNodeIndex] = 0;
+	}
 }
 
 /**
@@ -64,7 +95,7 @@ void FPBDConstraintGraph::ParticleRemove(TGeometryParticleHandle<FReal, 3>* Remo
 {
 	if (ParticleToNodeIndex.Contains(RemovedParticle))
 	{
-		int32 NodeIdx = ParticleToNodeIndex[RemovedParticle];
+		const int32 NodeIdx = ParticleToNodeIndex[RemovedParticle];
 		FreeIndexList.Push(NodeIdx);
 
 		FGraphNode& NodeRemoved = Nodes[NodeIdx];
@@ -131,7 +162,13 @@ void FPBDConstraintGraph::InitializeGraph(const TParticleView<TGeometryParticles
 	}
 	else
 	{
-		ensure(NumNonDisabledParticles <= Nodes.Num());
+		if (!ensure(NumNonDisabledParticles <= Nodes.Num()))
+		{
+			for (auto& Particle : Particles)
+			{
+				ParticleAdd(Particle.Handle());
+			}
+		}
 
 		// Update nodes may contain duplicate entries. To process in parallel we either need to prevent that, 
 		// or ensure we only process the first entry of a dupe. We are doing the latter for now...
@@ -208,6 +245,15 @@ void FPBDConstraintGraph::AddConstraint(const uint32 InContainerId, FConstraintH
 	FGraphEdge NewEdge;
 	NewEdge.Data = { InContainerId, InConstraintHandle };
 
+	if (ConstrainedParticles[0] && !ParticleToNodeIndex.Contains(ConstrainedParticles[0]))
+	{
+		ParticleAdd(ConstrainedParticles[0]);
+	}
+	if (ConstrainedParticles[1] && !ParticleToNodeIndex.Contains(ConstrainedParticles[1]))
+	{
+		ParticleAdd(ConstrainedParticles[1]);
+	}
+
 	int32* PNodeIndex0 = (ConstrainedParticles[0])? ParticleToNodeIndex.Find(ConstrainedParticles[0]) : nullptr;
 	int32* PNodeIndex1 = (ConstrainedParticles[1])? ParticleToNodeIndex.Find(ConstrainedParticles[1]) : nullptr;
 	if (ensure(PNodeIndex0 || PNodeIndex1))
@@ -244,6 +290,10 @@ void FPBDConstraintGraph::UpdateIslands(const TParticleView<TPBDRigidParticles<F
 	for (auto& PBDRigid : PBDRigids)
 	{
 		PBDRigid.Island() = INDEX_NONE;
+		if (!ensure(ParticleToNodeIndex.Contains(PBDRigid.Handle())))
+		{
+			ParticleAdd(PBDRigid.Handle());
+		}
 	}
 	ComputeIslands(PBDRigids, Particles);
 }
@@ -279,7 +329,8 @@ void FPBDConstraintGraph::ComputeIslands(const TParticleView<TPBDRigidParticles<
 
 	for (auto& Particle : PBDRigids)
 	{
-		int32 Idx = ParticleToNodeIndex[Particle.Handle()];
+		auto* ParticleHandle = Particle.Handle();
+		int32 Idx = ParticleToNodeIndex[ParticleHandle];  // ryan - FAILS!
 		// selective reset of islands, don't reset if has been visited due to being edge connected to earlier processed node
 		if (Visited[Idx] && Visited[Idx] != VisitToken)
 		{
diff --git a/Engine/Source/Runtime/Experimental/Chaos/Private/Chaos/PBDRigidClustering.cpp b/Engine/Source/Runtime/Experimental/Chaos/Private/Chaos/PBDRigidClustering.cpp
--- a/Engine/Source/Runtime/Experimental/Chaos/Private/Chaos/PBDRigidClustering.cpp
+++ b/Engine/Source/Runtime/Experimental/Chaos/Private/Chaos/PBDRigidClustering.cpp
@@ -113,7 +113,13 @@ namespace Chaos
 
 	DECLARE_CYCLE_STAT(TEXT("TPBDRigidClustering<>::RewindAndEvolve<BGF>()"), STAT_RewindAndEvolve_BGF, STATGROUP_Chaos);
 	template<class T, int d>
-	void RewindAndEvolve(FPBDRigidsEvolutionGBF& Evolution, TPBDRigidClusteredParticles<T, d>& InParticles, const TSet<int32>& IslandsToRecollide, const TSet<uint32>& AllActivatedChildren, const T Dt, TPBDCollisionConstraints<T, d>& CollisionRule)
+	void RewindAndEvolve(
+		FPBDRigidsEvolutionGBF& Evolution, 
+		TPBDRigidClusteredParticles<T, d>& InParticles, 
+		const TSet<int32>& IslandsToRecollide, 
+		const TSet<TPBDRigidParticleHandle<T, d>*> AllActivatedChildren,
+		const T Dt, 
+		TPBDCollisionConstraints<T, d>& CollisionRule)
 	{
 		SCOPE_CYCLE_COUNTER(STAT_RewindAndEvolve_BGF);
 		// Rewind active particles
@@ -189,51 +195,51 @@ namespace Chaos
 
 	DECLARE_CYCLE_STAT(TEXT("TPBDRigidClustering<>::UpdateClusterMassProperties()"), STAT_UpdateClusterMassProperties, STATGROUP_Chaos);
 	template<class T, int d>
-	void UpdateClusterMassProperties(TPBDRigidClusteredParticles<T, d>& Particles, const TArray<uint32>& Children, const uint32 NewIndex, const TRigidTransform<T, d>* ForceMassOrientation, const TArrayCollectionArray<TUniquePtr<TMultiChildProxyData<T, d>>>* MMultiChildProxyData, const TArrayCollectionArray<FMultiChildProxyId>* MMultiChildProxyId)
+	void UpdateClusterMassProperties(
+		Chaos::TPBDRigidClusteredParticleHandle<float, 3>* Parent, 
+		TSet<TPBDRigidParticleHandle<T, d>*>& Children, 
+		const TRigidTransform<T, d>* ForceMassOrientation)
 	{
 		SCOPE_CYCLE_COUNTER(STAT_UpdateClusterMassProperties);
-		check(0 <= NewIndex && NewIndex < Particles.Size());
 		check(Children.Num());
 
-		Particles.X(NewIndex) = TVector<T, d>(0);
-		Particles.R(NewIndex) = TRotation<T, d>(FQuat::MakeFromEuler(TVector<T, d>(0)));
-		Particles.V(NewIndex) = TVector<T, d>(0);
-		Particles.W(NewIndex) = TVector<T, d>(0);
-		Particles.M(NewIndex) = 0;
-		Particles.I(NewIndex) = PMatrix<T, d, d>(0);
+		Parent->SetX(TVector<T, d>(0));
+		Parent->SetR(TRotation<T, d>(FQuat::MakeFromEuler(TVector<T, d>(0))));
+		Parent->SetV(TVector<T, d>(0));
+		Parent->SetW(TVector<T, d>(0));
+		Parent->SetM(0);
+		Parent->SetI(PMatrix<T, d, d>(0));
 
-		auto GetMultiChildProxy = [MMultiChildProxyId, MMultiChildProxyData](const int32 OriginalChild, const TMultiChildProxyData<T, d>*& ProxyData) -> int32
+		bool bHasChild = false;
+		bool bHasProxyChild = false;
+		for (TPBDRigidParticleHandle<T, d>* OriginalChild : Children)
 		{
-			const int32 MultiChildProxyId = MassPropertiesFromMultiChildProxy && MMultiChildProxyId ? (*MMultiChildProxyId)[OriginalChild].Id : INDEX_NONE;
-			if (MultiChildProxyId != INDEX_NONE)
+			FMultiChildProxyId MultiChildProxyId; // sizeof(FMultiChildProxyId) = sizeof(void*), so copy
+			TMultiChildProxyData<T, d>* ProxyData = nullptr;
+			if (TPBDRigidClusteredParticleHandle<T, d>* ClusteredOriginalChild = OriginalChild->CastToClustered())
 			{
-				ProxyData = (*MMultiChildProxyData)[MultiChildProxyId].Get();
-				return ProxyData ? MultiChildProxyId : INDEX_NONE;
+				MultiChildProxyId = ClusteredOriginalChild->MultiChildProxyId();
+				ProxyData = ClusteredOriginalChild->MultiChildProxyData().Get();
 			}
-			return INDEX_NONE;
-		};
 
-		bool bHasChild = false;
-		bool bHasProxyChild = false;
-		for (const auto OriginalChild : Children)
-		{
-			const TMultiChildProxyData<T, d>* ProxyData;
-			const int32 MultiChildProxyId = GetMultiChildProxy(OriginalChild, ProxyData);
+			//int32 Child;
+			TPBDRigidParticleHandle<T, d>* Child;
 			TVector<T, d> ChildPosition;
 			TRotation<T, d> ChildRotation;
-			int32 Child;
-
-			if (MultiChildProxyId == INDEX_NONE)
+			if (MultiChildProxyId.Id == nullptr)
 			{
 				Child = OriginalChild;
-				ChildPosition = Particles.X(Child);
-				ChildRotation = Particles.R(Child);
+				ChildPosition = Child->X();
+				ChildRotation = Child->R();
 			}
 			else if (ProxyData->KeyChild == OriginalChild)
 			{
-				Child = MultiChildProxyId;
+				Child = MultiChildProxyId.Id;
 
-				const TRigidTransform<T, d> ProxyWorldTM = ProxyData->RelativeToKeyChild * TRigidTransform<T, d>(Particles.X(OriginalChild), Particles.R(OriginalChild));
+				const TRigidTransform<T, d> ProxyWorldTM =
+					ProxyData->RelativeToKeyChild *
+					TRigidTransform<T, d>(
+						OriginalChild->X(), OriginalChild->R());
 				ChildPosition = ProxyWorldTM.GetLocation();
 				ChildRotation = ProxyWorldTM.GetRotation();
 			}
@@ -242,100 +248,105 @@ namespace Chaos
 				continue; //using a proxy but we are not the key child
 			}
 
-			const auto& ChildMass = Particles.M(Child);
-			const PMatrix<T, d, d> ChildWorldSpaceI = (ChildRotation * FMatrix::Identity) * Particles.I(Child) * (ChildRotation * FMatrix::Identity).GetTransposed();
+			const T ChildMass = Child->M();
+			const PMatrix<T, d, d> ChildWorldSpaceI = 
+				(ChildRotation * FMatrix::Identity) * Child->I() * (ChildRotation * FMatrix::Identity).GetTransposed();
 			if (ChildWorldSpaceI.ContainsNaN())
 			{
 				continue;
 			}
 			bHasProxyChild = true;
 			bHasChild = true;
-			Particles.I(NewIndex) += ChildWorldSpaceI;
-			Particles.M(NewIndex) += ChildMass;
-			Particles.X(NewIndex) += ChildPosition * ChildMass;
-			Particles.V(NewIndex) += Particles.V(OriginalChild) * ChildMass; //use original child for veocities because we don't simulate proxy
-			Particles.W(NewIndex) += Particles.W(OriginalChild) * ChildMass;
+			bHasProxyChild = true;
+			Parent->I() += ChildWorldSpaceI;
+			Parent->M() += ChildMass;
+			Parent->X() += ChildPosition * ChildMass;
+			Parent->V() += OriginalChild->V() * ChildMass; // Use orig child for vel because we don't sim the proxy
+			Parent->W() += OriginalChild->W() * ChildMass;
 		}
 		if (!ensure(bHasProxyChild))
 		{
-			for (const auto OriginalChild : Children)
+			for (TPBDRigidParticleHandle<T, d>* OriginalChild : Children)
 			{
-				TVector<T, d> ChildPosition;
-				TRotation<T, d> ChildRotation;
-				int32 Child;
-
-				Child = OriginalChild;
-				ChildPosition = Particles.X(Child);
-				ChildRotation = Particles.R(Child);
+				TPBDRigidParticleHandle<T, d>* Child = OriginalChild;
+				const TVector<T, d>& ChildPosition = Child->X();
+				const TRotation<T, d>& ChildRotation = Child->R();
+				const T ChildMass = Child->M();
 
-				const auto& ChildMass = Particles.M(Child);
-				const PMatrix<T, d, d> ChildWorldSpaceI = (ChildRotation * FMatrix::Identity) * Particles.I(Child) * (ChildRotation * FMatrix::Identity).GetTransposed();
+				const PMatrix<T, d, d> ChildWorldSpaceI = 
+					(ChildRotation * FMatrix::Identity) * Child->I() * (ChildRotation * FMatrix::Identity).GetTransposed();
 				if (ChildWorldSpaceI.ContainsNaN())
 				{
 					continue;
 				}
 				bHasChild = true;
-				Particles.I(NewIndex) += ChildWorldSpaceI;
-				Particles.M(NewIndex) += ChildMass;
-				Particles.X(NewIndex) += ChildPosition * ChildMass;
-				Particles.V(NewIndex) += Particles.V(OriginalChild) * ChildMass;	//use original child for veocities because we don't simulate proxy
-				Particles.W(NewIndex) += Particles.W(OriginalChild) * ChildMass;
+				Parent->I() += ChildWorldSpaceI;
+				Parent->M() += ChildMass;
+				Parent->X() += ChildPosition * ChildMass;
+				Parent->V() += OriginalChild->V() * ChildMass; // Use orig child for vel because we don't sim the proxy
+				Parent->W() += OriginalChild->W() * ChildMass;
 			}
 		}
 		for (int32 i = 0; i < d; i++)
 		{
-			const PMatrix<T, d, d>& InertiaTensor = Particles.I(NewIndex);
+			const PMatrix<T, d, d>& InertiaTensor = Parent->I();
 			if (InertiaTensor.GetColumn(i)[i] < SMALL_NUMBER)
 			{
-				Particles.I(NewIndex) = PMatrix<T, d, d>(1.f, 1.f, 1.f);
+				Parent->SetI(PMatrix<T, d, d>(1.f, 1.f, 1.f));
 				break;
 			}
 		}
 
-		if (!ensure(bHasChild) || !ensure(Particles.M(NewIndex) > SMALL_NUMBER))
-		{
-			Particles.M(NewIndex) = 1;
-			Particles.X(NewIndex) = TVector<T, d>(0);
-			Particles.V(NewIndex) = TVector<T, d>(0);
-			Particles.PreV(NewIndex) = Particles.V(NewIndex);
-			Particles.InvM(NewIndex) = 1;
-			Particles.P(NewIndex) = Particles.X(NewIndex);
-			Particles.W(NewIndex) = TVector<T, d>(0);
-			Particles.PreW(NewIndex) = Particles.W(NewIndex);
-			Particles.R(NewIndex) = TRotation<T, d>(FMatrix::Identity);
-			Particles.Q(NewIndex) = Particles.R(NewIndex);
-			Particles.I(NewIndex) = FMatrix::Identity;
-			Particles.InvI(NewIndex) = FMatrix::Identity;
+		if (!ensure(bHasChild) || !ensure(Parent->M() > SMALL_NUMBER))
+		{
+			Parent->M() = 1.0;
+			Parent->X() = TVector<T, d>(0);
+			Parent->V() = TVector<T, d>(0);
+			Parent->PreV() = Parent->V();
+			Parent->InvM() = 1;
+			Parent->P() = Parent->X();
+			Parent->W() = TVector<T, d>(0);
+			Parent->PreW() = Parent->W();
+			Parent->R() = TRotation<T, d>(FMatrix::Identity);
+			Parent->Q() = Parent->R();
+			Parent->I() = FMatrix::Identity;
+			Parent->InvI() = FMatrix::Identity;
 			return;
 		}
 
-		check(Particles.M(NewIndex) > SMALL_NUMBER);
+		check(Parent->M() > SMALL_NUMBER);
 
-		Particles.X(NewIndex) /= Particles.M(NewIndex);
-		Particles.V(NewIndex) /= Particles.M(NewIndex);
-		Particles.PreV(NewIndex) = Particles.V(NewIndex);
-		Particles.InvM(NewIndex) = 1. / Particles.M(NewIndex);
+		Parent->X() /= Parent->M();
+		Parent->V() /= Parent->M();
+		Parent->PreV() = Parent->V();
+		Parent->InvM() = 1. / Parent->M();
 		if (ForceMassOrientation)
 		{
-			Particles.X(NewIndex) = ForceMassOrientation->GetLocation();
+			Parent->X() = ForceMassOrientation->GetLocation();
 		}
-		Particles.P(NewIndex) = Particles.X(NewIndex);
-		for (const auto OriginalChild : Children)
+		Parent->P() = Parent->X();
+		for (TPBDRigidParticleHandle<T, d>* OriginalChild : Children)
 		{
-			int32 Child;
-			const TMultiChildProxyData<T, d>* ProxyData = nullptr;
-			const int32 MultiChildProxyId = bHasProxyChild ? GetMultiChildProxy(OriginalChild, ProxyData) : INDEX_NONE;
+			FMultiChildProxyId MultiChildProxyId; // sizeof(FMultiChildProxyId) = sizeof(void*), so copy
+			TMultiChildProxyData<T, d>* ProxyData = nullptr;
+			if (TPBDRigidClusteredParticleHandle<T, d>* ClusteredOriginalChild = OriginalChild->CastToClustered())
+			{
+				MultiChildProxyId = ClusteredOriginalChild->MultiChildProxyId();
+				ProxyData = ClusteredOriginalChild->MultiChildProxyData().Get();
+			}
 
+			TPBDRigidParticleHandle<T, d>* Child;
 			TVector<T, d> ChildPosition;
-			if (MultiChildProxyId == INDEX_NONE)
+			if (MultiChildProxyId.Id == nullptr)
 			{
 				Child = OriginalChild;
-				ChildPosition = Particles.X(Child);
+				ChildPosition = Child->X();
 			}
 			else if (ProxyData->KeyChild == OriginalChild)
 			{
-				Child = MultiChildProxyId;
-				const TRigidTransform<T, d> ProxyWorldTM = ProxyData->RelativeToKeyChild * TRigidTransform<T, d>(Particles.X(OriginalChild), Particles.R(OriginalChild));
+				Child = MultiChildProxyId.Id;
+				const TRigidTransform<T, d> ProxyWorldTM = 
+					ProxyData->RelativeToKeyChild * TRigidTransform<T, d>(OriginalChild->X(), OriginalChild->R());
 				ChildPosition = ProxyWorldTM.GetLocation();
 			}
 			else
@@ -343,40 +354,49 @@ namespace Chaos
 				continue; //using a proxy but we are not the key child
 			}
 
-			TVector<T, d> ParentToChild = ChildPosition - Particles.X(NewIndex);
+			TVector<T, d> ParentToChild = ChildPosition - Parent->X();
 
-			const auto& ChildMass = Particles.M(Child);
-			Particles.W(NewIndex) += TVector<T, d>::CrossProduct(ParentToChild, Particles.V(OriginalChild) * ChildMass); //taking v from original child since we are not actually simulating the proxy child
+			const T ChildMass = Child->M();
+			// taking v from original child since we are not actually simulating the proxy child
+			Parent->W() += 
+				TVector<T, d>::CrossProduct(ParentToChild, 
+					OriginalChild->V() * ChildMass);
 			{
 				const T& p0 = ParentToChild[0];
 				const T& p1 = ParentToChild[1];
 				const T& p2 = ParentToChild[2];
-				const T& m = Particles.M(Child);
-				Particles.I(NewIndex) += PMatrix<T, d, d>(m * (p1 * p1 + p2 * p2), -m * p1 * p0, -m * p2 * p0, m * (p2 * p2 + p0 * p0), -m * p2 * p1, m * (p1 * p1 + p0 * p0));
+				const T& m = ChildMass;
+				Parent->I() += 
+					PMatrix<T, d, d>(
+						m * (p1 * p1 + p2 * p2), -m * p1 * p0, -m * p2 * p0, 
+						m * (p2 * p2 + p0 * p0), -m * p2 * p1, m * (p1 * p1 + p0 * p0));
 			}
 		}
-		if (Particles.I(NewIndex).ContainsNaN())
+		PMatrix<T, d, d>& InertiaTensor = Parent->I();
+		if (Parent->I().ContainsNaN())
 		{
-			Particles.I(NewIndex) = PMatrix<T, d, d>(1.f, 1.f, 1.f);
+			InertiaTensor = PMatrix<T, d, d>(1.f, 1.f, 1.f);
 		}
-		for (int32 i = 0; i < d; i++)
+		else
 		{
-			const PMatrix<T, d, d>& InertiaTensor = Particles.I(NewIndex);
-			if (InertiaTensor.GetColumn(i)[i] < SMALL_NUMBER)
+			for (int32 i = 0; i < d; i++)
 			{
-				Particles.I(NewIndex) = PMatrix<T, d, d>(1.f, 1.f, 1.f);
-				break;
+				if (InertiaTensor.GetColumn(i)[i] < SMALL_NUMBER)
+				{
+					InertiaTensor = PMatrix<T, d, d>(1.f, 1.f, 1.f);
+					break;
+				}
 			}
 		}
-		Particles.W(NewIndex) /= Particles.M(NewIndex);
-		Particles.PreW(NewIndex) = Particles.W(NewIndex);
-		Particles.R(NewIndex) = Chaos::TransformToLocalSpace<T, d>(Particles.I(NewIndex));
+		Parent->W() /= Parent->M();
+		Parent->PreW() = Parent->W();
+		Parent->R() = Chaos::TransformToLocalSpace<T, d>(InertiaTensor);
 		if (ForceMassOrientation)
 		{
-			Particles.R(NewIndex) = ForceMassOrientation->GetRotation();
+			Parent->R() = ForceMassOrientation->GetRotation();
 		}
-		Particles.Q(NewIndex) = Particles.R(NewIndex);
-		Particles.InvI(NewIndex) = Particles.I(NewIndex).Inverse();
+		Parent->Q() = Parent->R();
+		Parent->InvI() = Parent->I().Inverse();
 	}
 
 	//==========================================================================
@@ -384,97 +404,112 @@ namespace Chaos
 	//==========================================================================
 
 	template<class FPBDRigidsEvolution, class FPBDCollisionConstraint, class T, int d>
-	TPBDRigidClustering<FPBDRigidsEvolution, FPBDCollisionConstraint, T, d>::TPBDRigidClustering(FPBDRigidsEvolution& InEvolution, TPBDRigidClusteredParticles<T, d>& InParticles)
+	TPBDRigidClustering<FPBDRigidsEvolution, FPBDCollisionConstraint, T, d>::TPBDRigidClustering(
+		FPBDRigidsEvolution& InEvolution, 
+		TPBDRigidClusteredParticles<T, d>& InParticles)
 		: MEvolution(InEvolution)
 		, MParticles(InParticles)
 		, MCollisionImpulseArrayDirty(true)
 		, DoGenerateBreakingData(false)
 		, MClusterConnectionFactor(1.0)
 		, MClusterUnionConnectionType(FClusterCreationParameters<T>::EConnectionMethod::DelaunayTriangulation)
-	{
-	}
+	{}
 
 	template<class FPBDRigidsEvolution, class FPBDCollisionConstraint, class T, int d>
 	TPBDRigidClustering<FPBDRigidsEvolution, FPBDCollisionConstraint, T, d>::~TPBDRigidClustering()
-	{
-	}
+	{}
 
 	DECLARE_CYCLE_STAT(TEXT("TPBDRigidClustering<>::CreateClusterParticle"), STAT_CreateClusterParticle, STATGROUP_Chaos);
 	template<class FPBDRigidsEvolution, class FPBDCollisionConstraint, class T, int d>
-	int32 TPBDRigidClustering<FPBDRigidsEvolution, FPBDCollisionConstraint, T, d>::CreateClusterParticle(
-		int32 ClusterGroupIndex,
-		TArray<uint32>&& Children,
-		/*const*/ TSerializablePtr<FImplicitObject> ProxyGeometry,
-		const TRigidTransform<T, d>* ForceMassOrientation,
-		const FClusterCreationParameters<T>& Parameters)
+	Chaos::TPBDRigidClusteredParticleHandle<float, 3>* 
+	TPBDRigidClustering<FPBDRigidsEvolution, FPBDCollisionConstraint, T, d>::CreateClusterParticle(
+		const int32 ClusterGroupIndex,
+		TArray<Chaos::TPBDRigidParticleHandle<T,d>*>&& Children,
+		const FClusterCreationParameters<T>& Parameters,
+		TSerializablePtr<FImplicitObject> ProxyGeometry,
+		const TRigidTransform<T, d>* ForceMassOrientation)
 	{
 		SCOPE_CYCLE_COUNTER(STAT_CreateClusterParticle);
 
-		FClusterMap& MParentToChildren = GetChildrenMap();
-
-		int32 NewIndex = Parameters.RigidBodyIndex;
-		if (NewIndex == INDEX_NONE)
+		Chaos::TPBDRigidClusteredParticleHandle<float, 3>* NewParticle = Parameters.ClusterParticleHandle;
+		if (!NewParticle)
 		{
-			NewIndex = MParticles.Size();
-			MEvolution.CreateClusteredParticles(1);
+			//NewParticle = MEvolution.CreateClusteredParticles(1)[0]; // calls Evolution.DirtyParticle()
+			NewParticle = MEvolution.CreateClusteredParticles(1)[0]; // calls Evolution.DirtyParticle()
+		}
+		if (!NewParticle->UniqueIdx().IsValid())
+		{
+			NewParticle->SetUniqueIdx(MEvolution.GenerateUniqueIdx());
 		}
-		TGeometryParticleHandle<T, d>* NewParticle = MParticles.Handle(NewIndex);
-
-		MEvolution.EnableParticle(NewParticle, MParticles.Handle(Children[0]));
-		MParticles.CollisionGroup(NewIndex) = INT_MAX;
-		TopLevelClusterParents.Add(NewIndex);
 
-		MParticles.InternalCluster(NewIndex) = false;
-		MParticles.ClusterIds(NewIndex) = ClusterId(INDEX_NONE, Children.Num());
-		MParticles.ClusterGroupIndex(NewIndex) = ClusterGroupIndex;
+		// Must do this so that the constraint graph knows about this particle 
+		// prior to calling CreateIslands().  We could call MEvolution.CreateParticle()
+		// which does the same thing, but also calls DirtyParticle(), which is already
+		// done by MEvolution.CreateClusteredParticles(), and will be done again by
+		// MEvolution.EnableParticle().
+		//MEvolution.GetConstraintGraph().AddParticle(NewParticle); // PBDRigidsEvolutionGBF protects GetConstraintGraph().  Bah!
+		MEvolution.CreateParticle(NewParticle); // Doesn't create, just adds to constraint graph
+		MEvolution.EnableParticle(NewParticle, nullptr); // null for parent skips constraint graph EnableParticle()
+		NewParticle->SetCollisionGroup(INT_MAX);
+		TopLevelClusterParents.Add(NewParticle);
+
+		NewParticle->SetInternalCluster(false);
+		NewParticle->SetClusterId(ClusterId(nullptr, Children.Num()));
+		NewParticle->SetClusterGroupIndex(ClusterGroupIndex);
+		NewParticle->SetStrains(0.0);
 
 		//
 		// Update clustering data structures.
 		//
-		if (MParentToChildren.Contains(NewParticle))
+
+		if (MChildren.Contains(NewParticle))
 		{
-			MParentToChildren[NewParticle] = MoveTemp(Children);
+			MChildren[NewParticle] = MoveTemp(Children);
 		}
 		else
 		{
-			MParentToChildren.Add(NewParticle, MoveTemp(Children));
+			MChildren.Add(NewParticle, MoveTemp(Children));
 		}
 
-		bool bClusterIsAsleep = true;
-		MParticles.Strains(NewIndex) = 0.0;
+		const TArray<TPBDRigidParticleHandle<T, d>*>& ChildrenArray = MChildren[NewParticle];
+		TSet<TPBDRigidParticleHandle<T, d>*> ChildrenSet(ChildrenArray);
+		MEvolution.DisableParticles(reinterpret_cast<TSet<TGeometryParticleHandle<T, d>*>&>(ChildrenSet));
 
-		const TArray<uint32>& ChildrenArray = MParentToChildren[NewParticle];
-		//todo: avoid this iteration by passing a view into DisableParticles
-		TSet<TGeometryParticleHandle<T, d>*> ChildrenHandles;
-		for (auto Child : ChildrenArray)
-		{
-			ChildrenHandles.Add(MParticles.Handle(Child));
-		}
-		MEvolution.DisableParticles(ChildrenHandles);
-		for (const uint32 Child : ChildrenArray)
+		bool bClusterIsAsleep = true;
+		for (TPBDRigidParticleHandle<T, d>* Child : ChildrenSet)
 		{
-			TopLevelClusterParents.Remove(Child);
-			bClusterIsAsleep &= MParticles.Sleeping(Child);
-			MParticles.ClusterGroupIndex(Child) = 0;
-			MParticles.ClusterIds(Child).Id = NewIndex;
-			MParticles.Strains(NewIndex) += MParticles.Strains(Child);
+			bClusterIsAsleep &= Child->Sleeping();
 
-			MParticles.CollisionImpulses(NewIndex) = FMath::Max(MParticles.CollisionImpulses(NewIndex), MParticles.CollisionImpulses(Child));
-			MParticles.CollisionGroup(NewIndex) = (MParticles.CollisionGroup(NewIndex) < MParticles.CollisionGroup(Child)) ? MParticles.CollisionGroup(NewIndex) : MParticles.CollisionGroup(Child);
+			if (TPBDRigidClusteredParticleHandle<T, d>* ClusteredChild = Child->CastToClustered())
+			{
+				TopLevelClusterParents.Remove(ClusteredChild);
+
+				// Cluster group id 0 means "don't union with other things"
+				// TODO: Use INDEX_NONE instead of 0?
+				ClusteredChild->SetClusterGroupIndex(0);
+				ClusteredChild->ClusterIds().Id = NewParticle; // ClusterId needs to change?
+				NewParticle->Strains() += ClusteredChild->Strains();
+
+				NewParticle->SetCollisionImpulses(
+					FMath::Max(NewParticle->CollisionImpulses(), ClusteredChild->CollisionImpulses()));
+
+				const int32 NewCG = NewParticle->CollisionGroup();
+				const int32 ChildCG = ClusteredChild->CollisionGroup();
+				NewParticle->SetCollisionGroup(NewCG < ChildCG ? NewCG : ChildCG);
+			}
 		}
-		if (ChildrenArray.Num())
+		if (ChildrenSet.Num())
 		{
-			MParticles.Strains(NewIndex) /= ChildrenArray.Num();
+			NewParticle->Strains() /= ChildrenSet.Num();
 		}
 
 		ensureMsgf(!ProxyGeometry || ForceMassOrientation, TEXT("If ProxyGeometry is passed, we must override the mass orientation as they are tied"));
 
-		UpdateMassProperties(ChildrenArray, NewIndex, ForceMassOrientation);
-		UpdateGeometry(ChildrenArray, NewIndex, ProxyGeometry, Parameters);
-		GenerateConnectionGraph(NewIndex, Parameters);
-		MParticles.SetSleeping(NewIndex, bClusterIsAsleep);
-
-		return NewIndex;
+		UpdateMassProperties(NewParticle, ChildrenSet, ForceMassOrientation);
+		UpdateGeometry(NewParticle, ChildrenSet, ProxyGeometry, Parameters);
+		GenerateConnectionGraph(NewParticle, Parameters);
+		NewParticle->SetSleeping(bClusterIsAsleep);
+		return NewParticle;
 	}
 
 	int32 UnionsHaveCollisionParticles = 0;
@@ -482,78 +517,85 @@ namespace Chaos
 
 	DECLARE_CYCLE_STAT(TEXT("TPBDRigidClustering<>::CreateClusterParticleFromClusterChildren"), STAT_CreateClusterParticleFromClusterChildren, STATGROUP_Chaos);
 	template<class FPBDRigidsEvolution, class FPBDCollisionConstraint, class T, int d>
-	int32 TPBDRigidClustering<FPBDRigidsEvolution, FPBDCollisionConstraint, T, d>::CreateClusterParticleFromClusterChildren(TArray<uint32>&& Children, TPBDRigidClusteredParticleHandle<T,d>* Parent, const TRigidTransform<T, d>& ClusterWorldTM, const FClusterCreationParameters<T>& Parameters)
+	Chaos::TPBDRigidClusteredParticleHandle<float, 3>* 
+	TPBDRigidClustering<FPBDRigidsEvolution, FPBDCollisionConstraint, T, d>::CreateClusterParticleFromClusterChildren(
+		TArray<TPBDRigidParticleHandle<T,d>*>&& Children, 
+		TPBDRigidClusteredParticleHandle<T,d>* Parent, 
+		const TRigidTransform<T, d>& ClusterWorldTM, 
+		const FClusterCreationParameters<T>& Parameters)
 	{
 		SCOPE_CYCLE_COUNTER(STAT_CreateClusterParticleFromClusterChildren);
-		FClusterMap& MParentToChildren = GetChildrenMap();
 
 		//This cluster is made up of children that are currently in a cluster. This means we don't need to update or disable as much
-		int32 NewIndex = Parameters.RigidBodyIndex;
-		if (NewIndex == INDEX_NONE)
+		Chaos::TPBDRigidClusteredParticleHandle<float, 3>* NewParticle = Parameters.ClusterParticleHandle;
+		if (!NewParticle)
+		{
+			NewParticle = MEvolution.CreateClusteredParticles(1)[0]; // calls Evolution.DirtyParticle()
+		}
+		if (!NewParticle->UniqueIdx().IsValid())
 		{
-			NewIndex = MParticles.Size();
-			MEvolution.CreateClusteredParticles(1);
+			NewParticle->SetUniqueIdx(MEvolution.GenerateUniqueIdx());
 		}
-		TGeometryParticleHandle<T, d>* NewParticle = MParticles.Handle(NewIndex);
-		MEvolution.EnableParticle(MParticles.Handle(NewIndex), Parent);
+		MEvolution.CreateParticle(NewParticle);
+		MEvolution.EnableParticle(NewParticle, Parent);
 
-		MParticles.CollisionGroup(NewIndex) = INT_MAX;
-		TopLevelClusterParents.Add(NewIndex);
-		MParticles.InternalCluster(NewIndex) = true;
-		MParticles.ClusterIds(NewIndex) = ClusterId(INDEX_NONE, Children.Num());
+		NewParticle->SetCollisionGroup(INT_MAX);
+		TopLevelClusterParents.Add(NewParticle);
+		NewParticle->SetInternalCluster(true);
+		NewParticle->SetClusterId(ClusterId(Parent, Children.Num()));
 
 		//
 		// Update clustering data structures.
 		//
-		if (MParentToChildren.Contains(NewParticle))
+		if (MChildren.Contains(NewParticle))
 		{
-			MParentToChildren[NewParticle] = MoveTemp(Children);
+			MChildren[NewParticle] = MoveTemp(Children);
 		}
 		else
 		{
-			MParentToChildren.Add(NewParticle, MoveTemp(Children));
+			MChildren.Add(NewParticle, MoveTemp(Children));
 		}
 
-		const TArray<uint32>& ChildrenArray = MParentToChildren[NewParticle];
+		TArray<TPBDRigidParticleHandle<T, d>*>& ChildrenArray = MChildren[NewParticle];
 		//child transforms are out of date, need to update them. @todo(ocohen): if children transforms are relative we would not need to update this, but would simply have to do a final transform on the new cluster index
 		// TODO(mlentine): Why is this not needed? (Why is it ok to have DeactivateClusterChildren==false?)
 		if (DeactivateClusterChildren)
 		{
 			//TODO: avoid iteration just pass in a view
-			//TSet<uint32> ChildrenSet = TSet<uint32>(Children);	// @todo(ccaulfield): try to eliminate set creation (required by RemoveConstraints)
-			TSet<TGeometryParticleHandle<T, d>*> ChildrenHandles;
-			for (auto Child : ChildrenArray)
-			{
-				ChildrenHandles.Add(MParticles.Handle(Child));
-			}
+			TSet<TGeometryParticleHandle<T, d>*> ChildrenHandles(static_cast<TArray<TGeometryParticleHandle<T,d>*>>(ChildrenArray));
 			MEvolution.DisableParticles(ChildrenHandles);
 		}
-		for (uint32 Child : ChildrenArray)
+		for (TPBDRigidParticleHandle<T, d>* Child : ChildrenArray)
 		{
-			TRigidTransform<T, d> ChildFrame = MParticles.ChildToParent(Child) * ClusterWorldTM;
-			MParticles.X(Child) = ChildFrame.GetTranslation();
-			MParticles.R(Child) = ChildFrame.GetRotation();
-			MParticles.ClusterIds(Child).Id = NewIndex;
-			MParticles.ClusterGroupIndex(Child) = 0;
-			if (DeactivateClusterChildren)
+			if (TPBDRigidClusteredParticleHandle<T, d>* ClusteredChild = Child->CastToClustered())
 			{
-				TopLevelClusterParents.Remove(Child);
-			}
+				TRigidTransform<T, d> ChildFrame = ClusteredChild->ChildToParent() * ClusterWorldTM;
+				ClusteredChild->SetX(ChildFrame.GetTranslation());
+				ClusteredChild->SetR(ChildFrame.GetRotation());
+				ClusteredChild->ClusterIds().Id = NewParticle;
+				ClusteredChild->SetClusterGroupIndex(0);
+				if (DeactivateClusterChildren)
+				{
+					TopLevelClusterParents.Remove(ClusteredChild);
+				}
 
-			MParticles.CollisionImpulses(NewIndex) = FMath::Max(MParticles.CollisionImpulses(NewIndex), MParticles.CollisionImpulses(Child));
-			MParticles.CollisionGroup(NewIndex) = (MParticles.CollisionGroup(NewIndex) < MParticles.CollisionGroup(Child)) ? MParticles.CollisionGroup(NewIndex) : MParticles.CollisionGroup(Child);
+				ClusteredChild->SetCollisionImpulses(FMath::Max(NewParticle->CollisionImpulses(), ClusteredChild->CollisionImpulses()));
+				Child->SetCollisionGroup(FMath::Min(NewParticle->CollisionGroup(), Child->CollisionGroup()));
+			}
 		}
 
 		FClusterCreationParameters<T> NoCleanParams = Parameters;
 		NoCleanParams.bCleanCollisionParticles = false;
 		NoCleanParams.bCopyCollisionParticles = !!UnionsHaveCollisionParticles;
 
-		UpdateMassProperties(ChildrenArray, NewIndex, nullptr);
-		UpdateGeometry(ChildrenArray, NewIndex, TSerializablePtr<FImplicitObject>(), NoCleanParams);
+		TSet<TPBDRigidParticleHandle<T, d>*> ChildrenSet(ChildrenArray);
+		UpdateMassProperties(NewParticle, ChildrenSet, nullptr);
+		UpdateGeometry(NewParticle, ChildrenSet, TSerializablePtr<FImplicitObject>(), NoCleanParams);
 
-		return NewIndex;
+		return NewParticle;
 	}
 
+#if 0
 	int32 UseMultiChildProxy = 1;
 	FAutoConsoleVariableRef CVarUseMultiChildProxy(TEXT("p.UseMultiChildProxy"), UseMultiChildProxy, TEXT("Whether to merge multiple children into a single collision proxy when one is available"));
 
@@ -566,8 +608,6 @@ namespace Chaos
 	{
 		SCOPE_CYCLE_COUNTER(STAT_UnionClusterGroups);
 
-		FClusterMap& MParentToChildren = GetChildrenMap();
-
 		TMap<int32, TArray<uint32>> GroupMapping;
 		for (uint32 i = 0; i < MParticles.Size(); i++) // just loop active clusters here.
 		{
@@ -591,14 +631,14 @@ namespace Chaos
 				for (uint32& OriginalRootIdx : Group.Value)
 				{
 					TUniquePtr<TMultiChildProxyData<T, d>> ProxyData;
-					if (MParentToChildren.Contains(MParticles.Handle(OriginalRootIdx)))
+					if (MChildren.Contains(MParticles.Handle(OriginalRootIdx)))
 					{
-						if (UseMultiChildProxy && !MParticles.DynamicGeometry(OriginalRootIdx) && MParentToChildren[MParticles.Handle(OriginalRootIdx)].Num() > MinChildrenForMultiProxy) //Don't support dynamic geometry
+						if (UseMultiChildProxy && !MParticles.DynamicGeometry(OriginalRootIdx) && MChildren[MParticles.Handle(OriginalRootIdx)].Num() > MinChildrenForMultiProxy) //Don't support dynamic geometry
 						{
-							if (ensure(MParentToChildren[MParticles.Handle(OriginalRootIdx)].Num()))
+							if (ensure(MChildren[MParticles.Handle(OriginalRootIdx)].Num()))
 							{
 								ProxyData = MakeUnique<TMultiChildProxyData<T, d>>();
-								ProxyData->KeyChild = (MParentToChildren[MParticles.Handle(OriginalRootIdx)])[0];
+								ProxyData->KeyChild = (MChildren[MParticles.Handle(OriginalRootIdx)])[0];
 								ProxyData->RelativeToKeyChild = TRigidTransform<T, d>(MParticles.X(OriginalRootIdx), MParticles.R(OriginalRootIdx)); //store world space of original root. Need to break it up and then compute relative to world space of key child
 							}
 						}
@@ -610,7 +650,9 @@ namespace Chaos
 						{
 							//now that we have world space updated for key child, compute relative transform for original root
 							const TRigidTransform<T, d> OriginalRootWorldTM = ProxyData->RelativeToKeyChild;
-							ProxyData->RelativeToKeyChild = OriginalRootWorldTM.GetRelativeTransform(TRigidTransform<T, d>(MParticles.X(ProxyData->KeyChild), MParticles.R(ProxyData->KeyChild)));
+							ProxyData->RelativeToKeyChild = 
+								OriginalRootWorldTM.GetRelativeTransform(
+									TRigidTransform<T, d>(MParticles.X(ProxyData->KeyChild), MParticles.R(ProxyData->KeyChild)));
 							MParticles.MultiChildProxyData(OriginalRootIdx) = MoveTemp(ProxyData);
 
 							for (uint32 Child : OriginalRootChildren)
@@ -636,97 +678,107 @@ namespace Chaos
 			}
 		}
 	}
+#endif
 
+#if 0 // Not called currently
 	DECLARE_CYCLE_STAT(TEXT("TPBDRigidClustering<>::DeactivateClusterParticle"), STAT_DeactivateClusterParticle, STATGROUP_Chaos);
 	template<class FPBDRigidsEvolution, class FPBDCollisionConstraint, class T, int d>
-	TSet<uint32> TPBDRigidClustering<FPBDRigidsEvolution, FPBDCollisionConstraint, T, d>::DeactivateClusterParticle(TPBDRigidClusteredParticleHandle<T,d>* ClusteredParticle)
+	TSet<uint32> TPBDRigidClustering<FPBDRigidsEvolution, FPBDCollisionConstraint, T, d>::DeactivateClusterParticle(
+		TPBDRigidClusteredParticleHandle<T,d>* ClusteredParticle)
 	{
 		SCOPE_CYCLE_COUNTER(STAT_DeactivateClusterParticle);
-		FClusterMap& MParentToChildren = GetChildrenMap();
 
 		TSet<uint32> ActivatedChildren;
 		check(!ClusteredParticle->Disabled());
-		if (MParentToChildren.Contains(ClusteredParticle))
+		if (MChildren.Contains(ClusteredParticle))
 		{
-			ActivatedChildren = ReleaseClusterParticles(MParentToChildren[ClusteredParticle]);
+			ActivatedChildren = ReleaseClusterParticles(MChildren[ClusteredParticle]);
 		}
 		return ActivatedChildren;
 	}
+#endif // 0
 
 	DECLARE_CYCLE_STAT(TEXT("TPBDRigidClustering<>::ReleaseClusterParticles(STRAIN)"), STAT_ReleaseClusterParticles_STRAIN, STATGROUP_Chaos);
 	template<class FPBDRigidsEvolution, class FPBDCollisionConstraint, class T, int d>
-	TSet<uint32> TPBDRigidClustering<FPBDRigidsEvolution, FPBDCollisionConstraint, T, d>::ReleaseClusterParticles(
-		TPBDRigidClusteredParticleHandle<T,d>* ClusteredParticle, 
-		const TArrayView<T>& StrainArray)
+	TSet<TPBDRigidParticleHandle<T, d>*> 
+	TPBDRigidClustering<FPBDRigidsEvolution, FPBDCollisionConstraint, T, d>::ReleaseClusterParticles(
+		TPBDRigidClusteredParticleHandle<T,d>* ClusteredParticle,
+		const TArrayView<T>* ExternalStrainArray)
 	{
 		SCOPE_CYCLE_COUNTER(STAT_ReleaseClusterParticles_STRAIN);
-		FClusterMap& MParentToChildren = GetChildrenMap();
 
-		TSet<uint32> ActivatedChildren;
-		const float ClusterDistanceThreshold2 = ClusterDistanceThreshold * ClusterDistanceThreshold;
-
-		if (!ensureMsgf(MParentToChildren.Contains(ClusteredParticle), TEXT("Removing Cluster that does not exist!")))
+		TSet<TPBDRigidParticleHandle<T, d>*> ActivatedChildren;
+		if (!ensureMsgf(MChildren.Contains(ClusteredParticle), TEXT("Removing Cluster that does not exist!")))
 		{
 			return ActivatedChildren;
 		}
-		TArray<uint32>& Children = MParentToChildren[ClusteredParticle];
+		TArray<TPBDRigidParticleHandle<T, d>*>& Children = MChildren[ClusteredParticle];
+
 		bool bChildrenChanged = false;
 		const bool bRewindOnDecluster = ChildrenInheritVelocity < 1;
-		const TRigidTransform<T, d> PreSolveTM = bRewindOnDecluster ? TRigidTransform<T, d>(ClusteredParticle->X(), ClusteredParticle->R()) : TRigidTransform<T, d>(ClusteredParticle->P(), ClusteredParticle->Q());
+		const TRigidTransform<T, d> PreSolveTM = 
+			bRewindOnDecluster ? 
+			TRigidTransform<T, d>(ClusteredParticle->X(), ClusteredParticle->R()) : 
+			TRigidTransform<T, d>(ClusteredParticle->P(), ClusteredParticle->Q());
 
 		//@todo(ocohen): iterate with all the potential parents at once?
 		//find all children within some distance of contact point
 
-		auto RemoveChildLambda = [&](uint32 Child, int32 ChildIdx) {
-
-			MEvolution.EnableParticle(MParticles.Handle(Child), ClusteredParticle);
+		auto RemoveChildLambda = [&](TPBDRigidParticleHandle<T, d>* Child/*, const int32 Idx*/)
+		{
+			TPBDRigidClusteredParticleHandle<T,d>* ClusteredChild = Child->CastToClustered();
 
-			TopLevelClusterParents.Add(Child);
+			MEvolution.EnableParticle(Child, ClusteredParticle);
+			TopLevelClusterParents.Add(ClusteredChild);
 
 			//make sure to remove multi child proxy if it exists
-			if (MParticles.MultiChildProxyId(Child).Id != INDEX_NONE)
-			{
-				MParticles.MultiChildProxyData(MParticles.MultiChildProxyId(Child).Id).Reset();
-			}
+			ClusteredChild->MultiChildProxyData().Reset();
+			ClusteredChild->MultiChildProxyId().Id = nullptr;
+			ClusteredChild->SetClusterId(ClusterId());
 
-			TRigidTransform<T, d> ChildFrame = MParticles.ChildToParent(Child) * PreSolveTM;
-			MParticles.X(Child) = ChildFrame.GetTranslation();
-			MParticles.R(Child) = ChildFrame.GetRotation();
-			MParticles.ClusterIds(Child).Id = INDEX_NONE;
+			const TRigidTransform<T, d> ChildFrame = ClusteredChild->ChildToParent() * PreSolveTM;
+			Child->SetX(ChildFrame.GetTranslation());
+			Child->SetR(ChildFrame.GetRotation());
 
 			if (!bRewindOnDecluster)
 			{
-				MParticles.P(Child) = MParticles.X(Child);
-				MParticles.Q(Child) = MParticles.R(Child);
+				Child->SetP(Child->X());
+				Child->SetQ(Child->R());
 			}
 
 			//todo(ocohen): for now just inherit velocity at new COM. This isn't quite right for rotation
 			//todo(ocohen): in the presence of collisions, this will leave all children with the post-collision
 			// velocity. This should be controlled by material properties so we can allow the broken pieces to
 			// maintain the clusters pre-collision velocity.
-			MParticles.V(Child) = ClusteredParticle->V();
-			MParticles.W(Child) = ClusteredParticle->W();
-			MParticles.PreV(Child) = ClusteredParticle->PreV();
-			MParticles.PreW(Child) = ClusteredParticle->PreW();
-
+			Child->SetV(ClusteredParticle->V());
+			Child->SetW(ClusteredParticle->W());
+			Child->SetPreV(ClusteredParticle->PreV());
+			Child->SetPreW(ClusteredParticle->PreW());
 
 			ActivatedChildren.Add(Child);
-			if (ChildIdx != INDEX_NONE)
-			{
-				Children.RemoveAtSwap(ChildIdx, 1, /*bAllowShrinking=*/false); //@todo(ocohen): maybe avoid this until we know all children are not going away?
-			}
+			//if (ChildIdx != INDEX_NONE)
+			//{
+			//	Children.RemoveAtSwap(ChildIdx, 1, /*bAllowShrinking=*/false); //@todo(ocohen): maybe avoid this until we know all children are not going away?
+			//}
 
 			bChildrenChanged = true;
 		};
 
 		for (int32 ChildIdx = Children.Num() - 1; ChildIdx >= 0; --ChildIdx)
 		{
-			int32 Child = Children[ChildIdx];
-			if (StrainArray[Child] >= MParticles.Strains(Child))
+			TPBDRigidClusteredParticleHandle<T, d>* Child = Children[ChildIdx]->CastToClustered();
+			if (!Child)
+				continue;
+			if (ExternalStrainArray && (*ExternalStrainArray)[ChildIdx] ||
+			   !ExternalStrainArray && Child->CollisionImpulses() >= Child->Strain())
 			{
-				RemoveChildLambda(Child, ChildIdx); //the piece that hits just breaks off - we may want more control by looking at the edges of this piece which would give us cleaner breaks (this approach produces more rubble)
+				// The piece that hits just breaks off - we may want more control 
+				// by looking at the edges of this piece which would give us cleaner 
+				// breaks (this approach produces more rubble)
+				RemoveChildLambda(Child);
+				Children.RemoveAtSwap(ChildIdx);
 
-				if (MParticles.ToBeRemovedOnFracture(Child))
+				if (Child->ToBeRemovedOnFracture())
 				{
 					MActiveRemovalIndices.Add(Child);
 				}
@@ -734,13 +786,13 @@ namespace Chaos
 				{
 					if (DoGenerateBreakingData)
 					{
-						int32 NewIdx = MAllClusterBreakings.Add(TBreakingData<float, 3>());
-						TBreakingData<float, 3>& ClusterBreakingItem = MAllClusterBreakings[NewIdx];
-						ClusterBreakingItem.ParticleIndex = Child;
-						ClusterBreakingItem.Location = MParticles.X(Child);
-						ClusterBreakingItem.Velocity = MParticles.V(Child);
-						ClusterBreakingItem.AngularVelocity = MParticles.W(Child);
-						ClusterBreakingItem.Mass = MParticles.M(Child);
+						const int32 NewIdx = MAllClusterBreakings.Add(TBreakingData<float, 3>());
+						TBreakingData<float, 3>& ClusterBreak = MAllClusterBreakings[NewIdx];
+						ClusterBreak.Particle = Child->GTGeometryParticle();
+						ClusterBreak.Location = Child->X();
+						ClusterBreak.Velocity = Child->V();
+						ClusterBreak.AngularVelocity = Child->W();
+						ClusterBreak.Mass = Child->M();
 					}
 				}
 			}
@@ -751,37 +803,37 @@ namespace Chaos
 			if (UseConnectivity)
 			{
 				//cluster may have contained forests so find the connected pieces and cluster them together
-				TSet<uint32> PotentialActivatedChildren;
+				TSet<TPBDRigidParticleHandle<T,d>*> PotentialActivatedChildren;
 				PotentialActivatedChildren.Append(Children);
 
 				//first update the connected graph of the children we already removed
-				for (uint32 Child : ActivatedChildren)
+				for (TPBDRigidParticleHandle<T,d>* Child : ActivatedChildren)
 				{
 					RemoveNodeConnections(Child);
 				}
 
 				if (PotentialActivatedChildren.Num())
 				{
-					TArray<TArray<uint32>> ConnectedPiecesArray;
+					TArray<TArray<TPBDRigidParticleHandle<T,d>*>> ConnectedPiecesArray;
 					//traverse connectivity and see how many connected pieces we have
-					TSet<uint32> ProcessedChildren;
-					for (uint32 PotentialActivatedChild : PotentialActivatedChildren)
+					TSet<TPBDRigidParticleHandle<T,d>*> ProcessedChildren;
+					for (TPBDRigidParticleHandle<T,d>* PotentialActivatedChild : PotentialActivatedChildren)
 					{
 						if (!ProcessedChildren.Contains(PotentialActivatedChild))
 						{
 							ConnectedPiecesArray.AddDefaulted();
-							TArray<uint32>& ConnectedPieces = ConnectedPiecesArray.Last();
+							TArray<TPBDRigidParticleHandle<T,d>*>& ConnectedPieces = ConnectedPiecesArray.Last();
 
-							TArray<uint32> ProcessingQueue;
+							TArray<TPBDRigidParticleHandle<T,d>*> ProcessingQueue;
 							ProcessingQueue.Add(PotentialActivatedChild);
 							while (ProcessingQueue.Num())
 							{
-								uint32 Child = ProcessingQueue.Pop();
+								TPBDRigidParticleHandle<T,d>* Child = ProcessingQueue.Pop();
 								if (!ProcessedChildren.Contains(Child))
 								{
 									ProcessedChildren.Add(Child);
 									ConnectedPieces.Add(Child);
-									for (const TConnectivityEdge<T>& Edge : MParticles.ConnectivityEdges(Child))
+									for (const TConnectivityEdge<T>& Edge : Child->CastToClustered()->ConnectivityEdges())
 									{
 										if (!ProcessedChildren.Contains(Edge.Sibling))
 										{
@@ -793,38 +845,62 @@ namespace Chaos
 						}
 					}
 
-					for (const TArray<uint32>& ConnectedPieces : ConnectedPiecesArray)
+					int32 NumNewClusters = 0;
+					for (TArray<TPBDRigidParticleHandle<T,d>*>& ConnectedPieces : ConnectedPiecesArray)
 					{
 						if (ConnectedPieces.Num() == 1) //need to break single pieces first in case multi child proxy needs to be invalidated
 						{
-							const uint32 Child = ConnectedPieces[0];
-							RemoveChildLambda(Child, INDEX_NONE); //no need to remove child, we'll just empty it at the end
+							TPBDRigidParticleHandle<T,d>* Child = ConnectedPieces[0];
+							RemoveChildLambda(Child);
+						}
+						else if (ConnectedPieces.Num() > 1)
+						{
+							NumNewClusters++;
+						}
+					}
+					TArray<Chaos::TPBDRigidClusteredParticleHandle<float,3>*> NewClusterHandles = 
+						MEvolution.CreateClusteredParticles(NumNewClusters);
+					for (auto& Handle : NewClusterHandles)
+					{
+						if (!Handle->UniqueIdx().IsValid())
+						{
+							Handle->SetUniqueIdx(MEvolution.GenerateUniqueIdx());
 						}
 					}
-					for (TArray<uint32>& ConnectedPieces : ConnectedPiecesArray)
+					int32 ClusterHandlesIdx = 0;
+					for (TArray<TPBDRigidParticleHandle<T,d>*>& ConnectedPieces : ConnectedPiecesArray)
 					{
 						if (ConnectedPieces.Num() > 1) //now build the remaining pieces
 						{
-							int32 NewClusterIndex = CreateClusterParticleFromClusterChildren(MoveTemp(ConnectedPieces), ClusteredParticle, PreSolveTM);
-							MParticles.Strains(NewClusterIndex) = ClusteredParticle->Strains();
-							MEvolution.SetPhysicsMaterial(MParticles.Handle(NewClusterIndex), MEvolution.GetPhysicsMaterial(ClusteredParticle));
-
-							MParticles.V(NewClusterIndex) = ClusteredParticle->V();
-							MParticles.W(NewClusterIndex) = ClusteredParticle->W();
-							MParticles.PreV(NewClusterIndex) = ClusteredParticle->PreV();
-							MParticles.PreW(NewClusterIndex) = ClusteredParticle->PreW();
-							MParticles.P(NewClusterIndex) = MParticles.X(NewClusterIndex);
-							MParticles.Q(NewClusterIndex) = MParticles.R(NewClusterIndex);
-
-							ActivatedChildren.Add(NewClusterIndex);
+							Chaos::FClusterCreationParameters<float> CreationParameters;
+							CreationParameters.ClusterParticleHandle = NewClusterHandles[ClusterHandlesIdx++];
+							Chaos::TPBDRigidClusteredParticleHandle<float, 3>* NewCluster = 
+								CreateClusterParticleFromClusterChildren(
+									MoveTemp(ConnectedPieces), 
+									ClusteredParticle, 
+									PreSolveTM, 
+									CreationParameters);
+
+							NewCluster->SetStrain(ClusteredParticle->Strain());
+							MEvolution.SetPhysicsMaterial(
+								NewCluster, MEvolution.GetPhysicsMaterial(ClusteredParticle));
+
+							NewCluster->SetV(ClusteredParticle->V());
+							NewCluster->SetW(ClusteredParticle->W());
+							NewCluster->SetPreV(ClusteredParticle->PreV());
+							NewCluster->SetPreW(ClusteredParticle->PreW());
+							NewCluster->SetP(ClusteredParticle->X());
+							NewCluster->SetQ(ClusteredParticle->R());
+
+							ActivatedChildren.Add(NewCluster);
 						}
 					}
 				}
 			}
 
-			for (uint32 Active : ActivatedChildren)
+			for (TPBDRigidParticleHandle<T,d>* Child : ActivatedChildren)
 			{
-				UpdateKinematicProperties(Active);
+				UpdateKinematicProperties(Child);
 			}
 
 			//disable cluster
@@ -834,9 +910,11 @@ namespace Chaos
 		return ActivatedChildren;
 	}
 
+#if 0 // Not called currently
 	DECLARE_CYCLE_STAT(TEXT("TPBDRigidClustering<>::ReleaseClusterParticles(LIST)"), STAT_ReleaseClusterParticles_LIST, STATGROUP_Chaos);
 	template<class FPBDRigidsEvolution, class FPBDCollisionConstraint, class T, int d>
 	TSet<uint32> TPBDRigidClustering<FPBDRigidsEvolution, FPBDCollisionConstraint, T, d>::ReleaseClusterParticles(
+		TArray<TPBDRigidParticleHandle<T, d>*> ChildrenParticles)
 		const TArray<uint32>& ChildrenParticles)
 	{
 		SCOPE_CYCLE_COUNTER(STAT_ReleaseClusterParticles_LIST);
@@ -879,36 +957,41 @@ namespace Chaos
 
 		return ActivatedBodies;
 	}
+#endif // 0
 
 	DECLARE_CYCLE_STAT(TEXT("TPBDRigidClustering<>::AdvanceClustering"), STAT_AdvanceClustering, STATGROUP_Chaos);
 	DECLARE_CYCLE_STAT(TEXT("TPBDRigidClustering<>::Update Impulse from Strain"), STAT_UpdateImpulseStrain, STATGROUP_Chaos);
 	DECLARE_CYCLE_STAT(TEXT("TPBDRigidClustering<>::Update Dirty Impulses"), STAT_UpdateDirtyImpulses, STATGROUP_Chaos);
 	DECLARE_CYCLE_STAT(TEXT("TPBDRigidClustering<>::Rewind"), STAT_ClusterRewind, STATGROUP_Chaos);
 	template<class FPBDRigidsEvolution, class FPBDCollisionConstraint, class T, int d>
-	void TPBDRigidClustering<FPBDRigidsEvolution, FPBDCollisionConstraint, T, d>::AdvanceClustering(const T Dt, FPBDCollisionConstraint& CollisionRule)
+	void 
+	TPBDRigidClustering<FPBDRigidsEvolution, FPBDCollisionConstraint, T, d>::AdvanceClustering(
+		const T Dt, 
+		FPBDCollisionConstraint& CollisionRule)
 	{
 		SCOPE_CYCLE_COUNTER(STAT_AdvanceClustering);
 		UE_LOG(LogChaos, Verbose, TEXT("START FRAME with Dt %f"), Dt);
 
-		FClusterMap& MParentToChildren = GetChildrenMap();
-
 		double FrameTime = 0, Time = 0;
 		FDurationTimer Timer(Time);
 		Timer.Start();
 
 		{
 			const float Threshold = 1.0f;
-			auto RemovalIndicesCopy = MActiveRemovalIndices; // copy since RemoveOnFracture alters MActiveRemovalIndices
-			for (uint32 ParticleIdx : RemovalIndicesCopy)
+			TSet<Chaos::TPBDRigidParticleHandle<float, 3>*> RemovalIndicesCopy = MActiveRemovalIndices;
+			for (Chaos::TPBDRigidParticleHandle<float, 3>* Particle : RemovalIndicesCopy)
 			{
-				if (MParticles.ToBeRemovedOnFracture(ParticleIdx) && MParticles.V(ParticleIdx).SizeSquared() > Threshold && MParticles.PreV(ParticleIdx).SizeSquared() > Threshold)
+				//if (MParticles.ToBeRemovedOnFracture(ParticleIdx) && MParticles.V(ParticleIdx).SizeSquared() > Threshold && MParticles.PreV(ParticleIdx).SizeSquared() > Threshold)
+				if (Particle->ToBeRemovedOnFracture() && 
+					Particle->V().SizeSquared() > Threshold && 
+					Particle->PreV().SizeSquared() > Threshold)
 				{
-					DisableParticleWithBreakEvent(ParticleIdx);
+					DisableParticleWithBreakEvent(Particle);
 				}
 			}
 		}
 
-		if (MParticles.Size())
+		if(MChildren.Num())
 		{
 			//
 			//  Grab collision impulses for processing
@@ -933,13 +1016,16 @@ namespace Chaos
 				{
 					if (ActiveCluster->ClusterIds().NumChildren > 0) //active index is a cluster
 					{
-						const TArray<uint32>& ParentToChildren = MParentToChildren[ActiveCluster];
-						for (const uint32 Child : ParentToChildren)
+						TArray<TPBDRigidParticleHandle<T, d>*>& ParentToChildren = MChildren[ActiveCluster];
+						for (TPBDRigidParticleHandle<T, d>* Child : ParentToChildren)
 						{
-							if (MParticles.Strains(Child) <= 0.f)
+							if (TPBDRigidClusteredParticleHandle<T, d>* ClusteredChild = Child->CastToClustered())
 							{
-								MParticles.CollisionImpulses(Child) = FLT_MAX;
-								MCollisionImpulseArrayDirty = true;
+								if (ClusteredChild->Strain() <= 0.f)
+								{
+									ClusteredChild->CollisionImpulse() = FLT_MAX;
+									MCollisionImpulseArrayDirty = true;
+								}
 							}
 						}
 					}
@@ -948,21 +1034,22 @@ namespace Chaos
 
 			if (MCollisionImpulseArrayDirty)
 			{
+
 				SCOPE_CYCLE_COUNTER(STAT_UpdateDirtyImpulses);
-				TArrayView<T> CollisionImpulsesView(&(MParticles.CollisionImpulses(0)), MParticles.Size());
-				TMap<TPBDRigidClusteredParticleHandle<T, d>*, TSet<uint32>> ClusterToActivatedChildren = BreakingModel(CollisionImpulsesView);
+				TMap<TPBDRigidClusteredParticleHandle<T, d>*, TSet<TPBDRigidParticleHandle<T, d>*>> ClusterToActivatedChildren = 
+					BreakingModel();
 
-				TSet<uint32> AllActivatedChildren;
+				TSet<TPBDRigidParticleHandle<T, d>*> AllActivatedChildren;
 				TSet<int32> IslandsToRecollide;
 				for (auto Itr : ClusterToActivatedChildren)
 				{
 					//question: do we need to iterate all the children? Seems like island is known from cluster, but don't want to break anything at this point
-					const TSet<uint32>& ActivatedChildren = Itr.Value;
-					for (uint32 ActiveChild : ActivatedChildren)
+					TSet<TPBDRigidParticleHandle<T, d>*>& ActivatedChildren = Itr.Value;
+					for (TPBDRigidParticleHandle<T, d>* ActiveChild : ActivatedChildren)
 					{
-						if (ensure(!MParticles.Disabled(ActiveChild)))
+						if (ensure(!ActiveChild->Disabled()))
 						{
-							int32 Island = MParticles.Island(ActiveChild);
+							int32 Island = ActiveChild->Island();
 							if (!IslandsToRecollide.Contains(Island) && Island != INDEX_NONE) // todo ask mike
 							{
 								IslandsToRecollide.Add(Island);
@@ -986,62 +1073,60 @@ namespace Chaos
 					{
 						for (auto Itr : ClusterToActivatedChildren)
 						{
-							auto ClusteredParticle = Itr.Key;
-
-							const TSet<uint32>& ActivatedChildren = Itr.Value;
-							for (uint32 ActiveChild : ActivatedChildren)
+							TPBDRigidClusteredParticleHandle<T, d>* ClusteredParticle = Itr.Key;
+							TSet<TPBDRigidParticleHandle<T, d>*>& ActivatedChildren = Itr.Value;
+							for (TPBDRigidParticleHandle<T, d>* ActiveChild : ActivatedChildren)
 							{
-								MParticles.V(ActiveChild) = MParticles.V(ActiveChild) * (1.f - ChildrenInheritVelocity) + ClusteredParticle->V() * ChildrenInheritVelocity;
-								MParticles.W(ActiveChild) = MParticles.W(ActiveChild) * (1.f - ChildrenInheritVelocity) + ClusteredParticle->W() * ChildrenInheritVelocity;
+								ActiveChild->SetV(
+									ActiveChild->V() * (1.f - ChildrenInheritVelocity) + ClusteredParticle->V() * ChildrenInheritVelocity);
+								ActiveChild->SetW(
+									ActiveChild->W() * (1.f - ChildrenInheritVelocity) + ClusteredParticle->W() * ChildrenInheritVelocity);
 							}
 						}
 					}
 				}
-			}
-		}
-
+			} // end if MCollisionImpulseArrayDirty
+		} // end if MParticles.Size()
 		Timer.Stop();
 		UE_LOG(LogChaos, Verbose, TEXT("Cluster Break Update Time is %f"), Time);
 	}
 
 	DECLARE_CYCLE_STAT(TEXT("TPBDRigidClustering<>::BreakingModel()"), STAT_BreakingModel, STATGROUP_Chaos);
 	template<class FPBDRigidsEvolution, class FPBDCollisionConstraint, class T, int d>
-	TMap<TPBDRigidClusteredParticleHandle<T, d>*, TSet<uint32>> TPBDRigidClustering<FPBDRigidsEvolution, FPBDCollisionConstraint, T, d>::BreakingModel(TArrayView<T> ExternalStrain)
+	TMap<TPBDRigidClusteredParticleHandle<T, d>*, TSet<TPBDRigidParticleHandle<T, d>*>> 
+	TPBDRigidClustering<FPBDRigidsEvolution, FPBDCollisionConstraint, T, d>::BreakingModel(
+		TArrayView<T>* ExternalStrain)
 	{
 		SCOPE_CYCLE_COUNTER(STAT_BreakingModel);
-		FClusterMap& MParentToChildren = GetChildrenMap();
 
-		TMap<TPBDRigidClusteredParticleHandle<T,d>*, TSet<uint32>> AllActivatedChildren;
+		TMap<TPBDRigidClusteredParticleHandle<T,d>*, TSet<TPBDRigidParticleHandle<T, d>*>> AllActivatedChildren;
 
 		auto NonDisabledClusteredParticles = MEvolution.GetNonDisabledClusteredArray(); //make copy because release cluster modifies active indices. We want to iterate over original active indices
 		for (auto ClusteredParticle : NonDisabledClusteredParticles)
 		{
 			if (ClusteredParticle->ClusterIds().NumChildren)
 			{
-				AllActivatedChildren.Add(ClusteredParticle, ReleaseClusterParticles(ClusteredParticle, ExternalStrain));
+				AllActivatedChildren.Add(
+					ClusteredParticle, 
+					ReleaseClusterParticles(ClusteredParticle, ExternalStrain));
 			}
 			else
 			{
 				// there's no children to break but we need to process whether this single piece is to be removed when damaged
 				if (ClusteredParticle->ToBeRemovedOnFracture())
 				{
-					if (ExternalStrain[ClusteredParticle->TransientParticleIndex()] >= ClusteredParticle->Strains())
+					if (ClusteredParticle->CollisionImpulses() >= ClusteredParticle->Strains())
 					{
 						DisableCluster(ClusteredParticle);
-
 						if (DoGenerateBreakingData)
 						{
 							int32 NewIdx = MAllClusterBreakings.Add(TBreakingData<float, 3>());
-
-							TBreakingData<float, 3>& ClusterBreakingItem = MAllClusterBreakings[NewIdx];
-							ClusterBreakingItem.ParticleIndex = INDEX_NONE;// ClusteredParticle;
-#if CHAOS_PARTICLEHANDLE_TODO
-							ClusterBreakingItem.Particle = ClusteredParticle;
-#endif
-							ClusterBreakingItem.Location = ClusteredParticle->X();
-							ClusterBreakingItem.Velocity = ClusteredParticle->V();
-							ClusterBreakingItem.AngularVelocity = ClusteredParticle->W(); 
-							ClusterBreakingItem.Mass = ClusteredParticle->M();
+							TBreakingData<float, 3>& ClusterBreak = MAllClusterBreakings[NewIdx];
+							ClusterBreak.Particle = ClusteredParticle->GTGeometryParticle();
+							ClusterBreak.Location = ClusteredParticle->X();
+							ClusterBreak.Velocity = ClusteredParticle->V();
+							ClusterBreak.AngularVelocity = ClusteredParticle->W();
+							ClusterBreak.Mass = ClusteredParticle->M();
 						}
 					}
 				}
@@ -1053,68 +1138,70 @@ namespace Chaos
 
 	DECLARE_CYCLE_STAT(TEXT("TPBDRigidClustering<>::PromoteStrains()"), STAT_PromoteStrains, STATGROUP_Chaos);
 	template<class FPBDRigidsEvolution, class FPBDCollisionConstraint, class T, int d>
-	T TPBDRigidClustering<FPBDRigidsEvolution, FPBDCollisionConstraint, T, d>::PromoteStrains(uint32 CurrentNode, TArrayView<T>& ExternalStrains)
+	T TPBDRigidClustering<FPBDRigidsEvolution, FPBDCollisionConstraint, T, d>::PromoteStrains(
+		TPBDRigidParticleHandle<T, d>* CurrentNode)
 	{
 		SCOPE_CYCLE_COUNTER(STAT_PromoteStrains);
-		FClusterMap& MParentToChildren = GetChildrenMap();
 
 		T Result = 0;
-		if (MParentToChildren.Contains(MParticles.Handle(CurrentNode)))
+		if (MChildren.Contains(CurrentNode))
 		{
-			for (uint32 Child : MParentToChildren[MParticles.Handle(CurrentNode)])
+			for (TPBDRigidParticleHandle<T, d>* Child : MChildren[CurrentNode])
 			{
-				Result += PromoteStrains(Child, ExternalStrains);
+				Result += PromoteStrains(Child);
 			}
 		}
 		else
 		{
-			return ExternalStrains[CurrentNode];
+			return CurrentNode->CastToClustered()->Strains();
 		}
-		ExternalStrains[CurrentNode] += Result;
+		CurrentNode->CastToClustered()->Strains() += Result;
 		return Result;
 	}
 
 	DECLARE_CYCLE_STAT(TEXT("TPBDRigidClustering<>::UpdateKinematicProperties()"), STAT_UpdateKinematicProperties, STATGROUP_Chaos);
 	template<class FPBDRigidsEvolution, class FPBDCollisionConstraint, class T, int d>
-	void TPBDRigidClustering<FPBDRigidsEvolution, FPBDCollisionConstraint, T, d>::UpdateKinematicProperties(uint32 ClusterIndex)
+	void TPBDRigidClustering<FPBDRigidsEvolution, FPBDCollisionConstraint, T, d>::UpdateKinematicProperties(
+		Chaos::TPBDRigidParticleHandle<float, 3>* Parent)
 	{
 		SCOPE_CYCLE_COUNTER(STAT_UpdateKinematicProperties);
-		FClusterMap& MParentToChildren = GetChildrenMap();
 
 		EObjectStateType ObjectState = EObjectStateType::Dynamic;
-		check(0 <= ClusterIndex && ClusterIndex < MParticles.Size());
-		if (MParentToChildren.Contains(MParticles.Handle(ClusterIndex)) && MParentToChildren[MParticles.Handle(ClusterIndex)].Num())
+		check(Parent != nullptr);
+		if (MChildren.Contains(Parent) && MChildren[Parent].Num())
 		{
-			TQueue<int32> Queue;
-			for (const int32 Child : MParentToChildren[MParticles.Handle(ClusterIndex)])
+			// TQueue is a linked list, which has no preallocator.
+			TQueue<Chaos::TPBDRigidParticleHandle<float, 3>*> Queue;
+			for (Chaos::TPBDRigidParticleHandle<float, 3>* Child : MChildren[Parent])
 			{
 				Queue.Enqueue(Child);
 			}
 
-			int32 CurrentIndex;
-			while (Queue.Dequeue(CurrentIndex) && ObjectState == EObjectStateType::Dynamic)
+			Chaos::TPBDRigidParticleHandle<float, 3>* CurrentHandle;
+			while (Queue.Dequeue(CurrentHandle) && ObjectState == EObjectStateType::Dynamic)
 			{
 				// @question : Maybe we should just store the leaf node bodies in a
 				// map, that will require Memory(n*log(n))
-				if (MParentToChildren.Contains(MParticles.Handle(CurrentIndex)))
+				if (MChildren.Contains(CurrentHandle))
 				{
-					for (const auto Child : MParentToChildren[MParticles.Handle(CurrentIndex)])
+					for (Chaos::TPBDRigidParticleHandle<float, 3>* Child : MChildren[CurrentHandle])
 					{
 						Queue.Enqueue(Child);
 					}
 				}
 
-				if (MParticles.ObjectState(CurrentIndex) == EObjectStateType::Kinematic)
+				const EObjectStateType CurrState = CurrentHandle->ObjectState();
+				if (CurrState == EObjectStateType::Kinematic)
 				{
 					ObjectState = EObjectStateType::Kinematic;
 				}
-				if (MParticles.ObjectState(CurrentIndex) == EObjectStateType::Static)
+				else if (CurrState == EObjectStateType::Static)
 				{
 					ObjectState = EObjectStateType::Static;
 				}
 			}
 
-			MParticles.SetObjectState(ClusterIndex, ObjectState);
+			Parent->SetObjectState(ObjectState);
 		}
 	}
 
@@ -1122,6 +1209,12 @@ namespace Chaos
 	template<class FPBDRigidsEvolution, class FPBDCollisionConstraint, class T, int d>
 	void TPBDRigidClustering<FPBDRigidsEvolution, FPBDCollisionConstraint, T, d>::SwapBufferedData()
 	{
+		check(false);
+		// TODO: Ryan - this code currently uses MParticles as the only source of clustered particles.
+		// As things stand, clustered particles and geometry collection particles are both of clustered
+		// type, but are stored separately.  Geometry collection particles exist on both game and physics
+		// threads, cluster particles only exist on the physics thread.
+/*
 		SCOPE_CYCLE_COUNTER(STAT_SwapBufferedData);
 		ResourceLock.WriteLock();
 		//BufferResource.MChildren.Reset();
@@ -1132,40 +1225,42 @@ namespace Chaos
 		const auto& NonDisabledClusteredParticles = MEvolution.GetNonDisabledClusteredArray();
 		for (auto& ClusteredParticle : NonDisabledClusteredParticles)
 		{
-			const TArray<uint32>* Children = MChildren.Find(ClusteredParticle);
-			if (Children && ClusteredParticle->ClusterIds().Id == INDEX_NONE) //root cluster so copy children
+			//const TArray<uint32>* Children = MChildren.Find(ClusteredParticle);
+			TArray<TPBDRigidParticleHandle<T, d>*>* Children = MChildren.Find(ClusteredParticle);
+			//if (Children && ClusteredParticle->ClusterIds().Id == INDEX_NONE) //root cluster so copy children
+			if (Children && ClusteredParticle->ClusterIds().Id == nullptr) //root cluster so copy children
 			{
 				//TODO: record GT particle pointer instead
-				BufferResource.MChildren.Add(ClusteredParticle->TransientParticleIndex(), *Children);
-				BufferResource.ClusterParentTransforms.Add(ClusteredParticle->TransientParticleIndex(), TRigidTransform<float, 3>(ClusteredParticle->X(), ClusteredParticle->R()));
+				//BufferResource.MChildren.Add(ClusteredParticle->TransientParticleIndex(), *Children);
+				BufferResource.MChildren.Add(
+					ClusteredParticle, *Children);
+				//BufferResource.ClusterParentTransforms.Add(ClusteredParticle->TransientParticleIndex(), TRigidTransform<float, 3>(ClusteredParticle->X(), ClusteredParticle->R()));
+				BufferResource.ClusterParentTransforms.Add(
+					ClusteredParticle, TRigidTransform<float, 3>(ClusteredParticle->X(), ClusteredParticle->R()));
 			}
 		}
 
 		BufferResource.GeometryPtrs = AllGeom; //in future this should be sparse. SQ has fallback that relies on potentially all geom so can't do it yet
-
 		ResourceLock.WriteUnlock();
+*/
 	}
 
 	DECLARE_CYCLE_STAT(TEXT("TPBDRigidClustering<>::GetActiveClusterIndex"), STAT_GetActiveClusterIndex, STATGROUP_Chaos);
 	template<class FPBDRigidsEvolution, class FPBDCollisionConstraint, class T, int d>
-	int32 TPBDRigidClustering<FPBDRigidsEvolution, FPBDCollisionConstraint, T, d>::GetActiveClusterIndex(uint32 ChildIndex)
+	TPBDRigidParticleHandle<T, d>* TPBDRigidClustering<FPBDRigidsEvolution, FPBDCollisionConstraint, T, d>::GetActiveClusterIndex(
+		TPBDRigidParticleHandle<T, d>* Child)
 	{
 		SCOPE_CYCLE_COUNTER(STAT_GetActiveClusterIndex);
-		int32 NumParticles = MParticles.Size();
-		int32 CurrentIndex = (int32)ChildIndex;
-		while (0 <= CurrentIndex && CurrentIndex < NumParticles)
+		while (Child && Child->Disabled())
 		{
-			if (!MParticles.Disabled(CurrentIndex))
-			{
-				return CurrentIndex;
-			}
-			CurrentIndex = MParticles.ClusterIds(CurrentIndex).Id;
+			Child = Child->CastToClustered()->ClusterIds().Id;
 		}
-		return INDEX_NONE;
+		return Child; 
 	}
 
 	template<class FPBDRigidsEvolution, class FPBDCollisionConstraint, class T, int d>
-	void TPBDRigidClustering<FPBDRigidsEvolution, FPBDCollisionConstraint, T, d>::IncrementPendingClusterCounter(uint32 ClusterGroupID)
+	void TPBDRigidClustering<FPBDRigidsEvolution, FPBDCollisionConstraint, T, d>::IncrementPendingClusterCounter(
+		uint32 ClusterGroupID)
 	{
 		if (!PendingClusterCounter.Contains(ClusterGroupID))
 		{
@@ -1175,7 +1270,8 @@ namespace Chaos
 	}
 
 	template<class FPBDRigidsEvolution, class FPBDCollisionConstraint, class T, int d>
-	void TPBDRigidClustering<FPBDRigidsEvolution, FPBDCollisionConstraint, T, d>::DecrementPendingClusterCounter(uint32 ClusterGroupID)
+	void TPBDRigidClustering<FPBDRigidsEvolution, FPBDCollisionConstraint, T, d>::DecrementPendingClusterCounter(
+		uint32 ClusterGroupID)
 	{
 		PendingClusterCounter[ClusterGroupID]--;
 		ensure(0 <= PendingClusterCounter[ClusterGroupID]);
@@ -1183,10 +1279,13 @@ namespace Chaos
 
 	DECLARE_CYCLE_STAT(TEXT("TPBDRigidClustering<>::GenerateConnectionGraph"), STAT_GenerateConnectionGraph, STATGROUP_Chaos);
 	template<class FPBDRigidsEvolution, class FPBDCollisionConstraint, class T, int d>
-	void TPBDRigidClustering<FPBDRigidsEvolution, FPBDCollisionConstraint, T, d>::GenerateConnectionGraph(int32 ClusterIndex, const FClusterCreationParameters<T>& Parameters)
+	void TPBDRigidClustering<FPBDRigidsEvolution, FPBDCollisionConstraint, T, d>::GenerateConnectionGraph(
+		Chaos::TPBDRigidClusteredParticleHandle<float, 3>* Parent,
+		const FClusterCreationParameters<T>& Parameters)
 	{
 		SCOPE_CYCLE_COUNTER(STAT_GenerateConnectionGraph);
-		if (!GetChildrenMap().Contains(MParticles.Handle(ClusterIndex))) return;
+		if (!MChildren.Contains(Parent)) 
+			return;
 
 		// Connectivity Graph
 		//    Build a connectivity graph for the cluster. If the PointImplicit is specified
@@ -1198,7 +1297,8 @@ namespace Chaos
 			typename FClusterCreationParameters<T>::EConnectionMethod LocalConnectionMethod = Parameters.ConnectionMethod;
 
 			if (LocalConnectionMethod == FClusterCreationParameters<T>::EConnectionMethod::None ||
-				(LocalConnectionMethod == FClusterCreationParameters<T>::EConnectionMethod::PointImplicit && !MParticles.CollisionParticles(ClusterIndex)))
+				(LocalConnectionMethod == FClusterCreationParameters<T>::EConnectionMethod::PointImplicit && 
+				 !Parent->CollisionParticles()))
 			{
 				LocalConnectionMethod = FClusterCreationParameters<T>::EConnectionMethod::MinimalSpanningSubsetDelaunayTriangulation; // default method
 			}
@@ -1206,29 +1306,35 @@ namespace Chaos
 			if (LocalConnectionMethod == FClusterCreationParameters<T>::EConnectionMethod::PointImplicit ||
 				LocalConnectionMethod == FClusterCreationParameters<T>::EConnectionMethod::PointImplicitAugmentedWithMinimalDelaunay)
 			{
-				UpdateConnectivityGraphUsingPointImplicit(ClusterIndex, Parameters);
+				UpdateConnectivityGraphUsingPointImplicit(Parent, Parameters);
 			}
 
 			if (LocalConnectionMethod == FClusterCreationParameters<T>::EConnectionMethod::DelaunayTriangulation)
 			{
-				UpdateConnectivityGraphUsingDelaunayTriangulation(ClusterIndex, Parameters); // not thread safe
+				UpdateConnectivityGraphUsingDelaunayTriangulation(Parent, Parameters); // not thread safe
 			}
 
 			if (LocalConnectionMethod == FClusterCreationParameters<T>::EConnectionMethod::PointImplicitAugmentedWithMinimalDelaunay ||
 				LocalConnectionMethod == FClusterCreationParameters<T>::EConnectionMethod::MinimalSpanningSubsetDelaunayTriangulation)
 			{
-				FixConnectivityGraphUsingDelaunayTriangulation(ClusterIndex, Parameters);
+				FixConnectivityGraphUsingDelaunayTriangulation(Parent, Parameters);
 			}
 		}
 	}
 
 	DECLARE_CYCLE_STAT(TEXT("TPBDRigidClustering<>::UpdateMassProperties"), STAT_UpdateMassProperties, STATGROUP_Chaos);
 	template<class FPBDRigidsEvolution, class FPBDCollisionConstraint, class T, int d>
-	void TPBDRigidClustering<FPBDRigidsEvolution, FPBDCollisionConstraint, T, d>::UpdateMassProperties(const TArray<uint32>& Children, const uint32 ClusterIndex, const TRigidTransform<T, d>* ForceMassOrientation)
+	void TPBDRigidClustering<FPBDRigidsEvolution, FPBDCollisionConstraint, T, d>::UpdateMassProperties(
+		Chaos::TPBDRigidClusteredParticleHandle<float, 3>* Parent, 
+		TSet<TPBDRigidParticleHandle<T, d>*>& Children, 
+		const TRigidTransform<T, d>* ForceMassOrientation)
 	{
 		SCOPE_CYCLE_COUNTER(STAT_UpdateMassProperties);
-		UpdateClusterMassProperties(MParticles, Children, ClusterIndex, ForceMassOrientation, &MParticles.MultiChildProxyDataArray(), &MParticles.MultiChildProxyIdArray());
-		UpdateKinematicProperties(ClusterIndex);
+		UpdateClusterMassProperties(
+			Parent,
+			Children,
+			ForceMassOrientation);
+		UpdateKinematicProperties(Parent);
 	}
 
 	DECLARE_CYCLE_STAT(TEXT("TPBDRigidClustering<>::UpdateGeometry"), STAT_UpdateGeometry, STATGROUP_Chaos);
@@ -1238,106 +1344,130 @@ namespace Chaos
 	DECLARE_CYCLE_STAT(TEXT("TPBDRigidClustering<>::UpdateGeometry_PointsBVH"), STAT_UpdateGeometry_PointsBVH, STATGROUP_Chaos);
 
 	template<class FPBDRigidsEvolution, class FPBDCollisionConstraint, class T, int d>
-	void TPBDRigidClustering<FPBDRigidsEvolution, FPBDCollisionConstraint, T, d>::UpdateGeometry(const TArray<uint32>& Children,
-		const uint32 NewIndex, TSerializablePtr<FImplicitObject> ProxyGeometry, const FClusterCreationParameters<T>& Parameters)
+	void TPBDRigidClustering<FPBDRigidsEvolution, FPBDCollisionConstraint, T, d>::UpdateGeometry(
+		Chaos::TPBDRigidClusteredParticleHandle<float, 3>* Parent, 
+		const TSet<TPBDRigidParticleHandle<T, d>*>& Children, 
+		TSerializablePtr<FImplicitObject> ProxyGeometry, 
+		const FClusterCreationParameters<T>& Parameters)
 	{
 		SCOPE_CYCLE_COUNTER(STAT_UpdateGeometry);
-		const bool bUseCollisionPoints = (ProxyGeometry || Parameters.bCopyCollisionParticles) && !Parameters.CollisionParticles;
+
 		TArray<TUniquePtr<FImplicitObject>> Objects;
 		TArray<TUniquePtr<FImplicitObject>> Objects2; //todo: find a better way to reuse this
+		Objects.Reserve(Children.Num());
+		Objects2.Reserve(Children.Num());
+
+		//we should never update existing geometry since this is used by SQ threads.
+		ensure(!Parent->Geometry());
+		ensure(!Parent->DynamicGeometry());
 
-		ensure(!MParticles.Geometry(NewIndex)); //we should never update existing geometry since this is used by SQ threads.
-		ensure(!MParticles.DynamicGeometry(NewIndex)); //we should never update existing geometry since this is used by SQ threads.
+		const TRigidTransform<T, d> ClusterWorldTM(Parent->X(), Parent->R());
 
 		TArray<TVector<T, d>> OriginalPoints;
-		bool bUseParticleImplicit = false;
-		TArray<int32> GeomToOriginalParticlesHack;
+		TArray<TPBDRigidParticleHandle<T, d>*> GeomToOriginalParticlesHack;
 		GeomToOriginalParticlesHack.Reserve(Children.Num());
+
+		const bool bUseCollisionPoints = (ProxyGeometry || Parameters.bCopyCollisionParticles) && !Parameters.CollisionParticles;
+		bool bUseParticleImplicit = false;
 		bool bUsingMultiChildProxy = false;
 
-		const TRigidTransform<T, d> ClusterWorldTM(MParticles.X(NewIndex), MParticles.R(NewIndex));
-		{
+		{ // STAT_UpdateGeometry_GatherObjects
 			SCOPE_CYCLE_COUNTER(STAT_UpdateGeometry_GatherObjects);
-			for (const uint32 Child : Children)
+
+			if (bUseCollisionPoints)
 			{
-				const TRigidTransform<T, d> ChildWorldTM(MParticles.X(Child), MParticles.R(Child));
+				uint32 NumPoints = 0;
+				for (TPBDRigidParticleHandle<T, d>* Child : Children)
+				{
+					NumPoints += Child->CollisionParticlesSize();
+				}
+				OriginalPoints.Reserve(NumPoints);
+			}
+
+			for (TPBDRigidParticleHandle<T, d>* Child : Children)
+			{
+				const TRigidTransform<T, d> ChildWorldTM(Child->X(), Child->R());
 				TRigidTransform<T, d> Frame = ChildWorldTM.GetRelativeTransform(ClusterWorldTM);
-				int32 UsedGeomChild = Child;
-				if (MParticles.Geometry(Child))
+				TPBDRigidParticleHandle<T, d>* UsedGeomChild = Child;
+				if (Child->Geometry())
 				{
-					const int32 MultiChildProxyId = MParticles.MultiChildProxyId(Child).Id;
-					if (UseLevelsetCollision || MultiChildProxyId == INDEX_NONE || !MParticles.MultiChildProxyData(MultiChildProxyId))
+					TPBDRigidClusteredParticleHandle<T, d>* ClusteredChild = Child->CastToClustered();
+
+					const FMultiChildProxyId MultiChildProxyId = 
+						ClusteredChild ? ClusteredChild->MultiChildProxyId() : FMultiChildProxyId();
+					const TUniquePtr<TMultiChildProxyData<T, d>>& MultiChildProxyData = ClusteredChild->MultiChildProxyData();
+
+					if (UseLevelsetCollision || MultiChildProxyId.Id == nullptr || !MultiChildProxyData)
 					{
-						Objects.Add(TUniquePtr<FImplicitObject>(new TImplicitObjectTransformed<T, d>(MParticles.Geometry(Child), Frame)));
-						Objects2.Add(TUniquePtr<FImplicitObject>(new TImplicitObjectTransformed<T, d>(MParticles.Geometry(Child), Frame)));
+						Objects.Add(TUniquePtr<FImplicitObject>(new TImplicitObjectTransformed<T, d>(Child->Geometry(), Frame)));
+						Objects2.Add(TUniquePtr<FImplicitObject>(new TImplicitObjectTransformed<T, d>(Child->Geometry(), Frame)));
 						GeomToOriginalParticlesHack.Add(Child);
 					}
-					else if (MParticles.MultiChildProxyData(MultiChildProxyId)->KeyChild == Child)
+					else if (MultiChildProxyData->KeyChild == Child)
 					{
 						//using multi child proxy and this child is the key
-						const TRigidTransform<T, d> ProxyWorldTM = MParticles.MultiChildProxyData(MultiChildProxyId)->RelativeToKeyChild * ChildWorldTM;
+						const TRigidTransform<T, d> ProxyWorldTM = MultiChildProxyData->RelativeToKeyChild * ChildWorldTM;
 						const TRigidTransform<T, d> ProxyRelativeTM = ProxyWorldTM.GetRelativeTransform(ClusterWorldTM);
-						Objects.Add(TUniquePtr<FImplicitObject>(new TImplicitObjectTransformed<T, d>(MParticles.Geometry(MultiChildProxyId), ProxyRelativeTM)));
-						Objects2.Add(TUniquePtr<FImplicitObject>(new TImplicitObjectTransformed<T, d>(MParticles.Geometry(MultiChildProxyId), ProxyRelativeTM)));
-						UsedGeomChild = MultiChildProxyId;
+						Objects.Add(TUniquePtr<FImplicitObject>(new TImplicitObjectTransformed<T, d>(MultiChildProxyId.Id->Geometry(), ProxyRelativeTM)));
+						Objects2.Add(TUniquePtr<FImplicitObject>(new TImplicitObjectTransformed<T, d>(MultiChildProxyId.Id->Geometry(), ProxyRelativeTM)));
+						UsedGeomChild = MultiChildProxyId.Id;
 						GeomToOriginalParticlesHack.Add(UsedGeomChild);
 						bUsingMultiChildProxy = true;
 					}
 				}
 
-				ensure(MParticles.Disabled(Child) == true);
+				ensure(Child->Disabled() == true);
+				check(Child->CastToClustered()->ClusterIds().Id == Parent);
 
-				check(MParticles.ClusterIds(Child).Id == NewIndex);
-				MParticles.ChildToParent(Child) = Frame;
+				Child->CastToClustered()->SetChildToParent(Frame);
 
 				if (bUseCollisionPoints)
 				{
 					SCOPE_CYCLE_COUNTER(STAT_UpdateGeometry_GatherPoints);
-					if (MParticles.CollisionParticles(Child))
+					if (const TUniquePtr<TBVHParticles<T,d>>& CollisionParticles = Child->CollisionParticles())
 					{
-						for (uint32 i = 0; i < MParticles.CollisionParticles(Child)->Size(); ++i)
+						for (uint32 i = 0; i < CollisionParticles->Size(); ++i)
 						{
-							OriginalPoints.Add(Frame.TransformPosition(MParticles.CollisionParticles(Child)->X(i)));
+							OriginalPoints.Add(Frame.TransformPosition(CollisionParticles->X(i)));
 						}
 					}
 				}
-
-				if (MParticles.Geometry(Child) && MParticles.Geometry(Child)->GetType() == ImplicitObjectType::Unknown)
+				if (Child->Geometry() && Child->Geometry()->GetType() == ImplicitObjectType::Unknown)
 				{
 					bUseParticleImplicit = true;
 				}
-			}
-		}
+			} // end for
+		} // STAT_UpdateGeometry_GatherObjects
 
 		{
 			QUICK_SCOPE_CYCLE_COUNTER(SpatialBVH);
-			MParticles.ChildrenSpatial(NewIndex) = Objects2.Num() ? MakeUnique<FImplicitObjectUnionClustered>(MoveTemp(Objects2), GeomToOriginalParticlesHack) : nullptr;
+			TUniquePtr<FImplicitObjectUnionClustered>& ChildrenSpatial = Parent->ChildrenSpatial();
+			ChildrenSpatial.Reset(
+				Objects2.Num() ? 
+				new Chaos::FImplicitObjectUnionClustered(MoveTemp(Objects2), GeomToOriginalParticlesHack) : 
+				nullptr);
 		}
 
 		TArray<TVector<T, d>> CleanedPoints;
 		if (!Parameters.CollisionParticles)
 		{
-			if (Parameters.bCleanCollisionParticles)
-			{
-				CleanedPoints = CleanCollisionParticles(OriginalPoints, ClusterSnapDistance);
-			}
-			else
-			{
-				CleanedPoints = OriginalPoints;
-			}
+			CleanedPoints = 
+				Parameters.bCleanCollisionParticles ? 
+				CleanCollisionParticles(OriginalPoints, ClusterSnapDistance) : 
+				OriginalPoints;
 		}
 
 		if (ProxyGeometry)
 		{
 			//ensureMsgf(false, TEXT("Checking usage with proxy"));
 			//@coverage {production}
-			MParticles.SetGeometry(NewIndex, ProxyGeometry);
+			Parent->SetGeometry(ProxyGeometry);
 		}
 		else if (Objects.Num() == 0)
 		{
 			//ensureMsgf(false, TEXT("Checking usage with no proxy and no objects"));
 			//@coverage : {production}
-			MParticles.SetGeometry(NewIndex, Chaos::TSerializablePtr<Chaos::FImplicitObject>());
+			Parent->SetGeometry(Chaos::TSerializablePtr<Chaos::FImplicitObject>());
 		}
 		else
 		{
@@ -1375,59 +1505,62 @@ namespace Chaos
 							}
 						}
 					}
-					MParticles.SetDynamicGeometry(NewIndex, MoveTemp(LevelSet));
+					Parent->SetDynamicGeometry(MoveTemp(LevelSet));
 				}
 				else
 				{
-					MParticles.SetDynamicGeometry(NewIndex, MakeUnique<TSphere<T, d>>(TVector<T, d>(0), BoundsExtents.Size() * 0.5f));
+					Parent->SetDynamicGeometry(
+						MakeUnique<TSphere<T, d>>(TVector<T, d>(0), BoundsExtents.Size() * 0.5f));
 				}
 			}
-			else
+			else // !UseLevelsetCollision
 			{
 				if (!bUsingMultiChildProxy && Objects.Num() == 1)
 				{
 					// @coverage:{confidence tests}
 					//ensureMsgf(false, TEXT("Checking no proxy, not level set, a single object"));
-					MParticles.SetDynamicGeometry(NewIndex, MoveTemp(Objects[0]));
+					Parent->SetDynamicGeometry(MoveTemp(Objects[0]));
 				}
 				else
 				{
 					QUICK_SCOPE_CYCLE_COUNTER(UnionBVH);
 					// @coverage : { confidence tests}
 					//ensureMsgf(false, TEXT("Checking no proxy, not levelset, and multiple objects"));
-					MParticles.SetDynamicGeometry(NewIndex, MakeUnique<FImplicitObjectUnionClustered>(MoveTemp(Objects), GeomToOriginalParticlesHack));
+					Parent->SetDynamicGeometry(
+						MakeUnique<FImplicitObjectUnionClustered>(
+							MoveTemp(Objects), GeomToOriginalParticlesHack));
 				}
 			}
 		}
 
-		if (bUseParticleImplicit && MParticles.DynamicGeometry(NewIndex)) //if children are ignore analytic and this is a dynamic geom, mark it too. todo(ocohen): clean this up
+		//if children are ignore analytic and this is a dynamic geom, mark it too. todo(ocohen): clean this up
+		if (bUseParticleImplicit && Parent->DynamicGeometry()) 
 		{
-			MParticles.DynamicGeometry(NewIndex)->SetDoCollide(false);
+			Parent->DynamicGeometry()->SetDoCollide(false);
 		}
 
 		if (Parameters.CollisionParticles)
 		{
 			SCOPE_CYCLE_COUNTER(STAT_UpdateGeometry_CopyPoints);
-			MParticles.CollisionParticles(NewIndex).Reset(Parameters.CollisionParticles);
-
+			Parent->CollisionParticles().Reset(Parameters.CollisionParticles);
 		}
 		else
 		{
 			{
 				SCOPE_CYCLE_COUNTER(STAT_UpdateGeometry_GatherPoints);
-				MParticles.CollisionParticlesInitIfNeeded(NewIndex);
-				MParticles.CollisionParticles(NewIndex)->Resize(0);
-				MParticles.CollisionParticles(NewIndex)->AddParticles(CleanedPoints.Num());
+				Parent->CollisionParticlesInitIfNeeded();
+				TUniquePtr<TBVHParticles<T, d>>& CollisionParticles = Parent->CollisionParticles();
+				CollisionParticles->AddParticles(CleanedPoints.Num());
 				for (int32 i = 0; i < CleanedPoints.Num(); ++i)
 				{
-					MParticles.CollisionParticles(NewIndex)->X(i) = CleanedPoints[i];
+					CollisionParticles->X(i) = CleanedPoints[i];
 				}
 			}
 
 			if (bUseCollisionPoints)
 			{
 				SCOPE_CYCLE_COUNTER(STAT_UpdateGeometry_PointsBVH);
-				MParticles.CollisionParticles(NewIndex)->UpdateAccelerationStructures();
+				Parent->CollisionParticles()->UpdateAccelerationStructures();
 			}
 		}
 	}
@@ -1437,7 +1570,8 @@ namespace Chaos
 
 	DECLARE_CYCLE_STAT(TEXT("ComputeStrainFromCollision"), STAT_ComputeStrainFromCollision, STATGROUP_Chaos);
 	template<class FPBDRigidsEvolution, class FPBDCollisionConstraint, class T, int d>
-	void TPBDRigidClustering<FPBDRigidsEvolution, FPBDCollisionConstraint, T, d>::ComputeStrainFromCollision(const FPBDCollisionConstraint& CollisionRule)
+	void TPBDRigidClustering<FPBDRigidsEvolution, FPBDCollisionConstraint, T, d>::ComputeStrainFromCollision(
+		const FPBDCollisionConstraint& CollisionRule)
 	{
 		SCOPE_CYCLE_COUNTER(STAT_ComputeStrainFromCollision);
 		FClusterMap& MParentToChildren = GetChildrenMap();
@@ -1451,48 +1585,60 @@ namespace Chaos
 				continue;
 			}
 
-			auto ComputeStrainLambda = [&](const TPBDRigidClusteredParticleHandle<T, d>* Cluster, const TArray<uint32>& ParentToChildren) {
+			auto ComputeStrainLambda = [&](
+				const TPBDRigidClusteredParticleHandle<T, d>* Cluster, 
+				const TArray<TPBDRigidParticleHandle<T, d>*>& ParentToChildren)
+			{
 				const TRigidTransform<T, d> WorldToClusterTM = TRigidTransform<T, d>(Cluster->P(), Cluster->Q());
 				const TVector<T, d> ContactLocationClusterLocal = WorldToClusterTM.InverseTransformPosition(ContactHandle->GetContactLocation());
 				TAABB<T, d> ContactBox(ContactLocationClusterLocal, ContactLocationClusterLocal);
 				ContactBox.Thicken(ClusterDistanceThreshold);
 				if (Cluster->ChildrenSpatial())
 				{
-					const TArray<int32> Intersections = Cluster->ChildrenSpatial()->FindAllIntersectingChildren(ContactBox);
-					for (int32 Child : Intersections)
+					const TArray<TPBDRigidParticleHandle<FReal, 3>*> Intersections = Cluster->ChildrenSpatial()->FindAllIntersectingChildren(ContactBox);
+					for (TPBDRigidParticleHandle<FReal, 3>* Child : Intersections)
 					{
-						if (const TUniquePtr<TMultiChildProxyData<T, d>>& ProxyData = MParticles.MultiChildProxyData(Child))
+						if (TPBDRigidClusteredParticleHandle<FReal, 3>* ClusteredChild = Child->CastToClustered())
 						{
-							//multi child so get its children
-							const int32 KeyChild = ProxyData->KeyChild;
-							const TRigidTransform<T, d> ProxyToCluster = ProxyData->RelativeToKeyChild * MParticles.ChildToParent(KeyChild);
-							const TVector<T, d> ContactLocationProxyLocal = ProxyToCluster.InverseTransformPosition(ContactLocationClusterLocal);
-							TAABB<T, d> ContactBoxProxy(ContactLocationProxyLocal, ContactLocationProxyLocal);
-							ContactBoxProxy.Thicken(ClusterDistanceThreshold);
-							if (MParticles.ChildrenSpatial(Child))
+							const TUniquePtr<TMultiChildProxyData<T, d>>& ProxyData = ClusteredChild->MultiChildProxyData();
+							const TPBDRigidParticleHandle<T, d>* KeyChild = ProxyData ? ProxyData->KeyChild : nullptr;
+							const TPBDRigidClusteredParticleHandle<T, d>* ClusteredKeyChild = KeyChild ? KeyChild->CastToClustered() : nullptr;
+							if (ClusteredKeyChild)
 							{
-								const TArray<int32> SubIntersections = MParticles.ChildrenSpatial(Child)->FindAllIntersectingChildren(ContactBoxProxy);
-								for (int32 SubChild : SubIntersections)
+								//multi child so get its children
+								const TRigidTransform<T, d> ProxyToCluster = ProxyData->RelativeToKeyChild * ClusteredKeyChild->ChildToParent();
+								const TVector<T, d> ContactLocationProxyLocal = ProxyToCluster.InverseTransformPosition(ContactLocationClusterLocal);
+								TAABB<T, d> ContactBoxProxy(ContactLocationProxyLocal, ContactLocationProxyLocal);
+								ContactBoxProxy.Thicken(ClusterDistanceThreshold);
+								if (ClusteredChild->ChildrenSpatial())
 								{
-									MParticles.CollisionImpulses(SubChild) += ContactHandle->GetAccumulatedImpulse().Size();
+									const TArray<TPBDRigidParticleHandle<FReal, 3>*> SubIntersections = 
+										ClusteredChild->ChildrenSpatial()->FindAllIntersectingChildren(ContactBoxProxy);
+									for (TPBDRigidParticleHandle<FReal, 3>* SubChild : SubIntersections)
+									{
+										if (TPBDRigidClusteredParticleHandle<FReal, 3>* ClusteredSubChild = SubChild->CastToClustered())
+										{
+											ClusteredSubChild->CollisionImpulses() += ContactHandle->GetAccumulatedImpulse().Size();
+										}
+									}
 								}
 							}
-						}
-						else
-						{
-							MParticles.CollisionImpulses(Child) += ContactHandle->GetAccumulatedImpulse().Size();
+							else
+							{
+								ClusteredChild->CollisionImpulses() += ContactHandle->GetAccumulatedImpulse().Size();
+							}
 						}
 					}
 				}
 			};
 
 			TVector<const TGeometryParticleHandle<T, d>*, 2> ConstrainedParticles = ContactHandle->GetConstrainedParticles();
-			if (auto ChildrenPtr = MParentToChildren.Find(ConstrainedParticles[0]))
+			if (const TArray<TPBDRigidParticleHandle<T, d>*>* ChildrenPtr = MParentToChildren.Find(ConstrainedParticles[0]->CastToRigidParticle()))
 			{
 				ComputeStrainLambda(ConstrainedParticles[0]->CastToClustered(), *ChildrenPtr);
 			}
 
-			if (auto ChildrenPtr = MParentToChildren.Find(ConstrainedParticles[1]))
+			if (const TArray<TPBDRigidParticleHandle<T, d>*>* ChildrenPtr = MParentToChildren.Find(ConstrainedParticles[1]->CastToRigidParticle()))
 			{
 				ComputeStrainLambda(ConstrainedParticles[1]->CastToClustered(), *ChildrenPtr);
 			}
@@ -1508,20 +1654,18 @@ namespace Chaos
 		SCOPE_CYCLE_COUNTER(STAT_ResetCollisionImpulseArray);
 		if (MCollisionImpulseArrayDirty)
 		{
-			for (uint32 i = 0; i < MParticles.Size(); i++)
-			{
-				MParticles.CollisionImpulses(i) = (T)0.0;
-			}
+			TPBDRigidsSOAs<T, d>& ParticleStructures = MEvolution.GetParticles();
+			ParticleStructures.GetGeometryCollectionParticles().CollisionImpulsesArray().Fill(0.0f);
+			ParticleStructures.GetClusteredParticles().CollisionImpulsesArray().Fill(0.0f);
 			MCollisionImpulseArrayDirty = false;
 		}
 	}
 
 	template<class FPBDRigidsEvolution, class FPBDCollisionConstraint, class T, int d>
-	void TPBDRigidClustering<FPBDRigidsEvolution, FPBDCollisionConstraint, T, d>::DisableCluster(TPBDRigidClusteredParticleHandle<T,d>* ClusteredParticle)
+	void TPBDRigidClustering<FPBDRigidsEvolution, FPBDCollisionConstraint, T, d>::DisableCluster(
+		TPBDRigidClusteredParticleHandle<T,d>* ClusteredParticle)
 	{
 		// #note: we don't recursively descend to the children
-		FClusterMap& ParentToChildren = GetChildrenMap();
-
 		MEvolution.DisableParticle(ClusteredParticle);
 
 		if (MoveClustersWhenDeactivated)
@@ -1531,96 +1675,102 @@ namespace Chaos
 			ClusteredParticle->V() = FVector(0.f);            // HACK : Move them away to prevent reactivation.
 		}
 
-		TopLevelClusterParents.Remove(ClusteredParticle->TransientParticleIndex());
+		TopLevelClusterParents.Remove(ClusteredParticle);
+		TArray<TPBDRigidParticleHandle<T, d>*>& ParentToChildren = MChildren[ClusteredParticle];
 		ParentToChildren.Remove(ClusteredParticle);
 		ClusteredParticle->ClusterIds() = ClusterId();
 		ClusteredParticle->ClusterGroupIndex() = 0;
-		MActiveRemovalIndices.Remove(ClusteredParticle->TransientParticleIndex());
+		MActiveRemovalIndices.Remove(ClusteredParticle);
 	}
 
 	template<class FPBDRigidsEvolution, class FPBDCollisionConstraint, class T, int d>
-	void TPBDRigidClustering<FPBDRigidsEvolution, FPBDCollisionConstraint, T, d>::DisableParticleWithBreakEvent(uint32 ClusterIndex)
+	void TPBDRigidClustering<FPBDRigidsEvolution, FPBDCollisionConstraint, T, d>::DisableParticleWithBreakEvent(
+		Chaos::TPBDRigidParticleHandle<float, 3>* Particle)
 	{
-		DisableCluster(MParticles.Handle(ClusterIndex));
+		DisableCluster(Particle->CastToClustered());
 
 		if (DoGenerateBreakingData)
 		{
-			int32 NewIdx = MAllClusterBreakings.Add(TBreakingData<float, 3>());
-
-			TBreakingData<float, 3>& ClusterBreakingItem = MAllClusterBreakings[NewIdx];
-			ClusterBreakingItem.ParticleIndex = ClusterIndex;
-			ClusterBreakingItem.Location = MParticles.X(ClusterIndex);
-			ClusterBreakingItem.Velocity = MParticles.V(ClusterIndex);
-			ClusterBreakingItem.AngularVelocity = MParticles.W(ClusterIndex);
-			ClusterBreakingItem.Mass = MParticles.M(ClusterIndex);
+			const int32 NewIdx = MAllClusterBreakings.Add(TBreakingData<float, 3>());
+			TBreakingData<float, 3>& ClusterBreak = MAllClusterBreakings[NewIdx];
+			ClusterBreak.Particle = Particle->GTGeometryParticle();
+			ClusterBreak.Location = Particle->X();
+			ClusterBreak.Velocity = Particle->V();
+			ClusterBreak.AngularVelocity = Particle->W();
+			ClusterBreak.Mass = Particle->M();
 		}
 	}
 
 	DECLARE_CYCLE_STAT(TEXT("TPBDRigidClustering<>::UpdateConnectivityGraphUsingPointImplicit"), STAT_UpdateConnectivityGraphUsingPointImplicit, STATGROUP_Chaos);
 	template<class FPBDRigidsEvolution, class FPBDCollisionConstraint, class T, int d>
-	void TPBDRigidClustering<FPBDRigidsEvolution, FPBDCollisionConstraint, T, d>::UpdateConnectivityGraphUsingPointImplicit(uint32 ClusterIndex, const FClusterCreationParameters<T>& Parameters)
+	void TPBDRigidClustering<FPBDRigidsEvolution, FPBDCollisionConstraint, T, d>::UpdateConnectivityGraphUsingPointImplicit(
+		Chaos::TPBDRigidClusteredParticleHandle<float, 3>* Parent,
+		const FClusterCreationParameters<T>& Parameters)
 	{
 		SCOPE_CYCLE_COUNTER(STAT_UpdateConnectivityGraphUsingPointImplicit);
-		FClusterMap& MParentToChildren = GetChildrenMap();
 
-		if (UseConnectivity)
+		if (!UseConnectivity)
 		{
-			T Delta = FMath::Min(FMath::Max(Parameters.CoillisionThicknessPercent, (T)0), T(1));
-			const TArray<uint32>& Children = MParentToChildren[MParticles.Handle(ClusterIndex)];
-			for (int32 i = 0; i < Children.Num(); ++i)
+			return;
+		}
+
+		const T Delta = FMath::Min(FMath::Max(Parameters.CoillisionThicknessPercent, (T)0), T(1));
+		const TArray<TPBDRigidParticleHandle<T, d>*>& Children = MChildren[Parent];
+		for (int32 i = 0; i < Children.Num(); ++i)
+		{
+			TPBDRigidParticleHandle<T, d>* Child1 = Children[i];
+			if (!Child1->Geometry() || !Child1->Geometry()->HasBoundingBox())
+			{
+				continue;
+			}
+			const TVector<T,d>& Child1X = Child1->X();
+			if (!(ensure(!FMath::IsNaN(Child1X[0])) && ensure(!FMath::IsNaN(Child1X[1])) && ensure(!FMath::IsNaN(Child1X[2]))))
 			{
-				const uint32 Child1 = Children[i];
-				if (!MParticles.Geometry(Child1) || !MParticles.Geometry(Child1)->HasBoundingBox())
+				continue;
+			}
+			TRigidTransform<T, d> TM1 = TRigidTransform<T, d>(Child1X, Child1->R());
+			TBox<T, d> Box1 = Child1->Geometry()->BoundingBox();
+
+			const int32 Offset = i + 1;
+			const int32 NumRemainingChildren = Children.Num() - Offset;
+			typedef TPair<TPBDRigidParticleHandle<T, d>*, TPBDRigidParticleHandle<T, d>*> ParticlePair;
+			typedef TArray<ParticlePair> ParticlePairArray;
+			TArray<ParticlePairArray> Connections;
+			Connections.Init(ParticlePairArray(), NumRemainingChildren);
+			PhysicsParallelFor(NumRemainingChildren, [&](int32 Idx) 
+			{
+				const int32 ChildrenIdx = Offset + Idx;
+				TPBDRigidParticleHandle<T, d>* Child2 = Children[ChildrenIdx];
+				if(!Child2->CollisionParticles())
+					return;
+
+				const TVector<T, d>& Child2X = Child2->X();
+				if (!(ensure(!FMath::IsNaN(Child2X[0])) && ensure(!FMath::IsNaN(Child2X[1])) && ensure(!FMath::IsNaN(Child2X[2]))))
+					return;
+
+				const TRigidTransform<T, d> TM = TM1.GetRelativeTransform(TRigidTransform<T, d>(Child2X, Child2->R()));
+
+				bool bCollided = false;
+				for (uint32 CollisionIdx = 0; !bCollided && CollisionIdx < Child2->CollisionParticles()->Size(); ++CollisionIdx)
 				{
-					continue;
+					const TVector<T, d> LocalPoint = 
+						TM.TransformPositionNoScale(Child2->CollisionParticles()->X(CollisionIdx));
+					const T Phi = Child1->Geometry()->SignedDistance(LocalPoint - (LocalPoint * Delta));
+					if (Phi < 0.0)
+						bCollided = true;
 				}
-				if (!(ensure(!FMath::IsNaN(MParticles.X(Child1)[0])) && ensure(!FMath::IsNaN(MParticles.X(Child1)[1])) && ensure(!FMath::IsNaN(MParticles.X(Child1)[2]))))
+				if (bCollided)
 				{
-					continue;
+					Connections[Idx].Add(ParticlePair(Child1, Child2));
 				}
-				TRigidTransform<T, d> TM1 = TRigidTransform<T, d>(MParticles.X(Child1), MParticles.R(Child1));
-				TAABB<T, d> Box1 = MParticles.Geometry(Child1)->BoundingBox();
-
-				TArray<TArray<TPair<uint32, uint32>>> Connections;
-				Connections.Init(TArray<TPair<uint32, uint32>>(), Children.Num() - (i + 1));
-				PhysicsParallelFor(Children.Num() - (i + 1), [&](int32 Idx) {
-					const uint32 Child2 = Children[Idx + i + 1];
-					if (!MParticles.CollisionParticles(Child2))
-					{
-						return;
-					}
-					if (!(ensure(!FMath::IsNaN(MParticles.X(Child2)[0])) && ensure(!FMath::IsNaN(MParticles.X(Child2)[1])) && ensure(!FMath::IsNaN(MParticles.X(Child2)[2]))))
-					{
-						return;
-					}
-					TRigidTransform<T, d> TM = TM1.GetRelativeTransform(TRigidTransform<T, d>(MParticles.X(Child2), MParticles.R(Child2)));
-
-					bool bCollided = false;
-					for (uint32 CollisionIdx = 0; !bCollided && CollisionIdx < MParticles.CollisionParticles(Child2)->Size(); ++CollisionIdx)
-					{
-						TVector<T, d> Normal;
-						TVector<T, d> LocalPoint = TM.TransformPositionNoScale(MParticles.CollisionParticles(Child2)->X(CollisionIdx));
-						T Phi = MParticles.Geometry(Child1)->PhiWithNormal(LocalPoint - (LocalPoint * Delta), Normal);
-						if (Phi < 0.0)
-							bCollided = true;
-					}
-					if (bCollided)
-					{
-						Connections[Idx].Add(TPair<uint32, uint32>((uint32)Child1, (uint32)Child2));
-					}
-				});
+			});
 
-				// join results and make connections
-				for (const auto& ConnectionList : Connections)
+			// join results and make connections
+			for (const ParticlePairArray& ConnectionList : Connections)
+			{
+				for (const ParticlePair& Edge : ConnectionList)
 				{
-					for (const auto& Elem : ConnectionList)
-					{
-						if (MParticles.ConnectivityEdges(Elem.Key).Num() < Parameters.MaxNumConnections)
-						{
-							const T AvgStrain = (MParticles.Strains(Elem.Key) + MParticles.Strains(Elem.Value)) * (T)0.5f;
-							ConnectNodes(Elem.Key, Elem.Value, AvgStrain);
-						}
-					}
+					ConnectNodes(Edge.Key, Edge.Value);
 				}
 			}
 		}
@@ -1628,75 +1778,99 @@ namespace Chaos
 
 	DECLARE_CYCLE_STAT(TEXT("TPBDRigidClustering<>::FixConnectivityGraphUsingDelaunayTriangulation"), STAT_FixConnectivityGraphUsingDelaunayTriangulation, STATGROUP_Chaos);
 	template<class FPBDRigidsEvolution, class FPBDCollisionConstraint, class T, int d>
-	void TPBDRigidClustering<FPBDRigidsEvolution, FPBDCollisionConstraint, T, d>::FixConnectivityGraphUsingDelaunayTriangulation(uint32 ClusterIndex, const FClusterCreationParameters<T>& Parameters)
+	void TPBDRigidClustering<FPBDRigidsEvolution, FPBDCollisionConstraint, T, d>::FixConnectivityGraphUsingDelaunayTriangulation(
+		Chaos::TPBDRigidClusteredParticleHandle<float, 3>* Parent,
+		const FClusterCreationParameters<T>& Parameters)
 	{
 		// @todo(investigate) : This is trying to set multiple connections and throwing a warning in ConnectNodes
 		SCOPE_CYCLE_COUNTER(STAT_FixConnectivityGraphUsingDelaunayTriangulation);
 
-		const TArray<uint32>& Children = GetChildrenMap()[MParticles.Handle(ClusterIndex)];
+		const TArray<TPBDRigidParticleHandle<T, d>*>& Children = MChildren[Parent];
 
 		// Compute Delaunay neighbor graph on children centers
 		TArray<FVector> Pts;
 		Pts.AddUninitialized(Children.Num());
 		for (int32 i = 0; i < Children.Num(); ++i)
 		{
-			Pts[i] = MParticles.X(Children[i]);
+			Pts[i] = Children[i]->X();
 		}
-		TArray<TArray<int32>> Neighbors;
+		TArray<TArray<int32>> Neighbors; // Indexes into Children
 		VoronoiNeighbors(Pts, Neighbors);
 
 		// Build a UnionFind graph to find (indirectly) connected children
 		struct UnionFindInfo
 		{
-			int32 GroupIdx;
+			TPBDRigidParticleHandle<T, d>* GroupId;
 			int32 Size;
 		};
-		TArray<UnionFindInfo> UnionInfo;
-		UnionInfo.SetNumUninitialized(Children.Num());
-		TMap<int32, int32> ChildReverseIdx;
-		for (int32 i = 0; i < Children.Num(); ++i)
+		TMap<TPBDRigidParticleHandle<T, d>*, UnionFindInfo> UnionInfo;
+		UnionInfo.Reserve(Children.Num());
+
+		// Initialize UnionInfo:
+		//		0: GroupId = Children[0], Size = 1
+		//		1: GroupId = Children[1], Size = 1
+		//		2: GroupId = Children[2], Size = 1
+		//		3: GroupId = Children[3], Size = 1
+
+		for(TPBDRigidParticleHandle<T, d>* Child : Children)
 		{
-			UnionInfo[i].GroupIdx = i;
-			UnionInfo[i].Size = 1;
-			ChildReverseIdx.Add(Children[i], i);
+			UnionInfo.Add(Child, { Child, 1 }); // GroupId, Size
 		}
-		auto FindGroup = [&](int Idx) {
-			int GroupIdx = Idx;
-			if (GroupIdx >= 0 && GroupIdx < UnionInfo.Num())
+
+		auto FindGroup = [&](TPBDRigidParticleHandle<T, d>* Id) 
+		{
+			TPBDRigidParticleHandle<T, d>* GroupId = Id;
+			if (GroupId)
 			{
 				int findIters = 0;
-				while (UnionInfo[GroupIdx].GroupIdx != GroupIdx)
+				while (UnionInfo[GroupId].GroupId != GroupId)
 				{
 					ensure(findIters++ < 10); // if this while loop iterates more than a few times, there's probably a bug in the unionfind
-					UnionInfo[GroupIdx].GroupIdx = UnionInfo[UnionInfo[GroupIdx].GroupIdx].GroupIdx;
-					GroupIdx = UnionInfo[GroupIdx].GroupIdx;
-					if (GroupIdx < 0 || GroupIdx >= UnionInfo.Num()) break; // this is a error exit
+					auto& CurrInfo = UnionInfo[GroupId];
+					auto& NextInfo = UnionInfo[CurrInfo.GroupId];
+					CurrInfo.GroupId = NextInfo.GroupId;
+					GroupId = NextInfo.GroupId;
+					if (!GroupId) break; // error condidtion
 				}
 			}
-			return GroupIdx;
+			return GroupId;
 		};
-		auto MergeGroup = [&](int A, int B) {
-			int GroupA = FindGroup(A);
-			int GroupB = FindGroup(B);
+
+		// MergeGroup(Children[0], Children[1])
+		//		0: GroupId = Children[1], Size = 0
+		//		1: GroupId = Children[1], Size = 2
+		//		2: GroupId = Children[2], Size = 1
+		//		3: GroupId = Children[3], Size = 1
+
+		auto MergeGroup = [&](TPBDRigidParticleHandle<T, d>* A, TPBDRigidParticleHandle<T, d>* B) 
+		{
+			TPBDRigidParticleHandle<T, d>* GroupA = FindGroup(A);
+			TPBDRigidParticleHandle<T, d>* GroupB = FindGroup(B);
 			if (GroupA == GroupB)
 			{
 				return;
 			}
+			// Make GroupA the smaller of the two
 			if (UnionInfo[GroupA].Size > UnionInfo[GroupB].Size)
 			{
 				Swap(GroupA, GroupB);
 			}
-			UnionInfo[GroupA].GroupIdx = GroupB;
+			// Overwrite GroupA with GroupB
+			UnionInfo[GroupA].GroupId = GroupB;
 			UnionInfo[GroupB].Size += UnionInfo[GroupA].Size;
+			UnionInfo[GroupA].Size = 0; // not strictly necessary, but more correct
 		};
+
+		// Merge all groups with edges connecting them.
 		for (int32 i = 0; i < Children.Num(); ++i)
 		{
-			for (const TConnectivityEdge<T>& Edge : MParticles.ConnectivityEdges(Children[i]))
+			TPBDRigidParticleHandle<T, d>* Child = Children[i];
+			const TArray<TConnectivityEdge<T>>& Edges = Child->CastToClustered()->ConnectivityEdges();
+			for (const TConnectivityEdge<T>& Edge : Edges)
 			{
-				int* nbr = ChildReverseIdx.Find(Edge.Sibling);
-				if (nbr)
+				if (UnionInfo.Contains(Edge.Sibling))
 				{
-					MergeGroup(i, *nbr);
+					MergeGroup(Child, Edge.Sibling);
 				}
 			}
 		}
@@ -1704,37 +1878,42 @@ namespace Chaos
 		// Find candidate edges from the Delaunay graph to consider adding
 		struct LinkCandidate
 		{
-			int32 A, B;
+			//int32 A, B;
+			TPBDRigidParticleHandle<T, d>* A;
+			TPBDRigidParticleHandle<T, d>* B;
 			float DistSq;
 		};
 		TArray<LinkCandidate> Candidates;
-		float AlwaysAcceptBelowDistSqThreshold = 50.f*50.f*100.f*MClusterConnectionFactor;
+		Candidates.Reserve(Neighbors.Num());
+
+		const float AlwaysAcceptBelowDistSqThreshold = 50.f*50.f*100.f*MClusterConnectionFactor;
 		for (int32 i = 0; i < Neighbors.Num(); i++)
 		{
-			for (int32 Nbr : Neighbors[i])
+			TPBDRigidParticleHandle<T, d>* Child1 = Children[i];
+			const TArray<int32>& Child1Neighbors = Neighbors[i];
+			for (const int32 Nbr : Child1Neighbors)
 			{
 				if (Nbr < i)
 				{ // assume we'll get the symmetric connection; don't bother considering this one
 					continue;
 				}
+				TPBDRigidParticleHandle<T, d>* Child2 = Children[Nbr];
 
-				float DistSq = FVector::DistSquared(Pts[i], Pts[Nbr]);
+				const float DistSq = FVector::DistSquared(Pts[i], Pts[Nbr]);
 				if (DistSq < AlwaysAcceptBelowDistSqThreshold)
 				{ // below always-accept threshold: don't bother adding to candidates array, just merge now
-					MergeGroup(i, Nbr);
-					int32 ChildA = Children[i], ChildB = Children[Nbr];
-					const T AvgStrain = (MParticles.Strains(ChildA) + MParticles.Strains(ChildB)) * (T)0.5f;
-					ConnectNodes(ChildA, ChildB, AvgStrain);
+					MergeGroup(Child1, Child2);
+					ConnectNodes(Child1, Child2);
 					continue;
 				}
 
-				if (FindGroup(i) == FindGroup(Nbr))
+				if (FindGroup(Child1) == FindGroup(Child2))
 				{ // already part of the same group so we don't need Delaunay edge  
 					continue;
 				}
 
 				// add to array to sort and add as-needed
-				Candidates.Add({ i, Nbr, DistSq });
+				Candidates.Add({ Child1, Child2, DistSq });
 			}
 		}
 
@@ -1742,108 +1921,142 @@ namespace Chaos
 		Candidates.Sort([](const LinkCandidate& A, const LinkCandidate& B) { return A.DistSq < B.DistSq; });
 		for (const LinkCandidate& Candidate : Candidates)
 		{
-			int32 A = Candidate.A, B = Candidate.B;
-			if (FindGroup(A) != FindGroup(B))
+			TPBDRigidParticleHandle<T, d>* Child1 = Candidate.A;
+			TPBDRigidParticleHandle<T, d>* Child2 = Candidate.B;
+			if (FindGroup(Child1) != FindGroup(Child2))
 			{
-				MergeGroup(A, B);
-				int32 ChildA = Children[A], ChildB = Children[B];
-				const T AvgStrain = (MParticles.Strains(ChildA) + MParticles.Strains(ChildB)) * (T)0.5f;
-				ConnectNodes(ChildA, ChildB, AvgStrain);
+				MergeGroup(Child1, Child2);
+				ConnectNodes(Child1, Child2);
 			}
 		}
 	}
 
 	DECLARE_CYCLE_STAT(TEXT("TPBDRigidClustering<>::UpdateConnectivityGraphUsingDelaunayTriangulation"), STAT_UpdateConnectivityGraphUsingDelaunayTriangulation, STATGROUP_Chaos);
+
 	template<class FPBDRigidsEvolution, class FPBDCollisionConstraint, class T, int d>
-	void TPBDRigidClustering<FPBDRigidsEvolution, FPBDCollisionConstraint, T, d>::UpdateConnectivityGraphUsingDelaunayTriangulation(uint32 ClusterIndex, const FClusterCreationParameters<T>& Parameters)
+	void TPBDRigidClustering<FPBDRigidsEvolution, FPBDCollisionConstraint, T, d>::UpdateConnectivityGraphUsingDelaunayTriangulation(
+		Chaos::TPBDRigidClusteredParticleHandle<float, 3>* Parent,
+		const FClusterCreationParameters<T>& Parameters)
 	{
 		SCOPE_CYCLE_COUNTER(STAT_UpdateConnectivityGraphUsingDelaunayTriangulation);
 
-		const TArray<uint32>& Children = GetChildrenMap()[MParticles.Handle(ClusterIndex)];
+		const TArray<TPBDRigidParticleHandle<T, d>*>& Children = MChildren[Parent];
 
 		TArray<FVector> Pts;
 		Pts.AddUninitialized(Children.Num());
 		for (int32 i = 0; i < Children.Num(); ++i)
 		{
-			Pts[i] = MParticles.X(Children[i]);
+			Pts[i] = Children[i]->X();
 		}
-
 		TArray<TArray<int>> Neighbors;
 		VoronoiNeighbors(Pts, Neighbors);
 
-		TSet<TPair<int32, int32>> UniqueEdges;
+		TSet<TPair<TPBDRigidParticleHandle<T, d>*, TPBDRigidParticleHandle<T, d>*>> UniqueEdges;
 		for (int32 i = 0; i < Neighbors.Num(); i++)
 		{
 			for (int32 j = 0; j < Neighbors[i].Num(); j++)
 			{
-				const int32 First = Children[i];
-				const int32 Sibling = Children[Neighbors[i][j]];
-				const bool bFirstSmaller = First < Sibling;
-				TPair<int32, int32> SortedPair(bFirstSmaller ? First : Sibling, bFirstSmaller ? Sibling : First);
-				if (!UniqueEdges.Contains(SortedPair))
+				TPBDRigidParticleHandle<T, d>* Child1 = Children[i];
+				TPBDRigidParticleHandle<T, d>* Child2 = Children[Neighbors[i][j]];
+				const bool bFirstSmaller = Child1 < Child2;
+				TPair<TPBDRigidParticleHandle<T, d>*, TPBDRigidParticleHandle<T, d>*> SortedPair(
+					bFirstSmaller ? Child1 : Child2, 
+					bFirstSmaller ? Child2 : Child1);
+				if (!UniqueEdges.Find(SortedPair))
 				{
 					// this does not use ConnectNodes because Neighbors is bi-direction : as in (1,2),(2,1)
-					const T AvgStrain = (MParticles.Strains(First) + MParticles.Strains(Sibling)) * (T)0.5f;
-					//AddUniqueConnection(Children[i], Children[Neighbors[i][j]], AvgStrain);
-					ConnectNodes(First, Sibling, AvgStrain);
+					ConnectNodes(Child1, Child2);
 					UniqueEdges.Add(SortedPair);
 				}
 			}
 		}
 	}
 
+	//template<class FPBDRigidsEvolution, class FPBDCollisionConstraint, class T, int d>
+	//void TPBDRigidClustering<FPBDRigidsEvolution, FPBDCollisionConstraint, T, d>::AddUniqueConnection(uint32 Index1, uint32 Index2, T Strain)
+	//{
+	//	if (Index1 != Index2)
+	//	{
+	//		//todo(pref): This can be removed if we are sure there are no duplicate connections generated.
+	//		for (int32 i = 0; i < MParticles.ConnectivityEdges(Index1).Num(); i++)
+	//		{
+	//			if (MParticles.ConnectivityEdges(Index1)[i].Sibling == Index2)
+	//			{
+	//				// @todo(duplication connection) : re-enable post GDC.  
+	//				// FixConnectivityGraphUsingDelaunayTriangulation attempts to add multiple connections.
+	//				// so commenting out this msg to remove the noise from the confidence test. 
+	//				// ensureMsgf(false, TEXT("Duplicate graph connection."));
+	//				return;
+	//			}
+	//		}
+	//
+	//		MParticles.ConnectivityEdges(Index1).Add({ Index2, Strain });
+	//	}
+	//}
+
 	template<class FPBDRigidsEvolution, class FPBDCollisionConstraint, class T, int d>
-	void TPBDRigidClustering<FPBDRigidsEvolution, FPBDCollisionConstraint, T, d>::AddUniqueConnection(uint32 Index1, uint32 Index2, T Strain)
+	void TPBDRigidClustering<FPBDRigidsEvolution, FPBDCollisionConstraint, T, d>::ConnectNodes(
+		TPBDRigidParticleHandle<T, d>* Child1,
+		TPBDRigidParticleHandle<T, d>* Child2)
 	{
-		if (Index1 != Index2)
-		{
-			//todo(pref): This can be removed if we are sure there are no duplicate connections generated.
-			for (int32 i = 0; i < MParticles.ConnectivityEdges(Index1).Num(); i++)
-			{
-				if (MParticles.ConnectivityEdges(Index1)[i].Sibling == Index2)
-				{
-					// @todo(duplication connection) : re-enable post GDC.  
-					// FixConnectivityGraphUsingDelaunayTriangulation attempts to add multiple connections.
-					// so commenting out this msg to remove the noise from the confidence test. 
-					// ensureMsgf(false, TEXT("Duplicate graph connection."));
-					return;
-				}
-			}
+		check(Child1 != Child2);
+		TPBDRigidClusteredParticleHandle<T, d>* ClusteredChild1 = Child1->CastToClustered();
+		TPBDRigidClusteredParticleHandle<T, d>* ClusteredChild2 = Child2->CastToClustered();
+		ConnectNodes(ClusteredChild1, ClusteredChild2);
+	}
 
-			MParticles.ConnectivityEdges(Index1).Add({ Index2, Strain });
+	template<class FPBDRigidsEvolution, class FPBDCollisionConstraint, class T, int d>
+	void TPBDRigidClustering<FPBDRigidsEvolution, FPBDCollisionConstraint, T, d>::ConnectNodes(
+		TPBDRigidClusteredParticleHandle<T, d>* ClusteredChild1,
+		TPBDRigidClusteredParticleHandle<T, d>* ClusteredChild2)
+	{
+		check(ClusteredChild1 && ClusteredChild1);
+		if (ClusteredChild1 == ClusteredChild2)
+			return;
+		const T AvgStrain = (ClusteredChild1->Strains() + ClusteredChild2->Strains()) * (T)0.5;
+		TArray<TConnectivityEdge<T>>& Edges1 = ClusteredChild1->ConnectivityEdges();
+		TArray<TConnectivityEdge<T>>& Edges2 = ClusteredChild2->ConnectivityEdges();
+		if (//Edges1.Num() < Parameters.MaxNumConnections && 
+			!Edges1.FindByKey(ClusteredChild2))
+		{
+			Edges1.Add(TConnectivityEdge<T>(ClusteredChild2, AvgStrain));
+		}
+		if (//Edges2.Num() < Parameters.MaxNumConnections && 
+			!Edges2.FindByKey(ClusteredChild1))
+		{
+			Edges2.Add(TConnectivityEdge<T>(ClusteredChild1, AvgStrain));
 		}
 	}
 
 	template<class FPBDRigidsEvolution, class FPBDCollisionConstraint, class T, int d>
-	void TPBDRigidClustering<FPBDRigidsEvolution, FPBDCollisionConstraint, T, d>::ConnectNodes(uint32 Index1, uint32 Index2, T Strain)
+	void TPBDRigidClustering<FPBDRigidsEvolution, FPBDCollisionConstraint, T, d>::RemoveNodeConnections(
+		TPBDRigidParticleHandle<T, d>* Child)
 	{
-		AddUniqueConnection(Index1, Index2, Strain);
-		AddUniqueConnection(Index2, Index1, Strain);
+		RemoveNodeConnections(Child->CastToClustered());
 	}
 
 	DECLARE_CYCLE_STAT(TEXT("TPBDRigidClustering<>::RemoveNodeConnections"), STAT_RemoveNodeConnections, STATGROUP_Chaos);
 	template<class FPBDRigidsEvolution, class FPBDCollisionConstraint, class T, int d>
-	void TPBDRigidClustering<FPBDRigidsEvolution, FPBDCollisionConstraint, T, d>::RemoveNodeConnections(uint32 ParticleIndex)
+	void TPBDRigidClustering<FPBDRigidsEvolution, FPBDCollisionConstraint, T, d>::RemoveNodeConnections(
+		TPBDRigidClusteredParticleHandle<T, d>* ClusteredChild)
 	{
 		SCOPE_CYCLE_COUNTER(STAT_RemoveNodeConnections);
-		checkSlow(0 <= ParticleIndex && ParticleIndex < MParticles.Size());
-
-		for (int32 i = MParticles.ConnectivityEdges(ParticleIndex).Num() - 1; 0 <= i; i--)
-		{
-			const TConnectivityEdge<T>& Edge = MParticles.ConnectivityEdges(ParticleIndex)[i];
-			for (int Idx = MParticles.ConnectivityEdges(Edge.Sibling).Num() - 1; Idx >= 0; --Idx)
-			{
-				const TConnectivityEdge<T>& OtherEdge = MParticles.ConnectivityEdges(Edge.Sibling)[Idx];
-				if (OtherEdge.Sibling == ParticleIndex)
-				{
-					MParticles.ConnectivityEdges(Edge.Sibling).RemoveAtSwap(Idx); //Note: we shouldn't have to keep searching, but sometimes these over subscribe, should fix
-				}
-			}
+		check(ClusteredChild);
+		TArray<TConnectivityEdge<T>>& Edges = ClusteredChild->ConnectivityEdges();
+		for (TConnectivityEdge<T>& Edge : Edges)
+		{
+			TArray<TConnectivityEdge<T>>& OtherEdges = Edge.Sibling->CastToClustered()->ConnectivityEdges();
+			const int32 Idx = OtherEdges.IndexOfByKey(ClusteredChild);
+			if (Idx != INDEX_NONE)
+				OtherEdges.RemoveAtSwap(Idx);
+			// Make sure there are no duplicates!
+			check(OtherEdges.IndexOfByKey(ClusteredChild) == INDEX_NONE);
 		}
-		MParticles.ConnectivityEdges(ParticleIndex).SetNum(0);
+		Edges.SetNum(0);
 	}
 
-	/*
+
+/*
 	template<class FPBDRigidsEvolution, class FPBDCollisionConstraint, class T, int d>
 	void TPBDRigidClustering<FPBDRigidsEvolution, FPBDCollisionConstraint, T, d>::ClearPendingClusterCounter(uint32 ClusterGroupID)
 	{
@@ -1855,7 +2068,9 @@ namespace Chaos
 
 using namespace Chaos;
 template class CHAOS_API Chaos::TPBDRigidClustering<FPBDRigidsEvolutionGBF, TPBDCollisionConstraints<float, 3>, float, 3>;
-template CHAOS_API void Chaos::UpdateClusterMassProperties<float, 3>(TPBDRigidClusteredParticles<float, 3>& Particles, const TArray<uint32>& Children, const uint32 ClusterIndex, const TRigidTransform<float, 3>* ForceMassOrientation,
-	const TArrayCollectionArray<TUniquePtr<TMultiChildProxyData<float, 3>>>* MMultiChildProxyData,
-	const TArrayCollectionArray<FMultiChildProxyId>* MMultiChildProxyId);
 
+
+template CHAOS_API void Chaos::UpdateClusterMassProperties<float, 3>(
+	Chaos::TPBDRigidClusteredParticleHandle<float, 3>* NewParticle,
+	TSet<TPBDRigidParticleHandle<float, 3>*>& Children,
+	const TRigidTransform<float, 3>* ForceMassOrientation);
diff --git a/Engine/Source/Runtime/Experimental/Chaos/Public/Chaos/ClusterCreationParameters.h b/Engine/Source/Runtime/Experimental/Chaos/Public/Chaos/ClusterCreationParameters.h
--- a/Engine/Source/Runtime/Experimental/Chaos/Public/Chaos/ClusterCreationParameters.h
+++ b/Engine/Source/Runtime/Experimental/Chaos/Public/Chaos/ClusterCreationParameters.h
@@ -2,6 +2,8 @@
 
 #pragma once
 
+#include "Chaos/ParticleHandle.h"
+#include "Chaos/ParticleHandleFwd.h"
 #include "Chaos/Defines.h"
 
 namespace Chaos
@@ -22,8 +24,8 @@ namespace Chaos
 			, bool bCopyCollisionParticlesIn = true
 			, bool bGenerateConnectionGraphIn = true
 			, EConnectionMethod ConnectionMethodIn = EConnectionMethod::PointImplicitAugmentedWithMinimalDelaunay
-			, TBVHParticles<T, 3>* CollisionParticlesIn = nullptr
-			, int32 RigidBodyIndexIn = INDEX_NONE
+			, TBVHParticles<float, 3>* CollisionParticlesIn = nullptr
+			, Chaos::TPBDRigidClusteredParticleHandle<float,3>* ClusterParticleHandleIn = nullptr
 		)
 			: CoillisionThicknessPercent(CoillisionThicknessPercentIn)
 			, MaxNumConnections(MaxNumConnectionsIn)
@@ -32,7 +34,7 @@ namespace Chaos
 			, bGenerateConnectionGraph(bGenerateConnectionGraphIn)
 			, ConnectionMethod(ConnectionMethodIn)
 			, CollisionParticles(CollisionParticlesIn)
-			, RigidBodyIndex(RigidBodyIndexIn)
+			, ClusterParticleHandle(ClusterParticleHandleIn)
 		{}
 
 		T CoillisionThicknessPercent;
@@ -42,6 +44,6 @@ namespace Chaos
 		bool bGenerateConnectionGraph;
 		EConnectionMethod ConnectionMethod;
 		TBVHParticles<T, 3>* CollisionParticles;
-		int32 RigidBodyIndex;
+		Chaos::TPBDRigidClusteredParticleHandle<float, 3>* ClusterParticleHandle;
 	};
 }
diff --git a/Engine/Source/Runtime/Experimental/Chaos/Public/Chaos/ExternalCollisionData.h b/Engine/Source/Runtime/Experimental/Chaos/Public/Chaos/ExternalCollisionData.h
--- a/Engine/Source/Runtime/Experimental/Chaos/Public/Chaos/ExternalCollisionData.h
+++ b/Engine/Source/Runtime/Experimental/Chaos/Public/Chaos/ExternalCollisionData.h
@@ -192,21 +192,21 @@ namespace Chaos
 	struct TBreakingData
 	{
 		TBreakingData()
-			: Location(TVector<T, d>((T)0.0))
+			: Particle(nullptr)
+			, Location(TVector<T, d>((T)0.0))
 			, Velocity(TVector<T, d>((T)0.0))
 			, AngularVelocity(TVector<T, d>((T)0.0))
 			, Mass((T)0.0)
-			, Particle(nullptr)
 			, ParticleIndex(INDEX_NONE)
 			, ParticleIndexMesh(INDEX_NONE)
 			, BoundingBox(TAABB<T, d>(TVector<T, d>((T)0.0), TVector<T, d>((T)0.0)))
 		{}
 
+		TGeometryParticle<T, d>* Particle; // 64 bits go first for better packing
 		TVector<T, d> Location;
 		TVector<T, d> Velocity;
 		TVector<T, d> AngularVelocity;
 		T Mass;
-		TGeometryParticle<T, d>* Particle;
 		int32 ParticleIndex; //#todo: remove this in favor of TGeometryParticle?
 		int32 ParticleIndexMesh; // If ParticleIndex points to a cluster then this index will point to an actual mesh in the cluster
 								 // It is important to be able to get extra data from the component
diff --git a/Engine/Source/Runtime/Experimental/Chaos/Public/Chaos/ImplicitObjectUnion.h b/Engine/Source/Runtime/Experimental/Chaos/Public/Chaos/ImplicitObjectUnion.h
--- a/Engine/Source/Runtime/Experimental/Chaos/Public/Chaos/ImplicitObjectUnion.h
+++ b/Engine/Source/Runtime/Experimental/Chaos/Public/Chaos/ImplicitObjectUnion.h
@@ -174,26 +174,20 @@ protected:
 	}
 
   protected:
-	TArray<TUniquePtr<FImplicitObject>> MObjects;
-	TAABB<FReal,3> MLocalBoundingBox;
-	TUniquePtr<FLargeImplicitObjectUnionData> LargeUnionData;	//only needed when there are many objects
-
 	//needed for serialization
 	FImplicitObjectUnion();
 	friend FImplicitObject;	//needed for serialization
-public:
+
+	TArray<TUniquePtr<FImplicitObject>> MObjects;
+	TAABB<FReal,3> MLocalBoundingBox;
+	TUniquePtr<FLargeImplicitObjectUnionData> LargeUnionData;	//only needed when there are many objects
 };
 
 class CHAOS_API FImplicitObjectUnionClustered: public FImplicitObjectUnion
 {
 public:
-
-	FImplicitObjectUnionClustered(TArray<TUniquePtr<FImplicitObject>>&& Objects,const TArray<int32>& OriginalParticleLookupHack = TArray<int32>());
-	FImplicitObjectUnionClustered()
-		: FImplicitObjectUnion()
-	{
-		Type = ImplicitObjectType::UnionClustered;
-	}
+	FImplicitObjectUnionClustered();
+	FImplicitObjectUnionClustered(TArray<TUniquePtr<FImplicitObject>>&& Objects, const TArray<TPBDRigidParticleHandle<FReal, 3>*>& OriginalParticleLookupHack = TArray<TPBDRigidParticleHandle<FReal, 3>*>());
 	FImplicitObjectUnionClustered(const FImplicitObjectUnionClustered& Other) = delete;
 	FImplicitObjectUnionClustered(FImplicitObjectUnionClustered&& Other);
 	virtual ~FImplicitObjectUnionClustered() = default;
@@ -203,11 +197,51 @@ public:
 		return ImplicitObjectType::UnionClustered;
 	}
 
-	TArray<int32> FindAllIntersectingChildren(const TAABB<FReal,3>& LocalBounds) const;
+	TArray<TPBDRigidParticleHandle<FReal, 3>*> FindAllIntersectingChildren(const TAABB<FReal, 3>& LocalBounds) const;
+
+#if CHAOS_PARTICLEHANDLE_TODO
 	TArray<int32> FindAllIntersectingChildren(const TSpatialRay<FReal,3>& LocalRay) const;
+	{
+		TArray<int32> IntersectingChildren;
+		if (LargeUnionData) //todo: make this work when hierarchy is not built
+		{
+			IntersectingChildren = LargeUnionData->Hierarchy.FindAllIntersections(LocalRay);
+			for (int32 i = IntersectingChildren.Num() - 1; i >= 0; --i)
+			{
+				const int32 Idx = IntersectingChildren[i];
+				if (Idx < MOriginalParticleLookupHack.Num())
+				{
+					IntersectingChildren[i] = MOriginalParticleLookupHack[Idx];
+				}
+				else
+				{
+					IntersectingChildren.RemoveAtSwap(i);
+				}
+			}
+			/*for (int32& Idx : IntersectingChildren)
+			{
+				Idx = MOriginalParticleLookupHack[Idx];
+			}*/
+		}
+		else
+		{
+			IntersectingChildren = MOriginalParticleLookupHack;
+		}
+
+		return IntersectingChildren;
+	}
+#endif
+
+	const TPBDRigidParticleHandle<FReal, 3>* FindParticleForImplicitObject(const FImplicitObject* Object) const
+	{ 
+		typedef TPBDRigidParticleHandle<FReal, 3>* ValueType;
+		const ValueType* Handle = MCollisionParticleLookupHack.Find(Object);
+		return Handle ? *Handle : nullptr;
+	}
 
 private:
-	TArray<int32> MOriginalParticleLookupHack;	//temp hack for finding original particles
-	TMap<const FImplicitObject*,int32> MCollisionParticleLookupHack;	//temp hack for finding collision particles
+	// Temp hack for finding original particles
+	TArray<TPBDRigidParticleHandle<FReal, 3>*> MOriginalParticleLookupHack;	
+	TMap<const FImplicitObject*,TPBDRigidParticleHandle<FReal, 3>*> MCollisionParticleLookupHack;	//temp hack for finding collision particles
 };
 }
diff --git a/Engine/Source/Runtime/Experimental/Chaos/Public/Chaos/PBDConstraintGraph.h b/Engine/Source/Runtime/Experimental/Chaos/Public/Chaos/PBDConstraintGraph.h
--- a/Engine/Source/Runtime/Experimental/Chaos/Public/Chaos/PBDConstraintGraph.h
+++ b/Engine/Source/Runtime/Experimental/Chaos/Public/Chaos/PBDConstraintGraph.h
@@ -144,9 +144,11 @@ namespace Chaos
 		const FConstraintData& GetConstraintData(int32 ConstraintDataIndex) const;
 
 		/**
-		 * Enable a particle after initialization. Adds the particle to the island of the parent particle if supplied.
+		 * Adds \p ChildParticle to the constraint graph.  Copies the sleeping 
+		 * state and island of \p ParentParticle to \p ChildParticle.  Does nothing
+		 * if \p ParentParticle is not supplied.
 		 */
-		void EnableParticle(TGeometryParticleHandle<FReal, 3>* Particle, const TGeometryParticleHandle<FReal, 3>* ParentParticle);
+		void EnableParticle(TGeometryParticleHandle<FReal, 3>* ChildParticle, const TGeometryParticleHandle<FReal, 3>* ParentParticle);
 
 		/**
 		 * Disable a particle and remove it from its island.
@@ -154,6 +156,13 @@ namespace Chaos
 		 */
 		void DisableParticle(TGeometryParticleHandle<FReal, 3>* Particle);
 
+		/**
+		 * Preallocate buffers for \p Num particles.
+		 *
+		 * Returns the number of new slots created.
+		 */
+		int32 ReserveParticles(const int32 Num);
+
 		void AddParticle(TGeometryParticleHandle<FReal, 3>* AddedParticle)
 		{
 			ParticleAdd(AddedParticle);
diff --git a/Engine/Source/Runtime/Experimental/Chaos/Public/Chaos/PBDGeometryCollectionParticles.h b/Engine/Source/Runtime/Experimental/Chaos/Public/Chaos/PBDGeometryCollectionParticles.h
--- a/Engine/Source/Runtime/Experimental/Chaos/Public/Chaos/PBDGeometryCollectionParticles.h
+++ b/Engine/Source/Runtime/Experimental/Chaos/Public/Chaos/PBDGeometryCollectionParticles.h
@@ -2,23 +2,23 @@
 #pragma once
 
 #include "Chaos/ArrayCollectionArray.h"
-#include "Chaos/PBDRigidParticles.h"
+#include "Chaos/PBDRigidClusteredParticles.h"
 
 namespace Chaos
 {
 	template<class T, int d>
-	class TPBDGeometryCollectionParticles : public TPBDRigidParticles<T, d>
+	class TPBDGeometryCollectionParticles : public TPBDRigidClusteredParticles<T, d>
 	{
 	public:
 		TPBDGeometryCollectionParticles()
-		    : TPBDRigidParticles<T, d>()
+		    : TPBDRigidClusteredParticles<T, d>()
 		{
 			InitHelper();
 		}
-		TPBDGeometryCollectionParticles(const TPBDRigidParticles<T, d>& Other) = delete;
+		TPBDGeometryCollectionParticles(const TPBDRigidClusteredParticles<T, d>& Other) = delete;
 
-		TPBDGeometryCollectionParticles(TPBDRigidParticles<T, d>&& Other)
-		    : TPBDRigidParticles<T, d>(MoveTemp(Other))
+		TPBDGeometryCollectionParticles(TPBDRigidClusteredParticles<T, d>&& Other)
+		    : TPBDRigidClusteredParticles<T, d>(MoveTemp(Other))
 		{
 			InitHelper();
 		}
diff --git a/Engine/Source/Runtime/Experimental/Chaos/Public/Chaos/PBDRigidClusteredParticles.h b/Engine/Source/Runtime/Experimental/Chaos/Public/Chaos/PBDRigidClusteredParticles.h
--- a/Engine/Source/Runtime/Experimental/Chaos/Public/Chaos/PBDRigidClusteredParticles.h
+++ b/Engine/Source/Runtime/Experimental/Chaos/Public/Chaos/PBDRigidClusteredParticles.h
@@ -9,48 +9,84 @@ namespace Chaos
 {
 	class FImplicitObjectUnionClustered;
 
-// ClusterId 
-//    Used within the clustering system to describe the clustering hierarchy. The ID will
-//  store the children IDs, and a Parent ID. When Id==NONE_INDEX the cluster is not
-//  controlled by another body. 
+/** 
+ * Used within the clustering system to describe the clustering hierarchy. The ClusterId
+ * stores the children IDs, and a Parent ID. When Id == \c nullptr the cluster is not
+ * controlled by another body. 
+ *
+ * TODO: Chaos - Add dimension template param?  Add floating point param?
+ */
 struct ClusterId
 {
-	ClusterId() : Id(INDEX_NONE), NumChildren(0) {}
-	ClusterId(int32 NewId, int NumChildrenIn)
+	ClusterId() 
+		: Id(nullptr)
+		, NumChildren(0) 
+	{}
+	ClusterId(TPBDRigidParticleHandle<float, 3>* NewId, int NumChildrenIn)
 		: Id(NewId)
-		, NumChildren(NumChildrenIn) {}
-	int32 Id;
+		, NumChildren(NumChildrenIn) 
+	{}
+	TPBDRigidParticleHandle<float, 3>* Id;
 	int32 NumChildren;
 };
 
-/** When multiple children are active and can share one collision proxy. Only valid if all original children are still in the cluster*/
+/** 
+ * When multiple children are active and can share one collision proxy. Only 
+ * valid if all original children are still in the cluster.
+ */
 template <typename T, int d>
 struct CHAOS_API TMultiChildProxyData
 {
+	TMultiChildProxyData()
+		: KeyChild(nullptr)
+	{}
 	TRigidTransform<T, d> RelativeToKeyChild;	//Use one child's transform to determine where to place the geometry. Needed for partial fracture where all children are still present and can therefore use proxy
-	uint32 KeyChild;
+	TPBDRigidParticleHandle<T, d>* KeyChild;
 };
 
-/** Used with TMultiChildProxyData. INDEX_NONE indicates no proxy data available */
+/** 
+ * Used with \c TMultiChildProxyData. \c nullptr indicates no proxy data available.
+ *
+ * TODO: Chaos - Add dimension template param?  Add floating point param?
+ */
 struct FMultiChildProxyId
 {
-	FMultiChildProxyId() : Id(INDEX_NONE) {}
-	int32 Id;
+	FMultiChildProxyId()
+		: Id(nullptr) 
+	{}
+	TPBDRigidParticleHandle<float, 3>* Id;
 };
 
+/**
+ * An entry in a clustered particle's \c ConnectivityEdges array, indicating a
+ * connection between that body and \c Sibling, with a strength breakable by 
+ * a \c Strain threshold.
+ *
+ * TODO: Chaos - Add dimension template param?
+ */
 template <typename T>
 struct TConnectivityEdge
 {
-	TConnectivityEdge() {}
-	TConnectivityEdge(uint32 InSibling, T InStrain)
+	TConnectivityEdge() 
+		: Sibling(nullptr)
+		, Strain(0.0)
+	{}
+
+	TConnectivityEdge(TPBDRigidParticleHandle<T, 3>* InSibling, const T InStrain)
 		: Sibling(InSibling)
-		, Strain(InStrain) {}
+		, Strain(InStrain) 
+	{}
 
 	TConnectivityEdge(const TConnectivityEdge& Other)
 		: Sibling(Other.Sibling)
-		, Strain(Other.Strain) {}
+		, Strain(Other.Strain) 
+	{}
+
+	/** Compares by \p OtherSibling only, for \c TArray::FindByKey(). */
+	bool operator==(const TPBDRigidParticleHandle<T, 3>* OtherSibling) const
+	{ return Sibling == OtherSibling; }
 
-	uint32 Sibling;
+	TPBDRigidParticleHandle<T, 3>* Sibling;
 	T Strain;
 };
 
@@ -104,6 +140,7 @@ class TPBDRigidClusteredParticles : public TPBDRigidParticles<T, d>
 
 	const auto& CollisionImpulses(int32 Idx) const { return MCollisionImpulses[Idx]; }
 	auto& CollisionImpulses(int32 Idx) { return MCollisionImpulses[Idx]; }
+	auto& CollisionImpulsesArray() { return MCollisionImpulses; }
 
 	const auto& Strains(int32 Idx) const { return MStrains[Idx]; }
 	auto& Strains(int32 Idx) { return MStrains[Idx]; }
@@ -172,4 +209,4 @@ class TPBDRigidClusteredParticles : public TPBDRigidParticles<T, d>
 	  TArrayCollectionArray<TArray<TConnectivityEdge<T>>> MConnectivityEdges;
 };
 
-}
+} // namespace Chaos
diff --git a/Engine/Source/Runtime/Experimental/Chaos/Public/Chaos/PBDRigidClustering.h b/Engine/Source/Runtime/Experimental/Chaos/Public/Chaos/PBDRigidClustering.h
--- a/Engine/Source/Runtime/Experimental/Chaos/Public/Chaos/PBDRigidClustering.h
+++ b/Engine/Source/Runtime/Experimental/Chaos/Public/Chaos/PBDRigidClustering.h
@@ -18,8 +18,8 @@ template <typename T, int d>
 class CHAOS_API TClusterBuffer
 {
 public:
-	using FClusterChildrenMap = TMap<uint32, TArray<uint32>>;
-	using FClusterTransformMap = TMap<uint32, TRigidTransform<float, 3>>;
+	using FClusterChildrenMap = TMap<TPBDRigidParticleHandle<T, d>*, TArray<TPBDRigidParticleHandle<T, d>*>>;
+	using FClusterTransformMap = TMap<TPBDRigidParticleHandle<T, d>*, TRigidTransform<float, 3>>;
 
 	virtual ~TClusterBuffer() = default;
 
@@ -35,9 +35,10 @@ template<class FPBDRigidEvolution, class FPBDCollisionConstraint, class T, int d
 class CHAOS_API TPBDRigidClustering
 {
 	typedef typename FPBDCollisionConstraint::FPointContactConstraint FPointContactConstraint;
-
 public:
-	typedef TMap<const TGeometryParticleHandle<T,d>*, TArray<uint32> > FClusterMap;
+	/** Parent to children */
+	typedef TMap<TPBDRigidParticleHandle<T, d>*, TArray<TPBDRigidParticleHandle<T, d>*> > FClusterMap;
+
 	using FCollisionConstraintHandle = TPBDCollisionConstraintHandle<T, d>;
 
 	TPBDRigidClustering(FPBDRigidEvolution& InEvolution, TPBDRigidClusteredParticles<T, d>& InParticles);
@@ -47,57 +48,70 @@ public:
 	// Initialization
 	//
 
-	/*
-	*  Create Cluster : Initialize clusters in the simulation.
-	*    ClusterGroupIndex : Index to join cluster into.
-	*    Children : Rigid body ID to include in the cluster.
-	*    ProxyGeometry : Collision default for the cluster, automatically generated otherwise.
-	*    ForceMassOrientation : Inertial alignment into mass space.
-	*/
-	int32 CreateClusterParticle(int32 ClusterGroupIndex, TArray<uint32>&& Children,
+	/**
+	 *  Initialize a cluster with the specified children.
+	 *
+	 *	\p ClusterGroupIndex - Index to join cluster into.
+	 *	\p Children - List of children that should belong to the cluster
+	 *	\p Parameters - ClusterParticleHandle must be valid, this is the parent cluster body.
+	 *	ProxyGeometry : Collision default for the cluster, automatically generated otherwise.
+	 *		ForceMassOrientation : Inertial alignment into mass space.
+	 *	
+	 */
+	Chaos::TPBDRigidClusteredParticleHandle<float, 3>* CreateClusterParticle(
+		const int32 ClusterGroupIndex, 
+		TArray<Chaos::TPBDRigidParticleHandle<T,d>*>&& Children, 
+		const FClusterCreationParameters<T>& Parameters/* = FClusterCreationParameters<T>()*/,
 		TSerializablePtr<FImplicitObject> ProxyGeometry = TSerializablePtr<FImplicitObject>(),
-		const TRigidTransform<T, d>* ForceMassOrientation = nullptr, const FClusterCreationParameters<T>& Parameters = FClusterCreationParameters<T>());
-
-	/*
-	*  CreateClusterParticleFromClusterChildren
-	*    Children : Rigid body ID to include in the cluster.
-	*/
-	int32 CreateClusterParticleFromClusterChildren(TArray<uint32>&& Children, TPBDRigidClusteredParticleHandle<T,d>* Parent,
-		const TRigidTransform<T, d>& ClusterWorldTM, const FClusterCreationParameters<T>& Parameters = FClusterCreationParameters<T>());
+		const TRigidTransform<T, d>* ForceMassOrientation = nullptr);
 
-	/*
-	*  UnionClusterGroups
-	*    Clusters that share a group index should be unioned into a single cluster prior to simulation.
-	*    The GroupIndex should be set on creation, and never touched by the client again.
-	*/
+	/**
+	 *  CreateClusterParticleFromClusterChildren
+	 *    Children : Rigid body ID to include in the cluster.
+	 */
+	Chaos::TPBDRigidClusteredParticleHandle<float, 3>* CreateClusterParticleFromClusterChildren(
+		TArray<TPBDRigidParticleHandle<T,d>*>&& Children, 
+		TPBDRigidClusteredParticleHandle<T,d>* Parent,
+		const TRigidTransform<T, d>& ClusterWorldTM, 
+		const FClusterCreationParameters<T>& Parameters/* = FClusterCreationParameters<T>()*/);
+
+#if 0 // Not called currently
+	/**
+	 *  UnionClusterGroups
+	 *    Clusters that share a group index should be unioned into a single cluster prior to simulation.
+	 *    The GroupIndex should be set on creation, and never touched by the client again.
+	 */
 	void UnionClusterGroups();
-
+#endif // 0
 	//
 	// Releasing
 	//
 
+#if 0 // Not called currently
 	/*
 	*  DeactivateClusterParticle
 	*    Release all the particles within the cluster particle
 	*/
 	TSet<uint32> DeactivateClusterParticle(TPBDRigidClusteredParticleHandle<T,d>* ClusteredParticle);
-
+#endif // 0
 
 	/*
 	*  ReleaseClusterParticles (BasedOnStrain)
-	*    Release clusters based on the passed in strains. Any cluster bodies that
-	*    have a MStrain value less than its entry in the StrainArray will be
-	*    released from the cluster.
+	*    Release clusters based on the passed in \p ExternalStrainArray, or the 
+	*    particle handle's current \c CollisionImpulses() value. Any cluster bodies 
+	*    that have a strain value less than this valid will be released from the 
+	*    cluster.
 	*/
-	TSet<uint32> ReleaseClusterParticles(TPBDRigidClusteredParticleHandle<T, d>* ClusteredParticle, const TArrayView<T>& StrainArray);
+	TSet<TPBDRigidParticleHandle<T, d>*> ReleaseClusterParticles(TPBDRigidClusteredParticleHandle<T, d>* ClusteredParticle, const TArrayView<T>* ExternalStrainArray=nullptr);
 
+#if 0 // Not called currently
 	/*
 	*  ReleaseClusterParticles
 	*    Release all rigid body IDs passed,
 	*/
-	TSet<uint32> ReleaseClusterParticles(const TArray<uint32>& ClusteredParticles);
-
-
+	TSet<uint32> ReleaseClusterParticles(
+		TArray<TPBDRigidParticleHandle<T, d>*> ChildrenParticles);
+#endif // 0
 	//
 	// Operational 
 	//
@@ -118,7 +132,8 @@ public:
 	*    encoded strain. The remainder strains are propagated back down to
 	*    the children clusters.
 	*/
-	TMap<TPBDRigidClusteredParticleHandle<T, d>*, TSet<uint32>> BreakingModel(TArrayView<T> ExternalStrain);
+	TMap<TPBDRigidClusteredParticleHandle<T, d>*, TSet<TPBDRigidParticleHandle<T, d>*>> BreakingModel(
+		TArrayView<T>* ExternalStrain=nullptr);
 
 	/**
 	*  PromoteStrains
@@ -127,14 +142,13 @@ public:
 	*    a ExternalStrain entry of 7. Will only decent the current
 	*    node passed, and ignores the disabled flag.
 	*/
-	T PromoteStrains(uint32 CurrentNode, TArrayView<T>& ExternalStrains);
+	T PromoteStrains(TPBDRigidParticleHandle<T, d>* CurrentNode);
 
 	/*
 	*  Process the kinematic state of the clusters. Because the leaf node geometry can
 	*  be changed by the solver, it is necessary to check all the sub clusters.
 	*/
-	void UpdateKinematicProperties(uint32 ClusterIndex);
-
+	void UpdateKinematicProperties(Chaos::TPBDRigidParticleHandle<float, 3>* Parent);
 
 	//
 	// Access
@@ -152,7 +166,7 @@ public:
 	*    Get the current childs active cluster. Returns INDEX_NONE if
 	*    not active or driven.
 	*/
-	int32                                               GetActiveClusterIndex(uint32 ChildIndex);
+	TPBDRigidParticleHandle<T, d>* GetActiveClusterIndex(TPBDRigidParticleHandle<T, d>* Child);
 
 	/*
 	*  GetClusterIdsArray
@@ -187,14 +201,13 @@ public:
 
 	/**
 	*  GetParentToChildren
-	*    The parent to children map stores the currently active cluster ids ( Particle Indices) as
+	*    The parent to children map stores the currently active cluster ids (Particle Indices) as
 	*    the keys of the map. The value of the map is a pointer to an array  constrained
 	*    rigid bodies.
 	*/
 	FClusterMap &										GetChildrenMap() { return MChildren; }
 	const FClusterMap &                                 GetChildrenMap() const { return MChildren; }
 
-
 	/*
 	*  GetClusterGroupIndexArray
 	*    The group index is used to automatically bind disjoint clusters. This attribute it set
@@ -231,9 +244,7 @@ public:
 	* GetConnectivityEdges
 	*    Provides a list of each rigid body's current siblings and associated strain within the cluster.
 	*/
-//#if CHAOS_PARTICLEHANDLE_TODO: is this really needed?
 	const TArrayCollectionArray<TArray<TConnectivityEdge<T>>>& GetConnectivityEdges() const { return MParticles.ConnectivityEdgesArray(); }
-//#endif
 
 	/**
 	* GenerateConnectionGraph
@@ -242,10 +253,15 @@ public:
 	*/
 	void SetClusterConnectionFactor(float ClusterConnectionFactorIn) { MClusterConnectionFactor = ClusterConnectionFactorIn; }
 	void SetClusterUnionConnectionType(typename FClusterCreationParameters<T>::EConnectionMethod ClusterConnectionType) { MClusterUnionConnectionType = ClusterConnectionType; }
-	void GenerateConnectionGraph(int32 NewIndex, const FClusterCreationParameters<T> & Parameters = FClusterCreationParameters<T>());
-	const TSet<int32>& GetTopLevelClusterParents() const { return TopLevelClusterParents; }
-	TSet<int32>& GetTopLevelClusterParents() { return TopLevelClusterParents; }
 
+	void GenerateConnectionGraph(
+		Chaos::TPBDRigidClusteredParticleHandle<float, 3>* Parent,
+		const FClusterCreationParameters<T> & Parameters = FClusterCreationParameters<T>());
+
+	const TSet<Chaos::TPBDRigidClusteredParticleHandle<float, 3>*>& GetTopLevelClusterParents() const { return TopLevelClusterParents; }
+	TSet<Chaos::TPBDRigidClusteredParticleHandle<float, 3>*>& GetTopLevelClusterParents() { return TopLevelClusterParents; }
+
+	/* Ryan - do we still need this?
 	void InitTopLevelClusterParents(const int32 StartIndex)
 	{
 		if (!StartIndex)
@@ -260,34 +276,52 @@ public:
 			}
 		}
 	}
-
-
-
+	*/
  protected:
+	void UpdateMassProperties(
+		Chaos::TPBDRigidClusteredParticleHandle<float, 3>* Parent, 
+		TSet<TPBDRigidParticleHandle<T, d>*>& Children, 
+		const TRigidTransform<T, d>* ForceMassOrientation);
+	void UpdateGeometry(
+		Chaos::TPBDRigidClusteredParticleHandle<float, 3>* Parent, 
+		const TSet<TPBDRigidParticleHandle<T, d>*>& Children, 
+		TSerializablePtr<FImplicitObject> ProxyGeometry, 
+		const FClusterCreationParameters<T>& Parameters);
 
-	void UpdateMassProperties(const TArray<uint32>& Children, const uint32 NewIndex, const TRigidTransform<T, d>* ForceMassOrientation);
-	void UpdateGeometry(const TArray<uint32>& Children, const uint32 NewIndex, TSerializablePtr<FImplicitObject> ProxyGeometry, const FClusterCreationParameters<T>& Parameters);
 	void ComputeStrainFromCollision(const FPBDCollisionConstraint& CollisionRule);
 	void ResetCollisionImpulseArray();
 	void DisableCluster(TPBDRigidClusteredParticleHandle<T, d>* ClusteredParticle);
-	void DisableParticleWithBreakEvent(uint32 ClusterIndex);
+	void DisableParticleWithBreakEvent(Chaos::TPBDRigidParticleHandle<float, 3>* Particle);
 
 	/*
 	* Connectivity
 	*/
-	void UpdateConnectivityGraphUsingPointImplicit(uint32 ClusterIndex, const FClusterCreationParameters<T>& Parameters = FClusterCreationParameters<T>());
-	void FixConnectivityGraphUsingDelaunayTriangulation(uint32 ClusterIndex, const FClusterCreationParameters<T>& Parameters = FClusterCreationParameters<T>());
-	void UpdateConnectivityGraphUsingDelaunayTriangulation(uint32 ClusterIndex, const FClusterCreationParameters<T>& Parameters = FClusterCreationParameters<T>());
-	void AddUniqueConnection(uint32 Index1, uint32 Index2, T Strain);
-	void ConnectNodes(uint32 Index1, uint32 Index2, T Strain);
-	void RemoveNodeConnections(uint32 ParticleIndex);
+	void UpdateConnectivityGraphUsingPointImplicit(
+		Chaos::TPBDRigidClusteredParticleHandle<float, 3>* Parent,
+		const FClusterCreationParameters<T>& Parameters = FClusterCreationParameters<T>());
+	void FixConnectivityGraphUsingDelaunayTriangulation(
+		Chaos::TPBDRigidClusteredParticleHandle<float, 3>* Parent,
+		const FClusterCreationParameters<T>& Parameters = FClusterCreationParameters<T>());
+	void UpdateConnectivityGraphUsingDelaunayTriangulation(
+		Chaos::TPBDRigidClusteredParticleHandle<float, 3>* Parent,
+		const FClusterCreationParameters<T>& Parameters = FClusterCreationParameters<T>());
+
+	void ConnectNodes(
+		TPBDRigidParticleHandle<T, d>* Child1,
+		TPBDRigidParticleHandle<T, d>* Child2);
+	void ConnectNodes(
+		TPBDRigidClusteredParticleHandle<T, d>* Child1,
+		TPBDRigidClusteredParticleHandle<T, d>* Child2);
+
+	void RemoveNodeConnections(TPBDRigidParticleHandle<T, d>* Child);
+	void RemoveNodeConnections(TPBDRigidClusteredParticleHandle<T, d>* Child);
 
 private:
 
 	FPBDRigidEvolution& MEvolution;
 	TPBDRigidClusteredParticles<T, d>& MParticles;
-	TSet<int32> TopLevelClusterParents;
-	TSet<int32> MActiveRemovalIndices;
+	TSet<Chaos::TPBDRigidClusteredParticleHandle<float, 3>*> TopLevelClusterParents;
+	TSet<Chaos::TPBDRigidParticleHandle<float, 3>*> MActiveRemovalIndices;
 
 	// Cluster data
 	mutable FRWLock ResourceLock;
@@ -308,13 +342,16 @@ private:
 };
 
 template <typename T, int d>
-void UpdateClusterMassProperties(TPBDRigidClusteredParticles<T, d>& Particles, const TArray<uint32>& Children,
-	const uint32 NewIndex, const TRigidTransform<T, d>* ForceMassOrientation = nullptr,
-	const TArrayCollectionArray<TUniquePtr<TMultiChildProxyData<T, d>>>* MMultiChildProxyData = nullptr,
-	const TArrayCollectionArray<FMultiChildProxyId>* MMultiChildProxyId = nullptr);
+void UpdateClusterMassProperties(
+	Chaos::TPBDRigidClusteredParticleHandle<float, 3>* Parent,
+	TSet<TPBDRigidParticleHandle<T, d>*>& Children,
+	const TRigidTransform<T, d>* ForceMassOrientation = nullptr);
 
 template<typename T, int d>
-TArray<TVector<T, d>> CleanCollisionParticles(const TArray<TVector<T, d>>& Vertices, TAABB<T, d> BBox, const float SnapDistance=0.01)
+TArray<TVector<T, d>> CleanCollisionParticles(
+	const TArray<TVector<T, d>>& Vertices, 
+	TAABB<T, d> BBox, 
+	const float SnapDistance=0.01)
 {
 	const int32 NumPoints = Vertices.Num();
 	if (NumPoints <= 1)
@@ -378,7 +415,9 @@ TArray<TVector<T, d>> CleanCollisionParticles(const TArray<TVector<T, d>>& Verti
 }
 
 template<typename T, int d>
-TArray<TVector<T, d>> CleanCollisionParticles(const TArray<TVector<T, d>>& Vertices, const float SnapDistance=0.01)
+TArray<TVector<T, d>> CleanCollisionParticles(
+	const TArray<TVector<T, d>>& Vertices, 
+	const float SnapDistance=0.01)
 {
 	if (!Vertices.Num())
 	{
@@ -394,7 +433,10 @@ TArray<TVector<T, d>> CleanCollisionParticles(const TArray<TVector<T, d>>& Verti
 
 template <typename T, int d>
 TArray<TVector<T,d>> 
-CleanCollisionParticles(TTriangleMesh<T> &TriMesh, const TArrayView<const TVector<T,d>>& Vertices, float Fraction)
+CleanCollisionParticles(
+	TTriangleMesh<T> &TriMesh, 
+	const TArrayView<const TVector<T,d>>& Vertices, 
+	const float Fraction)
 {
 	TArray<TVector<T, d>> CollisionVertices;
 	if (Fraction <= 0.0)
@@ -407,33 +449,39 @@ CleanCollisionParticles(TTriangleMesh<T> &TriMesh, const TArrayView<const TVecto
 
 	// Particles are ordered from most important to least, with coincident 
 	// vertices at the very end.
-	const int32 numGoodPoints = Ordering.Num() - CoincidentVertices.Num();
+	const int32 NumGoodPoints = Ordering.Num() - CoincidentVertices.Num();
 
 #if DO_GUARD_SLOW
-	for (int i = numGoodPoints; i < Ordering.Num(); ++i)
+	for (int i = NumGoodPoints; i < Ordering.Num(); ++i)
 	{
 		ensure(CoincidentVertices.Contains(Ordering[i]));	//make sure all coincident vertices are at the back
 	}
 #endif
 
-	CollisionVertices.AddUninitialized(std::min(numGoodPoints, static_cast<int32>(ceil(numGoodPoints * Fraction))));
-	for(int i=0; i < CollisionVertices.Num(); i++)
+	CollisionVertices.AddUninitialized(std::min(NumGoodPoints, static_cast<int32>(ceil(NumGoodPoints * Fraction))));
+	for (int i = 0; i < CollisionVertices.Num(); i++)
+	{
 		CollisionVertices[i] = Vertices[Ordering[i]];
-
+	}
 	return CollisionVertices;
 }
 
-
 template <typename T, int d>
 void
-CleanCollisionParticles( /*const*/ TTriangleMesh<T> &TriMesh, const TArrayView<const TVector<T, d>>& Vertices, float Fraction, TSet<int32> & ResultingIndices)
+CleanCollisionParticles(
+	TTriangleMesh<T> &TriMesh, 
+	const TArrayView<const TVector<T, d>>& Vertices, 
+	const float Fraction, 
+	TSet<int32>& ResultingIndices)
 {
+	ResultingIndices.Reset();
 	if (Fraction <= 0.0)
 		return;
 
 	TArray<int32> CoincidentVertices;
 	const TArray<int32> Ordering = TriMesh.GetVertexImportanceOrdering(Vertices, &CoincidentVertices, true);
-	const int32 NumGoodPoints = Ordering.Num() - CoincidentVertices.Num();
+	int32 NumGoodPoints = Ordering.Num() - CoincidentVertices.Num();
+	NumGoodPoints = std::min(NumGoodPoints, static_cast<int32>(ceil(NumGoodPoints * Fraction)))
 
 	ResultingIndices.Reserve(NumGoodPoints);
 	for (int32 i = 0; i < NumGoodPoints; i++)
@@ -442,6 +490,4 @@ CleanCollisionParticles( /*const*/ TTriangleMesh<T> &TriMesh, const TArrayView<c
 	}
 }
 
-
-
 } // namespace Chaos
diff --git a/Engine/Source/Runtime/Experimental/Chaos/Public/Chaos/PBDRigidsEvolution.h b/Engine/Source/Runtime/Experimental/Chaos/Public/Chaos/PBDRigidsEvolution.h
--- a/Engine/Source/Runtime/Experimental/Chaos/Public/Chaos/PBDRigidsEvolution.h
+++ b/Engine/Source/Runtime/Experimental/Chaos/Public/Chaos/PBDRigidsEvolution.h
@@ -269,6 +269,16 @@ class FPBDRigidsEvolutionBase
 		return NewParticles;
 	}
 
+	CHAOS_API TArray<TPBDGeometryCollectionParticleHandle<FReal, 3>*> CreateGeometryCollectionParticles(int32 NumParticles,const FUniqueIdx* ExistingIndices = nullptr,  const TPBDRigidParticleParameters<FReal, 3>& Params = TPBDRigidParticleParameters<FReal, 3>())
+	{
+		auto NewParticles = Particles.CreateGeometryCollectionParticles(NumParticles, ExistingIndices, Params);
+		for (auto& Particle : NewParticles)
+		{
+			DirtyParticle(*Particle);
+		}
+		return NewParticles;
+	}
+
 	CHAOS_API void AddForceFunction(FForceRule ForceFunction) { ForceRules.Add(ForceFunction); }
 	CHAOS_API void AddImpulseFunction(FForceRule ImpulseFunction) { ImpulseRules.Add(ImpulseFunction); }
 	CHAOS_API void SetParticleUpdateVelocityFunction(FUpdateVelocityRule ParticleUpdate) { ParticleUpdateVelocity = ParticleUpdate; }
@@ -313,6 +323,17 @@ class FPBDRigidsEvolutionBase
 		RemoveConstraints(TSet<TGeometryParticleHandle<FReal, 3>*>({ Particle }));
 	}
 
+	CHAOS_API void DisableParticles(TSet<TGeometryParticleHandle<FReal, 3>*> &ParticlesIn)
+	{
+		for (TGeometryParticleHandle<FReal, 3>* Particle : ParticlesIn)
+		{
+			RemoveParticleFromAccelerationStructure(*Particle);
+			Particles.DisableParticle(Particle);
+			ConstraintGraph.DisableParticle(Particle);
+		}
+		RemoveConstraints(ParticlesIn);
+	}
+
 	template <bool bPersistent>
 	FORCEINLINE_DEBUGGABLE void DirtyParticle(TGeometryParticleHandleImp<FReal, 3, bPersistent>& Particle)
 	{
@@ -340,6 +361,19 @@ class FPBDRigidsEvolutionBase
 		Particles.DestroyParticle(Particle);
 	}
 
+	/**
+	 * Preallocate buffers for creating \p Num particles.
+	 */
+	CHAOS_API void ReserveParticles(const int32 Num)
+	{
+		if (const int32 NumNew = ConstraintGraph.ReserveParticles(Num))
+		{
+			InternalAccelerationQueue.PendingData.Reserve(InternalAccelerationQueue.Num() + NumNew);
+			AsyncAccelerationQueue.PendingData.Reserve(AsyncAccelerationQueue.Num() + NumNew);
+			ExternalAccelerationQueue.PendingData.Reserve(ExternalAccelerationQueue.Num() + NumNew);
+		}
+	}
+
 	CHAOS_API void CreateParticle(TGeometryParticleHandle<FReal, 3>* ParticleAdded)
 	{
 		ConstraintGraph.AddParticle(ParticleAdded);
diff --git a/Engine/Source/Runtime/Experimental/Chaos/Public/Chaos/PBDRigidsSOAs.h b/Engine/Source/Runtime/Experimental/Chaos/Public/Chaos/PBDRigidsSOAs.h
--- a/Engine/Source/Runtime/Experimental/Chaos/Public/Chaos/PBDRigidsSOAs.h
+++ b/Engine/Source/Runtime/Experimental/Chaos/Public/Chaos/PBDRigidsSOAs.h
@@ -227,6 +227,10 @@ public:
 			{
 				RemoveFromMapAndArray(PBDRigidClustered, NonDisabledClusteredToIndex, NonDisabledClusteredArray);
 				RemoveFromMapAndArray(PBDRigidClustered, ActiveClusteredToIndex, ActiveClusteredArray);
+				if (Particle->GetParticleType() == Chaos::EParticleType::GeometryCollection)
+				{
+					UpdateGeometryCollectionViews();
+				}
 			}
 			else
 			{
@@ -260,6 +264,10 @@ public:
 				{
 					InsertToMapAndArray(PBDRigidClustered, ActiveClusteredToIndex, ActiveClusteredArray);
 				}
+				if (Particle->GetParticleType() == Chaos::EParticleType::GeometryCollection)
+				{
+					UpdateGeometryCollectionViews();
+				}
 			}
 			else
 			{
@@ -454,26 +462,38 @@ public:
 	 * their object state (static, kinematic, dynamic, sleeping) and their disabled 
 	 * state.
 	 */
-	void UpdateGeometryCollectionViews()
+	void UpdateGeometryCollectionViews(const bool ForceUpdateViews=false)
 	{
-		int32 AIdx = 0, SIdx = 0, KIdx = 0, DIdx = 0;
+		if (!GeometryCollectionParticles)
+		{
+			if (ForceUpdateViews)
+				UpdateViews();
+			return;
+		}
+
+		int32 AIdx = 0, SIdx = 0, KIdx = 0, DIdx = 0; // Active, Static, Kinematic, Dynamic Index
 
-		for (TPBDGeometryCollectionParticleHandle<T, d>* Handle : ActiveGeometryCollectionArray)
+		for(int32 PIdx = 0; PIdx < (int32)GeometryCollectionParticles->Size(); PIdx++)
 		{
+			TPBDGeometryCollectionParticleHandle<T, d>* Handle = GeometryCollectionParticles->Handle(PIdx);
+			if (!Handle)
+				continue;
+
 			// If the particle is disabled we treat it as static, but for no reason 
 			// other than immediate convenience.
 			const Chaos::EObjectStateType State = Handle->Disabled() ? Chaos::EObjectStateType::Static : Handle->ObjectState();
 
+			// Count the number of particles in each state.
 			switch (State)
 			{
 			case Chaos::EObjectStateType::Static:
 				SIdx++;
-				AIdx += (int32)(!Handle->Disabled());
+				AIdx += (int32)(!Handle->Disabled()); // If Disabled == true, then it's not active
 				break;
 
 			case Chaos::EObjectStateType::Kinematic:
 				KIdx++;
-				AIdx += (int32)(!Handle->Disabled());
+				AIdx++;
 				break;
 
 			case Chaos::EObjectStateType::Sleeping: // Sleeping is a modified dynamic state
@@ -482,7 +502,7 @@ public:
 
 			case Chaos::EObjectStateType::Dynamic:
 				DIdx++;
-				AIdx += (int32)(!Handle->Disabled());
+				AIdx++;
 				break;
 
 			default:
@@ -490,12 +510,12 @@ public:
 			};
 		}
 
+		// Compare with the previous array sizes, and resize if needed.
 		bool Changed = 
 			ActiveGeometryCollectionArray.Num() != AIdx ||
 			StaticGeometryCollectionArray.Num() != SIdx || 
 			KinematicGeometryCollectionArray.Num() != KIdx || 
 			DynamicGeometryCollectionArray.Num() != DIdx;
-
 		if (Changed)
 		{
 			ActiveGeometryCollectionArray.SetNumUninitialized(AIdx);
@@ -504,10 +524,14 @@ public:
 			DynamicGeometryCollectionArray.SetNumUninitialized(DIdx);
 		}
 
+		// (Re)populate the arrays, making note if any prior entires differ from the current.
 		AIdx = SIdx = KIdx = DIdx = 0;
-
-		for (TPBDGeometryCollectionParticleHandle<T, d>* Handle : ActiveGeometryCollectionArray)
+		for(int32 PIdx = 0; PIdx < (int32)GeometryCollectionParticles->Size(); PIdx++)
 		{
+			TPBDGeometryCollectionParticleHandle<T, d>* Handle = GeometryCollectionParticles->Handle(PIdx);
+			if (!Handle)
+				continue;
+
 			const Chaos::EObjectStateType State = 
 				Handle->Disabled() ? Chaos::EObjectStateType::Static : Handle->ObjectState();
 			switch (State)
@@ -525,11 +549,9 @@ public:
 			case Chaos::EObjectStateType::Kinematic:
 				Changed |= KinematicGeometryCollectionArray[KIdx] != Handle;
 				KinematicGeometryCollectionArray[KIdx++] = Handle;
-				if (!Handle->Disabled())
-				{
-					Changed |= ActiveGeometryCollectionArray[AIdx] != Handle;
-					ActiveGeometryCollectionArray[AIdx++] = Handle;
-				}
+
+				Changed |= ActiveGeometryCollectionArray[AIdx] != Handle;
+				ActiveGeometryCollectionArray[AIdx++] = Handle;
 				break;
 
 			case Chaos::EObjectStateType::Sleeping: // Sleeping is a modified dynamic state
@@ -540,11 +562,9 @@ public:
 			case Chaos::EObjectStateType::Dynamic:
 				Changed |= DynamicGeometryCollectionArray[DIdx] != Handle;
 				DynamicGeometryCollectionArray[DIdx++] = Handle;
-				if (!Handle->Disabled())
-				{
-					Changed |= ActiveGeometryCollectionArray[AIdx] != Handle;
-					ActiveGeometryCollectionArray[AIdx++] = Handle;
-				}
+
+				Changed |= ActiveGeometryCollectionArray[AIdx] != Handle;
+				ActiveGeometryCollectionArray[AIdx++] = Handle;
 				break;
 
 			default:
@@ -552,7 +572,7 @@ public:
 			};
 		}
 
-		if(Changed)
+		if(Changed || ForceUpdateViews)
 		{
 			UpdateViews();
 		}
@@ -607,13 +627,29 @@ private:
 	template <typename TParticle1, typename TParticle2>
 	void InsertToMapAndArray(const TArray<TParticle1*>& ParticlesToInsert, TMap<TParticle2*, int32>& ParticleToIndex, TArray<TParticle2*>& ParticleArray)
 	{
+		TArray<bool> Contains;
+		Contains.AddZeroed(ParticlesToInsert.Num());
+
 		// TODO: Compile time check ensuring TParticle2 is derived from TParticle1?
 		int32 NextIdx = ParticleArray.Num();
-		for (auto Particle : ParticlesToInsert)
+		for(int32 Idx = 0; Idx < ParticlesToInsert.Num(); ++Idx)
+		{
+			auto Particle = ParticlesToInsert[Idx];
+			Contains[Idx] = ParticleToIndex.Contains(Particle);
+			if (!Contains[Idx])
+			{
+				ParticleToIndex.Add(Particle, NextIdx++);
+			}
+		}
+		ParticleArray.Reserve(ParticleArray.Num() + NextIdx - ParticleArray.Num());
+		for (int32 Idx = 0; Idx < ParticlesToInsert.Num(); ++Idx)
 		{
-			ParticleToIndex.Add(Particle, NextIdx++);
+			if (!Contains[Idx])
+			{
+				auto Particle = ParticlesToInsert[Idx];
+				ParticleArray.Add(Particle);
+			}
 		}
-		ParticleArray.Append(ParticlesToInsert);
 	}
 
 	template <typename TParticle>
@@ -686,13 +722,24 @@ private:
 			TArray<TSOAView<TPBDRigidParticles<T, d>>> TmpArray = 
 			{ 
 				{&ActiveParticlesArray},
+				{&NonDisabledClusteredArray}, // Ryan!
 				{&ActiveGeometryCollectionArray}
 			};
 			ActiveParticlesView = MakeParticleView(MoveTemp(TmpArray));
 		}
 		{
-			TArray<TSOAView<TGeometryParticles<T, d>>> TmpArray = { StaticParticles.Get(), StaticDisabledParticles.Get(), KinematicParticles.Get(), KinematicDisabledParticles.Get(),
-				DynamicParticles.Get(), DynamicDisabledParticles.Get(), DynamicKinematicParticles.Get(), ClusteredParticles.Get(), GeometryCollectionParticles.Get() };
+			TArray<TSOAView<TGeometryParticles<T, d>>> TmpArray = 
+			{ 
+				StaticParticles.Get(), 
+				StaticDisabledParticles.Get(), 
+				KinematicParticles.Get(), 
+				KinematicDisabledParticles.Get(),
+				DynamicParticles.Get(), 
+				DynamicDisabledParticles.Get(), 
+				DynamicKinematicParticles.Get(), 
+				ClusteredParticles.Get(), 
+				GeometryCollectionParticles.Get() 
+			};
 			AllParticlesView = MakeParticleView(MoveTemp(TmpArray));
 		}
 		{
diff --git a/Engine/Source/Runtime/Experimental/Chaos/Public/Chaos/ParticleHandle.h b/Engine/Source/Runtime/Experimental/Chaos/Public/Chaos/ParticleHandle.h
--- a/Engine/Source/Runtime/Experimental/Chaos/Public/Chaos/ParticleHandle.h
+++ b/Engine/Source/Runtime/Experimental/Chaos/Public/Chaos/ParticleHandle.h
@@ -421,7 +421,7 @@ public:
 	TSerializablePtr<FImplicitObject> Geometry() const { return GeometryParticles->Geometry(ParticleIdx); }
 	void SetGeometry(TSerializablePtr<FImplicitObject> InGeometry) { GeometryParticles->SetGeometry(ParticleIdx, InGeometry); }
 
-	TSerializablePtr<FImplicitObject> SharedGeometry() const { return GeometryParticles->SharedGeometry(ParticleIdx); }
+	TSharedPtr<FImplicitObject, ESPMode::ThreadSafe> SharedGeometry() const { return GeometryParticles->SharedGeometry(ParticleIdx); }
 	void SetSharedGeometry(TSharedPtr<FImplicitObject, ESPMode::ThreadSafe> InGeometry) { GeometryParticles->SetSharedGeometry(ParticleIdx, InGeometry); }
 
 	const TUniquePtr<FImplicitObject>& DynamicGeometry() const { return GeometryParticles->DynamicGeometry(ParticleIdx); }
@@ -735,7 +735,10 @@ public:
 	void SetToBeRemovedOnFracture(const bool bToBeRemovedOnFracture) { PBDRigidParticles->ToBeRemovedOnFracture(ParticleIdx) = bToBeRemovedOnFracture; }
 
 	EObjectStateType ObjectState() const { return PBDRigidParticles->ObjectState(ParticleIdx); }
+	void SetObjectState(EObjectStateType InState) { PBDRigidParticles->SetObjectState(ParticleIdx, InState); }
 	void SetObjectStateLowLevel(EObjectStateType InState) { PBDRigidParticles->SetObjectState(ParticleIdx, InState); }
+	
+	bool Sleeping() const { return PBDRigidParticles->Sleeping(ParticleIdx); }
 	void SetSleeping(bool bSleeping) { PBDRigidParticles->SetSleeping(ParticleIdx, bSleeping); }
 
 	//Really only useful when using a transient handle
@@ -782,32 +785,52 @@ public:
 		return Serializable;
 	}
 
-	const auto& ClusterIds() const { return PBDRigidClusteredParticles->ClusterIds(ParticleIdx); }
-	auto& ClusterIds() { return PBDRigidClusteredParticles->ClusterIds(ParticleIdx); }
+	void SetClusterId(const ClusterId& Id) { PBDRigidClusteredParticles->ClusterIds(ParticleIdx) = Id; }
+	const ClusterId& ClusterIds() const { return PBDRigidClusteredParticles->ClusterIds(ParticleIdx); }
+	ClusterId& ClusterIds() { return PBDRigidClusteredParticles->ClusterIds(ParticleIdx); }
+
+	const TRigidTransform<T,d>& ChildToParent() const { return PBDRigidClusteredParticles->ChildToParent(ParticleIdx); }
+	TRigidTransform<T,d>& ChildToParent() { return PBDRigidClusteredParticles->ChildToParent(ParticleIdx); }
+	void SetChildToParent(const TRigidTransform<T, d>& Xf) { PBDRigidClusteredParticles->ChildToParent(ParticleIdx) = Xf; }
 
-	const auto& ChildToParent() const { return PBDRigidClusteredParticles->ChildToParent(ParticleIdx); }
-	auto& ChildToParent() { return PBDRigidClusteredParticles->ChildToParent(ParticleIdx); }
+	const int32& ClusterGroupIndex() const { return PBDRigidClusteredParticles->ClusterGroupIndex(ParticleIdx); }
+	int32& ClusterGroupIndex() { return PBDRigidClusteredParticles->ClusterGroupIndex(ParticleIdx); }
+	void SetClusterGroupIndex(const int32 Idx) { PBDRigidClusteredParticles->ClusterGroupIndex(ParticleIdx) = Idx; }
 
-	const auto& ClusterGroupIndex() const { return PBDRigidClusteredParticles->ClusterGroupIndex(ParticleIdx); }
-	auto& ClusterGroupIndex() { return PBDRigidClusteredParticles->ClusterGroupIndex(ParticleIdx); }
+	const bool& InternalCluster() const { return PBDRigidClusteredParticles->InternalCluster(ParticleIdx); }
+	bool& InternalCluster() { return PBDRigidClusteredParticles->InternalCluster(ParticleIdx); }
+	void SetInternalCluster(const bool Value) { PBDRigidClusteredParticles->InternalCluster(ParticleIdx) = Value; }
 
-	const auto& InternalCluster() const { return PBDRigidClusteredParticles->InternalCluster(ParticleIdx); }
-	auto& InternalCluster() { return PBDRigidClusteredParticles->InternalCluster(ParticleIdx); }
+	const TUniquePtr<FImplicitObjectUnionClustered>& ChildrenSpatial() const { return PBDRigidClusteredParticles->ChildrenSpatial(ParticleIdx); }
+	TUniquePtr<FImplicitObjectUnionClustered>& ChildrenSpatial() { return PBDRigidClusteredParticles->ChildrenSpatial(ParticleIdx); }
+	void SetChildrenSpatial(TUniquePtr<FImplicitObjectUnion>& Obj) { PBDRigidClusteredParticles->ChildrenSpatial(ParticleIdx) = Obj; }
 
-	const auto& ChildrenSpatial() const { return PBDRigidClusteredParticles->ChildrenSpatial(ParticleIdx); }
-	auto& ChildrenSpatial() { return PBDRigidClusteredParticles->ChildrenSpatial(ParticleIdx); }
+	const FMultiChildProxyId& MultiChildProxyId() const { return PBDRigidClusteredParticles->MultiChildProxyId(ParticleIdx); }
+	FMultiChildProxyId& MultiChildProxyId() { return PBDRigidClusteredParticles->MultiChildProxyId(ParticleIdx); }
+	void SetMultiChildProxyId(const FMultiChildProxyId& Id) { PBDRigidClusteredParticles->MultiChildProxyId(ParticleIdx) = Id; }
 
-	//const auto& MultiChildProxyId() const { return PBDRigidClusteredParticles->MultiChildProxyId(ParticleIdx); }
-	//auto& MultiChildProxyId() { return PBDRigidClusteredParticles->MultiChildProxyId(ParticleIdx); }
+	const TUniquePtr<TMultiChildProxyData<T, d>>& MultiChildProxyData() const { return PBDRigidClusteredParticles->MultiChildProxyData(ParticleIdx); }
+	TUniquePtr<TMultiChildProxyData<T, d>>& MultiChildProxyData() { return PBDRigidClusteredParticles->MultiChildProxyData(ParticleIdx); }
+	void SetMultiChildProxyData(TUniquePtr<TMultiChildProxyData<T, d>>&& ProxyData) { PBDRigidClusteredParticles->MultiChildProxyData(ParticleIdx) = MoveTemp(ProxyData); }
 
-	//const auto& MultiChildProxyData() const { return PBDRigidClusteredParticles->MultiChildProxyData(ParticleIdx); }
-	//auto& MultiChildProxyData() { return PBDRigidClusteredParticles->MultiChildProxyData(ParticleIdx); }
+	const T& CollisionImpulse() const { return PBDRigidClusteredParticles->CollisionImpulses(ParticleIdx); }
+	T& CollisionImpulse() { return PBDRigidClusteredParticles->CollisionImpulses(ParticleIdx); }
+	void SetCollisionImpulse(const T Value) { PBDRigidClusteredParticles->CollisionImpulses(ParticleIdx) = Value; }
+	const T& CollisionImpulses() const { return PBDRigidClusteredParticles->CollisionImpulses(ParticleIdx); }
+	T& CollisionImpulses() { return PBDRigidClusteredParticles->CollisionImpulses(ParticleIdx); }
+	void SetCollisionImpulses(const T Value) { PBDRigidClusteredParticles->CollisionImpulses(ParticleIdx) = Value; }
 
-	const auto& CollisionImpulses() const { return PBDRigidClusteredParticles->CollisionImpulses(ParticleIdx); }
-	auto& CollisionImpulses() { return PBDRigidClusteredParticles->CollisionImpulses(ParticleIdx); }
+	const T& Strain() const { return PBDRigidClusteredParticles->Strains(ParticleIdx); }
+	T& Strain() { return PBDRigidClusteredParticles->Strains(ParticleIdx); }
+	void SetStrain(const T Value) { PBDRigidClusteredParticles->Strains(ParticleIdx) = Value; }
+	const T& Strains() const { return PBDRigidClusteredParticles->Strains(ParticleIdx); }
+	T& Strains() { return PBDRigidClusteredParticles->Strains(ParticleIdx); }
+	void SetStrains(const T Value) { PBDRigidClusteredParticles->Strains(ParticleIdx) = Value; }
 
-	const auto& Strains() const { return PBDRigidClusteredParticles->Strains(ParticleIdx); }
-	auto& Strains() { return PBDRigidClusteredParticles->Strains(ParticleIdx); }
+	const TArray<TConnectivityEdge<T>>& ConnectivityEdges() const { return PBDRigidClusteredParticles->ConnectivityEdges(ParticleIdx); }
+	TArray<TConnectivityEdge<T>>& ConnectivityEdges() { return PBDRigidClusteredParticles->ConnectivityEdges(ParticleIdx); }
+	void SetConnectivityEdges(const TArray<TConnectivityEdge<T>>& Edges) { PBDRigidClusteredParticles->ConnectivityEdges(ParticleIdx) = Edges; }
+	void SetConnectivityEdges(TArray<TConnectivityEdge<T>>&& Edges) { PBDRigidClusteredParticles->ConnectivityEdges(ParticleIdx) = MoveTemp(Edges); }
 
 	const TPBDRigidClusteredParticleHandleImp<T, d, true>* Handle() const { return PBDRigidClusteredParticles->Handle(ParticleIdx); }
 	TPBDRigidClusteredParticleHandleImp<T, d, true>* Handle() { return PBDRigidClusteredParticles->Handle(ParticleIdx); }
@@ -818,7 +841,7 @@ public:
 };
 
 template <typename T, int d, bool bPersistent = true>
-class TPBDGeometryCollectionParticleHandleImp : public TPBDRigidParticleHandleImp<T, d, bPersistent>
+class TPBDGeometryCollectionParticleHandleImp : public TPBDRigidClusteredParticleHandleImp<T, d, bPersistent>
 {
 public:
 	using TGeometryParticleHandleImp<T, d, bPersistent>::ParticleIdx;
@@ -832,7 +855,7 @@ protected:
 
 	//needed for serialization
 	TPBDGeometryCollectionParticleHandleImp()
-		: TPBDRigidParticleHandleImp<T, d, bPersistent>()
+		: TPBDRigidClusteredParticleHandleImp<T, d, bPersistent>()
 	{}
 
 	TPBDGeometryCollectionParticleHandleImp<T, d, bPersistent>(
@@ -840,8 +863,8 @@ protected:
 		int32 InIdx, 
 		int32 InGlobalIdx, 
 		const TPBDRigidParticleParameters<T, d>& Params = TPBDRigidParticleParameters<T, d>())
-		: TPBDRigidParticleHandleImp<T, d, bPersistent>(
-			TSerializablePtr<TPBDRigidParticles<T, d>>(Particles), InIdx, InGlobalIdx, Params)
+		: TPBDRigidClusteredParticleHandleImp<T, d, bPersistent>(
+			TSerializablePtr<TPBDRigidClusteredParticles<T, d>>(Particles), InIdx, InGlobalIdx, Params)
 	{}
 public:
 
@@ -1200,6 +1223,7 @@ public:
 	TGenericParticleHandle(TGeometryParticleHandle<T, d>* InHandle) : Imp(InHandle) {}
 
 	TGenericParticleHandleHandleImp<T, d>* operator->() const { return const_cast<TGenericParticleHandleHandleImp<T, d>*>(&Imp); }
+	TGenericParticleHandleHandleImp<T, d>* Get() const { return const_cast<TGenericParticleHandleHandleImp<T, d>*>(&Imp); }
 
 private:
 	TGenericParticleHandleHandleImp<T, d> Imp;
@@ -1579,12 +1603,19 @@ private:
 	// Pointer to some arbitrary data associated with the particle, but not used by Chaos. External systems may use this for whatever.
 	void* MUserData;
 
+public:
+	// Ryan: FGeometryCollectionPhysicsProxy needs access to GeometrySharedLowLevel(), 
+	// as it needs access for the same reason as ParticleData.  For some reason
+	// the friend declaration isn't working.  Exposing this function until this 
+	// can be straightened out.
+	//friend class FGeometryCollectionPhysicsProxy;
 	// This is only for use by ParticleData. This should be called only in one place,
 	// when the geometry is being copied from GT to PT.
 	TSharedPtr<FImplicitObject, ESPMode::ThreadSafe> GeometrySharedLowLevel() const
 	{
 		return MGeometry;
 	}
+private:
 
 #if CHAOS_CHECKED
 	FName MDebugName;
diff --git a/Engine/Source/Runtime/Experimental/Chaos/Public/Chaos/ParticleHandleFwd.h b/Engine/Source/Runtime/Experimental/Chaos/Public/Chaos/ParticleHandleFwd.h
--- a/Engine/Source/Runtime/Experimental/Chaos/Public/Chaos/ParticleHandleFwd.h
+++ b/Engine/Source/Runtime/Experimental/Chaos/Public/Chaos/ParticleHandleFwd.h
@@ -57,6 +57,7 @@ namespace Chaos
 
 	template <typename T, int d>
 	using TPBDRigidClusteredParticleHandle = TPBDRigidClusteredParticleHandleImp<T, d, true>;
+	using TPBDRigidClusteredParticleHandleFloat3 = TPBDRigidClusteredParticleHandle<float, 3>;
 
 	template <typename T, int d>
 	using TTransientPBDRigidClusteredParticleHandle = TPBDRigidClusteredParticleHandleImp<T, d, false>;
diff --git a/Engine/Source/Runtime/Experimental/ChaosCore/Public/Chaos/ArrayCollectionArray.h b/Engine/Source/Runtime/Experimental/ChaosCore/Public/Chaos/ArrayCollectionArray.h
--- a/Engine/Source/Runtime/Experimental/ChaosCore/Public/Chaos/ArrayCollectionArray.h
+++ b/Engine/Source/Runtime/Experimental/ChaosCore/Public/Chaos/ArrayCollectionArray.h
@@ -37,6 +37,14 @@ class TArrayCollectionArray : public TArrayCollectionArrayBase, public TArray<T>
 
 	virtual ~TArrayCollectionArray() {}
 
+	void Fill(const T& Value)
+	{
+		for (int32 Idx = 0; Idx < Num(); ++Idx)
+		{
+			TArray<T>::operator[](Idx) = Value;
+		}
+	}
+
 	TArrayCollectionArray<T> Clone()
 	{
 		TArrayCollectionArray<T> NewArray;
diff --git a/Engine/Source/Runtime/Experimental/ChaosSolvers/Private/PBDRigidActiveParticlesBuffer.cpp b/Engine/Source/Runtime/Experimental/ChaosSolvers/Private/PBDRigidActiveParticlesBuffer.cpp
--- a/Engine/Source/Runtime/Experimental/ChaosSolvers/Private/PBDRigidActiveParticlesBuffer.cpp
+++ b/Engine/Source/Runtime/Experimental/ChaosSolvers/Private/PBDRigidActiveParticlesBuffer.cpp
@@ -22,14 +22,24 @@ namespace Chaos
 	void FPBDRigidActiveParticlesBuffer::BufferPhysicsResults(FPBDRigidsSolver* Solver)
 	{
 		auto& ActiveGameThreadParticles = SolverDataOut->AccessProducerBuffer()->ActiveGameThreadParticles;
+		auto& PhysicsParticleProxies = SolverDataOut->AccessProducerBuffer()->PhysicsParticleProxies;
 
 		ActiveGameThreadParticles.Empty();
-		for (auto& ActiveParticle : Solver->GetParticles().GetActiveParticlesView())
+		TParticleView<TPBDRigidParticles<float, 3>>& ActiveParticlesView = 
+			Solver->GetParticles().GetActiveParticlesView();
+		for (auto& ActiveParticle : ActiveParticlesView)
 		{
 			if (ActiveParticle.Handle())
 			{
-				TGeometryParticle<float, 3>* Handle = ActiveParticle.Handle()->GTGeometryParticle();
-				ActiveGameThreadParticles.Add(Handle);
+				// Clustered particles don't have a game thread particle instance.
+				if (TGeometryParticle<float, 3>* Handle = ActiveParticle.Handle()->GTGeometryParticle())
+				{
+					ActiveGameThreadParticles.Add(Handle);
+				}
+				else if(IPhysicsProxyBase** Proxy = Solver->MParticleToProxy.Find(ActiveParticle.Handle()))
+				{
+					PhysicsParticleProxies.Add(*Proxy);
+				}
 			}
 		}
 	}
diff --git a/Engine/Source/Runtime/Experimental/ChaosSolvers/Private/PBDRigidsSolver.cpp b/Engine/Source/Runtime/Experimental/ChaosSolvers/Private/PBDRigidsSolver.cpp
--- a/Engine/Source/Runtime/Experimental/ChaosSolvers/Private/PBDRigidsSolver.cpp
+++ b/Engine/Source/Runtime/Experimental/ChaosSolvers/Private/PBDRigidsSolver.cpp
@@ -79,6 +79,9 @@ namespace Chaos
 					auto& GeomCollectionParticles = MSolver->GetEvolution()->GetParticles().GetGeometryCollectionParticles();
 					Obj->FieldParameterUpdateCallback(MSolver, GeomCollectionParticles, Strains, 
 						PositionTarget, PositionTargetedParticles, /*AnimatedPositions,*/ MSolver->GetSolverTime());
+					auto& ClusteredParticles = MSolver->GetEvolution()->GetParticles().GetClusteredParticles();
+					Obj->FieldParameterUpdateCallback(MSolver, ClusteredParticles, Strains,
+						PositionTarget, PositionTargetedParticles, /*AnimatedPositions,*/ MSolver->GetSolverTime());
 				}
 			}
 
@@ -113,6 +116,8 @@ namespace Chaos
 					{
 						auto& GeomCollectionParticles = MSolver->GetEvolution()->GetParticles().GetGeometryCollectionParticles();
 						Obj->FieldForcesUpdateCallback(MSolver, GeomCollectionParticles, Forces, Torques, MSolver->GetSolverTime());
+						auto& ClusteredParticles = MSolver->GetEvolution()->GetParticles().GetClusteredParticles();
+						Obj->FieldForcesUpdateCallback(MSolver, ClusteredParticles, Forces, Torques, MSolver->GetSolverTime());
 					}
 					MSolver->GetEvolution()->AdvanceOneTimeStep(DeltaTime);
 				}
@@ -121,6 +126,10 @@ namespace Chaos
 				// If time remains, then log why we have lost energy over the timestep.
 				if (TimeRemaining > 0.f)
 				{
+					auto& GeomCollectionParticles = MSolver->GetEvolution()->GetParticles().GetGeometryCollectionParticles();
+					Obj->FieldForcesUpdateCallback(MSolver, GeomCollectionParticles, Forces, Torques, MSolver->GetSolverTime());
+					auto& ClusteredParticles = MSolver->GetEvolution()->GetParticles().GetClusteredParticles();
+					Obj->FieldForcesUpdateCallback(MSolver, ClusteredParticles, Forces, Torques, MSolver->GetSolverTime());
 					if (StepsRemaining == 0)
 					{
 						UE_LOG(LogPBDRigidsSolver, Warning, TEXT("AdvanceOneTimeStepTask::DoWork() - Energy lost over %fs due to too many substeps over large timestep"), TimeRemaining);
@@ -440,6 +449,7 @@ namespace Chaos
 
 	bool FPBDRigidsSolver::UnregisterObject(FGeometryCollectionPhysicsProxy* InProxy)
 	{
+		InProxy->OnRemoveFromSolver(this);
 		InProxy->SetSolver(static_cast<FPBDRigidsSolver*>(nullptr));
 		return GeometryCollectionPhysicsProxies.Remove(InProxy) != 0;
 	}
@@ -656,7 +666,8 @@ namespace Chaos
 		auto Cmd = [Proxy, Solver](Chaos::FPersistentPhysicsTask* PhysThread)
 		{
 			auto* Evolution = Solver->GetEvolution();
-			TManagedArray<Chaos::TPBDGeometryCollectionParticleHandle<float,3>*>& Handles = Proxy->GetSolverParticleHandles();
+			TManagedArray<Chaos::TPBDRigidClusteredParticleHandle<float, 3>*>& Handles = 
+				Proxy->GetSolverParticleHandles();
 			for (auto* Handle : Handles)
 			{
 				if (Handle)
@@ -755,6 +766,9 @@ namespace Chaos
 				case Chaos::EParticleType::GeometryCollection:
 					ActiveGC.AddUnique((FGeometryCollectionPhysicsProxy*)(ActiveObject.GTGeometryParticle()->Proxy));
 					break;
+				case Chaos::EParticleType::Clustered:
+					ActiveGC.AddUnique((FGeometryCollectionPhysicsProxy*)(Proxy));
+					break;
 				default:
 					check(false);
 				}
@@ -793,6 +807,9 @@ namespace Chaos
 				case Chaos::EParticleType::GeometryCollection:
 					ActiveGC.AddUnique((FGeometryCollectionPhysicsProxy*)(ActiveObject.GTGeometryParticle()->Proxy));
 					break;
+				case Chaos::EParticleType::Clustered:
+					ActiveGC.AddUnique((FGeometryCollectionPhysicsProxy*)(Proxy));
+					break;
 				default:
 					check(false);
 				}
@@ -805,6 +822,9 @@ namespace Chaos
 		}
 	}
 
+	// This function is not called.  FPhysScene_ChaosInterface::EndFrame() calls 
+	// FPhysScene_ChaosInterface::SyncBodies() instead, and then immediately afterwards 
+	// calls FPBDRigidsSovler::SyncEvents_GameThread().
 	void FPBDRigidsSolver::UpdateGameThreadStructures()
 	{
 		//ensure(IsInGameThread());
@@ -831,6 +851,9 @@ namespace Chaos
 				case Chaos::EParticleType::GeometryCollection:
 					ActiveGC.AddUnique((FGeometryCollectionPhysicsProxy*)(ActiveObject.GTGeometryParticle()->Proxy));
 					break;
+				case Chaos::EParticleType::Clustered:
+					ActiveGC.AddUnique((FGeometryCollectionPhysicsProxy*)(Proxy));
+					break;
 				default:
 					check(false);
 				}
diff --git a/Engine/Source/Runtime/Experimental/ChaosSolvers/Private/PhysicsProxy/FieldSystemPhysicsProxy.cpp b/Engine/Source/Runtime/Experimental/ChaosSolvers/Private/PhysicsProxy/FieldSystemPhysicsProxy.cpp
--- a/Engine/Source/Runtime/Experimental/ChaosSolvers/Private/PhysicsProxy/FieldSystemPhysicsProxy.cpp
+++ b/Engine/Source/Runtime/Experimental/ChaosSolvers/Private/PhysicsProxy/FieldSystemPhysicsProxy.cpp
@@ -237,7 +237,7 @@ void FFieldSystemPhysicsProxy::FieldParameterUpdateCallback(
 					if (StateChanged)
 					{
 						// regenerate views
-						CurrentSolver->GetParticles().UpdateGeometryCollectionViews();
+						CurrentSolver->GetParticles().UpdateGeometryCollectionViews(true);
 					}
 
 #if TODO_REIMPLEMENT_RIGID_CLUSTERING
@@ -322,7 +322,6 @@ void FFieldSystemPhysicsProxy::FieldParameterUpdateCallback(
 				}
 				CommandsToRemove.Add(CommandIndex);
 			}
-#if TODO_REIMPLEMENT_RIGID_CLUSTERING
 			else if (Command.TargetAttribute == GetFieldPhysicsName(EFieldPhysicsType::Field_ExternalClusterStrain))
 			{
 				SCOPE_CYCLE_COUNTER(STAT_ParamUpdateField_ExternalClusterStrain);
@@ -344,30 +343,40 @@ void FFieldSystemPhysicsProxy::FieldParameterUpdateCallback(
 							Command.MetaData
 						};
 
+						// TODO: Chaos, Ryan
+						// As we're allocating a buffer the size of all particles every iteration, 
+						// I suspect this is a performance hit.  It seems like we should add a buffer
+						// that's recycled rather than reallocating.  It could live on the particles
+						// and its lifetime tied to an object that lives in the scope of the object 
+						// that's driving the sampling of the field.
+
 						TArray<float> StrainSamples;
-						StrainSamples.AddUninitialized(SamplesView.Num());
-						for (const ContextIndex& Index : IndicesArray)
-						{
-							StrainSamples[Index.Sample] = 0.f;
-						}
+						// There's 2 ways to think about initializing this array...
+						// Either we have a low number of indices, and the cost of iterating
+						// over the indices in addition to StrainSamples is lower than the
+						// cost of initializing them all, or it's cheaper and potentially 
+						// more cache coherent to just initialize them all.  I'm thinking
+						// the latter may be more likely...
+						StrainSamples.AddZeroed(SamplesView.Num());
+
 						TArrayView<float> FloatBuffer(&StrainSamples[0], StrainSamples.Num());
 						static_cast<const FFieldNode<float> *>(Command.RootNode.Get())->Evaluate(Context, FloatBuffer);
 
 						int32 Iterations = 1;
 						if (Command.MetaData.Contains(FFieldSystemMetaData::EMetaType::ECommandData_Iteration))
 						{
-							Iterations = static_cast<FFieldSystemMetaDataIteration*>(Command.MetaData[FFieldSystemMetaData::EMetaType::ECommandData_Iteration].Get())->Iterations;
+							Iterations = static_cast<FFieldSystemMetaDataIteration*>(
+								Command.MetaData[FFieldSystemMetaData::EMetaType::ECommandData_Iteration].Get())->Iterations;
 						}
 
 						if (StrainSamples.Num())
 						{
-							CurrentSolver->GetRigidClustering().BreakingModel(StrainSamples);
+							CurrentSolver->GetEvolution()->GetRigidClustering().BreakingModel(&FloatBuffer);
 						}
 					}
 				}
 				CommandsToRemove.Add(CommandIndex);
 			}
-#endif
 			else if (Command.TargetAttribute == GetFieldPhysicsName(EFieldPhysicsType::Field_Kill))
 			{
 				SCOPE_CYCLE_COUNTER(STAT_ParamUpdateField_Kill);
@@ -1104,7 +1113,8 @@ void FFieldSystemPhysicsProxy::GetParticleHandles(
 	}
 	else if (ResolutionType == EFieldResolutionType::Field_Resolution_Minimal)
 	{
-		const Chaos::TParticleView<Chaos::TGeometryParticles<float, 3>> &ParticleView = SolverParticles.GetNonDisabledView();
+		const Chaos::TParticleView<Chaos::TGeometryParticles<float, 3>> &ParticleView = 
+			SolverParticles.GetNonDisabledView();
 		Handles.Reserve(ParticleView.Num());
 		for (Chaos::TParticleIterator<Chaos::TGeometryParticles<float, 3>> It = ParticleView.Begin(), ItEnd = ParticleView.End();
 			It != ItEnd; ++It)
diff --git a/Engine/Source/Runtime/Experimental/ChaosSolvers/Private/PhysicsProxy/GeometryCollectionPhysicsProxy.cpp b/Engine/Source/Runtime/Experimental/ChaosSolvers/Private/PhysicsProxy/GeometryCollectionPhysicsProxy.cpp
--- a/Engine/Source/Runtime/Experimental/ChaosSolvers/Private/PhysicsProxy/GeometryCollectionPhysicsProxy.cpp
+++ b/Engine/Source/Runtime/Experimental/ChaosSolvers/Private/PhysicsProxy/GeometryCollectionPhysicsProxy.cpp
@@ -46,6 +46,7 @@ FAutoConsoleVariableRef CVarCollisionParticlesPerObjectFractionDefault(TEXT("p.C
 FName FGeometryCollectionPhysicsProxy::SimplicialsAttribute("CollisionParticles");
 FName FGeometryCollectionPhysicsProxy::ImplicitsAttribute("Implicits");
 FName FGeometryCollectionPhysicsProxy::SolverParticleHandlesAttribute("SolverParticleHandles");
+FName FGeometryCollectionPhysicsProxy::SolverClusterHandlesAttribute("SolverClusterHandles");
 FName FGeometryCollectionPhysicsProxy::GTGeometryParticleAttribute("GTGeometryParticles");
 
 DEFINE_LOG_CATEGORY_STATIC(UGCC_LOG, Error, All);
@@ -217,19 +218,19 @@ TArray<int32> ComputeRecursiveOrder(const FGeometryCollection& Collection)
 DECLARE_CYCLE_STAT(TEXT("FGeometryCollectionPhysicsProxy::PopulateSimulatedParticle"), STAT_PopulateSimulatedParticle, STATGROUP_Chaos);
 void PopulateSimulatedParticle(
 	Chaos::TPBDRigidParticleHandle<float,3>* Handle,
-	//FGeometryCollectionPhysicsProxy::FParticlesType& Particles,
 	const FSharedSimulationParameters& SharedParams,
 	const FCollisionStructureManager::FSimplicial* Simplicial,
 	Chaos::TSerializablePtr<Chaos::FImplicitObject> Implicit,
 	const float MassIn,
 	const FVector& InertiaTensorVec,
-	//int32 RigidBodyIndex,
 	const FTransform& WorldTransform, 
 	const uint8 DynamicState, 
 	const int16 CollisionGroup)
 {
 	SCOPE_CYCLE_COUNTER(STAT_PopulateSimulatedParticle);
 
+	check(Handle->UniqueIdx().IsValid());
+
 	Handle->SetDisabledLowLevel(false);
 	//Particles.SetDisabledLowLevel(RigidBodyIndex, false);
 
@@ -449,6 +450,7 @@ void FGeometryCollectionPhysicsProxy::Initialize()
 
 	// Back to engine for setup from components
 	InitFunc(Parameters);
+
 	PTDynamicCollection.SyncAllGroups(*Parameters.RestCollection);
 	CollisionParticlesPerObjectFraction = 
 		Parameters.CollisionSampleFraction * CollisionParticlesPerObjectFractionDefault;
@@ -566,7 +568,7 @@ void FGeometryCollectionPhysicsProxy::InitializeBodiesGT()
 		GeometryCollectionAlgo::GlobalMatrices(DynamicCollection->Transform, DynamicCollection->Parent, Transform);
 		check(DynamicCollection->Transform.Num() == Transform.Num());
 
-		NumParticles = SimulatableParticles.Num();
+		NumParticles = GetTransformGroupSize(); // SimulatableParticles.Count(true);
 		BaseParticleIndex = 0; // Are we always zero indexed now?
 
 		// Create game thread particles...
@@ -574,24 +576,59 @@ void FGeometryCollectionPhysicsProxy::InitializeBodiesGT()
 		{
 			if (SimulatableParticles[Idx])
 			{
-				//GTGeometryParticles[Idx] = Chaos::TGeometryParticle<float, 3>::CreateParticle();
-				//GTGeometryParticles[Idx] = Chaos::TKinematicGeometryParticle<float, 3>::CreateParticle();
-				//GTGeometryParticles[Idx] = Chaos::TPBDRigidParticle<float, 3>::CreateParticle();
-				GTGeometryParticles[Idx] = Chaos::TPBDGeometryCollectionParticle<float, 3>::CreateParticle();
-				Chaos::TPBDRigidParticle<float, 3>* GTParticle = GTGeometryParticles[Idx]->CastToRigidParticle();
+				// Clustering data only lives on the physics thread, so there are
+				// no game thread facing API for clustering.  Therefore, we don't 
+				// construct clustered particle instances (only clustered particleS,
+				// emphasis on the "S").
+				
+				TUniquePtr<Chaos::TPBDGeometryCollectionParticle<float, 3>> Particle = 
+					Chaos::TPBDGeometryCollectionParticle<float, 3>::CreateParticle();
+
+				GTGeometryParticles[Idx] = MoveTemp(Particle); // Particle is now undefined
+				Chaos::TPBDRigidParticle<float, 3>* GTParticle = 
+					GTGeometryParticles[Idx]->CastToRigidParticle();
+				GTParticle->SetUniqueIdx(GetSolver()->GetEvolution()->GenerateUniqueIdx(), false);
 
-//				Chaos::TSerializablePtr<Chaos::TImplicitObject<float, 3>> ImplicitShape = Implicits[Idx];
-//				GTParticle->SetGeometry(ImplicitShape);
-//				if (!ensure(GTParticle->IsParticleValid()))
-//				{
-//					check(false); // continue?
-//				}
 				FTransform ParticleTransform = MassToLocal[Idx] * Transform[Idx] * Parameters.WorldTransform;
 				if (GTParticle->ObjectState() == Chaos::EObjectStateType::Dynamic)
 				{
-					GTParticle->SetX(ParticleTransform.GetTranslation());
-					GTParticle->SetR(ParticleTransform.GetRotation());
+					// Calling SetX(), SetR(), etc will mark the particle dirty,
+					// which causes this proxy class to be put on a dirty list.
+					// That causes our PushToPhysicsState() to be called prematurely, 
+					// so we need to guard against that.
+					GTParticle->SetX(ParticleTransform.GetTranslation(), false); 
+					GTParticle->SetR(ParticleTransform.GetRotation().GetNormalized(), false);
+					GTParticle->SetIsland(INDEX_NONE); // doesn't mark dirty
 					GTParticle->Proxy = this;
+
+					// 
+					// AddGeometry
+					//
+
+					TArray<TUniquePtr<Chaos::FImplicitObject>> Geoms;
+					Chaos::TShapesArray<float, 3> Shapes;
+
+					if (Implicits[Idx])
+					{
+						GTParticle->SetGeometry(Implicits[Idx]->Copy());
+					}
+					if (Simplicials[Idx])
+					{
+						// Not sure what to do with particles at this point.
+					}
+
+					if (Parameters.InitialVelocityType == EInitialVelocityTypeEnum::Chaos_Initial_Velocity_User_Defined)
+					{
+						InitialLinearVelocity[Idx] = Parameters.InitialLinearVelocity;
+						InitialAngularVelocity[Idx] = Parameters.InitialAngularVelocity;
+						GTParticle->SetV(Parameters.InitialLinearVelocity, false);
+						GTParticle->SetW(Parameters.InitialAngularVelocity, false);
+					}
+					else
+					{
+						GTParticle->SetV(Chaos::TVector<float, 3>(0.f), false);
+						GTParticle->SetW(Chaos::TVector<float, 3>(0.f), false);
+					}
 				}
 			}
 		}
@@ -600,6 +637,9 @@ void FGeometryCollectionPhysicsProxy::InitializeBodiesGT()
 	}
 }
 
+int32 ReportTooManyChildrenNum = -1;
+FAutoConsoleVariableRef CVarReportTooManyChildrenNum(TEXT("p.ReportTooManyChildrenNum"), ReportTooManyChildrenNum, TEXT("Issue warning if more than this many children exist in a single cluster"));
+
 void FGeometryCollectionPhysicsProxy::InitializeBodiesPT(
 	Chaos::FPBDRigidsSolver* RigidsSolver,
 	Chaos::FPBDRigidsSolver::FParticlesType& Particles)
@@ -622,63 +662,75 @@ void FGeometryCollectionPhysicsProxy::InitializeBodiesPT(
 		const TManagedArray<float>& Mass = RestCollection->GetAttribute<float>("Mass", FTransformCollection::TransformGroup);
 		const TManagedArray<FVector>& InertiaTensor = RestCollection->GetAttribute<FVector>("InertiaTensor", FTransformCollection::TransformGroup);
 
+		const int32 NumTransforms = GetTransformGroupSize();
+
 		TArray<FTransform> Transform;
 		GeometryCollectionAlgo::GlobalMatrices(DynamicCollection->Transform, DynamicCollection->Parent, Transform);
 		check(DynamicCollection->Transform.Num() == Transform.Num());
 
-		// Count particles to add
-		const int NumSimulatedParticles = SimulatableParticles.Count(true);
-
-		//Chaos::FPhysicsSolver* FPSolver = GetSolver();
-		//const Chaos::FPhysicsSolver::FParticlesType& Particles = FPSolver->GetParticles();
-
 		const int NumRigids = 0; // ryan - Since we're doing SOA, we start at zero?
 		BaseParticleIndex = NumRigids;
 
+		// Count geometry collection leaf node particles to add
+		int NumSimulatedParticles = 0;
+		int NumLeafNodes = 0;
+		for (int32 Idx = 0; Idx < SimulatableParticles.Num(); ++Idx)
+		{
+			NumSimulatedParticles += SimulatableParticles[Idx];
+			NumLeafNodes += (SimulatableParticles[Idx] && Children[Idx].Num() == 0);
+		}
+
 		// Add entries into simulation array
+		RigidsSolver->GetEvolution()->ReserveParticles(NumSimulatedParticles);
 		TArray<Chaos::TPBDGeometryCollectionParticleHandle<float, 3>*> Handles = 
-			Particles.CreateGeometryCollectionParticles(NumSimulatedParticles);
+			Particles.CreateGeometryCollectionParticles(NumLeafNodes);
+		for (auto& Handle : Handles)
+		{
+			Handle->SetUniqueIdx(GetSolver()->GetEvolution()->GenerateUniqueIdx());
+		}
+
+
 		int32 NextIdx = 0;
 		for(int32 Idx=0; Idx < SimulatableParticles.Num(); ++Idx)
 		{
-			if (SimulatableParticles[Idx])
+			if (SimulatableParticles[Idx] && Children[Idx].Num() == 0) // Leaf node
 			{
 				check(GTGeometryParticles[Idx]);
 				Chaos::TPBDRigidParticle<float, 3>* GTParticle = GTGeometryParticles[Idx]->CastToRigidParticle();
-
 				Chaos::TPBDGeometryCollectionParticleHandle<float, 3>* Handle = Handles[NextIdx++];
 
 				if (GTParticle->ObjectState() == Chaos::EObjectStateType::Dynamic)
 				{
 					Handle->GTGeometryParticle() = GTParticle;
+					RigidsSolver->AddParticleToProxy(Handle, this);
 				}
 
 				SolverParticleHandles[Idx] = Handle;
+				HandleToTransformGroupIndex.Add(Handle, Idx);
+
+				check(SolverParticleHandles[Idx]->GetParticleType() == Handle->GetParticleType());
 				RigidsSolver->GetEvolution()->CreateParticle(Handle);
 			}
 		}
 
-#if TODO_REIMPLEMENT_RIGID_CLUSTERING
-		Chaos::TArrayCollectionArray<int32>& ClusterGroupIndex = ThisSolver->GetRigidClustering().GetClusterGroupIndexArray();
-		Chaos::TArrayCollectionArray<float>& StrainArray = ThisSolver->GetRigidClustering().GetStrainArray();
 		const float StrainDefault = Parameters.DamageThreshold.Num() ? Parameters.DamageThreshold[0] : 0;
-#endif // TODO_REIMPLEMENT_RIGID_CLUSTERING
 		// Add the rigid bodies
 		const int32 NumGeometries = DynamicCollection->NumElements(FGeometryCollection::GeometryGroup);
 		ParallelFor(NumGeometries, [&](int32 GeometryIndex)
 		{
 			const int32 TransformGroupIndex = TransformIndex[GeometryIndex];
 			//if (SimulatableParticles[TransformGroupIndex] == true)
-			if (Chaos::TPBDRigidParticleHandle<float, 3>* Handle = SolverParticleHandles[TransformGroupIndex])
+			if (Chaos::TPBDRigidClusteredParticleHandle<float, 3>* Handle = SolverParticleHandles[TransformGroupIndex])
 			{
 				//const int32 RigidBodyIndex = RigidBodyID[TransformGroupIndex];
-				const FTransform WorldTransform = MassToLocal[TransformGroupIndex] * Transform[TransformGroupIndex] * Parameters.WorldTransform;
+				const FTransform WorldTransform = 
+					MassToLocal[TransformGroupIndex] * Transform[TransformGroupIndex] * Parameters.WorldTransform;
 
 				PopulateSimulatedParticle(
 					Handle,
 					//Particles, 
 					Parameters.Shared, 
-					nullptr,// Simplicials[TransformGroupIndex].Get(),
+					Simplicials[TransformGroupIndex].Get(),
 					Implicits[TransformGroupIndex], 
 #if TODO_REIMPLEMENT_SOLVER_SETTINGS_ACCESSORS
 					RigidsSolver->GetMassScale() * Mass[TransformGroupIndex], 
@@ -692,10 +744,11 @@ void FGeometryCollectionPhysicsProxy::InitializeBodiesPT(
 					(uint8)Parameters.ObjectType, 
 					CollisionGroup[TransformGroupIndex]);
 
-#if TODO_REIMPLEMENT_RIGID_CLUSTERING
-				ClusterGroupIndex[RigidBodyIndex] = Parameters.ClusterGroupIndex;
-				StrainArray[RigidBodyIndex] = StrainDefault;
-#endif //TODO_REIMPLEMENT_RIGID_CLUSTERING
+				if (Parameters.EnableClustering)
+				{
+					Handle->SetClusterGroupIndex(Parameters.ClusterGroupIndex);
+					Handle->SetStrain(StrainDefault);
+				}
 
 				TUniquePtr<Chaos::TBVHParticles<float, 3>>& CollisionParticles = Handle->CollisionParticles();
 				CollisionParticles.Reset(Simplicials[TransformGroupIndex].Release()); // Steals!
@@ -779,87 +832,116 @@ void FGeometryCollectionPhysicsProxy::InitializeBodiesPT(
 		InitializeRemoveOnFracture(Particles, DynamicState);
 #endif // TODO_REIMPLEMENT_FRACTURE
 
-#if TODO_REIMPLEMENT_RIGID_CLUSTERING
 		// #BG Temporary - don't cluster when playing back. Needs to be changed when kinematics are per-proxy to support
 		// kinematic to dynamic transition for clusters.
 		if (Parameters.EnableClustering)// && Parameters.CacheType != EGeometryCollectionCacheType::Play)
 		{
+			// "RecursiveOrder" means bottom up - children come before their parents.
 			const TArray<int32> RecursiveOrder = ComputeRecursiveOrder(*RestCollection);
 
-			// num clusters
-			uint32 NumClusters=0;
+			// Propagate simulated particle flags up the hierarchy from children 
+			// to their parents, grandparents, etc...
 			TArray<bool> SubTreeContainsSimulatableParticle;
-			SubTreeContainsSimulatableParticle.SetNum(RecursiveOrder.Num());
+			SubTreeContainsSimulatableParticle.SetNumZeroed(RecursiveOrder.Num());
 			for (const int32 TransformGroupIndex : RecursiveOrder)
 			{
-				if (Children[TransformGroupIndex].Num() > 0)
+				if (Children[TransformGroupIndex].Num() == 0)
 				{
-					SubTreeContainsSimulatableParticle[TransformGroupIndex] = false;
-
-					TArray<uint32> RigidChildren, CollectionChildren;
-					for (const int32 ChildIndex : Children[TransformGroupIndex])
+					// Leaf node
+					SubTreeContainsSimulatableParticle[TransformGroupIndex] =
+						SolverParticleHandles[TransformGroupIndex] != nullptr;
+				}
+				else
+				{
+					// Cluster parent
+					const TSet<int32>& ChildIndices = Children[TransformGroupIndex];
+					for (const int32 ChildIndex : ChildIndices)
 					{
 						if(SubTreeContainsSimulatableParticle[ChildIndex])
 						{
-							NumClusters++;
 							SubTreeContainsSimulatableParticle[TransformGroupIndex] = true;
 							break;
 						}
 					}
 				}
-				else
-				{
-					SubTreeContainsSimulatableParticle[TransformGroupIndex] = (RigidBodyID[TransformGroupIndex] != INDEX_NONE);
-				}
 			}
 
-			const int32 ClusterStartIndex = Particles.Size();
-			Particles.AddParticles(NumClusters);
+			TArray<Chaos::TPBDRigidClusteredParticleHandle<float, 3>*> ClusterHandles;
+			// Ryan - It'd be better to batch allocate cluster particles ahead of time,
+			// but if ClusterHandles is empty, then new particles will be allocated
+			// on the fly by TPBDRigidClustering::CreateClusterParticle(), which 
+			// needs to work before this does...
+			//ClusterHandles = GetSolver()->GetEvolution()->CreateClusteredParticles(NumClusters);
 
-			int32 ClusterRigidBodyId = ClusterStartIndex;
+			int32 ClusterHandlesIndex = 0;
+			TArray<Chaos::TPBDRigidParticleHandle<float, 3>*> RigidChildren;
+			TArray<int32> RigidChildrenTransformGroupIndex;
 			for (const int32 TransformGroupIndex : RecursiveOrder)
 			{
-				TArray<uint32> RigidChildren, CollectionChildren;
+				// Don't construct particles for branches of the hierarchy that  
+				// don't contain any simulated particles.
+				if (!SubTreeContainsSimulatableParticle[TransformGroupIndex])
+				{
+					continue;
+				}
+
+				RigidChildren.Reset(Children.Num());
+				RigidChildrenTransformGroupIndex.Reset(Children.Num());
 				for (const int32 ChildIndex : Children[TransformGroupIndex])
 				{
-					if (RigidBodyID[ChildIndex] != INDEX_NONE)
+					if (Chaos::TPBDRigidClusteredParticleHandle<float, 3>* Handle = SolverParticleHandles[ChildIndex])
 					{
-						RigidChildren.Add(RigidBodyID[ChildIndex]);
-						CollectionChildren.Add(ChildIndex);
+						RigidChildren.Add(Handle);
+						RigidChildrenTransformGroupIndex.Add(ChildIndex);
 					}
 				}
 				if (RigidChildren.Num())
 				{
 					if (ReportTooManyChildrenNum >= 0 && RigidChildren.Num() > ReportTooManyChildrenNum)
 					{
-						UE_LOG(LogChaos, Warning, TEXT("Too many children (%d) in a single cluster:%s"), RigidChildren.Num(), *Parameters.Name);
+						UE_LOG(LogChaos, Warning, TEXT("Too many children (%d) in a single cluster:%s"), 
+							RigidChildren.Num(), *Parameters.Name);
 					}
 
-					RigidBodyID[TransformGroupIndex] = ClusterRigidBodyId;
 					Chaos::FClusterCreationParameters<float> CreationParameters;
-					CreationParameters.RigidBodyIndex = RigidBodyID[TransformGroupIndex];
-					BuildClusters(TransformGroupIndex, CollectionChildren, RigidChildren, CreationParameters);
-					ClusterRigidBodyId++;
+					CreationParameters.ClusterParticleHandle = 
+						ClusterHandles.Num() ? ClusterHandles[ClusterHandlesIndex++] : nullptr;
+
+					Chaos::TPBDRigidClusteredParticleHandle<float, 3>* Handle =
+						BuildClusters(
+							TransformGroupIndex, 
+							RigidChildren, 
+							RigidChildrenTransformGroupIndex, 
+							CreationParameters);
+
+					SolverClusterHandles[TransformGroupIndex] = Handle;
+					SolverParticleHandles[TransformGroupIndex] = Handle;
+					HandleToTransformGroupIndex.Add(Handle, TransformGroupIndex);
+					RigidsSolver->AddParticleToProxy(Handle, this);
 				}
 			}
-		}
-#endif // TODO_REIMPLEMENT_RIGID_CLUSTERING
- 
-		const int32 NumTransforms = GetTransformGroupSize();
-#if TODO_REIMPLEMENT_RIGID_CLUSTERING
-		// Set Connectivity
-		ParallelFor(NumTransforms, [&](int32 TransformGroupIndex)
-		{
-			if (Children[TransformGroupIndex].Num() > 0)
+
+			// We've likely chagned the state of leaf nodes, which are geometry
+			// collection particles.  Update which particle views they belong in,
+			// as well as views of clustered particles.
+			Particles.UpdateGeometryCollectionViews(true);
+
+			// Set cluster connectivity.  TPBDRigidClustering::CreateClusterParticle() 
+			// will optionally do this, but we switch that functionality off in BuildClusters().
+			ParallelFor(NumTransforms, [&](int32 TransformGroupIndex)
 			{
-				if (RigidBodyID[TransformGroupIndex] != INDEX_NONE)
+				const TManagedArray<TSet<int32>>& Children = RestCollection->Children;
+				if (Children[TransformGroupIndex].Num() > 0)
 				{
-					Solver->GetEvolution()->GetRigidClustering().GenerateConnectionGraph(RigidBodyID[TransformGroupIndex]);
-					//GetSolver()->GetRigidClustering().GenerateConnectionGraph(RigidBodyID[TransformGroupIndex]);
+					if (SolverClusterHandles[TransformGroupIndex])
+					{
+						GetSolver()->GetEvolution()->GetRigidClustering().GenerateConnectionGraph(
+							SolverClusterHandles[TransformGroupIndex]);
+					}
 				}
-			}
-		});
-#endif // TODO_REIMPLEMENT_RIGID_CLUSTERING
+			});
+		} // end if EnableClustering
+ 
 
 #if TODO_REIMPLEMENT_RIGID_CACHING
 		// If we're recording and want to start immediately caching then we should cache the rest state
@@ -878,13 +960,12 @@ void FGeometryCollectionPhysicsProxy::InitializeBodiesPT(
 			//  Activated bodies has already been called so we are good to go.
 			//
 			InitializedState = ESimulationInitializationState::Initialized;
-#if TODO_REIMPLEMENT_RIGID_CLUSTERING
+
 			if (Parameters.EnableClustering && Parameters.ClusterGroupIndex)
 			{
-				GetSolver()->GetRigidClustering().IncrementPendingClusterCounter(Parameters.ClusterGroupIndex);
-				GetSolver()->GetRigidClustering().DecrementPendingClusterCounter(Parameters.ClusterGroupIndex);
+				GetSolver()->GetEvolution()->GetRigidClustering().IncrementPendingClusterCounter(Parameters.ClusterGroupIndex);
+				GetSolver()->GetEvolution()->GetRigidClustering().DecrementPendingClusterCounter(Parameters.ClusterGroupIndex);
 			}
-#endif // TODO_REIMPLEMENT_RIGID_CLUSTERING
 		}
 		else if (InitializedState == ESimulationInitializationState::Unintialized)
 		{
@@ -913,7 +994,6 @@ void FGeometryCollectionPhysicsProxy::InitializeBodiesPT(
 				}
 			}
 
-#if TODO_REIMPLEMENT_RIGID_CLUSTERING
 			//
 			//  Clustering needs to advertise its group id to the cluster so
 			//  that the group is not initialized before all the bodies are
@@ -921,9 +1001,9 @@ void FGeometryCollectionPhysicsProxy::InitializeBodiesPT(
 			//
 			if (Parameters.EnableClustering && Parameters.ClusterGroupIndex )
 			{
-				GetSolver()->GetRigidClustering().IncrementPendingClusterCounter(Parameters.ClusterGroupIndex);
+				GetSolver()->GetEvolution()->GetRigidClustering().IncrementPendingClusterCounter(
+					Parameters.ClusterGroupIndex);
 			}
-#endif // TODO_REIMPLEMENT_RIGID_CLUSTERING
 		}
 		else
 		{
@@ -1562,16 +1642,17 @@ void FGeometryCollectionPhysicsProxy::EndFrameCallback(const float EndFrame)
 	Commands.Empty();
 }
 
-int32 ReportTooManyChildrenNum = -1;
-FAutoConsoleVariableRef CVarReportTooManyChildrenNum(TEXT("p.ReportTooManyChildrenNum"), ReportTooManyChildrenNum, TEXT("Issue warning if more than this many children exist in a single cluster"));
-
 void FGeometryCollectionPhysicsProxy::CreateRigidBodyCallback(FParticlesType& Particles)
 {
+	check(false); // This function is deprecated and should never be called.
+
 	const FGeometryCollection* RestCollection = Parameters.RestCollection;
 	FGeometryDynamicCollection* DynamicCollection = Parameters.DynamicCollection;
 	check(DynamicCollection);
 	
-	if (Parameters.Simulating && ((InitializedState == ESimulationInitializationState::Unintialized) || (InitializedState == ESimulationInitializationState::Activated)))
+	if (Parameters.Simulating && 
+		((InitializedState == ESimulationInitializationState::Unintialized) || 
+		 (InitializedState == ESimulationInitializationState::Activated)))
 	{
 #if TODO_REIMPLEMENT_RIGID_CLUSTERING
 		Chaos::FPhysicsSolver* ThisSolver = GetSolver();
@@ -1833,7 +1914,7 @@ void FGeometryCollectionPhysicsProxy::ActivateBodies()
 
 			if (Parameters.EnableClustering && Parameters.ClusterGroupIndex)
 			{
-				Chaos::FPhysicsSolver::FClusteringType & Clustering = GetSolver()->GetRigidClustering();
+				auto& Clustering = GetSolver()->GetEvolution()->GetRigidClustering();
 				Clustering.DecrementPendingClusterCounter(Parameters.ClusterGroupIndex);
 				ParentIndex = Parameters.ClusterGroupIndex;
 			}
@@ -2357,33 +2438,40 @@ FAutoConsoleVariableRef CVarReportNoLevelsetCluster(TEXT("p.gc.ReportNoLevelsetC
 
 DECLARE_CYCLE_STAT(TEXT("FGeometryCollectionPhysicsProxy::BuildClusters"), STAT_BuildClusters, STATGROUP_Chaos);
 DECLARE_CYCLE_STAT(TEXT("FGeometryCollectionPhysicsProxy::BuildClusters:GlobalMatrices"), STAT_BuildClustersGlobalMatrices, STATGROUP_Chaos);
-void FGeometryCollectionPhysicsProxy::BuildClusters(uint32 CollectionClusterIndex, const TArray<uint32>& CollectionChildIDs, const TArray<uint32>& ChildIDs, const Chaos::FClusterCreationParameters<float> & ClusterParameters)
+
+
+Chaos::TPBDRigidClusteredParticleHandle<float, 3>* 
+FGeometryCollectionPhysicsProxy::BuildClusters(
+	const uint32 CollectionClusterIndex, // TransformGroupIndex
+	TArray<Chaos::TPBDRigidParticleHandle<float,3>*>& ChildHandles,
+	const TArray<int32>& ChildTransformGroupIndices,
+	const Chaos::FClusterCreationParameters<float> & ClusterParameters)
 {
 	SCOPE_CYCLE_COUNTER(STAT_BuildClusters);
-	check(CollectionChildIDs.Num() == ChildIDs.Num());
+
 	check(CollectionClusterIndex != INDEX_NONE);
-	check(ChildIDs.Num() != 0);
+	check(ChildHandles.Num() != 0);
 
+	check(Parameters.DynamicCollection);
 	FGeometryDynamicCollection* Collection = Parameters.DynamicCollection;
-	check(Collection);
-
 	TManagedArray<int32>& DynamicState = Collection->DynamicState;
 	TManagedArray<int32>& CollisionGroup = Collection->CollisionGroup;
+	TManagedArray<int32>& ParentIndex = Collection->Parent;
+	TManagedArray<TSet<int32>>& Children = Collection->Children;
+	TManagedArray<FTransform>& Transform = Collection->Transform;
 
-#if TODO_REIMPLEMENT_GET_RIGID_PARTICLES
 	Chaos::FPhysicsSolver* ThisSolver = GetSolver();
-	Chaos::TPBDRigidParticles<float, 3>& Particles = ThisSolver->GetRigidParticles();
-	TManagedArray<int32>& Parent = Collection->Parent;
-	TManagedArray<TSet<int32>>& Children = Collection->Children;
-	TManagedArray<FTransform>& HierarchyTransform = Collection->Transform;
 
 	//todo(ocohen): cache this
+	check(Parameters.RestCollection);
 	const TManagedArray<float>& Mass = Parameters.RestCollection->GetAttribute<float>("Mass", FTransformCollection::TransformGroup);
 	const TManagedArray<FVector>& InertiaTensor = Parameters.RestCollection->GetAttribute<FVector>("InertiaTensor", FTransformCollection::TransformGroup);
 
 	//If we are a root particle use the world transform, otherwise set the relative transform
-	const FTransform CollectionSpaceTransform = GeometryCollectionAlgo::GlobalMatrix(Collection->Transform, Parent, CollectionClusterIndex);
-	const Chaos::TRigidTransform<float, 3> ParticleTM = MassToLocal[CollectionClusterIndex] * CollectionSpaceTransform * Parameters.WorldTransform;
+	const FTransform CollectionSpaceTransform = 
+		GeometryCollectionAlgo::GlobalMatrix(Transform, ParentIndex, CollectionClusterIndex);
+	const Chaos::TRigidTransform<float, 3> ParticleTM = 
+		MassToLocal[CollectionClusterIndex] * CollectionSpaceTransform * Parameters.WorldTransform;
 
 	//create new cluster particle
 	//The reason we need to pass in a mass orientation override is as follows:
@@ -2397,22 +2485,36 @@ void FGeometryCollectionPhysicsProxy::BuildClusters(uint32 CollectionClusterInde
 	ClusterCreationParameters.ConnectionMethod = Parameters.ClusterConnectionMethod;
 	if (ClusterCreationParameters.CollisionParticles)
 	{
-		int32 NumCollisionParticles = ClusterCreationParameters.CollisionParticles->Size();
-		int32 CollisionParticlesSize = FMath::Max(0, FMath::Min(int(NumCollisionParticles * CollisionParticlesPerObjectFraction), NumCollisionParticles));
-		ClusterCreationParameters.CollisionParticles->Resize(CollisionParticlesSize);
+		const int32 NumCollisionParticles = ClusterCreationParameters.CollisionParticles->Size();
+		const int32 ClampedCollisionParticlesSize = 
+			FMath::Max(0, FMath::Min(int(NumCollisionParticles * CollisionParticlesPerObjectFraction), NumCollisionParticles));
+		ClusterCreationParameters.CollisionParticles->Resize(ClampedCollisionParticlesSize);
+	}
+	TArray<Chaos::TPBDRigidParticleHandle<float, 3>*> ChildHandlesCopy(ChildHandles);
+
+	// Construct an active cluster particle, disable children, derive M and I from children:
+	Chaos::TPBDRigidClusteredParticleHandle<float, 3>* Parent = 
+		GetSolver()->GetEvolution()->GetRigidClustering().CreateClusterParticle(
+			Parameters.ClusterGroupIndex, 
+			MoveTemp(ChildHandlesCopy),
+			ClusterCreationParameters,
+			Implicits[CollectionClusterIndex],
+			&ParticleTM);
+	if (!Parent->UniqueIdx().IsValid())
+	{
+		Parent->SetUniqueIdx(GetSolver()->GetEvolution()->GenerateUniqueIdx());
 	}
-	const int NewSolverClusterID = GetSolver()->GetRigidClustering().CreateClusterParticle(Parameters.ClusterGroupIndex, ChildIDs, Implicits[CollectionClusterIndex], &ParticleTM, ClusterCreationParameters);
-	// @todo(mlentine): This is not the best solution to set parameters
 
-	if (ReportNoLevelsetCluster && GetSolver()->GetRigidParticles().DynamicGeometry(NewSolverClusterID))
+	if (ReportNoLevelsetCluster && 
+		Parent->DynamicGeometry())
 	{
 		//ensureMsgf(false, TEXT("Union object generated for cluster"));
 		UE_LOG(LogChaos, Warning, TEXT("Union object generated for cluster:%s"), *Parameters.Name);
 	}
 
-	GetSolver()->SetPhysicsMaterial(NewSolverClusterID, GetSolver()->GetPhysicsMaterial(ChildIDs[0]));
-	//GetSolver()->UpdateKinematicProperties(NewSolverClusterID);
-
+	GetSolver()->GetEvolution()->SetPhysicsMaterial(
+		Parent, 
+		GetSolver()->GetEvolution()->GetPhysicsMaterial(ChildHandles[0]));
 
 	if(Parameters.InitialVelocityType == EInitialVelocityTypeEnum::Chaos_Initial_Velocity_User_Defined)
 	{
@@ -2420,9 +2522,9 @@ void FGeometryCollectionPhysicsProxy::BuildClusters(uint32 CollectionClusterInde
 		InitialAngularVelocity[CollectionClusterIndex] = Parameters.InitialAngularVelocity;
 	}
 
-	if (Particles.InvM(NewSolverClusterID) == 0.0)
+	if (Parent->InvM() == 0.0)
 	{
-		if (Particles.ObjectState(NewSolverClusterID) == Chaos::EObjectStateType::Static)
+		if (Parent->ObjectState() == Chaos::EObjectStateType::Static)
 		{
 			DynamicState[CollectionClusterIndex] = (uint8)EObjectStateTypeEnum::Chaos_Object_Static;
 		}
@@ -2435,50 +2537,72 @@ void FGeometryCollectionPhysicsProxy::BuildClusters(uint32 CollectionClusterInde
 	//const FTransform ParticleTM(Particles.R(NewSolverClusterID), Particles.X(NewSolverClusterID));	//in theory we should be computing this and passing in to avoid inertia computation at runtime. If we do this we must account for leaf particles that have already been created in world space
 	MassToLocal[CollectionClusterIndex] = FTransform::Identity;
 
-	PopulateSimulatedParticle(Particles, Parameters.Shared, nullptr, //Simplicials[CollectionClusterIndex].Get(),
-		Implicits[CollectionClusterIndex], ThisSolver->GetMassScale() * Mass[CollectionClusterIndex], ThisSolver->GetMassScale() * InertiaTensor[CollectionClusterIndex],
-		NewSolverClusterID, ParticleTM, (uint8)DynamicState[CollectionClusterIndex], 0);
+	PopulateSimulatedParticle(
+		Parent,
+		Parameters.Shared, 
+		nullptr, // CollisionParticles is optionally set from CreateClusterParticle()
+		Chaos::TSerializablePtr<Chaos::FImplicitObject>(), // Parent->Geometry() ? Parent->Geometry() : Implicits[CollectionClusterIndex], 
+#if TODO_REIMPLEMENT_SOLVER_SETTINGS_ACCESSORS
+		// Note: if the children had this scaling baked in, then cluster parents
+		// will too, and we shouldn't apply the scaling to the parents twice!
+		ThisSolver->GetMassScale() * Mass[CollectionClusterIndex], 
+		ThisSolver->GetMassScale() * InertiaTensor[CollectionClusterIndex],
+#else
+		// We computed M and I from children in CreateClusterParticle().
+		Parent->M() > 0.0 ? Parent->M() : Mass[CollectionClusterIndex], 
+		Parent->I().GetDiagonal() != Chaos::TVector<float,3>(0.0) ? Parent->I().GetDiagonal() : InertiaTensor[CollectionClusterIndex],
+#endif
+		ParticleTM, 
+		(uint8)DynamicState[CollectionClusterIndex], 
+		0); // CollisionGroup
 
 	// two-way mapping
-	RigidBodyID[CollectionClusterIndex] = NewSolverClusterID;
+	SolverClusterHandles[CollectionClusterIndex] = Parent;
 
 	const int32 NumThresholds = Parameters.DamageThreshold.Num();
-	int32 Level = FMath::Clamp(CalculateHierarchyLevel(Collection, CollectionClusterIndex), 0, INT_MAX);
-	const float Default = NumThresholds > 0 ? Parameters.DamageThreshold[NumThresholds - 1] : 0;
+	const int32 Level = FMath::Clamp(CalculateHierarchyLevel(Collection, CollectionClusterIndex), 0, INT_MAX);
+	const float Default = NumThresholds > 0 ? Parameters.DamageThreshold[NumThresholds - 1] : 0.f;
 	float Damage = Level < NumThresholds ? Parameters.DamageThreshold[Level] : Default;
 	if(Level >= Parameters.MaxClusterLevel) 
 		Damage = FLT_MAX;
-	int32 MinCollisionGroup = INT_MAX;
-
-	Chaos::TArrayCollectionArray<float>& SolverStrainArray = GetSolver()->GetRigidClustering().GetStrainArray();
+	Parent->SetStrains(Damage);
 
-	SolverStrainArray[NewSolverClusterID] = Damage;
-	GetSolver()->SetPhysicsMaterial(NewSolverClusterID, Parameters.PhysicalMaterial);
+	GetSolver()->GetEvolution()->SetPhysicsMaterial(Parent, Parameters.PhysicalMaterial);
 
-	FTransform ParentTransform = GeometryCollectionAlgo::GlobalMatrix(Collection->Transform, Collection->Parent, CollectionClusterIndex);
+	const FTransform ParentTransform = 
+		GeometryCollectionAlgo::GlobalMatrix(Collection->Transform, Collection->Parent, CollectionClusterIndex);
 
-	for(int32 idx = 0; idx < ChildIDs.Num(); idx++)
+	int32 MinCollisionGroup = INT_MAX;
+	for(int32 Idx=0; Idx < ChildHandles.Num(); Idx++)
 	{
-		SolverStrainArray[ChildIDs[idx]]= Damage;
+		const int32 ChildTransformGroupIndex = ChildTransformGroupIndices[Idx];
+		Chaos::TPBDRigidParticleHandle<float, 3>* Child = ChildHandles[Idx];
+
+		if (Chaos::TPBDRigidClusteredParticleHandle<float, 3>* ClusteredChild = Child->CastToClustered())
+		{
+			//SolverStrainArray[ChildIDs[idx]]= Damage;
+			ClusteredChild->SetStrains(Damage);
+		}
 
-		const int32 TransformGroupIndex = CollectionChildIDs[idx];
-		SolverClusterID[TransformGroupIndex] = NewSolverClusterID;
-		MinCollisionGroup = FMath::Min(CollisionGroup[TransformGroupIndex], MinCollisionGroup);
+		SolverClusterHandles[ChildTransformGroupIndex] = Parent;
 
-		const FTransform ChildTransform(Particles.R(ChildIDs[idx]), Particles.X(ChildIDs[idx]));
-		if(Children[TransformGroupIndex].Num()) // clustered local transform
+		MinCollisionGroup = FMath::Min(Child->CollisionGroup(), MinCollisionGroup);
+
+		const FTransform ChildTransform(Child->R(), Child->X());
+		if(Children[ChildTransformGroupIndex].Num()) // clustered local transform
 		{
-			HierarchyTransform[TransformGroupIndex] = ChildTransform.GetRelativeTransform(ParticleTM);
+			Transform[ChildTransformGroupIndex] = ChildTransform.GetRelativeTransform(ParticleTM);
 		}
 		else // rigid local transform
 		{
-			const FTransform RestTransform = Parameters.RestCollection->Transform[TransformGroupIndex] * ParentTransform * Parameters.WorldTransform;
-			HierarchyTransform[TransformGroupIndex] = RestTransform.GetRelativeTransform(ParticleTM);
+			const FTransform RestTransform = 
+				Parameters.RestCollection->Transform[ChildTransformGroupIndex] * ParentTransform * Parameters.WorldTransform;
+			Transform[ChildTransformGroupIndex] = RestTransform.GetRelativeTransform(ParticleTM);
 		}
-		HierarchyTransform[TransformGroupIndex].NormalizeRotation();
+		Transform[ChildTransformGroupIndex].NormalizeRotation();
 	}
-	CollisionGroup[CollectionClusterIndex] = MinCollisionGroup;
-#endif
+	Parent->SetCollisionGroup(MinCollisionGroup);
+	return Parent;
 }
 
 int32 FGeometryCollectionPhysicsProxy::CalculateHierarchyLevel(FGeometryDynamicCollection* GeometryCollection, int32 TransformIndex)
@@ -2508,6 +2632,7 @@ void FGeometryCollectionPhysicsProxy::CreateDynamicAttributes()
 	DynamicCollection->AddExternalAttribute(SimplicialsAttribute, FTransformCollection::TransformGroup, Simplicials);
 	DynamicCollection->AddExternalAttribute(ImplicitsAttribute, FTransformCollection::TransformGroup, Implicits);
 	DynamicCollection->AddExternalAttribute(SolverParticleHandlesAttribute, FTransformCollection::TransformGroup, SolverParticleHandles);
+	DynamicCollection->AddExternalAttribute(SolverClusterHandlesAttribute, FTransformCollection::TransformGroup, SolverClusterHandles/*SolverParticleHandles*/);
 	DynamicCollection->AddExternalAttribute(GTGeometryParticleAttribute, FTransformCollection::TransformGroup, GTGeometryParticles);
 
 	if (Parameters.InitialVelocityType == EInitialVelocityTypeEnum::Chaos_Initial_Velocity_User_Defined)
@@ -2516,7 +2641,8 @@ void FGeometryCollectionPhysicsProxy::CreateDynamicAttributes()
 		DynamicCollection->AddExternalAttribute("InitialLinearVelocity", FTransformCollection::TransformGroup, InitialLinearVelocity);
 	}
 
-	const auto& RestImplicits = Parameters.RestCollection->GetAttribute<TUniquePtr<Chaos::FImplicitObject>>(ImplicitsAttribute, FTransformCollection::TransformGroup);
+	const auto& RestImplicits = Parameters.RestCollection->GetAttribute<TUniquePtr<Chaos::FImplicitObject>>(
+		ImplicitsAttribute, FTransformCollection::TransformGroup);
 	for (int32 Index = DynamicCollection->NumElements(FTransformCollection::TransformGroup) - 1; 0 <= Index; Index--)
 	{
 		Simplicials[Index] = TUniquePtr<FSimplicial>(nullptr);
@@ -2853,12 +2979,19 @@ void FGeometryCollectionPhysicsProxy::InitializeSharedCollisionStructures(
 		}
 	}
 
-#if TODO_REIMPLEMENT_RIGID_CLUSTERING
 	if (bEnableClustering)
 	{
 		//Put all children into collection space so we can compute mass properties.
-		TPBDRigidParticles<float, 3> CollectionSpaceParticles;
-		CollectionSpaceParticles.AddParticles(NumTransforms);
+		TUniquePtr<TPBDRigidClusteredParticles<float, 3>> CollectionSpaceParticles(new TPBDRigidClusteredParticles<float, 3>());
+		CollectionSpaceParticles->AddParticles(NumTransforms);
+		TArray<TUniquePtr<TPBDRigidClusteredParticleHandle<float, 3>>> Handles;
+		Handles.Reserve(NumTransforms);
+		for (int32 Idx = 0; Idx < NumTransforms; Idx++)
+		{
+			Handles.Add(TPBDRigidClusteredParticleHandle<float, 3>::CreateParticleHandle(
+				MakeSerializable(CollectionSpaceParticles), Idx, Idx));
+			Handles[Idx]->SetUniqueIdx(FUniqueIdx(Idx));
+		}
 
 		for (int32 GeometryIdx = 0; GeometryIdx < NumGeometries; ++GeometryIdx)
 		{
@@ -2866,11 +2999,16 @@ void FGeometryCollectionPhysicsProxy::InitializeSharedCollisionStructures(
 			if (CollectionSimulatableParticles[TransformGroupIndex])
 			{
 				FTransform MassToComponent = CollectionMassToLocal[TransformGroupIndex] * CollectionSpaceTransforms[TransformGroupIndex];
-				PopulateSimulatedParticle(CollectionSpaceParticles,
-					SharedParams, CollectionSimplicials[TransformGroupIndex].Get(),
-					Chaos::MakeSerializable(CollectionImplicits[TransformGroupIndex]), CollectionMass[TransformGroupIndex],
-					CollectionInertiaTensor[TransformGroupIndex], TransformGroupIndex, MassToComponent,
-					(uint8)EObjectStateTypeEnum::Chaos_Object_Dynamic, INDEX_NONE);
+				PopulateSimulatedParticle(
+					Handles[TransformGroupIndex].Get(),
+					SharedParams, 
+					CollectionSimplicials[TransformGroupIndex].Get(),
+					Chaos::MakeSerializable(CollectionImplicits[TransformGroupIndex]), 
+					CollectionMass[TransformGroupIndex],
+					CollectionInertiaTensor[TransformGroupIndex], 
+					MassToComponent,
+					(uint8)EObjectStateTypeEnum::Chaos_Object_Dynamic, 
+					INDEX_NONE); // CollisionGroup
 			}
 		}
 
@@ -2878,7 +3016,7 @@ void FGeometryCollectionPhysicsProxy::InitializeSharedCollisionStructures(
 		const TArray<int32> TransformToGeometry = ComputeTransformToGeometryMap(RestCollection);
 
 		TArray<bool> IsClusterSimulated;
-		IsClusterSimulated.Init(false, CollectionSpaceParticles.Size());
+		IsClusterSimulated.Init(false, CollectionSpaceParticles->Size());
 		//build collision structures depth first
 		for (const int32 TransformGroupIndex : RecursiveOrder)
 		{
@@ -2886,13 +3024,13 @@ void FGeometryCollectionPhysicsProxy::InitializeSharedCollisionStructures(
 			{
 				const int32 ClusterTransformIdx = TransformGroupIndex;
 				//update mass 
-				TArray<uint32> ChildrenIndices;
+				TSet<TPBDRigidParticleHandle<float,3>*> ChildrenIndices;
 				ChildrenIndices.Reserve(Children[ClusterTransformIdx].Num());
 				for (int32 ChildIdx : Children[ClusterTransformIdx])
 				{
 					if (CollectionSimulatableParticles[ChildIdx] || IsClusterSimulated[ChildIdx])
 					{
-						ChildrenIndices.Add(ChildIdx);
+						ChildrenIndices.Add(Handles[ChildIdx].Get());
 					}
 				}
 				if (!ChildrenIndices.Num())
@@ -2902,26 +3040,33 @@ void FGeometryCollectionPhysicsProxy::InitializeSharedCollisionStructures(
 				//CollectionSimulatableParticles[TransformGroupIndex] = true;
 				IsClusterSimulated[TransformGroupIndex] = true;
 
-				UpdateClusterMassProperties(CollectionSpaceParticles, ChildrenIndices, ClusterTransformIdx);	//compute mass properties
-				const FTransform ClusterMassToCollection = FTransform(CollectionSpaceParticles.R(ClusterTransformIdx), CollectionSpaceParticles.X(ClusterTransformIdx));
+				UpdateClusterMassProperties(Handles[ClusterTransformIdx].Get(), ChildrenIndices);	//compute mass properties
+				const FTransform ClusterMassToCollection = 
+					FTransform(CollectionSpaceParticles->R(ClusterTransformIdx), 
+							   CollectionSpaceParticles->X(ClusterTransformIdx));
 
-				//Compute MassToLocal as if the transform hierarchy stays fixed. In reality we modify the transform hierarchy so that MassToLocal is identity at runtime.
-				CollectionMassToLocal[ClusterTransformIdx] = ClusterMassToCollection.GetRelativeTransform(CollectionSpaceTransforms[ClusterTransformIdx]);
+				// Compute MassToLocal as if the transform hierarchy stays fixed. 
+				// In reality we modify the transform hierarchy so that MassToLocal 
+				// is identity at runtime.
+				CollectionMassToLocal[ClusterTransformIdx] = 
+					ClusterMassToCollection.GetRelativeTransform(
+						CollectionSpaceTransforms[ClusterTransformIdx]);
 
 				//update geometry
 				//merge children meshes and move them into cluster's mass space
 				TArray<TVector<int32, 3>> UnionMeshIndices;
 				FBox InstanceBoundingBox(EForceInit::ForceInitToZero);
-				for (uint32 Child : ChildrenIndices)
+				for (TPBDRigidParticleHandle<float,3>* Child : ChildrenIndices)
 				{
-					const int32 ChildTransformIdx = Child;
-
+					const int32 ChildTransformIdx = Child->UniqueIdx().Idx;
 					if (Chaos::TTriangleMesh<float>* ChildMesh = TriangleMeshesArray[ChildTransformIdx].Get())
 					{
 						const TArray<TVector<int32, 3>>& ChildIndices = ChildMesh->GetSurfaceElements();
 						UnionMeshIndices.Append(ChildIndices);
 
-						const FTransform ChildMassToClusterMass = (CollectionMassToLocal[ChildTransformIdx] * CollectionSpaceTransforms[ChildTransformIdx]).GetRelativeTransform(ClusterMassToCollection);
+						const FTransform ChildMassToClusterMass = 
+							(CollectionMassToLocal[ChildTransformIdx] * CollectionSpaceTransforms[ChildTransformIdx]).GetRelativeTransform(
+								ClusterMassToCollection);
 
 						TSet<int32> VertsAdded;
 						for (const TVector<int32, 3>& Tri : ChildIndices)
@@ -2932,7 +3077,8 @@ void FGeometryCollectionPhysicsProxy::InitializeSharedCollisionStructures(
 								if (!VertsAdded.Contains(VertIdx))
 								{
 									//Update particles so they are in the cluster's mass space
-									MassSpaceParticles.X(VertIdx) = ChildMassToClusterMass.TransformPosition(MassSpaceParticles.X(VertIdx));
+									MassSpaceParticles.X(VertIdx) = 
+										ChildMassToClusterMass.TransformPosition(MassSpaceParticles.X(VertIdx));
 									InstanceBoundingBox += MassSpaceParticles.X(VertIdx);
 									VertsAdded.Add(VertIdx);
 								}
@@ -2942,19 +3088,19 @@ void FGeometryCollectionPhysicsProxy::InitializeSharedCollisionStructures(
 				}
 
 				TUniquePtr<TTriangleMesh<float>> UnionMesh(new TTriangleMesh<float>(MoveTemp(UnionMeshIndices)));
-				const FMatrix& InertiaMatrix = CollectionSpaceParticles.I(ClusterTransformIdx);
+				const FMatrix& InertiaMatrix = CollectionSpaceParticles->I(ClusterTransformIdx);
 				const FVector InertiaDiagonal(InertiaMatrix.M[0][0], InertiaMatrix.M[1][1], InertiaMatrix.M[2][2]);
 				CollectionInertiaTensor[ClusterTransformIdx] = InertiaDiagonal;
-				CollectionMass[ClusterTransformIdx] = CollectionSpaceParticles.M(ClusterTransformIdx);
+				CollectionMass[ClusterTransformIdx] = CollectionSpaceParticles->M(ClusterTransformIdx);
 
 				const int32 SizeSpecificIdx = FindSizeSpecificIdx(SharedParams.SizeSpecificData, InstanceBoundingBox);
 				const FSharedSimulationSizeSpecificData& SizeSpecificData = SharedParams.SizeSpecificData[SizeSpecificIdx];
 
 				if (SizeSpecificData.ImplicitType == EImplicitTypeEnum::Chaos_Implicit_LevelSet)
 				{
-					TVector<float, 3> Scale = 2 * InstanceBoundingBox.GetExtent() / MaxChildBounds;
-					float ScaleMax = Scale.GetAbsMax();
-					float ScaleMin = Scale.GetAbsMin();
+					const TVector<float, 3> Scale = 2 * InstanceBoundingBox.GetExtent() / MaxChildBounds;
+					const float ScaleMax = Scale.GetAbsMax();
+					const float ScaleMin = Scale.GetAbsMin();
 					float MaxResolution = ScaleMax * SizeSpecificData.MaxLevelSetResolution;
 					float MinResolution = ScaleMin * SizeSpecificData.MinLevelSetResolution;
 					if (MaxResolution > SizeSpecificData.MaxClusterLevelSetResolution)
@@ -3016,7 +3162,13 @@ void FGeometryCollectionPhysicsProxy::InitializeSharedCollisionStructures(
 						FCollisionStructureManager::NewSimplicial(MassSpaceParticles, *UnionMesh, CollectionImplicits[ClusterTransformIdx].Get(),
 						SharedParams.MaximumCollisionParticleCount));
 				}
-				else
+				else if(SizeSpecificData.ImplicitType == EImplicitTypeEnum::Chaos_Implicit_Capsule)
+				{
+					ensure(false); // unsupported implicit type
+					CollectionImplicits[ClusterTransformIdx].Reset();
+					CollectionSimplicials[ClusterTransformIdx].Reset();
+				}
+				else // Assume it's a union???
 				{
 					CollectionImplicits[ClusterTransformIdx].Reset();	//union so just set as null
 					CollectionSimplicials[ClusterTransformIdx].Reset();
@@ -3028,7 +3180,6 @@ void FGeometryCollectionPhysicsProxy::InitializeSharedCollisionStructures(
 
 		InitRemoveOnFracture(RestCollection, SharedParams);
 	}
-#endif
 }
 
 void FGeometryCollectionPhysicsProxy::InitRemoveOnFracture(FGeometryCollection& RestCollection, const FSharedSimulationParameters& SharedParams)
@@ -3039,15 +3190,12 @@ void FGeometryCollectionPhysicsProxy::InitRemoveOnFracture(FGeometryCollection&
 	}
 
 	// Markup Node Hierarchy Status with FS_RemoveOnFracture flags where geometry is ALL glass
-	int32 NumGeometries = RestCollection.NumElements(FGeometryCollection::GeometryGroup);
-
+	const int32 NumGeometries = RestCollection.NumElements(FGeometryCollection::GeometryGroup);
 	for (int32 Idx = 0; Idx < NumGeometries; Idx++)
 	{
-		int32 TransformIndex = RestCollection.TransformIndex[Idx];
-
-		int32 Start = RestCollection.FaceStart[Idx];
-		int32 End = RestCollection.FaceCount[Idx];
-
+		const int32 TransformIndex = RestCollection.TransformIndex[Idx];
+		const int32 Start = RestCollection.FaceStart[Idx];
+		const int32 End = RestCollection.FaceCount[Idx];
 		bool IsToBeRemoved = true;
 		for (int32 Face = Start; Face < Start + End; Face++)
 		{
@@ -3060,14 +3208,12 @@ void FGeometryCollectionPhysicsProxy::InitRemoveOnFracture(FGeometryCollection&
 					break;
 				}
 			}
-
 			if (!FoundMatch)
 			{
 				IsToBeRemoved = false;
 				break;
 			}
 		}
-
 		if (IsToBeRemoved)
 		{
 			RestCollection.SetFlags(TransformIndex, FGeometryCollection::FS_RemoveOnFracture);
@@ -3076,9 +3222,7 @@ void FGeometryCollectionPhysicsProxy::InitRemoveOnFracture(FGeometryCollection&
 		{
 			RestCollection.ClearFlags(TransformIndex, FGeometryCollection::FS_RemoveOnFracture);
 		}
-
 	}
-
 }
 
 void FGeometryCollectionPhysicsProxy::InitializeKinematics(FParticlesType& Particles, const TManagedArray<int32>& DynamicState)
@@ -3704,6 +3848,14 @@ void FGeometryCollectionPhysicsProxy::UpdateRecordedState(float SolverTime, cons
 	});
 }
 
+void FGeometryCollectionPhysicsProxy::OnRemoveFromSolver(Chaos::FPBDRigidsSolver *RBDSolver)
+{
+	for (auto* Handle : SolverClusterHandles)
+	{
+		RBDSolver->RemoveParticleToProxy(Handle);
+	}
+}
+
 void FGeometryCollectionPhysicsProxy::OnRemoveFromScene()
 {
 #if TODO_REIMPLEMENT_GET_RIGID_PARTICLES
@@ -3774,7 +3926,49 @@ Chaos::FParticleData* FGeometryCollectionPhysicsProxy::NewData()
 
 		TArray<FTransform> Transforms; Transforms.SetNum(NumParticles);
 		Buffer.Transforms = TManagedArray<FTransform>(MoveTemp(Transforms));
+		
 		Buffer.DisabledStates.SetNum(NumParticles);
+
+		Buffer.SpatialIdx.SetNum(NumParticles);
+		Buffer.UniqueIdx.SetNum(NumParticles);
+		Buffer.UserData.SetNum(NumParticles);
+		Buffer.SharedGeometry.SetNum(NumParticles);
+		Buffer.ShapeSimData.SetNum(NumParticles);
+		Buffer.ShapeQueryData.SetNum(NumParticles);
+
+		// To be done once:
+		int32 Idx = 0;
+		for (auto& GTParticle : GTGeometryParticles)
+		{
+			if (GTParticle)
+			{
+				Buffer.DisabledStates[Idx] = false;
+				Buffer.SpatialIdx[Idx] = GTParticle->SpatialIdx();
+				Buffer.UniqueIdx[Idx] = GTParticle->UniqueIdx();
+				Buffer.UserData[Idx] = GTParticle->UserData();
+				Buffer.SharedGeometry[Idx] = GTParticle->GeometrySharedLowLevel();
+
+				for (auto& Shape : GTParticle->ShapesArray())
+				{
+					Buffer.ShapeQueryData[Idx].Add(Shape->QueryData);
+					Buffer.ShapeSimData[Idx].Add(Shape->SimData);
+
+					// Ryan - Hard code everything collides with everything:
+					const int32 Idx1 = Buffer.ShapeSimData[Idx].Num() - 1;
+					Buffer.ShapeSimData[Idx][Idx1].Word1 = 0xFFFF; // this body channel
+					Buffer.ShapeSimData[Idx][Idx1].Word3 = 0xFFFF; // collision candidate channels
+				}
+			}
+			else
+			{
+				Buffer.DisabledStates[Idx] = true;
+				// Buffer.SpatialIdx[Idx] = ???;
+				//Buffer.UniqueIdx[Idx].Idx = INDEX_NONE;
+				//Buffer.UserData[Idx] = nullptr;
+				//Buffer.SharedGeometry[Idx] = nullptr;
+			}
+			++Idx;
+		}
 	}
 
 	const FGeometryDynamicCollection* DynamicCollection = Parameters.DynamicCollection;
@@ -3793,18 +3987,22 @@ Chaos::FParticleData* FGeometryCollectionPhysicsProxy::NewData()
 
 	GeometryCollectionAlgo::GlobalMatrices(DynamicCollection->Transform, DynamicCollection->Parent, Buffer.GlobalTransforms);
 
-/*	for (int32 HandleIdx = 0; HandleIdx < NumParticles; ++HandleIdx)
+	
+	int32 Idx = 0;
+	for (auto& GTParticle : GTGeometryParticles)
 	{
-		if (Chaos::TPBDRigidParticleHandle<float, 3>* Handle = SolverParticleHandles[HandleIdx])
+		if (GTParticle)
 		{
-			const int32 Idx = Buffer.BaseIndex + HandleIdx;
-			Buffer.Transforms[Idx].SetTranslation(Handle->X());
-			Buffer.Transforms[Idx].SetRotation(Handle->R());
-			Buffer.DisabledStates[Idx] = Handle->Disabled();
-			//...
+			if (auto* GTRigidParticle = GTParticle->CastToRigidParticle())
+			{
+				Buffer.DisabledStates[Idx] = GTRigidParticle->Disabled();
+			}
+
+			//const Chaos::TShapesArray<float,3>& ShapesArray = GTParticle->ShapesArray();
 		}
+		++Idx;
 	}
-*/
+	
 	// Make the buffer available to the consumer.
 	GameToPhysInterchange->FlipProducer(); 
 	
@@ -3829,6 +4027,11 @@ void FGeometryCollectionPhysicsProxy::PushToPhysicsState(const Chaos::FParticleD
 	 * the solver
 	 */
 
+	if (InitializedState == ESimulationInitializationState::Unintialized)
+	{
+		return;
+	}
+
 	// The other proxies do a deep copy into dynamic transient memory, passed into
 	// this function via InData.  Instead of that, we're using persistent memory
 	// in a buffer to facilitate communicating between threads.
@@ -3859,16 +4062,49 @@ void FGeometryCollectionPhysicsProxy::PushToPhysicsState(const Chaos::FParticleD
 			Handle->SetX(ParticleTransform.GetTranslation());
 			Handle->SetR(ParticleTransform.GetRotation());
 			Handle->SetDisabled(GState->DisabledStates[Idx]);
-			//...
 
-			if(Handle->Geometry() && Handle->Geometry()->HasBoundingBox())
+			Handle->SetLinearEtherDrag(Parameters.LinearEtherDrag);
+			Handle->SetAngularEtherDrag(Parameters.AngularEtherDrag);
+
+			if (GState->SharedGeometry[Idx] && GState->SharedGeometry[Idx]->HasBoundingBox())
 			{
 				Handle->SetHasBounds(true);
-				Handle->SetLocalBounds(Handle->Geometry()->BoundingBox());
-				Handle->SetWorldSpaceInflatedBounds(Handle->Geometry()->BoundingBox().TransformedAABB(Chaos::TRigidTransform<float, 3>(ParticleTransform)));
+				Handle->SetLocalBounds(GState->SharedGeometry[Idx]->BoundingBox());
+
+				Chaos::TAABB<float, 3> WorldSpaceBox =
+					GState->SharedGeometry[Idx]->BoundingBox().TransformedAABB(
+						Chaos::TRigidTransform<float, 3>(Handle->X(), Handle->R()));
+				// TODO Ryan - WorldSpaceBox.ThickenSymmetrically(GState->V());
+				Handle->SetWorldSpaceInflatedBounds(WorldSpaceBox);
 			}
-		}
-	}
+
+			Handle->SetSpatialIdx(GState->SpatialIdx[Idx]);
+			Handle->SetUniqueIdx(GState->UniqueIdx[Idx]);
+			Handle->SetUserData(GState->UserData[Idx]);
+
+			// TODO Ryan - check dirty flags like TGeometryParticleData::DirtyFlags.
+			// In the mean time, we only set once.
+			if (!Handle->SharedGeometry() && GState->SharedGeometry.IsValidIndex(Idx) && GState->SharedGeometry[Idx])
+			{
+				Handle->SetSharedGeometry(GState->SharedGeometry[Idx]);
+			}
+
+			int32 CurrShape = 0;
+			for (const TUniquePtr<Chaos::TPerShapeData<Chaos::FReal, 3>>& Shape : Handle->ShapesArray())
+			{
+				// TODO Ryan - hardwired for now:
+				Shape->bDisable = false;
+				Shape->CollisionTraceType = Chaos::EChaosCollisionTraceFlag::Chaos_CTF_UseSimpleAndComplex; // Chaos_CTF_UseDefault?
+				if (CurrShape < GState->ShapeSimData[Idx].Num())
+				{
+					Shape->SimData = GState->ShapeSimData[Idx][CurrShape];
+					Shape->QueryData = GState->ShapeQueryData[Idx][CurrShape];
+				}
+				++CurrShape;
+			}
+
+		} // if rigid handle
+	} // for all SolverParticleHandles
 }
 
 void FGeometryCollectionPhysicsProxy::BufferPhysicsResults()
@@ -3891,6 +4127,8 @@ void FGeometryCollectionPhysicsProxy::BufferPhysicsResults()
 	//TManagedArray<TSet<int32>>& ChildrenCache = TargetResults.Children;
 	TManagedArray<int32>& SimulationTypeCache = TargetResults.SimulationType;
 	//TManagedArray<int32>& StatusFlagsCache = TargetResults.StatusFlags;
+	
+	TargetResults.DynamicState.Init(PTDynamicCollection.DynamicState);
 
 //#if TODO_REIMPLEMENT_GET_RIGID_PARTICLES
 
@@ -3962,7 +4200,14 @@ void FGeometryCollectionPhysicsProxy::BufferPhysicsResults()
 		}
 		//TargetResults.DisabledStates.Append(&Particles.DisabledRef(BaseParticleIndex), NumParticles);
 	}
+	const FTransform& ActorToWorld = Parameters.WorldTransform;
+	// It's not worth shrinking the EndFrameUnparentingBuffer array, at least until the solver supports deleting bodies.
+	if (EndFrameUnparentingBuffer.Num() < NumParticles)
 	{
+		EndFrameUnparentingBuffer.Init(INDEX_NONE, NumParticles);
+	}
+
+	{ // tmp scope
 		SCOPE_CYCLE_COUNTER(STAT_CalcParticleToWorld);
 		// Fill particle to world transforms
 		TargetResults.ParticleToWorldTransforms.SetNum(NumParticles);
@@ -3988,47 +4233,150 @@ void FGeometryCollectionPhysicsProxy::BufferPhysicsResults()
 		{
 			//only update roots and first children
 			const int32 ParticleIndex = BaseParticleIndex + TransformIndex;
-//			const int32 ParentIndex = ClusterID[ParticleIndex].Id;
+			//			const int32 ParentIndex = ClusterID[ParticleIndex].Id;
 
-			if(!SolverParticleHandles[ParticleIndex])
+			Chaos::TPBDRigidClusteredParticleHandle<float, 3>* Handle = SolverParticleHandles[ParticleIndex];
+			if (!Handle)
 			{
 				continue;
 			}
 
-			Chaos::TPBDRigidParticleHandle<float, 3>& Handle = *SolverParticleHandles[ParticleIndex];
-
-			if (/*ParentIndex == INDEX_NONE*/true)
+			// Dynamic state is also updated by the solver during field interaction.
+			if (!Handle->Sleeping())
 			{
-				// No need to copy disabled. If we don't guard against disabled, 
-				// proxy particles will override their transform.
-				if (!Handle.Disabled())	
+				const Chaos::EObjectStateType ObjectState = Handle->ObjectState();
+				switch (ObjectState)
 				{
-					FTransform Xf(Handle.R(), Handle.X());
+				case Chaos::EObjectStateType::Kinematic:
+					TargetResults.DynamicState[TransformIndex] = (int)EObjectStateTypeEnum::Chaos_Object_Kinematic;
+					break;
+				case Chaos::EObjectStateType::Static:
+					TargetResults.DynamicState[TransformIndex] = (int)EObjectStateTypeEnum::Chaos_Object_Static;
+					break;
+				case Chaos::EObjectStateType::Sleeping:
+					TargetResults.DynamicState[TransformIndex] = (int)EObjectStateTypeEnum::Chaos_Object_Sleeping;
+					break;
+				case Chaos::EObjectStateType::Dynamic:
+				case Chaos::EObjectStateType::Uninitialized:
+				default:
+					TargetResults.DynamicState[TransformIndex] = (int)EObjectStateTypeEnum::Chaos_Object_Dynamic;
+					break;
+				}
+			}
 
-					TargetResults.ParticleToWorldTransforms[TransformIndex] = Xf;
-					TargetResults.Transforms[TransformIndex] = Xf;
+			// Update the transform and parent hierarchy of the active rigid bodies. Active bodies can be either
+			// rigid geometry defined from the leaf nodes of the collection, or cluster bodies that drive an entire
+			// branch of the hierarchy within the GeometryCollection.
+			// - Active bodies are directly driven from the global position of the corresponding
+			//   rigid bodies within the solver ( cases where RigidBodyID[TransformGroupIndex] is not disabled ). 
+			// - Deactivated bodies are driven from the transforms of their active parents. However the solver can
+			//   take ownership of the parents during the simulation, so it might be necessary to force deactivated
+			//   bodies out of the collections hierarchy during the simulation.  
+			if (!Handle->Disabled())
+			{
+				// Update the transform of the active body. The active body can be either a single rigid
+				// or a collection of rigidly attached geometries (Clustering). The cluster is represented as a
+				// single transform in the GeometryCollection, and all children are stored in the local space
+				// of the parent cluster.
+				// ... When setting cluster transforms it is expected that the MassToLocal is identity.
+				//     Cluster initialization will set the vertices in the MassSpace of the rigid body.
+				// ... When setting individual rigid bodies that are not clustered, the MassToLocal will be 
+				//     non-Identity, and will reflect the difference between the geometric center of the geometry
+				//     and that corresponding rigid bodies center of mass. 
+
+				const FTransform ParticleToWorld(Handle->R(), Handle->X());
+
+#if 0 // TODO Ryan
+				TargetResults.Transforms[TransformIndex] =
+					MassToLocal[TransformIndex].GetRelativeTransformReverse(ParticleToWorld).GetRelativeTransform(ActorToWorld);
+#else
+				TargetResults.ParticleToWorldTransforms[TransformIndex] = ParticleToWorld;
+				TargetResults.Transforms[TransformIndex] = ParticleToWorld;
+#endif
+
+
+#if 0 // TODO Ryan
+				// Force all enabled rigid bodies out of the transform hierarchy
+				if (TargetResults.Parent[TransformIndex] != INDEX_NONE)
+				{
+					// Children in the hierarchy are stored in a TSet, which is not thread safe.  
+					// So we retain indices to remove afterwards.
+					EndFrameUnparentingBuffer[TransformIndex] = TargetResults.Parent[TransformIndex];
 				}
+#endif
+#if 0 // TODO Ryan
+				// When a leaf node rigid body is removed from a cluster, the rigid
+				// body will become active and needs its clusterID updated.  This just
+				// syncs the clusterID all the time.
+				CollectionClusterID[TransformGroupIndex] = ClusterID[RigidBodyIndex].Id;
+#endif
 			}
-/*			else if(ClusterID[ParentIndex].Id == INDEX_NONE)
+			else // Handle->Disabled()
 			{
-				const int32 MultiChildProxyId = MultiChildProxyIdArray[ParticleIndex].Id;
-				const TMultiChildProxyData<float, 3>* ProxyData = MultiChildProxyId == INDEX_NONE ? nullptr : MultiChildProxyDataArray[MultiChildProxyId].Get();
-				if (ProxyData && Particles.Geometry(ParentIndex) && Particles.Geometry(ParentIndex)->IsUnderlyingUnion())	//sq cannot find children without union. If we want levelset support need mapping
+#if 0 // TODO Ryan
+				// The rigid body parent cluster has changed within the solver, and its
+				// parent body is not tracked within the geometry collection. So we need to
+				// pull the rigid bodies out of the transform hierarchy, and just drive
+				// the positions directly from the solvers cluster particle. 
+				if (CollectionClusterID[TransformGroupIndex] != ClusterID[RigidBodyIndex].Id)
 				{
-					//only need to copy out the proxy particle's transform
-					if (ProxyData->KeyChild == ParticleIndex)
+					// Force all driven rigid bodies out of the transform hierarchy
+					if (Parent[TransformGroupIndex] != INDEX_NONE)
 					{
-						TargetResults.ParticleToWorldTransforms[MultiChildProxyId - BaseParticleIndex] = ProxyData->RelativeToKeyChild * ClusterChildToParentMap[ParticleIndex] * FTransform(Particles.R(ParentIndex), Particles.X(ParentIndex));
+						const int32 ParentIndex = Parent[TransformGroupIndex];
+						// Children in the hierarchy are stored in a TSet, which is not thread safe.  So we retain
+						// indices to remove afterwards.
+						EndFrameUnparentingBuffer[TransformGroupIndex] = ParentIndex;
 					}
+					CollectionClusterID[TransformGroupIndex] = ClusterID[RigidBodyIndex].Id;
 				}
-				else
+#endif
+				// Disabled rigid bodies that have valid cluster parents, and have been re-indexed by the
+				// solver ( As in, They were re-clustered outside of the geometry collection), These clusters 
+				// will need to be rendered based on the clusters position. 
+
+				if (TPBDRigidParticleHandle<float, 3>* ParentHandle = Handle->ClusterIds().Id)
 				{
-					TargetResults.ParticleToWorldTransforms[TransformIndex] = ClusterChildToParentMap[ParticleIndex] * FTransform(Particles.R(ParentIndex), Particles.X(ParentIndex));
+					if (TPBDRigidClusteredParticleHandle<float, 3>* ClusteredParentHandle = ParentHandle->CastToClustered())
+					{
+						if (ClusteredParentHandle->InternalCluster())
+						{
+							ensure(false); // TODO Ryan!
+#if 0 // TODO Ryan
+							const FTransform ClusterChildToWorld = ClusterChildToParentMap[RigidBodyIndex] * FTransform(Particles.R(ClusterParentIndex), Particles.X(ClusterParentIndex));
+							if (Parameters.IsCacheRecording())
+							{
+								Particles.X(RigidBodyIndex) = ClusterChildToWorld.GetTranslation();
+								Particles.R(RigidBodyIndex) = ClusterChildToWorld.GetRotation();
+							}
+							// GeomToActor = ActorToWorld.Inv() * ClusterChildToWorld * MassToLocal.Inv();
+							Transform[TransformGroupIndex] = MassToLocal[TransformGroupIndex].GetRelativeTransformReverse(ClusterChildToWorld).GetRelativeTransform(ActorToWorld);
+							Transform[TransformGroupIndex].NormalizeRotation();
+#endif
+						}
+					}
 				}
-			}
-*/
+
+			}// end if
+		} // end for
+	} // tmp scope
+
+#if 0 // TODO Ryan
+	for (int32 TransformGroupIndex = 0; TransformGroupIndex < TransformSize; TransformGroupIndex++)
+	{
+		const int32 ParentIndex = EndFrameUnparentingBuffer[TransformGroupIndex];
+		if (ParentIndex >= 0)
+		{
+			// We reuse EndFrameUnparentingBuffer potentially without reinitialization, so reset this index to -1 before it gets paged out.
+			EndFrameUnparentingBuffer[TransformGroupIndex] = -1;
+
+			Children[ParentIndex].Remove(TransformGroupIndex);
+			Parent[TransformGroupIndex] = INDEX_NONE;
 		}
 	}
+#endif
+
+
 	// We've updated a buffer.  There's no reason to wait for FGeometryCollectionPhysicsProxy::FlipBuffer()
 	// to be called.  Just flip.
 	PhysToGameInterchange->FlipProducer();
@@ -4047,6 +4395,7 @@ void FGeometryCollectionPhysicsProxy::FlipBuffer()
 	//PhysToGameInterchange->FlipProducer();
 }
 
+// Called from FPhysScene_ChaosInterface::SyncBodies(), NOT the solver.
 void FGeometryCollectionPhysicsProxy::PullFromPhysicsState()
 {
 	/**
@@ -4084,8 +4433,15 @@ void FGeometryCollectionPhysicsProxy::PullFromPhysicsState()
 		{
 			if(!TR.DisabledStates[TmIndex])
 			{
+#if 0 // TODO Ryan
+				const FTransform& ParticleToWorld = TR.ParticleToWorldTransforms[TmIndex];
+				GTDynamicCollection->Transform[TmIndex] = 
+					MassToLocal[TmIndex].GetRelativeTransformReverse(ParticleToWorld).GetRelativeTransform(
+						Parameters.WorldTransform);
+#else
 				GTDynamicCollection->Transform[TmIndex] = MassToLocal[TmIndex].GetRelativeTransformReverse(TR.Transforms[TmIndex]).GetRelativeTransform(Parameters.WorldTransform);
 				GTDynamicCollection->Transform[TmIndex].NormalizeRotation();
+#endif
 			}
 		}
 
@@ -4172,10 +4528,14 @@ void FGeometryCollectionPhysicsProxy::UpdateGeometryCollection(FGeometryCollecti
 				//Collection->Transform[TransformGroupIndex];
 
 			// GeomToActor = ActorToWorld.Inv() * ParticleToWorld * MassToLocal.Inv();
-			//Transform[TransformGroupIndex] = MassToLocal[TransformGroupIndex].GetRelativeTransformReverse(ParticleToWorld).GetRelativeTransform(ActorToWorld);
+#if 0 // TODO Ryan
+			Transform[TransformGroupIndex] = 
+				MassToLocal[TransformGroupIndex].GetRelativeTransformReverse(ParticleToWorld).GetRelativeTransform(
+					ActorToWorld);
+#else
 			Transform[TransformGroupIndex] = ParticleToWorld;
 			Transform[TransformGroupIndex].NormalizeRotation();
-
+#endif
 			// dynamic state is also updated by the solver during field interaction. 
 			if (/*!Particles.Sleeping(RigidBodyIndex)*/true)
 			{
diff --git a/Engine/Source/Runtime/Experimental/ChaosSolvers/Public/PBDRigidActiveParticlesBuffer.h b/Engine/Source/Runtime/Experimental/ChaosSolvers/Public/PBDRigidActiveParticlesBuffer.h
--- a/Engine/Source/Runtime/Experimental/ChaosSolvers/Public/PBDRigidActiveParticlesBuffer.h
+++ b/Engine/Source/Runtime/Experimental/ChaosSolvers/Public/PBDRigidActiveParticlesBuffer.h
@@ -15,6 +15,9 @@ namespace Chaos
 	struct CHAOSSOLVERS_API FPBDRigidActiveParticlesBufferOut
 	{
 		TArray<TGeometryParticle<float, 3>*> ActiveGameThreadParticles;
+		// Some particle types (clustered) only exist on the game thread, but we
+		// still need to pull data over via their proxies.
+		TSet<IPhysicsProxyBase*> PhysicsParticleProxies;
 	};
 
 
@@ -51,7 +54,6 @@ namespace Chaos
 
 		// Physics thread to game thread
 		TUniquePtr<IBufferResource<FPBDRigidActiveParticlesBufferOut>> SolverDataOut;
-
 	};
 
 	class FPBDRigidActiveParticlesBufferAccessor
diff --git a/Engine/Source/Runtime/Experimental/ChaosSolvers/Public/PBDRigidsSolver.h b/Engine/Source/Runtime/Experimental/ChaosSolvers/Public/PBDRigidsSolver.h
--- a/Engine/Source/Runtime/Experimental/ChaosSolvers/Public/PBDRigidsSolver.h
+++ b/Engine/Source/Runtime/Experimental/ChaosSolvers/Public/PBDRigidsSolver.h
@@ -280,6 +280,11 @@ namespace Chaos
 		FParticlesType& GetParticles() { return Particles; }
 		const FParticlesType& GetParticles() const { return Particles; }
 
+		void AddParticleToProxy(const Chaos::TGeometryParticleHandle<float, 3>* Particle, IPhysicsProxyBase* Proxy)
+		{ MParticleToProxy.Add(Particle, Proxy); }
+		void RemoveParticleToProxy(const Chaos::TGeometryParticleHandle<float, 3>* Particle)
+		{ MParticleToProxy.Remove(Particle); }
+
 		/**/
 		FEventManager* GetEventManager() { return MEventManager.Get(); }
 
diff --git a/Engine/Source/Runtime/Experimental/ChaosSolvers/Public/PhysicsProxy/GeometryCollectionPhysicsProxy.h b/Engine/Source/Runtime/Experimental/ChaosSolvers/Public/PhysicsProxy/GeometryCollectionPhysicsProxy.h
--- a/Engine/Source/Runtime/Experimental/ChaosSolvers/Public/PhysicsProxy/GeometryCollectionPhysicsProxy.h
+++ b/Engine/Source/Runtime/Experimental/ChaosSolvers/Public/PhysicsProxy/GeometryCollectionPhysicsProxy.h
@@ -8,7 +8,9 @@
 #include "Chaos/Framework/PhysicsProxy.h"
 #include "GeometryCollection/ManagedArray.h"
 #include "GeometryCollection/GeometryCollectionCollisionStructureManager.h"
+#include "Chaos/CollisionFilterData.h"
 #include "Chaos/Framework/BufferedData.h"
+#include "Chaos/GeometryParticlesfwd.h"
 #include "Chaos/ParticleHandle.h"
 #include "Chaos/ParticleHandleFwd.h"
 #include "PBDRigidsSolver.h"
@@ -20,6 +22,10 @@ namespace Chaos
 	template <typename T> struct FClusterCreationParameters;
 }
 
+/**
+ * Buffer structure for communicating simulation state between game and physics 
+ * threads.
+ */
 class FGeometryCollectionResults
 {
 public:
@@ -44,11 +50,20 @@ public:
 	TManagedArray<float> Mass;
 	TManagedArray<FVector> InertiaTensor;
 
+	TManagedArray<int32> ClusterId;
+
 	//TManagedArray<int32> RigidBodyIds;
 	TArray<FTransform> ParticleToWorldTransforms;
 	//TManagedArray<int32> Level;
 	//TManagedArray<int32> StatusFlags;
 
+	TArray<Chaos::FSpatialAccelerationIdx> SpatialIdx;
+	TArray<Chaos::FUniqueIdx> UniqueIdx;
+	TArray<void*> UserData;
+	TArray<TSharedPtr<Chaos::FImplicitObject, ESPMode::ThreadSafe>> SharedGeometry;
+	TArray<TArray<FCollisionFilterData>> ShapeSimData;
+	TArray<TArray<FCollisionFilterData>> ShapeQueryData;
+
 	bool IsObjectDynamic;
 	bool IsObjectLoading;
 
@@ -90,6 +105,7 @@ public:
 	static FName SimplicialsAttribute;
 	static FName ImplicitsAttribute;
 	static FName SolverParticleHandlesAttribute;
+	static FName SolverClusterHandlesAttribute;
 	static FName GTGeometryParticleAttribute;
 
 	typedef FCollisionStructureManager::FSimplicial FSimplicial;
@@ -147,6 +163,7 @@ public:
 	void ActivateBodies();
 
 	void SyncBeforeDestroy();
+	void OnRemoveFromSolver(Chaos::FPBDRigidsSolver *RBDSolver);
 	void OnRemoveFromScene();
 
 	Chaos::FParticleData* NewData();
@@ -163,7 +180,11 @@ public:
 
 	void MergeRecordedTracks(const FRecordedTransformTrack& A, const FRecordedTransformTrack& B, FRecordedTransformTrack& Target);
 	FRecordedFrame& InsertRecordedFrame(FRecordedTransformTrack& InTrack, float InTime);
-	void BuildClusters(uint32 CollectionClusterIndex, const TArray<uint32>& CollectionChildIDs, const TArray<uint32>& ChildIDs, const Chaos::FClusterCreationParameters<float> & Parameters);
+	Chaos::TPBDRigidClusteredParticleHandle<float, 3>* BuildClusters(
+		const uint32 CollectionClusterIndex, 
+		TArray<Chaos::TPBDRigidParticleHandle<float, 3>*>& ChildHandles,
+		const TArray<int32>& ChildTransformGroupIndices,
+		const Chaos::FClusterCreationParameters<float> & Parameters);
 	void BufferCommand(Chaos::FPhysicsSolver* InSolver, const FFieldSystemCommand& Command) { Commands.Add(Command); }
 
 	static void InitializeSharedCollisionStructures(Chaos::FErrorReporter& ErrorReporter, FGeometryCollection& RestCollection, const FSharedSimulationParameters& SharedParams);
@@ -181,7 +202,7 @@ public:
 	int32 GetNumAddedParticles() const { return NumParticles; }
 	int32 GetTransformGroupSize() const { return Parameters.DynamicCollection->NumElements(FGeometryCollection::TransformGroup); }
 
-	TManagedArray<Chaos::TPBDGeometryCollectionParticleHandle<float, 3>*>& GetSolverParticleHandles() { return SolverParticleHandles; }
+	TManagedArray<Chaos::TPBDRigidClusteredParticleHandle<float, 3>*>& GetSolverParticleHandles() { return SolverParticleHandles; }
 	void ClearSolverParticleHandles() { SolverParticleHandles.Fill(nullptr); }
 
 	/*
@@ -261,9 +282,11 @@ private:
 	//TUniquePtr<FGeometryDynamicCollection> PTDynamicCollection;
 	FGeometryDynamicCollection PTDynamicCollection;
 
-	//TManagedArray<Chaos::TPBDRigidParticleHandle<float, 3>*> SolverParticleHandles;
-	TManagedArray<Chaos::TPBDGeometryCollectionParticleHandle<float, 3>*> SolverParticleHandles;
-	TManagedArray<TUniquePtr<Chaos::TGeometryParticle<float, 3>>> GTGeometryParticles; // unique arrays for geo, kin, dyn?
+	TManagedArray<Chaos::TPBDRigidClusteredParticleHandle<float, 3>*> SolverParticleHandles;
+	TManagedArray<Chaos::TPBDRigidClusteredParticleHandle<float, 3>*> SolverClusterHandles;
+	TManagedArray<TUniquePtr<Chaos::TGeometryParticle<float, 3>>> GTGeometryParticles;
+
+	TMap<Chaos::TPBDRigidParticleHandle<float, 3>*, int32> HandleToTransformGroupIndex;
 
 	TManagedArray<FTransform> MassToLocal;
 	TManagedArray<int32> CollisionMask;
@@ -275,6 +298,8 @@ private:
 	TManagedArray<bool> SimulatableParticles;
 	TManagedArray<TUniquePtr<FSimplicial> > Simplicials; // FSimplicial = Chaos::TBVHParticles<float,3>
 	TManagedArray<Chaos::TSerializablePtr<Chaos::FImplicitObject>> Implicits;
+	//TManagedArray<Chaos::TUniquePtr<Chaos::FImplicitObject>> Implicits;
+	//TManagedArray<TSharedPtr<Chaos::FImplicitObject, ESPMode::ThreadSafe>> Implicits;
 	TArray<int32> EndFrameUnparentingBuffer;
 
 	// This is a subset of the geometry group that are used in the transform hierarchy to represent geometry
diff --git a/Engine/Source/Runtime/Experimental/GeometryCollectionCore/Public/GeometryCollection/ManagedArrayTypeValues.inl b/Engine/Source/Runtime/Experimental/GeometryCollectionCore/Public/GeometryCollection/ManagedArrayTypeValues.inl
--- a/Engine/Source/Runtime/Experimental/GeometryCollectionCore/Public/GeometryCollection/ManagedArrayTypeValues.inl
+++ b/Engine/Source/Runtime/Experimental/GeometryCollectionCore/Public/GeometryCollection/ManagedArrayTypeValues.inl
@@ -30,11 +30,13 @@ MANAGED_ARRAY_TYPE(TArray<FVector>*, VectorArrayPointer)
 MANAGED_ARRAY_TYPE(TUniquePtr<TArray<FVector>>, VectorArrayUniquePointer)
 MANAGED_ARRAY_TYPE(Chaos::FImplicitObject3*, FImplicitObject3Pointer)
 MANAGED_ARRAY_TYPE(TUniquePtr<Chaos::FImplicitObject3>, FImplicitObject3UniquePointer)
+MANAGED_ARRAY_TYPE(TSharedPtr<Chaos::FImplicitObject3>, FImplicitObject3SharedPointer)
 MANAGED_ARRAY_TYPE(Chaos::TSerializablePtr<Chaos::FImplicitObject3>, FImplicitObject3SerializablePtr)
 MANAGED_ARRAY_TYPE(Chaos::FBVHParticlesFloat3, FBVHParticlesFloat3Pointer)
 MANAGED_ARRAY_TYPE(TUniquePtr<Chaos::FBVHParticlesFloat3>, FBVHParticlesFloat3UniquePointer)
 MANAGED_ARRAY_TYPE(Chaos::TPBDRigidParticleHandleFloat3*, TPBDRigidParticleHandle3fPtr)
 MANAGED_ARRAY_TYPE(Chaos::TPBDGeometryCollectionParticleHandleFloat3*, TPBDGeometryCollectionParticleHandle3fPtr)
+MANAGED_ARRAY_TYPE(Chaos::TPBDRigidClusteredParticleHandleFloat3*, TPBDRigidClusteredParticleHandle3fPtr)
 MANAGED_ARRAY_TYPE(TUniquePtr<Chaos::TGeometryParticleFloat3>, TGeometryParticle3fUniquePtr)
 
 // NOTE: new types must be added at the bottom to keep serialization from breaking
diff --git a/Engine/Source/Runtime/Experimental/GeometryCollectionEngine/Private/GeometryCollection/GeometryCollectionComponent.cpp b/Engine/Source/Runtime/Experimental/GeometryCollectionEngine/Private/GeometryCollection/GeometryCollectionComponent.cpp
--- a/Engine/Source/Runtime/Experimental/GeometryCollectionEngine/Private/GeometryCollection/GeometryCollectionComponent.cpp
+++ b/Engine/Source/Runtime/Experimental/GeometryCollectionEngine/Private/GeometryCollection/GeometryCollectionComponent.cpp
@@ -109,6 +109,8 @@ UGeometryCollectionComponent::UGeometryCollectionComponent(const FObjectInitiali
 	, ClusterConnectionType(EClusterConnectionTypeEnum::Chaos_PointImplicit)
 	, CollisionGroup(0)
 	, CollisionSampleFraction(1.0)
+	, LinearEtherDrag(0.0)
+	, AngularEtherDrag(0.0)
 	, InitialVelocityType(EInitialVelocityTypeEnum::Chaos_Initial_Velocity_User_Defined)
 	, InitialLinearVelocity(0.f, 0.f, 0.f)
 	, InitialAngularVelocity(0.f, 0.f, 0.f)
@@ -1107,22 +1109,26 @@ void UGeometryCollectionComponent::ResetDynamicCollection()
 
 void UGeometryCollectionComponent::OnCreatePhysicsState()
 {
-#if WITH_PHYSX
+/*#if WITH_PHYSX
 	DummyBodySetup = NewObject<UBodySetup>(this, UBodySetup::StaticClass());
 	DummyBodySetup->AggGeom.BoxElems.Add(FKBoxElem(1.0f));
 	DummyBodyInstance.InitBody(DummyBodySetup, GetComponentToWorld(), this, nullptr);
 	DummyBodyInstance.bNotifyRigidBodyCollision = BodyInstance.bNotifyRigidBodyCollision;
 #endif
-
+*/
 	// Skip the chain - don't care about body instance setup
 	UActorComponent::OnCreatePhysicsState();
 	if (!Simulating) IsObjectLoading = false; // just mark as loaded if we are simulating.
 
-#if WITH_PHYSX
+/*#if WITH_PHYSX
 	DummyBodyInstance.SetCollisionEnabled(ECollisionEnabled::QueryOnly);
 	DummyBodyInstance.SetResponseToAllChannels(ECR_Block);
 #endif
-
+*/
+	// Static mesh uses an init framework that goes through FBodyInstance.  We
+	// do the same thing, but through the geometry collection proxy and lambdas
+	// defined below.  FBodyInstance doesn't work for geometry collections 
+	// because FBodyInstance manages a single particle, where we have many.
 	if (!PhysicsProxy)
 	{
 #if WITH_EDITOR && WITH_EDITORONLY_DATA
@@ -1135,7 +1141,6 @@ void UGeometryCollectionComponent::OnCreatePhysicsState()
 			RestCollectionMutable->EnsureDataIsCooked();
 		}
 #endif
-
 		const bool bValidWorld = GetWorld() && GetWorld()->IsGameWorld();
 		const bool bValidCollection = DynamicCollection && DynamicCollection->Transform.Num() > 0;
 		if (bValidWorld && bValidCollection)
@@ -1147,6 +1152,10 @@ void UGeometryCollectionComponent::OnCreatePhysicsState()
 				ChaosMaterial->SleepingLinearThreshold = PhysicalMaterial->SleepingLinearVelocityThreshold;
 				ChaosMaterial->SleepingAngularThreshold = PhysicalMaterial->SleepingAngularVelocityThreshold;
 			}
+
+			//
+			// Called from FGeometryCollectionPhysicsProxy::Initialize()
+			//
 			auto InitFunc = [this](FSimulationParameters& InParams)
 			{
 #if !(UE_BUILD_SHIPPING || UE_BUILD_TEST)
@@ -1165,6 +1174,8 @@ void UGeometryCollectionComponent::OnCreatePhysicsState()
 				InParams.ClusterConnectionMethod = (Chaos::FClusterCreationParameters<float>::EConnectionMethod)ClusterConnectionType;
 				InParams.CollisionGroup = CollisionGroup;
 				InParams.CollisionSampleFraction = CollisionSampleFraction;
+				InParams.LinearEtherDrag = LinearEtherDrag;
+				InParams.AngularEtherDrag = AngularEtherDrag;
 				InParams.InitialVelocityType = InitialVelocityType;
 				InParams.InitialLinearVelocity = InitialLinearVelocity;
 				InParams.InitialAngularVelocity = InitialAngularVelocity;
diff --git a/Engine/Source/Runtime/Experimental/GeometryCollectionEngine/Public/GeometryCollection/GeometryCollectionComponent.h b/Engine/Source/Runtime/Experimental/GeometryCollectionEngine/Public/GeometryCollection/GeometryCollectionComponent.h
--- a/Engine/Source/Runtime/Experimental/GeometryCollectionEngine/Public/GeometryCollection/GeometryCollectionComponent.h
+++ b/Engine/Source/Runtime/Experimental/GeometryCollectionEngine/Public/GeometryCollection/GeometryCollectionComponent.h
@@ -373,9 +373,7 @@ public:
 	bool Simulating;
 	ESimulationInitializationState InitializationState;
 
-	/*
-	*  ObjectType defines how to initialize the rigid objects state, Kinematic, Sleeping, Dynamic.
-	*/
+	/** ObjectType defines how to initialize the rigid objects state, Kinematic, Sleeping, Dynamic. */
 	UPROPERTY(EditAnywhere, BlueprintReadWrite, Category = "ChaosPhysics|General")
 	EObjectStateTypeEnum ObjectType;
 
@@ -383,64 +381,51 @@ public:
 	UPROPERTY(EditAnywhere, BlueprintReadWrite, Category = "ChaosPhysics|Clustering")
 	bool EnableClustering;
 
-	/**
-	* Maximum level for cluster breaks
-	*/
+	/** Maximum level for cluster breaks. */
 	UPROPERTY(EditAnywhere, BlueprintReadWrite, Category = "ChaosPhysics|Clustering")
 	int32 ClusterGroupIndex;
 
-	/**
-	* Maximum level for cluster breaks
-	*/
+	/** Maximum level for cluster breaks. */
 	UPROPERTY(EditAnywhere, BlueprintReadWrite, Category = "ChaosPhysics|Clustering")
 	int32 MaxClusterLevel;
 
-	/**
-	* Damage threshold for clusters at different levels.
-	*/
+	/** Damage threshold for clusters at different levels. */
 	UPROPERTY(EditAnywhere, BlueprintReadWrite, Category = "ChaosPhysics|Clustering")
 	TArray<float> DamageThreshold;
 
-	/*
-	*  ObjectType defines how to initialize the rigid objects state, Kinematic, Sleeping, Dynamic.
-	*/
+	/** */
 	UPROPERTY(EditAnywhere, BlueprintReadWrite, Category = "ChaosPhysics|Clustering")
 	EClusterConnectionTypeEnum ClusterConnectionType;
 
-	/**
-	* Uniform Friction
-	*/
+	/** */
 	UPROPERTY(EditAnywhere, BlueprintReadWrite, Category = "ChaosPhysics|Collisions")
 	int32 CollisionGroup;
 
-	/**
-	* Uniform Friction
-	*/
+	/** Uniform Friction */
 	UPROPERTY(EditAnywhere, BlueprintReadWrite, Category = "ChaosPhysics|Collisions")
 	float CollisionSampleFraction;
 
+	/** Uniform linear ether drag. */
+	UPROPERTY(EditAnywhere, BlueprintReadWrite, Category = "ChaosPhysics|Ether Drag")
+	float LinearEtherDrag;
 
-	/**
-	* Physical Properties
-	*/
+	/** Uniform angular ether drag. */
+	UPROPERTY(EditAnywhere, BlueprintReadWrite, Category = "ChaosPhysics|Ether Drag")
+	float AngularEtherDrag;
+
+	/** Physical Properties */
 	UPROPERTY(EditAnywhere, BlueprintReadOnly, Category = "ChaosPhysics")
 	const UChaosPhysicalMaterial* PhysicalMaterial;
 
-	/**
-	*
-	*/
+	/** */
 	UPROPERTY(EditAnywhere, BlueprintReadWrite, Category = "ChaosPhysics|Initial Velocity")
 	EInitialVelocityTypeEnum InitialVelocityType;
 
-	/**
-	*
-	*/
+	/** */
 	UPROPERTY(EditAnywhere, BlueprintReadWrite, Category = "ChaosPhysics|Initial Velocity")
 	FVector InitialLinearVelocity;
 
-	/**
-	*
-	*/
+	/** */
 	UPROPERTY(EditAnywhere, BlueprintReadWrite, Category = "ChaosPhysics|Initial Velocity")
 	FVector InitialAngularVelocity;
 
diff --git a/Engine/Source/Runtime/Experimental/GeometryCollectionSimulationCore/Public/GeometryCollection/GeometryCollectionSimulationCoreTypes.h b/Engine/Source/Runtime/Experimental/GeometryCollectionSimulationCore/Public/GeometryCollection/GeometryCollectionSimulationCoreTypes.h
--- a/Engine/Source/Runtime/Experimental/GeometryCollectionSimulationCore/Public/GeometryCollection/GeometryCollectionSimulationCoreTypes.h
+++ b/Engine/Source/Runtime/Experimental/GeometryCollectionSimulationCore/Public/GeometryCollection/GeometryCollectionSimulationCoreTypes.h
@@ -7,6 +7,7 @@
 #include "GeometryCollection/RecordedTransformTrack.h"
 #include "GeometryCollectionSimulationTypes.h"
 #include "Chaos/ClusterCreationParameters.h"
+#include "Chaos/CollisionFilterData.h"
 
 class FGeometryCollection;
 class FGeometryDynamicCollection;
@@ -128,6 +129,9 @@ struct FCollisionDataSimulationParameters
 	bool DoCollisionDataSpatialHash;
 	float CollisionDataSpatialHashRadius;
 	int32 MaxCollisionPerCell;
+
+	FCollisionFilterData QueryData;
+	FCollisionFilterData SimData;
 };
 
 struct FBreakingDataSimulationParameters
@@ -209,6 +213,8 @@ struct FSimulationParameters
 		, ClusterConnectionMethod(Chaos::FClusterCreationParameters<float>::EConnectionMethod::PointImplicit)
 		, CollisionGroup(0)
 		, CollisionSampleFraction(1.0)
+		, LinearEtherDrag(0.0)
+		, AngularEtherDrag(0.0)
 		, InitialVelocityType(EInitialVelocityTypeEnum::Chaos_Initial_Velocity_None)
 		, InitialLinearVelocity(FVector(0))
 		, InitialAngularVelocity(FVector(0))
@@ -235,6 +241,8 @@ struct FSimulationParameters
 		, ClusterConnectionMethod(Other.ClusterConnectionMethod)
 		, CollisionGroup(Other.CollisionGroup)
 		, CollisionSampleFraction(Other.CollisionSampleFraction)
+		, LinearEtherDrag(Other.LinearEtherDrag)
+		, AngularEtherDrag(Other.AngularEtherDrag)
 		, InitialVelocityType(Other.InitialVelocityType)
 		, InitialLinearVelocity(Other.InitialLinearVelocity)
 		, InitialAngularVelocity(Other.InitialAngularVelocity)
@@ -282,6 +290,9 @@ struct FSimulationParameters
 	int32 CollisionGroup;
 	float CollisionSampleFraction;
 
+	float LinearEtherDrag;
+	float AngularEtherDrag;
+
 	EInitialVelocityTypeEnum InitialVelocityType;
 	FVector InitialLinearVelocity;
 	FVector InitialAngularVelocity;
