Merge 11257801 to Release-4.25.

diff --git a/Engine/Source/Runtime/Experimental/Chaos/Private/Chaos/PBDConstraintGraph.cpp b/Engine/Source/Runtime/Experimental/Chaos/Private/Chaos/PBDConstraintGraph.cpp
--- a/Engine/Source/Runtime/Experimental/Chaos/Private/Chaos/PBDConstraintGraph.cpp
+++ b/Engine/Source/Runtime/Experimental/Chaos/Private/Chaos/PBDConstraintGraph.cpp
@@ -166,7 +166,8 @@ void FPBDConstraintGraph::InitializeGraph(const TParticleView<TGeometryParticles
 		{
 			for (auto& Particle : Particles)
 			{
-				ParticleAdd(Particle.Handle());
+				if(!ParticleToNodeIndex.Contains(Particle.Handle()))
+					ParticleAdd(Particle.Handle());
 			}
 		}
 
diff --git a/Engine/Source/Runtime/Experimental/Chaos/Private/Chaos/PBDRigidClustering.cpp b/Engine/Source/Runtime/Experimental/Chaos/Private/Chaos/PBDRigidClustering.cpp
--- a/Engine/Source/Runtime/Experimental/Chaos/Private/Chaos/PBDRigidClustering.cpp
+++ b/Engine/Source/Runtime/Experimental/Chaos/Private/Chaos/PBDRigidClustering.cpp
@@ -437,10 +437,6 @@ namespace Chaos
 			//NewParticle = MEvolution.CreateClusteredParticles(1)[0]; // calls Evolution.DirtyParticle()
 			NewParticle = MEvolution.CreateClusteredParticles(1)[0]; // calls Evolution.DirtyParticle()
 		}
-		if (!NewParticle->UniqueIdx().IsValid())
-		{
-			NewParticle->SetUniqueIdx(MEvolution.GenerateUniqueIdx());
-		}
 
 		// Must do this so that the constraint graph knows about this particle 
 		// prior to calling CreateIslands().  We could call MEvolution.CreateParticle()
@@ -532,10 +528,6 @@ namespace Chaos
 		{
 			NewParticle = MEvolution.CreateClusteredParticles(1)[0]; // calls Evolution.DirtyParticle()
 		}
-		if (!NewParticle->UniqueIdx().IsValid())
-		{
-			NewParticle->SetUniqueIdx(MEvolution.GenerateUniqueIdx());
-		}
 		MEvolution.CreateParticle(NewParticle);
 		MEvolution.EnableParticle(NewParticle, Parent);
 
@@ -860,13 +852,6 @@ namespace Chaos
 					}
 					TArray<Chaos::TPBDRigidClusteredParticleHandle<float,3>*> NewClusterHandles = 
 						MEvolution.CreateClusteredParticles(NumNewClusters);
-					for (auto& Handle : NewClusterHandles)
-					{
-						if (!Handle->UniqueIdx().IsValid())
-						{
-							Handle->SetUniqueIdx(MEvolution.GenerateUniqueIdx());
-						}
-					}
 					int32 ClusterHandlesIdx = 0;
 					for (TArray<TPBDRigidParticleHandle<T,d>*>& ConnectedPieces : ConnectedPiecesArray)
 					{
diff --git a/Engine/Source/Runtime/Experimental/ChaosSolvers/Private/PhysicsProxy/GeometryCollectionPhysicsProxy.cpp b/Engine/Source/Runtime/Experimental/ChaosSolvers/Private/PhysicsProxy/GeometryCollectionPhysicsProxy.cpp
--- a/Engine/Source/Runtime/Experimental/ChaosSolvers/Private/PhysicsProxy/GeometryCollectionPhysicsProxy.cpp
+++ b/Engine/Source/Runtime/Experimental/ChaosSolvers/Private/PhysicsProxy/GeometryCollectionPhysicsProxy.cpp
@@ -587,7 +587,6 @@ void FGeometryCollectionPhysicsProxy::InitializeBodiesGT()
 				GTGeometryParticles[Idx] = MoveTemp(Particle); // Particle is now undefined
 				Chaos::TPBDRigidParticle<float, 3>* GTParticle = 
 					GTGeometryParticles[Idx]->CastToRigidParticle();
-				GTParticle->SetUniqueIdx(GetSolver()->GetEvolution()->GenerateUniqueIdx(), false);
 
 				FTransform ParticleTransform = MassToLocal[Idx] * Transform[Idx] * Parameters.WorldTransform;
 				if (GTParticle->ObjectState() == Chaos::EObjectStateType::Dynamic)
@@ -631,6 +630,16 @@ void FGeometryCollectionPhysicsProxy::InitializeBodiesGT()
 					}
 				}
 			}
+/*			else
+			{
+				// This is likely a cluster parent.  Create a base particle, until 
+				// we have a reason to include more data.
+
+				TUniquePtr<Chaos::TGeometryParticle<float, 3>> Particle = 
+					Chaos::TGeometryParticle<float, 3>::CreateParticle();
+				GTGeometryParticles[Idx] = MoveTemp(Particle); // Particle is now undefined
+			}
+*/
 		}
 
 		InitializedState = Parameters.InitializationState;
@@ -683,12 +692,7 @@ void FGeometryCollectionPhysicsProxy::InitializeBodiesPT(
 		// Add entries into simulation array
 		RigidsSolver->GetEvolution()->ReserveParticles(NumSimulatedParticles);
 		TArray<Chaos::TPBDGeometryCollectionParticleHandle<float, 3>*> Handles = 
-			Particles.CreateGeometryCollectionParticles(NumLeafNodes);
-		for (auto& Handle : Handles)
-		{
-			Handle->SetUniqueIdx(GetSolver()->GetEvolution()->GenerateUniqueIdx());
-		}
-
+			RigidsSolver->GetEvolution()->CreateGeometryCollectionParticles(NumLeafNodes);
 
 		int32 NextIdx = 0;
 		for(int32 Idx=0; Idx < SimulatableParticles.Num(); ++Idx)
@@ -914,6 +918,10 @@ void FGeometryCollectionPhysicsProxy::InitializeBodiesPT(
 							RigidChildrenTransformGroupIndex, 
 							CreationParameters);
 
+					Chaos::TGeometryParticle<float, 3>* GTParticle = 
+						GTGeometryParticles[TransformGroupIndex].Get();
+					Handle->GTGeometryParticle() = GTParticle;
+
 					SolverClusterHandles[TransformGroupIndex] = Handle;
 					SolverParticleHandles[TransformGroupIndex] = Handle;
 					HandleToTransformGroupIndex.Add(Handle, TransformGroupIndex);
@@ -2500,10 +2508,6 @@ FGeometryCollectionPhysicsProxy::BuildClusters(
 			ClusterCreationParameters,
 			Implicits[CollectionClusterIndex],
 			&ParticleTM);
-	if (!Parent->UniqueIdx().IsValid())
-	{
-		Parent->SetUniqueIdx(GetSolver()->GetEvolution()->GenerateUniqueIdx());
-	}
 
 	if (ReportNoLevelsetCluster && 
 		Parent->DynamicGeometry())
@@ -3929,9 +3933,6 @@ Chaos::FParticleData* FGeometryCollectionPhysicsProxy::NewData()
 		
 		Buffer.DisabledStates.SetNum(NumParticles);
 
-		Buffer.SpatialIdx.SetNum(NumParticles);
-		Buffer.UniqueIdx.SetNum(NumParticles);
-		Buffer.UserData.SetNum(NumParticles);
 		Buffer.SharedGeometry.SetNum(NumParticles);
 		Buffer.ShapeSimData.SetNum(NumParticles);
 		Buffer.ShapeQueryData.SetNum(NumParticles);
@@ -3943,9 +3944,6 @@ Chaos::FParticleData* FGeometryCollectionPhysicsProxy::NewData()
 			if (GTParticle)
 			{
 				Buffer.DisabledStates[Idx] = false;
-				Buffer.SpatialIdx[Idx] = GTParticle->SpatialIdx();
-				Buffer.UniqueIdx[Idx] = GTParticle->UniqueIdx();
-				Buffer.UserData[Idx] = GTParticle->UserData();
 				Buffer.SharedGeometry[Idx] = GTParticle->GeometrySharedLowLevel();
 
 				for (auto& Shape : GTParticle->ShapesArray())
@@ -4078,10 +4076,6 @@ void FGeometryCollectionPhysicsProxy::PushToPhysicsState(const Chaos::FParticleD
 				Handle->SetWorldSpaceInflatedBounds(WorldSpaceBox);
 			}
 
-			Handle->SetSpatialIdx(GState->SpatialIdx[Idx]);
-			Handle->SetUniqueIdx(GState->UniqueIdx[Idx]);
-			Handle->SetUserData(GState->UserData[Idx]);
-
 			// TODO Ryan - check dirty flags like TGeometryParticleData::DirtyFlags.
 			// In the mean time, we only set once.
 			if (!Handle->SharedGeometry() && GState->SharedGeometry.IsValidIndex(Idx) && GState->SharedGeometry[Idx])
diff --git a/Engine/Source/Runtime/Experimental/ChaosSolvers/Public/PhysicsProxy/GeometryCollectionPhysicsProxy.h b/Engine/Source/Runtime/Experimental/ChaosSolvers/Public/PhysicsProxy/GeometryCollectionPhysicsProxy.h
--- a/Engine/Source/Runtime/Experimental/ChaosSolvers/Public/PhysicsProxy/GeometryCollectionPhysicsProxy.h
+++ b/Engine/Source/Runtime/Experimental/ChaosSolvers/Public/PhysicsProxy/GeometryCollectionPhysicsProxy.h
@@ -52,14 +52,8 @@ public:
 
 	TManagedArray<int32> ClusterId;
 
-	//TManagedArray<int32> RigidBodyIds;
 	TArray<FTransform> ParticleToWorldTransforms;
-	//TManagedArray<int32> Level;
-	//TManagedArray<int32> StatusFlags;
 
-	TArray<Chaos::FSpatialAccelerationIdx> SpatialIdx;
-	TArray<Chaos::FUniqueIdx> UniqueIdx;
-	TArray<void*> UserData;
 	TArray<TSharedPtr<Chaos::FImplicitObject, ESPMode::ThreadSafe>> SharedGeometry;
 	TArray<TArray<FCollisionFilterData>> ShapeSimData;
 	TArray<TArray<FCollisionFilterData>> ShapeQueryData;
