Merge 11289630 to Release-4.25.

diff --git a/Engine/Source/Runtime/Experimental/Chaos/Private/Chaos/PBDRigidClustering.cpp b/Engine/Source/Runtime/Experimental/Chaos/Private/Chaos/PBDRigidClustering.cpp
--- a/Engine/Source/Runtime/Experimental/Chaos/Private/Chaos/PBDRigidClustering.cpp
+++ b/Engine/Source/Runtime/Experimental/Chaos/Private/Chaos/PBDRigidClustering.cpp
@@ -426,7 +426,7 @@ namespace Chaos
 		const int32 ClusterGroupIndex,
 		TArray<Chaos::TPBDRigidParticleHandle<T,d>*>&& Children,
 		const FClusterCreationParameters<T>& Parameters,
-		TSerializablePtr<FImplicitObject> ProxyGeometry,
+		TSharedPtr<Chaos::FImplicitObject, ESPMode::ThreadSafe> ProxyGeometry,
 		const TRigidTransform<T, d>* ForceMassOrientation)
 	{
 		SCOPE_CYCLE_COUNTER(STAT_CreateClusterParticle);
@@ -582,7 +582,7 @@ namespace Chaos
 
 		TSet<TPBDRigidParticleHandle<T, d>*> ChildrenSet(ChildrenArray);
 		UpdateMassProperties(NewParticle, ChildrenSet, nullptr);
-		UpdateGeometry(NewParticle, ChildrenSet, TSerializablePtr<FImplicitObject>(), NoCleanParams);
+		UpdateGeometry(NewParticle, ChildrenSet, nullptr, NoCleanParams);
 
 		return NewParticle;
 	}
@@ -1332,7 +1332,7 @@ namespace Chaos
 	void TPBDRigidClustering<FPBDRigidsEvolution, FPBDCollisionConstraint, T, d>::UpdateGeometry(
 		Chaos::TPBDRigidClusteredParticleHandle<float, 3>* Parent, 
 		const TSet<TPBDRigidParticleHandle<T, d>*>& Children, 
-		TSerializablePtr<FImplicitObject> ProxyGeometry, 
+		TSharedPtr<Chaos::FImplicitObject, ESPMode::ThreadSafe> ProxyGeometry,
 		const FClusterCreationParameters<T>& Parameters)
 	{
 		SCOPE_CYCLE_COUNTER(STAT_UpdateGeometry);
@@ -1446,7 +1446,7 @@ namespace Chaos
 		{
 			//ensureMsgf(false, TEXT("Checking usage with proxy"));
 			//@coverage {production}
-			Parent->SetGeometry(ProxyGeometry);
+			Parent->SetSharedGeometry(ProxyGeometry);
 		}
 		else if (Objects.Num() == 0)
 		{
diff --git a/Engine/Source/Runtime/Experimental/Chaos/Public/Chaos/ImplicitObject.h b/Engine/Source/Runtime/Experimental/Chaos/Public/Chaos/ImplicitObject.h
--- a/Engine/Source/Runtime/Experimental/Chaos/Public/Chaos/ImplicitObject.h
+++ b/Engine/Source/Runtime/Experimental/Chaos/Public/Chaos/ImplicitObject.h
@@ -327,4 +327,5 @@ FORCEINLINE FArchive& operator<<(FArchive& Ar, FImplicitObject& Value)
 }
 
 typedef FImplicitObject FImplicitObject3;
+typedef TSharedPtr<Chaos::FImplicitObject, ESPMode::ThreadSafe> ThreadSafeSharedPtr_FImplicitObject;
 }
\ No newline at end of file
diff --git a/Engine/Source/Runtime/Experimental/Chaos/Public/Chaos/PBDRigidClustering.h b/Engine/Source/Runtime/Experimental/Chaos/Public/Chaos/PBDRigidClustering.h
--- a/Engine/Source/Runtime/Experimental/Chaos/Public/Chaos/PBDRigidClustering.h
+++ b/Engine/Source/Runtime/Experimental/Chaos/Public/Chaos/PBDRigidClustering.h
@@ -62,7 +62,7 @@ public:
 		const int32 ClusterGroupIndex, 
 		TArray<Chaos::TPBDRigidParticleHandle<T,d>*>&& Children, 
 		const FClusterCreationParameters<T>& Parameters/* = FClusterCreationParameters<T>()*/,
-		TSerializablePtr<FImplicitObject> ProxyGeometry = TSerializablePtr<FImplicitObject>(),
+		TSharedPtr<Chaos::FImplicitObject, ESPMode::ThreadSafe> ProxyGeometry = nullptr,
 		const TRigidTransform<T, d>* ForceMassOrientation = nullptr);
 
 	/**
@@ -285,7 +285,7 @@ public:
 	void UpdateGeometry(
 		Chaos::TPBDRigidClusteredParticleHandle<float, 3>* Parent, 
 		const TSet<TPBDRigidParticleHandle<T, d>*>& Children, 
-		TSerializablePtr<FImplicitObject> ProxyGeometry, 
+		TSharedPtr<Chaos::FImplicitObject, ESPMode::ThreadSafe> ProxyGeometry,
 		const FClusterCreationParameters<T>& Parameters);
 
 	void ComputeStrainFromCollision(const FPBDCollisionConstraint& CollisionRule);
diff --git a/Engine/Source/Runtime/Experimental/ChaosSolvers/Private/PhysicsProxy/GeometryCollectionPhysicsProxy.cpp b/Engine/Source/Runtime/Experimental/ChaosSolvers/Private/PhysicsProxy/GeometryCollectionPhysicsProxy.cpp
--- a/Engine/Source/Runtime/Experimental/ChaosSolvers/Private/PhysicsProxy/GeometryCollectionPhysicsProxy.cpp
+++ b/Engine/Source/Runtime/Experimental/ChaosSolvers/Private/PhysicsProxy/GeometryCollectionPhysicsProxy.cpp
@@ -45,6 +45,7 @@ FAutoConsoleVariableRef CVarCollisionParticlesPerObjectFractionDefault(TEXT("p.C
 
 FName FGeometryCollectionPhysicsProxy::SimplicialsAttribute("CollisionParticles");
 FName FGeometryCollectionPhysicsProxy::ImplicitsAttribute("Implicits");
+FName FGeometryCollectionPhysicsProxy::SharedImplicitsAttribute("SharedImplicits");
 FName FGeometryCollectionPhysicsProxy::SolverParticleHandlesAttribute("SolverParticleHandles");
 FName FGeometryCollectionPhysicsProxy::SolverClusterHandlesAttribute("SolverClusterHandles");
 FName FGeometryCollectionPhysicsProxy::GTGeometryParticleAttribute("GTGeometryParticles");
@@ -200,7 +201,8 @@ TArray<int32> ComputeRecursiveOrder(const FGeometryCollection& Collection)
 				CurCluster.State = FClusterProcessing::VisitingChildren;
 				ClustersToProcess.Add(CurCluster);
 
-				for(int32 ChildIdx : Children[ClusterTransformIdx])	//order of children doesn't matter as long as all children appear before parent
+				//order of children doesn't matter as long as all children appear before parent
+				for(int32 ChildIdx : Children[ClusterTransformIdx])
 				{
 					ClustersToProcess.Emplace(ChildIdx);
 				}
@@ -220,7 +222,7 @@ void PopulateSimulatedParticle(
 	Chaos::TPBDRigidParticleHandle<float,3>* Handle,
 	const FSharedSimulationParameters& SharedParams,
 	const FCollisionStructureManager::FSimplicial* Simplicial,
-	Chaos::TSerializablePtr<Chaos::FImplicitObject> Implicit,
+	TSharedPtr<Chaos::FImplicitObject, ESPMode::ThreadSafe> Implicit,
 	const float MassIn,
 	const FVector& InertiaTensorVec,
 	const FTransform& WorldTransform, 
@@ -291,14 +293,11 @@ void PopulateSimulatedParticle(
 	//Particles.CollisionGroup(RigidBodyIndex) = CollisionGroup;
 
 	{
-//#if TODO_REIMPLEMENT_RIGID_CLUSTERING
 		// @important (The solver can not free this memory)
 		if (Implicit)	//todo(ocohen): this is only needed for cases where clusters have no proxy. Kind of gross though, should refactor
 		{
-			Handle->SetGeometry(Implicit);
-			//Particles.SetGeometry(RigidBodyIndex, Implicit);
+			Handle->SetSharedGeometry(Implicit);
 		}
-//#endif
 
 		Handle->CollisionParticlesInitIfNeeded();
 		//Particles.CollisionParticlesInitIfNeeded(RigidBodyIndex);
@@ -609,7 +608,7 @@ void FGeometryCollectionPhysicsProxy::InitializeBodiesGT()
 
 					if (Implicits[Idx])
 					{
-						GTParticle->SetGeometry(Implicits[Idx]->Copy());
+						GTParticle->SetGeometry(Implicits[Idx]); // shared pointer version
 					}
 					if (Simplicials[Idx])
 					{
@@ -2461,19 +2460,11 @@ FGeometryCollectionPhysicsProxy::BuildClusters(
 	check(ChildHandles.Num() != 0);
 
 	check(Parameters.DynamicCollection);
-	FGeometryDynamicCollection* Collection = Parameters.DynamicCollection;
-	TManagedArray<int32>& DynamicState = Collection->DynamicState;
-	TManagedArray<int32>& CollisionGroup = Collection->CollisionGroup;
-	TManagedArray<int32>& ParentIndex = Collection->Parent;
-	TManagedArray<TSet<int32>>& Children = Collection->Children;
-	TManagedArray<FTransform>& Transform = Collection->Transform;
-
-	Chaos::FPhysicsSolver* ThisSolver = GetSolver();
-
-	//todo(ocohen): cache this
-	check(Parameters.RestCollection);
-	const TManagedArray<float>& Mass = Parameters.RestCollection->GetAttribute<float>("Mass", FTransformCollection::TransformGroup);
-	const TManagedArray<FVector>& InertiaTensor = Parameters.RestCollection->GetAttribute<FVector>("InertiaTensor", FTransformCollection::TransformGroup);
+	FGeometryDynamicCollection* DynamicCollection = Parameters.DynamicCollection;
+	TManagedArray<int32>& DynamicState = DynamicCollection->DynamicState;
+	TManagedArray<int32>& ParentIndex = DynamicCollection->Parent;
+	TManagedArray<TSet<int32>>& Children = DynamicCollection->Children;
+	TManagedArray<FTransform>& Transform = DynamicCollection->Transform;
 
 	//If we are a root particle use the world transform, otherwise set the relative transform
 	const FTransform CollectionSpaceTransform = 
@@ -2506,7 +2497,7 @@ FGeometryCollectionPhysicsProxy::BuildClusters(
 			Parameters.ClusterGroupIndex, 
 			MoveTemp(ChildHandlesCopy),
 			ClusterCreationParameters,
-			Implicits[CollectionClusterIndex],
+			Implicits[CollectionClusterIndex], // union from children if null
 			&ParticleTM);
 
 	if (ReportNoLevelsetCluster && 
@@ -2516,10 +2507,6 @@ FGeometryCollectionPhysicsProxy::BuildClusters(
 		UE_LOG(LogChaos, Warning, TEXT("Union object generated for cluster:%s"), *Parameters.Name);
 	}
 
-	GetSolver()->GetEvolution()->SetPhysicsMaterial(
-		Parent, 
-		GetSolver()->GetEvolution()->GetPhysicsMaterial(ChildHandles[0]));
-
 	if(Parameters.InitialVelocityType == EInitialVelocityTypeEnum::Chaos_Initial_Velocity_User_Defined)
 	{
 		InitialLinearVelocity[CollectionClusterIndex] = Parameters.InitialLinearVelocity;
@@ -2541,11 +2528,17 @@ FGeometryCollectionPhysicsProxy::BuildClusters(
 	//const FTransform ParticleTM(Particles.R(NewSolverClusterID), Particles.X(NewSolverClusterID));	//in theory we should be computing this and passing in to avoid inertia computation at runtime. If we do this we must account for leaf particles that have already been created in world space
 	MassToLocal[CollectionClusterIndex] = FTransform::Identity;
 
+	check(Parameters.RestCollection);
+	const TManagedArray<float>& Mass = 
+		Parameters.RestCollection->GetAttribute<float>("Mass", FTransformCollection::TransformGroup);
+	const TManagedArray<FVector>& InertiaTensor = 
+		Parameters.RestCollection->GetAttribute<FVector>("InertiaTensor", FTransformCollection::TransformGroup);
+
 	PopulateSimulatedParticle(
 		Parent,
 		Parameters.Shared, 
 		nullptr, // CollisionParticles is optionally set from CreateClusterParticle()
-		Chaos::TSerializablePtr<Chaos::FImplicitObject>(), // Parent->Geometry() ? Parent->Geometry() : Implicits[CollectionClusterIndex], 
+		nullptr, // Parent->Geometry() ? Parent->Geometry() : Implicits[CollectionClusterIndex], 
 #if TODO_REIMPLEMENT_SOLVER_SETTINGS_ACCESSORS
 		// Note: if the children had this scaling baked in, then cluster parents
 		// will too, and we shouldn't apply the scaling to the parents twice!
@@ -2564,9 +2557,9 @@ FGeometryCollectionPhysicsProxy::BuildClusters(
 	SolverClusterHandles[CollectionClusterIndex] = Parent;
 
 	const int32 NumThresholds = Parameters.DamageThreshold.Num();
-	const int32 Level = FMath::Clamp(CalculateHierarchyLevel(Collection, CollectionClusterIndex), 0, INT_MAX);
-	const float Default = NumThresholds > 0 ? Parameters.DamageThreshold[NumThresholds - 1] : 0.f;
-	float Damage = Level < NumThresholds ? Parameters.DamageThreshold[Level] : Default;
+	const int32 Level = FMath::Clamp(CalculateHierarchyLevel(DynamicCollection, CollectionClusterIndex), 0, INT_MAX);
+	const float DefaultDamage = NumThresholds > 0 ? Parameters.DamageThreshold[NumThresholds - 1] : 0.f;
+	float Damage = Level < NumThresholds ? Parameters.DamageThreshold[Level] : DefaultDamage;
 	if(Level >= Parameters.MaxClusterLevel) 
 		Damage = FLT_MAX;
 	Parent->SetStrains(Damage);
@@ -2574,24 +2567,20 @@ FGeometryCollectionPhysicsProxy::BuildClusters(
 	GetSolver()->GetEvolution()->SetPhysicsMaterial(Parent, Parameters.PhysicalMaterial);
 
 	const FTransform ParentTransform = 
-		GeometryCollectionAlgo::GlobalMatrix(Collection->Transform, Collection->Parent, CollectionClusterIndex);
+		GeometryCollectionAlgo::GlobalMatrix(DynamicCollection->Transform, DynamicCollection->Parent, CollectionClusterIndex);
 
 	int32 MinCollisionGroup = INT_MAX;
 	for(int32 Idx=0; Idx < ChildHandles.Num(); Idx++)
 	{
-		const int32 ChildTransformGroupIndex = ChildTransformGroupIndices[Idx];
 		Chaos::TPBDRigidParticleHandle<float, 3>* Child = ChildHandles[Idx];
-
 		if (Chaos::TPBDRigidClusteredParticleHandle<float, 3>* ClusteredChild = Child->CastToClustered())
 		{
-			//SolverStrainArray[ChildIDs[idx]]= Damage;
 			ClusteredChild->SetStrains(Damage);
 		}
 
+		const int32 ChildTransformGroupIndex = ChildTransformGroupIndices[Idx];
 		SolverClusterHandles[ChildTransformGroupIndex] = Parent;
 
-		MinCollisionGroup = FMath::Min(Child->CollisionGroup(), MinCollisionGroup);
-
 		const FTransform ChildTransform(Child->R(), Child->X());
 		if(Children[ChildTransformGroupIndex].Num()) // clustered local transform
 		{
@@ -2604,6 +2593,8 @@ FGeometryCollectionPhysicsProxy::BuildClusters(
 			Transform[ChildTransformGroupIndex] = RestTransform.GetRelativeTransform(ParticleTM);
 		}
 		Transform[ChildTransformGroupIndex].NormalizeRotation();
+
+		MinCollisionGroup = FMath::Min(Child->CollisionGroup(), MinCollisionGroup);
 	}
 	Parent->SetCollisionGroup(MinCollisionGroup);
 	return Parent;
@@ -2645,17 +2636,53 @@ void FGeometryCollectionPhysicsProxy::CreateDynamicAttributes()
 		DynamicCollection->AddExternalAttribute("InitialLinearVelocity", FTransformCollection::TransformGroup, InitialLinearVelocity);
 	}
 
-	const auto& RestImplicits = Parameters.RestCollection->GetAttribute<TUniquePtr<Chaos::FImplicitObject>>(
-		ImplicitsAttribute, FTransformCollection::TransformGroup);
+	// The RestCollection may have been authored with TUniquePtr<FImplicitObject> under the name 
+	// ImplicitsAttribute, as they originally were.  As the particle system now uses TSharedPtr,
+	// we're now authoring TSharedPtr<FImplicitObject> under the name SharedImplicitsAttribute.
+	FGeometryCollection& NCRestCollection = const_cast<FGeometryCollection&>(*Parameters.RestCollection);
+	TManagedArray<TUniquePtr<Chaos::FImplicitObject>>* UniqueRestImplicits = 
+		NCRestCollection.FindAttribute<TUniquePtr<Chaos::FImplicitObject>>(
+			ImplicitsAttribute, FTransformCollection::TransformGroup);
+	TManagedArray<TSharedPtr<Chaos::FImplicitObject, ESPMode::ThreadSafe>>* SharedRestImplicits = 
+		NCRestCollection.FindAttribute<TSharedPtr<Chaos::FImplicitObject, ESPMode::ThreadSafe>>(
+			SharedImplicitsAttribute, FTransformCollection::TransformGroup);
+	
+	// If we have unique pointers to geometry, we need to convert them to shared pointers.
+	TManagedArray<TSharedPtr<Chaos::FImplicitObject, ESPMode::ThreadSafe>>* RestCollectionSharedImplicits = nullptr;
+	if (!SharedRestImplicits)
+	{
+		RestCollectionSharedImplicits = 
+			&NCRestCollection.AddAttribute<TSharedPtr<Chaos::FImplicitObject, ESPMode::ThreadSafe>>(
+				SharedImplicitsAttribute, FTransformCollection::TransformGroup);
+	}
+
 	for (int32 Index = DynamicCollection->NumElements(FTransformCollection::TransformGroup) - 1; 0 <= Index; Index--)
 	{
 		Simplicials[Index] = TUniquePtr<FSimplicial>(nullptr);
-		Implicits[Index] = Chaos::MakeSerializable(RestImplicits[Index]);
+
+		//Implicits[Index] = Chaos::MakeSerializable(RestImplicits[Index]);
+		if (SharedRestImplicits)
+		{
+			Implicits[Index] = (*SharedRestImplicits)[Index];
+		}
+		else if(UniqueRestImplicits)
+		{
+			// Steal the memory from the unique ptr.
+			Implicits[Index] = TSharedPtr<Chaos::FImplicitObject, ESPMode::ThreadSafe>((*UniqueRestImplicits)[Index].Release());
+			// Update the rest collection with a copy of the shared pointer.
+			(*RestCollectionSharedImplicits)[Index] = Implicits[Index];
+		}
 
 		RigidBodyID[Index] = INDEX_NONE;
 		SolverClusterID[Index] = INDEX_NONE;
 	}
 
+	// Clean up now invalid unique pointers.
+	if (RestCollectionSharedImplicits)
+	{
+		NCRestCollection.RemoveAttribute(ImplicitsAttribute, FTransformCollection::TransformGroup);
+	}
+
 	if (Parameters.RestCollection->HasAttribute(SimplicialsAttribute, FTransformCollection::TransformGroup))
 	{
 		const auto& RestSimplicials = Parameters.RestCollection->GetAttribute<TUniquePtr<FSimplicial>>(SimplicialsAttribute, FTransformCollection::TransformGroup);
@@ -2665,7 +2692,6 @@ void FGeometryCollectionPhysicsProxy::CreateDynamicAttributes()
 		}
 	}
 
-
 	// Merge Active Flags.
 	const int32 NumTransforms = SimulatableParticles.Num();
 	if(Parameters.RestCollection->HasAttribute(FGeometryCollection::SimulatableParticlesAttribute, FTransformCollection::TransformGroup))
@@ -2761,9 +2787,9 @@ void FGeometryCollectionPhysicsProxy::InitializeSharedCollisionStructures(
 	TManagedArray<TUniquePtr<FSimplicial>>& CollectionSimplicials = 
 		RestCollection.AddAttribute<TUniquePtr<FSimplicial>>(
 			SimplicialsAttribute, FTransformCollection::TransformGroup);
-	TManagedArray<TUniquePtr<FImplicitObject>>& CollectionImplicits = 
-		RestCollection.AddAttribute<TUniquePtr<FImplicitObject>>(
-			ImplicitsAttribute, FTransformCollection::TransformGroup);
+	TManagedArray<TSharedPtr<Chaos::FImplicitObject, ESPMode::ThreadSafe>>& CollectionImplicits = 
+		RestCollection.AddAttribute<TSharedPtr<Chaos::FImplicitObject, ESPMode::ThreadSafe>>(
+			SharedImplicitsAttribute, FTransformCollection::TransformGroup);
 
 	FTransform IdentityXf(FQuat::Identity, FVector(0));
 	IdentityXf.NormalizeRotation();
@@ -2959,7 +2985,7 @@ void FGeometryCollectionPhysicsProxy::InitializeSharedCollisionStructures(
 
 
 			ErrorReporter.SetPrefix(BaseErrorPrefix + " | Transform Index: " + FString::FromInt(TransformGroupIndex));
-			CollectionImplicits[TransformGroupIndex] = TUniquePtr<FImplicitObject>(
+			CollectionImplicits[TransformGroupIndex] = TSharedPtr<Chaos::FImplicitObject, ESPMode::ThreadSafe>(
 				FCollisionStructureManager::NewImplicit(ErrorReporter, MassSpaceParticles, *TriMesh, InstanceBoundingBox,
 					InnerRadius[GeometryIndex], SizeSpecificData.MinLevelSetResolution, SizeSpecificData.MaxLevelSetResolution,
 					SizeSpecificData.CollisionObjectReductionPercentage, SizeSpecificData.CollisionType,
@@ -3007,7 +3033,7 @@ void FGeometryCollectionPhysicsProxy::InitializeSharedCollisionStructures(
 					Handles[TransformGroupIndex].Get(),
 					SharedParams, 
 					CollectionSimplicials[TransformGroupIndex].Get(),
-					Chaos::MakeSerializable(CollectionImplicits[TransformGroupIndex]), 
+					CollectionImplicits[TransformGroupIndex], 
 					CollectionMass[TransformGroupIndex],
 					CollectionInertiaTensor[TransformGroupIndex], 
 					MassToComponent,
@@ -3126,7 +3152,7 @@ void FGeometryCollectionPhysicsProxy::InitializeSharedCollisionStructures(
 
 					//don't support non level-set serialization
 					ErrorReporter.SetPrefix(BaseErrorPrefix + " | Cluster Transform Index: " + FString::FromInt(ClusterTransformIdx));
-					CollectionImplicits[ClusterTransformIdx] = TUniquePtr<FImplicitObject>(
+					CollectionImplicits[ClusterTransformIdx] = TSharedPtr<Chaos::FImplicitObject, ESPMode::ThreadSafe>(
 						FCollisionStructureManager::NewImplicit(ErrorReporter, MassSpaceParticles, *UnionMesh,
 							InstanceBoundingBox, 0, MinResolution, MaxResolution,
 							SizeSpecificData.CollisionObjectReductionPercentage, SizeSpecificData.CollisionType,
@@ -3140,7 +3166,7 @@ void FGeometryCollectionPhysicsProxy::InitializeSharedCollisionStructures(
 				else if (SizeSpecificData.ImplicitType == EImplicitTypeEnum::Chaos_Implicit_Box)
 				{
 					ErrorReporter.SetPrefix(BaseErrorPrefix + " | Cluster Transform Index: " + FString::FromInt(ClusterTransformIdx));
-					CollectionImplicits[ClusterTransformIdx] = TUniquePtr<FImplicitObject>(
+					CollectionImplicits[ClusterTransformIdx] = TSharedPtr<Chaos::FImplicitObject, ESPMode::ThreadSafe>(
 						FCollisionStructureManager::NewImplicit(ErrorReporter, MassSpaceParticles, *UnionMesh,
 							InstanceBoundingBox, 0, 0, 0,
 							SizeSpecificData.CollisionObjectReductionPercentage, SizeSpecificData.CollisionType,
@@ -3155,7 +3181,7 @@ void FGeometryCollectionPhysicsProxy::InitializeSharedCollisionStructures(
 				else if (SizeSpecificData.ImplicitType == EImplicitTypeEnum::Chaos_Implicit_Sphere)
 				{
 					ErrorReporter.SetPrefix(BaseErrorPrefix + " | Cluster Transform Index: " + FString::FromInt(ClusterTransformIdx));
-					CollectionImplicits[ClusterTransformIdx] = TUniquePtr<FImplicitObject>(
+					CollectionImplicits[ClusterTransformIdx] = TSharedPtr<Chaos::FImplicitObject, ESPMode::ThreadSafe>(
 						FCollisionStructureManager::NewImplicit(ErrorReporter, MassSpaceParticles, *UnionMesh,
 							InstanceBoundingBox, InstanceBoundingBox.GetExtent().GetAbsMin() / 2, 0, 0,
 							SizeSpecificData.CollisionObjectReductionPercentage, SizeSpecificData.CollisionType,
diff --git a/Engine/Source/Runtime/Experimental/ChaosSolvers/Public/PhysicsProxy/GeometryCollectionPhysicsProxy.h b/Engine/Source/Runtime/Experimental/ChaosSolvers/Public/PhysicsProxy/GeometryCollectionPhysicsProxy.h
--- a/Engine/Source/Runtime/Experimental/ChaosSolvers/Public/PhysicsProxy/GeometryCollectionPhysicsProxy.h
+++ b/Engine/Source/Runtime/Experimental/ChaosSolvers/Public/PhysicsProxy/GeometryCollectionPhysicsProxy.h
@@ -98,6 +98,7 @@ public:
 	// collection attributes
 	static FName SimplicialsAttribute;
 	static FName ImplicitsAttribute;
+	static FName SharedImplicitsAttribute;
 	static FName SolverParticleHandlesAttribute;
 	static FName SolverClusterHandlesAttribute;
 	static FName GTGeometryParticleAttribute;
@@ -291,9 +292,7 @@ private:
 	TManagedArray<FVector> InitialLinearVelocity;
 	TManagedArray<bool> SimulatableParticles;
 	TManagedArray<TUniquePtr<FSimplicial> > Simplicials; // FSimplicial = Chaos::TBVHParticles<float,3>
-	TManagedArray<Chaos::TSerializablePtr<Chaos::FImplicitObject>> Implicits;
-	//TManagedArray<Chaos::TUniquePtr<Chaos::FImplicitObject>> Implicits;
-	//TManagedArray<TSharedPtr<Chaos::FImplicitObject, ESPMode::ThreadSafe>> Implicits;
+	TManagedArray<TSharedPtr<Chaos::FImplicitObject, ESPMode::ThreadSafe>> Implicits;
 	TArray<int32> EndFrameUnparentingBuffer;
 
 	// This is a subset of the geometry group that are used in the transform hierarchy to represent geometry
diff --git a/Engine/Source/Runtime/Experimental/GeometryCollectionCore/Public/GeometryCollection/ManagedArrayCollection.h b/Engine/Source/Runtime/Experimental/GeometryCollectionCore/Public/GeometryCollection/ManagedArrayCollection.h
--- a/Engine/Source/Runtime/Experimental/GeometryCollectionCore/Public/GeometryCollection/ManagedArrayCollection.h
+++ b/Engine/Source/Runtime/Experimental/GeometryCollectionCore/Public/GeometryCollection/ManagedArrayCollection.h
@@ -177,6 +177,17 @@ public:
 		return nullptr;
 	};
 
+	template<typename T>
+	const TManagedArray<T>* FindAttribute(FName Name, FName Group) const
+	{
+		if (HasAttribute(Name, Group))
+		{
+			FKeyType Key = FManagedArrayCollection::MakeMapKey(Name, Group);
+			return static_cast<TManagedArray<T>*>(Map[Key].Value);
+		}
+		return nullptr;
+	};
+
 	/**
 	* Returns attribute access of Type(T) from the group
 	* @param Name - The name of the attribute
diff --git a/Engine/Source/Runtime/Experimental/GeometryCollectionCore/Public/GeometryCollection/ManagedArrayTypeValues.inl b/Engine/Source/Runtime/Experimental/GeometryCollectionCore/Public/GeometryCollection/ManagedArrayTypeValues.inl
--- a/Engine/Source/Runtime/Experimental/GeometryCollectionCore/Public/GeometryCollection/ManagedArrayTypeValues.inl
+++ b/Engine/Source/Runtime/Experimental/GeometryCollectionCore/Public/GeometryCollection/ManagedArrayTypeValues.inl
@@ -38,6 +38,7 @@ MANAGED_ARRAY_TYPE(Chaos::TPBDRigidParticleHandleFloat3*, TPBDRigidParticleHandl
 MANAGED_ARRAY_TYPE(Chaos::TPBDGeometryCollectionParticleHandleFloat3*, TPBDGeometryCollectionParticleHandle3fPtr)
 MANAGED_ARRAY_TYPE(Chaos::TPBDRigidClusteredParticleHandleFloat3*, TPBDRigidClusteredParticleHandle3fPtr)
 MANAGED_ARRAY_TYPE(TUniquePtr<Chaos::TGeometryParticleFloat3>, TGeometryParticle3fUniquePtr)
+MANAGED_ARRAY_TYPE(Chaos::ThreadSafeSharedPtr_FImplicitObject, FImplicitObject3ThreadSafeSharedPointer)
 
 // NOTE: new types must be added at the bottom to keep serialization from breaking
 
