Merge 11407298 to Release-4.25.

Fix geometry collection transforms.  Fix a bug with collision particles including the origin.  Fix a problem with the geometry collection componet not being delivered simulation results, and having to compute transforms itself.  Fix an issue where it seems transforms and bounds were computed on data for the prior frame rather that the current.  Implement support for "None" as an option for geometry collection collision geometry.  Add a cvar for disabling gravity for geometry collections.  Code cleanup and compaction.  Do array preallocations where we can.

diff --git a/Engine/Source/Runtime/Experimental/Chaos/Public/Chaos/Framework/MultiBufferResource.h b/Engine/Source/Runtime/Experimental/Chaos/Public/Chaos/Framework/MultiBufferResource.h
--- a/Engine/Source/Runtime/Experimental/Chaos/Public/Chaos/Framework/MultiBufferResource.h
+++ b/Engine/Source/Runtime/Experimental/Chaos/Public/Chaos/Framework/MultiBufferResource.h
@@ -242,6 +242,15 @@ namespace Chaos
 			return ConsumerThreadBuffer->Value.Get();
 		}
 
+		/**
+		 * Get access to the currently held consumer buffer, ignoring whether
+		 * it's already been consumed.
+		 */
+		const ResourceType* PeekConsumerBuffer() const
+		{
+			return ConsumerThreadBuffer->Value.Get();
+		}
+
 		/**
 		 * Make the current producer buffer available to the consumer thread.
 		 *
diff --git a/Engine/Source/Runtime/Experimental/Chaos/Public/Chaos/PerParticleGravity.h b/Engine/Source/Runtime/Experimental/Chaos/Public/Chaos/PerParticleGravity.h
--- a/Engine/Source/Runtime/Experimental/Chaos/Public/Chaos/PerParticleGravity.h
+++ b/Engine/Source/Runtime/Experimental/Chaos/Public/Chaos/PerParticleGravity.h
@@ -34,27 +34,26 @@ namespace Chaos
 
 		void Apply(TTransientPBDRigidParticleHandle<T, d>& Particle, const T Dt) const override //-V762
 		{
-			if (!UsesGravity.Contains(Particle.Handle()))
+			if (!OptOut.Contains(Particle.Handle()))
 			{
 				Particle.F() += MAcceleration * Particle.M();
 			}
 		}
 
+		bool GetEnabled(const TGeometryParticleHandleImp<T, d, true>& Handle)
+		{
+			return !OptOut.Contains(&Handle);
+		}
+
 		void SetEnabled(const TGeometryParticleHandleImp<T, d, true>& Handle, bool bEnabled)
 		{
 			if (bEnabled)
 			{
-				if (UsesGravity.Contains(&Handle))
-				{
-					UsesGravity.Remove(&Handle);
-				}
+				OptOut.Remove(&Handle);
 			}
 			else
 			{
-				if (!UsesGravity.Contains(&Handle))
-				{
-					UsesGravity.Add(&Handle);
-				}
+				OptOut.Add(&Handle);
 			}
 		}
 
@@ -65,6 +64,6 @@ namespace Chaos
 
 	private:
 		TVector<T, d> MAcceleration;
-		TSet< const TGeometryParticleHandleImp<T, d, true>* > UsesGravity;
+		TSet< const TGeometryParticleHandleImp<T, d, true>* > OptOut;
 	};
 }
diff --git a/Engine/Source/Runtime/Experimental/ChaosSolvers/Private/PhysicsProxy/GeometryCollectionPhysicsProxy.cpp b/Engine/Source/Runtime/Experimental/ChaosSolvers/Private/PhysicsProxy/GeometryCollectionPhysicsProxy.cpp
--- a/Engine/Source/Runtime/Experimental/ChaosSolvers/Private/PhysicsProxy/GeometryCollectionPhysicsProxy.cpp
+++ b/Engine/Source/Runtime/Experimental/ChaosSolvers/Private/PhysicsProxy/GeometryCollectionPhysicsProxy.cpp
@@ -14,6 +14,7 @@
 #include "Chaos/MassProperties.h"
 #include "ChaosSolversModule.h"
 #include "Chaos/PBDCollisionConstraintsUtil.h"
+#include "Chaos/PerParticleGravity.h"
 #include "Chaos/ImplicitObject.h"
 #include "Chaos/Serializable.h"
 #include "Chaos/ErrorReporter.h"
@@ -41,7 +42,16 @@
 #endif
 
 float CollisionParticlesPerObjectFractionDefault = 1.0f;
-FAutoConsoleVariableRef CVarCollisionParticlesPerObjectFractionDefault(TEXT("p.CollisionParticlesPerObjectFractionDefault"), CollisionParticlesPerObjectFractionDefault, TEXT("Fraction of verts"));
+FAutoConsoleVariableRef CVarCollisionParticlesPerObjectFractionDefault(
+	TEXT("p.CollisionParticlesPerObjectFractionDefault"), 
+	CollisionParticlesPerObjectFractionDefault, 
+	TEXT("Fraction of verts"));
+
+bool DisableGeometryCollectionGravity = false;
+FAutoConsoleVariableRef CVarGeometryCollectionDisableGravity(
+	TEXT("p.GeometryCollectionDisableGravity"),
+	DisableGeometryCollectionGravity,
+	TEXT("Disable gravity for geometry collections"));
 
 FName FGeometryCollectionPhysicsProxy::SimplicialsAttribute("CollisionParticles");
 FName FGeometryCollectionPhysicsProxy::ImplicitsAttribute("Implicits");
@@ -231,10 +241,7 @@ void PopulateSimulatedParticle(
 {
 	SCOPE_CYCLE_COUNTER(STAT_PopulateSimulatedParticle);
 
-	check(Handle->UniqueIdx().IsValid());
-
 	Handle->SetDisabledLowLevel(false);
-	//Particles.SetDisabledLowLevel(RigidBodyIndex, false);
 
 	Handle->SetX(WorldTransform.GetTranslation());
 	Handle->SetV(Chaos::TVector<float, 3>(0.f));
@@ -244,14 +251,6 @@ void PopulateSimulatedParticle(
 	Handle->SetQ(Handle->R());
 	Handle->SetIsland(INDEX_NONE);
 
-	//Particles.X(RigidBodyIndex) = WorldTransform.GetTranslation();
-	//Particles.V(RigidBodyIndex) = Chaos::TVector<float, 3>( FVector(0) );
-	//Particles.R(RigidBodyIndex) = WorldTransform.GetRotation().GetNormalized();
-	//Particles.W(RigidBodyIndex) = Chaos::TVector<float, 3>( FVector(0) );
-	//Particles.P(RigidBodyIndex) = Particles.X(RigidBodyIndex);
-	//Particles.Q(RigidBodyIndex) = Particles.R(RigidBodyIndex);
-	//Particles.Island(RigidBodyIndex) = INDEX_NONE;
-
 	//todo: If mass is too small use the right inertia
 	ensureMsgf(MassIn >= SharedParams.MinimumMassClamp, 
 		TEXT("Mass smaller than minimum mass clamp. Too late to change"));
@@ -262,25 +261,20 @@ void PopulateSimulatedParticle(
 
 	Handle->SetM(MassIn);
 	Handle->SetInvM(1.0f / MassIn);
-	//Particles.M(RigidBodyIndex) = MassIn;
 	if (FMath::IsNaN(InertiaTensorVec[0]) || FMath::IsNaN(InertiaTensorVec[1]) || FMath::IsNaN(InertiaTensorVec[2]) ||
 		InertiaTensorVec[0] < SMALL_NUMBER || InertiaTensorVec[1] < SMALL_NUMBER || InertiaTensorVec[2] < SMALL_NUMBER)
 	{
 		Handle->SetI(Chaos::PMatrix<float, 3, 3>(1.f, 1.f, 1.f));
-		//Particles.I(RigidBodyIndex) = Chaos::PMatrix<float, 3, 3>(1.f, 1.f, 1.f);
 	}
 	else
 	{
 		Handle->SetI(Chaos::PMatrix<float, 3, 3>(InertiaTensorVec[0], InertiaTensorVec[1], InertiaTensorVec[2]));
 		Handle->SetInvI(Chaos::PMatrix<float, 3, 3>(1.0f / InertiaTensorVec[0], 1.0f / InertiaTensorVec[1], 1.0f / InertiaTensorVec[2]));
-		//Particles.I(RigidBodyIndex) = Chaos::PMatrix<float, 3, 3>(InertiaTensorVec[0], InertiaTensorVec[1], InertiaTensorVec[2]);
 	}
 
 	// for validation set the body to dynamic and check the inverse masses. 
 	Handle->SetObjectStateLowLevel(Chaos::EObjectStateType::Dynamic);
-	//Particles.SetObjectState(RigidBodyIndex, Chaos::EObjectStateType::Dynamic);
 
-	//ensureMsgf(Particles.InvI(RigidBodyIndex).M[0][0] > SMALL_NUMBER && Particles.InvI(RigidBodyIndex).M[1][1] > SMALL_NUMBER && Particles.InvI(RigidBodyIndex).M[2][2] > SMALL_NUMBER, TEXT("Inertia tensor is too large. Too late to change"));
 	const Chaos::PMatrix<float,3,3>& InvI = Handle->InvI();
 	if (!(InvI.M[0][0] > SMALL_NUMBER && 
 		  InvI.M[1][1] > SMALL_NUMBER && 
@@ -290,50 +284,33 @@ void PopulateSimulatedParticle(
 	}
 
 	Handle->SetCollisionGroup(CollisionGroup);
-	//Particles.CollisionGroup(RigidBodyIndex) = CollisionGroup;
 
+	if (Implicit)	//todo(ocohen): this is only needed for cases where clusters have no proxy. Kind of gross though, should refactor
 	{
-		// @important (The solver can not free this memory)
-		if (Implicit)	//todo(ocohen): this is only needed for cases where clusters have no proxy. Kind of gross though, should refactor
-		{
-			Handle->SetSharedGeometry(Implicit);
-		}
+		Handle->SetSharedGeometry(Implicit);
+	}
 
-		Handle->CollisionParticlesInitIfNeeded();
-		//Particles.CollisionParticlesInitIfNeeded(RigidBodyIndex);
-		if (Simplicial)
+	Handle->CollisionParticlesInitIfNeeded();
+	if (Simplicial)
+	{
+		TUniquePtr<Chaos::TBVHParticles<float, 3>>& CollisionParticles = Handle->CollisionParticles();
+		if (Simplicial->Size())
 		{
-			TUniquePtr<Chaos::TBVHParticles<float, 3>>& CollisionParticles = Handle->CollisionParticles();
-			if (Simplicial->Size())
+			CollisionParticles->Resize(0);
+			CollisionParticles->AddParticles(Simplicial->Size());
+			for (int32 VertexIndex = 0; VertexIndex < (int32)Simplicial->Size(); ++VertexIndex)
 			{
-				CollisionParticles->Resize(0);
-				//Particles.CollisionParticles(RigidBodyIndex)->Resize(0);
-				CollisionParticles->AddParticles(Simplicial->Size());
-				//Particles.CollisionParticles(RigidBodyIndex)->AddParticles(Simplicial->Size());
-				for (int32 VertexIndex = 0; VertexIndex < (int32)Simplicial->Size(); ++VertexIndex)
-				{
-					CollisionParticles->X(VertexIndex) = Simplicial->X(VertexIndex);
-					//Particles.CollisionParticles(RigidBodyIndex)->X(VertexIndex) = Simplicial->X(VertexIndex);
-				}
-			}
-
-			// @todo(remove): IF there is no simplicial we should not be forcing one. 
-			if (CollisionParticles->Size())
-			//if (!Particles.CollisionParticles(RigidBodyIndex)->Size())
-			{
-				CollisionParticles->AddParticles(1);
-				//Particles.CollisionParticles(RigidBodyIndex)->AddParticles(1);
-				CollisionParticles->X(0) = Chaos::TVector<float, 3>(0);
-				///Particles.CollisionParticles(RigidBodyIndex)->X(0) = Chaos::TVector<float, 3>(0);
+				CollisionParticles->X(VertexIndex) = Simplicial->X(VertexIndex);
 			}
+		}
 
-			if (CollisionParticles->Size())
-			//if (Particles.CollisionParticles(RigidBodyIndex)->Size())
-			{
-				CollisionParticles->UpdateAccelerationStructures();
-				//Particles.CollisionParticles(RigidBodyIndex)->UpdateAccelerationStructures();
-			}
+		// @todo(remove): IF there is no simplicial we should not be forcing one. 
+		if (!CollisionParticles->Size())
+		{
+			CollisionParticles->AddParticles(1);
+			CollisionParticles->X(0) = Chaos::TVector<float, 3>(0);
 		}
+		CollisionParticles->UpdateAccelerationStructures();
 	}
 
 	//
@@ -345,22 +322,18 @@ void PopulateSimulatedParticle(
 	if (DynamicState == (uint8)EObjectStateTypeEnum::Chaos_Object_Sleeping)
 	{
 		Handle->SetObjectStateLowLevel(Chaos::EObjectStateType::Sleeping);
-		//Particles.SetObjectState(RigidBodyIndex,Chaos::EObjectStateType::Sleeping);
 	}
 	else if (DynamicState == (uint8)EObjectStateTypeEnum::Chaos_Object_Kinematic)
 	{
 		Handle->SetObjectStateLowLevel(Chaos::EObjectStateType::Kinematic);
-		//Particles.SetObjectState(RigidBodyIndex, Chaos::EObjectStateType::Kinematic);
 	}
 	else if (DynamicState == (uint8)EObjectStateTypeEnum::Chaos_Object_Static)
 	{
 		Handle->SetObjectStateLowLevel(Chaos::EObjectStateType::Static);
-		//Particles.SetObjectState(RigidBodyIndex, Chaos::EObjectStateType::Static);
 	}
 	else
 	{
 		Handle->SetObjectStateLowLevel(Chaos::EObjectStateType::Dynamic);
-		//Particles.SetObjectState(RigidBodyIndex,Chaos::EObjectStateType::Dynamic);
 	}
 }
 
@@ -395,8 +368,8 @@ FGeometryCollectionPhysicsProxy::FGeometryCollectionPhysicsProxy(
 	, LastSyncCountGT(MAX_uint32)
 	, CollisionParticlesPerObjectFraction(CollisionParticlesPerObjectFractionDefault)
 
-	, PhysToGameInterchange(Chaos::FMultiBufferFactory<FGeometryCollectionResults>::CreateBuffer(BufferMode))
-	, GameToPhysInterchange(Chaos::FMultiBufferFactory<FGeometryCollectionResults>::CreateBuffer(BufferMode))
+	, PhysToGameInterchange(MakeUnique<Chaos::FGuardedTripleBuffer<FGeometryCollectionResults>>())
+	, GameToPhysInterchange(MakeUnique<Chaos::FGuardedTripleBuffer<FGeometryCollectionResults>>())
 {
 	// We rely on a guarded buffer.
 	check(BufferMode == Chaos::EMultiBufferMode::TripleGuarded);
@@ -594,8 +567,8 @@ void FGeometryCollectionPhysicsProxy::InitializeBodiesGT()
 					// which causes this proxy class to be put on a dirty list.
 					// That causes our PushToPhysicsState() to be called prematurely, 
 					// so we need to guard against that.
-					GTParticle->SetX(ParticleTransform.GetTranslation(), false); 
-					GTParticle->SetR(ParticleTransform.GetRotation().GetNormalized(), false);
+					//GTParticle->SetX(ParticleTransform.GetTranslation(), false); 
+					//GTParticle->SetR(ParticleTransform.GetRotation().GetNormalized(), false);
 					GTParticle->SetIsland(INDEX_NONE); // doesn't mark dirty
 					GTParticle->Proxy = this;
 
@@ -614,19 +587,6 @@ void FGeometryCollectionPhysicsProxy::InitializeBodiesGT()
 					{
 						// Not sure what to do with particles at this point.
 					}
-
-					if (Parameters.InitialVelocityType == EInitialVelocityTypeEnum::Chaos_Initial_Velocity_User_Defined)
-					{
-						InitialLinearVelocity[Idx] = Parameters.InitialLinearVelocity;
-						InitialAngularVelocity[Idx] = Parameters.InitialAngularVelocity;
-						GTParticle->SetV(Parameters.InitialLinearVelocity, false);
-						GTParticle->SetW(Parameters.InitialAngularVelocity, false);
-					}
-					else
-					{
-						GTParticle->SetV(Chaos::TVector<float, 3>(0.f), false);
-						GTParticle->SetW(Chaos::TVector<float, 3>(0.f), false);
-					}
 				}
 			}
 /*			else
@@ -673,7 +633,8 @@ void FGeometryCollectionPhysicsProxy::InitializeBodiesPT(
 		const int32 NumTransforms = GetTransformGroupSize();
 
 		TArray<FTransform> Transform;
-		GeometryCollectionAlgo::GlobalMatrices(DynamicCollection->Transform, DynamicCollection->Parent, Transform);
+		GeometryCollectionAlgo::GlobalMatrices(
+			DynamicCollection->Transform, DynamicCollection->Parent, Transform);
 		check(DynamicCollection->Transform.Num() == Transform.Num());
 
 		const int NumRigids = 0; // ryan - Since we're doing SOA, we start at zero?
@@ -718,14 +679,26 @@ void FGeometryCollectionPhysicsProxy::InitializeBodiesPT(
 
 		const float StrainDefault = Parameters.DamageThreshold.Num() ? Parameters.DamageThreshold[0] : 0;
 		// Add the rigid bodies
+
+		// Iterating over the geometry group is a fast way of skipping everything that's
+		// not a leaf node, as each geometry has a transform index, which is a shortcut
+		// for the case when there's a 1-to-1 mapping between transforms and geometries.
+		// At the point that we start supporting instancing, this assumption will no longer
+		// hold, and those reverse mappints will be INDEX_NONE.
+
 		const int32 NumGeometries = DynamicCollection->NumElements(FGeometryCollection::GeometryGroup);
 		ParallelFor(NumGeometries, [&](int32 GeometryIndex)
+		//for (int32 GeometryIndex = 0; GeometryIndex < NumGeometries; ++GeometryIndex)
 		{
 			const int32 TransformGroupIndex = TransformIndex[GeometryIndex];
-			//if (SimulatableParticles[TransformGroupIndex] == true)
 			if (Chaos::TPBDRigidClusteredParticleHandle<float, 3>* Handle = SolverParticleHandles[TransformGroupIndex])
 			{
-				//const int32 RigidBodyIndex = RigidBodyID[TransformGroupIndex];
+				// TODO: We probably want to take better advantage of this solver
+				// feature, rather than doing it all here ourselves?
+				Handle->SetCenterOfMass(Chaos::TVector<float,3>(0.0f));
+				Handle->SetRotationOfMass(Chaos::TRotation<float,3>::FromIdentity());
+			
+				// Mass space -> Composed parent space -> world
 				const FTransform WorldTransform = 
 					MassToLocal[TransformGroupIndex] * Transform[TransformGroupIndex] * Parameters.WorldTransform;
 
@@ -747,6 +720,19 @@ void FGeometryCollectionPhysicsProxy::InitializeBodiesPT(
 					(uint8)Parameters.ObjectType, 
 					CollisionGroup[TransformGroupIndex]);
 
+				if (Parameters.InitialVelocityType == EInitialVelocityTypeEnum::Chaos_Initial_Velocity_User_Defined)
+				{
+					InitialLinearVelocity[TransformGroupIndex] = Parameters.InitialLinearVelocity;
+					InitialAngularVelocity[TransformGroupIndex] = Parameters.InitialAngularVelocity;
+					Handle->SetV(Parameters.InitialLinearVelocity);
+					Handle->SetW(Parameters.InitialAngularVelocity);
+				}
+				else
+				{
+					Handle->SetV(Chaos::TVector<float, 3>(0.f));
+					Handle->SetW(Chaos::TVector<float, 3>(0.f));
+				}
+
 				if (Parameters.EnableClustering)
 				{
 					Handle->SetClusterGroupIndex(Parameters.ClusterGroupIndex);
@@ -1401,6 +1387,7 @@ void FGeometryCollectionPhysicsProxy::StartFrameCallback(const float Dt, const f
 
 void FGeometryCollectionPhysicsProxy::EndFrameCallback(const float EndFrame)
 {
+	check(false); // not currently called
 	if (InitializedState == ESimulationInitializationState::Initialized)
 	{
 		FGeometryDynamicCollection* Collection = Parameters.DynamicCollection;
@@ -1973,6 +1960,7 @@ void FGeometryCollectionPhysicsProxy::BindParticleCallbackMapping(Chaos::TArrayC
 
 void FGeometryCollectionPhysicsProxy::ContiguousIndices(TArray<ContextIndex>& Array, const Chaos::FPhysicsSolver* RigidSolver, EFieldResolutionType ResolutionType, bool bForce = true)
 {
+	check(false); // Not currently called
 	if (bForce)
 	{
 #if TODO_REIMPLEMENT_GET_RIGID_PARTICLES
@@ -2772,9 +2760,10 @@ void FGeometryCollectionPhysicsProxy::InitializeSharedCollisionStructures(
 	const TManagedArray<int32>& BoneMap = RestCollection.BoneMap;
 	const TManagedArray<int32>& Parent = RestCollection.Parent;
 	const TManagedArray<TSet<int32>>& Children = RestCollection.Children;
-	TManagedArray<bool>& CollectionSimulatableParticles = 
+	const TManagedArray<bool>& CollectionSimulatableParticles = 
 		RestCollection.GetAttribute<bool>(
 			FGeometryCollection::SimulatableParticlesAttribute, FTransformCollection::TransformGroup);
+
 	TManagedArray<FVector>& CollectionInertiaTensor = 
 		RestCollection.AddAttribute<FVector>(
 			TEXT("InertiaTensor"), FTransformCollection::TransformGroup);
@@ -2793,10 +2782,7 @@ void FGeometryCollectionPhysicsProxy::InitializeSharedCollisionStructures(
 
 	FTransform IdentityXf(FQuat::Identity, FVector(0));
 	IdentityXf.NormalizeRotation();
-	for (int32 Index = 0; Index < CollectionMassToLocal.Num(); ++Index)
-	{
-		CollectionMassToLocal[Index] = IdentityXf;
-	}
+	CollectionMassToLocal.Fill(IdentityXf);
 
 	// VerticesGroup
 	const TManagedArray<FVector>& Vertex = RestCollection.Vertex;
@@ -2812,19 +2798,18 @@ void FGeometryCollectionPhysicsProxy::InitializeSharedCollisionStructures(
 	const TManagedArray<FIntVector>& Indices = RestCollection.Indices;
 
 	TArray<FTransform> CollectionSpaceTransforms;
-	TManagedArray<FTransform>& HierarchyTransform = RestCollection.Transform;
-	GeometryCollectionAlgo::GlobalMatrices(HierarchyTransform, Parent, CollectionSpaceTransforms);
-	check(HierarchyTransform.Num() == CollectionSpaceTransforms.Num());
+	{ // tmp scope
+		const TManagedArray<FTransform>& HierarchyTransform = RestCollection.Transform;
+		GeometryCollectionAlgo::GlobalMatrices(HierarchyTransform, Parent, CollectionSpaceTransforms);
+		check(HierarchyTransform.Num() == CollectionSpaceTransforms.Num());
+	} // tmp scope
 
 	const int32 NumTransforms = CollectionSpaceTransforms.Num();
 	const int32 NumGeometries = RestCollection.NumElements(FGeometryCollection::GeometryGroup);
 
-	const bool bEnableClustering = true;	//question: at the moment we always build cluster data in the asset. This allows for per instance toggling. Is this needed? It increases memory usage for all geometry collection assets
-
 	TArray<TUniquePtr<TTriangleMesh<float>>> TriangleMeshesArray;	//use to union trimeshes in cluster case
 	TriangleMeshesArray.AddDefaulted(NumTransforms);
 
-	float TotalVolume = 0.f;
 	TArray<TSet<int32>> VertsAddedArray;
 	VertsAddedArray.AddDefaulted(NumGeometries);
 
@@ -2842,38 +2827,54 @@ void FGeometryCollectionPhysicsProxy::InitializeSharedCollisionStructures(
 	InertiaComputationNeeded.AddUninitialized(NumGeometries);
 
 	const float MinBoundsExtents = 10.f;
+
+	float TotalVolume = 0.f;
 	int32 NumSimulatableParticles = 0;
 
+	// The geometry group has a set of transform indices that maps a geometry index
+	// to a transform index, but only in the case where there is a 1-to-1 mapping 
+	// between the two.  In the event where a geometry is instanced for multiple
+	// transforms, the transform index on the geometry group should be INDEX_NONE.
+	// Otherwise, iterating over the geometry group is a convenient way to iterate
+	// over all the leaves of the hierarchy.
+	check(!TransformIndex.Contains(INDEX_NONE)); // TODO: implement support for instanced bodies
 	for (int32 GeometryIndex = 0; GeometryIndex < NumGeometries; GeometryIndex++)
 	{
-		int32 TransformGroupIndex = TransformIndex[GeometryIndex];
+		const int32 TransformGroupIndex = TransformIndex[GeometryIndex];
 		if (CollectionSimulatableParticles[TransformGroupIndex])
 		{
 			NumSimulatableParticles++;
 
-			TMassProperties<float, 3>& MassProperties = MassPropertiesArray[GeometryIndex];
-
-			TUniquePtr<TTriangleMesh<float>> TriMesh(CreateTriangleMesh(FaceCount[GeometryIndex], VertexStart[GeometryIndex], FaceStart[GeometryIndex], Vertex, Visible, Indices, VertsAddedArray[GeometryIndex]));
+			TUniquePtr<TTriangleMesh<float>> TriMesh(
+				CreateTriangleMesh(
+					FaceCount[GeometryIndex], 
+					VertexStart[GeometryIndex], 
+					FaceStart[GeometryIndex], 
+					Vertex, 
+					Visible, 
+					Indices, 
+					VertsAddedArray[GeometryIndex]));
 
-			const float BoundsVolume = FCollisionStructureManager::CalculateVolume(BoundingBox[GeometryIndex], InnerRadius[GeometryIndex], EImplicitTypeEnum::Chaos_Implicit_Box);
+			TMassProperties<float, 3>& MassProperties = MassPropertiesArray[GeometryIndex];
 
 			//CalculateVolumeAndCenterOfMass(MassSpaceParticles, *TriMesh, MassProperties.Volume, MassProperties.CenterOfMass);
 			//InertiaComputationNeeded[GeometryIndex] = true;
 			//if (MassProperties.Volume < MinVolume)
 			//just assume everything is a box, need better computation later but requires mesh cleanup
 			{
-				InertiaComputationNeeded[GeometryIndex] = false;	//since volume is too small we just use a fallback one
+				//since volume is too small we just use a fallback one
+				InertiaComputationNeeded[GeometryIndex] = false;	
 
-				CollectionMassToLocal[TransformGroupIndex] = FTransform(TRotation<float, 3>(FQuat(0, 0, 0, 1)), BoundingBox[GeometryIndex].GetCenter());
-				FVector Size = BoundingBox[GeometryIndex].GetSize();
-				Size.X = FMath::Max(Size.X, MinBoundsExtents);
-				Size.Y = FMath::Max(Size.Y, MinBoundsExtents);
-				Size.Z = FMath::Max(Size.Z, MinBoundsExtents);
+				CollectionMassToLocal[TransformGroupIndex] = 
+					FTransform(TRotation<float, 3>(FQuat(0, 0, 0, 1)), BoundingBox[GeometryIndex].GetCenter());
 
-				FVector SideSquared(Size.X * Size.X, Size.Y * Size.Y, Size.Z * Size.Z);
-				MassProperties.Volume = Size.X * Size.Y * Size.Z;	//just fake a box if tiny
-				MassProperties.InertiaTensor = PMatrix<float, 3, 3>((SideSquared.Y + SideSquared.Z) / 12.f, (SideSquared.X + SideSquared.Z) / 12.f, (SideSquared.X + SideSquared.Y) / 12.f);
 
+				const TVector<float, 3> Size = BoundingBox[GeometryIndex].GetSize().ComponentMax(FVector(MinBoundsExtents));
+				MassProperties.Volume = Size.Product();	//just fake a box if tiny
+				MassProperties.InertiaTensor = 
+					Chaos::TBox<float, 3>::GetInertiaTensor(
+						1.0f, // Mass
+						Size);
 				/*
 				if (MassProperties.Volume < MinVolume)
 				{
@@ -2923,7 +2924,7 @@ void FGeometryCollectionPhysicsProxy::InitializeSharedCollisionStructures(
 
 	for (int32 GeometryIndex = 0; GeometryIndex < NumGeometries; GeometryIndex++)
 	{
-		int32 TransformGroupIndex = TransformIndex[GeometryIndex];
+		const int32 TransformGroupIndex = TransformIndex[GeometryIndex];
 		if (CollectionSimulatableParticles[TransformGroupIndex])
 		{
 			TUniquePtr<TTriangleMesh<float>>& TriMesh = TriangleMeshesArray[TransformGroupIndex];
@@ -2935,26 +2936,35 @@ void FGeometryCollectionPhysicsProxy::InitializeSharedCollisionStructures(
 			{
 				ensure(false);
 			}
-
 			const float Mass_i = FMath::Max(DesiredDensity * Volume_i, SharedParams.MinimumMassClamp);
 			const float Density_i = Mass_i / Volume_i;
 			CollectionMass[TransformGroupIndex] = Mass_i;
 
-
 			if (InertiaComputationNeeded[GeometryIndex])
 			{
-				CalculateInertiaAndRotationOfMass(MassSpaceParticles, TriMesh->GetSurfaceElements(), Density_i, MassProperties.CenterOfMass, MassProperties.InertiaTensor, MassProperties.RotationOfMass);
-				CollectionMassToLocal[TransformGroupIndex] = FTransform(MassProperties.RotationOfMass, MassProperties.CenterOfMass);
-			}
-
-			const TVector<float, 3> DiagonalInertia(MassProperties.InertiaTensor.M[0][0], MassProperties.InertiaTensor.M[1][1], MassProperties.InertiaTensor.M[2][2]);
-			if(InertiaComputationNeeded[GeometryIndex])
-			{
+				check(false); // Ryan - Currently everything is hard wired to be a box in the loop above, so this should never run.
+				CalculateInertiaAndRotationOfMass(
+					MassSpaceParticles, 
+					TriMesh->GetSurfaceElements(), 
+					Density_i, 
+					MassProperties.CenterOfMass, 
+					MassProperties.InertiaTensor, 
+					MassProperties.RotationOfMass);
+
+				CollectionMassToLocal[TransformGroupIndex] = 
+					FTransform(MassProperties.RotationOfMass, MassProperties.CenterOfMass);
+
+				const TVector<float, 3> DiagonalInertia(
+					MassProperties.InertiaTensor.M[0][0], MassProperties.InertiaTensor.M[1][1], MassProperties.InertiaTensor.M[2][2]);
 				//computation includes mass already
 				CollectionInertiaTensor[TransformGroupIndex] = DiagonalInertia;
 			}
 			else
 			{
+				const TVector<float, 3> DiagonalInertia(
+					MassProperties.InertiaTensor.M[0][0], 
+					MassProperties.InertiaTensor.M[1][1], 
+					MassProperties.InertiaTensor.M[2][2]);
 				//using fallback computation, but adjust for mass
 				CollectionInertiaTensor[TransformGroupIndex] = DiagonalInertia * Mass_i;
 			}
@@ -2966,7 +2976,9 @@ void FGeometryCollectionPhysicsProxy::InitializeSharedCollisionStructures(
 			{
 				if (VertsAdded.Contains(VertIdx))	//only consider verts from the trimesh
 				{
-					MassSpaceParticles.X(VertIdx) = CollectionMassToLocal[TransformGroupIndex].InverseTransformPosition(MassSpaceParticles.X(VertIdx));
+					MassSpaceParticles.X(VertIdx) = 
+						CollectionMassToLocal[TransformGroupIndex].InverseTransformPosition(
+							MassSpaceParticles.X(VertIdx));
 					InstanceBoundingBox += MassSpaceParticles.X(VertIdx);	//build bounding box for visible verts in mass space
 				}
 			}
@@ -2974,46 +2986,84 @@ void FGeometryCollectionPhysicsProxy::InitializeSharedCollisionStructures(
 			const int32 SizeSpecificIdx = FindSizeSpecificIdx(SharedParams.SizeSpecificData, InstanceBoundingBox);
 			const FSharedSimulationSizeSpecificData& SizeSpecificData = SharedParams.SizeSpecificData[SizeSpecificIdx];
 
-
 			//
 			//  Build the simplicial for the rest collection. This will be used later in the DynamicCollection to 
 			//  populate the collision structures of the simulation. 
 			//
-			Chaos::TBVHParticles<float, 3> * Simplicial = FCollisionStructureManager::NewSimplicial(MassSpaceParticles, BoneMap, SizeSpecificData.CollisionType, *TriMesh, SizeSpecificData.CollisionParticlesFraction);
+			Chaos::TBVHParticles<float, 3> * Simplicial = 
+				FCollisionStructureManager::NewSimplicial(
+					MassSpaceParticles, 
+					BoneMap, 
+					SizeSpecificData.CollisionType, 
+					*TriMesh, 
+					SizeSpecificData.CollisionParticlesFraction);
 			CollectionSimplicials[TransformGroupIndex] = TUniquePtr<FSimplicial>(Simplicial);
 			ensureMsgf(CollectionSimplicials[TransformGroupIndex], TEXT("No simplicial representation."));
 
+			if (SizeSpecificData.ImplicitType == EImplicitTypeEnum::Chaos_Implicit_LevelSet)
+			{
+				ErrorReporter.SetPrefix(BaseErrorPrefix + " | Transform Index: " + FString::FromInt(TransformGroupIndex));
+				CollectionImplicits[TransformGroupIndex] = TSharedPtr<Chaos::FImplicitObject, ESPMode::ThreadSafe>(
+					FCollisionStructureManager::NewImplicitLevelset(
+						ErrorReporter,
+						MassSpaceParticles,
+						*TriMesh,
+						InstanceBoundingBox,
+						SizeSpecificData.MinLevelSetResolution,
+						SizeSpecificData.MaxLevelSetResolution,
+						SizeSpecificData.CollisionObjectReductionPercentage,
+						SizeSpecificData.CollisionType));
+			}
+			else if (SizeSpecificData.ImplicitType == EImplicitTypeEnum::Chaos_Implicit_Box)
+			{
+				CollectionImplicits[TransformGroupIndex] = TSharedPtr<Chaos::FImplicitObject, ESPMode::ThreadSafe>(
+					FCollisionStructureManager::NewImplicitBox(
+						InstanceBoundingBox,
+						SizeSpecificData.CollisionObjectReductionPercentage,
+						SizeSpecificData.CollisionType));
+			}
+			else if (SizeSpecificData.ImplicitType == EImplicitTypeEnum::Chaos_Implicit_Sphere)
+			{
+				CollectionImplicits[TransformGroupIndex] = TSharedPtr<Chaos::FImplicitObject, ESPMode::ThreadSafe>(
+					FCollisionStructureManager::NewImplicitSphere(
+						InnerRadius[GeometryIndex],
+						SizeSpecificData.CollisionObjectReductionPercentage,
+						SizeSpecificData.CollisionType));
+			}
+			else if (SizeSpecificData.ImplicitType == EImplicitTypeEnum::Chaos_Implicit_None)
+			{
+				CollectionImplicits[TransformGroupIndex] = nullptr;
+			}
+			else
+			{
+				ensure(false); // unsupported implicit type!
+			}
 
-			ErrorReporter.SetPrefix(BaseErrorPrefix + " | Transform Index: " + FString::FromInt(TransformGroupIndex));
-			CollectionImplicits[TransformGroupIndex] = TSharedPtr<Chaos::FImplicitObject, ESPMode::ThreadSafe>(
-				FCollisionStructureManager::NewImplicit(ErrorReporter, MassSpaceParticles, *TriMesh, InstanceBoundingBox,
-					InnerRadius[GeometryIndex], SizeSpecificData.MinLevelSetResolution, SizeSpecificData.MaxLevelSetResolution,
-					SizeSpecificData.CollisionObjectReductionPercentage, SizeSpecificData.CollisionType,
-					SizeSpecificData.ImplicitType));
 			if (CollectionImplicits[TransformGroupIndex])
 			{
-				TVector<float, 3> Extents = CollectionImplicits[TransformGroupIndex]->BoundingBox().Extents();
-				if (Extents.X > MaxChildBounds.X)
-				{
-					MaxChildBounds.X = Extents.X;
-				}
-				if (Extents.Y > MaxChildBounds.Y)
-				{
-					MaxChildBounds.Y = Extents.Y;
-				}
-				if (Extents.Z > MaxChildBounds.Z)
-				{
-					MaxChildBounds.Z = Extents.Z;
-				}
+				const TVector<float, 3> Extents = 
+					CollectionImplicits[TransformGroupIndex]->BoundingBox().Extents(); // Chaos::TAABB::Extents() is Max - Min
+				MaxChildBounds = MaxChildBounds.ComponentwiseMax(Extents);
 			}
 		}
 	}
 
+	// question: at the moment we always build cluster data in the asset. This 
+	// allows for per instance toggling. Is this needed? It increases memory 
+	// usage for all geometry collection assets.
+	const bool bEnableClustering = true;	
 	if (bEnableClustering)
 	{
 		//Put all children into collection space so we can compute mass properties.
 		TUniquePtr<TPBDRigidClusteredParticles<float, 3>> CollectionSpaceParticles(new TPBDRigidClusteredParticles<float, 3>());
 		CollectionSpaceParticles->AddParticles(NumTransforms);
+
+		// Init to -FLT_MAX for debugging purposes
+		for (int32 Idx = 0; Idx < NumTransforms; Idx++)
+		{
+			CollectionSpaceParticles->X(Idx) = Chaos::TVector<float, 3>(-TNumericLimits<float>::Max());
+		}
+
 		TArray<TUniquePtr<TPBDRigidClusteredParticleHandle<float, 3>>> Handles;
 		Handles.Reserve(NumTransforms);
 		for (int32 Idx = 0; Idx < NumTransforms; Idx++)
@@ -3025,7 +3075,7 @@ void FGeometryCollectionPhysicsProxy::InitializeSharedCollisionStructures(
 
 		for (int32 GeometryIdx = 0; GeometryIdx < NumGeometries; ++GeometryIdx)
 		{
-			int32 TransformGroupIndex = TransformIndex[GeometryIdx];
+			const int32 TransformGroupIndex = TransformIndex[GeometryIdx];
 			if (CollectionSimulatableParticles[TransformGroupIndex])
 			{
 				FTransform MassToComponent = CollectionMassToLocal[TransformGroupIndex] * CollectionSpaceTransforms[TransformGroupIndex];
@@ -3055,18 +3105,21 @@ void FGeometryCollectionPhysicsProxy::InitializeSharedCollisionStructures(
 				const int32 ClusterTransformIdx = TransformGroupIndex;
 				//update mass 
 				TSet<TPBDRigidParticleHandle<float,3>*> ChildrenIndices;
-				ChildrenIndices.Reserve(Children[ClusterTransformIdx].Num());
-				for (int32 ChildIdx : Children[ClusterTransformIdx])
-				{
-					if (CollectionSimulatableParticles[ChildIdx] || IsClusterSimulated[ChildIdx])
+				{ // tmp scope
+					ChildrenIndices.Reserve(Children[ClusterTransformIdx].Num());
+					for (int32 ChildIdx : Children[ClusterTransformIdx])
 					{
-						ChildrenIndices.Add(Handles[ChildIdx].Get());
+						if (CollectionSimulatableParticles[ChildIdx] || IsClusterSimulated[ChildIdx])
+						{
+							ChildrenIndices.Add(Handles[ChildIdx].Get());
+						}
 					}
-				}
-				if (!ChildrenIndices.Num())
-				{
-					continue;
-				}
+					if (!ChildrenIndices.Num())
+					{
+						continue;
+					}
+				} // tmp scope
+
 				//CollectionSimulatableParticles[TransformGroupIndex] = true;
 				IsClusterSimulated[TransformGroupIndex] = true;
 
@@ -3085,37 +3138,48 @@ void FGeometryCollectionPhysicsProxy::InitializeSharedCollisionStructures(
 				//update geometry
 				//merge children meshes and move them into cluster's mass space
 				TArray<TVector<int32, 3>> UnionMeshIndices;
+				int32 BiggestNumElements = 0;
+				{ // tmp scope
+					int32 NumChildIndices = 0;
+					for (TPBDRigidParticleHandle<float, 3>* Child : ChildrenIndices)
+					{
+						const int32 ChildTransformIdx = Child->UniqueIdx().Idx;
+						if (Chaos::TTriangleMesh<float>* ChildMesh = TriangleMeshesArray[ChildTransformIdx].Get())
+						{
+							BiggestNumElements = FMath::Max(BiggestNumElements, ChildMesh->GetNumElements());
+							NumChildIndices += ChildMesh->GetNumElements();
+						}
+					}
+					UnionMeshIndices.Reserve(NumChildIndices);
+				} // tmp scope
+
 				FBox InstanceBoundingBox(EForceInit::ForceInitToZero);
-				for (TPBDRigidParticleHandle<float,3>* Child : ChildrenIndices)
-				{
-					const int32 ChildTransformIdx = Child->UniqueIdx().Idx;
-					if (Chaos::TTriangleMesh<float>* ChildMesh = TriangleMeshesArray[ChildTransformIdx].Get())
+				{ // tmp scope
+					TSet<int32> VertsAdded;
+					VertsAdded.Reserve(BiggestNumElements);
+					for (TPBDRigidParticleHandle<float, 3>* Child : ChildrenIndices)
 					{
-						const TArray<TVector<int32, 3>>& ChildIndices = ChildMesh->GetSurfaceElements();
-						UnionMeshIndices.Append(ChildIndices);
+						const int32 ChildTransformIdx = Child->UniqueIdx().Idx;
+						if (Chaos::TTriangleMesh<float>* ChildMesh = TriangleMeshesArray[ChildTransformIdx].Get())
+						{
+							const TArray<TVector<int32, 3>>& ChildIndices = ChildMesh->GetSurfaceElements();
+							UnionMeshIndices.Append(ChildIndices);
 
-						const FTransform ChildMassToClusterMass = 
-							(CollectionMassToLocal[ChildTransformIdx] * CollectionSpaceTransforms[ChildTransformIdx]).GetRelativeTransform(
-								ClusterMassToCollection);
+							const FTransform ChildMassToClusterMass =
+								(CollectionMassToLocal[ChildTransformIdx] * CollectionSpaceTransforms[ChildTransformIdx]).GetRelativeTransform(
+									ClusterMassToCollection);
 
-						TSet<int32> VertsAdded;
-						for (const TVector<int32, 3>& Tri : ChildIndices)
-						{
-							for (int32 Axis = 0; Axis < 3; ++Axis)
+							ChildMesh->GetVertexSet(VertsAdded);
+							for (const int32 VertIdx : VertsAdded)
 							{
-								const int32 VertIdx = Tri[Axis];
-								if (!VertsAdded.Contains(VertIdx))
-								{
-									//Update particles so they are in the cluster's mass space
-									MassSpaceParticles.X(VertIdx) = 
-										ChildMassToClusterMass.TransformPosition(MassSpaceParticles.X(VertIdx));
-									InstanceBoundingBox += MassSpaceParticles.X(VertIdx);
-									VertsAdded.Add(VertIdx);
-								}
+								//Update particles so they are in the cluster's mass space
+								MassSpaceParticles.X(VertIdx) =
+									ChildMassToClusterMass.TransformPosition(MassSpaceParticles.X(VertIdx));
+								InstanceBoundingBox += MassSpaceParticles.X(VertIdx);
 							}
 						}
 					}
-				}
+				} // tmp scope
 
 				TUniquePtr<TTriangleMesh<float>> UnionMesh(new TTriangleMesh<float>(MoveTemp(UnionMeshIndices)));
 				const FMatrix& InertiaMatrix = CollectionSpaceParticles->I(ClusterTransformIdx);
@@ -3131,33 +3195,25 @@ void FGeometryCollectionPhysicsProxy::InitializeSharedCollisionStructures(
 					const TVector<float, 3> Scale = 2 * InstanceBoundingBox.GetExtent() / MaxChildBounds;
 					const float ScaleMax = Scale.GetAbsMax();
 					const float ScaleMin = Scale.GetAbsMin();
-					float MaxResolution = ScaleMax * SizeSpecificData.MaxLevelSetResolution;
+
 					float MinResolution = ScaleMin * SizeSpecificData.MinLevelSetResolution;
-					if (MaxResolution > SizeSpecificData.MaxClusterLevelSetResolution)
-					{
-						MaxResolution = SizeSpecificData.MaxClusterLevelSetResolution;
-					}
-					if (MinResolution > SizeSpecificData.MinClusterLevelSetResolution)
-					{
-						MinResolution = SizeSpecificData.MinClusterLevelSetResolution;
-					}
-					if (MaxResolution < SizeSpecificData.MaxLevelSetResolution)
-					{
-						MaxResolution = SizeSpecificData.MaxLevelSetResolution;
-					}
-					if (MinResolution < SizeSpecificData.MinLevelSetResolution)
-					{
-						MinResolution = SizeSpecificData.MinLevelSetResolution;
-					}
+					MinResolution = FMath::Clamp<float>(MinResolution, SizeSpecificData.MinLevelSetResolution, SizeSpecificData.MinClusterLevelSetResolution);
+
+					float MaxResolution = ScaleMax * SizeSpecificData.MaxLevelSetResolution;
+					MaxResolution = FMath::Clamp<float>(MaxResolution, SizeSpecificData.MaxLevelSetResolution, SizeSpecificData.MaxClusterLevelSetResolution);
 
 					//don't support non level-set serialization
 					ErrorReporter.SetPrefix(BaseErrorPrefix + " | Cluster Transform Index: " + FString::FromInt(ClusterTransformIdx));
 					CollectionImplicits[ClusterTransformIdx] = TSharedPtr<Chaos::FImplicitObject, ESPMode::ThreadSafe>(
-						FCollisionStructureManager::NewImplicit(ErrorReporter, MassSpaceParticles, *UnionMesh,
-							InstanceBoundingBox, 0, MinResolution, MaxResolution,
-							SizeSpecificData.CollisionObjectReductionPercentage, SizeSpecificData.CollisionType,
-							SizeSpecificData.ImplicitType)
-						);
+						FCollisionStructureManager::NewImplicitLevelset(
+							ErrorReporter,
+							MassSpaceParticles,
+							*UnionMesh,
+							InstanceBoundingBox,
+							MinResolution,
+							MaxResolution,
+							SizeSpecificData.CollisionObjectReductionPercentage,
+							SizeSpecificData.CollisionType));
 
 					CollectionSimplicials[ClusterTransformIdx] = TUniquePtr< FSimplicial >(
 						FCollisionStructureManager::NewSimplicial(MassSpaceParticles, *UnionMesh, CollectionImplicits[ClusterTransformIdx].Get(),
@@ -3167,26 +3223,23 @@ void FGeometryCollectionPhysicsProxy::InitializeSharedCollisionStructures(
 				{
 					ErrorReporter.SetPrefix(BaseErrorPrefix + " | Cluster Transform Index: " + FString::FromInt(ClusterTransformIdx));
 					CollectionImplicits[ClusterTransformIdx] = TSharedPtr<Chaos::FImplicitObject, ESPMode::ThreadSafe>(
-						FCollisionStructureManager::NewImplicit(ErrorReporter, MassSpaceParticles, *UnionMesh,
-							InstanceBoundingBox, 0, 0, 0,
-							SizeSpecificData.CollisionObjectReductionPercentage, SizeSpecificData.CollisionType,
-							SizeSpecificData.ImplicitType)
-						);
+						FCollisionStructureManager::NewImplicitBox(
+							InstanceBoundingBox,
+							SizeSpecificData.CollisionObjectReductionPercentage,
+							SizeSpecificData.CollisionType));
 
 					CollectionSimplicials[ClusterTransformIdx] = TUniquePtr< FSimplicial >(
 						FCollisionStructureManager::NewSimplicial(MassSpaceParticles, *UnionMesh, CollectionImplicits[ClusterTransformIdx].Get(),
 						SharedParams.MaximumCollisionParticleCount));
-
 				}
 				else if (SizeSpecificData.ImplicitType == EImplicitTypeEnum::Chaos_Implicit_Sphere)
 				{
 					ErrorReporter.SetPrefix(BaseErrorPrefix + " | Cluster Transform Index: " + FString::FromInt(ClusterTransformIdx));
 					CollectionImplicits[ClusterTransformIdx] = TSharedPtr<Chaos::FImplicitObject, ESPMode::ThreadSafe>(
-						FCollisionStructureManager::NewImplicit(ErrorReporter, MassSpaceParticles, *UnionMesh,
-							InstanceBoundingBox, InstanceBoundingBox.GetExtent().GetAbsMin() / 2, 0, 0,
-							SizeSpecificData.CollisionObjectReductionPercentage, SizeSpecificData.CollisionType,
-							SizeSpecificData.ImplicitType)
-						);
+						FCollisionStructureManager::NewImplicitSphere(
+							InstanceBoundingBox.GetExtent().GetAbsMin() / 2,
+							SizeSpecificData.CollisionObjectReductionPercentage,
+							SizeSpecificData.CollisionType));
 
 					CollectionSimplicials[ClusterTransformIdx] = TUniquePtr< FSimplicial >(
 						FCollisionStructureManager::NewSimplicial(MassSpaceParticles, *UnionMesh, CollectionImplicits[ClusterTransformIdx].Get(),
@@ -3950,15 +4003,12 @@ Chaos::FParticleData* FGeometryCollectionPhysicsProxy::NewData()
 	FGeometryCollectionResults& Buffer = *GameToPhysInterchange->AccessProducerBuffer();
 	if (Buffer.NumParticlesAdded != NumParticles)
 	{
+		Buffer.InitArrays(*GTDynamicCollection);
+
 		// Initialize the buffer
 		Buffer.NumParticlesAdded = NumParticles;
 		Buffer.BaseIndex = BaseParticleIndex;
 
-		TArray<FTransform> Transforms; Transforms.SetNum(NumParticles);
-		Buffer.Transforms = TManagedArray<FTransform>(MoveTemp(Transforms));
-		
-		Buffer.DisabledStates.SetNum(NumParticles);
-
 		Buffer.SharedGeometry.SetNum(NumParticles);
 		Buffer.ShapeSimData.SetNum(NumParticles);
 		Buffer.ShapeQueryData.SetNum(NumParticles);
@@ -3998,7 +4048,7 @@ Chaos::FParticleData* FGeometryCollectionPhysicsProxy::NewData()
 	const FGeometryDynamicCollection* DynamicCollection = Parameters.DynamicCollection;
 	const FGeometryCollection* RestCollection = Parameters.RestCollection;
 
-	Buffer.TransformIndex.Init(RestCollection->TransformIndex);
+	//Buffer.TransformIndex.Init(RestCollection->TransformIndex);
 	Buffer.BoneMap.Init(RestCollection->BoneMap);
 	Buffer.Parent.Init(RestCollection->Parent);
 	Buffer.Children.Init(RestCollection->Children);
@@ -4009,8 +4059,8 @@ Chaos::FParticleData* FGeometryCollectionPhysicsProxy::NewData()
 	Buffer.InertiaTensor.Init(RestCollection->GetAttribute<FVector>("InertiaTensor", FTransformCollection::TransformGroup));
 	Buffer.Transforms.Init(DynamicCollection->Transform);
 
-	GeometryCollectionAlgo::GlobalMatrices(DynamicCollection->Transform, DynamicCollection->Parent, Buffer.GlobalTransforms);
-
+	GeometryCollectionAlgo::GlobalMatrices(
+		DynamicCollection->Transform, DynamicCollection->Parent, Buffer.GlobalTransforms);
 	
 	int32 Idx = 0;
 	for (auto& GTParticle : GTGeometryParticles)
@@ -4071,7 +4121,7 @@ void FGeometryCollectionPhysicsProxy::PushToPhysicsState(const Chaos::FParticleD
 		ensure(false);
 		return;
 	}
-	
+
 	//const FStubGeometryCollectionData* GCData = static_cast<const FStubGeometryCollectionData*>(InData);
 	//const FGeometryCollectionResults* GCState = GCData->GetStateData();
 
@@ -4082,7 +4132,9 @@ void FGeometryCollectionPhysicsProxy::PushToPhysicsState(const Chaos::FParticleD
 		{
 			const int32 Idx = GState->BaseIndex + HandleIdx;
 
-			FTransform ParticleTransform = GState->Transforms[Idx] * Parameters.WorldTransform;
+			FTransform ParticleTransform =
+				MassToLocal[Idx] * GState->Transforms[Idx] * Parameters.WorldTransform;
+
 			Handle->SetX(ParticleTransform.GetTranslation());
 			Handle->SetR(ParticleTransform.GetRotation());
 			Handle->SetDisabled(GState->DisabledStates[Idx]);
@@ -4122,11 +4174,27 @@ void FGeometryCollectionPhysicsProxy::PushToPhysicsState(const Chaos::FParticleD
 				}
 				++CurrShape;
 			}
-
 		} // if rigid handle
 	} // for all SolverParticleHandles
-}
 
+	if (DisableGeometryCollectionGravity) // cvar
+	{
+		// Our assumption is that you'd only ever want to wholesale opt geometry 
+		// collections out of gravity for debugging, so we keep this conditional
+		// out of the loop above and on it's own.  This means we can't turn gravity
+		// back on once it's off, but even if we didn't enclose this in an if(),
+		// this function won't be called again unless something dirties the proxy.
+
+		Chaos::TPerParticleGravity<float, 3>& GravityForces = GetSolver()->GetEvolution()->GetGravityForces();
+		for (int32 HandleIdx = 0; HandleIdx < SolverParticleHandles.Num(); ++HandleIdx)
+		{
+			if (Chaos::TPBDRigidParticleHandle<float, 3>* Handle = SolverParticleHandles[HandleIdx])
+			{
+				GravityForces.SetEnabled(*Handle, false);
+			}
+		}
+	}
+}
 void FGeometryCollectionPhysicsProxy::BufferPhysicsResults()
 {
 	/**
@@ -4140,86 +4208,34 @@ void FGeometryCollectionPhysicsProxy::BufferPhysicsResults()
 	//FGeometryCollectionResults& TargetResults = Results.GetPhysicsDataForWrite();
 	FGeometryCollectionResults& TargetResults = *PhysToGameInterchange->AccessProducerBuffer();
 
-	TManagedArray<FTransform>& TransformCache = TargetResults.Transforms;
-	TArray<FMatrix>& GlobalTransformCache = TargetResults.GlobalTransforms;	
-	//TManagedArray<int32>& IdCache = TargetResults.RigidBodyIds;
-	TManagedArray<int32>& ParentCache = TargetResults.Parent;
-	//TManagedArray<TSet<int32>>& ChildrenCache = TargetResults.Children;
-	TManagedArray<int32>& SimulationTypeCache = TargetResults.SimulationType;
-	//TManagedArray<int32>& StatusFlagsCache = TargetResults.StatusFlags;
-	
-	TargetResults.DynamicState.Init(PTDynamicCollection.DynamicState);
-
-//#if TODO_REIMPLEMENT_GET_RIGID_PARTICLES
-
-	TransformCache.Init(PTDynamicCollection.Transform);
-	//IdCache.Init(RigidBodyID);
-	ParentCache.Init(PTDynamicCollection.Parent);
-	//ChildrenCache.Init(PTDynamicCollection.Children);
-	SimulationTypeCache.Init(PTDynamicCollection.SimulationType);
-	//StatusFlagsCache.Init(PTDynamicCollection.StatusFlags);
-
-	// Base particle index to calculate index from a global particle index on the game thread
-	TargetResults.BaseIndex = BaseParticleIndex;
-	TargetResults.NumParticlesAdded = NumParticles;
-
-	// SQ requires full knowledge of active/inactive particles
-	TargetResults.DisabledStates.Reset(NumParticles);
-
-	// Advertise to game thread
-	TargetResults.IsObjectDynamic = IsObjectDynamic;
-
-	// Advertise to game thread
-	TargetResults.IsObjectLoading = IsObjectLoading;
-	
-	// If object is dynamic, compute global matrices	
-	if (IsObjectDynamic || GlobalTransformCache.Num() == 0)
+	if (TargetResults.NumTransformGroup() != PTDynamicCollection.NumElements(FGeometryCollection::TransformGroup))
 	{
-		SCOPE_CYCLE_COUNTER(STAT_CalcGlobalGCMatrices);
-		check(TransformCache.Num() == ParentCache.Num());
-		GeometryCollectionAlgo::GlobalMatrices(TransformCache, ParentCache, GlobalTransformCache);
-	}
+		TargetResults.InitArrays(PTDynamicCollection);
 
-	// compute world bounds
-	// #note: this is a loose bounds based on the circumscribed box of a bounding sphere for the geometry.		
-	if (IsObjectDynamic || TargetResults.WorldBounds.GetSphere().W < 1e-5)
-	{
-		SCOPE_CYCLE_COUNTER(STAT_CalcGlobalGCBounds);
-		FBox BoundingBox(ForceInit);
-		const FMatrix& ActorToWorld = Parameters.WorldTransform.ToMatrixWithScale();
-		for (int i = 0; i < ValidGeometryBoundingBoxes.Num(); ++i)
-		{
-			BoundingBox += ValidGeometryBoundingBoxes[i].TransformBy(
-				GlobalTransformCache[ValidGeometryTransformIndices[i]] * ActorToWorld);
-		}
-		TargetResults.WorldBounds = FBoxSphereBounds(BoundingBox);		
+		// Base particle index to calculate index from a global particle index on the game thread
+		TargetResults.BaseIndex = BaseParticleIndex;
+		TargetResults.NumParticlesAdded = NumParticles;
 	}
 
-	//Chaos::FPhysicsSolver::FParticlesType& Particles = GetSolver()->GetRigidParticles();
-	//if(NumParticles > 0)
-	{
-		SCOPE_CYCLE_COUNTER(STAT_CaptureDisabledState);
-		TargetResults.DisabledStates.SetNumUninitialized(NumParticles);
-		for (int32 Idx = 0; Idx < NumParticles; Idx++)
+	{ // tmp scope
+		IsObjectDynamic = false;
+		int32 Idx = 0;
+		for (auto* Handle : SolverParticleHandles)
 		{
-			if(!SolverParticleHandles[Idx])
+			if (!Handle || Handle->Disabled())
 			{
-				TargetResults.DisabledStates[Idx] = true;
+				TargetResults.DisabledStates[Idx++] = true;
 				continue;
 			}
-
-			if(SolverParticleHandles[Idx]->Disabled())
-			{
-				TargetResults.DisabledStates[Idx] = true;
-			}
-			else
-			{
-				TargetResults.DisabledStates[Idx] = false;
-				IsObjectDynamic = true;
-			}
+			TargetResults.DisabledStates[Idx++] = false;
+			IsObjectDynamic = true;
 		}
-		//TargetResults.DisabledStates.Append(&Particles.DisabledRef(BaseParticleIndex), NumParticles);
-	}
+	} // tmp scope
+
+	// Advertise to game thread
+	TargetResults.IsObjectDynamic = IsObjectDynamic;
+	TargetResults.IsObjectLoading = IsObjectLoading;
+	
 	const FTransform& ActorToWorld = Parameters.WorldTransform;
 	// It's not worth shrinking the EndFrameUnparentingBuffer array, at least until the solver supports deleting bodies.
 	if (EndFrameUnparentingBuffer.Num() < NumParticles)
@@ -4230,17 +4246,18 @@ void FGeometryCollectionPhysicsProxy::BufferPhysicsResults()
 	{ // tmp scope
 		SCOPE_CYCLE_COUNTER(STAT_CalcParticleToWorld);
 		// Fill particle to world transforms
-		TargetResults.ParticleToWorldTransforms.SetNum(NumParticles);
-/*
-		for (int32 TransformIndex = 0; TransformIndex < NumParticles; ++TransformIndex)
+//		TargetResults.ParticleToWorldTransforms.SetNum(NumParticles);
+/*		for (int32 TransformIndex = 0; TransformIndex < NumParticles; ++TransformIndex)
 		{
 			//only update roots and first children
-			if (ParentCache[TransformIndex] == INDEX_NONE)
+			if (TargetResults.Parent[TransformIndex] == INDEX_NONE)
 			{
-				Chaos::TPBDGeometryCollectionParticleHandle<float, 3>* Handle = SolverParticleHandles[TransformIndex];
-				const FTransform Xf(Handle->R(), Handle->X());
-				//const int32 ParticleIndex = BaseParticleIndex + TransformIndex;
-				TargetResults.ParticleToWorldTransforms[TransformIndex] = Xf;
+				if (auto* Handle = SolverParticleHandles[TransformIndex])
+				{
+					const FTransform Xf(Handle->R(), Handle->X());
+					//const int32 ParticleIndex = BaseParticleIndex + TransformIndex;
+					TargetResults.ParticleToWorldTransforms[TransformIndex] = Xf;
+				}
 			}
 		}
 */
@@ -4251,11 +4268,7 @@ void FGeometryCollectionPhysicsProxy::BufferPhysicsResults()
 
 		for (int32 TransformIndex = 0; TransformIndex < NumParticles; ++TransformIndex)
 		{
-			//only update roots and first children
-			const int32 ParticleIndex = BaseParticleIndex + TransformIndex;
-			//			const int32 ParentIndex = ClusterID[ParticleIndex].Id;
-
-			Chaos::TPBDRigidClusteredParticleHandle<float, 3>* Handle = SolverParticleHandles[ParticleIndex];
+			Chaos::TPBDRigidClusteredParticleHandle<float, 3>* Handle = SolverParticleHandles[TransformIndex];
 			if (!Handle)
 			{
 				continue;
@@ -4304,15 +4317,12 @@ void FGeometryCollectionPhysicsProxy::BufferPhysicsResults()
 				//     non-Identity, and will reflect the difference between the geometric center of the geometry
 				//     and that corresponding rigid bodies center of mass. 
 
-				const FTransform ParticleToWorld(Handle->R(), Handle->X());
+				FTransform& ParticleToWorld = TargetResults.ParticleToWorldTransforms[TransformIndex];
+				ParticleToWorld = FTransform(Handle->R(), Handle->X());
 
-#if 0 // TODO Ryan
 				TargetResults.Transforms[TransformIndex] =
 					MassToLocal[TransformIndex].GetRelativeTransformReverse(ParticleToWorld).GetRelativeTransform(ActorToWorld);
-#else
-				TargetResults.ParticleToWorldTransforms[TransformIndex] = ParticleToWorld;
-				TargetResults.Transforms[TransformIndex] = ParticleToWorld;
-#endif
+				TargetResults.Transforms[TransformIndex].NormalizeRotation();
 
 
 #if 0 // TODO Ryan
@@ -4352,7 +4362,7 @@ void FGeometryCollectionPhysicsProxy::BufferPhysicsResults()
 				}
 #endif
 				// Disabled rigid bodies that have valid cluster parents, and have been re-indexed by the
-				// solver ( As in, They were re-clustered outside of the geometry collection), These clusters 
+				// solver (as in, They were re-clustered outside of the geometry collection), These clusters 
 				// will need to be rendered based on the clusters position. 
 
 				if (TPBDRigidParticleHandle<float, 3>* ParentHandle = Handle->ClusterIds().Id)
@@ -4361,18 +4371,20 @@ void FGeometryCollectionPhysicsProxy::BufferPhysicsResults()
 					{
 						if (ClusteredParentHandle->InternalCluster())
 						{
-							ensure(false); // TODO Ryan!
-#if 0 // TODO Ryan
-							const FTransform ClusterChildToWorld = ClusterChildToParentMap[RigidBodyIndex] * FTransform(Particles.R(ClusterParentIndex), Particles.X(ClusterParentIndex));
+							//const FTransform ClusterChildToWorld = ClusterChildToParentMap[RigidBodyIndex] * FTransform(Particles.R(ClusterParentIndex), Particles.X(ClusterParentIndex));
+							const TRigidTransform<float, 3>& ChildToParent = Handle->ChildToParent();
+							const FTransform ClusterChildToWorld = ChildToParent * FTransform(ParentHandle->R(), ParentHandle->X());
 							if (Parameters.IsCacheRecording())
 							{
-								Particles.X(RigidBodyIndex) = ClusterChildToWorld.GetTranslation();
-								Particles.R(RigidBodyIndex) = ClusterChildToWorld.GetRotation();
+								//Particles.X(RigidBodyIndex) = ClusterChildToWorld.GetTranslation();
+								Handle->SetX(ClusterChildToWorld.GetTranslation());
+								//Particles.R(RigidBodyIndex) = ClusterChildToWorld.GetRotation();
+								Handle->SetR(ClusterChildToWorld.GetRotation());
 							}
 							// GeomToActor = ActorToWorld.Inv() * ClusterChildToWorld * MassToLocal.Inv();
-							Transform[TransformGroupIndex] = MassToLocal[TransformGroupIndex].GetRelativeTransformReverse(ClusterChildToWorld).GetRelativeTransform(ActorToWorld);
-							Transform[TransformGroupIndex].NormalizeRotation();
-#endif
+							TargetResults.Transforms[TransformIndex] = 
+								MassToLocal[TransformIndex].GetRelativeTransformReverse(ClusterChildToWorld).GetRelativeTransform(ActorToWorld);
+							TargetResults.Transforms[TransformIndex].NormalizeRotation();
 						}
 					}
 				}
@@ -4381,6 +4393,30 @@ void FGeometryCollectionPhysicsProxy::BufferPhysicsResults()
 		} // end for
 	} // tmp scope
 
+	// If object is dynamic, compute global matrices	
+	if (IsObjectDynamic || TargetResults.GlobalTransforms.Num() == 0)
+	{
+		SCOPE_CYCLE_COUNTER(STAT_CalcGlobalGCMatrices);
+		check(TargetResults.Transforms.Num() == TargetResults.Parent.Num());
+		GeometryCollectionAlgo::GlobalMatrices(
+			TargetResults.Transforms, TargetResults.Parent, TargetResults.GlobalTransforms);
+	}
+
+	// compute world bounds
+	// #note: this is a loose bounds based on the circumscribed box of a bounding sphere for the geometry.		
+	if (IsObjectDynamic || TargetResults.WorldBounds.GetSphere().W < 1e-5)
+	{
+		SCOPE_CYCLE_COUNTER(STAT_CalcGlobalGCBounds);
+		FBox BoundingBox(ForceInit);
+		//const FMatrix& ActorToWorld = Parameters.WorldTransform.ToMatrixWithScale();
+		for (int i = 0; i < ValidGeometryBoundingBoxes.Num(); ++i)
+		{
+			BoundingBox += ValidGeometryBoundingBoxes[i].TransformBy(
+				TargetResults.GlobalTransforms[ValidGeometryTransformIndices[i]] * ActorToWorld);
+		}
+		TargetResults.WorldBounds = FBoxSphereBounds(BoundingBox);		
+	}
+
 #if 0 // TODO Ryan
 	for (int32 TransformGroupIndex = 0; TransformGroupIndex < TransformSize; TransformGroupIndex++)
 	{
@@ -4396,6 +4432,9 @@ void FGeometryCollectionPhysicsProxy::BufferPhysicsResults()
 	}
 #endif
 
+	// We pulled state out of particles.  Push results into PTDynamicCollection, 
+	// so that it has an updated view.
+	UpdateGeometryCollection(TargetResults);
 
 	// We've updated a buffer.  There's no reason to wait for FGeometryCollectionPhysicsProxy::FlipBuffer()
 	// to be called.  Just flip.
@@ -4427,23 +4466,13 @@ void FGeometryCollectionPhysicsProxy::PullFromPhysicsState()
 	 * Note: A read lock will have been acquired for this - so the physics thread won't force a buffer flip while this
 	 * sync is ongoing
 	 */
-/* 
-	// We don't need to do this if we're using a guarded buffer, as it'll return 
-	// null if it doesn't have new data for us.
-	const uint32 LastSyncCountFromPhysics = Results.GetGameDataSyncCount();
-	if (LastSyncCountFromPhysics == LastSyncCountGT)
-	{
-		return;
-	}
-	LastSyncCountGT = LastSyncCountFromPhysics;
-*/
-	//FGeometryCollectionResults& TargetResult = Results.GetGameDataForWrite();
+
 	const FGeometryCollectionResults* TargetResultPtr = PhysToGameInterchange->GetConsumerBuffer();
 	if (!TargetResultPtr)
 		return;
 	// Remove const-ness as we're going to do the ExchangeArrays() thing.
 	FGeometryCollectionResults& TR = *const_cast<FGeometryCollectionResults*>(TargetResultPtr);
-/**/
+
 	// We should never be changing the number of entries, this would break other 
 	// attributes in the transform group.
 	const int32 NumTransforms = GTDynamicCollection->Transform.Num();
@@ -4453,15 +4482,7 @@ void FGeometryCollectionPhysicsProxy::PullFromPhysicsState()
 		{
 			if(!TR.DisabledStates[TmIndex])
 			{
-#if 0 // TODO Ryan
-				const FTransform& ParticleToWorld = TR.ParticleToWorldTransforms[TmIndex];
-				GTDynamicCollection->Transform[TmIndex] = 
-					MassToLocal[TmIndex].GetRelativeTransformReverse(ParticleToWorld).GetRelativeTransform(
-						Parameters.WorldTransform);
-#else
-				GTDynamicCollection->Transform[TmIndex] = MassToLocal[TmIndex].GetRelativeTransformReverse(TR.Transforms[TmIndex]).GetRelativeTransform(Parameters.WorldTransform);
-				GTDynamicCollection->Transform[TmIndex].NormalizeRotation();
-#endif
+				GTDynamicCollection->Transform[TmIndex] = TR.Transforms[TmIndex];
 			}
 		}
 
@@ -4473,7 +4494,6 @@ void FGeometryCollectionPhysicsProxy::PullFromPhysicsState()
 			CacheSyncFunc(TR);
 		}
 	}
-/**/
 /*
 	for (int32 TransformIndex = 0; TransformIndex < NumParticles; ++TransformIndex)
 	{
@@ -4492,15 +4512,20 @@ void FGeometryCollectionPhysicsProxy::PullFromPhysicsState()
 //	UpdateGeometryCollection(TR);
 }
 
-void FGeometryCollectionPhysicsProxy::UpdateGeometryCollection(FGeometryCollectionResults& GCResults)
+void FGeometryCollectionPhysicsProxy::UpdateGeometryCollection(
+	const FGeometryCollectionResults& GCResults,
+	FGeometryDynamicCollection* Collection)
 {
 	if (InitializedState != ESimulationInitializationState::Initialized)
 	{
 		return;
 	}
-
-	FGeometryDynamicCollection* Collection = Parameters.DynamicCollection;
+	if (!Collection)
+	{
+		Collection = Parameters.DynamicCollection;
+	}
 	check(Collection);
+
 	TManagedArray<FTransform>& Transform = Collection->Transform;
 	TManagedArray<int32>& Parent = Collection->Parent;
 	TManagedArray<TSet<int32>>& Children = Collection->Children;
@@ -4522,6 +4547,14 @@ void FGeometryCollectionPhysicsProxy::UpdateGeometryCollection(FGeometryCollecti
 		if (!GTGeometryParticles[TransformGroupIndex])
 			continue;
 
+		auto* Handle = SolverParticleHandles[TransformGroupIndex];
+
+		// dynamic state is also updated by the solver during field interaction. 
+		if(!Handle->Sleeping())
+		{
+			DynamicState[TransformGroupIndex] = GCResults.DynamicState[TransformGroupIndex];
+		}
+
 		// Update the transform and parent hierarchy of the active rigid bodies. Active bodies can be either
 		// rigid geometry defined from the leaf nodes of the collection, or cluster bodies that drive an entire
 		// branch of the hierarchy within the GeometryCollection.
@@ -4531,7 +4564,7 @@ void FGeometryCollectionPhysicsProxy::UpdateGeometryCollection(FGeometryCollecti
 		//   take ownership of the parents during the simulation, so it might be necessary to force deactivated
 		//   bodies out of the collections hierarchy during the simulation.  
 
-		if (/*!Particles.Disabled(RigidBodyID[TransformGroupIndex])*/true)
+		if(!GCResults.DisabledStates[TransformGroupIndex])
 		{
 			// Update the transform of the active body. The active body can be either a single rigid
 			// or a collection of rigidly attached geometries (Clustering). The cluster is represented as a
@@ -4543,46 +4576,8 @@ void FGeometryCollectionPhysicsProxy::UpdateGeometryCollection(FGeometryCollecti
 			//     non-Identity, and will reflect the difference between the geometric center of the geometry
 			//     and that corresponding rigid bodies center of mass. 
 
-			//const FTransform ParticleToWorld(Particles.R(RigidBodyIndex), Particles.X(RigidBodyIndex));
-			const FTransform& ParticleToWorld = GCResults.ParticleToWorldTransforms[TransformGroupIndex];
-				//Collection->Transform[TransformGroupIndex];
-
 			// GeomToActor = ActorToWorld.Inv() * ParticleToWorld * MassToLocal.Inv();
-#if 0 // TODO Ryan
-			Transform[TransformGroupIndex] = 
-				MassToLocal[TransformGroupIndex].GetRelativeTransformReverse(ParticleToWorld).GetRelativeTransform(
-					ActorToWorld);
-#else
-			Transform[TransformGroupIndex] = ParticleToWorld;
-			Transform[TransformGroupIndex].NormalizeRotation();
-#endif
-			// dynamic state is also updated by the solver during field interaction. 
-			if (/*!Particles.Sleeping(RigidBodyIndex)*/true)
-			{
-/*
-				if (Particles.ObjectState(RigidBodyIndex) == Chaos::EObjectStateType::Kinematic)
-				{
-					DynamicState[TransformGroupIndex] = (int)EObjectStateTypeEnum::Chaos_Object_Kinematic;
-				}
-				else if (Particles.ObjectState(RigidBodyIndex) == Chaos::EObjectStateType::Static)
-				{
-					DynamicState[TransformGroupIndex] = (int)EObjectStateTypeEnum::Chaos_Object_Static;
-				}
-				else
-				{
-*/
-					DynamicState[TransformGroupIndex] = (int)EObjectStateTypeEnum::Chaos_Object_Dynamic;
-//				}
-			}
-
-			// Force all enabled rigid bodies out of the transform hierarchy
-			if (Parent[TransformGroupIndex] != INDEX_NONE)
-			{
-				const int32 ParentIndex = Parent[TransformGroupIndex];
-				// Children in the hierarchy are stored in a TSet, which is not thread safe.  So we retain
-				// indices to remove afterwards.
-				EndFrameUnparentingBuffer[TransformGroupIndex] = ParentIndex;
-			}
+			Transform[TransformGroupIndex] = GCResults.Transforms[TransformGroupIndex];
 
 			// When a leaf node rigid body is removed from a cluster the rigid
 			// body will become active and needs its clusterID updated. This just
@@ -4591,7 +4586,7 @@ void FGeometryCollectionPhysicsProxy::UpdateGeometryCollection(FGeometryCollecti
 		}
 		else //if (Particles.Disabled(RigidBodyIndex))
 		{
-			//...
+			Transform[TransformGroupIndex] = GCResults.Transforms[TransformGroupIndex];
 		}
 	}
 }
diff --git a/Engine/Source/Runtime/Experimental/ChaosSolvers/Public/PhysicsProxy/GeometryCollectionPhysicsProxy.h b/Engine/Source/Runtime/Experimental/ChaosSolvers/Public/PhysicsProxy/GeometryCollectionPhysicsProxy.h
--- a/Engine/Source/Runtime/Experimental/ChaosSolvers/Public/PhysicsProxy/GeometryCollectionPhysicsProxy.h
+++ b/Engine/Source/Runtime/Experimental/ChaosSolvers/Public/PhysicsProxy/GeometryCollectionPhysicsProxy.h
@@ -33,11 +33,29 @@ public:
 
 	void Reset();
 
+	int32 NumTransformGroup() const { return Transforms.Num(); }
+
+	void InitArrays(const FGeometryDynamicCollection &Other)
+	{
+		// Managed arrays
+		Transforms.Init(Other.Transform);
+		DynamicState.Init(Other.DynamicState);
+		Parent.Init(Other.Parent);
+		Children.Init(Other.Children);
+		SimulationType.Init(Other.SimulationType);
+		
+		// Arrays
+		const int32 NumTransforms = Other.NumElements(FGeometryCollection::TransformGroup);
+		DisabledStates.SetNumUninitialized(NumTransforms);
+		GlobalTransforms.SetNumUninitialized(NumTransforms);
+		ParticleToWorldTransforms.SetNumUninitialized(NumTransforms);
+	}
+
 	int32 BaseIndex;
 	int32 NumParticlesAdded;
 	TArray<bool> DisabledStates;
-
 	TArray<FMatrix> GlobalTransforms;
+	TArray<FTransform> ParticleToWorldTransforms;
 
 	TManagedArray<int32> TransformIndex;
 
@@ -52,8 +70,6 @@ public:
 
 	TManagedArray<int32> ClusterId;
 
-	TArray<FTransform> ParticleToWorldTransforms;
-
 	TArray<TSharedPtr<Chaos::FImplicitObject, ESPMode::ThreadSafe>> SharedGeometry;
 	TArray<TArray<FCollisionFilterData>> ShapeSimData;
 	TArray<TArray<FCollisionFilterData>> ShapeQueryData;
@@ -211,7 +227,7 @@ public:
 	//
 	void SetCollisionParticlesPerObjectFraction(float CollisionParticlesPerObjectFractionIn) {CollisionParticlesPerObjectFraction = CollisionParticlesPerObjectFractionIn;}
 
-	const FGeometryCollectionResults* GetConsumerResultsGT() { return PhysToGameInterchange->GetConsumerBuffer(); }
+	const FGeometryCollectionResults* GetConsumerResultsGT() { return PhysToGameInterchange->PeekConsumerBuffer(); }
 
 	/*
 	* For Testing ONLY!
@@ -220,7 +236,9 @@ public:
 	const TManagedArray<FTransform>& MassToLocal_GameThreadAccess() const { return MassToLocal; }
 
 protected:
-	void UpdateGeometryCollection(FGeometryCollectionResults& Results);
+	void UpdateGeometryCollection(
+		const FGeometryCollectionResults& Results,
+		FGeometryDynamicCollection* TargetCollection=nullptr);
 
 protected:
 	int32 CalculateHierarchyLevel(FGeometryDynamicCollection* GeometryCollection, int32 TransformIndex);
@@ -335,6 +353,7 @@ private:
 
 	// Double buffer of geom collection result data - TODO (Ryan): deprecated?
 	Chaos::TBufferedData<FGeometryCollectionResults> Results;
+	Chaos::TBufferedData<TArray<FMatrix>> ResultGlobalMatrices;
 
 	// TODO (Ryan) - Currently this is using triple buffers for game-physics and 
 	// physics-game thread communication, but not for any reason other than this 
@@ -345,8 +364,8 @@ private:
 	// currently managing the exchange is built upon.  However, I believe that 
 	// logic locks, and the triple buffer would enable a decoupled lock-free 
 	// paradigm, at least for this component of the handshake.
-	TUniquePtr<Chaos::IBufferResource<FGeometryCollectionResults>> PhysToGameInterchange;
-	TUniquePtr<Chaos::IBufferResource<FGeometryCollectionResults>> GameToPhysInterchange;
+	TUniquePtr<Chaos::FGuardedTripleBuffer<FGeometryCollectionResults>> PhysToGameInterchange;
+	TUniquePtr<Chaos::FGuardedTripleBuffer<FGeometryCollectionResults>> GameToPhysInterchange;
 };
 
 CHAOSSOLVERS_API void BuildSimulationData(Chaos::FErrorReporter& ErrorReporter, FGeometryCollection& GeometryCollection, const FSharedSimulationParameters& SharedParams);
diff --git a/Engine/Source/Runtime/Experimental/GeometryCollectionSimulationCore/Private/GeometryCollection/GeometryCollectionCollisionStructureManager.cpp b/Engine/Source/Runtime/Experimental/GeometryCollectionSimulationCore/Private/GeometryCollection/GeometryCollectionCollisionStructureManager.cpp
--- a/Engine/Source/Runtime/Experimental/GeometryCollectionSimulationCore/Private/GeometryCollection/GeometryCollectionCollisionStructureManager.cpp
+++ b/Engine/Source/Runtime/Experimental/GeometryCollectionSimulationCore/Private/GeometryCollection/GeometryCollectionCollisionStructureManager.cpp
@@ -33,9 +33,9 @@ FAutoConsoleVariableRef CVarCollisionParticlesMax(TEXT("p.CollisionParticlesMax"
 FCollisionStructureManager::FSimplicial*
 FCollisionStructureManager::NewSimplicial(
 	const Chaos::TParticles<float, 3>& Vertices,
-	Chaos::TTriangleMesh<float>& TriMesh,
+	const Chaos::TTriangleMesh<float>& TriMesh,
 	const Chaos::FImplicitObject* Implicit,
-	int32 CollisionParticlesMaxInput)
+	const int32 CollisionParticlesMaxInput)
 {
 	FCollisionStructureManager::FSimplicial * Simplicial = new FCollisionStructureManager::FSimplicial();
 	if (Implicit || Vertices.Size())
@@ -124,22 +124,19 @@ FCollisionStructureManager::NewSimplicial(
 	return Simplicial;
 }
 
-
-
 FCollisionStructureManager::FSimplicial*
 FCollisionStructureManager::NewSimplicial(
 	const Chaos::TParticles<float,3>& AllParticles,
 	const TManagedArray<int32>& BoneMap,
 	const ECollisionTypeEnum CollisionType,
 	Chaos::TTriangleMesh<float>& TriMesh,
-	const float CollisionParticlesFraction
-)
+	const float CollisionParticlesFraction)
 {
 	const bool bEnableCollisionParticles = (CollisionType == ECollisionTypeEnum::Chaos_Surface_Volumetric);
 	if (bEnableCollisionParticles || true)
 	{
 		// @todo : Clean collision particles need to operate on the collision mask from the DynamicCollection,
-		//         then transfer only the good collision particles during the initialization. 
+		//         then transfer only the good collision particles during the initialization. `
 		FCollisionStructureManager::FSimplicial * Simplicial = new FCollisionStructureManager::FSimplicial();
 		const TArrayView<const Chaos::TVector<float, 3>> ArrayView(&AllParticles.X(0), AllParticles.Size());
 		const TArray<Chaos::TVector<float,3>>& Result = Chaos::CleanCollisionParticles(TriMesh, ArrayView, CollisionParticlesFraction);
@@ -172,8 +169,9 @@ FCollisionStructureManager::NewSimplicial(
 	return new FCollisionStructureManager::FSimplicial();
 }
 
-
-void FCollisionStructureManager::UpdateImplicitFlags(FImplicit* Implicit, ECollisionTypeEnum CollisionType)
+void FCollisionStructureManager::UpdateImplicitFlags(
+	FImplicit* Implicit, 
+	const ECollisionTypeEnum CollisionType)
 {
 	if (Implicit && (CollisionType == ECollisionTypeEnum::Chaos_Surface_Volumetric))
 	{
@@ -182,15 +180,93 @@ void FCollisionStructureManager::UpdateImplicitFlags(FImplicit* Implicit, EColli
 	}
 }
 
+FCollisionStructureManager::FImplicit* 
+FCollisionStructureManager::NewImplicit(
+	Chaos::FErrorReporter& ErrorReporter,
+	const Chaos::TParticles<float, 3>& MeshParticles,
+	const Chaos::TTriangleMesh<float>& TriMesh,
+	const FBox& CollisionBounds,
+	const float Radius, 
+	const int32 MinRes,
+	const int32 MaxRes,
+	const float CollisionObjectReduction,
+	const ECollisionTypeEnum CollisionType,
+	const EImplicitTypeEnum ImplicitType)
+{
+	if (ImplicitType == EImplicitTypeEnum::Chaos_Implicit_Box)
+	{
+		return NewImplicitBox(CollisionBounds, CollisionObjectReduction, CollisionType);
+	}
+	else if (ImplicitType == EImplicitTypeEnum::Chaos_Implicit_Sphere)
+	{
+		return NewImplicitSphere(Radius, CollisionObjectReduction, CollisionType);
+	}
+	else if (ImplicitType == EImplicitTypeEnum::Chaos_Implicit_LevelSet)
+	{
+		return NewImplicitLevelset(ErrorReporter, MeshParticles, TriMesh, CollisionBounds, MinRes, MaxRes, CollisionObjectReduction, CollisionType);
+	}
+	
+	return nullptr;
+}
+
+FCollisionStructureManager::FImplicit*
+FCollisionStructureManager::NewImplicitBox(
+	const FBox& CollisionBounds,
+	const float CollisionObjectReduction,
+	const ECollisionTypeEnum CollisionType)
+{
+	const FVector HalfExtents = CollisionBounds.GetExtent() * (1 - CollisionObjectReduction / 100.f);
+	const FVector Center = CollisionBounds.GetCenter();
+	Chaos::FImplicitObject* Implicit = new Chaos::TBox<float, 3>(Center - HalfExtents, Center + HalfExtents);
+	UpdateImplicitFlags(Implicit, CollisionType);
+	return Implicit;
+}
+
+FCollisionStructureManager::FImplicit*
+FCollisionStructureManager::NewImplicitSphere(
+	const float Radius,
+	const float CollisionObjectReduction,
+	const ECollisionTypeEnum CollisionType)
+{
+	Chaos::FImplicitObject* Implicit = new Chaos::TSphere<float, 3>(Chaos::TVector<float, 3>(0), Radius * (1 - CollisionObjectReduction / 100.f));
+	UpdateImplicitFlags(Implicit, CollisionType);
+	return Implicit;
+}
+
+FCollisionStructureManager::FImplicit*
+FCollisionStructureManager::NewImplicitLevelset(
+	Chaos::FErrorReporter& ErrorReporter,
+	const Chaos::TParticles<float, 3>& MeshParticles,
+	const Chaos::TTriangleMesh<float>& TriMesh,
+	const FBox& CollisionBounds,
+	const int32 MinRes,
+	const int32 MaxRes,
+	const float CollisionObjectReduction,
+	const ECollisionTypeEnum CollisionType)
+{
+	FVector HalfExtents = CollisionBounds.GetExtent();
+	if (HalfExtents.GetAbsMin() < KINDA_SMALL_NUMBER)
+	{
+		return nullptr;
+	}
+	HalfExtents *= CollisionObjectReduction / 100.f;
+	const float MinExtent = FMath::Min(HalfExtents[0], FMath::Min(HalfExtents[1], HalfExtents[2]));
+	Chaos::TLevelSet<float, 3>* LevelSet = NewLevelset(ErrorReporter, MeshParticles, TriMesh, CollisionBounds, MinRes, MaxRes, CollisionType);
+	if (LevelSet && MinExtent > 0)
+	{
+		LevelSet->Shrink(MinExtent);
+	}
+	return LevelSet;
+}
+
 Chaos::TLevelSet<float, 3>* FCollisionStructureManager::NewLevelset(
 	Chaos::FErrorReporter& ErrorReporter,
 	const Chaos::TParticles<float, 3>& MeshParticles,
 	const Chaos::TTriangleMesh<float>& TriMesh,
 	const FBox& CollisionBounds,
-	int32 MinRes,
-	int32 MaxRes,
-	ECollisionTypeEnum CollisionType
-	)
+	const int32 MinRes,
+	const int32 MaxRes,
+	const ECollisionTypeEnum CollisionType)
 {
 	Chaos::TVector<int32, 3> Counts;
 	const FVector Extents = CollisionBounds.GetExtent();
@@ -236,51 +312,6 @@ Chaos::TLevelSet<float, 3>* FCollisionStructureManager::NewLevelset(
 	return Implicit;
 }
 
-FCollisionStructureManager::FImplicit* 
-FCollisionStructureManager::NewImplicit(
-	Chaos::FErrorReporter& ErrorReporter,
-	const Chaos::TParticles<float, 3>& MeshParticles,
-	const Chaos::TTriangleMesh<float>& TriMesh,
-	const FBox& CollisionBounds,
-	const float Radius, 
-	const int32 MinRes,
-	const int32 MaxRes,
-	const float CollisionObjectReduction,
-	const ECollisionTypeEnum CollisionType,
-	const EImplicitTypeEnum ImplicitType)
-{
-	Chaos::FImplicitObject* Implicit = nullptr;
-	if (ImplicitType == EImplicitTypeEnum::Chaos_Implicit_Box)
-	{
-		FVector HalfExtents = CollisionBounds.GetExtent() * (1 - CollisionObjectReduction / 100.f);
-		FVector Center = CollisionBounds.GetCenter();
-		Implicit = new Chaos::TBox<float, 3>(Center - HalfExtents, Center + HalfExtents);
-	}
-	else if (ImplicitType == EImplicitTypeEnum::Chaos_Implicit_Sphere)
-	{
-		Implicit = new Chaos::TSphere<float, 3>(Chaos::TVector<float, 3>(0), Radius * (1 - CollisionObjectReduction / 100.f));
-	}
-	else if (ImplicitType == EImplicitTypeEnum::Chaos_Implicit_LevelSet)
-	{
-		FVector HalfExtents = CollisionBounds.GetExtent();
-		if (HalfExtents.X < KINDA_SMALL_NUMBER || HalfExtents.Y < KINDA_SMALL_NUMBER || HalfExtents.Z < KINDA_SMALL_NUMBER)
-		{
-			return nullptr;
-		}
-		HalfExtents *= CollisionObjectReduction / 100.f;
-		float MinExtent = FMath::Min(HalfExtents[0], FMath::Min(HalfExtents[1], HalfExtents[2]));
-		Chaos::TLevelSet<float, 3>* LevelSet = NewLevelset(ErrorReporter, MeshParticles, TriMesh, CollisionBounds, MinRes, MaxRes, CollisionType);
-		if (LevelSet && MinExtent > 0)
-		{
-			LevelSet->Shrink(MinExtent);
-		}
-		return LevelSet;
-	}
-	
-	UpdateImplicitFlags(Implicit, CollisionType);
-	return Implicit;
-}
-
 FVector 
 FCollisionStructureManager::CalculateUnitMassInertiaTensor(
 	const FBox& Bounds,
diff --git a/Engine/Source/Runtime/Experimental/GeometryCollectionSimulationCore/Public/GeometryCollection/GeometryCollectionCollisionStructureManager.h b/Engine/Source/Runtime/Experimental/GeometryCollectionSimulationCore/Public/GeometryCollection/GeometryCollectionCollisionStructureManager.h
--- a/Engine/Source/Runtime/Experimental/GeometryCollectionSimulationCore/Public/GeometryCollection/GeometryCollectionCollisionStructureManager.h
+++ b/Engine/Source/Runtime/Experimental/GeometryCollectionSimulationCore/Public/GeometryCollection/GeometryCollectionCollisionStructureManager.h
@@ -30,57 +30,100 @@ public:
 	typedef Chaos::TBVHParticles<float,3> FSimplicial;
 	typedef Chaos::FImplicitObject FImplicit;
 
+	/**
+	 * Culls particles inside \p Implicit, and coincident particles (defined by 
+	 * being within 1/20'th of the size of the domain from one another), truncates
+	 * at \p CollisionParticlesMaxInput, and returns a bounding volume hierarchy
+	 * of the remainder.
+	 */
 	static FSimplicial* NewSimplicial(
 		const Chaos::TParticles<float, 3>& Vertices,
-		Chaos::TTriangleMesh<float>& TriMesh,
+		const Chaos::TTriangleMesh<float>& TriMesh,
 		const Chaos::FImplicitObject* Implicit,
-		int32 CollisionParticlesMaxInput
-	);
+		const int32 CollisionParticlesMaxInput);
 	
-	
-	static FSimplicial * NewSimplicial(
+	/**
+	 * Culls particles by importance (See \c TTriangleMesh::GetVertexImportanceOrdering()),
+	 * and returns a bounding volume hierarchy of the remainder.
+	 */
+	static FSimplicial* NewSimplicial(
 		const Chaos::TParticles<float,3>& AllParticles,
 		const TManagedArray<int32>& BoneMap,
 		const ECollisionTypeEnum CollisionType,
 		Chaos::TTriangleMesh<float>& TriMesh,
-		const float CollisionParticlesFraction
-	);
+		const float CollisionParticlesFraction);
 
+	/**
+	 * Calls \c SetDoCollide(false) and \c SetConvex(false) on \p Implicit if 
+	 * \p CollisionType is \c ECollisionTypeEnum::Chaos_Surface_Volumetric.
+	 */
+	static void UpdateImplicitFlags(
+		FImplicit* Implicit, 
+		const ECollisionTypeEnum CollisionType);
+
+	/**
+	 * Build a box, sphere, or level set based on \p ImplicitType.
+	 * 
+	 *	\p ErrorReporter - level set only
+	 *	\p MeshParticles - level set only
+	 *	\p TriMesh - level set only
+	 *	\p CollisionBounds - box and level set
+	 *	\p Radius - sphere only
+	 *	\p MinRes - level set only
+	 *	\p MaxRes - level set only
+	 *	\p CollisionObjectReduction - shrink percentage; value of 10 reduces by 
+	 *	   10%, 0 does nothing, 100 shrinks to zero.
+	 *	\p CollisionType - param forwarded to \c UpdateImplictFlags().
+	 *	\p ImplicitType - type of implicit shape to build.
+	 */
 	static FImplicit * NewImplicit(
 		Chaos::FErrorReporter& ErrorReporter,
 		const Chaos::TParticles<float, 3>& MeshParticles,
 		const Chaos::TTriangleMesh<float>& TriMesh,
-		const FBox& CollisionBoundsArray,
+		const FBox& CollisionBounds,
 		const float Radius,
 		const int32 MinRes,
 		const int32 MaxRes,
 		const float CollisionObjectReduction,
 		const ECollisionTypeEnum CollisionType,
-		const EImplicitTypeEnum ImplicitType
-	);
+		const EImplicitTypeEnum ImplicitType);
 
-	static FVector CalculateUnitMassInertiaTensor(
-		const FBox& BoundingBox,
-		const float Radius,
-		const EImplicitTypeEnum ImplicitType
-	);
+	static FImplicit* NewImplicitBox(
+		const FBox& CollisionBounds,
+		const float CollisionObjectReduction,
+		const ECollisionTypeEnum CollisionType);
 
-	static float CalculateVolume(
-		const FBox& BoundingBox,
+	static FImplicit* NewImplicitSphere(
 		const float Radius,
-		const EImplicitTypeEnum ImplicitType
-	);
+		const float CollisionObjectReduction,
+		const ECollisionTypeEnum CollisionType);
 
+	static FImplicit* NewImplicitLevelset(
+		Chaos::FErrorReporter& ErrorReporter,
+		const Chaos::TParticles<float, 3>& MeshParticles,
+		const Chaos::TTriangleMesh<float>& TriMesh,
+		const FBox& CollisionBounds,
+		const int32 MinRes,
+		const int32 MaxRes,
+		const float CollisionObjectReduction,
+		const ECollisionTypeEnum CollisionType);
 
 	static Chaos::TLevelSet<float, 3>* NewLevelset(
 		Chaos::FErrorReporter& ErrorReporter,
 		const Chaos::TParticles<float, 3>& MeshParticles,
 		const Chaos::TTriangleMesh<float>& TriMesh,
 		const FBox& CollisionBounds,
-		int32 MinRes,
-		int32 MaxRes,
-		ECollisionTypeEnum CollisionType
-	);
+		const int32 MinRes,
+		const int32 MaxRes,
+		const ECollisionTypeEnum CollisionType);
 
-	static void UpdateImplicitFlags(FImplicit* Implicit, ECollisionTypeEnum CollisionType);
+	static FVector CalculateUnitMassInertiaTensor(
+		const FBox& BoundingBox,
+		const float Radius,
+		const EImplicitTypeEnum ImplicitType);
+
+	static float CalculateVolume(
+		const FBox& BoundingBox,
+		const float Radius,
+		const EImplicitTypeEnum ImplicitType);
 };
