Merge 11447502 to Release-4.24.

Add error checking and reporting for level set generation and geometry collection collision shape initialization.

diff --git a/Engine/Source/Runtime/Experimental/Chaos/Public/Chaos/AABB.h b/Engine/Source/Runtime/Experimental/Chaos/Public/Chaos/AABB.h
--- a/Engine/Source/Runtime/Experimental/Chaos/Public/Chaos/AABB.h
+++ b/Engine/Source/Runtime/Experimental/Chaos/Public/Chaos/AABB.h
@@ -472,6 +472,7 @@ namespace Chaos
 
 		FORCEINLINE static TAABB<T, d> EmptyAABB() { return TAABB<T, d>(TVector<T, d>(TNumericLimits<T>::Max()), TVector<T, d>(-TNumericLimits<T>::Max())); }
 		FORCEINLINE static TAABB<T, d> ZeroAABB() { return TAABB<T, d>(TVector<T, d>((T)0), TVector<T, d>((T)0)); }
+		FORCEINLINE static TAABB<T, d> FullAABB() { return TAABB<T, d>(TVector<T, d>(-TNumericLimits<T>::Max()), TVector<T, d>(TNumericLimits<T>::Max())); }
 
 		FORCEINLINE void Serialize(FArchive &Ar) 
 		{
diff --git a/Engine/Source/Runtime/Experimental/Chaos/Public/Chaos/Levelset.h b/Engine/Source/Runtime/Experimental/Chaos/Public/Chaos/Levelset.h
--- a/Engine/Source/Runtime/Experimental/Chaos/Public/Chaos/Levelset.h
+++ b/Engine/Source/Runtime/Experimental/Chaos/Public/Chaos/Levelset.h
@@ -80,6 +80,24 @@ class CHAOS_API TLevelSet final : public FImplicitObject
 		SerializeImp(Ar);
 	}
 
+	/** Do the simplest thing possible to estimate the volume bounded by 
+	 * the zero'th isocontour of the level set. 
+	 */
+	T ApproximateNegativeMaterial() const
+	{
+		const TVector<T,d>& CellDim = MGrid.Dx();
+		const T QuarterCellVolume = CellDim.Product() * 0.25;
+		T Volume = 0.0;
+		for (int32 Idx = 0; Idx < MPhi.Num(); ++Idx)
+		{
+			if (MPhi[Idx] <= 0.0)
+			{
+				Volume += QuarterCellVolume;
+			}
+		}
+		return Volume;
+	}
+
 	bool ComputeMassProperties(T& OutVolume, TVector<T, d>& OutCOM, PMatrix<T,d,d>& OutInertia, TRotation<T, d>& OutRotationOfMass) const;
 
 	T ComputeLevelSetError(const TParticles<T, d>& InParticles, const TArray<TVector<T, 3>>& Normals, const TTriangleMesh<T>& Mesh, T& AngleError, T& MaxDistError);
diff --git a/Engine/Source/Runtime/Experimental/ChaosSolvers/Private/PhysicsProxy/GeometryCollectionPhysicsProxy.cpp b/Engine/Source/Runtime/Experimental/ChaosSolvers/Private/PhysicsProxy/GeometryCollectionPhysicsProxy.cpp
--- a/Engine/Source/Runtime/Experimental/ChaosSolvers/Private/PhysicsProxy/GeometryCollectionPhysicsProxy.cpp
+++ b/Engine/Source/Runtime/Experimental/ChaosSolvers/Private/PhysicsProxy/GeometryCollectionPhysicsProxy.cpp
@@ -296,11 +296,18 @@ void PopulateSimulatedParticle(
 		TUniquePtr<Chaos::TBVHParticles<float, 3>>& CollisionParticles = Handle->CollisionParticles();
 		if (Simplicial->Size())
 		{
+			const Chaos::TAABB<float,3> ImplicitShapeDomain = 
+				Implicit && Implicit->HasBoundingBox() ? Implicit->BoundingBox() : Chaos::TAABB<float,3>::FullAABB();
+
 			CollisionParticles->Resize(0);
 			CollisionParticles->AddParticles(Simplicial->Size());
 			for (int32 VertexIndex = 0; VertexIndex < (int32)Simplicial->Size(); ++VertexIndex)
 			{
 				CollisionParticles->X(VertexIndex) = Simplicial->X(VertexIndex);
+
+				// Make sure the collision particles are at least in the domain 
+				// of the implicit shape.
+				ensure(ImplicitShapeDomain.Contains(CollisionParticles->X(VertexIndex)));
 			}
 		}
 
@@ -3002,7 +3009,7 @@ void FGeometryCollectionPhysicsProxy::InitializeSharedCollisionStructures(
 
 			if (SizeSpecificData.ImplicitType == EImplicitTypeEnum::Chaos_Implicit_LevelSet)
 			{
-				ErrorReporter.SetPrefix(BaseErrorPrefix + " | Transform Index: " + FString::FromInt(TransformGroupIndex));
+				ErrorReporter.SetPrefix(BaseErrorPrefix + " | Transform Index: " + FString::FromInt(TransformGroupIndex) + " of " + FString::FromInt(TransformIndex.Num()));
 				CollectionImplicits[TransformGroupIndex] = TSharedPtr<Chaos::FImplicitObject, ESPMode::ThreadSafe>(
 					FCollisionStructureManager::NewImplicitLevelset(
 						ErrorReporter,
@@ -3221,7 +3228,6 @@ void FGeometryCollectionPhysicsProxy::InitializeSharedCollisionStructures(
 				}
 				else if (SizeSpecificData.ImplicitType == EImplicitTypeEnum::Chaos_Implicit_Box)
 				{
-					ErrorReporter.SetPrefix(BaseErrorPrefix + " | Cluster Transform Index: " + FString::FromInt(ClusterTransformIdx));
 					CollectionImplicits[ClusterTransformIdx] = TSharedPtr<Chaos::FImplicitObject, ESPMode::ThreadSafe>(
 						FCollisionStructureManager::NewImplicitBox(
 							InstanceBoundingBox,
@@ -3234,8 +3240,7 @@ void FGeometryCollectionPhysicsProxy::InitializeSharedCollisionStructures(
 				}
 				else if (SizeSpecificData.ImplicitType == EImplicitTypeEnum::Chaos_Implicit_Sphere)
 				{
-					ErrorReporter.SetPrefix(BaseErrorPrefix + " | Cluster Transform Index: " + FString::FromInt(ClusterTransformIdx));
-					CollectionImplicits[ClusterTransformIdx] = TSharedPtr<Chaos::FImplicitObject, ESPMode::ThreadSafe>(
+ 					CollectionImplicits[ClusterTransformIdx] = TSharedPtr<Chaos::FImplicitObject, ESPMode::ThreadSafe>(
 						FCollisionStructureManager::NewImplicitSphere(
 							InstanceBoundingBox.GetExtent().GetAbsMin() / 2,
 							SizeSpecificData.CollisionObjectReductionPercentage,
diff --git a/Engine/Source/Runtime/Experimental/GeometryCollectionEngine/Public/GeometryCollection/GeometryCollectionComponent.h b/Engine/Source/Runtime/Experimental/GeometryCollectionEngine/Public/GeometryCollection/GeometryCollectionComponent.h
--- a/Engine/Source/Runtime/Experimental/GeometryCollectionEngine/Public/GeometryCollection/GeometryCollectionComponent.h
+++ b/Engine/Source/Runtime/Experimental/GeometryCollectionEngine/Public/GeometryCollection/GeometryCollectionComponent.h
@@ -232,27 +232,37 @@ private:
 //GetArrayRest (gives original rest value)
 //This generates pointers to arrays marked private. Macro assumes getters are public
 //todo(ocohen): may want to take in a static name
-#define COPY_ON_WRITE_ATTRIBUTE(Type, Name, Group)\
-FORCEINLINE const TManagedArray<Type>& Get##Name##Array() const {\
-	return Indirect##Name##Array ? *Indirect##Name##Array : RestCollection->GetGeometryCollection()->Name;\
-}\
-FORCEINLINE TManagedArray<Type>& Get##Name##ArrayCopyOnWrite(){\
-	if(!Indirect##Name##Array)\
-	{\
-		static FName StaticName(#Name);\
-		DynamicCollection->AddAttribute<Type>(StaticName, Group);\
-		DynamicCollection->CopyAttribute(*RestCollection->GetGeometryCollection(), StaticName, Group);\
-		Indirect##Name##Array = &DynamicCollection->GetAttribute<Type>(StaticName, Group);\
-		CopyOnWriteAttributeList.Add(reinterpret_cast<FManagedArrayBase**>(&Indirect##Name##Array));\
-	}\
-	return *Indirect##Name##Array;\
-}\
-FORCEINLINE void Reset##Name##ArrayDynamic(){\
-	Indirect##Name##Array = NULL;\
-}\
-FORCEINLINE const TManagedArray<Type>& Get##Name##ArrayRest() const{ return RestCollection->GetGeometryCollection()->Name; }\
-private:\
-TManagedArray<Type>* Indirect##Name##Array;\
+#define COPY_ON_WRITE_ATTRIBUTE(Type, Name, Group)								\
+FORCEINLINE const TManagedArray<Type>& Get##Name##Array() const 				\
+{																				\
+	return Indirect##Name##Array ?												\
+		*Indirect##Name##Array : RestCollection->GetGeometryCollection()->Name;	\
+}																				\
+FORCEINLINE TManagedArray<Type>& Get##Name##ArrayCopyOnWrite()					\
+{																				\
+	if(!Indirect##Name##Array)													\
+	{																			\
+		static FName StaticName(#Name);											\
+		DynamicCollection->AddAttribute<Type>(StaticName, Group);				\
+		DynamicCollection->CopyAttribute(										\
+			*RestCollection->GetGeometryCollection(), StaticName, Group);		\
+		Indirect##Name##Array =													\
+			&DynamicCollection->GetAttribute<Type>(StaticName, Group);			\
+		CopyOnWriteAttributeList.Add(											\
+			reinterpret_cast<FManagedArrayBase**>(&Indirect##Name##Array));		\
+	}																			\
+	return *Indirect##Name##Array;												\
+}																				\
+FORCEINLINE void Reset##Name##ArrayDynamic()									\
+{																				\
+	Indirect##Name##Array = NULL;												\
+}																				\
+FORCEINLINE const TManagedArray<Type>& Get##Name##ArrayRest() const				\
+{																				\
+	return RestCollection->GetGeometryCollection()->Name;						\
+}																				\
+private:																		\
+	TManagedArray<Type>* Indirect##Name##Array;									\
 public:
 
 /**
diff --git a/Engine/Source/Runtime/Experimental/GeometryCollectionSimulationCore/Private/GeometryCollection/GeometryCollectionCollisionStructureManager.cpp b/Engine/Source/Runtime/Experimental/GeometryCollectionSimulationCore/Private/GeometryCollection/GeometryCollectionCollisionStructureManager.cpp
--- a/Engine/Source/Runtime/Experimental/GeometryCollectionSimulationCore/Private/GeometryCollection/GeometryCollectionCollisionStructureManager.cpp
+++ b/Engine/Source/Runtime/Experimental/GeometryCollectionSimulationCore/Private/GeometryCollection/GeometryCollectionCollisionStructureManager.cpp
@@ -252,9 +252,24 @@ FCollisionStructureManager::NewImplicitLevelset(
 	HalfExtents *= CollisionObjectReduction / 100.f;
 	const float MinExtent = FMath::Min(HalfExtents[0], FMath::Min(HalfExtents[1], HalfExtents[2]));
 	Chaos::TLevelSet<float, 3>* LevelSet = NewLevelset(ErrorReporter, MeshParticles, TriMesh, CollisionBounds, MinRes, MaxRes, CollisionType);
-	if (LevelSet && MinExtent > 0)
+	if (LevelSet)
 	{
-		LevelSet->Shrink(MinExtent);
+		const float DomainVolume = LevelSet->BoundingBox().GetVolume();
+		const float FilledVolume = LevelSet->ApproximateNegativeMaterial();
+		if (FilledVolume < DomainVolume * 0.05)
+		{
+			ErrorReporter.ReportError(
+				*FString::Printf(TEXT(
+					"Level set is small or empty:\n"
+					"    domain volume: %g\n"
+					"    estimated level set volume: %g\n"
+					"    percentage filled: %g%%"),
+					DomainVolume, FilledVolume, FilledVolume / DomainVolume * 100.0));
+		}
+		if (MinExtent > 0)
+		{
+			LevelSet->Shrink(MinExtent);
+		}
 	}
 	return LevelSet;
 }
