Integrate 1145372 to Release-4.25.

Fix chaos niagra interface so that it pulls simulation results from the right place.  Remove a bunch of dead code, and set aside deprecated methods.  Be more selective with lambda captures.

diff --git a/Engine/Plugins/Experimental/ChaosNiagara/Source/ChaosNiagara/Private/NiagaraDataInterfaceChaosDestruction.cpp b/Engine/Plugins/Experimental/ChaosNiagara/Source/ChaosNiagara/Private/NiagaraDataInterfaceChaosDestruction.cpp
--- a/Engine/Plugins/Experimental/ChaosNiagara/Source/ChaosNiagara/Private/NiagaraDataInterfaceChaosDestruction.cpp
+++ b/Engine/Plugins/Experimental/ChaosNiagara/Source/ChaosNiagara/Private/NiagaraDataInterfaceChaosDestruction.cpp
@@ -583,8 +583,10 @@ void GetMeshExtData(FSolverData SolverData,
 				}
 				if (const FGeometryCollectionPhysicsProxy* GeomCollectionPhysicsProxy = GeometryCollectionComponent->GetPhysicsProxy())
 				{
-					const FGeometryCollectionResults& PhysResult = GeomCollectionPhysicsProxy->GetPhysicsResults().GetGameDataForRead();
-					Transform = PhysResult.ParticleToWorldTransforms[ParticleIndex - PhysResult.BaseIndex];
+					if (const FGeometryCollectionResults* PhysResult = GeomCollectionPhysicsProxy->GetConsumerResultsGT())
+					{
+						Transform = PhysResult->ParticleToWorldTransforms[ParticleIndex - PhysResult->BaseIndex];
+					}
 				}
 			}
 		}
diff --git a/Engine/Source/Runtime/Experimental/ChaosSolvers/Private/PBDRigidsSolver.cpp b/Engine/Source/Runtime/Experimental/ChaosSolvers/Private/PBDRigidsSolver.cpp
--- a/Engine/Source/Runtime/Experimental/ChaosSolvers/Private/PBDRigidsSolver.cpp
+++ b/Engine/Source/Runtime/Experimental/ChaosSolvers/Private/PBDRigidsSolver.cpp
@@ -435,14 +435,16 @@ namespace Chaos
 		InProxy->Initialize();
 		InProxy->InitializeBodiesGT();
 		InProxy->NewData(); // Buffers data on the proxy.
+		FParticlesType* InParticles = &GetParticles();
 
 		// Finish registration on the physics thread...
 		FChaosSolversModule::GetModule()->GetDispatcher()->EnqueueCommandImmediate(this,
-			[this, InProxy](FPBDRigidsSolver* Solver)
+			[InParticles, InProxy](FPBDRigidsSolver* Solver)
 		{
 			UE_LOG(LogPBDRigidsSolver, Verbose, 
 				TEXT("FPBDRigidsSolver::RegisterObject(FGeometryCollectionPhysicsProxy*)"));
-			InProxy->InitializeBodiesPT(Solver, Particles);
+			check(InParticles);
+			InProxy->InitializeBodiesPT(Solver, *InParticles);
 			InProxy->PushToPhysicsState(nullptr); // Uses buffer on proxy for transport, not ProxyData.
 		});
 	}
@@ -464,7 +466,7 @@ namespace Chaos
 
 		FChaosSolversModule::GetModule()->GetDispatcher()->EnqueueCommandImmediate(
 			this,
-			[this, InProxy, ProxyData](FPBDRigidsSolver* Solver)
+			[InProxy, ProxyData](FPBDRigidsSolver* Solver)
 			{
 				UE_LOG(LogPBDRigidsSolver, Verbose,
 					TEXT("FPBDRigidsSolver::RegisterObject(FFieldSystemPhysicsProxy*)"));
diff --git a/Engine/Source/Runtime/Experimental/ChaosSolvers/Private/PhysicsProxy/GeometryCollectionPhysicsProxy.cpp b/Engine/Source/Runtime/Experimental/ChaosSolvers/Private/PhysicsProxy/GeometryCollectionPhysicsProxy.cpp
--- a/Engine/Source/Runtime/Experimental/ChaosSolvers/Private/PhysicsProxy/GeometryCollectionPhysicsProxy.cpp
+++ b/Engine/Source/Runtime/Experimental/ChaosSolvers/Private/PhysicsProxy/GeometryCollectionPhysicsProxy.cpp
@@ -297,7 +297,8 @@
 		if (Simplicial->Size())
 		{
 			const Chaos::TAABB<float,3> ImplicitShapeDomain = 
-				Implicit && Implicit->HasBoundingBox() ? Implicit->BoundingBox() : Chaos::TAABB<float,3>::FullAABB();
+				Implicit && Implicit->GetType() == Chaos::ImplicitObjectType::LevelSet && Implicit->HasBoundingBox() ? 
+				Implicit->BoundingBox() : Chaos::TAABB<float,3>::FullAABB();
 
 			CollisionParticles->Resize(0);
 			CollisionParticles->AddParticles(Simplicial->Size());
@@ -473,18 +474,17 @@
 	RecordedTracks.Records.Reset();
 	
 	// Setup double buffer data
-	Results.Get(0).Transforms.Init(PTDynamicCollection.Transform);
-	//Results.Get(0).RigidBodyIds.Init(RigidBodyID);
-	Results.Get(1).Transforms.Init(PTDynamicCollection.Transform);
-	//Results.Get(1).RigidBodyIds.Init(RigidBodyID);
+	//Results.Get(0).Transforms.Init(PTDynamicCollection.Transform);
+	//Results.Get(1).Transforms.Init(PTDynamicCollection.Transform);
 	
 	LastSyncCountGT = 0;
 
 	// Initialize global transforms
 
-	TArray<FMatrix> &TmpGlobalTransforms = Results.Get(0).GlobalTransforms;
+	TArray<FMatrix> &TmpGlobalTransforms = GameToPhysInterchange->AccessProducerBuffer()->GlobalTransforms;		
+	//TArray<FMatrix> &TmpGlobalTransforms = Results.Get(0).GlobalTransforms;
 	GeometryCollectionAlgo::GlobalMatrices(PTDynamicCollection.Transform, PTDynamicCollection.Parent, TmpGlobalTransforms);
-	Results.Get(1).GlobalTransforms = TmpGlobalTransforms;
+	//Results.Get(1).GlobalTransforms = TmpGlobalTransforms;
 
 
 	// Initialize data for faster bound calculations
@@ -516,8 +516,8 @@
 				TmpGlobalTransforms[ValidGeometryTransformIndices[i]] * ActorToWorld);
 		}
 
-		Results.Get(0).WorldBounds = FBoxSphereBounds(BoundingBox);
-		Results.Get(1).WorldBounds = FBoxSphereBounds(BoundingBox);
+		//Results.Get(0).WorldBounds = FBoxSphereBounds(BoundingBox);
+		//Results.Get(1).WorldBounds = FBoxSphereBounds(BoundingBox);
 	}
 
 	InitializedState = ESimulationInitializationState::Unintialized;
@@ -727,19 +727,6 @@
 					(uint8)Parameters.ObjectType, 
 					CollisionGroup[TransformGroupIndex]);
 
-				if (Parameters.InitialVelocityType == EInitialVelocityTypeEnum::Chaos_Initial_Velocity_User_Defined)
-				{
-					InitialLinearVelocity[TransformGroupIndex] = Parameters.InitialLinearVelocity;
-					InitialAngularVelocity[TransformGroupIndex] = Parameters.InitialAngularVelocity;
-					Handle->SetV(Parameters.InitialLinearVelocity);
-					Handle->SetW(Parameters.InitialAngularVelocity);
-				}
-				else
-				{
-					Handle->SetV(Chaos::TVector<float, 3>(0.f));
-					Handle->SetW(Chaos::TVector<float, 3>(0.f));
-				}
-
 				if (Parameters.EnableClustering)
 				{
 					Handle->SetClusterGroupIndex(Parameters.ClusterGroupIndex);
@@ -785,42 +772,23 @@
 
 		if (Parameters.InitialVelocityType == EInitialVelocityTypeEnum::Chaos_Initial_Velocity_User_Defined)
 		{
-			for (int32 GeometryIndex = 0; GeometryIndex < NumGeometries; ++GeometryIndex)
+			// A previous implementation of this went wide on this loop.  The general 
+			// rule of thumb for parallelization is that each thread needs at least
+			// 1000 operations in order to overcome the expense of threading.  I don't
+			// think that's generally going to be the case here...
+			for (int32 TransformGroupIndex = 0; TransformGroupIndex < GetTransformGroupSize(); ++TransformGroupIndex)
 			{
-				const int32 TransformGroupIndex = TransformIndex[GeometryIndex];
 				if (Chaos::TPBDRigidParticleHandle<float, 3>* Handle = SolverParticleHandles[TransformGroupIndex])
 				{
-					InitialLinearVelocity[TransformGroupIndex] = Parameters.InitialLinearVelocity;
-					InitialAngularVelocity[TransformGroupIndex] = Parameters.InitialAngularVelocity;
 					if (DynamicState[TransformGroupIndex] == (int32)EObjectStateTypeEnum::Chaos_Object_Dynamic)
 					{
-						Handle->SetV(Parameters.InitialLinearVelocity);
-						Handle->SetW(Parameters.InitialAngularVelocity);
+						Handle->SetV(InitialLinearVelocity[TransformGroupIndex]);
+						Handle->SetW(InitialAngularVelocity[TransformGroupIndex]);
 					}
 				}
 			}
 		}
-/*
-		ParallelFor(NumGeometries, [&](int32 GeometryIndex)
-		{
-			const int32 TransformGroupIndex = TransformIndex[GeometryIndex];
-			//if (SimulatableParticles[TransformGroupIndex] == true)
-			if (Chaos::TPBDRigidParticleHandle<float, 3>* Handle = SolverParticleHandles[TransformGroupIndex])
-			{
-				//const int32 RigidBodyIndex = RigidBodyID[TransformGroupIndex];
-				if (Parameters.InitialVelocityType == EInitialVelocityTypeEnum::Chaos_Initial_Velocity_User_Defined)
-				{
-					InitialLinearVelocity[TransformGroupIndex] = Parameters.InitialLinearVelocity;
-					InitialAngularVelocity[TransformGroupIndex] = Parameters.InitialAngularVelocity;
-					if (DynamicState[TransformGroupIndex] == (int32)EObjectStateTypeEnum::Chaos_Object_Dynamic)
-					{
-						Particles.V(RigidBodyIndex) = Parameters.InitialLinearVelocity;
-						Particles.W(RigidBodyIndex) = Parameters.InitialAngularVelocity;
-					}
-				}
-			}
-		});
-*/
+
 #if TODO_REIMPLEMENT_KINEMATICS
 		InitializeKinematics(Particles, DynamicState);
 #endif // TODO_REIMPLEMENT_KINEMATICS
@@ -1013,1435 +981,12 @@
 	} // end if simulating...
 }
 
-void FGeometryCollectionPhysicsProxy::UpdateKinematicBodiesCallback(const FParticlesType& Particles, const float Dt, const float Time, FKinematicProxy& Proxy)
-{
-	if (InitializedState == ESimulationInitializationState::Initialized)
-	{
-		SCOPE_CYCLE_COUNTER(STAT_KinematicUpdate);
-		FGeometryDynamicCollection* Collection = Parameters.DynamicCollection;
-		check(Collection);
-
-		bool bIsCachePlaying = Parameters.IsCachePlaying();
-		bool bIsReverseCachePlaying = Parameters.IsCacheRecording() && Parameters.ReverseCacheBeginTime != 0 && Parameters.ReverseCacheBeginTime < Time;
-		if ((!bIsCachePlaying && !bIsReverseCachePlaying) || !Parameters.RecordedTrack)
-		{
-			return;
-		}
-
-#if TODO_REIMPLEMENT_KINEMATIC_PROXY
-		bool bFirst = !Proxy.Ids.Num();
-		if (bFirst)
-		{
-			Proxy.Position.Reset(RigidBodyID.Num());
-			Proxy.Rotation.Reset(RigidBodyID.Num());
-			Proxy.NextPosition.Reset(RigidBodyID.Num());
-			Proxy.NextRotation.Reset(RigidBodyID.Num());
-
-			Proxy.Position.AddUninitialized(RigidBodyID.Num());
-			Proxy.Rotation.AddUninitialized(RigidBodyID.Num());
-			Proxy.NextPosition.AddUninitialized(RigidBodyID.Num());
-			Proxy.NextRotation.AddUninitialized(RigidBodyID.Num());
-
-			for (int32 i = 0; i < RigidBodyID.Num(); ++i)
-			{
-				Proxy.Ids.Add(RigidBodyID[i]);
-
-				// Initialise to rest state
-				const int32 RbId = Proxy.Ids.Last();
-				Proxy.Position[i] = RbId != INDEX_NONE ? Particles.X(RbId) : FVector::ZeroVector;
-				Proxy.Rotation[i] = RbId != INDEX_NONE ? Particles.R(RbId) : FQuat::Identity;
-				Proxy.NextPosition[i] = Proxy.Position[i];
-				Proxy.NextRotation[i] = Proxy.Rotation[i];
-			}
-		}
-
-		if (bIsCachePlaying && !bIsReverseCachePlaying && (Time < Parameters.CacheBeginTime || !Parameters.RecordedTrack->IsTimeValid(Time)))
-		{
-			return;
-		}
-
-		float ReverseTime = Parameters.RecordedTrack->GetLastTime() - Time + Parameters.ReverseCacheBeginTime;
-		if (bIsReverseCachePlaying && !Parameters.RecordedTrack->IsTimeValid(ReverseTime))
-		{
-			return;
-		}
-
-		const FRecordedFrame* FirstFrame = nullptr;
-		const FRecordedFrame* SecondFrame = nullptr;
-		float PlaybackTime = bIsReverseCachePlaying ? ReverseTime : Time;
-		Parameters.RecordedTrack->GetFramesForTime(PlaybackTime, FirstFrame, SecondFrame);
-
-		if (FirstFrame && !SecondFrame)
-		{
-			// Only one frame to take information from (simpler case)
-			const int32 NumActives = FirstFrame->TransformIndices.Num();
-
-			// Actives
-			Chaos::PhysicsParallelFor(NumActives, [&](int32 Index)
-			{
-				const int32 InternalIndex = FirstFrame->TransformIndices[Index];
-				if (InternalIndex >= RigidBodyID.Num() || InternalIndex < 0)
-				{
-					UE_LOG(UGCC_LOG, Error, 
-						TEXT("%s: Cache index %d out of range: [%d, %d).  Regenerate the cache."),
-						*Parameters.Name, InternalIndex, 0, RigidBodyID.Num());
-					return;
-				}
-				const int32 ExternalIndex = RigidBodyID[InternalIndex];
-
-				if (ExternalIndex != INDEX_NONE && Particles.InvM(ExternalIndex) == 0.0f && !Particles.Disabled(ExternalIndex))
-				{
-					const FTransform& ParticleTransform = FirstFrame->Transforms[Index];
-					Proxy.Position[InternalIndex] = Particles.X(ExternalIndex);
-					Proxy.Rotation[InternalIndex] = Particles.R(ExternalIndex);
-					Proxy.NextPosition[InternalIndex] = ParticleTransform.GetTranslation();
-					Proxy.NextRotation[InternalIndex] = ParticleTransform.GetRotation();
-				}
-			});
-		}
-		else if (FirstFrame && SecondFrame)
-		{
-			// Both frames valid, second frame has all the indices we need
-			const int32 NumActives = SecondFrame->TransformIndices.Num();
-
-			const float Alpha = (PlaybackTime - FirstFrame->Timestamp) / (SecondFrame->Timestamp - FirstFrame->Timestamp);
-			check(0 <= Alpha && Alpha <= 1.0f);
-
-			Chaos::PhysicsParallelFor(NumActives, [&](int32 Index)
-			{
-				const int32 InternalIndex = SecondFrame->TransformIndices[Index];
-				if (InternalIndex >= RigidBodyID.Num() || InternalIndex < 0)
-				{
-					UE_LOG(UGCC_LOG, Error, 
-						TEXT("%s: Cache index %d out of range: [%d, %d).  Regenerate the cache."),
-						*Parameters.Name, InternalIndex, 0, RigidBodyID.Num());
-					return;
-				}
-				const int32 ExternalIndex = RigidBodyID[InternalIndex];
-				const int32 PreviousIndexSlot = Index < SecondFrame->PreviousTransformIndices.Num() ? SecondFrame->PreviousTransformIndices[Index] : INDEX_NONE;
-
-				if (ExternalIndex != INDEX_NONE && Particles.InvM(ExternalIndex) == 0.0f && !Particles.Disabled(ExternalIndex))
-				{
-					if (PreviousIndexSlot != INDEX_NONE)
-					{
-						Proxy.Position[InternalIndex] = Proxy.NextPosition[InternalIndex];
-						Proxy.Rotation[InternalIndex] = Proxy.NextRotation[InternalIndex];
-
-						FTransform BlendedTM;
-						BlendedTM.Blend(FirstFrame->Transforms[PreviousIndexSlot], SecondFrame->Transforms[Index], Alpha);
-
-						Proxy.NextPosition[InternalIndex] = BlendedTM.GetTranslation();
-						Proxy.NextRotation[InternalIndex] = BlendedTM.GetRotation();
-					}
-					else
-					{
-						// NewActive case
-						Proxy.Position[InternalIndex] = Proxy.NextPosition[InternalIndex];
-						Proxy.Rotation[InternalIndex] = Proxy.NextRotation[InternalIndex];
-
-						FTransform BlendedTM;
-						BlendedTM.Blend(FTransform(Particles.R(ExternalIndex), Particles.X(ExternalIndex), FVector::OneVector), SecondFrame->Transforms[Index], Alpha);
-
-						Proxy.NextPosition[InternalIndex] = BlendedTM.GetTranslation();
-						Proxy.NextRotation[InternalIndex] = BlendedTM.GetRotation();
-					}
-				}
-			});
-			// #BGallagher Handle new inactives. If it's a cluster parent and it's fully disabled we'll need to decluster it here.
-		}
-#endif
-	}
-}
-
-void FGeometryCollectionPhysicsProxy::StartFrameCallback(const float Dt, const float Time)
-{
-	SCOPE_CYCLE_COUNTER(STAT_GeomBeginFrame);
-	if (InitializedState == ESimulationInitializationState::Initialized)
-	{
-		// Reverse playback only plays back what we just recorded.  So, the condition 
-		// is, "are we in record mode, but ready to play back what we've recorded?"
-		const bool bIsReverseCachePlaying = Parameters.IsCacheRecording() && Parameters.ReverseCacheBeginTime != 0 && Time > Parameters.ReverseCacheBeginTime;
-		if (Parameters.IsCachePlaying() || bIsReverseCachePlaying)
-		{
-#if TODO_REIMPLEMENT_GET_RIGID_PARTICLES
-			// Update the enabled/disabled state for kinematic particles for the upcoming frame
-			Chaos::FPhysicsSolver* ThisSolver = GetSolver();
-			Chaos::FPhysicsSolver::FParticlesType& Particles = ThisSolver->GetRigidParticles();
-
-			if (!Parameters.RecordedTrack)
-			{
-				if (!ensure(Parameters.CacheType == EGeometryCollectionCacheType::Record))
-				{
-					return;
-				}
-				Parameters.RecordedTrack = new FRecordedTransformTrack();
-				Parameters.bOwnsTrack = true;
-			}
-			if (Parameters.bClearCache && bIsReverseCachePlaying)
-			{
-				check(!CommitRecordedStateCallback);
-				*const_cast<FRecordedTransformTrack*>(Parameters.RecordedTrack) = FRecordedTransformTrack::ProcessRawRecordedData(RecordedTracks);
-				Parameters.bClearCache = false;
-			}
-
-			bool bParticlesUpdated = false;
-
-			const float ThisFrameTime = bIsReverseCachePlaying ? (Parameters.RecordedTrack->GetLastTime() - Time + Parameters.ReverseCacheBeginTime) : Time;
-			if (!Parameters.RecordedTrack->IsTimeValid(ThisFrameTime))
-			{
-				// Invalid cache time, nothing to update
-				return;
-			}
-
-			FCriticalSection ParticleUpdateLock;
-			const int32 NumMappings = RigidBodyID.Num();
-			Chaos::PhysicsParallelFor(NumMappings, [&](int32 InternalParticleIndex)
-				//for(int32 Index = 0; Index < NumMappings; ++Index)
-			{
-				const int32 ExternalParticleIndex = RigidBodyID[InternalParticleIndex];
-
-				if (ExternalParticleIndex == INDEX_NONE)
-				{
-					return;
-				}
-
-				if (Particles.InvM(ExternalParticleIndex) != 0)
-				{
-					return;
-				}
-
-				// We need to check a window of Now - Dt to Now and see if we ever activated in that time.
-				// This is for short activations because if we miss one then the playback will be incorrect
-				const bool bShouldBeDisabled = !Parameters.RecordedTrack->GetWasActiveInWindow(InternalParticleIndex, ThisFrameTime, bIsReverseCachePlaying ? (ThisFrameTime - Dt) : (ThisFrameTime + Dt));
-				const bool bDisabledNow = Particles.Disabled(ExternalParticleIndex);
-				if (bDisabledNow != bShouldBeDisabled)
-				{
-					Particles.SetDisabledLowLevel(ExternalParticleIndex, bShouldBeDisabled);
-					ParticleUpdateLock.Lock();
-					if (!bShouldBeDisabled)
-					{
-						Particles.SetObjectState(ExternalParticleIndex, Chaos::EObjectStateType::Kinematic);
-						if (!ensure(ThisSolver->ActiveIndices().Find(ExternalParticleIndex) == nullptr))
-						{
-							bParticlesUpdated = true;
-						}
-						else if(!bParticlesUpdated)
-						{
-							ThisSolver->NonDisabledIndices().Add(ExternalParticleIndex);
-							ThisSolver->ActiveIndices().Add(ExternalParticleIndex);
-						}
-					}
-					else
-					{
-						if (!ensure(ThisSolver->ActiveIndices().Find(ExternalParticleIndex) != nullptr))
-						{
-							bParticlesUpdated = true;
-						}
-						else if(!bParticlesUpdated)
-						{
-							ThisSolver->NonDisabledIndices().Remove(ExternalParticleIndex);
-							ThisSolver->ActiveIndices().Remove(ExternalParticleIndex);
-						}
-					}
-					ParticleUpdateLock.Unlock();
-				}
-
-			});
-
-			// Do not add collisions if reverse
-			if (!bIsReverseCachePlaying)
-			{
-				const FRecordedFrame* RecordedFrame = Parameters.RecordedTrack->FindRecordedFrame(ThisFrameTime);
-				if (RecordedFrame == nullptr)
-				{
-					const int32 Index = Parameters.RecordedTrack->FindLastKeyBefore(ThisFrameTime);
-					if (Index > 0 && Index < Parameters.RecordedTrack->Records.Num())
-					{
-						RecordedFrame = &Parameters.RecordedTrack->Records[Index];
-					}
-				}
-
-				if (RecordedFrame)
-				{
-					// Collisions
-					if (Parameters.CollisionData.DoGenerateCollisionData && ThisFrameTime > 0.f && Parameters.CollisionData.CollisionDataSizeMax > 0)
-					{
-						if (RecordedFrame->Collisions.Num() > 0)
-						{
-							Chaos::FPhysicsSolver::FCollisionDataArray& AllCollisionsDataArray = ThisSolver->GetAllCollisionsDataArray();
-							TMap<IPhysicsProxyBase*, TArray<int32>>& AllCollisionsIndicesByPhysicsProxy = ThisSolver->GetAllCollisionsIndicesByPhysicsProxy();
-
-							if (!AllCollisionsIndicesByPhysicsProxy.Contains(this))
-							{
-								AllCollisionsIndicesByPhysicsProxy.Add(this, TArray<int32>());
-							}
-
-							for (int32 Idx = 0; Idx < RecordedFrame->Collisions.Num(); ++Idx)
-							{
-								// Check if the particle is still kinematic
-								if (RecordedFrame->Collisions[Idx].ParticleIndex < 0 ||
-									(RecordedFrame->Collisions[Idx].ParticleIndex >= 0 && RecordedFrame->Collisions[Idx].ParticleIndex < static_cast<int32>(Particles.Size()) && Particles.ObjectState(RecordedFrame->Collisions[Idx].ParticleIndex) == Chaos::EObjectStateType::Kinematic))
-								{
-									const int32 NewIdx = AllCollisionsDataArray.Add(Chaos::TCollisionData<float, 3>());
-									Chaos::TCollisionData<float, 3>& AllCollisionsDataArrayItem = AllCollisionsDataArray[NewIdx];
-
-									AllCollisionsDataArrayItem.Location = RecordedFrame->Collisions[Idx].Location;
-									AllCollisionsDataArrayItem.AccumulatedImpulse = RecordedFrame->Collisions[Idx].AccumulatedImpulse;
-									AllCollisionsDataArrayItem.Normal = RecordedFrame->Collisions[Idx].Normal;
-									AllCollisionsDataArrayItem.Velocity1 = RecordedFrame->Collisions[Idx].Velocity1;
-									AllCollisionsDataArrayItem.Velocity2 = RecordedFrame->Collisions[Idx].Velocity2;
-									AllCollisionsDataArrayItem.AngularVelocity1 = RecordedFrame->Collisions[Idx].AngularVelocity1;
-									AllCollisionsDataArrayItem.AngularVelocity2 = RecordedFrame->Collisions[Idx].AngularVelocity2;
-									AllCollisionsDataArrayItem.Mass1 = RecordedFrame->Collisions[Idx].Mass1;
-									AllCollisionsDataArrayItem.Mass2 = RecordedFrame->Collisions[Idx].Mass2;
-#if TODO_CONVERT_GEOMETRY_COLLECTION_PARTICLE_INDICES_TO_PARTICLE_POINTERS
-									AllCollisionsDataArrayItem.ParticleIndex = RecordedFrame->Collisions[Idx].ParticleIndex;
-#endif
-									AllCollisionsDataArrayItem.LevelsetIndex = RecordedFrame->Collisions[Idx].LevelsetIndex;
-									AllCollisionsDataArrayItem.ParticleIndexMesh = RecordedFrame->Collisions[Idx].ParticleIndexMesh;
-									AllCollisionsDataArrayItem.LevelsetIndexMesh = RecordedFrame->Collisions[Idx].LevelsetIndexMesh;
-
-									AllCollisionsIndicesByPhysicsProxy[this].Add(NewIdx);
-								}
-							}
-						}
-					}
-
-					// Breaking
-					if (Parameters.BreakingData.DoGenerateBreakingData && ThisFrameTime > 0.f && Parameters.BreakingData.BreakingDataSizeMax > 0)
-					{
-						if (RecordedFrame->Breakings.Num() > 0)
-						{
-							Chaos::FPhysicsSolver::FBreakingDataArray& AllBreakingsDataArray = ThisSolver->GetAllBreakingsDataArray();
-							TMap<IPhysicsProxyBase*, TArray<int32>>& AllBreakingsIndicesByPhysicsProxy = ThisSolver->GetAllBreakingsIndicesByPhysicsProxy();
-
-							if (!AllBreakingsIndicesByPhysicsProxy.Contains(this))
-							{
-								AllBreakingsIndicesByPhysicsProxy.Add(this, TArray<int32>());
-							}
-
-							for (int32 Idx = 0; Idx < RecordedFrame->Breakings.Num(); ++Idx)
-							{
-								// Check if the particle is still kinematic
-								if (RecordedFrame->Breakings[Idx].ParticleIndex < 0 ||
-									(RecordedFrame->Breakings[Idx].ParticleIndex >= 0 && RecordedFrame->Breakings[Idx].ParticleIndex < static_cast<int32>(Particles.Size()) && Particles.ObjectState(RecordedFrame->Breakings[Idx].ParticleIndex) == Chaos::EObjectStateType::Kinematic))
-								{
-									const int32 NewIdx = AllBreakingsDataArray.Add(Chaos::TBreakingData<float, 3>());
-									Chaos::TBreakingData<float, 3>& AllBreakingsDataArrayItem = AllBreakingsDataArray[NewIdx];
-
-									AllBreakingsDataArrayItem.Location = RecordedFrame->Breakings[Idx].Location;
-									AllBreakingsDataArrayItem.Velocity = RecordedFrame->Breakings[Idx].Velocity;
-									AllBreakingsDataArrayItem.AngularVelocity = RecordedFrame->Breakings[Idx].AngularVelocity;
-									AllBreakingsDataArrayItem.Mass = RecordedFrame->Breakings[Idx].Mass;
-#if TODO_CONVERT_GEOMETRY_COLLECTION_PARTICLE_INDICES_TO_PARTICLE_POINTERS
-									AllBreakingsDataArrayItem.ParticleIndex = RecordedFrame->Breakings[Idx].ParticleIndex;
-#endif
-									AllBreakingsDataArrayItem.ParticleIndexMesh = RecordedFrame->Breakings[Idx].ParticleIndexMesh;
-
-									AllBreakingsIndicesByPhysicsProxy[this].Add(NewIdx);
-								}
-							}
-						}
-					}
-
-					// Trailing
-					if (Parameters.TrailingData.DoGenerateTrailingData && ThisFrameTime > 0.f && Parameters.TrailingData.TrailingDataSizeMax > 0)
-					{
-						if (RecordedFrame->Trailings.Num() > 0)
-						{
-							Chaos::FPhysicsSolver::FTrailingDataArray& AllTrailingsDataArray = ThisSolver->GetAllTrailingsDataArray();
-							TMap<IPhysicsProxyBase*, TArray<int32>>& AllTrailingsIndicesByPhysicsProxy = ThisSolver->GetAllTrailingsIndicesByPhysicsProxy();
-
-							if (!AllTrailingsIndicesByPhysicsProxy.Contains(this))
-							{
-								AllTrailingsIndicesByPhysicsProxy.Add(this, TArray<int32>());
-							}
-
-							for (const FSolverTrailingData& Trailing : RecordedFrame->Trailings)
-							{
-								// Check if the particle is still kinematic
-								if (Trailing.ParticleIndex < 0 ||
-									(Trailing.ParticleIndex >= 0 && Trailing.ParticleIndex < static_cast<int32>(Particles.Size()) && Particles.ObjectState(Trailing.ParticleIndex) == Chaos::EObjectStateType::Kinematic))
-								{
-									const int32 NewIdx = AllTrailingsDataArray.Add(Chaos::TTrailingData<float, 3>());
-									Chaos::TTrailingData<float, 3>& AllTrailingsDataArrayItem = AllTrailingsDataArray[NewIdx];
-
-									AllTrailingsDataArrayItem.Location = Trailing.Location;
-									AllTrailingsDataArrayItem.Velocity = Trailing.Velocity;
-									AllTrailingsDataArrayItem.AngularVelocity = Trailing.AngularVelocity;
-									AllTrailingsDataArrayItem.Mass = Trailing.Mass;
-#if TODO_CONVERT_GEOMETRY_COLLECTION_PARTICLE_INDICES_TO_PARTICLE_POINTERS
-									AllTrailingsDataArrayItem.ParticleIndex = Trailing.ParticleIndex;
-#endif
-									AllTrailingsDataArrayItem.ParticleIndexMesh = Trailing.ParticleIndexMesh;
-
-									AllTrailingsIndicesByPhysicsProxy[this].Add(NewIdx);
-								}
-							}
-						}
-					}
-				}
-			}
-
-			if (bParticlesUpdated)
-			{
-				ThisSolver->InitializeFromParticleData(0);
-			}
-#endif
-		}
-	}
-}
-
-void FGeometryCollectionPhysicsProxy::EndFrameCallback(const float EndFrame)
-{
-	check(false); // not currently called
-	if (InitializedState == ESimulationInitializationState::Initialized)
-	{
-		FGeometryDynamicCollection* Collection = Parameters.DynamicCollection;
-		check(Collection);
-
-		ProxySimDuration += EndFrame;
-
-		if(Collection->HasAttribute("RigidBodyID", FGeometryCollection::TransformGroup))
-		{
-#if TODO_REIMPLEMENT_GET_RIGID_PARTICLES
-			//
-			//  Update transforms for the simulated transforms
-			//
-			TManagedArray<int32>& CollectionClusterID = SolverClusterID;
-			TManagedArray<FTransform>& Transform = Collection->Transform;
-			TManagedArray<int32>& Parent = Collection->Parent;
-			TManagedArray<TSet<int32>>& Children = Collection->Children;
-			TManagedArray<int32>& SimulationType = Collection->SimulationType;
-
-			TManagedArray<int32>& DynamicState = Collection->GetAttribute<int32>("DynamicState", FGeometryCollection::TransformGroup);
-
-			Chaos::TPBDRigidParticles<float, 3>& Particles = GetSolver()->GetRigidParticles();
-			const Chaos::TPBDCollisionConstraints<float, 3>& CollisionRule = GetSolver()->GetCollisionConstraints();
-			const Chaos::TArrayCollectionArray<Chaos::ClusterId>& ClusterID = GetSolver()->GetRigidClustering().GetClusterIdsArray();
-			const Chaos::TArrayCollectionArray<Chaos::TRigidTransform<float, 3>>& ClusterChildToParentMap = GetSolver()->GetRigidClustering().GetChildToParentMap();
-			const Chaos::TArrayCollectionArray<bool>& InternalCluster = GetSolver()->GetRigidClustering().GetInternalClusterArray();
-
-			//Particles X and R are aligned with center of mass and inertia principal axes.
-			//Renderer doesn't know about this and simply does ActorToWorld * GeomToActor * LocalSpaceVerts
-			//In proper math multiplication order:
-			//ParticleToWorld = ActorToWorld * GeomToActor * MassToLocal
-			//GeomToWorld = ActorToWorld * GeomToActor
-			//=> GeomToWorld = ParticleToWorld * MassToLocal.Inv()
-			//=> GeomToActor = ActorToWorld.Inv() * ParticleToWorld * MassToLocal.Inv()
-			const int32 TransformSize = Collection->NumElements(FGeometryCollection::TransformGroup);
-			const FTransform& ActorToWorld = Parameters.WorldTransform;
-
-			// It's not worth shrinking the EndFrameUnparentingBuffer array, at least until the solver supports deleting bodies.
-			if (EndFrameUnparentingBuffer.Num() < TransformSize)
-			{
-				EndFrameUnparentingBuffer.Init(-1, TransformSize);
-			}
-
-			{
-				for (int32 TransformGroupIndex = 0; TransformGroupIndex < TransformSize; ++TransformGroupIndex)
-				{
-					const int32 RigidBodyIndex = RigidBodyID[TransformGroupIndex];
-					if (RigidBodyIndex != INDEX_NONE)
-					{
-						// Update the transform and parent hierarchy of the active rigid bodies. Active bodies can be either
-						// rigid geometry defined from the leaf nodes of the collection, or cluster bodies that drive an entire
-						// branch of the hierarchy within the GeometryCollection.
-						// - Active bodies are directly driven from the global position of the corresponding
-						//   rigid bodies within the solver ( cases where RigidBodyID[TransformGroupIndex] is not disabled ). 
-						// - Deactivated bodies are driven from the transforms of their active parents. However the solver can
-						//   take ownership of the parents during the simulation, so it might be necessary to force deactivated
-						//   bodies out of the collections hierarchy during the simulation.  
-						if (!Particles.Disabled(RigidBodyID[TransformGroupIndex]))
-						{
-							// Update the transform of the active body. The active body can be either a single rigid
-							// or a collection of rigidly attached geometries (Clustering). The cluster is represented as a
-							// single transform in the GeometryCollection, and all children are stored in the local space
-							// of the parent cluster.
-							// ... When setting cluster transforms it is expected that the MassToLocal is identity.
-							//     Cluster initialization will set the vertices in the MassSpace of the rigid body.
-							// ... When setting individual rigid bodies that are not clustered, the MassToLocal will be 
-							//     non-Identity, and will reflect the difference between the geometric center of the geometry
-							//     and that corresponding rigid bodies center of mass. 
-							const FTransform ParticleToWorld(Particles.R(RigidBodyIndex), Particles.X(RigidBodyIndex));
-							// GeomToActor = ActorToWorld.Inv() * ParticleToWorld * MassToLocal.Inv();
-							Transform[TransformGroupIndex] = MassToLocal[TransformGroupIndex].GetRelativeTransformReverse(ParticleToWorld).GetRelativeTransform(ActorToWorld);
-							Transform[TransformGroupIndex].NormalizeRotation();
-
-							// dynamic state is also updated by the solver during field interaction. 
-							if (!Particles.Sleeping(RigidBodyIndex))
-							{
-								if (Particles.ObjectState(RigidBodyIndex) == Chaos::EObjectStateType::Kinematic)
-								{
-									DynamicState[TransformGroupIndex] = (int)EObjectStateTypeEnum::Chaos_Object_Kinematic;
-								}
-								else if (Particles.ObjectState(RigidBodyIndex) == Chaos::EObjectStateType::Static)
-								{
-									DynamicState[TransformGroupIndex] = (int)EObjectStateTypeEnum::Chaos_Object_Static;
-								}
-								else
-								{
-									DynamicState[TransformGroupIndex] = (int)EObjectStateTypeEnum::Chaos_Object_Dynamic;
-								}
-							}
-
-							// Force all enabled rigid bodies out of the transform hierarchy
-							if (Parent[TransformGroupIndex] != INDEX_NONE)
-							{
-								const int32 ParentIndex = Parent[TransformGroupIndex];
-								// Children in the hierarchy are stored in a TSet, which is not thread safe.  So we retain
-								// indices to remove afterwards.
-								EndFrameUnparentingBuffer[TransformGroupIndex] = ParentIndex;
-							}
-
-							// When a leaf node rigid body is removed from a cluster the rigid
-							// body will become active and needs its clusterID updated. This just
-							// syncs the clusterID all the time. 
-							CollectionClusterID[TransformGroupIndex] = ClusterID[RigidBodyIndex].Id;
-						}
-						else if (Particles.Disabled(RigidBodyIndex))
-						{
-							// dynamic state is also updated by the solver during field interaction. 
-							if (!Particles.Sleeping(RigidBodyIndex))
-							{
-								if (Particles.ObjectState(RigidBodyIndex) == Chaos::EObjectStateType::Kinematic)
-								{
-									DynamicState[TransformGroupIndex] = (int)EObjectStateTypeEnum::Chaos_Object_Kinematic;
-								}
-								else if (Particles.ObjectState(RigidBodyIndex) == Chaos::EObjectStateType::Static)
-								{
-									DynamicState[TransformGroupIndex] = (int)EObjectStateTypeEnum::Chaos_Object_Static;
-								}
-								else
-								{
-									DynamicState[TransformGroupIndex] = (int)EObjectStateTypeEnum::Chaos_Object_Dynamic;
-								}
-							}
-
-							// The rigid body parent cluster has changed within the solver, and its
-							// parent body is not tracked within the geometry collection. So we need to
-							// pull the rigid bodies out of the transform hierarchy, and just drive
-							// the positions directly from the solvers cluster particle. 
-							if (CollectionClusterID[TransformGroupIndex] != ClusterID[RigidBodyIndex].Id)
-							{
-								// Force all driven rigid bodies out of the transform hierarchy
-								if (Parent[TransformGroupIndex] != INDEX_NONE)
-								{
-									const int32 ParentIndex = Parent[TransformGroupIndex];
-									// Children in the hierarchy are stored in a TSet, which is not thread safe.  So we retain
-									// indices to remove afterwards.
-									EndFrameUnparentingBuffer[TransformGroupIndex] = ParentIndex;
-								}
-								CollectionClusterID[TransformGroupIndex] = ClusterID[RigidBodyIndex].Id;
-							}
-
-							// Disabled rigid bodies that have valid cluster parents, and have been re-indexed by the
-							// solver ( As in, They were re-clustered outside of the geometry collection), These clusters 
-							// will need to be rendered based on the clusters position. 
-							const int32 ClusterParentIndex = CollectionClusterID[TransformGroupIndex];
-							if (ClusterParentIndex != INDEX_NONE)
-							{
-								if (InternalCluster[ClusterParentIndex])
-								{
-									const FTransform ClusterChildToWorld = ClusterChildToParentMap[RigidBodyIndex] * FTransform(Particles.R(ClusterParentIndex), Particles.X(ClusterParentIndex));
-									if (Parameters.IsCacheRecording())
-									{
-										Particles.X(RigidBodyIndex) = ClusterChildToWorld.GetTranslation();
-										Particles.R(RigidBodyIndex) = ClusterChildToWorld.GetRotation();
-									}
-									// GeomToActor = ActorToWorld.Inv() * ClusterChildToWorld * MassToLocal.Inv();
-									Transform[TransformGroupIndex] = MassToLocal[TransformGroupIndex].GetRelativeTransformReverse(ClusterChildToWorld).GetRelativeTransform(ActorToWorld);
-									Transform[TransformGroupIndex].NormalizeRotation();
-								}
-							}
-						}
-					}
-				}
-			}
-			for (int32 TransformGroupIndex = 0; TransformGroupIndex < TransformSize; TransformGroupIndex++)
-			{
-				const int32 ParentIndex = EndFrameUnparentingBuffer[TransformGroupIndex];
-				if (ParentIndex >= 0)
-				{
-					// We reuse EndFrameUnparentingBuffer potentially without reinitialization, so reset this index to -1 before it gets paged out.
-					EndFrameUnparentingBuffer[TransformGroupIndex] = -1;
-
-					Children[ParentIndex].Remove(TransformGroupIndex);
-					Parent[TransformGroupIndex] = INDEX_NONE;
-				}
-			}
-
-			//
-			//  Set rest cache on simulated object.
-			//
-			if (Parameters.IsCacheRecording())
-			{
-				check(!UpdateRecordedStateCallback);
-				UpdateRecordedState(ProxySimDuration, RigidBodyID, CollectionClusterID, InternalCluster, Particles, CollisionRule);
-			}
-
-			// one way trigger from non-simulating to simulating
-			if (!IsObjectDynamic)
-			{
-				const Chaos::TArrayCollectionArray<bool>& ExternalID = GetSolver()->GetRigidClustering().GetInternalClusterArray();
-
-				for (int32 TransformGroupIndex = 0; TransformGroupIndex < TransformSize; TransformGroupIndex++)
-				{
-					const int32 RigidBodyIndex = RigidBodyID[TransformGroupIndex];
-					if (RigidBodyIndex != INDEX_NONE)
-					{
-						bool HasBeenRemoved = (PTDynamicCollection.StatusFlags[TransformGroupIndex] & FGeometryCollection::FS_RemoveOnFracture) && Particles.Disabled(RigidBodyIndex) && ClusterID[RigidBodyIndex].Id == INDEX_NONE;
-
-						if (HasBeenRemoved || !Particles.Disabled(RigidBodyIndex))
-						{
-							if (DynamicState[TransformGroupIndex] != (int32)EObjectStateTypeEnum::Chaos_Object_Static)
-							{
-								IsObjectDynamic = true;
-								break;
-							}
-						}
-						else
-						{
-							const int32 ParentClusterID = ClusterID[RigidBodyIndex].Id;
-							if (ParentClusterID != INDEX_NONE && ExternalID[ParentClusterID]
-								&& Particles.ObjectState(ParentClusterID) != Chaos::EObjectStateType::Static)
-							{
-								IsObjectDynamic = true;
-								break;
-							}
-						}
-					}
-				}
-			}
-
-			// Can't change visibility right now so setting scale to zero instead
-			// only process if enabled on this object (ones with glass), most will skip over
-			if (Parameters.RemoveOnFractureEnabled && IsObjectDynamic)
-			{
-				for (int32 TransformGroupIndex = 0; TransformGroupIndex < TransformSize; TransformGroupIndex++)
-				{
-					const int32 RigidBodyIndex = RigidBodyID[TransformGroupIndex];
-					if (RigidBodyIndex != INDEX_NONE)
-					{
-						if ((PTDynamicCollection.StatusFlags[TransformGroupIndex] & FGeometryCollection::FS_RemoveOnFracture)
-							&& Particles.Disabled(RigidBodyIndex) && ClusterID[RigidBodyIndex].Id == INDEX_NONE)
-						{
-							PTDynamicCollection.Transform[TransformGroupIndex].SetScale3D(FVector::ZeroVector);
-						}
-					}
-				}
-			}
-#endif
-		}
-
-#if TODO_REIMPLEMENT_RIGID_CLUSTERING
-		{
-			IsObjectLoading = GetSolver()->GetRigidClustering().NumberOfPendingClusters() != 0;
-		}
-#endif
-	}
-	Commands.Empty();
-}
-
-void FGeometryCollectionPhysicsProxy::CreateRigidBodyCallback(FParticlesType& Particles)
-{
-	check(false); // This function is deprecated and should never be called.
-
-	const FGeometryCollection* RestCollection = Parameters.RestCollection;
-	FGeometryDynamicCollection* DynamicCollection = Parameters.DynamicCollection;
-	check(DynamicCollection);
-	
-	if (Parameters.Simulating && 
-		((InitializedState == ESimulationInitializationState::Unintialized) || 
-		 (InitializedState == ESimulationInitializationState::Activated)))
-	{
-#if TODO_REIMPLEMENT_RIGID_CLUSTERING
-		Chaos::FPhysicsSolver* ThisSolver = GetSolver();
-		Chaos::TArrayCollectionArray<int32>& ClusterGroupIndex = ThisSolver->GetRigidClustering().GetClusterGroupIndexArray();
-		Chaos::TArrayCollectionArray<float>& StrainArray = ThisSolver->GetRigidClustering().GetStrainArray();
-		const float StrainDefault = Parameters.DamageThreshold.Num() ? Parameters.DamageThreshold[0] : 0;
-
-		const TManagedArray<int32>& TransformIndex = RestCollection->TransformIndex;
-		const TManagedArray<int32>& BoneMap = RestCollection->BoneMap;
-		const TManagedArray<int32>& Parent = RestCollection->Parent;
-		const TManagedArray<TSet<int32>>& Children = RestCollection->Children;
-		const TManagedArray<int32>& SimulationType = RestCollection->SimulationType;
-		const TManagedArray<FVector>& Vertex = RestCollection->Vertex;
-		const TManagedArray<int32>& DynamicState = DynamicCollection->DynamicState;
-		const TManagedArray<int32>& CollisionGroup = DynamicCollection->CollisionGroup;
-		const TManagedArray<float>& Mass = RestCollection->GetAttribute<float>("Mass", FTransformCollection::TransformGroup);
-		const TManagedArray<FVector>& InertiaTensor = RestCollection->GetAttribute<FVector>("InertiaTensor", FTransformCollection::TransformGroup);
-
-		TArray<FTransform> Transform;
-		GeometryCollectionAlgo::GlobalMatrices(DynamicCollection->Transform, DynamicCollection->Parent, Transform);
-		check(DynamicCollection->Transform.Num() == Transform.Num());
-
-		// count particles to add
-		int NumSimulatedParticles = 0;
-		for (int32 Index = 0; Index < SimulatableParticles.Num(); Index++)
-		{
-			if (SimulatableParticles[Index] == true)
-			{
-				NumSimulatedParticles++;
-			}
-		}
-
-		// Add entries into simulation array
-		int NumRigids = Particles.Size();
-		BaseParticleIndex = NumRigids;
-		Particles.AddParticles(NumSimulatedParticles);
-		for (int32 Index = 0, NextId = 0; Index < SimulatableParticles.Num(); Index++)
-		{
-			if (SimulatableParticles[Index] == true)
-			{
-				RigidBodyID[Index] = NumRigids + NextId++;
-			}
-		}
-
-		// Add the rigid bodies
-		//for (int32 Index = 0; Index < TransformIndex.Num(); Index++)
-		const int32 NumGeometries = DynamicCollection->NumElements(FGeometryCollection::GeometryGroup);
-		ParallelFor(NumGeometries, [&](int32 GeometryIndex)
-		{
-			const int32 TransformGroupIndex = TransformIndex[GeometryIndex];
-			if (SimulatableParticles[TransformGroupIndex] == true)
-			{
-				const int32 RigidBodyIndex = RigidBodyID[TransformGroupIndex];
-				const FTransform WorldTransform = MassToLocal[TransformGroupIndex] * Transform[TransformGroupIndex] * Parameters.WorldTransform;
-
-				PopulateSimulatedParticle(Particles, Parameters.Shared, nullptr,// Simplicials[TransformGroupIndex].Get(),
-					Implicits[TransformGroupIndex], ThisSolver->GetMassScale() * Mass[TransformGroupIndex], ThisSolver->GetMassScale() * InertiaTensor[TransformGroupIndex],
-					RigidBodyIndex, WorldTransform, (uint8)EObjectStateTypeEnum::Chaos_Object_Dynamic, CollisionGroup[TransformGroupIndex]);
-				ClusterGroupIndex[RigidBodyIndex] = Parameters.ClusterGroupIndex;
-				Particles.CollisionParticles(RigidBodyIndex).Reset(Simplicials[TransformGroupIndex].Release());
-				if (Particles.CollisionParticles(RigidBodyIndex))
-				{
-					int32 NumCollisionParticles = Particles.CollisionParticles(RigidBodyIndex)->Size();
-					int32 CollisionParticlesSize = FMath::Max(0, FMath::Min(int(NumCollisionParticles*CollisionParticlesPerObjectFraction), NumCollisionParticles));
-					Particles.CollisionParticles(RigidBodyIndex)->Resize(CollisionParticlesSize);
-				}
-
-				StrainArray[RigidBodyIndex] = StrainDefault;
-				GetSolver()->SetPhysicsMaterial(RigidBodyIndex, Parameters.PhysicalMaterial);
-			}
-
-		});
-
-		for (FFieldSystemCommand & Cmd : Parameters.InitializationCommands)
-		{
-			if (Cmd.MetaData.Contains(FFieldSystemMetaData::EMetaType::ECommandData_ProcessingResolution)) Cmd.MetaData.Remove(FFieldSystemMetaData::EMetaType::ECommandData_ProcessingResolution);
-			FFieldSystemMetaDataProcessingResolution* ResolutionData = new FFieldSystemMetaDataProcessingResolution(EFieldResolutionType::Field_Resolution_Maximum);
-			Cmd.MetaData.Add(FFieldSystemMetaData::EMetaType::ECommandData_ProcessingResolution, TUniquePtr< FFieldSystemMetaDataProcessingResolution >(ResolutionData));
-			Commands.Add(Cmd);
-		}
-		Parameters.InitializationCommands.Empty();
-		ProcessCommands(Particles, GetSolver()->GetSolverTime());
-
-		ParallelFor(NumGeometries, [&](int32 GeometryIndex)
-		{
-			const int32 TransformGroupIndex = TransformIndex[GeometryIndex];
-			if (SimulatableParticles[TransformGroupIndex] == true)
-			{
-				const int32 RigidBodyIndex = RigidBodyID[TransformGroupIndex];
-
-				if (Parameters.InitialVelocityType == EInitialVelocityTypeEnum::Chaos_Initial_Velocity_User_Defined)
-				{
-					InitialLinearVelocity[TransformGroupIndex] = Parameters.InitialLinearVelocity;
-					InitialAngularVelocity[TransformGroupIndex] = Parameters.InitialAngularVelocity;
-					if (DynamicState[TransformGroupIndex] == (int32)EObjectStateTypeEnum::Chaos_Object_Dynamic)
-					{
-						Particles.V(RigidBodyIndex) = Parameters.InitialLinearVelocity;
-						Particles.W(RigidBodyIndex) = Parameters.InitialAngularVelocity;
-					}
-				}
-			}
-		});
-
-		InitializeKinematics(Particles, DynamicState);
-
-		InitializeRemoveOnFracture(Particles, DynamicState);
-
-		// #BG Temporary - don't cluster when playing back. Needs to be changed when kinematics are per-proxy to support
-		// kinematic to dynamic transition for clusters.
-		if (Parameters.EnableClustering)// && Parameters.CacheType != EGeometryCollectionCacheType::Play)
-		{
-			const TArray<int32> RecursiveOrder = ComputeRecursiveOrder(*RestCollection);
-
-			// num clusters
-			uint32 NumClusters=0;
-			TArray<bool> SubTreeContainsSimulatableParticle;
-			SubTreeContainsSimulatableParticle.SetNum(RecursiveOrder.Num());
-			for (const int32 TransformGroupIndex : RecursiveOrder)
-			{
-				if (Children[TransformGroupIndex].Num() > 0)
-				{
-					SubTreeContainsSimulatableParticle[TransformGroupIndex] = false;
-
-					TArray<uint32> RigidChildren, CollectionChildren;
-					for (const int32 ChildIndex : Children[TransformGroupIndex])
-					{
-						if(SubTreeContainsSimulatableParticle[ChildIndex])
-						{
-							NumClusters++;
-							SubTreeContainsSimulatableParticle[TransformGroupIndex] = true;
-							break;
-						}
-					}
-				}
-				else
-				{
-					SubTreeContainsSimulatableParticle[TransformGroupIndex] = (RigidBodyID[TransformGroupIndex] != INDEX_NONE);
-				}
-			}
-
-			const int32 ClusterStartIndex = Particles.Size();
-			Particles.AddParticles(NumClusters);
-
-			int32 ClusterRigidBodyId = ClusterStartIndex;
-			for (const int32 TransformGroupIndex : RecursiveOrder)
-			{
-				TArray<uint32> RigidChildren, CollectionChildren;
-				for (const int32 ChildIndex : Children[TransformGroupIndex])
-				{
-					if (RigidBodyID[ChildIndex] != INDEX_NONE)
-					{
-						RigidChildren.Add(RigidBodyID[ChildIndex]);
-						CollectionChildren.Add(ChildIndex);
-					}
-				}
-				if (RigidChildren.Num())
-				{
-					if (ReportTooManyChildrenNum >= 0 && RigidChildren.Num() > ReportTooManyChildrenNum)
-					{
-						UE_LOG(LogChaos, Warning, TEXT("Too many children (%d) in a single cluster:%s"), RigidChildren.Num(), *Parameters.Name);
-					}
-
-					RigidBodyID[TransformGroupIndex] = ClusterRigidBodyId;
-					Chaos::FClusterCreationParameters<float> CreationParameters;
-					CreationParameters.RigidBodyIndex = RigidBodyID[TransformGroupIndex];
-					BuildClusters(TransformGroupIndex, CollectionChildren, RigidChildren, CreationParameters);
-					ClusterRigidBodyId++;
-				}
-			}
-		}
-
-		NumParticles = Particles.Size() - BaseParticleIndex;
-
-		// Set Connectivity
-		const int32 NumTransforms = DynamicCollection->NumElements(FGeometryCollection::TransformGroup);
-		ParallelFor(NumTransforms, [&](int32 TransformGroupIndex)
-		{
-			if (Children[TransformGroupIndex].Num() > 0)
-			{
-				if (RigidBodyID[TransformGroupIndex] != INDEX_NONE)
-				{
-					GetSolver()->GetRigidClustering().GenerateConnectionGraph(RigidBodyID[TransformGroupIndex]);
-				}
-			}
-		});
-
-		// If we're recording and want to start immediately caching then we should cache the rest state
-		if (Parameters.IsCacheRecording() && Parameters.CacheBeginTime == 0.0f)
-		{
-			if (UpdateRecordedStateCallback)
-			{
-				UpdateRecordedStateCallback(0.0f, RigidBodyID, Particles, GetSolver()->GetCollisionConstraints());
-			}
-		}
-
-
-		if (InitializedState == ESimulationInitializationState::Activated)
-		{
-			//
-			//  Activated bodies has already been called so we are good to go.
-			//
-			InitializedState = ESimulationInitializationState::Initialized;
-
-			if (Parameters.EnableClustering && Parameters.ClusterGroupIndex)
-			{
-				GetSolver()->GetRigidClustering().IncrementPendingClusterCounter(Parameters.ClusterGroupIndex);
-				GetSolver()->GetRigidClustering().DecrementPendingClusterCounter(Parameters.ClusterGroupIndex);
-			}
-		}
-		else if (InitializedState == ESimulationInitializationState::Unintialized)
-		{
-
-			//
-			//  Activated bodies has not been called, se we are waiting 
-			//  to become active. Deactivate all bodies, and wait for
-			//  ActivatedBodies to be called, and defer the cluster initialization
-			//
-			InitializedState = ESimulationInitializationState::Created;
-			for (int32 TransformGroupIndex = 0; TransformGroupIndex < NumTransforms; TransformGroupIndex++)
-			{
-				const int32 RigidBodyIndex = RigidBodyID[TransformGroupIndex];
-				if (RigidBodyIndex != INDEX_NONE)
-				{
-					if (!Particles.Disabled(RigidBodyIndex))
-					{
-						PendingActivationList.Add(TransformGroupIndex);
-						GetSolver()->GetEvolution()->DisableParticle(RigidBodyIndex);
-					}
-				}
-			}
-
-			//
-			//  Clustering needs to advertise its group id to the cluster so
-			//  that the group is not initialized before all the bodies are
-			//  loaded and created. 
-			//
-			if (Parameters.EnableClustering && Parameters.ClusterGroupIndex )
-			{
-				GetSolver()->GetRigidClustering().IncrementPendingClusterCounter(Parameters.ClusterGroupIndex);
-			}
-		}
-		else
-		{
-			// unknown initialization state in creation callback
-			ensure(false);
-		}
-#endif
-	}
-}
-
-void FGeometryCollectionPhysicsProxy::ActivateBodies()
-{
-	if (Parameters.Simulating)
-	{
-		if (InitializedState == ESimulationInitializationState::Created)
-		{
-#if TODO_REIMPLEMENT_RIGID_CLUSTERING
-			int32 ParentIndex = INDEX_NONE;
-
-			if (Parameters.EnableClustering && Parameters.ClusterGroupIndex)
-			{
-				auto& Clustering = GetSolver()->GetEvolution()->GetRigidClustering();
-				Clustering.DecrementPendingClusterCounter(Parameters.ClusterGroupIndex);
-				ParentIndex = Parameters.ClusterGroupIndex;
-			}
-#endif
-#if TODO_REIMPLEMENT_DEFERRED_ACTIVATION
-			Chaos::FPhysicsSolver::FParticlesType & Particles = GetSolver()->GetRigidParticles();
-			for (uint32 TransformGroupIndex : PendingActivationList)
-			{
-//				int32 RigidBodyIndex = RigidBodyID[TransformGroupIndex];
-//				checkSlow(RigidBodyIndex != INDEX_NONE);
-				if(Chaos::TPBDRigidParticleHandle<float, 3>* Handle = SolverParticleHandles[TransformGroupIndex])
-//				if (Particles.Disabled(RigidBodyIndex))
-				{
-					GetSolver()->GetEvolution()->EnableParticle(RigidBodyIndex, ParentIndex);
-				}
-			}
-#endif
-			PendingActivationList.Reset(0);
-
-			InitializedState = ESimulationInitializationState::Initialized;
-		}
-		else if (InitializedState == ESimulationInitializationState::Unintialized)
-		{
-			InitializedState = ESimulationInitializationState::Activated;
-		}
-		else
-		{
-			// unknown initialization state in activate bodies
-			ensure(false);
-		}
-	}
-}
-
-void FGeometryCollectionPhysicsProxy::BindParticleCallbackMapping(Chaos::TArrayCollectionArray<PhysicsProxyWrapper> & PhysicsProxyReverseMap, Chaos::TArrayCollectionArray<int32> & ParticleIDReverseMap)
-{
-	if (InitializedState== ESimulationInitializationState::Initialized)
-	{
-		for (int32 Index = 0; Index < RigidBodyID.Num(); Index++)
-		{
-			const int32 RigidBodyIndex = RigidBodyID[Index];
-			if (RigidBodyIndex != INDEX_NONE)
-			{
-				PhysicsProxyReverseMap[RigidBodyIndex] = {this, EPhysicsProxyType::GeometryCollectionType};
-				ParticleIDReverseMap[RigidBodyIndex] = Index;
-			}
-		}
-	}
-}
-
-void FGeometryCollectionPhysicsProxy::ContiguousIndices(TArray<ContextIndex>& Array, const Chaos::FPhysicsSolver* RigidSolver, EFieldResolutionType ResolutionType, bool bForce = true)
-{
-	check(false); // Not currently called
-	if (bForce)
-	{
-#if TODO_REIMPLEMENT_GET_RIGID_PARTICLES
-		const Chaos::FPhysicsSolver::FParticlesType & Particles = RigidSolver->GetRigidParticles();
-		if (ResolutionType == EFieldResolutionType::Field_Resolution_Minimal)
-		{
-			const Chaos::TArrayCollectionArray<Chaos::ClusterId>& ClusterIdArray = RigidSolver->GetRigidClustering().GetClusterIdsArray();
-
-
-			//  Generate a Index mapping between the rigid body indices and 
-			//  the particle indices. This allows the geometry collection to
-			//  evaluate only its own ACTIVE particles + ClusterChildren
-			int32 NumIndices = 0;
-			Array.SetNumUninitialized(RigidBodyID.Num());
-			for (int32 i = 0; i < RigidBodyID.Num(); i++)
-			{
-				const int32 RigidBodyIndex = RigidBodyID[i];
-				if (RigidBodyIndex != INDEX_NONE && !Particles.Disabled(RigidBodyIndex)) // active bodies
-				{
-					Array[NumIndices] = { RigidBodyID[i],i };
-					NumIndices++;
-				}
-				if (ClusterIdArray[RigidBodyIndex].Id != INDEX_NONE && !Particles.Disabled(ClusterIdArray[RigidBodyIndex].Id)) // children
-				{
-					Array[NumIndices] = { RigidBodyID[i],i };
-					NumIndices++;
-				}
-			}
-			Array.SetNum(NumIndices);
-		}
-		else if (ResolutionType == EFieldResolutionType::Field_Resolution_Maximum)
-		{
-			//  Generate a Index mapping between the rigid body indices and 
-			//  the particle indices. This allows the geometry collection to
-			//  evaluate only its own particles. 
-			int32 NumIndices = 0;
-			Array.SetNumUninitialized(RigidBodyID.Num());
-			for (int32 i = 0; i < RigidBodyID.Num(); i++)
-			{
-				const int32 RigidBodyIndex = RigidBodyID[i];
-				if (RigidBodyIndex != INDEX_NONE)
-				{
-					Array[NumIndices] = { RigidBodyIndex, i };
-					NumIndices++;
-				}
-			}
-			Array.SetNum(NumIndices);
-		}
-#endif
-	}
-}
-
-void FGeometryCollectionPhysicsProxy::ProcessCommands(FParticlesType& Particles, const float Time)
-{
-	FGeometryDynamicCollection* Collection = Parameters.DynamicCollection;
-	check(Collection);
-
-	// Process Particle-Collection commands
-	if (Commands.Num())
-	{
-		TArray<ContextIndex> IndicesArray;
-		Chaos::FPhysicsSolver* CurrentSolver = GetSolver();
-
-		for (int32 CommandIndex = Commands.Num()-1; 0<=CommandIndex; CommandIndex--)
-		{
-
-			//
-			// Extract command and set metadata
-			//
-			FFieldSystemCommand & Command = Commands[CommandIndex];
-			EFieldResolutionType ResolutionType = EFieldResolutionType::Field_Resolution_Minimal;
-			if (Command.MetaData.Contains(FFieldSystemMetaData::EMetaType::ECommandData_ProcessingResolution))
-			{
-				check(Command.MetaData[FFieldSystemMetaData::EMetaType::ECommandData_ProcessingResolution] != nullptr);
-				ResolutionType = static_cast<FFieldSystemMetaDataProcessingResolution*>(Command.MetaData[FFieldSystemMetaData::EMetaType::ECommandData_ProcessingResolution].Get())->ProcessingResolution;
-			}
-
-			if (Command.TargetAttribute == GetGeometryCollectionPhysicsTypeName(EGeometryCollectionPhysicsTypeEnum::Chaos_DynamicState))
-			{
-				if (ensureMsgf(Command.RootNode->Type() == FFieldNode<int32>::StaticType(),
-					TEXT("Field based evaluation of the simulations 'DynamicState' parameter expects int32 field inputs.")))
-				{
-					FGeometryCollectionPhysicsProxy::ContiguousIndices(IndicesArray, CurrentSolver, ResolutionType, IndicesArray.Num() != Particles.Size());
-					if (IndicesArray.Num())
-					{
-						TArrayView<ContextIndex> IndexView(&(IndicesArray[0]), IndicesArray.Num());
-
-						FVector * xptr = &(Particles.X(0));
-						TArrayView<FVector> SamplesView(xptr, Particles.Size());
-
-						FFieldContext Context{
-							IndexView,
-							SamplesView,
-							Command.MetaData
-						};
-
-						TArrayView<int32> DynamicStateView(&(Collection->DynamicState[0]), Collection->DynamicState.Num());
-						static_cast<const FFieldNode<int32> *>(Command.RootNode.Get())->Evaluate(Context, DynamicStateView);
-
-						PushKinematicStateToSolver(Particles);
-					}
-				}
-				Commands.RemoveAt(CommandIndex);
-			}
-			else if (Command.TargetAttribute == GetGeometryCollectionPhysicsTypeName(EGeometryCollectionPhysicsTypeEnum::Chaos_InitialLinearVelocity))
-			{
-				if (ensureMsgf(Parameters.InitialVelocityType == EInitialVelocityTypeEnum::Chaos_Initial_Velocity_User_Defined,
-					TEXT("Field based evaluation of the simulations 'InitialLinearVelocity' requires the geometry collection be set to User Defined Initial Velocity")))
-				{
-					if (ensureMsgf(Command.RootNode->Type() == FFieldNode<FVector>::StaticType(),
-						TEXT("Field based evaluation of the simulations 'InitialLinearVelocity' parameter expects FVector field inputs.")))
-					{
-						FGeometryCollectionPhysicsProxy::ContiguousIndices(IndicesArray, CurrentSolver, ResolutionType, IndicesArray.Num() != Particles.Size());
-						if (IndicesArray.Num())
-						{
-							TArrayView<ContextIndex> IndexView(&(IndicesArray[0]), IndicesArray.Num());
-
-							FVector * xptr = &(Particles.X(0));
-							TArrayView<FVector> SamplesView(xptr, Particles.Size());
-
-							FFieldContext Context{
-								IndexView,
-								SamplesView,
-								Command.MetaData
-							};
-
-							TArrayView<FVector> ResultsView(&(InitialLinearVelocity[0]), InitialLinearVelocity.Num());
-							static_cast<const FFieldNode<FVector> *>(Command.RootNode.Get())->Evaluate(Context, ResultsView);
-						}
-					}
-				}
-				Commands.RemoveAt(CommandIndex);
-			}
-			else if (Command.TargetAttribute == GetGeometryCollectionPhysicsTypeName(EGeometryCollectionPhysicsTypeEnum::Chaos_InitialAngularVelocity))
-			{
-				if (ensureMsgf(Parameters.InitialVelocityType == EInitialVelocityTypeEnum::Chaos_Initial_Velocity_User_Defined,
-					TEXT("Field based evaluation of the simulations 'InitialAngularVelocity' requires the geometry collection be set to User Defined Initial Velocity")))
-				{
-					if (ensureMsgf(Command.RootNode->Type() == FFieldNode<FVector>::StaticType(),
-						TEXT("Field based evaluation of the simulations 'InitialAngularVelocity' parameter expects FVector field inputs.")))
-					{
-						FGeometryCollectionPhysicsProxy::ContiguousIndices(IndicesArray, CurrentSolver, ResolutionType, IndicesArray.Num() != Particles.Size());
-						if (IndicesArray.Num())
-						{
-							TArrayView<ContextIndex> IndexView(&(IndicesArray[0]), IndicesArray.Num());
-
-							FVector * xptr = &(Particles.X(0));
-							TArrayView<FVector> SamplesView(xptr, Particles.Size());
-
-							FFieldContext Context{
-								IndexView,
-								SamplesView,
-								Command.MetaData
-							};
-
-							TArrayView<FVector> ResultsView(&(InitialAngularVelocity[0]), InitialAngularVelocity.Num());
-							static_cast<const FFieldNode<FVector> *>(Command.RootNode.Get())->Evaluate(Context, ResultsView);
-						}
-					}
-				}
-				Commands.RemoveAt(CommandIndex);
-			}
-		}
-	}
-
-	// Process Particle-Particle commands
-	if (Commands.Num())
-	{
-		Chaos::FPhysicsSolver* CurrentSolver = GetSolver();
-
-		//  Generate a Index mapping between the rigid body indices and 
-		//  the particle indices. This allows the geometry collection to
-		//  evaluate only its own particles. 
-		TArray<ContextIndex> IndicesArray;
-		int32 NumIndices = 0;
-		IndicesArray.SetNumUninitialized(RigidBodyID.Num());
-		for (int32 i = 0; i < RigidBodyID.Num(); i++)
-		{
-			if (RigidBodyID[i] != INDEX_NONE)
-			{
-				IndicesArray[NumIndices] = { RigidBodyID[i], RigidBodyID[i] };
-				NumIndices++;
-			}
-		}
-		IndicesArray.SetNum(NumIndices);
-
-		for (int32 CommandIndex = Commands.Num() - 1; 0 <= CommandIndex; CommandIndex--)
-		{
-			FFieldSystemCommand & Command = Commands[CommandIndex];
-			if (Command.TargetAttribute == GetGeometryCollectionPhysicsTypeName(EGeometryCollectionPhysicsTypeEnum::Chaos_LinearVelocity))
-			{
-				if (ensureMsgf(Command.RootNode->Type() == FFieldNode<FVector>::StaticType(),
-					TEXT("Field based evaluation of the simulations 'LinearVelocity' parameter expects FVector field inputs.")))
-				{
-					FVector * xptr = &(Particles.X(0));
-					TArrayView<FVector> SamplesView(xptr, Particles.Size());
-					TArrayView<ContextIndex> IndexView(&(IndicesArray[0]), IndicesArray.Num());
-
-					FFieldContext Context{
-						IndexView,
-						SamplesView,
-						Command.MetaData
-					};
-
-					FVector * vptr = &(Particles.V(0));
-					TArrayView<FVector> ResultsView(vptr, Particles.Size());
-					static_cast<const FFieldNode<FVector> *>(Command.RootNode.Get())->Evaluate(Context, ResultsView);
-				}
-				Commands.RemoveAt(CommandIndex);
-			}
-			else if (Command.TargetAttribute == GetGeometryCollectionPhysicsTypeName(EGeometryCollectionPhysicsTypeEnum::Chaos_AngularVelocity))
-			{
-				if (ensureMsgf(Command.RootNode->Type() == FFieldNode<FVector>::StaticType(),
-					TEXT("Field based evaluation of the simulations 'AngularVelocity' parameter expects FVector field inputs.")))
-				{
-
-					FVector * xptr = &(Particles.X(0));
-					TArrayView<FVector> SamplesView(xptr, Particles.Size());
-					TArrayView<ContextIndex> IndexView(&(IndicesArray[0]), IndicesArray.Num());
-
-					FFieldContext Context{
-						IndexView,
-						SamplesView,
-						Command.MetaData
-					};
-
-					FVector * vptr = &(Particles.W(0));
-					TArrayView<FVector> ResultsView(vptr, Particles.Size());
-					static_cast<const FFieldNode<FVector> *>(Command.RootNode.Get())->Evaluate(Context, ResultsView);
-				}
-				Commands.RemoveAt(CommandIndex);
-			}
-			else if (Command.TargetAttribute == GetGeometryCollectionPhysicsTypeName(EGeometryCollectionPhysicsTypeEnum::Chaos_CollisionGroup))
-			{
-				if (ensureMsgf(Command.RootNode->Type() == FFieldNode<int32>::StaticType(),
-					TEXT("Field based evaluation of the simulations 'CollisionGroup' parameter expects int32 field inputs.")))
-				{
-					FVector * xptr = &(Particles.X(0));
-					TArrayView<FVector> SamplesView(xptr, Particles.Size());
-					TArrayView<ContextIndex> IndexView(&(IndicesArray[0]), IndicesArray.Num());
-
-					FFieldContext Context{
-						IndexView,
-						SamplesView,
-						Command.MetaData
-					};
-
-					int32 * cptr = &(Particles.CollisionGroup(0));
-					TArrayView<int32> ResultsView(cptr, Particles.Size());
-					static_cast<const FFieldNode<int32> *>(Command.RootNode.Get())->Evaluate(Context, ResultsView);
-				}
-				Commands.RemoveAt(CommandIndex);
-			}
-		}
-	}
-}
-
-void FGeometryCollectionPhysicsProxy::PushKinematicStateToSolver(FParticlesType& Particles)
-{
-	if (InitializedState == ESimulationInitializationState::Initialized)
-	{
-		FGeometryDynamicCollection* Collection = Parameters.DynamicCollection;
-		if (Collection->Transform.Num())
-		{
-			TManagedArray<int32>& DynamicState = Collection->GetAttribute<int32>("DynamicState", FGeometryCollection::TransformGroup);
-
-#if TODO_REIMPLEMENT_RIGID_CLUSTERING
-			TSet<int32> ClustersToUpdate;
-			const Chaos::TArrayCollectionArray<Chaos::ClusterId>& ClusterID = GetSolver()->GetRigidClustering().GetClusterIdsArray();
-
-			for (int32 TransformGroupIndex = 0; TransformGroupIndex < DynamicState.Num(); TransformGroupIndex++)
-			{
-				if (RigidBodyID[TransformGroupIndex] != INDEX_NONE)
-				{
-					int32 RigidBodyIndex = RigidBodyID[TransformGroupIndex];
-					if (DynamicState[TransformGroupIndex] == (int32)EObjectStateTypeEnum::Chaos_Object_Dynamic
-						&& (Particles.ObjectState(RigidBodyIndex) == Chaos::EObjectStateType::Kinematic || Particles.ObjectState(RigidBodyIndex) == Chaos::EObjectStateType::Static)
-						&& FLT_EPSILON < Particles.M(RigidBodyIndex))
-					{
-						Particles.SetObjectState(RigidBodyIndex, Chaos::EObjectStateType::Dynamic);
-
-						if (Parameters.InitialVelocityType == EInitialVelocityTypeEnum::Chaos_Initial_Velocity_User_Defined)
-						{
-							Particles.V(RigidBodyIndex) = Chaos::TVector<float, 3>(InitialLinearVelocity[TransformGroupIndex]);
-							Particles.W(RigidBodyIndex) = Chaos::TVector<float, 3>(InitialAngularVelocity[TransformGroupIndex]);
-						}
-
-						if (!Particles.Disabled(RigidBodyIndex) && Particles.Sleeping(RigidBodyIndex))
-						{
-							Particles.SetSleeping(RigidBodyIndex, false);
-							GetSolver()->ActiveIndices().Add(RigidBodyIndex);
-						}
-						else
-						{
-							if (ClusterID[RigidBodyIndex].Id != INDEX_NONE)
-							{
-								int32 ActiveClusterIndex = GetSolver()->GetRigidClustering().GetActiveClusterIndex(RigidBodyIndex);
-								if (ActiveClusterIndex != INDEX_NONE)
-								{
-									ClustersToUpdate.Add(ActiveClusterIndex);
-								}
-							}
-						}
-					}
-					else if ((DynamicState[TransformGroupIndex] == (int32)EObjectStateTypeEnum::Chaos_Object_Kinematic)
-						&& (Particles.ObjectState(RigidBodyIndex) == Chaos::EObjectStateType::Dynamic)
-						&& FLT_EPSILON < Particles.M(RigidBodyIndex))
-					{
-						Particles.SetObjectState(RigidBodyIndex, Chaos::EObjectStateType::Kinematic);
-
-						if (ClusterID[RigidBodyIndex].Id != INDEX_NONE)
-						{
-							int32 ActiveClusterIndex = GetSolver()->GetRigidClustering().GetActiveClusterIndex(RigidBodyIndex);
-							if (ActiveClusterIndex != INDEX_NONE)
-							{
-								ClustersToUpdate.Add(ActiveClusterIndex);
-							}
-						}
-					}
-					else if ((DynamicState[TransformGroupIndex] == (int32)EObjectStateTypeEnum::Chaos_Object_Static)
-						&& (Particles.ObjectState(RigidBodyIndex) == Chaos::EObjectStateType::Dynamic)
-						&& FLT_EPSILON < Particles.M(RigidBodyIndex))
-					{
-						Particles.SetObjectState(RigidBodyIndex, Chaos::EObjectStateType::Static);
-
-						if (ClusterID[RigidBodyIndex].Id != INDEX_NONE)
-						{
-							int32 ActiveClusterIndex = GetSolver()->GetRigidClustering().GetActiveClusterIndex(RigidBodyIndex);
-							if (ActiveClusterIndex != INDEX_NONE)
-							{
-								ClustersToUpdate.Add(ActiveClusterIndex);
-							}
-						}
-					}
-					else if ((DynamicState[TransformGroupIndex] == (int32)EObjectStateTypeEnum::Chaos_Object_Sleeping)
-						&& (Particles.ObjectState(RigidBodyIndex) == Chaos::EObjectStateType::Dynamic))
-					{
-						Particles.SetObjectState(RigidBodyIndex, Chaos::EObjectStateType::Sleeping);
-						GetSolver()->ActiveIndices().Remove(RigidBodyIndex);
-					}
-					else if ((DynamicState[TransformGroupIndex] == (int32)EObjectStateTypeEnum::Chaos_Object_Dynamic)
-						&& (Particles.ObjectState(RigidBodyIndex) == Chaos::EObjectStateType::Sleeping))
-					{
-						Particles.SetObjectState(RigidBodyIndex, Chaos::EObjectStateType::Dynamic);
-						GetSolver()->ActiveIndices().Add(RigidBodyIndex);
-					}
-				}
-			}
-#endif
-		}
-	}
-}
-
-void FGeometryCollectionPhysicsProxy::ParameterUpdateCallback(FParticlesType& Particles, const float Time)
-{
-	if (InitializedState == ESimulationInitializationState::Initialized)
-	{
-		FGeometryDynamicCollection* Collection = Parameters.DynamicCollection;
-		check(Collection);
-
-		if (Collection->Transform.Num())
-		{
-			ProcessCommands(Particles, Time);
-
-
-			if (Parameters.RecordedTrack)
-			{
-				float ReverseTime = Parameters.RecordedTrack->GetLastTime() - Time + Parameters.ReverseCacheBeginTime;
-				// @todo(mlentine): We shouldn't need to do this every frame
-				if (Parameters.IsCacheRecording() && Time > Parameters.ReverseCacheBeginTime && Parameters.ReverseCacheBeginTime != 0 && Parameters.RecordedTrack->IsTimeValid(ReverseTime))
-				{
-					for (int32 Index = 0; Index < RigidBodyID.Num(); Index++)
-					{
-						int32 RigidBodyIndex = RigidBodyID[Index];
-
-						// Check index, will be invalid for cluster parents.
-						if (RigidBodyIndex != INDEX_NONE)
-						{
-							Particles.InvM(RigidBodyIndex) = 0.f;
-							Particles.InvI(RigidBodyIndex) = Chaos::PMatrix<float, 3, 3>(0.f);
-						}
-					}
-				}
-			}
-			/* @question : Should we tell the solver the mass has changed ? */
-		}
-	}
-}
-
-void FGeometryCollectionPhysicsProxy::DisableCollisionsCallback(TSet<TTuple<int32, int32>>& InPairs)
-{
-
-}
-
-void FGeometryCollectionPhysicsProxy::AddForceCallback(FParticlesType& InParticles, const float InDt, const int32 InIndex)
-{
-	// gravity forces managed directly on the solver for now
-}
-
-void FGeometryCollectionPhysicsProxy::FieldForcesUpdateCallback(Chaos::FPhysicsSolver* InSolver, FParticlesType& Particles, Chaos::TArrayCollectionArray<FVector> & Force, Chaos::TArrayCollectionArray<FVector> & Torque, const float Time)
-{
-	if (InitializedState == ESimulationInitializationState::Initialized)
-	{
-		if (Commands.Num())
-		{
-			// @todo: This seems like a waste if we just want to get everything
-			int32 Counter = 0;
-			TArray<ContextIndex> IndicesArray;
-			IndicesArray.AddUninitialized(RigidBodyID.Num());
-			for (int32 i = 0; i < RigidBodyID.Num(); i++)
-			{
-				if (RigidBodyID[i] != INDEX_NONE)
-				{
-					IndicesArray[i] = { RigidBodyID[i],RigidBodyID[i] };
-					Counter++;
-				}
-			}
-			IndicesArray.SetNum(Counter, false);
-			TArrayView<ContextIndex> IndexView(&(IndicesArray[0]), IndicesArray.Num());
-
-
-			for (int32 CommandIndex = 0; CommandIndex < Commands.Num(); CommandIndex++)
-			{
-				FFieldSystemCommand & Command = Commands[CommandIndex];
-
-				if (Command.TargetAttribute == GetGeometryCollectionPhysicsTypeName(EGeometryCollectionPhysicsTypeEnum::Chaos_LinearForce))
-				{
-					if (ensureMsgf(Command.RootNode->Type() == FFieldNode<FVector>::StaticType(),
-						TEXT("Field based evaluation of the simulations 'LinearForce' parameter expects FVector field inputs.")))
-					{
-						FVector * tptr = &(Particles.X(0));
-						TArrayView<FVector> SamplesView(tptr, int32(Particles.Size()));
-
-						FFieldContext Context{
-							IndexView,
-							SamplesView,
-							Command.MetaData
-						};
-						TArrayView<FVector> ForceView(&(Force[0]), Force.Num());
-						static_cast<const FFieldNode<FVector> *>(Command.RootNode.Get())->Evaluate(Context, ForceView);
-					}
-					Commands.RemoveAt(CommandIndex);
-				}
-				else if (Command.TargetAttribute == GetGeometryCollectionPhysicsTypeName(EGeometryCollectionPhysicsTypeEnum::Chaos_AngularTorque))
-				{
-					if (ensureMsgf(Command.RootNode->Type() == FFieldNode<FVector>::StaticType(),
-						TEXT("Field based evaluation of the simulations 'AngularTorque' parameter expects FVector field inputs.")))
-					{
-						FVector * tptr = &(Particles.X(0));
-						TArrayView<FVector> SamplesView(tptr, int32(Particles.Size()));
-
-						FFieldContext Context{
-							IndexView,
-							SamplesView,
-							Command.MetaData
-						};
-						TArrayView<FVector> TorqueView(&(Torque[0]), Torque.Num());
-						static_cast<const FFieldNode<FVector> *>(Command.RootNode.Get())->Evaluate(Context, TorqueView);
-					}
-					Commands.RemoveAt(CommandIndex);
-				}
-
-			}
-		}
-	}
-}
-
 int32 ReportNoLevelsetCluster = 0;
 FAutoConsoleVariableRef CVarReportNoLevelsetCluster(TEXT("p.gc.ReportNoLevelsetCluster"), ReportNoLevelsetCluster, TEXT("Report any cluster objects without levelsets"));
 
 DECLARE_CYCLE_STAT(TEXT("FGeometryCollectionPhysicsProxy::BuildClusters"), STAT_BuildClusters, STATGROUP_Chaos);
 DECLARE_CYCLE_STAT(TEXT("FGeometryCollectionPhysicsProxy::BuildClusters:GlobalMatrices"), STAT_BuildClustersGlobalMatrices, STATGROUP_Chaos);
 
-
 Chaos::TPBDRigidClusteredParticleHandle<float, 3>* 
 FGeometryCollectionPhysicsProxy::BuildClusters(
 	const uint32 CollectionClusterIndex, // TransformGroupIndex
@@ -2502,12 +1047,6 @@
 		UE_LOG(LogChaos, Warning, TEXT("Union object generated for cluster:%s"), *Parameters.Name);
 	}
 
-	if(Parameters.InitialVelocityType == EInitialVelocityTypeEnum::Chaos_Initial_Velocity_User_Defined)
-	{
-		InitialLinearVelocity[CollectionClusterIndex] = Parameters.InitialLinearVelocity;
-		InitialAngularVelocity[CollectionClusterIndex] = Parameters.InitialAngularVelocity;
-	}
-
 	if (Parent->InvM() == 0.0)
 	{
 		if (Parent->ObjectState() == Chaos::EObjectStateType::Static)
@@ -2595,16 +1134,71 @@
 	return Parent;
 }
 
-int32 FGeometryCollectionPhysicsProxy::CalculateHierarchyLevel(FGeometryDynamicCollection* GeometryCollection, int32 TransformIndex)
+void FGeometryCollectionPhysicsProxy::ContiguousIndices(TArray<ContextIndex>& Array, const Chaos::FPhysicsSolver* RigidSolver, EFieldResolutionType ResolutionType, bool bForce = true)
 {
+	check(false); // Not currently called
+	if (bForce)
+	{
+#if TODO_REIMPLEMENT_GET_RIGID_PARTICLES
+		const Chaos::FPhysicsSolver::FParticlesType & Particles = RigidSolver->GetRigidParticles();
+		if (ResolutionType == EFieldResolutionType::Field_Resolution_Minimal)
+		{
+			const Chaos::TArrayCollectionArray<Chaos::ClusterId>& ClusterIdArray = RigidSolver->GetRigidClustering().GetClusterIdsArray();
+
+
+			//  Generate a Index mapping between the rigid body indices and 
+			//  the particle indices. This allows the geometry collection to
+			//  evaluate only its own ACTIVE particles + ClusterChildren
+			int32 NumIndices = 0;
+			Array.SetNumUninitialized(RigidBodyID.Num());
+			for (int32 i = 0; i < RigidBodyID.Num(); i++)
+			{
+				const int32 RigidBodyIndex = RigidBodyID[i];
+				if (RigidBodyIndex != INDEX_NONE && !Particles.Disabled(RigidBodyIndex)) // active bodies
+				{
+					Array[NumIndices] = { RigidBodyID[i],i };
+					NumIndices++;
+				}
+				if (ClusterIdArray[RigidBodyIndex].Id != INDEX_NONE && !Particles.Disabled(ClusterIdArray[RigidBodyIndex].Id)) // children
+				{
+					Array[NumIndices] = { RigidBodyID[i],i };
+					NumIndices++;
+				}
+			}
+			Array.SetNum(NumIndices);
+		}
+		else if (ResolutionType == EFieldResolutionType::Field_Resolution_Maximum)
+		{
+			//  Generate a Index mapping between the rigid body indices and 
+			//  the particle indices. This allows the geometry collection to
+			//  evaluate only its own particles. 
+			int32 NumIndices = 0;
+			Array.SetNumUninitialized(RigidBodyID.Num());
+			for (int32 i = 0; i < RigidBodyID.Num(); i++)
+			{
+				const int32 RigidBodyIndex = RigidBodyID[i];
+				if (RigidBodyIndex != INDEX_NONE)
+				{
+					Array[NumIndices] = { RigidBodyIndex, i };
+					NumIndices++;
+				}
+			}
+			Array.SetNum(NumIndices);
+		}
+#endif
+	}
+}
+
+int32 FGeometryCollectionPhysicsProxy::CalculateHierarchyLevel(
+	const FGeometryDynamicCollection* GeometryCollection, 
+	int32 TransformIndex) const
+{
 	int32 Level = 0;
-
 	while (GeometryCollection->Parent[TransformIndex] != -1)
 	{
 		TransformIndex = GeometryCollection->Parent[TransformIndex];
 		Level++;
 	}
-
 	return Level;
 }
 
@@ -2614,21 +1208,33 @@
 	FGeometryDynamicCollection* DynamicCollection = Parameters.DynamicCollection;
 
 	DynamicCollection->AddExternalAttribute("SimulatableParticles", FGeometryCollection::TransformGroup, SimulatableParticles);
-	DynamicCollection->AddExternalAttribute("RigidBodyID", FTransformCollection::TransformGroup, RigidBodyID);
 	DynamicCollection->AddExternalAttribute("SolverClusterID", FTransformCollection::TransformGroup, SolverClusterID);
 	DynamicCollection->AddExternalAttribute("MassToLocal", FTransformCollection::TransformGroup, MassToLocal);
 	DynamicCollection->CopyAttribute(RestCollection, "MassToLocal", FTransformCollection::TransformGroup);
 	DynamicCollection->AddExternalAttribute("CollisionStructureID", FTransformCollection::TransformGroup, CollisionStructureID);
+
+	// TODO: Move these to the geometry group?
 	DynamicCollection->AddExternalAttribute(SimplicialsAttribute, FTransformCollection::TransformGroup, Simplicials);
 	DynamicCollection->AddExternalAttribute(ImplicitsAttribute, FTransformCollection::TransformGroup, Implicits);
-	DynamicCollection->AddExternalAttribute(SolverParticleHandlesAttribute, FTransformCollection::TransformGroup, SolverParticleHandles);
-	DynamicCollection->AddExternalAttribute(SolverClusterHandlesAttribute, FTransformCollection::TransformGroup, SolverClusterHandles/*SolverParticleHandles*/);
+
+	// We don't have access to rigid body ids anymore, so we aren't using this currently.
+	// There is still a bunch of legacy code in this file that builds against the RigidBodyID
+	// variable, so it's still defined in the header.
+	//DynamicCollection->AddExternalAttribute("RigidBodyID", FTransformCollection::TransformGroup, RigidBodyID);
+
+	// Holds TUniquePtr<Chaos::TPBDGeometryCollectionParticle<float,3>> for leaf nodes (non parents).
 	DynamicCollection->AddExternalAttribute(GTGeometryParticleAttribute, FTransformCollection::TransformGroup, GTGeometryParticles);
+	// Holds Chaos::TPBDGeometryCollectionParticleHandle<float, 3>* for all nodes
+	DynamicCollection->AddExternalAttribute(SolverParticleHandlesAttribute, FTransformCollection::TransformGroup, SolverParticleHandles);
+	// Holds Chaos::TPBDRigidClusteredParticleHandle<float, 3>* for cluster parent nodes
+	DynamicCollection->AddExternalAttribute(SolverClusterHandlesAttribute, FTransformCollection::TransformGroup, SolverClusterHandles);
 
 	if (Parameters.InitialVelocityType == EInitialVelocityTypeEnum::Chaos_Initial_Velocity_User_Defined)
 	{
-		DynamicCollection->AddExternalAttribute("InitialAngularVelocity", FTransformCollection::TransformGroup, InitialAngularVelocity);
 		DynamicCollection->AddExternalAttribute("InitialLinearVelocity", FTransformCollection::TransformGroup, InitialLinearVelocity);
+		DynamicCollection->AddExternalAttribute("InitialAngularVelocity", FTransformCollection::TransformGroup, InitialAngularVelocity);
+		InitialLinearVelocity.Fill(Parameters.InitialLinearVelocity);
+		InitialAngularVelocity.Fill(Parameters.InitialAngularVelocity);
 	}
 
 	// The RestCollection may have been authored with TUniquePtr<FImplicitObject> under the name 
@@ -2668,7 +1274,6 @@
 			(*RestCollectionSharedImplicits)[Index] = Implicits[Index];
 		}
 
-		RigidBodyID[Index] = INDEX_NONE;
 		SolverClusterID[Index] = INDEX_NONE;
 	}
 
@@ -2680,7 +1285,8 @@
 
 	if (Parameters.RestCollection->HasAttribute(SimplicialsAttribute, FTransformCollection::TransformGroup))
 	{
-		const auto& RestSimplicials = Parameters.RestCollection->GetAttribute<TUniquePtr<FSimplicial>>(SimplicialsAttribute, FTransformCollection::TransformGroup);
+		const auto& RestSimplicials = Parameters.RestCollection->GetAttribute<TUniquePtr<FSimplicial>>(
+			SimplicialsAttribute, FTransformCollection::TransformGroup);
 		for (int32 Index = DynamicCollection->NumElements(FTransformCollection::TransformGroup) - 1; 0 <= Index; Index--)
 		{
 			Simplicials[Index] = TUniquePtr<FSimplicial>(RestSimplicials[Index] ? RestSimplicials[Index]->NewCopy() : nullptr);
@@ -2692,7 +1298,9 @@
 	if(Parameters.RestCollection->HasAttribute(FGeometryCollection::SimulatableParticlesAttribute, FTransformCollection::TransformGroup))
 	{
 		// When the rest collection has been pre configured with simulation data, use that to determine it ability to simulate. 
-		const TManagedArray<bool>& ActiveRestParticles = Parameters.RestCollection->GetAttribute<bool>(FGeometryCollection::SimulatableParticlesAttribute, FTransformCollection::TransformGroup);
+		const TManagedArray<bool>& ActiveRestParticles = 
+			Parameters.RestCollection->GetAttribute<bool>(
+				FGeometryCollection::SimulatableParticlesAttribute, FTransformCollection::TransformGroup);
 		check(ActiveRestParticles.Num() == DynamicCollection->Active.Num());
 		check(ActiveRestParticles.Num() == SimulatableParticles.Num());
 		for(int32 TransformIdx = 0; TransformIdx < NumTransforms; TransformIdx++)
@@ -2717,30 +1325,1259 @@
 	}
 }
 
+void FGeometryCollectionPhysicsProxy::InitializeKinematics(FParticlesType& Particles, const TManagedArray<int32>& DynamicState)
+{
+	if(Parameters.DynamicCollection)
+	{
+		for(int TransformGroupIndex = 0; TransformGroupIndex < RigidBodyID.Num(); TransformGroupIndex++)
+		{
+			if(RigidBodyID[TransformGroupIndex] != INDEX_NONE)
+			{
+				int32 RigidBodyIndex = RigidBodyID[TransformGroupIndex];
+				if(DynamicState[TransformGroupIndex] == (uint8)EObjectStateTypeEnum::Chaos_Object_Kinematic)
+				{
+					Particles.SetObjectState(RigidBodyIndex, Chaos::EObjectStateType::Kinematic);
+				}
+				else if(DynamicState[TransformGroupIndex] == (uint8)EObjectStateTypeEnum::Chaos_Object_Static)
+				{
+					Particles.SetObjectState(RigidBodyIndex, Chaos::EObjectStateType::Static);
+				}
+				else if(DynamicState[TransformGroupIndex] == (uint8)EObjectStateTypeEnum::Chaos_Object_Sleeping)
+				{
+					Particles.SetObjectState(RigidBodyIndex, Chaos::EObjectStateType::Sleeping);
+				}
+			}
+		}
+	}
+}
+
+void FGeometryCollectionPhysicsProxy::InitializeRemoveOnFracture(FParticlesType& Particles, const TManagedArray<int32>& DynamicState)
+{
+	if (Parameters.DynamicCollection && Parameters.RemoveOnFractureEnabled)
+	{
+	//	TManagedArray<FGeometryCollectionBoneNode>& Hierarchy = Parameters.DynamicCollection->BoneHierarchy;
+
+		for (int TransformGroupIndex = 0; TransformGroupIndex < RigidBodyID.Num(); TransformGroupIndex++)
+		{
+			if (RigidBodyID[TransformGroupIndex] != INDEX_NONE)
+			{
+				int32 RigidBodyIndex = RigidBodyID[TransformGroupIndex];
+
+				if (Parameters.DynamicCollection->StatusFlags[TransformGroupIndex] & FGeometryCollection::FS_RemoveOnFracture)
+				{
+					Particles.ToBeRemovedOnFracture(RigidBodyIndex) = true;
+				}
+			}
+		}
+	}
+}
+
+void FGeometryCollectionPhysicsProxy::AddCollisionToCollisionData(FRecordedFrame* ExistingFrame, const FParticlesType& Particles, const Chaos::TPBDCollisionConstraints<float, 3>::FPointContactConstraint& Constraint)
+{
+#if TODO_REIMPLEMENT_RIGID_CLUSTERING
+	const Chaos::TArrayCollectionArray<Chaos::ClusterId>& ClusterIdsArray = GetSolver()->GetRigidClustering().GetClusterIdsArray();
+	const TMap<uint32, TUniquePtr<TArray<uint32>>>& ParentToChildrenMap = GetSolver()->GetRigidClustering().GetChildrenMap();
+	const Chaos::TArrayCollectionArray<PhysicsProxyWrapper>& PhysicsProxyReverseMapping = GetSolver()->GetPhysicsProxyReverseMapping();
+
+	int32 NewIdx = ExistingFrame->Collisions.Add(FSolverCollisionData());
+	FSolverCollisionData& Collision = ExistingFrame->Collisions[NewIdx];
+
+	Collision.Location = Constraint.Location;
+	Collision.AccumulatedImpulse = Constraint.AccumulatedImpulse;
+	Collision.Normal = Constraint.Normal;
+	Collision.Velocity1 = Particles.V(Constraint.ParticleIndex);
+	Collision.Velocity2 = Particles.V(Constraint.LevelsetIndex);
+	Collision.AngularVelocity1 = Particles.W(Constraint.ParticleIndex);
+	Collision.AngularVelocity2 = Particles.W(Constraint.LevelsetIndex);
+	Collision.Mass1 = Particles.M(Constraint.ParticleIndex);
+	Collision.Mass2 = Particles.M(Constraint.LevelsetIndex);
+	Collision.ParticleIndex = Constraint.ParticleIndex;
+	Collision.LevelsetIndex = Constraint.LevelsetIndex;
+
+	// If Constraint.ParticleIndex is a cluster store an index for a mesh in this cluster
+	if (ClusterIdsArray[Constraint.ParticleIndex].NumChildren > 0)
+	{
+		int32 ParticleIndexMesh = GetSolver()->GetParticleIndexMesh(ParentToChildrenMap, Constraint.ParticleIndex);
+		ensure(ParticleIndexMesh != INDEX_NONE);
+		Collision.ParticleIndexMesh = ParticleIndexMesh;
+	}
+	// If Constraint.LevelsetIndex is a cluster store an index for a mesh in this cluster
+	if (ClusterIdsArray[Constraint.LevelsetIndex].NumChildren > 0)
+	{
+		int32 LevelsetIndexMesh = GetSolver()->GetParticleIndexMesh(ParentToChildrenMap, Constraint.LevelsetIndex);
+		ensure(LevelsetIndexMesh != INDEX_NONE);
+		Collision.LevelsetIndexMesh = LevelsetIndexMesh;
+	}
+#endif
+}
+
+void FGeometryCollectionPhysicsProxy::UpdateCollisionData(const FParticlesType& Particles, const FCollisionConstraintsType& CollisionRule, FRecordedFrame* ExistingFrame)
+{
+	ExistingFrame->Collisions.Reset();
+
+	if (Parameters.CollisionData.SaveCollisionData && ExistingFrame->Timestamp > 0.f && Parameters.CollisionData.CollisionDataSizeMax > 0)
+	{
+#if TODO_REIMPLEMENT_PHYSICS_PROXY_REVERSE_MAPPING
+		const TArray<Chaos::TPBDCollisionConstraints<float, 3>::FRigidBodyContactConstraint>& AllConstraintsArray = CollisionRule.GetAllConstraints();
+		if (AllConstraintsArray.Num() > 0)
+		{
+			const Chaos::TArrayCollectionArray<PhysicsProxyWrapper>& PhysicsProxyReverseMapping = GetSolver()->GetPhysicsProxyReverseMapping();
+
+			TArray<Chaos::TPBDCollisionConstraints<float, 3>::FRigidBodyContactConstraint> ConstraintsArray;
+			ConstraintsArray.SetNumUninitialized(AllConstraintsArray.Num());
+
+			FBox BoundingBox(ForceInitToZero);
+			int32 NumConstraints = 0;
+			for (int32 Idx = 0; Idx < AllConstraintsArray.Num(); ++Idx)
+			{
+				// Check if the collision is for this PhysicsProxy
+				IPhysicsProxyBase* PhysicsProxy = PhysicsProxyReverseMapping[AllConstraintsArray[Idx].ParticleIndex].PhysicsProxy;
+				if (PhysicsProxy == this)
+				{
+					if (ensure(!AllConstraintsArray[Idx].AccumulatedImpulse.ContainsNaN()))
+					{
+						if (!AllConstraintsArray[Idx].AccumulatedImpulse.IsZero())
+						{
+							if (ensure(!AllConstraintsArray[Idx].Location.ContainsNaN() &&
+								!AllConstraintsArray[Idx].Normal.ContainsNaN()) &&
+								!Particles.V(AllConstraintsArray[Idx].ParticleIndex).ContainsNaN() &&
+								!Particles.V(AllConstraintsArray[Idx].LevelsetIndex).ContainsNaN() &&
+								!Particles.W(AllConstraintsArray[Idx].ParticleIndex).ContainsNaN() &&
+								!Particles.W(AllConstraintsArray[Idx].LevelsetIndex).ContainsNaN())
+							{
+								BoundingBox += AllConstraintsArray[Idx].Location;
+
+								ConstraintsArray[NumConstraints] = AllConstraintsArray[Idx];
+								NumConstraints++;
+							}
+						}
+					}
+				}
+			}
+
+			ConstraintsArray.SetNum(NumConstraints);
+
+			if (ConstraintsArray.Num() > 0)
+			{
+				if (Parameters.CollisionData.DoCollisionDataSpatialHash &&
+					Parameters.CollisionData.CollisionDataSpatialHashRadius > 0.f &&
+					ConstraintsArray.Num() > 1 &&
+					(BoundingBox.GetExtent().X > 0.f || BoundingBox.GetExtent().Y > 0.f || BoundingBox.GetExtent().Z > 0.f))
+				{
+					// Validate Parameters.CollisionData.CollisionDataSpatialHashRadius
+					// CellSize must be smaller than the smallest bbox extent
+					float SpatialHashRadius = Parameters.CollisionData.CollisionDataSpatialHashRadius;
+					TArray<float> ExtentArray = { BoundingBox.GetExtent().X , BoundingBox.GetExtent().Y, BoundingBox.GetExtent().Z };
+					ExtentArray.Sort();
+					if (ExtentArray[0] != 0.f)
+					{
+						if (2.f * SpatialHashRadius > ExtentArray[0])
+						{
+							SpatialHashRadius = 0.5f * ExtentArray[0];
+						}
+					}
+					else
+					{
+						if (ExtentArray[1] != 0.f)
+						{
+							if (2.f * SpatialHashRadius > ExtentArray[1])
+							{
+								SpatialHashRadius = 0.5f * ExtentArray[1];
+							}
+						}
+						else
+						{
+							if (2.f * SpatialHashRadius > ExtentArray[2])
+							{
+								SpatialHashRadius = 0.5f * ExtentArray[2];
+							}
+						}
+					}
+
+					// Spatial hash the constraints
+					TMultiMap<int32, int32> HashTableMap;
+					ComputeHashTable(ConstraintsArray, BoundingBox, HashTableMap, SpatialHashRadius);
+
+					TArray<int32> UsedCellsArray;
+					HashTableMap.GetKeys(UsedCellsArray);
+
+					int32 NumCollisionsThisFrame = 0;
+					for (int32 IdxCell = 0; IdxCell < UsedCellsArray.Num(); ++IdxCell)
+					{
+						TArray<int32> ConstraintsInCellArray;
+						HashTableMap.MultiFind(UsedCellsArray[IdxCell], ConstraintsInCellArray);
+
+						int32 NumConstraintsToGetFromCell = FMath::Min(Parameters.CollisionData.MaxCollisionPerCell, ConstraintsInCellArray.Num());
+						for (int32 IdxConstraint = 0; IdxConstraint < NumConstraintsToGetFromCell; ++IdxConstraint)
+						{
+							AddCollisionToCollisionData(ExistingFrame, Particles, ConstraintsArray[ConstraintsInCellArray[IdxConstraint]]);
+						}
+					}
+
+					if (ExistingFrame->Collisions.Num() > Parameters.CollisionData.CollisionDataSizeMax)
+					{
+						TArray<FSolverCollisionData> CollisionsArray1;
+
+						float FInc = (float)ExistingFrame->Collisions.Num() / (float)Parameters.CollisionData.CollisionDataSizeMax;
+
+						CollisionsArray1.SetNumUninitialized(Parameters.CollisionData.CollisionDataSizeMax);
+						for (int32 IdxCollision = 0; IdxCollision < Parameters.CollisionData.CollisionDataSizeMax; ++IdxCollision)
+						{
+							int32 NewIdx = FMath::FloorToInt((float)IdxCollision * FInc);
+							CollisionsArray1[IdxCollision] = ExistingFrame->Collisions[NewIdx];
+						}
+
+						ExistingFrame->Collisions.SetNumUninitialized(Parameters.CollisionData.CollisionDataSizeMax);
+						for (int32 IdxCollision = 0; IdxCollision < Parameters.CollisionData.CollisionDataSizeMax; ++IdxCollision)
+						{
+							ExistingFrame->Collisions[IdxCollision] = CollisionsArray1[IdxCollision];
+						}
+					}
+				}
+				else
+				{
+					if (ConstraintsArray.Num() <= Parameters.CollisionData.CollisionDataSizeMax)
+					{
+						for (int32 IdxConstraint = 0; IdxConstraint < ConstraintsArray.Num(); ++IdxConstraint)
+						{
+							AddCollisionToCollisionData(ExistingFrame, Particles, ConstraintsArray[IdxConstraint]);
+						}
+					}
+					else
+					{
+						float FInc = (float)ConstraintsArray.Num() / (float)Parameters.CollisionData.CollisionDataSizeMax;
+
+						for (int32 IdxConstraint = 0; IdxConstraint < Parameters.CollisionData.CollisionDataSizeMax; ++IdxConstraint)
+						{
+							int32 Idx = FMath::FloorToInt((float)IdxConstraint * FInc);
+
+							AddCollisionToCollisionData(ExistingFrame, Particles, ConstraintsArray[Idx]);
+						}
+					}
+				}
+			}
+		}
+#endif
+	}
+}
+
+void FGeometryCollectionPhysicsProxy::AddBreakingToBreakingData(FRecordedFrame* ExistingFrame, const FParticlesType& Particles, const Chaos::TBreakingData<float, 3>& Breaking)
+{
+#if TODO_REIMPLEMENT_RIGID_CLUSTERING
+	const Chaos::TArrayCollectionArray<Chaos::ClusterId>& ClusterIdsArray = GetSolver()->GetRigidClustering().GetClusterIdsArray();
+	const TMap<uint32, TUniquePtr<TArray<uint32>>>& ParentToChildrenMap = GetSolver()->GetRigidClustering().GetChildrenMap();
+	const Chaos::TArrayCollectionArray<PhysicsProxyWrapper>& PhysicsProxyReverseMapping = GetSolver()->GetPhysicsProxyReverseMapping();
+
+	int32 NewIdx = ExistingFrame->Breakings.Add(FSolverBreakingData());
+	FSolverBreakingData& NewBreaking = ExistingFrame->Breakings[NewIdx];
+
+	NewBreaking.Location = Breaking.Location;
+	NewBreaking.Velocity = Particles.V(Breaking.ParticleIndex);
+	NewBreaking.AngularVelocity = Particles.W(Breaking.ParticleIndex);
+	NewBreaking.Mass = Particles.M(Breaking.ParticleIndex);
+	NewBreaking.ParticleIndex = Breaking.ParticleIndex;
+
+	// If Breaking.ParticleIndex is a cluster store an index for a mesh in this cluster
+	if (ClusterIdsArray[Breaking.ParticleIndex].NumChildren > 0)
+	{
+		int32 ParticleIndexMesh = GetSolver()->GetParticleIndexMesh(ParentToChildrenMap, Breaking.ParticleIndex);
+		ensure(ParticleIndexMesh != INDEX_NONE);
+		NewBreaking.ParticleIndexMesh = ParticleIndexMesh;
+	}
+#endif
+}
+
+void FGeometryCollectionPhysicsProxy::UpdateBreakingData(const FParticlesType& Particles, FRecordedFrame* ExistingFrame)
+{
+	ExistingFrame->Breakings.Reset();
+
+	if (Parameters.BreakingData.SaveBreakingData && ExistingFrame->Timestamp > 0.f && Parameters.BreakingData.BreakingDataSizeMax > 0)
+	{
+#if TODO_REIMPLEMENT_GETALLCLUSTERBREAKINGS
+		const TArray<Chaos::TBreakingData<float, 3>>& AllBreakingsArray = GetSolver()->GetAllClusterBreakings();
+		if (AllBreakingsArray.Num() > 0)
+		{
+			const Chaos::TArrayCollectionArray<PhysicsProxyWrapper>& PhysicsProxyReverseMapping = GetSolver()->GetPhysicsProxyReverseMapping();
+
+			TArray<Chaos::TBreakingData<float, 3>> BreakingsArray;
+			BreakingsArray.SetNumUninitialized(AllBreakingsArray.Num());
+
+			FBox BoundingBox(ForceInitToZero);
+			int32 NumBreakings = 0;
+			for (int32 Idx = 0; Idx < AllBreakingsArray.Num(); ++Idx)
+			{
+				// Check if the breaking is for this PhysicsProxy
+				IPhysicsProxyBase* PhysicsProxy = PhysicsProxyReverseMapping[AllBreakingsArray[Idx].ParticleIndex].PhysicsProxy;
+				if (PhysicsProxy == this)
+				{
+					if (ensure(!AllBreakingsArray[Idx].Location.ContainsNaN() &&
+						!Particles.V(AllBreakingsArray[Idx].ParticleIndex).ContainsNaN() &&
+						!Particles.W(AllBreakingsArray[Idx].ParticleIndex).ContainsNaN()))
+					{
+						BoundingBox += AllBreakingsArray[Idx].Location;
+
+						BreakingsArray[NumBreakings] = AllBreakingsArray[Idx];
+						NumBreakings++;
+					}
+				}
+			}
+
+			BreakingsArray.SetNum(NumBreakings);
+
+			if (BreakingsArray.Num() > 0)
+			{
+				if (Parameters.BreakingData.DoBreakingDataSpatialHash &&
+					Parameters.BreakingData.BreakingDataSpatialHashRadius > 0.f &&
+					BreakingsArray.Num() > 1 &&
+					(BoundingBox.GetExtent().X > 0.f || BoundingBox.GetExtent().Y > 0.f || BoundingBox.GetExtent().Z > 0.f))
+				{
+					// Validate Parameters.CollisionData.CollisionDataSpatialHashRadius
+					// CellSize must be smaller than the smallest bbox extent
+					float SpatialHashRadius = Parameters.BreakingData.BreakingDataSpatialHashRadius;
+					TArray<float> ExtentArray = { BoundingBox.GetExtent().X , BoundingBox.GetExtent().Y, BoundingBox.GetExtent().Z };
+					ExtentArray.Sort();
+					if (ExtentArray[0] != 0.f)
+					{
+						if (2.f * SpatialHashRadius > ExtentArray[0])
+						{
+							SpatialHashRadius = 0.5f * ExtentArray[0];
+						}
+					}
+					else
+					{
+						if (ExtentArray[1] != 0.f)
+						{
+							if (2.f * SpatialHashRadius > ExtentArray[1])
+							{
+								SpatialHashRadius = 0.5f * ExtentArray[1];
+							}
+						}
+						else
+						{
+							if (2.f * SpatialHashRadius > ExtentArray[2])
+							{
+								SpatialHashRadius = 0.5f * ExtentArray[2];
+							}
+						}
+					}
+
+					// Spatial hash the constraints
+					TMultiMap<int32, int32> HashTableMap;
+					ComputeHashTable(BreakingsArray, BoundingBox, HashTableMap, SpatialHashRadius);
+
+					TArray<int32> UsedCellsArray;
+					HashTableMap.GetKeys(UsedCellsArray);
+
+					for (int32 IdxCell = 0; IdxCell < UsedCellsArray.Num(); ++IdxCell)
+					{
+						TArray<int32> BreakingsInCellArray;
+						HashTableMap.MultiFind(UsedCellsArray[IdxCell], BreakingsInCellArray);
+
+						int32 NumBreakingsToGetFromCell = FMath::Min(Parameters.BreakingData.MaxBreakingPerCell, BreakingsInCellArray.Num());
+						for (int32 IdxBreaking = 0; IdxBreaking < NumBreakingsToGetFromCell; ++IdxBreaking)
+						{
+							AddBreakingToBreakingData(ExistingFrame, Particles, BreakingsArray[BreakingsInCellArray[IdxBreaking]]);
+						}
+					}
+
+					if (ExistingFrame->Breakings.Num() > Parameters.BreakingData.BreakingDataSizeMax)
+					{
+						TArray<FSolverBreakingData> BreakingsArray1;
+
+						float FInc = (float)ExistingFrame->Breakings.Num() / (float)Parameters.BreakingData.BreakingDataSizeMax;
+
+						BreakingsArray1.SetNumUninitialized(Parameters.BreakingData.BreakingDataSizeMax);
+						for (int32 IdxBreaking = 0; IdxBreaking < Parameters.BreakingData.BreakingDataSizeMax; ++IdxBreaking)
+						{
+							int32 NewIdx = FMath::FloorToInt((float)IdxBreaking * FInc);
+							BreakingsArray1[IdxBreaking] = ExistingFrame->Breakings[NewIdx];
+						}
+
+						ExistingFrame->Breakings.SetNumUninitialized(Parameters.BreakingData.BreakingDataSizeMax);
+						for (int32 IdxBreaking = 0; IdxBreaking < Parameters.BreakingData.BreakingDataSizeMax; ++IdxBreaking)
+						{
+							ExistingFrame->Breakings[IdxBreaking] = BreakingsArray1[IdxBreaking];
+						}
+					}
+				}
+				else
+				{
+					if (BreakingsArray.Num() <= Parameters.BreakingData.BreakingDataSizeMax)
+					{
+						for (int32 IdxBreaking = 0; IdxBreaking < BreakingsArray.Num(); ++IdxBreaking)
+						{
+							AddBreakingToBreakingData(ExistingFrame, Particles, BreakingsArray[IdxBreaking]);
+						}
+					}
+					else
+					{
+						float FInc = (float)BreakingsArray.Num() / (float)Parameters.BreakingData.BreakingDataSizeMax;
+
+						for (int32 IdxBreaking = 0; IdxBreaking < Parameters.BreakingData.BreakingDataSizeMax; ++IdxBreaking)
+						{
+							int32 Idx = FMath::FloorToInt((float)IdxBreaking * FInc);
+
+							AddBreakingToBreakingData(ExistingFrame, Particles, BreakingsArray[Idx]);
+						}
+					}
+				}
+			}
+		}
+#endif
+	}
+}
+
+void FGeometryCollectionPhysicsProxy::UpdateTrailingData(const FParticlesType& Particles, FRecordedFrame* ExistingFrame)
+{
+	ExistingFrame->Trailings.Reset();
+
+	if (Parameters.TrailingData.SaveTrailingData && ExistingFrame->Timestamp > 0.f && Parameters.TrailingData.TrailingDataSizeMax > 0)
+	{
+		const float TrailingMinSpeedThresholdSquared = Parameters.TrailingData.TrailingMinSpeedThreshold * Parameters.TrailingData.TrailingMinSpeedThreshold;
+
+		// Find previous frame trailing data
+		int32 LastKey = RecordedTracks.FindLastKeyBefore(ExistingFrame->Timestamp);
+		if (LastKey != INDEX_NONE)
+		{
+			FRecordedFrame& PrevFrame = RecordedTracks.Records[LastKey];
+			if (PrevFrame.Trailings.Num() > 0)
+			{
+				for (FSolverTrailingData& Trailing : PrevFrame.Trailings)
+				{
+					if (Particles.ObjectState(Trailing.ParticleIndex) != Chaos::EObjectStateType::Sleeping &&
+						!Particles.Disabled(Trailing.ParticleIndex) &&
+						Particles.V(Trailing.ParticleIndex).SizeSquared() >= TrailingMinSpeedThresholdSquared)
+					{
+						ExistingFrame->Trailings.Add(Trailing);
+					}
+				}
+			}
+		}
+
+		if (ExistingFrame->Trailings.Num() < Parameters.TrailingData.TrailingDataSizeMax)
+		{
+#if TODO_REIMPLEMENT_PHYSICS_PROXY_REVERSE_MAPPING
+			const Chaos::TArrayCollectionArray<PhysicsProxyWrapper>& PhysicsProxyReverseMapping = GetSolver()->GetPhysicsProxyReverseMapping();
+			const Chaos::TArrayCollectionArray<Chaos::ClusterId>& ClusterIdsArray = GetSolver()->GetRigidClustering().GetClusterIdsArray();
+			const TMap<uint32, TUniquePtr<TArray<uint32>>>& ParentToChildrenMap = GetSolver()->GetRigidClustering().GetChildrenMap();
+
+			TArray<FSolverTrailingData> AllTrailings;
+			for (int32 IdxParticle = 0; IdxParticle < (int32)Particles.Size(); ++IdxParticle)
+			{
+				// Check if the particle is for this PhysicsProxy
+				IPhysicsProxyBase* PhysicsProxy = PhysicsProxyReverseMapping[IdxParticle].PhysicsProxy;
+				if (PhysicsProxy == this)
+				{
+					if (ensure(FMath::IsFinite(Particles.InvM(IdxParticle))))
+					{
+						if (Particles.ObjectState(IdxParticle) != Chaos::EObjectStateType::Sleeping &&
+							!Particles.Disabled(IdxParticle) &&
+							Particles.InvM(IdxParticle) != 0.f)
+						{
+							if (Particles.Geometry(IdxParticle) && 
+								Particles.Geometry(IdxParticle)->HasBoundingBox())
+							{
+								if (ensure(!Particles.X(IdxParticle).ContainsNaN()) &&
+									!Particles.V(IdxParticle).ContainsNaN() &&
+									!Particles.W(IdxParticle).ContainsNaN() &&
+									FMath::IsFinite(Particles.M(IdxParticle)))
+								{
+									Chaos::TAABB<float, 3> BoundingBox = Particles.Geometry(IdxParticle)->BoundingBox();
+									Chaos::TVector<float, 3> Extents = BoundingBox.Extents();
+									float Volume = Extents[0] * Extents[1] * Extents[2];
+									float SpeedSquared = Particles.V(IdxParticle).SizeSquared();
+
+									if (SpeedSquared >= TrailingMinSpeedThresholdSquared &&
+										Volume > Parameters.TrailingData.TrailingMinVolumeThreshold)
+									{
+										FSolverTrailingData TrailingData(Particles.X(IdxParticle),
+											Particles.V(IdxParticle),
+											Particles.W(IdxParticle),
+											Particles.M(IdxParticle),
+											IdxParticle,
+											INDEX_NONE);
+
+										// If IdxParticle is a cluster store an index for a mesh in this cluster
+										if (ClusterIdsArray[IdxParticle].NumChildren > 0)
+										{
+											int32 ParticleIndexMesh = GetSolver()->GetParticleIndexMesh(ParentToChildrenMap, IdxParticle);
+											ensure(ParticleIndexMesh != INDEX_NONE);
+											TrailingData.ParticleIndexMesh = ParticleIndexMesh;
+										}
+
+										FSetElementId Id = ExistingFrame->Trailings.FindId(TrailingData);
+										if (!Id.IsValidId())
+										{
+											AllTrailings.Add(TrailingData);
+										}
+										else
+										{
+											ExistingFrame->Trailings[Id].Location = Particles.X(IdxParticle);
+											ExistingFrame->Trailings[Id].Velocity = Particles.V(IdxParticle);
+											ExistingFrame->Trailings[Id].AngularVelocity = Particles.W(IdxParticle);
+										}
+									}
+								}
+							}
+						}
+					}
+				}
+			}
+
+			if (AllTrailings.Num() > 0)
+			{
+				int32 NumTrailingsToAdd = Parameters.TrailingData.TrailingDataSizeMax - ExistingFrame->Trailings.Num();
+
+				if (AllTrailings.Num() <= NumTrailingsToAdd)
+				{
+					for (int32 IdxTrailing = 0; IdxTrailing < AllTrailings.Num(); ++IdxTrailing)
+					{
+						ExistingFrame->Trailings.Add(AllTrailings[IdxTrailing]);
+					}
+				}
+				else
+				{
+					float FInc = (float)AllTrailings.Num() / (float)(NumTrailingsToAdd);
+
+					for (int32 IdxTrailing = 0; IdxTrailing < NumTrailingsToAdd; ++IdxTrailing)
+					{
+						int32 Idx = FMath::FloorToInt((float)IdxTrailing * FInc);
+						ExistingFrame->Trailings.Add(AllTrailings[Idx]);
+					}
+				}
+			}
+#endif
+		}
+	}
+}
+
+void FGeometryCollectionPhysicsProxy::UpdateRecordedState(float SolverTime, const TManagedArray<int32>& InRigidBodyID, const TManagedArray<int32>& InCollectionClusterID, const Chaos::TArrayCollectionArray<bool>& InInternalCluster, const FParticlesType& Particles, const FCollisionConstraintsType& CollisionRule)
+{
+	FRecordedFrame* ExistingFrame = RecordedTracks.FindRecordedFrame(SolverTime);
+
+	if(!ExistingFrame)
+	{
+		ExistingFrame = &InsertRecordedFrame(RecordedTracks, SolverTime);
+	}
+
+	ExistingFrame->Reset(InRigidBodyID.Num());
+	ExistingFrame->Timestamp = SolverTime;
+
+	// Collision
+	UpdateCollisionData(Particles, CollisionRule, ExistingFrame);
+
+	// Breaking
+	UpdateBreakingData(Particles, ExistingFrame);
+
+	// Trailing
+	UpdateTrailingData(Particles, ExistingFrame);
+
+	ParallelFor(InRigidBodyID.Num(), [&](int32 Index)
+	{
+		const int32 ExternalIndex = InRigidBodyID[Index];
+
+		if(ExternalIndex >= 0)
+		{
+			FTransform& NewTransform = ExistingFrame->Transforms[Index];
+
+			NewTransform.SetTranslation(Particles.X(ExternalIndex));
+			NewTransform.SetRotation(Particles.R(ExternalIndex));
+			NewTransform.SetScale3D(FVector(1.0f));
+
+			int32 ClusterParentIndex = InCollectionClusterID[Index];
+			if (ClusterParentIndex != INDEX_NONE && InInternalCluster[ClusterParentIndex])
+			{
+				ExistingFrame->DisabledFlags[Index] = Particles.Disabled(ClusterParentIndex);
+			}
+			else
+			{
+				ExistingFrame->DisabledFlags[Index] = Particles.Disabled(ExternalIndex);
+			}
+		}
+	});
+}
+
+void FGeometryCollectionPhysicsProxy::OnRemoveFromSolver(Chaos::FPBDRigidsSolver *RBDSolver)
+{
+	for (auto* Handle : SolverClusterHandles)
+	{
+		RBDSolver->RemoveParticleToProxy(Handle);
+	}
+}
+
+void FGeometryCollectionPhysicsProxy::OnRemoveFromScene()
+{
+#if TODO_REIMPLEMENT_GET_RIGID_PARTICLES
+	// #BG TODO This isn't great - we currently cannot handle things being removed from the solver.
+	// need to refactor how we handle this and actually remove the particles instead of just constantly
+	// growing the array. Currently everything is just tracked by index though so the solver will have
+	// to notify all the proxies that a chunk of data was removed - or use a sparse array (undesireable)
+	Chaos::FPhysicsSolver::FParticlesType& Particles = GetSolver()->GetRigidParticles();
+
+	// #BG TODO Special case here because right now we reset/realloc the evolution per geom component
+	// in endplay which clears this out. That needs to not happen and be based on world shutdown
+	if(Particles.Size() == 0)
+	{
+		return;
+	}
+
+	const int32 Begin = BaseParticleIndex;
+	const int32 Count = NumParticles;
+
+	if (ensure((int32)Particles.Size() > 0 && (Begin + Count) <= (int32)Particles.Size()))
+	{
+		for (int32 ParticleIndex = 0; ParticleIndex < Count; ++ParticleIndex)
+		{
+			GetSolver()->GetEvolution()->DisableParticle(Begin + ParticleIndex);
+			GetSolver()->GetRigidClustering().GetTopLevelClusterParents().Remove(Begin + ParticleIndex);
+		}
+	}
+#endif
+}
+
+void FGeometryCollectionPhysicsProxy::SyncBeforeDestroy()
+{
+	if(FinalSyncFunc)
+	{
+		FinalSyncFunc(RecordedTracks);
+	}
+}
+
+void FGeometryCollectionPhysicsProxy::BufferGameState() 
+{
+	// NewData() redirects to this function.
+	/**
+	 * CONTEXT: GAMETHREAD
+	 * Returns a new unmanaged allocation of the data saved on the handle, otherwise nullptr
+	 */
+
+	// The other physics proxies allocate memory and returns a copy of game thread
+	// data as a TGeometryParticleData<T,d>, TKinematicGeometryParticleData<T,d>, 
+	// or a TPBDRigidParticleData<T,d>, which are all derived from FParticleData.
+	// That then gets passed to the physics thread via PushToPhysicsState().
+	//
+	// So all we're doing is pushing game thread state to the physics thread.
+	// I don't think we need to allocate memory to do this.  I think the other 
+	// proxies get away with that because they're dealing with a low number of 
+	// bodies; which is not the case here. Rather, let's just use a triple buffer.
+
+	if (!Parameters.Simulating || 
+		!ensure(NumParticles != INDEX_NONE)) // Make sure InitBodiesGT() has been called!
+	{
+		return;
+	}
+
+	FGeometryCollectionResults& Buffer = *GameToPhysInterchange->AccessProducerBuffer();
+	if (Buffer.NumParticlesAdded != NumParticles)
+	{
+		Buffer.InitArrays(*GTDynamicCollection);
+
+		// Initialize the buffer
+		Buffer.NumParticlesAdded = NumParticles;
+		Buffer.BaseIndex = BaseParticleIndex;
+
+		Buffer.SharedGeometry.SetNum(NumParticles);
+		Buffer.ShapeSimData.SetNum(NumParticles);
+		Buffer.ShapeQueryData.SetNum(NumParticles);
+
+		// To be done once:
+		int32 Idx = 0;
+		for (auto& GTParticle : GTGeometryParticles)
+		{
+			if (GTParticle)
+			{
+				Buffer.DisabledStates[Idx] = false;
+				Buffer.SharedGeometry[Idx] = GTParticle->GeometrySharedLowLevel();
+
+				for (auto& Shape : GTParticle->ShapesArray())
+				{
+					Buffer.ShapeQueryData[Idx].Add(Shape->QueryData);
+					Buffer.ShapeSimData[Idx].Add(Shape->SimData);
+
+					// Ryan - Hard code everything collides with everything:
+					const int32 Idx1 = Buffer.ShapeSimData[Idx].Num() - 1;
+					Buffer.ShapeSimData[Idx][Idx1].Word1 = 0xFFFF; // this body channel
+					Buffer.ShapeSimData[Idx][Idx1].Word3 = 0xFFFF; // collision candidate channels
+				}
+			}
+			else
+			{
+				Buffer.DisabledStates[Idx] = true;
+				// Buffer.SpatialIdx[Idx] = ???;
+				//Buffer.UniqueIdx[Idx].Idx = INDEX_NONE;
+				//Buffer.UserData[Idx] = nullptr;
+				//Buffer.SharedGeometry[Idx] = nullptr;
+			}
+			++Idx;
+		}
+	}
+
+	const FGeometryDynamicCollection* DynamicCollection = Parameters.DynamicCollection;
+	const FGeometryCollection* RestCollection = Parameters.RestCollection;
+
+	//Buffer.TransformIndex.Init(RestCollection->TransformIndex);
+	Buffer.BoneMap.Init(RestCollection->BoneMap);
+	Buffer.Parent.Init(RestCollection->Parent);
+	Buffer.Children.Init(RestCollection->Children);
+	Buffer.SimulationType.Init(RestCollection->SimulationType);
+
+	Buffer.DynamicState.Init(DynamicCollection->DynamicState);
+	Buffer.Mass.Init(RestCollection->GetAttribute<float>("Mass", FTransformCollection::TransformGroup));
+	Buffer.InertiaTensor.Init(RestCollection->GetAttribute<FVector>("InertiaTensor", FTransformCollection::TransformGroup));
+	Buffer.Transforms.Init(DynamicCollection->Transform);
+
+	GeometryCollectionAlgo::GlobalMatrices(
+		DynamicCollection->Transform, DynamicCollection->Parent, Buffer.GlobalTransforms);
+	
+	int32 Idx = 0;
+	for (auto& GTParticle : GTGeometryParticles)
+	{
+		if (GTParticle)
+		{
+			if (auto* GTRigidParticle = GTParticle->CastToRigidParticle())
+			{
+				Buffer.DisabledStates[Idx] = GTRigidParticle->Disabled();
+			}
+
+			//const Chaos::TShapesArray<float,3>& ShapesArray = GTParticle->ShapesArray();
+		}
+		++Idx;
+	}
+	
+	// Make the buffer available to the consumer.
+	GameToPhysInterchange->FlipProducer(); 
+	
+	// If we wanted to emulate the other proxies, we could augment the 
+	// FStubGeometryCollectionData class to carry along the consumer buffer:
+	//
+	// return new FStubGeometryCollectionData(GametoPhysInterchange.GetConsumerBuffer());
+	//
+	// However, at this point, it seems pretty unnecessary to do the additional
+	// allocation and deallocation of the transient FStubGeometryCollectionData
+	// class, as I believe all these operations are happening with the same 
+	// instance of this class.  So skip all that, and just return null.
+	return;
+}
+
+void FGeometryCollectionPhysicsProxy::PushToPhysicsState(const Chaos::FParticleData *InData)
+{
+	/**
+	 * CONTEXT: GAMETHREAD -> to -> PHYSICSTHREAD
+	 * Called on the game thread when the solver is about to advance forward. This
+	 * callback should Enqueue commands on the PhysicsThread to update the state of
+	 * the solver
+	 */
+
+	if (InitializedState == ESimulationInitializationState::Unintialized)
+	{
+		return;
+	}
+
+	// The other proxies do a deep copy into dynamic transient memory, passed into
+	// this function via InData.  Instead of that, we're using persistent memory
+	// in a buffer to facilitate communicating between threads.
+
+	const FGeometryCollectionResults* GState = GameToPhysInterchange->GetConsumerBuffer();
+	if (!GState)
+	{
+		// This will happen if GameToPhysInterchange hasn't been updated since
+		// the last time GetConsumerBuffer() was called.  At this point, that's 
+		// probably a bug, as the interface code is assuming it needs to lock
+		// to prevent stomping on a double buffer.  We should look into doing 
+		// away with that lock, but for the time being, we'll assert:
+		ensure(false);
+		return;
+	}
+
+	//const FStubGeometryCollectionData* GCData = static_cast<const FStubGeometryCollectionData*>(InData);
+	//const FGeometryCollectionResults* GCState = GCData->GetStateData();
+
+	// Move the copied game thread data into the handles.
+	for (int32 HandleIdx = 0; HandleIdx < SolverParticleHandles.Num(); ++HandleIdx)
+	{
+		if (Chaos::TPBDRigidParticleHandle<float, 3>* Handle = SolverParticleHandles[HandleIdx])
+		{
+			const int32 Idx = GState->BaseIndex + HandleIdx;
+
+			FTransform ParticleTransform =
+				MassToLocal[Idx] * GState->Transforms[Idx] * Parameters.WorldTransform;
+
+			Handle->SetX(ParticleTransform.GetTranslation());
+			Handle->SetR(ParticleTransform.GetRotation());
+			Handle->SetDisabled(GState->DisabledStates[Idx]);
+
+			Handle->SetLinearEtherDrag(Parameters.LinearEtherDrag);
+			Handle->SetAngularEtherDrag(Parameters.AngularEtherDrag);
+
+			if (GState->SharedGeometry[Idx] && GState->SharedGeometry[Idx]->HasBoundingBox())
+			{
+				Handle->SetHasBounds(true);
+				Handle->SetLocalBounds(GState->SharedGeometry[Idx]->BoundingBox());
+
+				Chaos::TAABB<float, 3> WorldSpaceBox =
+					GState->SharedGeometry[Idx]->BoundingBox().TransformedAABB(
+						Chaos::TRigidTransform<float, 3>(Handle->X(), Handle->R()));
+				// TODO Ryan - WorldSpaceBox.ThickenSymmetrically(GState->V());
+				Handle->SetWorldSpaceInflatedBounds(WorldSpaceBox);
+			}
+
+			// TODO Ryan - check dirty flags like TGeometryParticleData::DirtyFlags.
+			// In the mean time, we only set once.
+			if (!Handle->SharedGeometry() && GState->SharedGeometry.IsValidIndex(Idx) && GState->SharedGeometry[Idx])
+			{
+				Handle->SetSharedGeometry(GState->SharedGeometry[Idx]);
+			}
+
+			int32 CurrShape = 0;
+			for (const TUniquePtr<Chaos::TPerShapeData<Chaos::FReal, 3>>& Shape : Handle->ShapesArray())
+			{
+				// TODO Ryan - hardwired for now:
+				Shape->bDisable = false;
+				Shape->CollisionTraceType = Chaos::EChaosCollisionTraceFlag::Chaos_CTF_UseSimpleAndComplex; // Chaos_CTF_UseDefault?
+				if (CurrShape < GState->ShapeSimData[Idx].Num())
+				{
+					Shape->SimData = GState->ShapeSimData[Idx][CurrShape];
+					Shape->QueryData = GState->ShapeQueryData[Idx][CurrShape];
+				}
+				++CurrShape;
+			}
+		} // if rigid handle
+	} // for all SolverParticleHandles
+
+	if (DisableGeometryCollectionGravity) // cvar
+	{
+		// Our assumption is that you'd only ever want to wholesale opt geometry 
+		// collections out of gravity for debugging, so we keep this conditional
+		// out of the loop above and on it's own.  This means we can't turn gravity
+		// back on once it's off, but even if we didn't enclose this in an if(),
+		// this function won't be called again unless something dirties the proxy.
+
+		Chaos::TPerParticleGravity<float, 3>& GravityForces = GetSolver()->GetEvolution()->GetGravityForces();
+		for (int32 HandleIdx = 0; HandleIdx < SolverParticleHandles.Num(); ++HandleIdx)
+		{
+			if (Chaos::TPBDRigidParticleHandle<float, 3>* Handle = SolverParticleHandles[HandleIdx])
+			{
+				GravityForces.SetEnabled(*Handle, false);
+			}
+		}
+	}
+}
+void FGeometryCollectionPhysicsProxy::BufferPhysicsResults()
+{
+	/**
+	 * CONTEXT: PHYSICSTHREAD
+	 * Called per-tick after the simulation has completed. The proxy should cache the results of their
+	 * simulation into the local buffer. 
+	 */
+	using namespace Chaos;
+	SCOPE_CYCLE_COUNTER(STAT_CacheResultGeomCollection);
+
+	//FGeometryCollectionResults& TargetResults = Results.GetPhysicsDataForWrite();
+	FGeometryCollectionResults& TargetResults = *PhysToGameInterchange->AccessProducerBuffer();
+
+	if (TargetResults.NumTransformGroup() != PTDynamicCollection.NumElements(FGeometryCollection::TransformGroup))
+	{
+		TargetResults.InitArrays(PTDynamicCollection);
+
+		// Base particle index to calculate index from a global particle index on the game thread
+		TargetResults.BaseIndex = BaseParticleIndex;
+		TargetResults.NumParticlesAdded = NumParticles;
+	}
+
+	{ // tmp scope
+		IsObjectDynamic = false;
+		int32 Idx = 0;
+		for (auto* Handle : SolverParticleHandles)
+		{
+			if (!Handle || Handle->Disabled())
+			{
+				TargetResults.DisabledStates[Idx++] = true;
+				continue;
+			}
+			TargetResults.DisabledStates[Idx++] = false;
+			IsObjectDynamic = true;
+		}
+	} // tmp scope
+
+	// Advertise to game thread
+	TargetResults.IsObjectDynamic = IsObjectDynamic;
+	TargetResults.IsObjectLoading = IsObjectLoading;
+	
+	const FTransform& ActorToWorld = Parameters.WorldTransform;
+	// It's not worth shrinking the EndFrameUnparentingBuffer array, at least until the solver supports deleting bodies.
+	if (EndFrameUnparentingBuffer.Num() < NumParticles)
+	{
+		EndFrameUnparentingBuffer.Init(INDEX_NONE, NumParticles);
+	}
+
+	{ // tmp scope
+		SCOPE_CYCLE_COUNTER(STAT_CalcParticleToWorld);
+		// Fill particle to world transforms
+//		TargetResults.ParticleToWorldTransforms.SetNum(NumParticles);
+/*		for (int32 TransformIndex = 0; TransformIndex < NumParticles; ++TransformIndex)
+		{
+			//only update roots and first children
+			if (TargetResults.Parent[TransformIndex] == INDEX_NONE)
+			{
+				if (auto* Handle = SolverParticleHandles[TransformIndex])
+				{
+					const FTransform Xf(Handle->R(), Handle->X());
+					//const int32 ParticleIndex = BaseParticleIndex + TransformIndex;
+					TargetResults.ParticleToWorldTransforms[TransformIndex] = Xf;
+				}
+			}
+		}
+*/
+		//const TArrayCollectionArray<Chaos::ClusterId>& ClusterID = GetSolver()->GetRigidClustering().GetClusterIdsArray();
+		//const TArrayCollectionArray<Chaos::TRigidTransform<float, 3>>& ClusterChildToParentMap = GetSolver()->GetRigidClustering().GetChildToParentMap();
+		//const TArrayCollectionArray<FMultiChildProxyId>& MultiChildProxyIdArray = GetSolver()->GetRigidClustering().GetMultiChildProxyIdArray();
+		//const TArrayCollectionArray<TUniquePtr<TMultiChildProxyData<float, 3>>>& MultiChildProxyDataArray = GetSolver()->GetRigidClustering().GetMultiChildProxyDataArray();
+
+		for (int32 TransformIndex = 0; TransformIndex < NumParticles; ++TransformIndex)
+		{
+			Chaos::TPBDRigidClusteredParticleHandle<float, 3>* Handle = SolverParticleHandles[TransformIndex];
+			if (!Handle)
+			{
+				continue;
+			}
+
+			// Dynamic state is also updated by the solver during field interaction.
+			if (!Handle->Sleeping())
+			{
+				const Chaos::EObjectStateType ObjectState = Handle->ObjectState();
+				switch (ObjectState)
+				{
+				case Chaos::EObjectStateType::Kinematic:
+					TargetResults.DynamicState[TransformIndex] = (int)EObjectStateTypeEnum::Chaos_Object_Kinematic;
+					break;
+				case Chaos::EObjectStateType::Static:
+					TargetResults.DynamicState[TransformIndex] = (int)EObjectStateTypeEnum::Chaos_Object_Static;
+					break;
+				case Chaos::EObjectStateType::Sleeping:
+					TargetResults.DynamicState[TransformIndex] = (int)EObjectStateTypeEnum::Chaos_Object_Sleeping;
+					break;
+				case Chaos::EObjectStateType::Dynamic:
+				case Chaos::EObjectStateType::Uninitialized:
+				default:
+					TargetResults.DynamicState[TransformIndex] = (int)EObjectStateTypeEnum::Chaos_Object_Dynamic;
+					break;
+				}
+			}
+
+			// Update the transform and parent hierarchy of the active rigid bodies. Active bodies can be either
+			// rigid geometry defined from the leaf nodes of the collection, or cluster bodies that drive an entire
+			// branch of the hierarchy within the GeometryCollection.
+			// - Active bodies are directly driven from the global position of the corresponding
+			//   rigid bodies within the solver ( cases where RigidBodyID[TransformGroupIndex] is not disabled ). 
+			// - Deactivated bodies are driven from the transforms of their active parents. However the solver can
+			//   take ownership of the parents during the simulation, so it might be necessary to force deactivated
+			//   bodies out of the collections hierarchy during the simulation.  
+			if (!Handle->Disabled())
+			{
+				// Update the transform of the active body. The active body can be either a single rigid
+				// or a collection of rigidly attached geometries (Clustering). The cluster is represented as a
+				// single transform in the GeometryCollection, and all children are stored in the local space
+				// of the parent cluster.
+				// ... When setting cluster transforms it is expected that the MassToLocal is identity.
+				//     Cluster initialization will set the vertices in the MassSpace of the rigid body.
+				// ... When setting individual rigid bodies that are not clustered, the MassToLocal will be 
+				//     non-Identity, and will reflect the difference between the geometric center of the geometry
+				//     and that corresponding rigid bodies center of mass. 
+
+				FTransform& ParticleToWorld = TargetResults.ParticleToWorldTransforms[TransformIndex];
+				ParticleToWorld = FTransform(Handle->R(), Handle->X());
+
+				TargetResults.Transforms[TransformIndex] =
+					MassToLocal[TransformIndex].GetRelativeTransformReverse(ParticleToWorld).GetRelativeTransform(ActorToWorld);
+				TargetResults.Transforms[TransformIndex].NormalizeRotation();
+
+
+#if 0 // TODO Ryan
+				// Force all enabled rigid bodies out of the transform hierarchy
+				if (TargetResults.Parent[TransformIndex] != INDEX_NONE)
+				{
+					// Children in the hierarchy are stored in a TSet, which is not thread safe.  
+					// So we retain indices to remove afterwards.
+					EndFrameUnparentingBuffer[TransformIndex] = TargetResults.Parent[TransformIndex];
+				}
+#endif
+#if 0 // TODO Ryan
+				// When a leaf node rigid body is removed from a cluster, the rigid
+				// body will become active and needs its clusterID updated.  This just
+				// syncs the clusterID all the time.
+				CollectionClusterID[TransformGroupIndex] = ClusterID[RigidBodyIndex].Id;
+#endif
+			}
+			else // Handle->Disabled()
+			{
+#if 0 // TODO Ryan
+				// The rigid body parent cluster has changed within the solver, and its
+				// parent body is not tracked within the geometry collection. So we need to
+				// pull the rigid bodies out of the transform hierarchy, and just drive
+				// the positions directly from the solvers cluster particle. 
+				if (CollectionClusterID[TransformGroupIndex] != ClusterID[RigidBodyIndex].Id)
+				{
+					// Force all driven rigid bodies out of the transform hierarchy
+					if (Parent[TransformGroupIndex] != INDEX_NONE)
+					{
+						const int32 ParentIndex = Parent[TransformGroupIndex];
+						// Children in the hierarchy are stored in a TSet, which is not thread safe.  So we retain
+						// indices to remove afterwards.
+						EndFrameUnparentingBuffer[TransformGroupIndex] = ParentIndex;
+					}
+					CollectionClusterID[TransformGroupIndex] = ClusterID[RigidBodyIndex].Id;
+				}
+#endif
+				// Disabled rigid bodies that have valid cluster parents, and have been re-indexed by the
+				// solver (as in, They were re-clustered outside of the geometry collection), These clusters 
+				// will need to be rendered based on the clusters position. 
+
+				if (TPBDRigidParticleHandle<float, 3>* ParentHandle = Handle->ClusterIds().Id)
+				{
+					if (TPBDRigidClusteredParticleHandle<float, 3>* ClusteredParentHandle = ParentHandle->CastToClustered())
+					{
+						if (ClusteredParentHandle->InternalCluster())
+						{
+							//const FTransform ClusterChildToWorld = ClusterChildToParentMap[RigidBodyIndex] * FTransform(Particles.R(ClusterParentIndex), Particles.X(ClusterParentIndex));
+							const TRigidTransform<float, 3>& ChildToParent = Handle->ChildToParent();
+							const FTransform ClusterChildToWorld = ChildToParent * FTransform(ParentHandle->R(), ParentHandle->X());
+							if (Parameters.IsCacheRecording())
+							{
+								//Particles.X(RigidBodyIndex) = ClusterChildToWorld.GetTranslation();
+								Handle->SetX(ClusterChildToWorld.GetTranslation());
+								//Particles.R(RigidBodyIndex) = ClusterChildToWorld.GetRotation();
+								Handle->SetR(ClusterChildToWorld.GetRotation());
+							}
+							// GeomToActor = ActorToWorld.Inv() * ClusterChildToWorld * MassToLocal.Inv();
+							TargetResults.Transforms[TransformIndex] = 
+								MassToLocal[TransformIndex].GetRelativeTransformReverse(ClusterChildToWorld).GetRelativeTransform(ActorToWorld);
+							TargetResults.Transforms[TransformIndex].NormalizeRotation();
+						}
+					}
+				}
+
+			}// end if
+		} // end for
+	} // tmp scope
+
+	// If object is dynamic, compute global matrices	
+	if (IsObjectDynamic || TargetResults.GlobalTransforms.Num() == 0)
+	{
+		SCOPE_CYCLE_COUNTER(STAT_CalcGlobalGCMatrices);
+		check(TargetResults.Transforms.Num() == TargetResults.Parent.Num());
+		GeometryCollectionAlgo::GlobalMatrices(
+			TargetResults.Transforms, TargetResults.Parent, TargetResults.GlobalTransforms);
+	}
+
+	// compute world bounds
+	// #note: this is a loose bounds based on the circumscribed box of a bounding sphere for the geometry.		
+	if (IsObjectDynamic || TargetResults.WorldBounds.GetSphere().W < 1e-5)
+	{
+		SCOPE_CYCLE_COUNTER(STAT_CalcGlobalGCBounds);
+		FBox BoundingBox(ForceInit);
+		//const FMatrix& ActorToWorld = Parameters.WorldTransform.ToMatrixWithScale();
+		for (int i = 0; i < ValidGeometryBoundingBoxes.Num(); ++i)
+		{
+			BoundingBox += ValidGeometryBoundingBoxes[i].TransformBy(
+				TargetResults.GlobalTransforms[ValidGeometryTransformIndices[i]] * ActorToWorld);
+		}
+		TargetResults.WorldBounds = FBoxSphereBounds(BoundingBox);		
+	}
+
+#if 0 // TODO Ryan
+	for (int32 TransformGroupIndex = 0; TransformGroupIndex < TransformSize; TransformGroupIndex++)
+	{
+		const int32 ParentIndex = EndFrameUnparentingBuffer[TransformGroupIndex];
+		if (ParentIndex >= 0)
+		{
+			// We reuse EndFrameUnparentingBuffer potentially without reinitialization, so reset this index to -1 before it gets paged out.
+			EndFrameUnparentingBuffer[TransformGroupIndex] = -1;
+
+			Children[ParentIndex].Remove(TransformGroupIndex);
+			Parent[TransformGroupIndex] = INDEX_NONE;
+		}
+	}
+#endif
+
+	// We pulled state out of particles.  Push results into PTDynamicCollection, 
+	// so that it has an updated view.
+	UpdateGeometryCollection(TargetResults);
+
+	// We've updated a buffer.  There's no reason to wait for FGeometryCollectionPhysicsProxy::FlipBuffer()
+	// to be called.  Just flip.
+	PhysToGameInterchange->FlipProducer();
+//#endif // TODO_REIMPLEMENT_GET_RIGID_PARTICLES
+}
+
+void FGeometryCollectionPhysicsProxy::FlipBuffer()
+{
+	/**
+	 * CONTEXT: PHYSICSTHREAD (Write Locked)
+	 * Called by the physics thread to signal that it is safe to perform any double-buffer flips here.
+	 * The physics thread has pre-locked an RW lock for this operation so the game thread won't be reading
+	 * the data
+	 */
+}
+
+// Called from FPhysScene_ChaosInterface::SyncBodies(), NOT the solver.
+void FGeometryCollectionPhysicsProxy::PullFromPhysicsState()
+{
+	/**
+	 * CONTEXT: GAMETHREAD (Read Locked)
+	 * Perform a similar operation to Sync, but take the data from a gamethread-safe buffer. This will be called
+	 * from the game thread when it cannot sync to the physics thread. The simulation is very likely to be running
+	 * when this happens so never read any physics thread data here!
+	 *
+	 * Note: A read lock will have been acquired for this - so the physics thread won't force a buffer flip while this
+	 * sync is ongoing
+	 */
+
+	const FGeometryCollectionResults* TargetResultPtr = PhysToGameInterchange->GetConsumerBuffer();
+	if (!TargetResultPtr)
+		return;
+	// Remove const-ness as we're going to do the ExchangeArrays() thing.
+	FGeometryCollectionResults& TR = *const_cast<FGeometryCollectionResults*>(TargetResultPtr);
+
+	// We should never be changing the number of entries, this would break other 
+	// attributes in the transform group.
+	const int32 NumTransforms = GTDynamicCollection->Transform.Num();
+	if (ensure(NumTransforms == TR.Transforms.Num()))	
+	{
+		for(int32 TmIndex = 0; TmIndex < NumTransforms; ++TmIndex)
+		{
+			if(!TR.DisabledStates[TmIndex])
+			{
+				GTDynamicCollection->Transform[TmIndex] = TR.Transforms[TmIndex];
+			}
+		}
+
+		//question: why do we need this? Sleeping objects will always have to update GPU
+		GTDynamicCollection->MakeDirty();
+
+		if (CacheSyncFunc)
+		{
+			CacheSyncFunc(TR);
+		}
+	}
+/*
+	for (int32 TransformIndex = 0; TransformIndex < NumParticles; ++TransformIndex)
+	{
+		const int32 ParticleIndex = BaseParticleIndex + TransformIndex;
+		//Chaos::TPBDRigidParticleHandle<float, 3>& Handle = *SolverParticleHandles[ParticleIndex];
+		if (Chaos::TGeometryParticle<float, 3>* Particle = GTGeometryParticles[ParticleIndex].Get())
+		{
+			const FTransform& Xf = TR.Transforms[ParticleIndex];
+			Particle->SetX(Xf.GetTranslation());
+			Particle->SetR(Xf.GetRotation());
+
+			GTDynamicCollection->Transform[TransformIndex] = Xf;
+		}
+	}
+*/
+//	UpdateGeometryCollection(TR);
+}
+
+void FGeometryCollectionPhysicsProxy::UpdateGeometryCollection(
+	const FGeometryCollectionResults& GCResults,
+	FGeometryDynamicCollection* Collection)
+{
+	if (InitializedState != ESimulationInitializationState::Initialized)
+	{
+		return;
+	}
+	if (!Collection)
+	{
+		Collection = Parameters.DynamicCollection;
+	}
+	check(Collection);
+
+	TManagedArray<FTransform>& Transform = Collection->Transform;
+	TManagedArray<int32>& Parent = Collection->Parent;
+	TManagedArray<TSet<int32>>& Children = Collection->Children;
+	TManagedArray<int32>& SimulationType = Collection->SimulationType;
+	TManagedArray<int32>& DynamicState = Collection->GetAttribute<int32>("DynamicState", FGeometryCollection::TransformGroup);
+
+	// Particles X and R are aligned with center of mass and inertia principal axes.
+	// Renderer doesn't know about this and simply does ActorToWorld * GeomToActor * LocalSpaceVerts
+	// In proper math multiplication order:
+	//		ParticleToWorld = ActorToWorld * GeomToActor * MassToLocal
+	//		GeomToWorld = ActorToWorld * GeomToActor
+	//		=> GeomToWorld = ParticleToWorld * MassToLocal.Inv()
+	//		=> GeomToActor = ActorToWorld.Inv() * ParticleToWorld * MassToLocal.Inv()
+
+	const FTransform& ActorToWorld = Parameters.WorldTransform;
+	const int32 TransformSize = GetTransformGroupSize();
+	for (int32 TransformGroupIndex = 0; TransformGroupIndex < TransformSize; ++TransformGroupIndex)
+	{
+		if (!GTGeometryParticles[TransformGroupIndex])
+			continue;
+
+		auto* Handle = SolverParticleHandles[TransformGroupIndex];
+
+		// dynamic state is also updated by the solver during field interaction. 
+		if(!Handle->Sleeping())
+		{
+			DynamicState[TransformGroupIndex] = GCResults.DynamicState[TransformGroupIndex];
+		}
+
+		// Update the transform and parent hierarchy of the active rigid bodies. Active bodies can be either
+		// rigid geometry defined from the leaf nodes of the collection, or cluster bodies that drive an entire
+		// branch of the hierarchy within the GeometryCollection.
+		// - Active bodies are directly driven from the global position of the corresponding
+		//   rigid bodies within the solver ( cases where RigidBodyID[TransformGroupIndex] is not disabled ). 
+		// - Deactivated bodies are driven from the transforms of their active parents. However the solver can
+		//   take ownership of the parents during the simulation, so it might be necessary to force deactivated
+		//   bodies out of the collections hierarchy during the simulation.  
+
+		if(!GCResults.DisabledStates[TransformGroupIndex])
+		{
+			// Update the transform of the active body. The active body can be either a single rigid
+			// or a collection of rigidly attached geometries (Clustering). The cluster is represented as a
+			// single transform in the GeometryCollection, and all children are stored in the local space
+			// of the parent cluster.
+			// ... When setting cluster transforms it is expected that the MassToLocal is identity.
+			//     Cluster initialization will set the vertices in the MassSpace of the rigid body.
+			// ... When setting individual rigid bodies that are not clustered, the MassToLocal will be 
+			//     non-Identity, and will reflect the difference between the geometric center of the geometry
+			//     and that corresponding rigid bodies center of mass. 
+
+			// GeomToActor = ActorToWorld.Inv() * ParticleToWorld * MassToLocal.Inv();
+			Transform[TransformGroupIndex] = GCResults.Transforms[TransformGroupIndex];
+
+			// When a leaf node rigid body is removed from a cluster the rigid
+			// body will become active and needs its clusterID updated. This just
+			// syncs the clusterID all the time. 
+//			CollectionClusterID[TransformGroupIndex] = ClusterID[RigidBodyIndex].Id;
+		}
+		else //if (Particles.Disabled(RigidBodyIndex))
+		{
+			Transform[TransformGroupIndex] = GCResults.Transforms[TransformGroupIndex];
+		}
+	}
+}
+
+
+//==============================================================================
+// STATIC SETUP FUNCTIONS
+//==============================================================================
+
+
 int32 FindSizeSpecificIdx(const TArray<FSharedSimulationSizeSpecificData>& SizeSpecificData, const FBox& Bounds)
 {
 	const FVector Extents = Bounds.GetExtent();
 	const float Size = Extents.GetAbsMin();
-
 	check(SizeSpecificData.Num());
 	int32 UseIdx = 0;
 	float PreSize = FLT_MAX;
-
 	for (int32 Idx = SizeSpecificData.Num() - 1; Idx >=0 ; --Idx)
 	{
 		ensureMsgf(PreSize >= SizeSpecificData[Idx].MaxSize, TEXT("SizeSpecificData is not sorted"));
 		PreSize = SizeSpecificData[Idx].MaxSize;
-
 		if (Size < SizeSpecificData[Idx].MaxSize)
-		{
 			UseIdx = Idx;
-		}
 		else
-		{
 			break;
-		}
 	}
-
 	return UseIdx;
 }
 
@@ -3313,51 +3150,72 @@
 	}
 }
 
-void FGeometryCollectionPhysicsProxy::InitializeKinematics(FParticlesType& Particles, const TManagedArray<int32>& DynamicState)
+void IdentifySimulatableElements(Chaos::FErrorReporter& ErrorReporter, FGeometryCollection& GeometryCollection)
 {
-	if(Parameters.DynamicCollection)
+	// Determine which collection particles to simulate
+	const TManagedArray<FBox>& BoundingBox = GeometryCollection.BoundingBox;
+	const TManagedArray<int32>& VertexCount = GeometryCollection.VertexCount;
+	const TManagedArray<int32>& TransformIndex = GeometryCollection.TransformIndex;
+	const int32 NumTransforms = GeometryCollection.NumElements(FGeometryCollection::TransformGroup);
+
+	const int32 NumTransformMappings = TransformIndex.Num();
+
+	// Do not simulate hidden geometry
+	TArray<bool> HiddenObject;
+	HiddenObject.Init(true, NumTransforms);
+	const TManagedArray<bool>& Visible = GeometryCollection.Visible;
+	const TManagedArray<int32>& BoneMap = GeometryCollection.BoneMap;
+
+	const TManagedArray<FIntVector>& Indices = GeometryCollection.Indices;
+	int32 PrevObject = -1;
+	bool bContiguous = true;
+	for(int32 i = 0; i < Indices.Num(); i++)
 	{
-		for(int TransformGroupIndex = 0; TransformGroupIndex < RigidBodyID.Num(); TransformGroupIndex++)
+		if(Visible[i])
 		{
-			if(RigidBodyID[TransformGroupIndex] != INDEX_NONE)
+			int32 ObjIdx = BoneMap[Indices[i][0]];
+			HiddenObject[ObjIdx] = false;
+			if (!ensureMsgf(ObjIdx >= PrevObject, TEXT("Objects are not contiguous. This breaks assumptions later in the pipeline")))
 			{
-				int32 RigidBodyIndex = RigidBodyID[TransformGroupIndex];
-				if(DynamicState[TransformGroupIndex] == (uint8)EObjectStateTypeEnum::Chaos_Object_Kinematic)
-				{
-					Particles.SetObjectState(RigidBodyIndex, Chaos::EObjectStateType::Kinematic);
-				}
-				else if(DynamicState[TransformGroupIndex] == (uint8)EObjectStateTypeEnum::Chaos_Object_Static)
-				{
-					Particles.SetObjectState(RigidBodyIndex, Chaos::EObjectStateType::Static);
-				}
-				else if(DynamicState[TransformGroupIndex] == (uint8)EObjectStateTypeEnum::Chaos_Object_Sleeping)
-				{
-					Particles.SetObjectState(RigidBodyIndex, Chaos::EObjectStateType::Sleeping);
-				}
+				bContiguous = false;
 			}
+			
+			PrevObject = ObjIdx;
+		}
+	}
+
+	if (!bContiguous)
+	{
+		ErrorReporter.ReportError(TEXT("Objects are not contiguous. This breaks assumptions later in the pipeline"));
+		ErrorReporter.HandleLatestError();
+	}
+
+
+	//For now all simulation data is a non compiled attribute. Not clear what we want for simulated vs kinematic collections
+	TManagedArray<bool>& SimulatableParticles = 
+		GeometryCollection.AddAttribute<bool>(
+			FGeometryCollection::SimulatableParticlesAttribute, FTransformCollection::TransformGroup);
+
+	for(int32 TransformIdx = 0; TransformIdx < NumTransforms; TransformIdx++)
+	{
+		SimulatableParticles[TransformIdx] = false;
+	}
+
+	for(int i = 0; i < NumTransformMappings; i++)
+	{
+		int32 Tdx = TransformIndex[i];
+		checkSlow(0 <= Tdx && Tdx < NumTransforms);
+		if(GeometryCollection.IsGeometry(Tdx) && VertexCount[i] && 0.f < BoundingBox[i].GetSize().SizeSquared() && !HiddenObject[Tdx])
+		{
+			SimulatableParticles[Tdx] = true;
 		}
 	}
 }
 
-void FGeometryCollectionPhysicsProxy::InitializeRemoveOnFracture(FParticlesType& Particles, const TManagedArray<int32>& DynamicState)
+void BuildSimulationData(Chaos::FErrorReporter& ErrorReporter, FGeometryCollection& GeometryCollection, const FSharedSimulationParameters& SharedParams)
 {
-	if (Parameters.DynamicCollection && Parameters.RemoveOnFractureEnabled)
-	{
-	//	TManagedArray<FGeometryCollectionBoneNode>& Hierarchy = Parameters.DynamicCollection->BoneHierarchy;
-
-		for (int TransformGroupIndex = 0; TransformGroupIndex < RigidBodyID.Num(); TransformGroupIndex++)
-		{
-			if (RigidBodyID[TransformGroupIndex] != INDEX_NONE)
-			{
-				int32 RigidBodyIndex = RigidBodyID[TransformGroupIndex];
-
-				if (Parameters.DynamicCollection->StatusFlags[TransformGroupIndex] & FGeometryCollection::FS_RemoveOnFracture)
-				{
-					Particles.ToBeRemovedOnFracture(RigidBodyIndex) = true;
-				}
-			}
-		}
-	}
+	IdentifySimulatableElements(ErrorReporter, GeometryCollection);
+	FGeometryCollectionPhysicsProxy::InitializeSharedCollisionStructures(ErrorReporter, GeometryCollection, SharedParams);
 }
 
 void FGeometryCollectionPhysicsProxy::MergeRecordedTracks(const FRecordedTransformTrack& A, const FRecordedTransformTrack& B, FRecordedTransformTrack& Target)
@@ -3411,1255 +3269,1383 @@
 	return InTrack.Records[NewRecordIndex];
 }
 
-void FGeometryCollectionPhysicsProxy::AddCollisionToCollisionData(FRecordedFrame* ExistingFrame, const FParticlesType& Particles, const Chaos::TPBDCollisionConstraints<float, 3>::FPointContactConstraint& Constraint)
+
+//==============================================================================
+// FIELDS
+//==============================================================================
+
+
+void FGeometryCollectionPhysicsProxy::ProcessCommands(FParticlesType& Particles, const float Time)
 {
-#if TODO_REIMPLEMENT_RIGID_CLUSTERING
-	const Chaos::TArrayCollectionArray<Chaos::ClusterId>& ClusterIdsArray = GetSolver()->GetRigidClustering().GetClusterIdsArray();
-	const TMap<uint32, TUniquePtr<TArray<uint32>>>& ParentToChildrenMap = GetSolver()->GetRigidClustering().GetChildrenMap();
-	const Chaos::TArrayCollectionArray<PhysicsProxyWrapper>& PhysicsProxyReverseMapping = GetSolver()->GetPhysicsProxyReverseMapping();
+	FGeometryDynamicCollection* Collection = Parameters.DynamicCollection;
+	check(Collection);
 
-	int32 NewIdx = ExistingFrame->Collisions.Add(FSolverCollisionData());
-	FSolverCollisionData& Collision = ExistingFrame->Collisions[NewIdx];
+	// Process Particle-Collection commands
+	if (Commands.Num())
+	{
+		TArray<ContextIndex> IndicesArray;
+		Chaos::FPhysicsSolver* CurrentSolver = GetSolver();
 
-	Collision.Location = Constraint.Location;
-	Collision.AccumulatedImpulse = Constraint.AccumulatedImpulse;
-	Collision.Normal = Constraint.Normal;
-	Collision.Velocity1 = Particles.V(Constraint.ParticleIndex);
-	Collision.Velocity2 = Particles.V(Constraint.LevelsetIndex);
-	Collision.AngularVelocity1 = Particles.W(Constraint.ParticleIndex);
-	Collision.AngularVelocity2 = Particles.W(Constraint.LevelsetIndex);
-	Collision.Mass1 = Particles.M(Constraint.ParticleIndex);
-	Collision.Mass2 = Particles.M(Constraint.LevelsetIndex);
-	Collision.ParticleIndex = Constraint.ParticleIndex;
-	Collision.LevelsetIndex = Constraint.LevelsetIndex;
+		for (int32 CommandIndex = Commands.Num()-1; 0<=CommandIndex; CommandIndex--)
+		{
 
-	// If Constraint.ParticleIndex is a cluster store an index for a mesh in this cluster
-	if (ClusterIdsArray[Constraint.ParticleIndex].NumChildren > 0)
+			//
+			// Extract command and set metadata
+			//
+			FFieldSystemCommand & Command = Commands[CommandIndex];
+			EFieldResolutionType ResolutionType = EFieldResolutionType::Field_Resolution_Minimal;
+			if (Command.MetaData.Contains(FFieldSystemMetaData::EMetaType::ECommandData_ProcessingResolution))
+			{
+				check(Command.MetaData[FFieldSystemMetaData::EMetaType::ECommandData_ProcessingResolution] != nullptr);
+				ResolutionType = static_cast<FFieldSystemMetaDataProcessingResolution*>(Command.MetaData[FFieldSystemMetaData::EMetaType::ECommandData_ProcessingResolution].Get())->ProcessingResolution;
+			}
+
+			if (Command.TargetAttribute == GetGeometryCollectionPhysicsTypeName(EGeometryCollectionPhysicsTypeEnum::Chaos_DynamicState))
+			{
+				if (ensureMsgf(Command.RootNode->Type() == FFieldNode<int32>::StaticType(),
+					TEXT("Field based evaluation of the simulations 'DynamicState' parameter expects int32 field inputs.")))
+				{
+					FGeometryCollectionPhysicsProxy::ContiguousIndices(IndicesArray, CurrentSolver, ResolutionType, IndicesArray.Num() != Particles.Size());
+					if (IndicesArray.Num())
+					{
+						TArrayView<ContextIndex> IndexView(&(IndicesArray[0]), IndicesArray.Num());
+
+						FVector * xptr = &(Particles.X(0));
+						TArrayView<FVector> SamplesView(xptr, Particles.Size());
+
+						FFieldContext Context{
+							IndexView,
+							SamplesView,
+							Command.MetaData
+						};
+
+						TArrayView<int32> DynamicStateView(&(Collection->DynamicState[0]), Collection->DynamicState.Num());
+						static_cast<const FFieldNode<int32> *>(Command.RootNode.Get())->Evaluate(Context, DynamicStateView);
+
+						PushKinematicStateToSolver(Particles);
+					}
+				}
+				Commands.RemoveAt(CommandIndex);
+			}
+			else if (Command.TargetAttribute == GetGeometryCollectionPhysicsTypeName(EGeometryCollectionPhysicsTypeEnum::Chaos_InitialLinearVelocity))
+			{
+				if (ensureMsgf(Parameters.InitialVelocityType == EInitialVelocityTypeEnum::Chaos_Initial_Velocity_User_Defined,
+					TEXT("Field based evaluation of the simulations 'InitialLinearVelocity' requires the geometry collection be set to User Defined Initial Velocity")))
+				{
+					if (ensureMsgf(Command.RootNode->Type() == FFieldNode<FVector>::StaticType(),
+						TEXT("Field based evaluation of the simulations 'InitialLinearVelocity' parameter expects FVector field inputs.")))
+					{
+						FGeometryCollectionPhysicsProxy::ContiguousIndices(IndicesArray, CurrentSolver, ResolutionType, IndicesArray.Num() != Particles.Size());
+						if (IndicesArray.Num())
+						{
+							TArrayView<ContextIndex> IndexView(&(IndicesArray[0]), IndicesArray.Num());
+
+							FVector * xptr = &(Particles.X(0));
+							TArrayView<FVector> SamplesView(xptr, Particles.Size());
+
+							FFieldContext Context{
+								IndexView,
+								SamplesView,
+								Command.MetaData
+							};
+
+							TArrayView<FVector> ResultsView(&(InitialLinearVelocity[0]), InitialLinearVelocity.Num());
+							static_cast<const FFieldNode<FVector> *>(Command.RootNode.Get())->Evaluate(Context, ResultsView);
+						}
+					}
+				}
+				Commands.RemoveAt(CommandIndex);
+			}
+			else if (Command.TargetAttribute == GetGeometryCollectionPhysicsTypeName(EGeometryCollectionPhysicsTypeEnum::Chaos_InitialAngularVelocity))
+			{
+				if (ensureMsgf(Parameters.InitialVelocityType == EInitialVelocityTypeEnum::Chaos_Initial_Velocity_User_Defined,
+					TEXT("Field based evaluation of the simulations 'InitialAngularVelocity' requires the geometry collection be set to User Defined Initial Velocity")))
+				{
+					if (ensureMsgf(Command.RootNode->Type() == FFieldNode<FVector>::StaticType(),
+						TEXT("Field based evaluation of the simulations 'InitialAngularVelocity' parameter expects FVector field inputs.")))
+					{
+						FGeometryCollectionPhysicsProxy::ContiguousIndices(IndicesArray, CurrentSolver, ResolutionType, IndicesArray.Num() != Particles.Size());
+						if (IndicesArray.Num())
+						{
+							TArrayView<ContextIndex> IndexView(&(IndicesArray[0]), IndicesArray.Num());
+
+							FVector * xptr = &(Particles.X(0));
+							TArrayView<FVector> SamplesView(xptr, Particles.Size());
+
+							FFieldContext Context{
+								IndexView,
+								SamplesView,
+								Command.MetaData
+							};
+
+							TArrayView<FVector> ResultsView(&(InitialAngularVelocity[0]), InitialAngularVelocity.Num());
+							static_cast<const FFieldNode<FVector> *>(Command.RootNode.Get())->Evaluate(Context, ResultsView);
+						}
+					}
+				}
+				Commands.RemoveAt(CommandIndex);
+			}
+		}
+	}
+
+	// Process Particle-Particle commands
+	if (Commands.Num())
 	{
-		int32 ParticleIndexMesh = GetSolver()->GetParticleIndexMesh(ParentToChildrenMap, Constraint.ParticleIndex);
-		ensure(ParticleIndexMesh != INDEX_NONE);
-		Collision.ParticleIndexMesh = ParticleIndexMesh;
+		Chaos::FPhysicsSolver* CurrentSolver = GetSolver();
+
+		//  Generate a Index mapping between the rigid body indices and 
+		//  the particle indices. This allows the geometry collection to
+		//  evaluate only its own particles. 
+		TArray<ContextIndex> IndicesArray;
+		int32 NumIndices = 0;
+		IndicesArray.SetNumUninitialized(RigidBodyID.Num());
+		for (int32 i = 0; i < RigidBodyID.Num(); i++)
+		{
+			if (RigidBodyID[i] != INDEX_NONE)
+			{
+				IndicesArray[NumIndices] = { RigidBodyID[i], RigidBodyID[i] };
+				NumIndices++;
+			}
+		}
+		IndicesArray.SetNum(NumIndices);
+
+		for (int32 CommandIndex = Commands.Num() - 1; 0 <= CommandIndex; CommandIndex--)
+		{
+			FFieldSystemCommand & Command = Commands[CommandIndex];
+			if (Command.TargetAttribute == GetGeometryCollectionPhysicsTypeName(EGeometryCollectionPhysicsTypeEnum::Chaos_LinearVelocity))
+			{
+				if (ensureMsgf(Command.RootNode->Type() == FFieldNode<FVector>::StaticType(),
+					TEXT("Field based evaluation of the simulations 'LinearVelocity' parameter expects FVector field inputs.")))
+				{
+					FVector * xptr = &(Particles.X(0));
+					TArrayView<FVector> SamplesView(xptr, Particles.Size());
+					TArrayView<ContextIndex> IndexView(&(IndicesArray[0]), IndicesArray.Num());
+
+					FFieldContext Context{
+						IndexView,
+						SamplesView,
+						Command.MetaData
+					};
+
+					FVector * vptr = &(Particles.V(0));
+					TArrayView<FVector> ResultsView(vptr, Particles.Size());
+					static_cast<const FFieldNode<FVector> *>(Command.RootNode.Get())->Evaluate(Context, ResultsView);
+				}
+				Commands.RemoveAt(CommandIndex);
+			}
+			else if (Command.TargetAttribute == GetGeometryCollectionPhysicsTypeName(EGeometryCollectionPhysicsTypeEnum::Chaos_AngularVelocity))
+			{
+				if (ensureMsgf(Command.RootNode->Type() == FFieldNode<FVector>::StaticType(),
+					TEXT("Field based evaluation of the simulations 'AngularVelocity' parameter expects FVector field inputs.")))
+				{
+
+					FVector * xptr = &(Particles.X(0));
+					TArrayView<FVector> SamplesView(xptr, Particles.Size());
+					TArrayView<ContextIndex> IndexView(&(IndicesArray[0]), IndicesArray.Num());
+
+					FFieldContext Context{
+						IndexView,
+						SamplesView,
+						Command.MetaData
+					};
+
+					FVector * vptr = &(Particles.W(0));
+					TArrayView<FVector> ResultsView(vptr, Particles.Size());
+					static_cast<const FFieldNode<FVector> *>(Command.RootNode.Get())->Evaluate(Context, ResultsView);
+				}
+				Commands.RemoveAt(CommandIndex);
+			}
+			else if (Command.TargetAttribute == GetGeometryCollectionPhysicsTypeName(EGeometryCollectionPhysicsTypeEnum::Chaos_CollisionGroup))
+			{
+				if (ensureMsgf(Command.RootNode->Type() == FFieldNode<int32>::StaticType(),
+					TEXT("Field based evaluation of the simulations 'CollisionGroup' parameter expects int32 field inputs.")))
+				{
+					FVector * xptr = &(Particles.X(0));
+					TArrayView<FVector> SamplesView(xptr, Particles.Size());
+					TArrayView<ContextIndex> IndexView(&(IndicesArray[0]), IndicesArray.Num());
+
+					FFieldContext Context{
+						IndexView,
+						SamplesView,
+						Command.MetaData
+					};
+
+					int32 * cptr = &(Particles.CollisionGroup(0));
+					TArrayView<int32> ResultsView(cptr, Particles.Size());
+					static_cast<const FFieldNode<int32> *>(Command.RootNode.Get())->Evaluate(Context, ResultsView);
+				}
+				Commands.RemoveAt(CommandIndex);
+			}
+		}
 	}
-	// If Constraint.LevelsetIndex is a cluster store an index for a mesh in this cluster
-	if (ClusterIdsArray[Constraint.LevelsetIndex].NumChildren > 0)
-	{
-		int32 LevelsetIndexMesh = GetSolver()->GetParticleIndexMesh(ParentToChildrenMap, Constraint.LevelsetIndex);
-		ensure(LevelsetIndexMesh != INDEX_NONE);
-		Collision.LevelsetIndexMesh = LevelsetIndexMesh;
-	}
-#endif
 }
 
-void FGeometryCollectionPhysicsProxy::UpdateCollisionData(const FParticlesType& Particles, const FCollisionConstraintsType& CollisionRule, FRecordedFrame* ExistingFrame)
+void FGeometryCollectionPhysicsProxy::FieldForcesUpdateCallback(Chaos::FPhysicsSolver* InSolver, FParticlesType& Particles, Chaos::TArrayCollectionArray<FVector> & Force, Chaos::TArrayCollectionArray<FVector> & Torque, const float Time)
 {
-	ExistingFrame->Collisions.Reset();
+	if (InitializedState == ESimulationInitializationState::Initialized)
+	{
+		if (Commands.Num())
+		{
+			// @todo: This seems like a waste if we just want to get everything
+			int32 Counter = 0;
+			TArray<ContextIndex> IndicesArray;
+			IndicesArray.AddUninitialized(RigidBodyID.Num());
+			for (int32 i = 0; i < RigidBodyID.Num(); i++)
+			{
+				if (RigidBodyID[i] != INDEX_NONE)
+				{
+					IndicesArray[i] = { RigidBodyID[i],RigidBodyID[i] };
+					Counter++;
+				}
+			}
+			IndicesArray.SetNum(Counter, false);
+			TArrayView<ContextIndex> IndexView(&(IndicesArray[0]), IndicesArray.Num());
 
-	if (Parameters.CollisionData.SaveCollisionData && ExistingFrame->Timestamp > 0.f && Parameters.CollisionData.CollisionDataSizeMax > 0)
+
+			for (int32 CommandIndex = 0; CommandIndex < Commands.Num(); CommandIndex++)
+			{
+				FFieldSystemCommand & Command = Commands[CommandIndex];
+
+				if (Command.TargetAttribute == GetGeometryCollectionPhysicsTypeName(EGeometryCollectionPhysicsTypeEnum::Chaos_LinearForce))
+				{
+					if (ensureMsgf(Command.RootNode->Type() == FFieldNode<FVector>::StaticType(),
+						TEXT("Field based evaluation of the simulations 'LinearForce' parameter expects FVector field inputs.")))
+					{
+						FVector * tptr = &(Particles.X(0));
+						TArrayView<FVector> SamplesView(tptr, int32(Particles.Size()));
+
+						FFieldContext Context{
+							IndexView,
+							SamplesView,
+							Command.MetaData
+						};
+						TArrayView<FVector> ForceView(&(Force[0]), Force.Num());
+						static_cast<const FFieldNode<FVector> *>(Command.RootNode.Get())->Evaluate(Context, ForceView);
+					}
+					Commands.RemoveAt(CommandIndex);
+				}
+				else if (Command.TargetAttribute == GetGeometryCollectionPhysicsTypeName(EGeometryCollectionPhysicsTypeEnum::Chaos_AngularTorque))
+				{
+					if (ensureMsgf(Command.RootNode->Type() == FFieldNode<FVector>::StaticType(),
+						TEXT("Field based evaluation of the simulations 'AngularTorque' parameter expects FVector field inputs.")))
+					{
+						FVector * tptr = &(Particles.X(0));
+						TArrayView<FVector> SamplesView(tptr, int32(Particles.Size()));
+
+						FFieldContext Context{
+							IndexView,
+							SamplesView,
+							Command.MetaData
+						};
+						TArrayView<FVector> TorqueView(&(Torque[0]), Torque.Num());
+						static_cast<const FFieldNode<FVector> *>(Command.RootNode.Get())->Evaluate(Context, TorqueView);
+					}
+					Commands.RemoveAt(CommandIndex);
+				}
+
+			}
+		}
+	}
+}
+
+
+//==============================================================================
+// DEPRECATED
+//==============================================================================
+
+
+void FGeometryCollectionPhysicsProxy::UpdateKinematicBodiesCallback(const FParticlesType& Particles, const float Dt, const float Time, FKinematicProxy& Proxy)
+{
+	check(false); // not currently called
+	if (InitializedState == ESimulationInitializationState::Initialized)
 	{
-#if TODO_REIMPLEMENT_PHYSICS_PROXY_REVERSE_MAPPING
-		const TArray<Chaos::TPBDCollisionConstraints<float, 3>::FRigidBodyContactConstraint>& AllConstraintsArray = CollisionRule.GetAllConstraints();
-		if (AllConstraintsArray.Num() > 0)
+		SCOPE_CYCLE_COUNTER(STAT_KinematicUpdate);
+		FGeometryDynamicCollection* Collection = Parameters.DynamicCollection;
+		check(Collection);
+
+		bool bIsCachePlaying = Parameters.IsCachePlaying();
+		bool bIsReverseCachePlaying = Parameters.IsCacheRecording() && Parameters.ReverseCacheBeginTime != 0 && Parameters.ReverseCacheBeginTime < Time;
+		if ((!bIsCachePlaying && !bIsReverseCachePlaying) || !Parameters.RecordedTrack)
 		{
-			const Chaos::TArrayCollectionArray<PhysicsProxyWrapper>& PhysicsProxyReverseMapping = GetSolver()->GetPhysicsProxyReverseMapping();
+			return;
+		}
 
-			TArray<Chaos::TPBDCollisionConstraints<float, 3>::FRigidBodyContactConstraint> ConstraintsArray;
-			ConstraintsArray.SetNumUninitialized(AllConstraintsArray.Num());
+#if TODO_REIMPLEMENT_KINEMATIC_PROXY
+		bool bFirst = !Proxy.Ids.Num();
+		if (bFirst)
+		{
+			Proxy.Position.Reset(RigidBodyID.Num());
+			Proxy.Rotation.Reset(RigidBodyID.Num());
+			Proxy.NextPosition.Reset(RigidBodyID.Num());
+			Proxy.NextRotation.Reset(RigidBodyID.Num());
 
-			FBox BoundingBox(ForceInitToZero);
-			int32 NumConstraints = 0;
-			for (int32 Idx = 0; Idx < AllConstraintsArray.Num(); ++Idx)
+			Proxy.Position.AddUninitialized(RigidBodyID.Num());
+			Proxy.Rotation.AddUninitialized(RigidBodyID.Num());
+			Proxy.NextPosition.AddUninitialized(RigidBodyID.Num());
+			Proxy.NextRotation.AddUninitialized(RigidBodyID.Num());
+
+			for (int32 i = 0; i < RigidBodyID.Num(); ++i)
 			{
-				// Check if the collision is for this PhysicsProxy
-				IPhysicsProxyBase* PhysicsProxy = PhysicsProxyReverseMapping[AllConstraintsArray[Idx].ParticleIndex].PhysicsProxy;
-				if (PhysicsProxy == this)
+				Proxy.Ids.Add(RigidBodyID[i]);
+
+				// Initialise to rest state
+				const int32 RbId = Proxy.Ids.Last();
+				Proxy.Position[i] = RbId != INDEX_NONE ? Particles.X(RbId) : FVector::ZeroVector;
+				Proxy.Rotation[i] = RbId != INDEX_NONE ? Particles.R(RbId) : FQuat::Identity;
+				Proxy.NextPosition[i] = Proxy.Position[i];
+				Proxy.NextRotation[i] = Proxy.Rotation[i];
+			}
+		}
+
+		if (bIsCachePlaying && !bIsReverseCachePlaying && (Time < Parameters.CacheBeginTime || !Parameters.RecordedTrack->IsTimeValid(Time)))
+		{
+			return;
+		}
+
+		float ReverseTime = Parameters.RecordedTrack->GetLastTime() - Time + Parameters.ReverseCacheBeginTime;
+		if (bIsReverseCachePlaying && !Parameters.RecordedTrack->IsTimeValid(ReverseTime))
+		{
+			return;
+		}
+
+		const FRecordedFrame* FirstFrame = nullptr;
+		const FRecordedFrame* SecondFrame = nullptr;
+		float PlaybackTime = bIsReverseCachePlaying ? ReverseTime : Time;
+		Parameters.RecordedTrack->GetFramesForTime(PlaybackTime, FirstFrame, SecondFrame);
+
+		if (FirstFrame && !SecondFrame)
+		{
+			// Only one frame to take information from (simpler case)
+			const int32 NumActives = FirstFrame->TransformIndices.Num();
+
+			// Actives
+			Chaos::PhysicsParallelFor(NumActives, [&](int32 Index)
+			{
+				const int32 InternalIndex = FirstFrame->TransformIndices[Index];
+				if (InternalIndex >= RigidBodyID.Num() || InternalIndex < 0)
 				{
-					if (ensure(!AllConstraintsArray[Idx].AccumulatedImpulse.ContainsNaN()))
+					UE_LOG(UGCC_LOG, Error, 
+						TEXT("%s: Cache index %d out of range: [%d, %d).  Regenerate the cache."),
+						*Parameters.Name, InternalIndex, 0, RigidBodyID.Num());
+					return;
+				}
+				const int32 ExternalIndex = RigidBodyID[InternalIndex];
+
+				if (ExternalIndex != INDEX_NONE && Particles.InvM(ExternalIndex) == 0.0f && !Particles.Disabled(ExternalIndex))
+				{
+					const FTransform& ParticleTransform = FirstFrame->Transforms[Index];
+					Proxy.Position[InternalIndex] = Particles.X(ExternalIndex);
+					Proxy.Rotation[InternalIndex] = Particles.R(ExternalIndex);
+					Proxy.NextPosition[InternalIndex] = ParticleTransform.GetTranslation();
+					Proxy.NextRotation[InternalIndex] = ParticleTransform.GetRotation();
+				}
+			});
+		}
+		else if (FirstFrame && SecondFrame)
+		{
+			// Both frames valid, second frame has all the indices we need
+			const int32 NumActives = SecondFrame->TransformIndices.Num();
+
+			const float Alpha = (PlaybackTime - FirstFrame->Timestamp) / (SecondFrame->Timestamp - FirstFrame->Timestamp);
+			check(0 <= Alpha && Alpha <= 1.0f);
+
+			Chaos::PhysicsParallelFor(NumActives, [&](int32 Index)
+			{
+				const int32 InternalIndex = SecondFrame->TransformIndices[Index];
+				if (InternalIndex >= RigidBodyID.Num() || InternalIndex < 0)
+				{
+					UE_LOG(UGCC_LOG, Error, 
+						TEXT("%s: Cache index %d out of range: [%d, %d).  Regenerate the cache."),
+						*Parameters.Name, InternalIndex, 0, RigidBodyID.Num());
+					return;
+				}
+				const int32 ExternalIndex = RigidBodyID[InternalIndex];
+				const int32 PreviousIndexSlot = Index < SecondFrame->PreviousTransformIndices.Num() ? SecondFrame->PreviousTransformIndices[Index] : INDEX_NONE;
+
+				if (ExternalIndex != INDEX_NONE && Particles.InvM(ExternalIndex) == 0.0f && !Particles.Disabled(ExternalIndex))
+				{
+					if (PreviousIndexSlot != INDEX_NONE)
 					{
-						if (!AllConstraintsArray[Idx].AccumulatedImpulse.IsZero())
+						Proxy.Position[InternalIndex] = Proxy.NextPosition[InternalIndex];
+						Proxy.Rotation[InternalIndex] = Proxy.NextRotation[InternalIndex];
+
+						FTransform BlendedTM;
+						BlendedTM.Blend(FirstFrame->Transforms[PreviousIndexSlot], SecondFrame->Transforms[Index], Alpha);
+
+						Proxy.NextPosition[InternalIndex] = BlendedTM.GetTranslation();
+						Proxy.NextRotation[InternalIndex] = BlendedTM.GetRotation();
+					}
+					else
+					{
+						// NewActive case
+						Proxy.Position[InternalIndex] = Proxy.NextPosition[InternalIndex];
+						Proxy.Rotation[InternalIndex] = Proxy.NextRotation[InternalIndex];
+
+						FTransform BlendedTM;
+						BlendedTM.Blend(FTransform(Particles.R(ExternalIndex), Particles.X(ExternalIndex), FVector::OneVector), SecondFrame->Transforms[Index], Alpha);
+
+						Proxy.NextPosition[InternalIndex] = BlendedTM.GetTranslation();
+						Proxy.NextRotation[InternalIndex] = BlendedTM.GetRotation();
+					}
+				}
+			});
+			// #BGallagher Handle new inactives. If it's a cluster parent and it's fully disabled we'll need to decluster it here.
+		}
+#endif
+	}
+}
+
+void FGeometryCollectionPhysicsProxy::StartFrameCallback(const float Dt, const float Time)
+{
+	SCOPE_CYCLE_COUNTER(STAT_GeomBeginFrame);
+	if (InitializedState == ESimulationInitializationState::Initialized)
+	{
+		// Reverse playback only plays back what we just recorded.  So, the condition 
+		// is, "are we in record mode, but ready to play back what we've recorded?"
+		const bool bIsReverseCachePlaying = Parameters.IsCacheRecording() && Parameters.ReverseCacheBeginTime != 0 && Time > Parameters.ReverseCacheBeginTime;
+		if (Parameters.IsCachePlaying() || bIsReverseCachePlaying)
+		{
+#if TODO_REIMPLEMENT_GET_RIGID_PARTICLES
+			// Update the enabled/disabled state for kinematic particles for the upcoming frame
+			Chaos::FPhysicsSolver* ThisSolver = GetSolver();
+			Chaos::FPhysicsSolver::FParticlesType& Particles = ThisSolver->GetRigidParticles();
+
+			if (!Parameters.RecordedTrack)
+			{
+				if (!ensure(Parameters.CacheType == EGeometryCollectionCacheType::Record))
+				{
+					return;
+				}
+				Parameters.RecordedTrack = new FRecordedTransformTrack();
+				Parameters.bOwnsTrack = true;
+			}
+			if (Parameters.bClearCache && bIsReverseCachePlaying)
+			{
+				check(!CommitRecordedStateCallback);
+				*const_cast<FRecordedTransformTrack*>(Parameters.RecordedTrack) = FRecordedTransformTrack::ProcessRawRecordedData(RecordedTracks);
+				Parameters.bClearCache = false;
+			}
+
+			bool bParticlesUpdated = false;
+
+			const float ThisFrameTime = bIsReverseCachePlaying ? (Parameters.RecordedTrack->GetLastTime() - Time + Parameters.ReverseCacheBeginTime) : Time;
+			if (!Parameters.RecordedTrack->IsTimeValid(ThisFrameTime))
+			{
+				// Invalid cache time, nothing to update
+				return;
+			}
+
+			FCriticalSection ParticleUpdateLock;
+			const int32 NumMappings = RigidBodyID.Num();
+			Chaos::PhysicsParallelFor(NumMappings, [&](int32 InternalParticleIndex)
+				//for(int32 Index = 0; Index < NumMappings; ++Index)
+			{
+				const int32 ExternalParticleIndex = RigidBodyID[InternalParticleIndex];
+
+				if (ExternalParticleIndex == INDEX_NONE)
+				{
+					return;
+				}
+
+				if (Particles.InvM(ExternalParticleIndex) != 0)
+				{
+					return;
+				}
+
+				// We need to check a window of Now - Dt to Now and see if we ever activated in that time.
+				// This is for short activations because if we miss one then the playback will be incorrect
+				const bool bShouldBeDisabled = !Parameters.RecordedTrack->GetWasActiveInWindow(InternalParticleIndex, ThisFrameTime, bIsReverseCachePlaying ? (ThisFrameTime - Dt) : (ThisFrameTime + Dt));
+				const bool bDisabledNow = Particles.Disabled(ExternalParticleIndex);
+				if (bDisabledNow != bShouldBeDisabled)
+				{
+					Particles.SetDisabledLowLevel(ExternalParticleIndex, bShouldBeDisabled);
+					ParticleUpdateLock.Lock();
+					if (!bShouldBeDisabled)
+					{
+						Particles.SetObjectState(ExternalParticleIndex, Chaos::EObjectStateType::Kinematic);
+						if (!ensure(ThisSolver->ActiveIndices().Find(ExternalParticleIndex) == nullptr))
 						{
-							if (ensure(!AllConstraintsArray[Idx].Location.ContainsNaN() &&
-								!AllConstraintsArray[Idx].Normal.ContainsNaN()) &&
-								!Particles.V(AllConstraintsArray[Idx].ParticleIndex).ContainsNaN() &&
-								!Particles.V(AllConstraintsArray[Idx].LevelsetIndex).ContainsNaN() &&
-								!Particles.W(AllConstraintsArray[Idx].ParticleIndex).ContainsNaN() &&
-								!Particles.W(AllConstraintsArray[Idx].LevelsetIndex).ContainsNaN())
+							bParticlesUpdated = true;
+						}
+						else if(!bParticlesUpdated)
+						{
+							ThisSolver->NonDisabledIndices().Add(ExternalParticleIndex);
+							ThisSolver->ActiveIndices().Add(ExternalParticleIndex);
+						}
+					}
+					else
+					{
+						if (!ensure(ThisSolver->ActiveIndices().Find(ExternalParticleIndex) != nullptr))
+						{
+							bParticlesUpdated = true;
+						}
+						else if(!bParticlesUpdated)
+						{
+							ThisSolver->NonDisabledIndices().Remove(ExternalParticleIndex);
+							ThisSolver->ActiveIndices().Remove(ExternalParticleIndex);
+						}
+					}
+					ParticleUpdateLock.Unlock();
+				}
+
+			});
+
+			// Do not add collisions if reverse
+			if (!bIsReverseCachePlaying)
+			{
+				const FRecordedFrame* RecordedFrame = Parameters.RecordedTrack->FindRecordedFrame(ThisFrameTime);
+				if (RecordedFrame == nullptr)
+				{
+					const int32 Index = Parameters.RecordedTrack->FindLastKeyBefore(ThisFrameTime);
+					if (Index > 0 && Index < Parameters.RecordedTrack->Records.Num())
+					{
+						RecordedFrame = &Parameters.RecordedTrack->Records[Index];
+					}
+				}
+
+				if (RecordedFrame)
+				{
+					// Collisions
+					if (Parameters.CollisionData.DoGenerateCollisionData && ThisFrameTime > 0.f && Parameters.CollisionData.CollisionDataSizeMax > 0)
+					{
+						if (RecordedFrame->Collisions.Num() > 0)
+						{
+							Chaos::FPhysicsSolver::FCollisionDataArray& AllCollisionsDataArray = ThisSolver->GetAllCollisionsDataArray();
+							TMap<IPhysicsProxyBase*, TArray<int32>>& AllCollisionsIndicesByPhysicsProxy = ThisSolver->GetAllCollisionsIndicesByPhysicsProxy();
+
+							if (!AllCollisionsIndicesByPhysicsProxy.Contains(this))
 							{
-								BoundingBox += AllConstraintsArray[Idx].Location;
+								AllCollisionsIndicesByPhysicsProxy.Add(this, TArray<int32>());
+							}
 
-								ConstraintsArray[NumConstraints] = AllConstraintsArray[Idx];
-								NumConstraints++;
+							for (int32 Idx = 0; Idx < RecordedFrame->Collisions.Num(); ++Idx)
+							{
+								// Check if the particle is still kinematic
+								if (RecordedFrame->Collisions[Idx].ParticleIndex < 0 ||
+									(RecordedFrame->Collisions[Idx].ParticleIndex >= 0 && RecordedFrame->Collisions[Idx].ParticleIndex < static_cast<int32>(Particles.Size()) && Particles.ObjectState(RecordedFrame->Collisions[Idx].ParticleIndex) == Chaos::EObjectStateType::Kinematic))
+								{
+									const int32 NewIdx = AllCollisionsDataArray.Add(Chaos::TCollisionData<float, 3>());
+									Chaos::TCollisionData<float, 3>& AllCollisionsDataArrayItem = AllCollisionsDataArray[NewIdx];
+
+									AllCollisionsDataArrayItem.Location = RecordedFrame->Collisions[Idx].Location;
+									AllCollisionsDataArrayItem.AccumulatedImpulse = RecordedFrame->Collisions[Idx].AccumulatedImpulse;
+									AllCollisionsDataArrayItem.Normal = RecordedFrame->Collisions[Idx].Normal;
+									AllCollisionsDataArrayItem.Velocity1 = RecordedFrame->Collisions[Idx].Velocity1;
+									AllCollisionsDataArrayItem.Velocity2 = RecordedFrame->Collisions[Idx].Velocity2;
+									AllCollisionsDataArrayItem.AngularVelocity1 = RecordedFrame->Collisions[Idx].AngularVelocity1;
+									AllCollisionsDataArrayItem.AngularVelocity2 = RecordedFrame->Collisions[Idx].AngularVelocity2;
+									AllCollisionsDataArrayItem.Mass1 = RecordedFrame->Collisions[Idx].Mass1;
+									AllCollisionsDataArrayItem.Mass2 = RecordedFrame->Collisions[Idx].Mass2;
+#if TODO_CONVERT_GEOMETRY_COLLECTION_PARTICLE_INDICES_TO_PARTICLE_POINTERS
+									AllCollisionsDataArrayItem.ParticleIndex = RecordedFrame->Collisions[Idx].ParticleIndex;
+#endif
+									AllCollisionsDataArrayItem.LevelsetIndex = RecordedFrame->Collisions[Idx].LevelsetIndex;
+									AllCollisionsDataArrayItem.ParticleIndexMesh = RecordedFrame->Collisions[Idx].ParticleIndexMesh;
+									AllCollisionsDataArrayItem.LevelsetIndexMesh = RecordedFrame->Collisions[Idx].LevelsetIndexMesh;
+
+									AllCollisionsIndicesByPhysicsProxy[this].Add(NewIdx);
+								}
+							}
+						}
+					}
+
+					// Breaking
+					if (Parameters.BreakingData.DoGenerateBreakingData && ThisFrameTime > 0.f && Parameters.BreakingData.BreakingDataSizeMax > 0)
+					{
+						if (RecordedFrame->Breakings.Num() > 0)
+						{
+							Chaos::FPhysicsSolver::FBreakingDataArray& AllBreakingsDataArray = ThisSolver->GetAllBreakingsDataArray();
+							TMap<IPhysicsProxyBase*, TArray<int32>>& AllBreakingsIndicesByPhysicsProxy = ThisSolver->GetAllBreakingsIndicesByPhysicsProxy();
+
+							if (!AllBreakingsIndicesByPhysicsProxy.Contains(this))
+							{
+								AllBreakingsIndicesByPhysicsProxy.Add(this, TArray<int32>());
+							}
+
+							for (int32 Idx = 0; Idx < RecordedFrame->Breakings.Num(); ++Idx)
+							{
+								// Check if the particle is still kinematic
+								if (RecordedFrame->Breakings[Idx].ParticleIndex < 0 ||
+									(RecordedFrame->Breakings[Idx].ParticleIndex >= 0 && RecordedFrame->Breakings[Idx].ParticleIndex < static_cast<int32>(Particles.Size()) && Particles.ObjectState(RecordedFrame->Breakings[Idx].ParticleIndex) == Chaos::EObjectStateType::Kinematic))
+								{
+									const int32 NewIdx = AllBreakingsDataArray.Add(Chaos::TBreakingData<float, 3>());
+									Chaos::TBreakingData<float, 3>& AllBreakingsDataArrayItem = AllBreakingsDataArray[NewIdx];
+
+									AllBreakingsDataArrayItem.Location = RecordedFrame->Breakings[Idx].Location;
+									AllBreakingsDataArrayItem.Velocity = RecordedFrame->Breakings[Idx].Velocity;
+									AllBreakingsDataArrayItem.AngularVelocity = RecordedFrame->Breakings[Idx].AngularVelocity;
+									AllBreakingsDataArrayItem.Mass = RecordedFrame->Breakings[Idx].Mass;
+#if TODO_CONVERT_GEOMETRY_COLLECTION_PARTICLE_INDICES_TO_PARTICLE_POINTERS
+									AllBreakingsDataArrayItem.ParticleIndex = RecordedFrame->Breakings[Idx].ParticleIndex;
+#endif
+									AllBreakingsDataArrayItem.ParticleIndexMesh = RecordedFrame->Breakings[Idx].ParticleIndexMesh;
+
+									AllBreakingsIndicesByPhysicsProxy[this].Add(NewIdx);
+								}
+							}
+						}
+					}
+
+					// Trailing
+					if (Parameters.TrailingData.DoGenerateTrailingData && ThisFrameTime > 0.f && Parameters.TrailingData.TrailingDataSizeMax > 0)
+					{
+						if (RecordedFrame->Trailings.Num() > 0)
+						{
+							Chaos::FPhysicsSolver::FTrailingDataArray& AllTrailingsDataArray = ThisSolver->GetAllTrailingsDataArray();
+							TMap<IPhysicsProxyBase*, TArray<int32>>& AllTrailingsIndicesByPhysicsProxy = ThisSolver->GetAllTrailingsIndicesByPhysicsProxy();
+
+							if (!AllTrailingsIndicesByPhysicsProxy.Contains(this))
+							{
+								AllTrailingsIndicesByPhysicsProxy.Add(this, TArray<int32>());
+							}
+
+							for (const FSolverTrailingData& Trailing : RecordedFrame->Trailings)
+							{
+								// Check if the particle is still kinematic
+								if (Trailing.ParticleIndex < 0 ||
+									(Trailing.ParticleIndex >= 0 && Trailing.ParticleIndex < static_cast<int32>(Particles.Size()) && Particles.ObjectState(Trailing.ParticleIndex) == Chaos::EObjectStateType::Kinematic))
+								{
+									const int32 NewIdx = AllTrailingsDataArray.Add(Chaos::TTrailingData<float, 3>());
+									Chaos::TTrailingData<float, 3>& AllTrailingsDataArrayItem = AllTrailingsDataArray[NewIdx];
+
+									AllTrailingsDataArrayItem.Location = Trailing.Location;
+									AllTrailingsDataArrayItem.Velocity = Trailing.Velocity;
+									AllTrailingsDataArrayItem.AngularVelocity = Trailing.AngularVelocity;
+									AllTrailingsDataArrayItem.Mass = Trailing.Mass;
+#if TODO_CONVERT_GEOMETRY_COLLECTION_PARTICLE_INDICES_TO_PARTICLE_POINTERS
+									AllTrailingsDataArrayItem.ParticleIndex = Trailing.ParticleIndex;
+#endif
+									AllTrailingsDataArrayItem.ParticleIndexMesh = Trailing.ParticleIndexMesh;
+
+									AllTrailingsIndicesByPhysicsProxy[this].Add(NewIdx);
+								}
 							}
 						}
 					}
 				}
 			}
 
-			ConstraintsArray.SetNum(NumConstraints);
-
-			if (ConstraintsArray.Num() > 0)
+			if (bParticlesUpdated)
 			{
-				if (Parameters.CollisionData.DoCollisionDataSpatialHash &&
-					Parameters.CollisionData.CollisionDataSpatialHashRadius > 0.f &&
-					ConstraintsArray.Num() > 1 &&
-					(BoundingBox.GetExtent().X > 0.f || BoundingBox.GetExtent().Y > 0.f || BoundingBox.GetExtent().Z > 0.f))
-				{
-					// Validate Parameters.CollisionData.CollisionDataSpatialHashRadius
-					// CellSize must be smaller than the smallest bbox extent
-					float SpatialHashRadius = Parameters.CollisionData.CollisionDataSpatialHashRadius;
-					TArray<float> ExtentArray = { BoundingBox.GetExtent().X , BoundingBox.GetExtent().Y, BoundingBox.GetExtent().Z };
-					ExtentArray.Sort();
-					if (ExtentArray[0] != 0.f)
-					{
-						if (2.f * SpatialHashRadius > ExtentArray[0])
-						{
-							SpatialHashRadius = 0.5f * ExtentArray[0];
-						}
-					}
-					else
-					{
-						if (ExtentArray[1] != 0.f)
-						{
-							if (2.f * SpatialHashRadius > ExtentArray[1])
-							{
-								SpatialHashRadius = 0.5f * ExtentArray[1];
-							}
-						}
-						else
-						{
-							if (2.f * SpatialHashRadius > ExtentArray[2])
-							{
-								SpatialHashRadius = 0.5f * ExtentArray[2];
-							}
-						}
-					}
-
-					// Spatial hash the constraints
-					TMultiMap<int32, int32> HashTableMap;
-					ComputeHashTable(ConstraintsArray, BoundingBox, HashTableMap, SpatialHashRadius);
-
-					TArray<int32> UsedCellsArray;
-					HashTableMap.GetKeys(UsedCellsArray);
-
-					int32 NumCollisionsThisFrame = 0;
-					for (int32 IdxCell = 0; IdxCell < UsedCellsArray.Num(); ++IdxCell)
-					{
-						TArray<int32> ConstraintsInCellArray;
-						HashTableMap.MultiFind(UsedCellsArray[IdxCell], ConstraintsInCellArray);
-
-						int32 NumConstraintsToGetFromCell = FMath::Min(Parameters.CollisionData.MaxCollisionPerCell, ConstraintsInCellArray.Num());
-						for (int32 IdxConstraint = 0; IdxConstraint < NumConstraintsToGetFromCell; ++IdxConstraint)
-						{
-							AddCollisionToCollisionData(ExistingFrame, Particles, ConstraintsArray[ConstraintsInCellArray[IdxConstraint]]);
-						}
-					}
-
-					if (ExistingFrame->Collisions.Num() > Parameters.CollisionData.CollisionDataSizeMax)
-					{
-						TArray<FSolverCollisionData> CollisionsArray1;
-
-						float FInc = (float)ExistingFrame->Collisions.Num() / (float)Parameters.CollisionData.CollisionDataSizeMax;
-
-						CollisionsArray1.SetNumUninitialized(Parameters.CollisionData.CollisionDataSizeMax);
-						for (int32 IdxCollision = 0; IdxCollision < Parameters.CollisionData.CollisionDataSizeMax; ++IdxCollision)
-						{
-							int32 NewIdx = FMath::FloorToInt((float)IdxCollision * FInc);
-							CollisionsArray1[IdxCollision] = ExistingFrame->Collisions[NewIdx];
-						}
-
-						ExistingFrame->Collisions.SetNumUninitialized(Parameters.CollisionData.CollisionDataSizeMax);
-						for (int32 IdxCollision = 0; IdxCollision < Parameters.CollisionData.CollisionDataSizeMax; ++IdxCollision)
-						{
-							ExistingFrame->Collisions[IdxCollision] = CollisionsArray1[IdxCollision];
-						}
-					}
-				}
-				else
-				{
-					if (ConstraintsArray.Num() <= Parameters.CollisionData.CollisionDataSizeMax)
-					{
-						for (int32 IdxConstraint = 0; IdxConstraint < ConstraintsArray.Num(); ++IdxConstraint)
-						{
-							AddCollisionToCollisionData(ExistingFrame, Particles, ConstraintsArray[IdxConstraint]);
-						}
-					}
-					else
-					{
-						float FInc = (float)ConstraintsArray.Num() / (float)Parameters.CollisionData.CollisionDataSizeMax;
-
-						for (int32 IdxConstraint = 0; IdxConstraint < Parameters.CollisionData.CollisionDataSizeMax; ++IdxConstraint)
-						{
-							int32 Idx = FMath::FloorToInt((float)IdxConstraint * FInc);
-
-							AddCollisionToCollisionData(ExistingFrame, Particles, ConstraintsArray[Idx]);
-						}
-					}
-				}
+				ThisSolver->InitializeFromParticleData(0);
 			}
+#endif
 		}
-#endif
 	}
 }
 
-void FGeometryCollectionPhysicsProxy::AddBreakingToBreakingData(FRecordedFrame* ExistingFrame, const FParticlesType& Particles, const Chaos::TBreakingData<float, 3>& Breaking)
+void FGeometryCollectionPhysicsProxy::EndFrameCallback(const float EndFrame)
 {
-#if TODO_REIMPLEMENT_RIGID_CLUSTERING
-	const Chaos::TArrayCollectionArray<Chaos::ClusterId>& ClusterIdsArray = GetSolver()->GetRigidClustering().GetClusterIdsArray();
-	const TMap<uint32, TUniquePtr<TArray<uint32>>>& ParentToChildrenMap = GetSolver()->GetRigidClustering().GetChildrenMap();
-	const Chaos::TArrayCollectionArray<PhysicsProxyWrapper>& PhysicsProxyReverseMapping = GetSolver()->GetPhysicsProxyReverseMapping();
+	check(false); // not currently called
+	if (InitializedState == ESimulationInitializationState::Initialized)
+	{
+		FGeometryDynamicCollection* Collection = Parameters.DynamicCollection;
+		check(Collection);
 
-	int32 NewIdx = ExistingFrame->Breakings.Add(FSolverBreakingData());
-	FSolverBreakingData& NewBreaking = ExistingFrame->Breakings[NewIdx];
+		ProxySimDuration += EndFrame;
 
-	NewBreaking.Location = Breaking.Location;
-	NewBreaking.Velocity = Particles.V(Breaking.ParticleIndex);
-	NewBreaking.AngularVelocity = Particles.W(Breaking.ParticleIndex);
-	NewBreaking.Mass = Particles.M(Breaking.ParticleIndex);
-	NewBreaking.ParticleIndex = Breaking.ParticleIndex;
+		if(Collection->HasAttribute("RigidBodyID", FGeometryCollection::TransformGroup))
+		{
+#if TODO_REIMPLEMENT_GET_RIGID_PARTICLES
+			//
+			//  Update transforms for the simulated transforms
+			//
+			TManagedArray<int32>& CollectionClusterID = SolverClusterID;
+			TManagedArray<FTransform>& Transform = Collection->Transform;
+			TManagedArray<int32>& Parent = Collection->Parent;
+			TManagedArray<TSet<int32>>& Children = Collection->Children;
+			TManagedArray<int32>& SimulationType = Collection->SimulationType;
 
-	// If Breaking.ParticleIndex is a cluster store an index for a mesh in this cluster
-	if (ClusterIdsArray[Breaking.ParticleIndex].NumChildren > 0)
-	{
-		int32 ParticleIndexMesh = GetSolver()->GetParticleIndexMesh(ParentToChildrenMap, Breaking.ParticleIndex);
-		ensure(ParticleIndexMesh != INDEX_NONE);
-		NewBreaking.ParticleIndexMesh = ParticleIndexMesh;
-	}
-#endif
-}
+			TManagedArray<int32>& DynamicState = Collection->GetAttribute<int32>("DynamicState", FGeometryCollection::TransformGroup);
 
-void FGeometryCollectionPhysicsProxy::UpdateBreakingData(const FParticlesType& Particles, FRecordedFrame* ExistingFrame)
-{
-	ExistingFrame->Breakings.Reset();
+			Chaos::TPBDRigidParticles<float, 3>& Particles = GetSolver()->GetRigidParticles();
+			const Chaos::TPBDCollisionConstraints<float, 3>& CollisionRule = GetSolver()->GetCollisionConstraints();
+			const Chaos::TArrayCollectionArray<Chaos::ClusterId>& ClusterID = GetSolver()->GetRigidClustering().GetClusterIdsArray();
+			const Chaos::TArrayCollectionArray<Chaos::TRigidTransform<float, 3>>& ClusterChildToParentMap = GetSolver()->GetRigidClustering().GetChildToParentMap();
+			const Chaos::TArrayCollectionArray<bool>& InternalCluster = GetSolver()->GetRigidClustering().GetInternalClusterArray();
 
-	if (Parameters.BreakingData.SaveBreakingData && ExistingFrame->Timestamp > 0.f && Parameters.BreakingData.BreakingDataSizeMax > 0)
-	{
-#if TODO_REIMPLEMENT_GETALLCLUSTERBREAKINGS
-		const TArray<Chaos::TBreakingData<float, 3>>& AllBreakingsArray = GetSolver()->GetAllClusterBreakings();
-		if (AllBreakingsArray.Num() > 0)
-		{
-			const Chaos::TArrayCollectionArray<PhysicsProxyWrapper>& PhysicsProxyReverseMapping = GetSolver()->GetPhysicsProxyReverseMapping();
+			//Particles X and R are aligned with center of mass and inertia principal axes.
+			//Renderer doesn't know about this and simply does ActorToWorld * GeomToActor * LocalSpaceVerts
+			//In proper math multiplication order:
+			//ParticleToWorld = ActorToWorld * GeomToActor * MassToLocal
+			//GeomToWorld = ActorToWorld * GeomToActor
+			//=> GeomToWorld = ParticleToWorld * MassToLocal.Inv()
+			//=> GeomToActor = ActorToWorld.Inv() * ParticleToWorld * MassToLocal.Inv()
+			const int32 TransformSize = Collection->NumElements(FGeometryCollection::TransformGroup);
+			const FTransform& ActorToWorld = Parameters.WorldTransform;
 
-			TArray<Chaos::TBreakingData<float, 3>> BreakingsArray;
-			BreakingsArray.SetNumUninitialized(AllBreakingsArray.Num());
-
-			FBox BoundingBox(ForceInitToZero);
-			int32 NumBreakings = 0;
-			for (int32 Idx = 0; Idx < AllBreakingsArray.Num(); ++Idx)
+			// It's not worth shrinking the EndFrameUnparentingBuffer array, at least until the solver supports deleting bodies.
+			if (EndFrameUnparentingBuffer.Num() < TransformSize)
 			{
-				// Check if the breaking is for this PhysicsProxy
-				IPhysicsProxyBase* PhysicsProxy = PhysicsProxyReverseMapping[AllBreakingsArray[Idx].ParticleIndex].PhysicsProxy;
-				if (PhysicsProxy == this)
-				{
-					if (ensure(!AllBreakingsArray[Idx].Location.ContainsNaN() &&
-						!Particles.V(AllBreakingsArray[Idx].ParticleIndex).ContainsNaN() &&
-						!Particles.W(AllBreakingsArray[Idx].ParticleIndex).ContainsNaN()))
-					{
-						BoundingBox += AllBreakingsArray[Idx].Location;
-
-						BreakingsArray[NumBreakings] = AllBreakingsArray[Idx];
-						NumBreakings++;
-					}
-				}
+				EndFrameUnparentingBuffer.Init(-1, TransformSize);
 			}
 
-			BreakingsArray.SetNum(NumBreakings);
+			{
+				for (int32 TransformGroupIndex = 0; TransformGroupIndex < TransformSize; ++TransformGroupIndex)
+				{
+					const int32 RigidBodyIndex = RigidBodyID[TransformGroupIndex];
+					if (RigidBodyIndex != INDEX_NONE)
+					{
+						// Update the transform and parent hierarchy of the active rigid bodies. Active bodies can be either
+						// rigid geometry defined from the leaf nodes of the collection, or cluster bodies that drive an entire
+						// branch of the hierarchy within the GeometryCollection.
+						// - Active bodies are directly driven from the global position of the corresponding
+						//   rigid bodies within the solver ( cases where RigidBodyID[TransformGroupIndex] is not disabled ). 
+						// - Deactivated bodies are driven from the transforms of their active parents. However the solver can
+						//   take ownership of the parents during the simulation, so it might be necessary to force deactivated
+						//   bodies out of the collections hierarchy during the simulation.  
+						if (!Particles.Disabled(RigidBodyID[TransformGroupIndex]))
+						{
+							// Update the transform of the active body. The active body can be either a single rigid
+							// or a collection of rigidly attached geometries (Clustering). The cluster is represented as a
+							// single transform in the GeometryCollection, and all children are stored in the local space
+							// of the parent cluster.
+							// ... When setting cluster transforms it is expected that the MassToLocal is identity.
+							//     Cluster initialization will set the vertices in the MassSpace of the rigid body.
+							// ... When setting individual rigid bodies that are not clustered, the MassToLocal will be 
+							//     non-Identity, and will reflect the difference between the geometric center of the geometry
+							//     and that corresponding rigid bodies center of mass. 
+							const FTransform ParticleToWorld(Particles.R(RigidBodyIndex), Particles.X(RigidBodyIndex));
+							// GeomToActor = ActorToWorld.Inv() * ParticleToWorld * MassToLocal.Inv();
+							Transform[TransformGroupIndex] = MassToLocal[TransformGroupIndex].GetRelativeTransformReverse(ParticleToWorld).GetRelativeTransform(ActorToWorld);
+							Transform[TransformGroupIndex].NormalizeRotation();
 
-			if (BreakingsArray.Num() > 0)
-			{
-				if (Parameters.BreakingData.DoBreakingDataSpatialHash &&
-					Parameters.BreakingData.BreakingDataSpatialHashRadius > 0.f &&
-					BreakingsArray.Num() > 1 &&
-					(BoundingBox.GetExtent().X > 0.f || BoundingBox.GetExtent().Y > 0.f || BoundingBox.GetExtent().Z > 0.f))
-				{
-					// Validate Parameters.CollisionData.CollisionDataSpatialHashRadius
-					// CellSize must be smaller than the smallest bbox extent
-					float SpatialHashRadius = Parameters.BreakingData.BreakingDataSpatialHashRadius;
-					TArray<float> ExtentArray = { BoundingBox.GetExtent().X , BoundingBox.GetExtent().Y, BoundingBox.GetExtent().Z };
-					ExtentArray.Sort();
-					if (ExtentArray[0] != 0.f)
-					{
-						if (2.f * SpatialHashRadius > ExtentArray[0])
+							// dynamic state is also updated by the solver during field interaction. 
+							if (!Particles.Sleeping(RigidBodyIndex))
+							{
+								if (Particles.ObjectState(RigidBodyIndex) == Chaos::EObjectStateType::Kinematic)
+								{
+									DynamicState[TransformGroupIndex] = (int)EObjectStateTypeEnum::Chaos_Object_Kinematic;
+								}
+								else if (Particles.ObjectState(RigidBodyIndex) == Chaos::EObjectStateType::Static)
+								{
+									DynamicState[TransformGroupIndex] = (int)EObjectStateTypeEnum::Chaos_Object_Static;
+								}
+								else
+								{
+									DynamicState[TransformGroupIndex] = (int)EObjectStateTypeEnum::Chaos_Object_Dynamic;
+								}
+							}
+
+							// Force all enabled rigid bodies out of the transform hierarchy
+							if (Parent[TransformGroupIndex] != INDEX_NONE)
+							{
+								const int32 ParentIndex = Parent[TransformGroupIndex];
+								// Children in the hierarchy are stored in a TSet, which is not thread safe.  So we retain
+								// indices to remove afterwards.
+								EndFrameUnparentingBuffer[TransformGroupIndex] = ParentIndex;
+							}
+
+							// When a leaf node rigid body is removed from a cluster the rigid
+							// body will become active and needs its clusterID updated. This just
+							// syncs the clusterID all the time. 
+							CollectionClusterID[TransformGroupIndex] = ClusterID[RigidBodyIndex].Id;
+						}
+						else if (Particles.Disabled(RigidBodyIndex))
 						{
-							SpatialHashRadius = 0.5f * ExtentArray[0];
-						}
-					}
-					else
-					{
-						if (ExtentArray[1] != 0.f)
-						{
-							if (2.f * SpatialHashRadius > ExtentArray[1])
+							// dynamic state is also updated by the solver during field interaction. 
+							if (!Particles.Sleeping(RigidBodyIndex))
 							{
-								SpatialHashRadius = 0.5f * ExtentArray[1];
+								if (Particles.ObjectState(RigidBodyIndex) == Chaos::EObjectStateType::Kinematic)
+								{
+									DynamicState[TransformGroupIndex] = (int)EObjectStateTypeEnum::Chaos_Object_Kinematic;
+								}
+								else if (Particles.ObjectState(RigidBodyIndex) == Chaos::EObjectStateType::Static)
+								{
+									DynamicState[TransformGroupIndex] = (int)EObjectStateTypeEnum::Chaos_Object_Static;
+								}
+								else
+								{
+									DynamicState[TransformGroupIndex] = (int)EObjectStateTypeEnum::Chaos_Object_Dynamic;
+								}
 							}
-						}
-						else
-						{
-							if (2.f * SpatialHashRadius > ExtentArray[2])
+
+							// The rigid body parent cluster has changed within the solver, and its
+							// parent body is not tracked within the geometry collection. So we need to
+							// pull the rigid bodies out of the transform hierarchy, and just drive
+							// the positions directly from the solvers cluster particle. 
+							if (CollectionClusterID[TransformGroupIndex] != ClusterID[RigidBodyIndex].Id)
 							{
-								SpatialHashRadius = 0.5f * ExtentArray[2];
+								// Force all driven rigid bodies out of the transform hierarchy
+								if (Parent[TransformGroupIndex] != INDEX_NONE)
+								{
+									const int32 ParentIndex = Parent[TransformGroupIndex];
+									// Children in the hierarchy are stored in a TSet, which is not thread safe.  So we retain
+									// indices to remove afterwards.
+									EndFrameUnparentingBuffer[TransformGroupIndex] = ParentIndex;
+								}
+								CollectionClusterID[TransformGroupIndex] = ClusterID[RigidBodyIndex].Id;
 							}
-						}
-					}
 
-					// Spatial hash the constraints
-					TMultiMap<int32, int32> HashTableMap;
-					ComputeHashTable(BreakingsArray, BoundingBox, HashTableMap, SpatialHashRadius);
-
-					TArray<int32> UsedCellsArray;
-					HashTableMap.GetKeys(UsedCellsArray);
-
-					for (int32 IdxCell = 0; IdxCell < UsedCellsArray.Num(); ++IdxCell)
-					{
-						TArray<int32> BreakingsInCellArray;
-						HashTableMap.MultiFind(UsedCellsArray[IdxCell], BreakingsInCellArray);
-
-						int32 NumBreakingsToGetFromCell = FMath::Min(Parameters.BreakingData.MaxBreakingPerCell, BreakingsInCellArray.Num());
-						for (int32 IdxBreaking = 0; IdxBreaking < NumBreakingsToGetFromCell; ++IdxBreaking)
-						{
-							AddBreakingToBreakingData(ExistingFrame, Particles, BreakingsArray[BreakingsInCellArray[IdxBreaking]]);
-						}
-					}
-
-					if (ExistingFrame->Breakings.Num() > Parameters.BreakingData.BreakingDataSizeMax)
-					{
-						TArray<FSolverBreakingData> BreakingsArray1;
-
-						float FInc = (float)ExistingFrame->Breakings.Num() / (float)Parameters.BreakingData.BreakingDataSizeMax;
-
-						BreakingsArray1.SetNumUninitialized(Parameters.BreakingData.BreakingDataSizeMax);
-						for (int32 IdxBreaking = 0; IdxBreaking < Parameters.BreakingData.BreakingDataSizeMax; ++IdxBreaking)
-						{
-							int32 NewIdx = FMath::FloorToInt((float)IdxBreaking * FInc);
-							BreakingsArray1[IdxBreaking] = ExistingFrame->Breakings[NewIdx];
-						}
-
-						ExistingFrame->Breakings.SetNumUninitialized(Parameters.BreakingData.BreakingDataSizeMax);
-						for (int32 IdxBreaking = 0; IdxBreaking < Parameters.BreakingData.BreakingDataSizeMax; ++IdxBreaking)
-						{
-							ExistingFrame->Breakings[IdxBreaking] = BreakingsArray1[IdxBreaking];
-						}
-					}
-				}
-				else
-				{
-					if (BreakingsArray.Num() <= Parameters.BreakingData.BreakingDataSizeMax)
-					{
-						for (int32 IdxBreaking = 0; IdxBreaking < BreakingsArray.Num(); ++IdxBreaking)
-						{
-							AddBreakingToBreakingData(ExistingFrame, Particles, BreakingsArray[IdxBreaking]);
-						}
-					}
-					else
-					{
-						float FInc = (float)BreakingsArray.Num() / (float)Parameters.BreakingData.BreakingDataSizeMax;
-
-						for (int32 IdxBreaking = 0; IdxBreaking < Parameters.BreakingData.BreakingDataSizeMax; ++IdxBreaking)
-						{
-							int32 Idx = FMath::FloorToInt((float)IdxBreaking * FInc);
-
-							AddBreakingToBreakingData(ExistingFrame, Particles, BreakingsArray[Idx]);
-						}
-					}
-				}
-			}
-		}
-#endif
-	}
-}
-
-void FGeometryCollectionPhysicsProxy::UpdateTrailingData(const FParticlesType& Particles, FRecordedFrame* ExistingFrame)
-{
-	ExistingFrame->Trailings.Reset();
-
-	if (Parameters.TrailingData.SaveTrailingData && ExistingFrame->Timestamp > 0.f && Parameters.TrailingData.TrailingDataSizeMax > 0)
-	{
-		const float TrailingMinSpeedThresholdSquared = Parameters.TrailingData.TrailingMinSpeedThreshold * Parameters.TrailingData.TrailingMinSpeedThreshold;
-
-		// Find previous frame trailing data
-		int32 LastKey = RecordedTracks.FindLastKeyBefore(ExistingFrame->Timestamp);
-		if (LastKey != INDEX_NONE)
-		{
-			FRecordedFrame& PrevFrame = RecordedTracks.Records[LastKey];
-			if (PrevFrame.Trailings.Num() > 0)
-			{
-				for (FSolverTrailingData& Trailing : PrevFrame.Trailings)
-				{
-					if (Particles.ObjectState(Trailing.ParticleIndex) != Chaos::EObjectStateType::Sleeping &&
-						!Particles.Disabled(Trailing.ParticleIndex) &&
-						Particles.V(Trailing.ParticleIndex).SizeSquared() >= TrailingMinSpeedThresholdSquared)
-					{
-						ExistingFrame->Trailings.Add(Trailing);
-					}
-				}
-			}
-		}
-
-		if (ExistingFrame->Trailings.Num() < Parameters.TrailingData.TrailingDataSizeMax)
-		{
-#if TODO_REIMPLEMENT_PHYSICS_PROXY_REVERSE_MAPPING
-			const Chaos::TArrayCollectionArray<PhysicsProxyWrapper>& PhysicsProxyReverseMapping = GetSolver()->GetPhysicsProxyReverseMapping();
-			const Chaos::TArrayCollectionArray<Chaos::ClusterId>& ClusterIdsArray = GetSolver()->GetRigidClustering().GetClusterIdsArray();
-			const TMap<uint32, TUniquePtr<TArray<uint32>>>& ParentToChildrenMap = GetSolver()->GetRigidClustering().GetChildrenMap();
-
-			TArray<FSolverTrailingData> AllTrailings;
-			for (int32 IdxParticle = 0; IdxParticle < (int32)Particles.Size(); ++IdxParticle)
-			{
-				// Check if the particle is for this PhysicsProxy
-				IPhysicsProxyBase* PhysicsProxy = PhysicsProxyReverseMapping[IdxParticle].PhysicsProxy;
-				if (PhysicsProxy == this)
-				{
-					if (ensure(FMath::IsFinite(Particles.InvM(IdxParticle))))
-					{
-						if (Particles.ObjectState(IdxParticle) != Chaos::EObjectStateType::Sleeping &&
-							!Particles.Disabled(IdxParticle) &&
-							Particles.InvM(IdxParticle) != 0.f)
-						{
-							if (Particles.Geometry(IdxParticle) && 
-								Particles.Geometry(IdxParticle)->HasBoundingBox())
+							// Disabled rigid bodies that have valid cluster parents, and have been re-indexed by the
+							// solver ( As in, They were re-clustered outside of the geometry collection), These clusters 
+							// will need to be rendered based on the clusters position. 
+							const int32 ClusterParentIndex = CollectionClusterID[TransformGroupIndex];
+							if (ClusterParentIndex != INDEX_NONE)
 							{
-								if (ensure(!Particles.X(IdxParticle).ContainsNaN()) &&
-									!Particles.V(IdxParticle).ContainsNaN() &&
-									!Particles.W(IdxParticle).ContainsNaN() &&
-									FMath::IsFinite(Particles.M(IdxParticle)))
+								if (InternalCluster[ClusterParentIndex])
 								{
-									Chaos::TAABB<float, 3> BoundingBox = Particles.Geometry(IdxParticle)->BoundingBox();
-									Chaos::TVector<float, 3> Extents = BoundingBox.Extents();
-									float Volume = Extents[0] * Extents[1] * Extents[2];
-									float SpeedSquared = Particles.V(IdxParticle).SizeSquared();
-
-									if (SpeedSquared >= TrailingMinSpeedThresholdSquared &&
-										Volume > Parameters.TrailingData.TrailingMinVolumeThreshold)
+									const FTransform ClusterChildToWorld = ClusterChildToParentMap[RigidBodyIndex] * FTransform(Particles.R(ClusterParentIndex), Particles.X(ClusterParentIndex));
+									if (Parameters.IsCacheRecording())
 									{
-										FSolverTrailingData TrailingData(Particles.X(IdxParticle),
-											Particles.V(IdxParticle),
-											Particles.W(IdxParticle),
-											Particles.M(IdxParticle),
-											IdxParticle,
-											INDEX_NONE);
-
-										// If IdxParticle is a cluster store an index for a mesh in this cluster
-										if (ClusterIdsArray[IdxParticle].NumChildren > 0)
-										{
-											int32 ParticleIndexMesh = GetSolver()->GetParticleIndexMesh(ParentToChildrenMap, IdxParticle);
-											ensure(ParticleIndexMesh != INDEX_NONE);
-											TrailingData.ParticleIndexMesh = ParticleIndexMesh;
-										}
-
-										FSetElementId Id = ExistingFrame->Trailings.FindId(TrailingData);
-										if (!Id.IsValidId())
-										{
-											AllTrailings.Add(TrailingData);
-										}
-										else
-										{
-											ExistingFrame->Trailings[Id].Location = Particles.X(IdxParticle);
-											ExistingFrame->Trailings[Id].Velocity = Particles.V(IdxParticle);
-											ExistingFrame->Trailings[Id].AngularVelocity = Particles.W(IdxParticle);
-										}
+										Particles.X(RigidBodyIndex) = ClusterChildToWorld.GetTranslation();
+										Particles.R(RigidBodyIndex) = ClusterChildToWorld.GetRotation();
 									}
+									// GeomToActor = ActorToWorld.Inv() * ClusterChildToWorld * MassToLocal.Inv();
+									Transform[TransformGroupIndex] = MassToLocal[TransformGroupIndex].GetRelativeTransformReverse(ClusterChildToWorld).GetRelativeTransform(ActorToWorld);
+									Transform[TransformGroupIndex].NormalizeRotation();
 								}
 							}
 						}
 					}
 				}
 			}
+			for (int32 TransformGroupIndex = 0; TransformGroupIndex < TransformSize; TransformGroupIndex++)
+			{
+				const int32 ParentIndex = EndFrameUnparentingBuffer[TransformGroupIndex];
+				if (ParentIndex >= 0)
+				{
+					// We reuse EndFrameUnparentingBuffer potentially without reinitialization, so reset this index to -1 before it gets paged out.
+					EndFrameUnparentingBuffer[TransformGroupIndex] = -1;
 
-			if (AllTrailings.Num() > 0)
+					Children[ParentIndex].Remove(TransformGroupIndex);
+					Parent[TransformGroupIndex] = INDEX_NONE;
+				}
+			}
+
+			//
+			//  Set rest cache on simulated object.
+			//
+			if (Parameters.IsCacheRecording())
 			{
-				int32 NumTrailingsToAdd = Parameters.TrailingData.TrailingDataSizeMax - ExistingFrame->Trailings.Num();
+				check(!UpdateRecordedStateCallback);
+				UpdateRecordedState(ProxySimDuration, RigidBodyID, CollectionClusterID, InternalCluster, Particles, CollisionRule);
+			}
 
-				if (AllTrailings.Num() <= NumTrailingsToAdd)
+			// one way trigger from non-simulating to simulating
+			if (!IsObjectDynamic)
+			{
+				const Chaos::TArrayCollectionArray<bool>& ExternalID = GetSolver()->GetRigidClustering().GetInternalClusterArray();
+
+				for (int32 TransformGroupIndex = 0; TransformGroupIndex < TransformSize; TransformGroupIndex++)
 				{
-					for (int32 IdxTrailing = 0; IdxTrailing < AllTrailings.Num(); ++IdxTrailing)
+					const int32 RigidBodyIndex = RigidBodyID[TransformGroupIndex];
+					if (RigidBodyIndex != INDEX_NONE)
 					{
-						ExistingFrame->Trailings.Add(AllTrailings[IdxTrailing]);
+						bool HasBeenRemoved = (PTDynamicCollection.StatusFlags[TransformGroupIndex] & FGeometryCollection::FS_RemoveOnFracture) && Particles.Disabled(RigidBodyIndex) && ClusterID[RigidBodyIndex].Id == INDEX_NONE;
+
+						if (HasBeenRemoved || !Particles.Disabled(RigidBodyIndex))
+						{
+							if (DynamicState[TransformGroupIndex] != (int32)EObjectStateTypeEnum::Chaos_Object_Static)
+							{
+								IsObjectDynamic = true;
+								break;
+							}
+						}
+						else
+						{
+							const int32 ParentClusterID = ClusterID[RigidBodyIndex].Id;
+							if (ParentClusterID != INDEX_NONE && ExternalID[ParentClusterID]
+								&& Particles.ObjectState(ParentClusterID) != Chaos::EObjectStateType::Static)
+							{
+								IsObjectDynamic = true;
+								break;
+							}
+						}
 					}
 				}
-				else
+			}
+
+			// Can't change visibility right now so setting scale to zero instead
+			// only process if enabled on this object (ones with glass), most will skip over
+			if (Parameters.RemoveOnFractureEnabled && IsObjectDynamic)
+			{
+				for (int32 TransformGroupIndex = 0; TransformGroupIndex < TransformSize; TransformGroupIndex++)
 				{
-					float FInc = (float)AllTrailings.Num() / (float)(NumTrailingsToAdd);
-
-					for (int32 IdxTrailing = 0; IdxTrailing < NumTrailingsToAdd; ++IdxTrailing)
+					const int32 RigidBodyIndex = RigidBodyID[TransformGroupIndex];
+					if (RigidBodyIndex != INDEX_NONE)
 					{
-						int32 Idx = FMath::FloorToInt((float)IdxTrailing * FInc);
-						ExistingFrame->Trailings.Add(AllTrailings[Idx]);
+						if ((PTDynamicCollection.StatusFlags[TransformGroupIndex] & FGeometryCollection::FS_RemoveOnFracture)
+							&& Particles.Disabled(RigidBodyIndex) && ClusterID[RigidBodyIndex].Id == INDEX_NONE)
+						{
+							PTDynamicCollection.Transform[TransformGroupIndex].SetScale3D(FVector::ZeroVector);
+						}
 					}
 				}
 			}
 #endif
 		}
+
+#if TODO_REIMPLEMENT_RIGID_CLUSTERING
+		{
+			IsObjectLoading = GetSolver()->GetRigidClustering().NumberOfPendingClusters() != 0;
+		}
+#endif
 	}
+	Commands.Empty();
 }
 
-void FGeometryCollectionPhysicsProxy::UpdateRecordedState(float SolverTime, const TManagedArray<int32>& InRigidBodyID, const TManagedArray<int32>& InCollectionClusterID, const Chaos::TArrayCollectionArray<bool>& InInternalCluster, const FParticlesType& Particles, const FCollisionConstraintsType& CollisionRule)
+void FGeometryCollectionPhysicsProxy::CreateRigidBodyCallback(FParticlesType& Particles)
 {
-	FRecordedFrame* ExistingFrame = RecordedTracks.FindRecordedFrame(SolverTime);
+	check(false); // This function is deprecated and should never be called.
 
-	if(!ExistingFrame)
-	{
-		ExistingFrame = &InsertRecordedFrame(RecordedTracks, SolverTime);
-	}
-
-	ExistingFrame->Reset(InRigidBodyID.Num());
-	ExistingFrame->Timestamp = SolverTime;
-
-	// Collision
-	UpdateCollisionData(Particles, CollisionRule, ExistingFrame);
-
-	// Breaking
-	UpdateBreakingData(Particles, ExistingFrame);
-
-	// Trailing
-	UpdateTrailingData(Particles, ExistingFrame);
-
-	ParallelFor(InRigidBodyID.Num(), [&](int32 Index)
-	{
-		const int32 ExternalIndex = InRigidBodyID[Index];
-
-		if(ExternalIndex >= 0)
-		{
-			FTransform& NewTransform = ExistingFrame->Transforms[Index];
-
-			NewTransform.SetTranslation(Particles.X(ExternalIndex));
-			NewTransform.SetRotation(Particles.R(ExternalIndex));
-			NewTransform.SetScale3D(FVector(1.0f));
-
-			int32 ClusterParentIndex = InCollectionClusterID[Index];
-			if (ClusterParentIndex != INDEX_NONE && InInternalCluster[ClusterParentIndex])
-			{
-				ExistingFrame->DisabledFlags[Index] = Particles.Disabled(ClusterParentIndex);
-			}
-			else
-			{
-				ExistingFrame->DisabledFlags[Index] = Particles.Disabled(ExternalIndex);
-			}
-		}
-	});
-}
-
-void FGeometryCollectionPhysicsProxy::OnRemoveFromSolver(Chaos::FPBDRigidsSolver *RBDSolver)
-{
-	for (auto* Handle : SolverClusterHandles)
-	{
-		RBDSolver->RemoveParticleToProxy(Handle);
-	}
-}
-
-void FGeometryCollectionPhysicsProxy::OnRemoveFromScene()
-{
-#if TODO_REIMPLEMENT_GET_RIGID_PARTICLES
-	// #BG TODO This isn't great - we currently cannot handle things being removed from the solver.
-	// need to refactor how we handle this and actually remove the particles instead of just constantly
-	// growing the array. Currently everything is just tracked by index though so the solver will have
-	// to notify all the proxies that a chunk of data was removed - or use a sparse array (undesireable)
-	Chaos::FPhysicsSolver::FParticlesType& Particles = GetSolver()->GetRigidParticles();
-
-	// #BG TODO Special case here because right now we reset/realloc the evolution per geom component
-	// in endplay which clears this out. That needs to not happen and be based on world shutdown
-	if(Particles.Size() == 0)
-	{
-		return;
-	}
-
-	const int32 Begin = BaseParticleIndex;
-	const int32 Count = NumParticles;
-
-	if (ensure((int32)Particles.Size() > 0 && (Begin + Count) <= (int32)Particles.Size()))
-	{
-		for (int32 ParticleIndex = 0; ParticleIndex < Count; ++ParticleIndex)
-		{
-			GetSolver()->GetEvolution()->DisableParticle(Begin + ParticleIndex);
-			GetSolver()->GetRigidClustering().GetTopLevelClusterParents().Remove(Begin + ParticleIndex);
-		}
-	}
-#endif
-}
-
-void FGeometryCollectionPhysicsProxy::SyncBeforeDestroy()
-{
-	if(FinalSyncFunc)
-	{
-		FinalSyncFunc(RecordedTracks);
-	}
-}
-
-Chaos::FParticleData* FGeometryCollectionPhysicsProxy::NewData() 
-{
-	/**
-	 * CONTEXT: GAMETHREAD
-	 * Returns a new unmanaged allocation of the data saved on the handle, otherwise nullptr
-	 */
-
-	// The other physics proxies allocate memory and returns a copy of game thread
-	// data as a TGeometryParticleData<T,d>, TKinematicGeometryParticleData<T,d>, 
-	// or a TPBDRigidParticleData<T,d>, which are all derived from FParticleData.
-	// That then gets passed to the physics thread via PushToPhysicsState().
-	//
-	// So all we're doing is pushing game thread state to the physics thread.
-	// I don't think we need to allocate memory to do this.  I think the other 
-	// proxies get away with that because they're dealing with a low number of 
-	// bodies; which is not the case here. Rather, let's just use a triple buffer.
-
-	if (!Parameters.Simulating || 
-		!ensure(NumParticles != INDEX_NONE)) // Make sure InitBodiesGT() has been called!
-	{
-		return nullptr;
-	}
-
-	FGeometryCollectionResults& Buffer = *GameToPhysInterchange->AccessProducerBuffer();
-	if (Buffer.NumParticlesAdded != NumParticles)
-	{
-		Buffer.InitArrays(*GTDynamicCollection);
-
-		// Initialize the buffer
-		Buffer.NumParticlesAdded = NumParticles;
-		Buffer.BaseIndex = BaseParticleIndex;
-
-		Buffer.SharedGeometry.SetNum(NumParticles);
-		Buffer.ShapeSimData.SetNum(NumParticles);
-		Buffer.ShapeQueryData.SetNum(NumParticles);
-
-		// To be done once:
-		int32 Idx = 0;
-		for (auto& GTParticle : GTGeometryParticles)
-		{
-			if (GTParticle)
-			{
-				Buffer.DisabledStates[Idx] = false;
-				Buffer.SharedGeometry[Idx] = GTParticle->GeometrySharedLowLevel();
-
-				for (auto& Shape : GTParticle->ShapesArray())
-				{
-					Buffer.ShapeQueryData[Idx].Add(Shape->QueryData);
-					Buffer.ShapeSimData[Idx].Add(Shape->SimData);
-
-					// Ryan - Hard code everything collides with everything:
-					const int32 Idx1 = Buffer.ShapeSimData[Idx].Num() - 1;
-					Buffer.ShapeSimData[Idx][Idx1].Word1 = 0xFFFF; // this body channel
-					Buffer.ShapeSimData[Idx][Idx1].Word3 = 0xFFFF; // collision candidate channels
-				}
-			}
-			else
-			{
-				Buffer.DisabledStates[Idx] = true;
-				// Buffer.SpatialIdx[Idx] = ???;
-				//Buffer.UniqueIdx[Idx].Idx = INDEX_NONE;
-				//Buffer.UserData[Idx] = nullptr;
-				//Buffer.SharedGeometry[Idx] = nullptr;
-			}
-			++Idx;
-		}
-	}
-
-	const FGeometryDynamicCollection* DynamicCollection = Parameters.DynamicCollection;
 	const FGeometryCollection* RestCollection = Parameters.RestCollection;
+	FGeometryDynamicCollection* DynamicCollection = Parameters.DynamicCollection;
+	check(DynamicCollection);
+	
+	if (Parameters.Simulating && 
+		((InitializedState == ESimulationInitializationState::Unintialized) || 
+		 (InitializedState == ESimulationInitializationState::Activated)))
+	{
+#if TODO_REIMPLEMENT_RIGID_CLUSTERING
+		Chaos::FPhysicsSolver* ThisSolver = GetSolver();
+		Chaos::TArrayCollectionArray<int32>& ClusterGroupIndex = ThisSolver->GetRigidClustering().GetClusterGroupIndexArray();
+		Chaos::TArrayCollectionArray<float>& StrainArray = ThisSolver->GetRigidClustering().GetStrainArray();
+		const float StrainDefault = Parameters.DamageThreshold.Num() ? Parameters.DamageThreshold[0] : 0;
 
-	//Buffer.TransformIndex.Init(RestCollection->TransformIndex);
-	Buffer.BoneMap.Init(RestCollection->BoneMap);
-	Buffer.Parent.Init(RestCollection->Parent);
-	Buffer.Children.Init(RestCollection->Children);
-	Buffer.SimulationType.Init(RestCollection->SimulationType);
+		const TManagedArray<int32>& TransformIndex = RestCollection->TransformIndex;
+		const TManagedArray<int32>& BoneMap = RestCollection->BoneMap;
+		const TManagedArray<int32>& Parent = RestCollection->Parent;
+		const TManagedArray<TSet<int32>>& Children = RestCollection->Children;
+		const TManagedArray<int32>& SimulationType = RestCollection->SimulationType;
+		const TManagedArray<FVector>& Vertex = RestCollection->Vertex;
+		const TManagedArray<int32>& DynamicState = DynamicCollection->DynamicState;
+		const TManagedArray<int32>& CollisionGroup = DynamicCollection->CollisionGroup;
+		const TManagedArray<float>& Mass = RestCollection->GetAttribute<float>("Mass", FTransformCollection::TransformGroup);
+		const TManagedArray<FVector>& InertiaTensor = RestCollection->GetAttribute<FVector>("InertiaTensor", FTransformCollection::TransformGroup);
 
-	Buffer.DynamicState.Init(DynamicCollection->DynamicState);
-	Buffer.Mass.Init(RestCollection->GetAttribute<float>("Mass", FTransformCollection::TransformGroup));
-	Buffer.InertiaTensor.Init(RestCollection->GetAttribute<FVector>("InertiaTensor", FTransformCollection::TransformGroup));
-	Buffer.Transforms.Init(DynamicCollection->Transform);
+		TArray<FTransform> Transform;
+		GeometryCollectionAlgo::GlobalMatrices(DynamicCollection->Transform, DynamicCollection->Parent, Transform);
+		check(DynamicCollection->Transform.Num() == Transform.Num());
 
-	GeometryCollectionAlgo::GlobalMatrices(
-		DynamicCollection->Transform, DynamicCollection->Parent, Buffer.GlobalTransforms);
-	
-	int32 Idx = 0;
-	for (auto& GTParticle : GTGeometryParticles)
-	{
-		if (GTParticle)
+		// count particles to add
+		int NumSimulatedParticles = 0;
+		for (int32 Index = 0; Index < SimulatableParticles.Num(); Index++)
 		{
-			if (auto* GTRigidParticle = GTParticle->CastToRigidParticle())
+			if (SimulatableParticles[Index] == true)
 			{
-				Buffer.DisabledStates[Idx] = GTRigidParticle->Disabled();
+				NumSimulatedParticles++;
+			}
+		}
+
+		// Add entries into simulation array
+		int NumRigids = Particles.Size();
+		BaseParticleIndex = NumRigids;
+		Particles.AddParticles(NumSimulatedParticles);
+		for (int32 Index = 0, NextId = 0; Index < SimulatableParticles.Num(); Index++)
+		{
+			if (SimulatableParticles[Index] == true)
+			{
+				RigidBodyID[Index] = NumRigids + NextId++;
+			}
+		}
+
+		// Add the rigid bodies
+		//for (int32 Index = 0; Index < TransformIndex.Num(); Index++)
+		const int32 NumGeometries = DynamicCollection->NumElements(FGeometryCollection::GeometryGroup);
+		ParallelFor(NumGeometries, [&](int32 GeometryIndex)
+		{
+			const int32 TransformGroupIndex = TransformIndex[GeometryIndex];
+			if (SimulatableParticles[TransformGroupIndex] == true)
+			{
+				const int32 RigidBodyIndex = RigidBodyID[TransformGroupIndex];
+				const FTransform WorldTransform = MassToLocal[TransformGroupIndex] * Transform[TransformGroupIndex] * Parameters.WorldTransform;
+
+				PopulateSimulatedParticle(Particles, Parameters.Shared, nullptr,// Simplicials[TransformGroupIndex].Get(),
+					Implicits[TransformGroupIndex], ThisSolver->GetMassScale() * Mass[TransformGroupIndex], ThisSolver->GetMassScale() * InertiaTensor[TransformGroupIndex],
+					RigidBodyIndex, WorldTransform, (uint8)EObjectStateTypeEnum::Chaos_Object_Dynamic, CollisionGroup[TransformGroupIndex]);
+				ClusterGroupIndex[RigidBodyIndex] = Parameters.ClusterGroupIndex;
+				Particles.CollisionParticles(RigidBodyIndex).Reset(Simplicials[TransformGroupIndex].Release());
+				if (Particles.CollisionParticles(RigidBodyIndex))
+				{
+					int32 NumCollisionParticles = Particles.CollisionParticles(RigidBodyIndex)->Size();
+					int32 CollisionParticlesSize = FMath::Max(0, FMath::Min(int(NumCollisionParticles*CollisionParticlesPerObjectFraction), NumCollisionParticles));
+					Particles.CollisionParticles(RigidBodyIndex)->Resize(CollisionParticlesSize);
+				}
+
+				StrainArray[RigidBodyIndex] = StrainDefault;
+				GetSolver()->SetPhysicsMaterial(RigidBodyIndex, Parameters.PhysicalMaterial);
 			}
 
-			//const Chaos::TShapesArray<float,3>& ShapesArray = GTParticle->ShapesArray();
+		});
+
+		for (FFieldSystemCommand & Cmd : Parameters.InitializationCommands)
+		{
+			if (Cmd.MetaData.Contains(FFieldSystemMetaData::EMetaType::ECommandData_ProcessingResolution)) Cmd.MetaData.Remove(FFieldSystemMetaData::EMetaType::ECommandData_ProcessingResolution);
+			FFieldSystemMetaDataProcessingResolution* ResolutionData = new FFieldSystemMetaDataProcessingResolution(EFieldResolutionType::Field_Resolution_Maximum);
+			Cmd.MetaData.Add(FFieldSystemMetaData::EMetaType::ECommandData_ProcessingResolution, TUniquePtr< FFieldSystemMetaDataProcessingResolution >(ResolutionData));
+			Commands.Add(Cmd);
 		}
-		++Idx;
+		Parameters.InitializationCommands.Empty();
+		ProcessCommands(Particles, GetSolver()->GetSolverTime());
+
+		ParallelFor(NumGeometries, [&](int32 GeometryIndex)
+		{
+			const int32 TransformGroupIndex = TransformIndex[GeometryIndex];
+			if (SimulatableParticles[TransformGroupIndex] == true)
+			{
+				const int32 RigidBodyIndex = RigidBodyID[TransformGroupIndex];
+
+				if (Parameters.InitialVelocityType == EInitialVelocityTypeEnum::Chaos_Initial_Velocity_User_Defined)
+				{
+					InitialLinearVelocity[TransformGroupIndex] = Parameters.InitialLinearVelocity;
+					InitialAngularVelocity[TransformGroupIndex] = Parameters.InitialAngularVelocity;
+					if (DynamicState[TransformGroupIndex] == (int32)EObjectStateTypeEnum::Chaos_Object_Dynamic)
+					{
+						Particles.V(RigidBodyIndex) = Parameters.InitialLinearVelocity;
+						Particles.W(RigidBodyIndex) = Parameters.InitialAngularVelocity;
+					}
+				}
+			}
+		});
+
+		InitializeKinematics(Particles, DynamicState);
+
+		InitializeRemoveOnFracture(Particles, DynamicState);
+
+		// #BG Temporary - don't cluster when playing back. Needs to be changed when kinematics are per-proxy to support
+		// kinematic to dynamic transition for clusters.
+		if (Parameters.EnableClustering)// && Parameters.CacheType != EGeometryCollectionCacheType::Play)
+		{
+			const TArray<int32> RecursiveOrder = ComputeRecursiveOrder(*RestCollection);
+
+			// num clusters
+			uint32 NumClusters=0;
+			TArray<bool> SubTreeContainsSimulatableParticle;
+			SubTreeContainsSimulatableParticle.SetNum(RecursiveOrder.Num());
+			for (const int32 TransformGroupIndex : RecursiveOrder)
+			{
+				if (Children[TransformGroupIndex].Num() > 0)
+				{
+					SubTreeContainsSimulatableParticle[TransformGroupIndex] = false;
+
+					TArray<uint32> RigidChildren, CollectionChildren;
+					for (const int32 ChildIndex : Children[TransformGroupIndex])
+					{
+						if(SubTreeContainsSimulatableParticle[ChildIndex])
+						{
+							NumClusters++;
+							SubTreeContainsSimulatableParticle[TransformGroupIndex] = true;
+							break;
+						}
+					}
+				}
+				else
+				{
+					SubTreeContainsSimulatableParticle[TransformGroupIndex] = (RigidBodyID[TransformGroupIndex] != INDEX_NONE);
+				}
+			}
+
+			const int32 ClusterStartIndex = Particles.Size();
+			Particles.AddParticles(NumClusters);
+
+			int32 ClusterRigidBodyId = ClusterStartIndex;
+			for (const int32 TransformGroupIndex : RecursiveOrder)
+			{
+				TArray<uint32> RigidChildren, CollectionChildren;
+				for (const int32 ChildIndex : Children[TransformGroupIndex])
+				{
+					if (RigidBodyID[ChildIndex] != INDEX_NONE)
+					{
+						RigidChildren.Add(RigidBodyID[ChildIndex]);
+						CollectionChildren.Add(ChildIndex);
+					}
+				}
+				if (RigidChildren.Num())
+				{
+					if (ReportTooManyChildrenNum >= 0 && RigidChildren.Num() > ReportTooManyChildrenNum)
+					{
+						UE_LOG(LogChaos, Warning, TEXT("Too many children (%d) in a single cluster:%s"), RigidChildren.Num(), *Parameters.Name);
+					}
+
+					RigidBodyID[TransformGroupIndex] = ClusterRigidBodyId;
+					Chaos::FClusterCreationParameters<float> CreationParameters;
+					CreationParameters.RigidBodyIndex = RigidBodyID[TransformGroupIndex];
+					BuildClusters(TransformGroupIndex, CollectionChildren, RigidChildren, CreationParameters);
+					ClusterRigidBodyId++;
+				}
+			}
+		}
+
+		NumParticles = Particles.Size() - BaseParticleIndex;
+
+		// Set Connectivity
+		const int32 NumTransforms = DynamicCollection->NumElements(FGeometryCollection::TransformGroup);
+		ParallelFor(NumTransforms, [&](int32 TransformGroupIndex)
+		{
+			if (Children[TransformGroupIndex].Num() > 0)
+			{
+				if (RigidBodyID[TransformGroupIndex] != INDEX_NONE)
+				{
+					GetSolver()->GetRigidClustering().GenerateConnectionGraph(RigidBodyID[TransformGroupIndex]);
+				}
+			}
+		});
+
+		// If we're recording and want to start immediately caching then we should cache the rest state
+		if (Parameters.IsCacheRecording() && Parameters.CacheBeginTime == 0.0f)
+		{
+			if (UpdateRecordedStateCallback)
+			{
+				UpdateRecordedStateCallback(0.0f, RigidBodyID, Particles, GetSolver()->GetCollisionConstraints());
+			}
+		}
+
+
+		if (InitializedState == ESimulationInitializationState::Activated)
+		{
+			//
+			//  Activated bodies has already been called so we are good to go.
+			//
+			InitializedState = ESimulationInitializationState::Initialized;
+
+			if (Parameters.EnableClustering && Parameters.ClusterGroupIndex)
+			{
+				GetSolver()->GetRigidClustering().IncrementPendingClusterCounter(Parameters.ClusterGroupIndex);
+				GetSolver()->GetRigidClustering().DecrementPendingClusterCounter(Parameters.ClusterGroupIndex);
+			}
+		}
+		else if (InitializedState == ESimulationInitializationState::Unintialized)
+		{
+
+			//
+			//  Activated bodies has not been called, se we are waiting 
+			//  to become active. Deactivate all bodies, and wait for
+			//  ActivatedBodies to be called, and defer the cluster initialization
+			//
+			InitializedState = ESimulationInitializationState::Created;
+			for (int32 TransformGroupIndex = 0; TransformGroupIndex < NumTransforms; TransformGroupIndex++)
+			{
+				const int32 RigidBodyIndex = RigidBodyID[TransformGroupIndex];
+				if (RigidBodyIndex != INDEX_NONE)
+				{
+					if (!Particles.Disabled(RigidBodyIndex))
+					{
+						PendingActivationList.Add(TransformGroupIndex);
+						GetSolver()->GetEvolution()->DisableParticle(RigidBodyIndex);
+					}
+				}
+			}
+
+			//
+			//  Clustering needs to advertise its group id to the cluster so
+			//  that the group is not initialized before all the bodies are
+			//  loaded and created. 
+			//
+			if (Parameters.EnableClustering && Parameters.ClusterGroupIndex )
+			{
+				GetSolver()->GetRigidClustering().IncrementPendingClusterCounter(Parameters.ClusterGroupIndex);
+			}
+		}
+		else
+		{
+			// unknown initialization state in creation callback
+			ensure(false);
+		}
+#endif
 	}
-	
-	// Make the buffer available to the consumer.
-	GameToPhysInterchange->FlipProducer(); 
-	
-	// If we wanted to emulate the other proxies, we could augment the 
-	// FStubGeometryCollectionData class to carry along the consumer buffer:
-	//
-	// return new FStubGeometryCollectionData(GametoPhysInterchange.GetConsumerBuffer());
-	//
-	// However, at this point, it seems pretty unnecessary to do the additional
-	// allocation and deallocation of the transient FStubGeometryCollectionData
-	// class, as I believe all these operations are happening with the same 
-	// instance of this class.  So skip all that, and just return null.
-	return nullptr;
 }
 
-void FGeometryCollectionPhysicsProxy::PushToPhysicsState(const Chaos::FParticleData *InData)
+void FGeometryCollectionPhysicsProxy::ActivateBodies()
 {
-	/**
-	 * CONTEXT: GAMETHREAD -> to -> PHYSICSTHREAD
-	 * Called on the game thread when the solver is about to advance forward. This
-	 * callback should Enqueue commands on the PhysicsThread to update the state of
-	 * the solver
-	 */
+	if (Parameters.Simulating)
+	{
+		if (InitializedState == ESimulationInitializationState::Created)
+		{
+#if TODO_REIMPLEMENT_RIGID_CLUSTERING
+			int32 ParentIndex = INDEX_NONE;
 
-	if (InitializedState == ESimulationInitializationState::Unintialized)
+			if (Parameters.EnableClustering && Parameters.ClusterGroupIndex)
+			{
+				auto& Clustering = GetSolver()->GetEvolution()->GetRigidClustering();
+				Clustering.DecrementPendingClusterCounter(Parameters.ClusterGroupIndex);
+				ParentIndex = Parameters.ClusterGroupIndex;
+			}
+#endif
+#if TODO_REIMPLEMENT_DEFERRED_ACTIVATION
+			Chaos::FPhysicsSolver::FParticlesType & Particles = GetSolver()->GetRigidParticles();
+			for (uint32 TransformGroupIndex : PendingActivationList)
+			{
+//				int32 RigidBodyIndex = RigidBodyID[TransformGroupIndex];
+//				checkSlow(RigidBodyIndex != INDEX_NONE);
+				if(Chaos::TPBDRigidParticleHandle<float, 3>* Handle = SolverParticleHandles[TransformGroupIndex])
+//				if (Particles.Disabled(RigidBodyIndex))
+				{
+					GetSolver()->GetEvolution()->EnableParticle(RigidBodyIndex, ParentIndex);
+				}
+			}
+#endif
+			PendingActivationList.Reset(0);
+
+			InitializedState = ESimulationInitializationState::Initialized;
+		}
+		else if (InitializedState == ESimulationInitializationState::Unintialized)
+		{
+			InitializedState = ESimulationInitializationState::Activated;
+		}
+		else
+		{
+			// unknown initialization state in activate bodies
+			ensure(false);
+		}
+	}
+}
+
+void FGeometryCollectionPhysicsProxy::BindParticleCallbackMapping(Chaos::TArrayCollectionArray<PhysicsProxyWrapper> & PhysicsProxyReverseMap, Chaos::TArrayCollectionArray<int32> & ParticleIDReverseMap)
+{
+	if (InitializedState== ESimulationInitializationState::Initialized)
 	{
-		return;
-	}
-
-	// The other proxies do a deep copy into dynamic transient memory, passed into
-	// this function via InData.  Instead of that, we're using persistent memory
-	// in a buffer to facilitate communicating between threads.
-
-	const FGeometryCollectionResults* GState = GameToPhysInterchange->GetConsumerBuffer();
-	if (!GState)
-	{
-		// This will happen if GameToPhysInterchange hasn't been updated since
-		// the last time GetConsumerBuffer() was called.  At this point, that's 
-		// probably a bug, as the interface code is assuming it needs to lock
-		// to prevent stomping on a double buffer.  We should look into doing 
-		// away with that lock, but for the time being, we'll assert:
-		ensure(false);
-		return;
-	}
-
-	//const FStubGeometryCollectionData* GCData = static_cast<const FStubGeometryCollectionData*>(InData);
-	//const FGeometryCollectionResults* GCState = GCData->GetStateData();
-
-	// Move the copied game thread data into the handles.
-	for (int32 HandleIdx = 0; HandleIdx < SolverParticleHandles.Num(); ++HandleIdx)
-	{
-		if (Chaos::TPBDRigidParticleHandle<float, 3>* Handle = SolverParticleHandles[HandleIdx])
+		for (int32 Index = 0; Index < RigidBodyID.Num(); Index++)
 		{
-			const int32 Idx = GState->BaseIndex + HandleIdx;
-
-			FTransform ParticleTransform =
-				MassToLocal[Idx] * GState->Transforms[Idx] * Parameters.WorldTransform;
-
-			Handle->SetX(ParticleTransform.GetTranslation());
-			Handle->SetR(ParticleTransform.GetRotation());
-			Handle->SetDisabled(GState->DisabledStates[Idx]);
-
-			Handle->SetLinearEtherDrag(Parameters.LinearEtherDrag);
-			Handle->SetAngularEtherDrag(Parameters.AngularEtherDrag);
-
-			if (GState->SharedGeometry[Idx] && GState->SharedGeometry[Idx]->HasBoundingBox())
+			const int32 RigidBodyIndex = RigidBodyID[Index];
+			if (RigidBodyIndex != INDEX_NONE)
 			{
-				Handle->SetHasBounds(true);
-				Handle->SetLocalBounds(GState->SharedGeometry[Idx]->BoundingBox());
-
-				Chaos::TAABB<float, 3> WorldSpaceBox =
-					GState->SharedGeometry[Idx]->BoundingBox().TransformedAABB(
-						Chaos::TRigidTransform<float, 3>(Handle->X(), Handle->R()));
-				// TODO Ryan - WorldSpaceBox.ThickenSymmetrically(GState->V());
-				Handle->SetWorldSpaceInflatedBounds(WorldSpaceBox);
-			}
-
-			// TODO Ryan - check dirty flags like TGeometryParticleData::DirtyFlags.
-			// In the mean time, we only set once.
-			if (!Handle->SharedGeometry() && GState->SharedGeometry.IsValidIndex(Idx) && GState->SharedGeometry[Idx])
-			{
-				Handle->SetSharedGeometry(GState->SharedGeometry[Idx]);
-			}
-
-			int32 CurrShape = 0;
-			for (const TUniquePtr<Chaos::TPerShapeData<Chaos::FReal, 3>>& Shape : Handle->ShapesArray())
-			{
-				// TODO Ryan - hardwired for now:
-				Shape->bDisable = false;
-				Shape->CollisionTraceType = Chaos::EChaosCollisionTraceFlag::Chaos_CTF_UseSimpleAndComplex; // Chaos_CTF_UseDefault?
-				if (CurrShape < GState->ShapeSimData[Idx].Num())
-				{
-					Shape->SimData = GState->ShapeSimData[Idx][CurrShape];
-					Shape->QueryData = GState->ShapeQueryData[Idx][CurrShape];
-				}
-				++CurrShape;
-			}
-		} // if rigid handle
-	} // for all SolverParticleHandles
-
-	if (DisableGeometryCollectionGravity) // cvar
-	{
-		// Our assumption is that you'd only ever want to wholesale opt geometry 
-		// collections out of gravity for debugging, so we keep this conditional
-		// out of the loop above and on it's own.  This means we can't turn gravity
-		// back on once it's off, but even if we didn't enclose this in an if(),
-		// this function won't be called again unless something dirties the proxy.
-
-		Chaos::TPerParticleGravity<float, 3>& GravityForces = GetSolver()->GetEvolution()->GetGravityForces();
-		for (int32 HandleIdx = 0; HandleIdx < SolverParticleHandles.Num(); ++HandleIdx)
-		{
-			if (Chaos::TPBDRigidParticleHandle<float, 3>* Handle = SolverParticleHandles[HandleIdx])
-			{
-				GravityForces.SetEnabled(*Handle, false);
+				PhysicsProxyReverseMap[RigidBodyIndex] = {this, EPhysicsProxyType::GeometryCollectionType};
+				ParticleIDReverseMap[RigidBodyIndex] = Index;
 			}
 		}
 	}
 }
-void FGeometryCollectionPhysicsProxy::BufferPhysicsResults()
+
+void FGeometryCollectionPhysicsProxy::PushKinematicStateToSolver(FParticlesType& Particles)
 {
-	/**
-	 * CONTEXT: PHYSICSTHREAD
-	 * Called per-tick after the simulation has completed. The proxy should cache the results of their
-	 * simulation into the local buffer. 
-	 */
-	using namespace Chaos;
-	SCOPE_CYCLE_COUNTER(STAT_CacheResultGeomCollection);
+	if (InitializedState == ESimulationInitializationState::Initialized)
+	{
+		FGeometryDynamicCollection* Collection = Parameters.DynamicCollection;
+		if (Collection->Transform.Num())
+		{
+			TManagedArray<int32>& DynamicState = Collection->GetAttribute<int32>("DynamicState", FGeometryCollection::TransformGroup);
 
-	//FGeometryCollectionResults& TargetResults = Results.GetPhysicsDataForWrite();
-	FGeometryCollectionResults& TargetResults = *PhysToGameInterchange->AccessProducerBuffer();
+#if TODO_REIMPLEMENT_RIGID_CLUSTERING
+			TSet<int32> ClustersToUpdate;
+			const Chaos::TArrayCollectionArray<Chaos::ClusterId>& ClusterID = GetSolver()->GetRigidClustering().GetClusterIdsArray();
 
-	if (TargetResults.NumTransformGroup() != PTDynamicCollection.NumElements(FGeometryCollection::TransformGroup))
-	{
-		TargetResults.InitArrays(PTDynamicCollection);
+			for (int32 TransformGroupIndex = 0; TransformGroupIndex < DynamicState.Num(); TransformGroupIndex++)
+			{
+				if (RigidBodyID[TransformGroupIndex] != INDEX_NONE)
+				{
+					int32 RigidBodyIndex = RigidBodyID[TransformGroupIndex];
+					if (DynamicState[TransformGroupIndex] == (int32)EObjectStateTypeEnum::Chaos_Object_Dynamic
+						&& (Particles.ObjectState(RigidBodyIndex) == Chaos::EObjectStateType::Kinematic || Particles.ObjectState(RigidBodyIndex) == Chaos::EObjectStateType::Static)
+						&& FLT_EPSILON < Particles.M(RigidBodyIndex))
+					{
+						Particles.SetObjectState(RigidBodyIndex, Chaos::EObjectStateType::Dynamic);
 
-		// Base particle index to calculate index from a global particle index on the game thread
-		TargetResults.BaseIndex = BaseParticleIndex;
-		TargetResults.NumParticlesAdded = NumParticles;
-	}
+						if (Parameters.InitialVelocityType == EInitialVelocityTypeEnum::Chaos_Initial_Velocity_User_Defined)
+						{
+							Particles.V(RigidBodyIndex) = Chaos::TVector<float, 3>(InitialLinearVelocity[TransformGroupIndex]);
+							Particles.W(RigidBodyIndex) = Chaos::TVector<float, 3>(InitialAngularVelocity[TransformGroupIndex]);
+						}
 
-	{ // tmp scope
-		IsObjectDynamic = false;
-		int32 Idx = 0;
-		for (auto* Handle : SolverParticleHandles)
-		{
-			if (!Handle || Handle->Disabled())
-			{
-				TargetResults.DisabledStates[Idx++] = true;
-				continue;
-			}
-			TargetResults.DisabledStates[Idx++] = false;
-			IsObjectDynamic = true;
-		}
-	} // tmp scope
+						if (!Particles.Disabled(RigidBodyIndex) && Particles.Sleeping(RigidBodyIndex))
+						{
+							Particles.SetSleeping(RigidBodyIndex, false);
+							GetSolver()->ActiveIndices().Add(RigidBodyIndex);
+						}
+						else
+						{
+							if (ClusterID[RigidBodyIndex].Id != INDEX_NONE)
+							{
+								int32 ActiveClusterIndex = GetSolver()->GetRigidClustering().GetActiveClusterIndex(RigidBodyIndex);
+								if (ActiveClusterIndex != INDEX_NONE)
+								{
+									ClustersToUpdate.Add(ActiveClusterIndex);
+								}
+							}
+						}
+					}
+					else if ((DynamicState[TransformGroupIndex] == (int32)EObjectStateTypeEnum::Chaos_Object_Kinematic)
+						&& (Particles.ObjectState(RigidBodyIndex) == Chaos::EObjectStateType::Dynamic)
+						&& FLT_EPSILON < Particles.M(RigidBodyIndex))
+					{
+						Particles.SetObjectState(RigidBodyIndex, Chaos::EObjectStateType::Kinematic);
 
-	// Advertise to game thread
-	TargetResults.IsObjectDynamic = IsObjectDynamic;
-	TargetResults.IsObjectLoading = IsObjectLoading;
-	
-	const FTransform& ActorToWorld = Parameters.WorldTransform;
-	// It's not worth shrinking the EndFrameUnparentingBuffer array, at least until the solver supports deleting bodies.
-	if (EndFrameUnparentingBuffer.Num() < NumParticles)
-	{
-		EndFrameUnparentingBuffer.Init(INDEX_NONE, NumParticles);
-	}
+						if (ClusterID[RigidBodyIndex].Id != INDEX_NONE)
+						{
+							int32 ActiveClusterIndex = GetSolver()->GetRigidClustering().GetActiveClusterIndex(RigidBodyIndex);
+							if (ActiveClusterIndex != INDEX_NONE)
+							{
+								ClustersToUpdate.Add(ActiveClusterIndex);
+							}
+						}
+					}
+					else if ((DynamicState[TransformGroupIndex] == (int32)EObjectStateTypeEnum::Chaos_Object_Static)
+						&& (Particles.ObjectState(RigidBodyIndex) == Chaos::EObjectStateType::Dynamic)
+						&& FLT_EPSILON < Particles.M(RigidBodyIndex))
+					{
+						Particles.SetObjectState(RigidBodyIndex, Chaos::EObjectStateType::Static);
 
-	{ // tmp scope
-		SCOPE_CYCLE_COUNTER(STAT_CalcParticleToWorld);
-		// Fill particle to world transforms
-//		TargetResults.ParticleToWorldTransforms.SetNum(NumParticles);
-/*		for (int32 TransformIndex = 0; TransformIndex < NumParticles; ++TransformIndex)
-		{
-			//only update roots and first children
-			if (TargetResults.Parent[TransformIndex] == INDEX_NONE)
-			{
-				if (auto* Handle = SolverParticleHandles[TransformIndex])
-				{
-					const FTransform Xf(Handle->R(), Handle->X());
-					//const int32 ParticleIndex = BaseParticleIndex + TransformIndex;
-					TargetResults.ParticleToWorldTransforms[TransformIndex] = Xf;
+						if (ClusterID[RigidBodyIndex].Id != INDEX_NONE)
+						{
+							int32 ActiveClusterIndex = GetSolver()->GetRigidClustering().GetActiveClusterIndex(RigidBodyIndex);
+							if (ActiveClusterIndex != INDEX_NONE)
+							{
+								ClustersToUpdate.Add(ActiveClusterIndex);
+							}
+						}
+					}
+					else if ((DynamicState[TransformGroupIndex] == (int32)EObjectStateTypeEnum::Chaos_Object_Sleeping)
+						&& (Particles.ObjectState(RigidBodyIndex) == Chaos::EObjectStateType::Dynamic))
+					{
+						Particles.SetObjectState(RigidBodyIndex, Chaos::EObjectStateType::Sleeping);
+						GetSolver()->ActiveIndices().Remove(RigidBodyIndex);
+					}
+					else if ((DynamicState[TransformGroupIndex] == (int32)EObjectStateTypeEnum::Chaos_Object_Dynamic)
+						&& (Particles.ObjectState(RigidBodyIndex) == Chaos::EObjectStateType::Sleeping))
+					{
+						Particles.SetObjectState(RigidBodyIndex, Chaos::EObjectStateType::Dynamic);
+						GetSolver()->ActiveIndices().Add(RigidBodyIndex);
+					}
 				}
 			}
-		}
-*/
-		//const TArrayCollectionArray<Chaos::ClusterId>& ClusterID = GetSolver()->GetRigidClustering().GetClusterIdsArray();
-		//const TArrayCollectionArray<Chaos::TRigidTransform<float, 3>>& ClusterChildToParentMap = GetSolver()->GetRigidClustering().GetChildToParentMap();
-		//const TArrayCollectionArray<FMultiChildProxyId>& MultiChildProxyIdArray = GetSolver()->GetRigidClustering().GetMultiChildProxyIdArray();
-		//const TArrayCollectionArray<TUniquePtr<TMultiChildProxyData<float, 3>>>& MultiChildProxyDataArray = GetSolver()->GetRigidClustering().GetMultiChildProxyDataArray();
-
-		for (int32 TransformIndex = 0; TransformIndex < NumParticles; ++TransformIndex)
-		{
-			Chaos::TPBDRigidClusteredParticleHandle<float, 3>* Handle = SolverParticleHandles[TransformIndex];
-			if (!Handle)
-			{
-				continue;
-			}
-
-			// Dynamic state is also updated by the solver during field interaction.
-			if (!Handle->Sleeping())
-			{
-				const Chaos::EObjectStateType ObjectState = Handle->ObjectState();
-				switch (ObjectState)
-				{
-				case Chaos::EObjectStateType::Kinematic:
-					TargetResults.DynamicState[TransformIndex] = (int)EObjectStateTypeEnum::Chaos_Object_Kinematic;
-					break;
-				case Chaos::EObjectStateType::Static:
-					TargetResults.DynamicState[TransformIndex] = (int)EObjectStateTypeEnum::Chaos_Object_Static;
-					break;
-				case Chaos::EObjectStateType::Sleeping:
-					TargetResults.DynamicState[TransformIndex] = (int)EObjectStateTypeEnum::Chaos_Object_Sleeping;
-					break;
-				case Chaos::EObjectStateType::Dynamic:
-				case Chaos::EObjectStateType::Uninitialized:
-				default:
-					TargetResults.DynamicState[TransformIndex] = (int)EObjectStateTypeEnum::Chaos_Object_Dynamic;
-					break;
-				}
-			}
-
-			// Update the transform and parent hierarchy of the active rigid bodies. Active bodies can be either
-			// rigid geometry defined from the leaf nodes of the collection, or cluster bodies that drive an entire
-			// branch of the hierarchy within the GeometryCollection.
-			// - Active bodies are directly driven from the global position of the corresponding
-			//   rigid bodies within the solver ( cases where RigidBodyID[TransformGroupIndex] is not disabled ). 
-			// - Deactivated bodies are driven from the transforms of their active parents. However the solver can
-			//   take ownership of the parents during the simulation, so it might be necessary to force deactivated
-			//   bodies out of the collections hierarchy during the simulation.  
-			if (!Handle->Disabled())
-			{
-				// Update the transform of the active body. The active body can be either a single rigid
-				// or a collection of rigidly attached geometries (Clustering). The cluster is represented as a
-				// single transform in the GeometryCollection, and all children are stored in the local space
-				// of the parent cluster.
-				// ... When setting cluster transforms it is expected that the MassToLocal is identity.
-				//     Cluster initialization will set the vertices in the MassSpace of the rigid body.
-				// ... When setting individual rigid bodies that are not clustered, the MassToLocal will be 
-				//     non-Identity, and will reflect the difference between the geometric center of the geometry
-				//     and that corresponding rigid bodies center of mass. 
-
-				FTransform& ParticleToWorld = TargetResults.ParticleToWorldTransforms[TransformIndex];
-				ParticleToWorld = FTransform(Handle->R(), Handle->X());
-
-				TargetResults.Transforms[TransformIndex] =
-					MassToLocal[TransformIndex].GetRelativeTransformReverse(ParticleToWorld).GetRelativeTransform(ActorToWorld);
-				TargetResults.Transforms[TransformIndex].NormalizeRotation();
-
-
-#if 0 // TODO Ryan
-				// Force all enabled rigid bodies out of the transform hierarchy
-				if (TargetResults.Parent[TransformIndex] != INDEX_NONE)
-				{
-					// Children in the hierarchy are stored in a TSet, which is not thread safe.  
-					// So we retain indices to remove afterwards.
-					EndFrameUnparentingBuffer[TransformIndex] = TargetResults.Parent[TransformIndex];
-				}
 #endif
-#if 0 // TODO Ryan
-				// When a leaf node rigid body is removed from a cluster, the rigid
-				// body will become active and needs its clusterID updated.  This just
-				// syncs the clusterID all the time.
-				CollectionClusterID[TransformGroupIndex] = ClusterID[RigidBodyIndex].Id;
-#endif
-			}
-			else // Handle->Disabled()
-			{
-#if 0 // TODO Ryan
-				// The rigid body parent cluster has changed within the solver, and its
-				// parent body is not tracked within the geometry collection. So we need to
-				// pull the rigid bodies out of the transform hierarchy, and just drive
-				// the positions directly from the solvers cluster particle. 
-				if (CollectionClusterID[TransformGroupIndex] != ClusterID[RigidBodyIndex].Id)
-				{
-					// Force all driven rigid bodies out of the transform hierarchy
-					if (Parent[TransformGroupIndex] != INDEX_NONE)
-					{
-						const int32 ParentIndex = Parent[TransformGroupIndex];
-						// Children in the hierarchy are stored in a TSet, which is not thread safe.  So we retain
-						// indices to remove afterwards.
-						EndFrameUnparentingBuffer[TransformGroupIndex] = ParentIndex;
-					}
-					CollectionClusterID[TransformGroupIndex] = ClusterID[RigidBodyIndex].Id;
-				}
-#endif
-				// Disabled rigid bodies that have valid cluster parents, and have been re-indexed by the
-				// solver (as in, They were re-clustered outside of the geometry collection), These clusters 
-				// will need to be rendered based on the clusters position. 
-
-				if (TPBDRigidParticleHandle<float, 3>* ParentHandle = Handle->ClusterIds().Id)
-				{
-					if (TPBDRigidClusteredParticleHandle<float, 3>* ClusteredParentHandle = ParentHandle->CastToClustered())
-					{
-						if (ClusteredParentHandle->InternalCluster())
-						{
-							//const FTransform ClusterChildToWorld = ClusterChildToParentMap[RigidBodyIndex] * FTransform(Particles.R(ClusterParentIndex), Particles.X(ClusterParentIndex));
-							const TRigidTransform<float, 3>& ChildToParent = Handle->ChildToParent();
-							const FTransform ClusterChildToWorld = ChildToParent * FTransform(ParentHandle->R(), ParentHandle->X());
-							if (Parameters.IsCacheRecording())
-							{
-								//Particles.X(RigidBodyIndex) = ClusterChildToWorld.GetTranslation();
-								Handle->SetX(ClusterChildToWorld.GetTranslation());
-								//Particles.R(RigidBodyIndex) = ClusterChildToWorld.GetRotation();
-								Handle->SetR(ClusterChildToWorld.GetRotation());
-							}
-							// GeomToActor = ActorToWorld.Inv() * ClusterChildToWorld * MassToLocal.Inv();
-							TargetResults.Transforms[TransformIndex] = 
-								MassToLocal[TransformIndex].GetRelativeTransformReverse(ClusterChildToWorld).GetRelativeTransform(ActorToWorld);
-							TargetResults.Transforms[TransformIndex].NormalizeRotation();
-						}
-					}
-				}
-
-			}// end if
-		} // end for
-	} // tmp scope
-
-	// If object is dynamic, compute global matrices	
-	if (IsObjectDynamic || TargetResults.GlobalTransforms.Num() == 0)
-	{
-		SCOPE_CYCLE_COUNTER(STAT_CalcGlobalGCMatrices);
-		check(TargetResults.Transforms.Num() == TargetResults.Parent.Num());
-		GeometryCollectionAlgo::GlobalMatrices(
-			TargetResults.Transforms, TargetResults.Parent, TargetResults.GlobalTransforms);
-	}
-
-	// compute world bounds
-	// #note: this is a loose bounds based on the circumscribed box of a bounding sphere for the geometry.		
-	if (IsObjectDynamic || TargetResults.WorldBounds.GetSphere().W < 1e-5)
-	{
-		SCOPE_CYCLE_COUNTER(STAT_CalcGlobalGCBounds);
-		FBox BoundingBox(ForceInit);
-		//const FMatrix& ActorToWorld = Parameters.WorldTransform.ToMatrixWithScale();
-		for (int i = 0; i < ValidGeometryBoundingBoxes.Num(); ++i)
-		{
-			BoundingBox += ValidGeometryBoundingBoxes[i].TransformBy(
-				TargetResults.GlobalTransforms[ValidGeometryTransformIndices[i]] * ActorToWorld);
-		}
-		TargetResults.WorldBounds = FBoxSphereBounds(BoundingBox);		
-	}
-
-#if 0 // TODO Ryan
-	for (int32 TransformGroupIndex = 0; TransformGroupIndex < TransformSize; TransformGroupIndex++)
-	{
-		const int32 ParentIndex = EndFrameUnparentingBuffer[TransformGroupIndex];
-		if (ParentIndex >= 0)
-		{
-			// We reuse EndFrameUnparentingBuffer potentially without reinitialization, so reset this index to -1 before it gets paged out.
-			EndFrameUnparentingBuffer[TransformGroupIndex] = -1;
-
-			Children[ParentIndex].Remove(TransformGroupIndex);
-			Parent[TransformGroupIndex] = INDEX_NONE;
-		}
-	}
-#endif
-
-	// We pulled state out of particles.  Push results into PTDynamicCollection, 
-	// so that it has an updated view.
-	UpdateGeometryCollection(TargetResults);
-
-	// We've updated a buffer.  There's no reason to wait for FGeometryCollectionPhysicsProxy::FlipBuffer()
-	// to be called.  Just flip.
-	PhysToGameInterchange->FlipProducer();
-//#endif // TODO_REIMPLEMENT_GET_RIGID_PARTICLES
-}
-
-void FGeometryCollectionPhysicsProxy::FlipBuffer()
-{
-	/**
-	 * CONTEXT: PHYSICSTHREAD (Write Locked)
-	 * Called by the physics thread to signal that it is safe to perform any double-buffer flips here.
-	 * The physics thread has pre-locked an RW lock for this operation so the game thread won't be reading
-	 * the data
-	 */
-	//Results.Flip();
-	//PhysToGameInterchange->FlipProducer();
-}
-
-// Called from FPhysScene_ChaosInterface::SyncBodies(), NOT the solver.
-void FGeometryCollectionPhysicsProxy::PullFromPhysicsState()
-{
-	/**
-	 * CONTEXT: GAMETHREAD (Read Locked)
-	 * Perform a similar operation to Sync, but take the data from a gamethread-safe buffer. This will be called
-	 * from the game thread when it cannot sync to the physics thread. The simulation is very likely to be running
-	 * when this happens so never read any physics thread data here!
-	 *
-	 * Note: A read lock will have been acquired for this - so the physics thread won't force a buffer flip while this
-	 * sync is ongoing
-	 */
-
-	const FGeometryCollectionResults* TargetResultPtr = PhysToGameInterchange->GetConsumerBuffer();
-	if (!TargetResultPtr)
-		return;
-	// Remove const-ness as we're going to do the ExchangeArrays() thing.
-	FGeometryCollectionResults& TR = *const_cast<FGeometryCollectionResults*>(TargetResultPtr);
-
-	// We should never be changing the number of entries, this would break other 
-	// attributes in the transform group.
-	const int32 NumTransforms = GTDynamicCollection->Transform.Num();
-	if (ensure(NumTransforms == TR.Transforms.Num()))	
-	{
-		for(int32 TmIndex = 0; TmIndex < NumTransforms; ++TmIndex)
-		{
-			if(!TR.DisabledStates[TmIndex])
-			{
-				GTDynamicCollection->Transform[TmIndex] = TR.Transforms[TmIndex];
-			}
-		}
-
-		//question: why do we need this? Sleeping objects will always have to update GPU
-		GTDynamicCollection->MakeDirty();
-
-		if (CacheSyncFunc)
-		{
-			CacheSyncFunc(TR);
-		}
-	}
-/*
-	for (int32 TransformIndex = 0; TransformIndex < NumParticles; ++TransformIndex)
-	{
-		const int32 ParticleIndex = BaseParticleIndex + TransformIndex;
-		//Chaos::TPBDRigidParticleHandle<float, 3>& Handle = *SolverParticleHandles[ParticleIndex];
-		if (Chaos::TGeometryParticle<float, 3>* Particle = GTGeometryParticles[ParticleIndex].Get())
-		{
-			const FTransform& Xf = TR.Transforms[ParticleIndex];
-			Particle->SetX(Xf.GetTranslation());
-			Particle->SetR(Xf.GetRotation());
-
-			GTDynamicCollection->Transform[TransformIndex] = Xf;
-		}
-	}
-*/
-//	UpdateGeometryCollection(TR);
-}
-
-void FGeometryCollectionPhysicsProxy::UpdateGeometryCollection(
-	const FGeometryCollectionResults& GCResults,
-	FGeometryDynamicCollection* Collection)
-{
-	if (InitializedState != ESimulationInitializationState::Initialized)
-	{
-		return;
-	}
-	if (!Collection)
-	{
-		Collection = Parameters.DynamicCollection;
-	}
-	check(Collection);
-
-	TManagedArray<FTransform>& Transform = Collection->Transform;
-	TManagedArray<int32>& Parent = Collection->Parent;
-	TManagedArray<TSet<int32>>& Children = Collection->Children;
-	TManagedArray<int32>& SimulationType = Collection->SimulationType;
-	TManagedArray<int32>& DynamicState = Collection->GetAttribute<int32>("DynamicState", FGeometryCollection::TransformGroup);
-
-	// Particles X and R are aligned with center of mass and inertia principal axes.
-	// Renderer doesn't know about this and simply does ActorToWorld * GeomToActor * LocalSpaceVerts
-	// In proper math multiplication order:
-	//		ParticleToWorld = ActorToWorld * GeomToActor * MassToLocal
-	//		GeomToWorld = ActorToWorld * GeomToActor
-	//		=> GeomToWorld = ParticleToWorld * MassToLocal.Inv()
-	//		=> GeomToActor = ActorToWorld.Inv() * ParticleToWorld * MassToLocal.Inv()
-
-	const FTransform& ActorToWorld = Parameters.WorldTransform;
-	const int32 TransformSize = GetTransformGroupSize();
-	for (int32 TransformGroupIndex = 0; TransformGroupIndex < TransformSize; ++TransformGroupIndex)
-	{
-		if (!GTGeometryParticles[TransformGroupIndex])
-			continue;
-
-		auto* Handle = SolverParticleHandles[TransformGroupIndex];
-
-		// dynamic state is also updated by the solver during field interaction. 
-		if(!Handle->Sleeping())
-		{
-			DynamicState[TransformGroupIndex] = GCResults.DynamicState[TransformGroupIndex];
-		}
-
-		// Update the transform and parent hierarchy of the active rigid bodies. Active bodies can be either
-		// rigid geometry defined from the leaf nodes of the collection, or cluster bodies that drive an entire
-		// branch of the hierarchy within the GeometryCollection.
-		// - Active bodies are directly driven from the global position of the corresponding
-		//   rigid bodies within the solver ( cases where RigidBodyID[TransformGroupIndex] is not disabled ). 
-		// - Deactivated bodies are driven from the transforms of their active parents. However the solver can
-		//   take ownership of the parents during the simulation, so it might be necessary to force deactivated
-		//   bodies out of the collections hierarchy during the simulation.  
-
-		if(!GCResults.DisabledStates[TransformGroupIndex])
-		{
-			// Update the transform of the active body. The active body can be either a single rigid
-			// or a collection of rigidly attached geometries (Clustering). The cluster is represented as a
-			// single transform in the GeometryCollection, and all children are stored in the local space
-			// of the parent cluster.
-			// ... When setting cluster transforms it is expected that the MassToLocal is identity.
-			//     Cluster initialization will set the vertices in the MassSpace of the rigid body.
-			// ... When setting individual rigid bodies that are not clustered, the MassToLocal will be 
-			//     non-Identity, and will reflect the difference between the geometric center of the geometry
-			//     and that corresponding rigid bodies center of mass. 
-
-			// GeomToActor = ActorToWorld.Inv() * ParticleToWorld * MassToLocal.Inv();
-			Transform[TransformGroupIndex] = GCResults.Transforms[TransformGroupIndex];
-
-			// When a leaf node rigid body is removed from a cluster the rigid
-			// body will become active and needs its clusterID updated. This just
-			// syncs the clusterID all the time. 
-//			CollectionClusterID[TransformGroupIndex] = ClusterID[RigidBodyIndex].Id;
-		}
-		else //if (Particles.Disabled(RigidBodyIndex))
-		{
-			Transform[TransformGroupIndex] = GCResults.Transforms[TransformGroupIndex];
 		}
 	}
 }
 
-void IdentifySimulatableElements(Chaos::FErrorReporter& ErrorReporter, FGeometryCollection& GeometryCollection)
+void FGeometryCollectionPhysicsProxy::ParameterUpdateCallback(FParticlesType& Particles, const float Time)
 {
-	// Determine which collection particles to simulate
-	const TManagedArray<FBox>& BoundingBox = GeometryCollection.BoundingBox;
-	const TManagedArray<int32>& VertexCount = GeometryCollection.VertexCount;
-	const TManagedArray<int32>& TransformIndex = GeometryCollection.TransformIndex;
-	const int32 NumTransforms = GeometryCollection.NumElements(FGeometryCollection::TransformGroup);
+	if (InitializedState == ESimulationInitializationState::Initialized)
+	{
+		FGeometryDynamicCollection* Collection = Parameters.DynamicCollection;
+		check(Collection);
 
-	const int32 NumTransformMappings = TransformIndex.Num();
-
-	// Do not simulate hidden geometry
-	TArray<bool> HiddenObject;
-	HiddenObject.Init(true, NumTransforms);
-	const TManagedArray<bool>& Visible = GeometryCollection.Visible;
-	const TManagedArray<int32>& BoneMap = GeometryCollection.BoneMap;
-
-	const TManagedArray<FIntVector>& Indices = GeometryCollection.Indices;
-	int32 PrevObject = -1;
-	bool bContiguous = true;
-	for(int32 i = 0; i < Indices.Num(); i++)
-	{
-		if(Visible[i])
+		if (Collection->Transform.Num())
 		{
-			int32 ObjIdx = BoneMap[Indices[i][0]];
-			HiddenObject[ObjIdx] = false;
-			if (!ensureMsgf(ObjIdx >= PrevObject, TEXT("Objects are not contiguous. This breaks assumptions later in the pipeline")))
-			{
-				bContiguous = false;
-			}
-			
-			PrevObject = ObjIdx;
-		}
-	}
-
-	if (!bContiguous)
-	{
-		ErrorReporter.ReportError(TEXT("Objects are not contiguous. This breaks assumptions later in the pipeline"));
-		ErrorReporter.HandleLatestError();
-	}
+			ProcessCommands(Particles, Time);
 
 
-	//For now all simulation data is a non compiled attribute. Not clear what we want for simulated vs kinematic collections
-	TManagedArray<bool>& SimulatableParticles = 
-		GeometryCollection.AddAttribute<bool>(
-			FGeometryCollection::SimulatableParticlesAttribute, FTransformCollection::TransformGroup);
+			if (Parameters.RecordedTrack)
+			{
+				float ReverseTime = Parameters.RecordedTrack->GetLastTime() - Time + Parameters.ReverseCacheBeginTime;
+				// @todo(mlentine): We shouldn't need to do this every frame
+				if (Parameters.IsCacheRecording() && Time > Parameters.ReverseCacheBeginTime && Parameters.ReverseCacheBeginTime != 0 && Parameters.RecordedTrack->IsTimeValid(ReverseTime))
+				{
+					for (int32 Index = 0; Index < RigidBodyID.Num(); Index++)
+					{
+						int32 RigidBodyIndex = RigidBodyID[Index];
 
-	for(int32 TransformIdx = 0; TransformIdx < NumTransforms; TransformIdx++)
-	{
-		SimulatableParticles[TransformIdx] = false;
-	}
-
-	for(int i = 0; i < NumTransformMappings; i++)
-	{
-		int32 Tdx = TransformIndex[i];
-		checkSlow(0 <= Tdx && Tdx < NumTransforms);
-		if(GeometryCollection.IsGeometry(Tdx) && VertexCount[i] && 0.f < BoundingBox[i].GetSize().SizeSquared() && !HiddenObject[Tdx])
-		{
-			SimulatableParticles[Tdx] = true;
+						// Check index, will be invalid for cluster parents.
+						if (RigidBodyIndex != INDEX_NONE)
+						{
+							Particles.InvM(RigidBodyIndex) = 0.f;
+							Particles.InvI(RigidBodyIndex) = Chaos::PMatrix<float, 3, 3>(0.f);
+						}
+					}
+				}
+			}
+			/* @question : Should we tell the solver the mass has changed ? */
 		}
 	}
 }
 
-void BuildSimulationData(Chaos::FErrorReporter& ErrorReporter, FGeometryCollection& GeometryCollection, const FSharedSimulationParameters& SharedParams)
+void FGeometryCollectionPhysicsProxy::DisableCollisionsCallback(TSet<TTuple<int32, int32>>& InPairs)
 {
-	IdentifySimulatableElements(ErrorReporter, GeometryCollection);
-	FGeometryCollectionPhysicsProxy::InitializeSharedCollisionStructures(ErrorReporter, GeometryCollection, SharedParams);
+
 }
+
+void FGeometryCollectionPhysicsProxy::AddForceCallback(FParticlesType& InParticles, const float InDt, const int32 InIndex)
+{
+	// gravity forces managed directly on the solver for now
+}
+
diff --git a/Engine/Source/Runtime/Experimental/ChaosSolvers/Public/PhysicsProxy/GeometryCollectionPhysicsProxy.h b/Engine/Source/Runtime/Experimental/ChaosSolvers/Public/PhysicsProxy/GeometryCollectionPhysicsProxy.h
--- a/Engine/Source/Runtime/Experimental/ChaosSolvers/Public/PhysicsProxy/GeometryCollectionPhysicsProxy.h
+++ b/Engine/Source/Runtime/Experimental/ChaosSolvers/Public/PhysicsProxy/GeometryCollectionPhysicsProxy.h
@@ -125,7 +125,6 @@
 	using FInitFunc = TFunction<void(FSimulationParameters&)>;
 	using FCacheSyncFunc = TFunction<void(const FGeometryCollectionResults&)>;
 	using FFinalSyncFunc = TFunction<void(const FRecordedTransformTrack&)>;
-	/** ----------------------- */
 
 	FGeometryCollectionPhysicsProxy() = delete;
 	/**
@@ -149,19 +148,97 @@
 	 * and prepare for simulation.
 	 */
 	void Initialize();
+	void Reset() { InitializedState = ESimulationInitializationState::Unintialized; }
 
-	/** Initialize rigid body particles on the game thread. */
+	/** 
+	 * Initialize rigid body particles on the game thread. 
+	 *
+	 * Called by \c FPBDRigidsSolver::RegisterObject().
+	 */
 	void InitializeBodiesGT();
 
-	/** Called from the physics thread. */
+	/** 
+	 * Finish initialization on the physics thread. 
+	 *
+	 * Called by solver command registered by \c FPBDRigidsSolver::RegisterObject().
+	 */
 	void InitializeBodiesPT(
 		Chaos::FPBDRigidsSolver* RigidsSolver,
 		Chaos::FPBDRigidsSolver::FParticlesType& Particles);
 
-	void Reset() { InitializedState = ESimulationInitializationState::Unintialized; }
-
-	/** Solver Object interface */
+	/** */
 	bool IsSimulating() const { return Parameters.Simulating; }
+
+	/**
+	 * Pushes current game thread particle state into the \c GameToPhysInterchange.
+	 *
+	 * Redirects to \c BufferGameState(), and returns nullptr as this class manages 
+	 * data transport to the physics thread itself, without allocating memory.
+	 */
+	Chaos::FParticleData* NewData() { BufferGameState(); return nullptr; }
+	void BufferGameState();
+
+	/** 
+	 * Push game data to the physics thread.
+	 * 
+	 * The solver has determined that this proxy needs to update its physics state 
+	 * with game data, either as part of the solver registration process, or
+	 * a game thread particle instance was updated, which marked this proxy dirty.
+	 * 
+	 * Pulls data out of the \c GameToPhysInterchange, and updates solver particle 
+	 * data.
+	 *
+	 * Invoked by \c FPBDRigidsSolver::RegisterObject() and via particle dirtying. 
+	 */
+	void PushToPhysicsState(const Chaos::FParticleData*);
+
+	/** Called at the end of \c FPBDRigidsSolver::PushPhysicsStateExec(). */
+	void ClearAccumulatedData() {}
+
+	/** Push physics state into the \c PhysToGameInterchange. */
+	void BufferPhysicsResults();
+
+	/** Does nothing as \c BufferPhysicsResults() already did this. */
+	void FlipBuffer();
+	
+	/** 
+	 * Pulls data out of the PhysToGameInterchange and updates \c GTDynamicCollection. 
+	 * Called from FPhysScene_ChaosInterface::SyncBodies(), NOT the solver.
+	 */
+	void PullFromPhysicsState();
+
+	bool IsDirty() { return false; }
+
+	EPhysicsProxyType ConcreteType() { return EPhysicsProxyType::GeometryCollectionType; }
+
+	void SyncBeforeDestroy();
+	void OnRemoveFromSolver(Chaos::FPBDRigidsSolver *RBDSolver);
+	void OnRemoveFromScene();
+
+	void SetCollisionParticlesPerObjectFraction(float CollisionParticlesPerObjectFractionIn) 
+	{CollisionParticlesPerObjectFraction = CollisionParticlesPerObjectFractionIn;}
+
+	TManagedArray<Chaos::TPBDRigidClusteredParticleHandle<float, 3>*>& GetSolverParticleHandles() 
+	{ return SolverParticleHandles; }
+
+	const FGeometryCollectionResults* GetConsumerResultsGT() const 
+	{ return PhysToGameInterchange ? PhysToGameInterchange->PeekConsumerBuffer() : nullptr; }
+
+	/** Enqueue a field \p Command to be processed by \c ProcessCommands() or 
+	 * \c FieldForcesUpdateCallback(). 
+	 */
+	void BufferCommand(Chaos::FPhysicsSolver* , const FFieldSystemCommand& Command) 
+	{ Commands.Add(Command); }
+
+	static void InitializeSharedCollisionStructures(Chaos::FErrorReporter& ErrorReporter, FGeometryCollection& RestCollection, const FSharedSimulationParameters& SharedParams);
+	static void InitRemoveOnFracture(FGeometryCollection& RestCollection, const FSharedSimulationParameters& SharedParams);
+
+	static void MergeRecordedTracks(const FRecordedTransformTrack& A, const FRecordedTransformTrack& B, FRecordedTransformTrack& Target);
+	static FRecordedFrame& InsertRecordedFrame(FRecordedTransformTrack& InTrack, float InTime);
+
+	// 
+	// DEPRECATED
+	//
 	void UpdateKinematicBodiesCallback(const FParticlesType& InParticles, const float InDt, const float InTime, FKinematicProxy& InKinematicProxy);
 	void StartFrameCallback(const float InDt, const float InTime);
 	void EndFrameCallback(const float InDt);
@@ -173,81 +250,62 @@
 	void FieldForcesUpdateCallback(Chaos::FPhysicsSolver* InSolver, FParticlesType& Particles, Chaos::TArrayCollectionArray<FVector> & Force, Chaos::TArrayCollectionArray<FVector> & Torque, const float Time);
 	void ActivateBodies();
 
-	void SyncBeforeDestroy();
-	void OnRemoveFromSolver(Chaos::FPBDRigidsSolver *RBDSolver);
-	void OnRemoveFromScene();
-
-	Chaos::FParticleData* NewData();
-	/** Invoked by \c FPBDRigidsSolver::RegisterObject(). */
-	void PushToPhysicsState(const Chaos::FParticleData*);
-	void ClearAccumulatedData() {}
-	void BufferPhysicsResults();
-	void FlipBuffer();
-	/** Updates \c GTDynamicCollection with curr state of the \c Results buffer. */
-	void PullFromPhysicsState();
-	bool IsDirty() { return false; }
-	EPhysicsProxyType ConcreteType() { return EPhysicsProxyType::GeometryCollectionType; }
-	/** ----------------------- */
-
-	void MergeRecordedTracks(const FRecordedTransformTrack& A, const FRecordedTransformTrack& B, FRecordedTransformTrack& Target);
-	FRecordedFrame& InsertRecordedFrame(FRecordedTransformTrack& InTrack, float InTime);
+protected:
+	/**
+	 * Build a physics thread cluster parent particle.
+	 *	\p CollectionClusterIndex - the source geometry collection transform index.
+	 *	\p ChildHandles - physics particle handles of the cluster children.
+	 *  \p ChildTransformGroupIndices - geometry collection indices of the children.
+	 *  \P Parameters - uh, yeah...  Other parameters.
+	 */
 	Chaos::TPBDRigidClusteredParticleHandle<float, 3>* BuildClusters(
 		const uint32 CollectionClusterIndex, 
 		TArray<Chaos::TPBDRigidParticleHandle<float, 3>*>& ChildHandles,
 		const TArray<int32>& ChildTransformGroupIndices,
 		const Chaos::FClusterCreationParameters<float> & Parameters);
-	void BufferCommand(Chaos::FPhysicsSolver* InSolver, const FFieldSystemCommand& Command) { Commands.Add(Command); }
 
-	static void InitializeSharedCollisionStructures(Chaos::FErrorReporter& ErrorReporter, FGeometryCollection& RestCollection, const FSharedSimulationParameters& SharedParams);
-
-	static void InitRemoveOnFracture(FGeometryCollection& RestCollection, const FSharedSimulationParameters& SharedParams);
-	const Chaos::TBufferedData<FGeometryCollectionResults>& GetPhysicsResults() const { return Results; }
-
-	// Game thread and physics thread access to the dynamic collection. 
-	const FGeometryDynamicCollection* GetGeometryDynamicCollection_PhysicsThread() { return &PTDynamicCollection; }
-	const FGeometryDynamicCollection* GetGeometryDynamicCollection_GameThread() { return GTDynamicCollection; }
-
-	const FSimulationParameters& GetSimulationParameters() { return Parameters; }
-
-	int32 GetBaseParticleIndex() const { return BaseParticleIndex; }
-	int32 GetNumAddedParticles() const { return NumParticles; }
-	int32 GetTransformGroupSize() const { return Parameters.DynamicCollection->NumElements(FGeometryCollection::TransformGroup); }
-
-	TManagedArray<Chaos::TPBDRigidClusteredParticleHandle<float, 3>*>& GetSolverParticleHandles() { return SolverParticleHandles; }
-	void ClearSolverParticleHandles() { SolverParticleHandles.Fill(nullptr); }
-
-	/*
-	* ContiguousIndices
-	*   Generates a mapping between the Position array and the results array. When EFieldResolutionType is set to Maximum the complete
-	*   particle mapping is provided from the Particles.X to Particles.Attribute, when Minimum is set only the ActiveIndices and the
-	*   direct children of the active clusters are set in the IndicesArray.
-	*/
-	void ContiguousIndices(TArray<ContextIndex>& IndicesArray, const Chaos::FPhysicsSolver* RigidSolver, EFieldResolutionType ResolutionType, bool bForce);
-
-	//
-	void SetCollisionParticlesPerObjectFraction(float CollisionParticlesPerObjectFractionIn) {CollisionParticlesPerObjectFraction = CollisionParticlesPerObjectFractionIn;}
-
-	const FGeometryCollectionResults* GetConsumerResultsGT() { return PhysToGameInterchange->PeekConsumerBuffer(); }
-
-	/*
-	* For Testing ONLY!
-	*/
-	const TManagedArray<int32>& RigidBodyIDArray_TestingAccess() const { return RigidBodyID; }
-	const TManagedArray<FTransform>& MassToLocal_GameThreadAccess() const { return MassToLocal; }
-
-protected:
+	/** 
+	 * Copy \p Results into \p TargetCollection, or \c Parameters.DynamicCollection 
+	 * if \c nullptr.
+	 */
 	void UpdateGeometryCollection(
 		const FGeometryCollectionResults& Results,
 		FGeometryDynamicCollection* TargetCollection=nullptr);
 
-protected:
-	int32 CalculateHierarchyLevel(FGeometryDynamicCollection* GeometryCollection, int32 TransformIndex);
+	/**
+	 * Generates a mapping between the Position array and the results array. 
+	 * When EFieldResolutionType is set to Maximum the complete particle mapping 
+	 * is provided from the Particles.X to Particles.Attribute, when Minimum is 
+	 * set only the ActiveIndices and the direct children of the active clusters 
+	 * are set in the IndicesArray.
+	 */
+	void ContiguousIndices(
+		TArray<ContextIndex>& IndicesArray, 
+		const Chaos::FPhysicsSolver* RigidSolver, 
+		EFieldResolutionType ResolutionType, 
+		bool bForce);
+
+	/** The size of the transform group, as reported by the \c DynamicCollection 
+	 * currently stored in \c Parameters. 
+	 */
+	int32 GetTransformGroupSize() const 
+	{ return Parameters.DynamicCollection ? Parameters.DynamicCollection->NumElements(FGeometryCollection::TransformGroup) : 0; }
+
+	/** 
+	 * Traverses the parents of \p TransformIndex in \p GeometryCollection, counting
+	 * the number of levels until the next parent is \c INDEX_NONE.
+	 */
+	int32 CalculateHierarchyLevel(
+		const FGeometryDynamicCollection* GeometryCollection, 
+		int32 TransformIndex) const;
+
 	void CreateDynamicAttributes();
 	void InitializeKinematics(FParticlesType& Particles, const TManagedArray<int32>& ObjectType);
 	void InitializeRemoveOnFracture(FParticlesType& Particles, const TManagedArray<int32>& DynamicState);
-	void ProcessCommands(FParticlesType& Particles, const float Time);
 	void PushKinematicStateToSolver(FParticlesType& Particles);
 
+	void ProcessCommands(FParticlesType& Particles, const float Time);
+
 private:
 	void UpdateRecordedState(
 		float SolverTime, 
@@ -304,7 +362,7 @@
 	TManagedArray<FTransform> MassToLocal;
 	TManagedArray<int32> CollisionMask;
 	TManagedArray<int32> CollisionStructureID;
-	TManagedArray<int32> RigidBodyID;
+	TManagedArray<int32> RigidBodyID; // Deprecated.  Not added to dynamic collection.
 	TManagedArray<int32> SolverClusterID;
 	TManagedArray<FVector> InitialAngularVelocity;
 	TManagedArray<FVector> InitialLinearVelocity;
@@ -352,8 +410,7 @@
 	float CollisionParticlesPerObjectFraction;
 
 	// Double buffer of geom collection result data - TODO (Ryan): deprecated?
-	Chaos::TBufferedData<FGeometryCollectionResults> Results;
-	Chaos::TBufferedData<TArray<FMatrix>> ResultGlobalMatrices;
+	//Chaos::TBufferedData<FGeometryCollectionResults> Results;
 
 	// TODO (Ryan) - Currently this is using triple buffers for game-physics and 
 	// physics-game thread communication, but not for any reason other than this 
