Merging

//UE4/Private-Reverb/Engine/Source/Runtime/Experimental/ChaosSolvers/Private/PhysicsProxy/FieldSystemPhysicsProxy.cpp

to //UE4/Release-4.25/Engine/Source/Runtime/Experimental/ChaosSolvers/Private/PhysicsProxy/FieldSystemPhysicsProxy.cpp

diff --git a/Engine/Source/Runtime/Experimental/ChaosSolvers/Private/PhysicsProxy/FieldSystemPhysicsProxy.cpp b/Engine/Source/Runtime/Experimental/ChaosSolvers/Private/PhysicsProxy/FieldSystemPhysicsProxy.cpp
--- a/Engine/Source/Runtime/Experimental/ChaosSolvers/Private/PhysicsProxy/FieldSystemPhysicsProxy.cpp
+++ b/Engine/Source/Runtime/Experimental/ChaosSolvers/Private/PhysicsProxy/FieldSystemPhysicsProxy.cpp
@@ -184,6 +184,7 @@ void FFieldSystemPhysicsProxy::FieldParameterUpdateCallback(
 							// downgrading from Dynamic in those cases.
 							if (!bIsGC && HandleState != Chaos::EObjectStateType::Dynamic)
 							{
+								++i;
 								continue;
 							}
 
@@ -387,28 +388,36 @@ void FFieldSystemPhysicsProxy::FieldParameterUpdateCallback(
 					FFieldSystemPhysicsProxy::ContiguousIndices(IndicesArray, CurrentSolver, ResolutionType, IndicesArray.Num() != Particles.Size());
 					if (IndicesArray.Num())
 					{
-						TArrayView<ContextIndex> IndexView(&(IndicesArray[0]), IndicesArray.Num());
-
-						FVector * tptr = &(Particles.X(0));
-						TArrayView<FVector> SamplesView(tptr, int32(Particles.Size()));
+						TArrayView<FVector> SamplePointsView(&(SamplePoints[0]), SamplePoints.Num());
+						TArrayView<ContextIndex> SampleIndicesView(&(SampleIndices[0]), SampleIndices.Num());
+						FFieldContext Context(
+							SampleIndicesView,
+							SamplePointsView,
+							Command.MetaData);
 
-						FFieldContext Context{
-							IndexView,
-							SamplesView,
-							Command.MetaData
-						};
+						TArray<float> LocalResults;
+						LocalResults.AddZeroed(Handles.Num());
+						TArrayView<float> ResultsView(&(LocalResults[0]), LocalResults.Num());
+						static_cast<const FFieldNode<float>*>(Command.RootNode.Get())->Evaluate(
+							Context, ResultsView);
+						
 
-						// @todo(better) : Dont copy, support native type conversion. 
-						TArray<float> Results;
-						Results.AddUninitialized(Particles.Size());
-						for (const ContextIndex& Index : IndicesArray)
+						bool HasDisabled = false;
+						int32 i = 0;
+						for (Chaos::TGeometryParticleHandle<float, 3>* Handle : Handles)
 						{
-							Results[Index.Sample] = 0.f;
+							Chaos::TPBDRigidParticleHandle<float, 3>* RigidHandle = Handle->CastToRigidParticle();
+							
+							if (ResultsView[i] > 0.0)
+							{
+								CurrentSolver->GetEvolution()->DisableParticle(Handle);
+								HasDisabled = true;
+							}
+							++i;
 						}
-						TArrayView<float> ResultsView(&(Results[0]), Results.Num());
-						static_cast<const FFieldNode<float> *>(Command.RootNode.Get())->Evaluate(Context, ResultsView);
 
-#if TODO_REIMPLEMENT_GETFLOORINDEX
+						// #todo: no special casing for the floor
+						/*
 						bool bHasFloor = false;
 						int32 FloorIndex = CurrentSolver->GetFloorIndex();
 						if (FloorIndex != INDEX_NONE)
@@ -432,7 +441,7 @@ void FFieldSystemPhysicsProxy::FieldParameterUpdateCallback(
 							CurrentSolver->GetEvolution()->DisableParticle(FloorIndex);
 							Particles.SetObjectState(FloorIndex, Chaos::EObjectStateType::Static);
 						}
-#endif
+						*/
 					}
 				}
 				CommandsToRemove.Add(CommandIndex);
@@ -508,87 +517,55 @@ void FFieldSystemPhysicsProxy::FieldParameterUpdateCallback(
 				if (ensureMsgf(Command.RootNode->Type() == FFieldNode<float>::StaticType(),
 					TEXT("Field based evaluation of the simulations 'Disable' parameter expects scale field inputs.")))
 				{
-#if TODO_REIMPLEMENT_PHYSICS_PROXY_REVERSE_MAPPING
-					const Chaos::TArrayCollectionArray<PhysicsProxyWrapper>& PhysicsProxyMapping = CurrentSolver->GetPhysicsProxyReverseMapping();
-					
-					FFieldSystemPhysicsProxy::ContiguousIndices(IndicesArray, CurrentSolver, ResolutionType, IndicesArray.Num() != Particles.Size());
-					if (IndicesArray.Num())
+					if (Handles.Num())
 					{
-						TArrayView<ContextIndex> IndexView(&(IndicesArray[0]), IndicesArray.Num());
-
-						FVector * tptr = &(Particles.X(0));
-						TArrayView<FVector> SamplesView(tptr, int32(Particles.Size()));
+						TArrayView<FVector> SamplePointsView(&(SamplePoints[0]), SamplePoints.Num());
+						TArrayView<ContextIndex> SampleIndicesView(&(SampleIndices[0]), SampleIndices.Num());
+						FFieldContext Context(
+							SampleIndicesView,
+							SamplePointsView,
+							Command.MetaData);
 
-						FFieldContext Context{
-							IndexView,
-							SamplesView,
-							Command.MetaData
-						};
+						TArray<float> LocalResults;
+						LocalResults.AddZeroed(Handles.Num());
+						TArrayView<float> ResultsView(&(LocalResults[0]), LocalResults.Num());
+						static_cast<const FFieldNode<float>*>(Command.RootNode.Get())->Evaluate(
+							Context, ResultsView);
 
-						TArray<float> Results;
-						Results.AddUninitialized(Particles.Size());
-						for (const ContextIndex Index : IndicesArray)
+						int32 i = 0;
+						for (Chaos::TGeometryParticleHandle<float, 3>* Handle : Handles)
 						{
-							const SolverObjectWrapper& ParticleObjectWrapper = SolverObjectMapping[Index.Result];
-							TSerializablePtr<Chaos::FChaosPhysicsMaterial> Material = CurrentSolver->GetPhysicsMaterial(Index.Result);
-							if (ensure(Material) && ParticleObjectWrapper.SolverObject)
+							Chaos::TPBDRigidParticleHandle<float, 3>* RigidHandle = Handle->CastToRigidParticle();
+
+							if (RigidHandle && ResultsView.Num() > 0)
 							{
-								const TUniquePtr<Chaos::FChaosPhysicsMaterial>& InstanceMaterial = CurrentSolver->GetPerParticlePhysicsMaterial(Index.Result);
-								if (InstanceMaterial)
-								{
-									Results[Index.Result] = InstanceMaterial->SleepingLinearThreshold;
-								}
-								else
+								// if no per particle physics material is set, make one
+								if (!CurrentSolver->GetEvolution()->GetPerParticlePhysicsMaterial(RigidHandle).IsValid())
 								{
-									Results[Index.Result] = Material->SleepingLinearThreshold;
-								}
-							}
-							else
-							{
-								Results[Index.Result] = 0.0f;
-							}
-						}
 
-						TArrayView<float> ResultsView(&(Results[0]), Results.Num());
-						static_cast<const FFieldNode<float> *>(Command.RootNode.Get())->Evaluate(Context, ResultsView);
+									TUniquePtr<Chaos::FChaosPhysicsMaterial> NewMaterial = MakeUnique< Chaos::FChaosPhysicsMaterial>();
+									NewMaterial->SleepingLinearThreshold = ResultsView[i];
+									NewMaterial->SleepingAngularThreshold = ResultsView[i];
 
-						for (const ContextIndex Index : IndicesArray)
-						{
-							const int32 i = Index.Result;
-							const PhysicsProxyWrapper& ParticleObjectWrapper = PhysicsProxyMapping[i];
-							TSerializablePtr<FChaosPhysicsMaterial> Material = CurrentSolver->GetPhysicsMaterial(i);
-							if (!ensure(Material) || !ParticleObjectWrapper.PhysicsProxy)	//question: do we actually need to check for solver object?
-							{
-								continue;
-							}
 
-							//per instance override
-							if (!CurrentSolver->GetPerParticlePhysicsMaterial(Index.Result))
-							{
-								if (Results[i] != Material->SleepingLinearThreshold)
+									CurrentSolver->GetEvolution()->SetPhysicsMaterial(RigidHandle, MakeSerializable(NewMaterial));
+									CurrentSolver->GetEvolution()->SetPerParticlePhysicsMaterial(RigidHandle, NewMaterial);
+								}
+								else
 								{
-									// value changed from shared material, make unique material.
-									TUniquePtr<Chaos::FChaosPhysicsMaterial> NewMaterial = MakeUnique< Chaos::FChaosPhysicsMaterial>(*Material);
-									CurrentSolver->SetPerParticlePhysicsMaterial(Index.Result, MoveTemp(NewMaterial));
-									const TUniquePtr<FChaosPhysicsMaterial>& InstanceMaterial = CurrentSolver->GetPerParticlePhysicsMaterial(i);
+									const TUniquePtr<FChaosPhysicsMaterial>& InstanceMaterial = CurrentSolver->GetEvolution()->GetPerParticlePhysicsMaterial(RigidHandle);
 
-									InstanceMaterial->SleepingLinearThreshold = Results[i];
-									InstanceMaterial->SleepingAngularThreshold = Results[i];
+									if (ResultsView[i] != InstanceMaterial->DisabledLinearThreshold)
+									{
+										InstanceMaterial->SleepingLinearThreshold = ResultsView[i];
+										InstanceMaterial->SleepingAngularThreshold = ResultsView[i];
+									}
 								}
 							}
-							else
-							{
-								const TUniquePtr<FChaosPhysicsMaterial>& InstanceMaterial = CurrentSolver->GetPerParticlePhysicsMaterial(i);
 
-								if (InstanceMaterial->SleepingLinearThreshold != Results[i])
-								{
-									InstanceMaterial->SleepingLinearThreshold = Results[i];
-									InstanceMaterial->SleepingAngularThreshold = Results[i];
-								}
-							}
+							++i;
 						}
-					}
-#endif
+					}				
 				}
 				CommandsToRemove.Add(CommandIndex);
 			}
@@ -961,7 +938,7 @@ void FFieldSystemPhysicsProxy::FieldForcesUpdateCallback(
 							Command.MetaData);
 
 						TArray<FVector> LocalForce;
-						LocalForce.AddUninitialized(Handles.Num());					
+						LocalForce.AddZeroed(Handles.Num());					
 						TArrayView<FVector> ForceView(&(LocalForce[0]), LocalForce.Num());
 						static_cast<const FFieldNode<FVector> *>(Command.RootNode.Get())->Evaluate(Context, ForceView);
 		
@@ -970,7 +947,7 @@ void FFieldSystemPhysicsProxy::FieldForcesUpdateCallback(
 						{
 							Chaos::TPBDRigidParticleHandle<float, 3>* RigidHandle = Handle->CastToRigidParticle();
 							if(RigidHandle && RigidHandle->ObjectState() == Chaos::EObjectStateType::Dynamic)
-							{
+							{								
 								RigidHandle->F() += ForceView[i];
 							}
 							++i;
