Compile fix.

diff --git a/Engine/Source/Runtime/Experimental/Chaos/Public/Chaos/PBDRigidClustering.h b/Engine/Source/Runtime/Experimental/Chaos/Public/Chaos/PBDRigidClustering.h
--- a/Engine/Source/Runtime/Experimental/Chaos/Public/Chaos/PBDRigidClustering.h
+++ b/Engine/Source/Runtime/Experimental/Chaos/Public/Chaos/PBDRigidClustering.h
@@ -481,7 +481,7 @@ CleanCollisionParticles(
 	TArray<int32> CoincidentVertices;
 	const TArray<int32> Ordering = TriMesh.GetVertexImportanceOrdering(Vertices, &CoincidentVertices, true);
 	int32 NumGoodPoints = Ordering.Num() - CoincidentVertices.Num();
-	NumGoodPoints = std::min(NumGoodPoints, static_cast<int32>(ceil(NumGoodPoints * Fraction)))
+	NumGoodPoints = std::min(NumGoodPoints, static_cast<int32>(ceil(NumGoodPoints * Fraction)));
 
 	ResultingIndices.Reserve(NumGoodPoints);
 	for (int32 i = 0; i < NumGoodPoints; i++)
