Fixed some Chaos compile issues under clang.

diff --git a/Engine/Source/Runtime/Experimental/Chaos/Private/Chaos/PBDRigidClustering.cpp b/Engine/Source/Runtime/Experimental/Chaos/Private/Chaos/PBDRigidClustering.cpp
--- a/Engine/Source/Runtime/Experimental/Chaos/Private/Chaos/PBDRigidClustering.cpp
+++ b/Engine/Source/Runtime/Experimental/Chaos/Private/Chaos/PBDRigidClustering.cpp
@@ -761,8 +761,8 @@ namespace Chaos
 			TPBDRigidClusteredParticleHandle<T, d>* Child = Children[ChildIdx]->CastToClustered();
 			if (!Child)
 				continue;
-			if (ExternalStrainArray && (*ExternalStrainArray)[ChildIdx] ||
-			   !ExternalStrainArray && Child->CollisionImpulses() >= Child->Strain())
+			if ((ExternalStrainArray && (*ExternalStrainArray)[ChildIdx]) ||
+			   (!ExternalStrainArray && Child->CollisionImpulses() >= Child->Strain()))
 			{
 				// The piece that hits just breaks off - we may want more control 
 				// by looking at the edges of this piece which would give us cleaner 
diff --git a/Engine/Source/Runtime/Experimental/ChaosCore/Public/Chaos/ArrayCollectionArray.h b/Engine/Source/Runtime/Experimental/ChaosCore/Public/Chaos/ArrayCollectionArray.h
--- a/Engine/Source/Runtime/Experimental/ChaosCore/Public/Chaos/ArrayCollectionArray.h
+++ b/Engine/Source/Runtime/Experimental/ChaosCore/Public/Chaos/ArrayCollectionArray.h
@@ -39,7 +39,7 @@ class TArrayCollectionArray : public TArrayCollectionArrayBase, public TArray<T>
 
 	void Fill(const T& Value)
 	{
-		for (int32 Idx = 0; Idx < Num(); ++Idx)
+		for (int32 Idx = 0; Idx < TArray<T>::Num(); ++Idx)
 		{
 			TArray<T>::operator[](Idx) = Value;
 		}
