Fix for Chaos compilation under Switch.

Under Switch SharedPtr's default to thread safe (FORCE_THREADSAFE_SHAREDPTRS) so without an explicit definition we end up with a multiple definition compile error when we get to MANAGED_ARRAY_TYPE(Chaos::ThreadSafeSharedPtr_FImplicitObject, FImplicitObject3ThreadSafeSharedPointer)

Not sure if this is the right fix, but it is a fix...

diff --git a/Engine/Source/Runtime/Experimental/Chaos/Public/Chaos/ImplicitObject.h b/Engine/Source/Runtime/Experimental/Chaos/Public/Chaos/ImplicitObject.h
--- a/Engine/Source/Runtime/Experimental/Chaos/Public/Chaos/ImplicitObject.h
+++ b/Engine/Source/Runtime/Experimental/Chaos/Public/Chaos/ImplicitObject.h
@@ -328,4 +328,5 @@ FORCEINLINE FArchive& operator<<(FArchive& Ar, FImplicitObject& Value)
 
 typedef FImplicitObject FImplicitObject3;
 typedef TSharedPtr<Chaos::FImplicitObject, ESPMode::ThreadSafe> ThreadSafeSharedPtr_FImplicitObject;
+typedef TSharedPtr<Chaos::FImplicitObject, ESPMode::NotThreadSafe> NotThreadSafeSharedPtr_FImplicitObject;
 }
\ No newline at end of file
diff --git a/Engine/Source/Runtime/Experimental/GeometryCollectionCore/Public/GeometryCollection/ManagedArrayTypeValues.inl b/Engine/Source/Runtime/Experimental/GeometryCollectionCore/Public/GeometryCollection/ManagedArrayTypeValues.inl
--- a/Engine/Source/Runtime/Experimental/GeometryCollectionCore/Public/GeometryCollection/ManagedArrayTypeValues.inl
+++ b/Engine/Source/Runtime/Experimental/GeometryCollectionCore/Public/GeometryCollection/ManagedArrayTypeValues.inl
@@ -30,7 +30,7 @@ MANAGED_ARRAY_TYPE(TArray<FVector>*, VectorArrayPointer)
 MANAGED_ARRAY_TYPE(TUniquePtr<TArray<FVector>>, VectorArrayUniquePointer)
 MANAGED_ARRAY_TYPE(Chaos::FImplicitObject3*, FImplicitObject3Pointer)
 MANAGED_ARRAY_TYPE(TUniquePtr<Chaos::FImplicitObject3>, FImplicitObject3UniquePointer)
-MANAGED_ARRAY_TYPE(TSharedPtr<Chaos::FImplicitObject3>, FImplicitObject3SharedPointer)
+MANAGED_ARRAY_TYPE(Chaos::NotThreadSafeSharedPtr_FImplicitObject, FImplicitObject3SharedPointer)
 MANAGED_ARRAY_TYPE(Chaos::TSerializablePtr<Chaos::FImplicitObject3>, FImplicitObject3SerializablePtr)
 MANAGED_ARRAY_TYPE(Chaos::FBVHParticlesFloat3, FBVHParticlesFloat3Pointer)
 MANAGED_ARRAY_TYPE(TUniquePtr<Chaos::FBVHParticlesFloat3>, FBVHParticlesFloat3UniquePointer)
