CIS compile fix.

diff --git a/Engine/Source/Runtime/Experimental/ChaosSolvers/Private/PhysicsProxy/GeometryCollectionPhysicsProxy.cpp b/Engine/Source/Runtime/Experimental/ChaosSolvers/Private/PhysicsProxy/GeometryCollectionPhysicsProxy.cpp
--- a/Engine/Source/Runtime/Experimental/ChaosSolvers/Private/PhysicsProxy/GeometryCollectionPhysicsProxy.cpp
+++ b/Engine/Source/Runtime/Experimental/ChaosSolvers/Private/PhysicsProxy/GeometryCollectionPhysicsProxy.cpp
@@ -898,8 +898,8 @@ void FGeometryCollectionPhysicsProxy::InitializeBodiesPT(
 			// will optionally do this, but we switch that functionality off in BuildClusters().
 			ParallelFor(NumTransforms, [&](int32 TransformGroupIndex)
 			{
-				const TManagedArray<TSet<int32>>& Children = RestCollection->Children;
-				if (Children[TransformGroupIndex].Num() > 0)
+				const TManagedArray<TSet<int32>>& RestChildren = RestCollection->Children;
+				if (RestChildren[TransformGroupIndex].Num() > 0)
 				{
 					if (SolverClusterHandles[TransformGroupIndex])
 					{
