Paint Mode uses a different active-Tool model than Modeling Mode, where the user must always be in either Select or Paint tool, and can never be in a 'no-tool' state. Previously active-tool changes were not transacted, and 'undo' out of an active tool canceled it, and starting a new tool pushed a new transaction which discards future transaction. The default-tool behavior in UEdMode would as a result clear the future history if the user did Undo out of an active Tool.

Now this behavior is configurable, and ToolManager can optionally emit FChange transactions that will replay active-tool-change events. Update UMeshPaintMode to use this new alternative. Disable default tool behavior in UEdMode.

To implement this, ToolManager now tracks name of active tool and active tool builders.

#rnx

diff --git a/Engine/Plugins/MeshPainting/Source/MeshPaintEditorMode/Private/MeshPaintMode.cpp b/Engine/Plugins/MeshPainting/Source/MeshPaintEditorMode/Private/MeshPaintMode.cpp
--- a/Engine/Plugins/MeshPainting/Source/MeshPaintEditorMode/Private/MeshPaintMode.cpp
+++ b/Engine/Plugins/MeshPainting/Source/MeshPaintEditorMode/Private/MeshPaintMode.cpp
@@ -127,7 +127,10 @@ void UMeshPaintMode::Enter()
 	RegisterTool(ToolManagerCommands.TexturePaint, TEXT("TextureBrushTool"), NewObject<UMeshTexturePaintingToolBuilder>());
 	UpdateSelectedMeshes();
 
+	// disable tool change tracking to activate default tool, and then switch to full undo/redo tracking mode
+	GetToolManager()->ConfigureChangeTrackingMode(EToolChangeTrackingMode::NoChangeTracking);
 	ActivateDefaultTool();
+	GetToolManager()->ConfigureChangeTrackingMode(EToolChangeTrackingMode::FullUndoRedo);
 }
 
 void UMeshPaintMode::Exit()
diff --git a/Engine/Source/Editor/UnrealEd/Private/Tools/UEdMode.cpp b/Engine/Source/Editor/UnrealEd/Private/Tools/UEdMode.cpp
--- a/Engine/Source/Editor/UnrealEd/Private/Tools/UEdMode.cpp
+++ b/Engine/Source/Editor/UnrealEd/Private/Tools/UEdMode.cpp
@@ -396,7 +396,7 @@ UEdMode::UEdMode()
 	ToolsContext = nullptr;
 	ToolsContextClass = UInteractiveToolsContext::StaticClass();
 	ToolCommandList = MakeShareable(new FUICommandList);
-	bCheckIfDefaultToolNeeded = true;
+	bCheckIfDefaultToolNeeded = false;
 }
 
 
diff --git a/Engine/Source/Runtime/Experimental/InteractiveToolsFramework/Private/InteractiveToolManager.cpp b/Engine/Source/Runtime/Experimental/InteractiveToolsFramework/Private/InteractiveToolManager.cpp
--- a/Engine/Source/Runtime/Experimental/InteractiveToolsFramework/Private/InteractiveToolManager.cpp
+++ b/Engine/Source/Runtime/Experimental/InteractiveToolsFramework/Private/InteractiveToolManager.cpp
@@ -20,6 +20,8 @@ UInteractiveToolManager::UInteractiveToolManager()
 
 	ActiveRightBuilder = nullptr;
 	ActiveRightTool = nullptr;
+
+	ActiveToolChangeTrackingMode = EToolChangeTrackingMode::UndoToExit;
 }
 
 
@@ -64,10 +66,12 @@ bool UInteractiveToolManager::SelectActiveToolType(EToolSide Side, const FString
 		if (Side == EToolSide::Right)
 		{
 			ActiveRightBuilder = Builder;
+			ActiveRightBuilderName = Identifier;
 		}
 		else
 		{
 			ActiveLeftBuilder = Builder;
+			ActiveLeftBuilderName = Identifier;
 		}
 		return true;
 	}
@@ -97,23 +101,56 @@ bool UInteractiveToolManager::ActivateTool(EToolSide Side)
 {
 	check(Side == EToolSide::Left);   // TODO: support right-side tool
 
-	if (ActiveLeftTool != nullptr) 
+	// wrap tool change in a transaction so that deactivate and activate are grouped
+	bool bInTransaction = false;
+	if (ActiveToolChangeTrackingMode == EToolChangeTrackingMode::FullUndoRedo)
+	{
+		BeginUndoTransaction(LOCTEXT("ToolChange", "Change Tool"));
+		bInTransaction = true;
+	}
+
+	if (ActiveLeftTool != nullptr)
 	{
 		DeactivateTool(EToolSide::Left, EToolShutdownType::Accept);
 	}
 
-	if (ActiveLeftBuilder == nullptr) 
+	if (ActiveLeftBuilder == nullptr || ActivateToolInternal(Side) == false)
 	{
+		if (bInTransaction)
+		{
+			EndUndoTransaction();
+		}
 		return false;
 	}
 
+	if (ActiveToolChangeTrackingMode == EToolChangeTrackingMode::FullUndoRedo)
+	{
+		check(TransactionsAPI);
+		TransactionsAPI->AppendChange(this, MakeUnique<FActivateToolChange>(Side, ActiveLeftToolName), LOCTEXT("ActivateToolChange", "Activate Tool"));
+	} 
+	else if (ActiveToolChangeTrackingMode == EToolChangeTrackingMode::UndoToExit)
+	{
+		EmitObjectChange(this, MakeUnique<FBeginToolChange>(), LOCTEXT("ActivateToolChange", "Activate Tool"));
+	}
+
+	if (bInTransaction)
+	{
+		EndUndoTransaction();
+	}
+
+	return true;
+}
+
+
+bool UInteractiveToolManager::ActivateToolInternal(EToolSide Side)
+{
 	// construct input state we will pass to tools
 	FToolBuilderState InputState;
 	QueriesAPI->GetCurrentSelectionState(InputState);
 
 	if (ActiveLeftBuilder->CanBuildTool(InputState) == false)
 	{
-		TransactionsAPI->DisplayMessage( LOCTEXT("ActivateToolCanBuildFailMessage", "UInteractiveToolManager::ActivateTool: CanBuildTool returned false."), EToolMessageLevel::Internal);
+		TransactionsAPI->DisplayMessage(LOCTEXT("ActivateToolCanBuildFailMessage", "UInteractiveToolManager::ActivateTool: CanBuildTool returned false."), EToolMessageLevel::Internal);
 		return false;
 	}
 
@@ -122,6 +159,7 @@ bool UInteractiveToolManager::ActivateTool(EToolSide Side)
 	{
 		return false;
 	}
+	ActiveLeftToolName = ActiveLeftBuilderName;
 
 	ActiveLeftTool->Setup();
 
@@ -132,19 +170,32 @@ bool UInteractiveToolManager::ActivateTool(EToolSide Side)
 
 	OnToolStarted.Broadcast(this, ActiveLeftTool);
 
-	// emit a change so that we can undo to cancel the tool
-	EmitObjectChange(this, MakeUnique<FBeginToolChange>(), LOCTEXT("ActivateToolChange", "Activate Tool"));
-
 	return true;
 }
 
 
+
 void UInteractiveToolManager::DeactivateTool(EToolSide Side, EToolShutdownType ShutdownType)
 {
 	check(Side == EToolSide::Left);   // TODO: support right-side tool
-
 	if (ActiveLeftTool != nullptr)
 	{
+		if (ActiveToolChangeTrackingMode == EToolChangeTrackingMode::FullUndoRedo)
+		{
+			check(TransactionsAPI);
+			TransactionsAPI->AppendChange(this, MakeUnique<FActivateToolChange>(Side, ActiveLeftToolName, ShutdownType), LOCTEXT("DeactivateToolChange", "Deactivate Tool"));
+		}
+
+		DeactivateToolInternal(Side, ShutdownType);
+	}
+}
+
+
+void UInteractiveToolManager::DeactivateToolInternal(EToolSide Side, EToolShutdownType ShutdownType)
+{
+	if (Side == EToolSide::Left)
+	{
+		check(ActiveLeftTool);
 		InputRouter->ForceTerminateSource(ActiveLeftTool);
 
 		ActiveLeftTool->Shutdown(ShutdownType);
@@ -153,6 +204,7 @@ void UInteractiveToolManager::DeactivateTool(EToolSide Side, EToolShutdownType S
 
 		UInteractiveTool* DoneTool = ActiveLeftTool;
 		ActiveLeftTool = nullptr;
+		ActiveLeftToolName.Empty();
 
 		PostInvalidation();
 
@@ -161,6 +213,7 @@ void UInteractiveToolManager::DeactivateTool(EToolSide Side, EToolShutdownType S
 }
 
 
+
 bool UInteractiveToolManager::HasActiveTool(EToolSide Side) const
 {
 	return (Side == EToolSide::Left) ? (ActiveLeftTool != nullptr) : (ActiveRightTool != nullptr);
@@ -203,6 +256,10 @@ bool UInteractiveToolManager::CanCancelActiveTool(EToolSide Side)
 
 
 
+void UInteractiveToolManager::ConfigureChangeTrackingMode(EToolChangeTrackingMode ChangeMode)
+{
+	ActiveToolChangeTrackingMode = ChangeMode;
+}
 
 
 
@@ -299,7 +356,7 @@ void FBeginToolChange::Revert(UObject* Object)
 	UInteractiveToolManager* ToolManager = CastChecked<UInteractiveToolManager>(Object);
 	if (ToolManager->HasAnyActiveTool())
 	{
-		ToolManager->DeactivateTool(EToolSide::Left, EToolShutdownType::Cancel);
+		ToolManager->DeactivateToolInternal(EToolSide::Left, EToolShutdownType::Cancel);
 	}
 }
 
@@ -318,6 +375,57 @@ FString FBeginToolChange::ToString() const
 
 
 
+void FActivateToolChange::Apply(UObject* Object)
+{
+	UInteractiveToolManager* ToolManager = CastChecked<UInteractiveToolManager>(Object);
+	if (ToolManager)
+	{
+		if (bIsDeactivate)
+		{
+			ToolManager->DeactivateToolInternal(Side, ShutdownType);
+		}
+		else
+		{
+			ToolManager->SelectActiveToolType(Side, ToolType);
+			ToolManager->ActivateToolInternal(Side);
+		}
+	}
+}
+
+void FActivateToolChange::Revert(UObject* Object)
+{
+	UInteractiveToolManager* ToolManager = CastChecked<UInteractiveToolManager>(Object);
+	if (ToolManager)
+	{
+		if (bIsDeactivate)
+		{
+			ToolManager->SelectActiveToolType(Side, ToolType);
+			ToolManager->ActivateToolInternal(Side);
+		}
+		else
+		{
+			ToolManager->DeactivateToolInternal(Side, ShutdownType);
+		}
+	}
+}
+
+bool FActivateToolChange::HasExpired(UObject* Object) const
+{
+	UInteractiveToolManager* ToolManager = CastChecked<UInteractiveToolManager>(Object);
+	return (ToolManager == nullptr);
+}
+
+FString FActivateToolChange::ToString() const
+{
+	return FString(TEXT("Change Tool"));
+}
+
+
+
+
+
+
+
 void FToolChangeWrapperChange::Apply(UObject* Object)
 {
 	if (ToolChange.IsValid())
diff --git a/Engine/Source/Runtime/Experimental/InteractiveToolsFramework/Public/InteractiveToolManager.h b/Engine/Source/Runtime/Experimental/InteractiveToolsFramework/Public/InteractiveToolManager.h
--- a/Engine/Source/Runtime/Experimental/InteractiveToolsFramework/Public/InteractiveToolManager.h
+++ b/Engine/Source/Runtime/Experimental/InteractiveToolsFramework/Public/InteractiveToolManager.h
@@ -25,6 +25,23 @@ enum class EToolSide
 };
 
 
+/**
+ * UInteractiveToolManager can emit change events for the active tool in various ways.
+ * This allows different modes to control how tools activate/deactivate on undo/redo, which is necessary
+ * because some modes (eg Modeling Mode) do not support redo "into" a Tool, while others require it (like Paint Mode)
+ */
+UENUM()
+enum class EToolChangeTrackingMode
+{
+	/** Do not emit any Active Tool change events */
+	NoChangeTracking = 1,
+	/** When Activating a new Tool, emit a change that will cancel/deactivate that Tool on Undo, but not reactivate it on Redo */
+	UndoToExit = 2,
+	/** Full change tracking of active Tool. Note that on Activation when an existing Tool is auto-shutdown, two separate FChanges are emitted, wrapped in a single Transaction */
+	FullUndoRedo = 3
+};
+
+
 /**
  * UInteractiveToolManager allows users of the tools framework to create and operate Tool instances.
  * For each Tool, a (string,ToolBuilder) pair is registered with the ToolManager.
@@ -142,6 +159,11 @@ public:
 	virtual void DeactivateTool(EToolSide Side, EToolShutdownType ShutdownType);
 
 
+	/**
+	 * Configure how tool changes emit change events. See EToolChangeTrackingMode for details.
+	 */
+	virtual void ConfigureChangeTrackingMode(EToolChangeTrackingMode ChangeMode);
+
 
 	//
 	// Functions that Tools can call to interact with Transactions API
@@ -233,10 +255,23 @@ protected:
 	TMap<FString, UInteractiveToolBuilder*> ToolBuilders;
 
 	/** Currently-active Left ToolBuilder */
+	FString ActiveLeftBuilderName;
 	UInteractiveToolBuilder* ActiveLeftBuilder;
 	/** Currently-active Right ToolBuilder */
+	FString ActiveRightBuilderName;
 	UInteractiveToolBuilder* ActiveRightBuilder;
 
+	EToolChangeTrackingMode ActiveToolChangeTrackingMode;
+
+	FString ActiveLeftToolName;
+	FString ActiveRightToolName;
+
+
+	virtual bool ActivateToolInternal(EToolSide Side);
+	virtual void DeactivateToolInternal(EToolSide Side, EToolShutdownType ShutdownType);
+
+	friend class FBeginToolChange;
+	friend class FActivateToolChange;
 };
 
 
@@ -258,6 +293,31 @@ public:
 	virtual FString ToString() const override;
 };
 
+/**
+ * FActivateToolChange is used by UInteractiveToolManager to change the active tool.
+ * This Change has two modes, either activating or deactivating.
+ */
+class INTERACTIVETOOLSFRAMEWORK_API FActivateToolChange : public FToolCommandChange
+{
+public:
+	EToolSide Side;
+	FString ToolType;
+	bool bIsDeactivate = false;
+	EToolShutdownType ShutdownType;
+
+	FActivateToolChange(EToolSide SideIn, FString ToolTypeIn)
+		: Side(SideIn), ToolType(ToolTypeIn), bIsDeactivate(false) {}
+	FActivateToolChange(EToolSide SideIn, FString ToolTypeIn, EToolShutdownType ShutdownTypeIn)
+		: Side(SideIn), ToolType(ToolTypeIn), bIsDeactivate(true), ShutdownType(ShutdownTypeIn) {}
+
+	virtual void Apply(UObject* Object) override;
+	virtual void Revert(UObject* Object) override;
+	virtual bool HasExpired(UObject* Object) const override;
+	virtual FString ToString() const override;
+};
+
+
+
 
 
 /**
