Make static analysis happy.  Check for null, even though it should be uncecessary.  Remove a redundant conditional, even though we'll need it soon.

diff --git a/Engine/Source/Runtime/Experimental/ChaosSolvers/Private/PhysicsProxy/GeometryCollectionPhysicsProxy.cpp b/Engine/Source/Runtime/Experimental/ChaosSolvers/Private/PhysicsProxy/GeometryCollectionPhysicsProxy.cpp
--- a/Engine/Source/Runtime/Experimental/ChaosSolvers/Private/PhysicsProxy/GeometryCollectionPhysicsProxy.cpp
+++ b/Engine/Source/Runtime/Experimental/ChaosSolvers/Private/PhysicsProxy/GeometryCollectionPhysicsProxy.cpp
@@ -1266,7 +1266,7 @@ void FGeometryCollectionPhysicsProxy::CreateDynamicAttributes()
 		{
 			Implicits[Index] = (*SharedRestImplicits)[Index];
 		}
-		else if(UniqueRestImplicits)
+		else if(UniqueRestImplicits && RestCollectionSharedImplicits) // RestCollectionSharedImplicits is redundant, but static analysis build complains without it
 		{
 			// Steal the memory from the unique ptr.
 			Implicits[Index] = TSharedPtr<Chaos::FImplicitObject, ESPMode::ThreadSafe>((*UniqueRestImplicits)[Index].Release());
@@ -2525,8 +2525,12 @@ void FGeometryCollectionPhysicsProxy::UpdateGeometryCollection(
 		//   take ownership of the parents during the simulation, so it might be necessary to force deactivated
 		//   bodies out of the collections hierarchy during the simulation.  
 
+		// Static analysis build complains that the else clause does the same thing as the if/then,
+		// so remove the conditional for now...
+#if 0
 		if(!GCResults.DisabledStates[TransformGroupIndex])
 		{
+#endif
 			// Update the transform of the active body. The active body can be either a single rigid
 			// or a collection of rigidly attached geometries (Clustering). The cluster is represented as a
 			// single transform in the GeometryCollection, and all children are stored in the local space
@@ -2544,11 +2548,13 @@ void FGeometryCollectionPhysicsProxy::UpdateGeometryCollection(
 			// body will become active and needs its clusterID updated. This just
 			// syncs the clusterID all the time. 
 //			CollectionClusterID[TransformGroupIndex] = ClusterID[RigidBodyIndex].Id;
+#if 0
 		}
 		else //if (Particles.Disabled(RigidBodyIndex))
 		{
 			Transform[TransformGroupIndex] = GCResults.Transforms[TransformGroupIndex];
 		}
+#endif
 	}
 }
 
