Clustering fixes.  Make the logic that arranges geometry collection particles into particle SOA views simpler by adding a "Sleeping" array, and not treating sleeping geometry collection particles as static.  Also, only regenerate geometry collection views on pull, rather than on push - which is every authoring event.

Adds SOA particle serialization for clustered and geometry collection particles, but commented out due to ExternalPhysicsCustomObjectVersion.h being locked.

diff --git a/Engine/Source/Runtime/Experimental/Chaos/Private/Chaos/PBDRigidClustering.cpp b/Engine/Source/Runtime/Experimental/Chaos/Private/Chaos/PBDRigidClustering.cpp
--- a/Engine/Source/Runtime/Experimental/Chaos/Private/Chaos/PBDRigidClustering.cpp
+++ b/Engine/Source/Runtime/Experimental/Chaos/Private/Chaos/PBDRigidClustering.cpp
@@ -466,9 +466,10 @@ namespace Chaos
 		{
 			MChildren.Add(NewParticle, MoveTemp(Children));
 		}
-
 		const TArray<TPBDRigidParticleHandle<T, d>*>& ChildrenArray = MChildren[NewParticle];
 		TSet<TPBDRigidParticleHandle<T, d>*> ChildrenSet(ChildrenArray);
+
+		// Disable the children
 		MEvolution.DisableParticles(reinterpret_cast<TSet<TGeometryParticleHandle<T, d>*>&>(ChildrenSet));
 
 		bool bClusterIsAsleep = true;
@@ -483,7 +484,7 @@ namespace Chaos
 				// Cluster group id 0 means "don't union with other things"
 				// TODO: Use INDEX_NONE instead of 0?
 				ClusteredChild->SetClusterGroupIndex(0);
-				ClusteredChild->ClusterIds().Id = NewParticle; // ClusterId needs to change?
+				ClusteredChild->ClusterIds().Id = NewParticle;
 				NewParticle->Strains() += ClusteredChild->Strains();
 
 				NewParticle->SetCollisionImpulses(
@@ -768,7 +769,11 @@ namespace Chaos
 				// by looking at the edges of this piece which would give us cleaner 
 				// breaks (this approach produces more rubble)
 				RemoveChildLambda(Child);
-				Children.RemoveAtSwap(ChildIdx);
+
+				// Remove from the children array without freeing memory yet. 
+				// We're looping over Children and it'd be silly to free the array
+				// 1 entry at a time.
+				Children.RemoveAtSwap(ChildIdx, 1, false);
 
 				if (Child->ToBeRemovedOnFracture())
 				{
@@ -792,11 +797,15 @@ namespace Chaos
 
 		if (bChildrenChanged)
 		{
+			if (Children.Num() == 0)
+			{
+				// Free the memory if we can do so cheaply (no data copies).
+				Children.Empty(); 
+			}
+
 			if (UseConnectivity)
 			{
-				//cluster may have contained forests so find the connected pieces and cluster them together
-				TSet<TPBDRigidParticleHandle<T,d>*> PotentialActivatedChildren;
-				PotentialActivatedChildren.Append(Children);
+				// The cluster may have contained forests, so find the connected pieces and cluster them together.
 
 				//first update the connected graph of the children we already removed
 				for (TPBDRigidParticleHandle<T,d>* Child : ActivatedChildren)
@@ -804,23 +813,30 @@ namespace Chaos
 					RemoveNodeConnections(Child);
 				}
 
-				if (PotentialActivatedChildren.Num())
+				if (Children.Num())
 				{
 					TArray<TArray<TPBDRigidParticleHandle<T,d>*>> ConnectedPiecesArray;
-					//traverse connectivity and see how many connected pieces we have
-					TSet<TPBDRigidParticleHandle<T,d>*> ProcessedChildren;
-					for (TPBDRigidParticleHandle<T,d>* PotentialActivatedChild : PotentialActivatedChildren)
-					{
-						if (!ProcessedChildren.Contains(PotentialActivatedChild))
+
+					{ // tmp scope
+
+						//traverse connectivity and see how many connected pieces we have
+						TSet<TPBDRigidParticleHandle<T, d>*> ProcessedChildren;
+						ProcessedChildren.Reserve(Children.Num());
+
+						for (TPBDRigidParticleHandle<T, d>* PotentialActivatedChild : Children)
 						{
+							if (ProcessedChildren.Contains(PotentialActivatedChild))
+							{
+								continue;
+							}
 							ConnectedPiecesArray.AddDefaulted();
-							TArray<TPBDRigidParticleHandle<T,d>*>& ConnectedPieces = ConnectedPiecesArray.Last();
+							TArray<TPBDRigidParticleHandle<T, d>*>& ConnectedPieces = ConnectedPiecesArray.Last();
 
-							TArray<TPBDRigidParticleHandle<T,d>*> ProcessingQueue;
+							TArray<TPBDRigidParticleHandle<T, d>*> ProcessingQueue;
 							ProcessingQueue.Add(PotentialActivatedChild);
 							while (ProcessingQueue.Num())
 							{
-								TPBDRigidParticleHandle<T,d>* Child = ProcessingQueue.Pop();
+								TPBDRigidParticleHandle<T, d>* Child = ProcessingQueue.Pop();
 								if (!ProcessedChildren.Contains(Child))
 								{
 									ProcessedChildren.Add(Child);
@@ -835,7 +851,7 @@ namespace Chaos
 								}
 							}
 						}
-					}
+					} // tmp scope
 
 					int32 NumNewClusters = 0;
 					for (TArray<TPBDRigidParticleHandle<T,d>*>& ConnectedPieces : ConnectedPiecesArray)
@@ -890,7 +906,7 @@ namespace Chaos
 
 			//disable cluster
 			DisableCluster(ClusteredParticle);
-		}
+		} // bChildrenChanged
 
 		return ActivatedChildren;
 	}
diff --git a/Engine/Source/Runtime/Experimental/Chaos/Public/Chaos/PBDRigidsSOAs.h b/Engine/Source/Runtime/Experimental/Chaos/Public/Chaos/PBDRigidsSOAs.h
--- a/Engine/Source/Runtime/Experimental/Chaos/Public/Chaos/PBDRigidsSOAs.h
+++ b/Engine/Source/Runtime/Experimental/Chaos/Public/Chaos/PBDRigidsSOAs.h
@@ -115,7 +115,8 @@ public:
 
 		GeometryCollectionParticles = MakeUnique<TPBDGeometryCollectionParticles<T, d>>();
 		GeometryCollectionParticles->RemoveParticleBehavior() = ERemoveParticleBehavior::Remove;	//clustered particles maintain relative ordering
-		
+		bGeometryCollectionDirty = false;
+
 		UpdateViews();
 	}
 
@@ -129,7 +130,7 @@ public:
 	
 	TArray<TGeometryParticleHandle<T, d>*> CreateStaticParticles(int32 NumParticles, const FUniqueIdx* ExistingIndices = nullptr, const TGeometryParticleParameters<T, d>& Params = TGeometryParticleParameters<T, d>())
 	{
-		auto Results =  CreateParticlesHelper<TGeometryParticleHandle<T, d>>(NumParticles, ExistingIndices, Params.bDisabled ? StaticDisabledParticles : StaticParticles, Params);
+		auto Results = CreateParticlesHelper<TGeometryParticleHandle<T, d>>(NumParticles, ExistingIndices, Params.bDisabled ? StaticDisabledParticles : StaticParticles, Params);
 		UpdateViews();
 		return Results;
 	}
@@ -154,13 +155,20 @@ public:
 	{
 		TArray<TPBDGeometryCollectionParticleHandle<T, d>*> Results = CreateParticlesHelper<TPBDGeometryCollectionParticleHandle<T, d>>(
 			NumParticles, ExistingIndices, GeometryCollectionParticles, Params);
-		//TArray<TPBDRigidParticleHandle<T, d>*>& RigidHandles = (TArray<TPBDRigidParticleHandle<T, d>*>*)&Results;//*static_cast<TArray<TPBDRigidParticleHandle<T, d>*>*>(&Results);
-		if (!Params.bStartSleeping)
+		for (auto* Handle : Results)
 		{
-			InsertToMapAndArray(Results, ActiveGeometryCollectionToIndex, ActiveGeometryCollectionArray);
+			if (Params.bStartSleeping)
+			{
+				Handle->SetObjectState(Chaos::EObjectStateType::Sleeping);
+				Handle->SetSleeping(true);
+			}
+			else
+			{
+				Handle->SetObjectState(Chaos::EObjectStateType::Dynamic);
+				Handle->SetSleeping(false);
+			}
 		}
-		UpdateGeometryCollectionViews();
-		UpdateViews();
+		bGeometryCollectionDirty = true;
 		return Results;
 	}
 
@@ -215,8 +223,14 @@ public:
 		UpdateViews();
 	}
 
+	/**
+	 * A disabled particle is ignored by the solver.
+	 */
 	void DisableParticle(TGeometryParticleHandle<T, d>* Particle)
 	{
+		// Rigid particles express their disabled state with a boolean.
+		// Disabled kinematic and static particles get shuffled to differnt SOAs.
+
 		if (auto PBDRigid = Particle->CastToRigidParticle())
 		{
 			PBDRigid->Disabled() = true;
@@ -225,11 +239,20 @@ public:
 
 			if (auto PBDRigidClustered = Particle->CastToClustered())
 			{
-				RemoveFromMapAndArray(PBDRigidClustered, NonDisabledClusteredToIndex, NonDisabledClusteredArray);
-				RemoveFromMapAndArray(PBDRigidClustered, ActiveClusteredToIndex, ActiveClusteredArray);
 				if (Particle->GetParticleType() == Chaos::EParticleType::GeometryCollection)
 				{
-					UpdateGeometryCollectionViews();
+					bGeometryCollectionDirty = true;
+					return;
+				}
+				else // clustered
+				{
+					RemoveFromMapAndArray(PBDRigidClustered, 
+						NonDisabledClusteredToIndex, NonDisabledClusteredArray);
+					if (Particle->ObjectState() == EObjectStateType::Dynamic)
+					{
+						RemoveFromMapAndArray(PBDRigidClustered, 
+							ActiveClusteredToIndex, ActiveClusteredArray);
+					}
 				}
 			}
 			else
@@ -237,6 +260,7 @@ public:
 				Particle->MoveToSOA(*DynamicDisabledParticles);
 			}
 
+			// All active particles RIGID particles
 			if (Particle->ObjectState() == EObjectStateType::Dynamic)
 			{
 				RemoveFromMapAndArray(PBDRigid, ActiveParticlesToIndex, ActiveParticlesArray);
@@ -255,31 +279,44 @@ public:
 
 	void EnableParticle(TGeometryParticleHandle<T, d>* Particle)
 	{
+		// Rigid particles express their disabled state with a boolean.
+		// Disabled kinematic and static particles get shuffled to differnt SOAs.
+
 		if (auto PBDRigid = Particle->CastToRigidParticle())
 		{
+			PBDRigid->Disabled() = false;
+			// DisableParticle() zeros V and W.  We do nothing here and assume the client
+			// sets appropriate values.
+
 			if (auto PBDRigidClustered = Particle->CastToClustered())
 			{
-				InsertToMapAndArray(PBDRigidClustered, NonDisabledClusteredToIndex, NonDisabledClusteredArray);
-				if (!PBDRigid->Sleeping() && Particle->ObjectState() == EObjectStateType::Dynamic)
+				if (Particle->GetParticleType() == Chaos::EParticleType::GeometryCollection)
 				{
-					InsertToMapAndArray(PBDRigidClustered, ActiveClusteredToIndex, ActiveClusteredArray);
+					bGeometryCollectionDirty = true;
+					return;
 				}
-				if (Particle->GetParticleType() == Chaos::EParticleType::GeometryCollection)
+				else // clustered
 				{
-					UpdateGeometryCollectionViews();
+					InsertToMapAndArray(PBDRigidClustered, 
+						NonDisabledClusteredToIndex, NonDisabledClusteredArray);
+					if (!PBDRigid->Sleeping() && Particle->ObjectState() == EObjectStateType::Dynamic)
+					{
+						// Clustered, enabled, (dynamic, !sleeping)
+						InsertToMapAndArray(PBDRigidClustered, 
+							ActiveClusteredToIndex, ActiveClusteredArray);
+					}
 				}
 			}
 			else
 			{
-				SetDynamicParticleSOA(Particle->CastToRigidParticle());
+				SetDynamicParticleSOA(PBDRigid);
 			}
 
 			if (!PBDRigid->Sleeping() && Particle->ObjectState() == EObjectStateType::Dynamic)
 			{
-				InsertToMapAndArray(PBDRigid, ActiveParticlesToIndex, ActiveParticlesArray);
+				InsertToMapAndArray(PBDRigid, 
+					ActiveParticlesToIndex, ActiveParticlesArray);
 			}
-
-			PBDRigid->Disabled() = false;
 		}
 		else if (Particle->CastToKinematicParticle())
 		{
@@ -292,54 +329,129 @@ public:
 		UpdateViews();
 	}
 
+	/**
+	 * Wake a sleeping dynamic non-disabled particle.
+	 */
 	void ActivateParticle(TGeometryParticleHandle<T, d>* Particle)
 	{
-		auto PBDRigid = Particle->CastToRigidParticle();
-		if(PBDRigid && PBDRigid->ObjectState() == EObjectStateType::Dynamic)
+		if (auto PBDRigid = Particle->CastToRigidParticle())
 		{
-			if (ensure(!PBDRigid->Disabled()))
+			if (PBDRigid->ObjectState() == EObjectStateType::Sleeping ||
+				PBDRigid->ObjectState() == EObjectStateType::Dynamic)
 			{
-				if (auto PBDRigidClustered = Particle->CastToClustered())
+				if (ensure(!PBDRigid->Disabled()))
 				{
-					InsertToMapAndArray(PBDRigidClustered, ActiveClusteredToIndex, ActiveClusteredArray);
-				}
+					// Sleeping state is currently expressed in 2 places...
+					PBDRigid->SetSleeping(false);
+					PBDRigid->SetObjectState(EObjectStateType::Dynamic);
+		
+					bool bUpdateGeometryCollection = false;
+					if (auto PBDRigidClustered = Particle->CastToClustered())
+					{
+						if (Particle->GetParticleType() == Chaos::EParticleType::GeometryCollection)
+						{
+							bGeometryCollectionDirty = true;
+							return;
+						}
+						else
+						{
+							// Clustered, non geometry collection:
+							InsertToMapAndArray(PBDRigidClustered, 
+								ActiveClusteredToIndex, ActiveClusteredArray);
+						}
+					}
+					else
+					{
+						// Non clustered rigid particles:
+						InsertToMapAndArray(PBDRigid, 
+							ActiveParticlesToIndex, ActiveParticlesArray);
+					}
 
-				InsertToMapAndArray(PBDRigid, ActiveParticlesToIndex, ActiveParticlesArray);
+					UpdateViews();
+				}
 			}
 		}
-		
-		UpdateViews();
 	}
 
-	void DeactivateParticle(TGeometryParticleHandle<T, d>* Particle)
+	/**
+	 * Wake multiple dynamic non-disabled particles.
+	 */
+	void ActivateParticles(const TArray<TGeometryParticleHandle<T, d>*>& Particles)
 	{
-		auto PBDRigid = Particle->CastToRigidParticle();
-		if(PBDRigid)
+		bool bUpdateGeometryCollection = false;
+		for (auto Particle : Particles)
+		{
+			bUpdateGeometryCollection |= ActivateParticle(Particle);
+		}
+		if (bUpdateGeometryCollection)
+		{
+			UpdateGeometryCollectionViews();
+		}
+		else
+		{
+			UpdateViews();
+		}
+	}
+
+	/**
+	 * Put a non-disabled dynamic particle to sleep.
+	 *
+	 * If \p DeferUpdateViews is \c true, then it's assumed this function
+	 * is being called in a loop and it won't update the SOA view arrays.
+	 */
+	void DeactivateParticle(
+		TGeometryParticleHandle<T, d>* Particle,
+		const bool DeferUpdateViews=false)
+	{
+		if(auto PBDRigid = Particle->CastToRigidParticle())
 		{
-			if (   PBDRigid->ObjectState() == EObjectStateType::Dynamic
-				|| PBDRigid->ObjectState() == EObjectStateType::Sleeping)
+			if (PBDRigid->ObjectState() == EObjectStateType::Dynamic ||
+				PBDRigid->ObjectState() == EObjectStateType::Sleeping)
 			{
 				if (ensure(!PBDRigid->Disabled()))
 				{
+					// Sleeping state is currently expressed in 2 places...
+					PBDRigid->SetSleeping(true);
+					PBDRigid->SetObjectState(EObjectStateType::Sleeping);
+
 					if (auto PBDRigidClustered = Particle->CastToClustered())
 					{
-						RemoveFromMapAndArray(PBDRigidClustered, ActiveClusteredToIndex, ActiveClusteredArray);
+						if (Particle->GetParticleType() == Chaos::EParticleType::GeometryCollection)
+						{
+							bGeometryCollectionDirty = true;
+						}
+						else
+						{
+							RemoveFromMapAndArray(PBDRigidClustered, 
+								ActiveClusteredToIndex, ActiveClusteredArray);
+						}
+					}
+					else
+					{
+						RemoveFromMapAndArray(PBDRigid, 
+							ActiveParticlesToIndex, ActiveParticlesArray);
 					}
 
-					RemoveFromMapAndArray(PBDRigid, ActiveParticlesToIndex, ActiveParticlesArray);
+					if (!DeferUpdateViews)
+					{
+						UpdateViews();
+					}
 				}
 			}
 		}
-
-		UpdateViews();
 	}
 
+	/**
+	 * Put multiple dynamic non-disabled particles to sleep.
+	 */
 	void DeactivateParticles(const TArray<TGeometryParticleHandle<T, d>*>& Particles)
 	{
 		for (auto Particle : Particles)
 		{
-			DeactivateParticle(Particle);
+			DeactivateParticle(Particle, true);
 		}
+		UpdateIfNeeded();
+		UpdateViews();
 	}
 
 	void SetDynamicParticleSOA(TPBDRigidParticleHandle<T, d>* Particle)
@@ -389,6 +501,9 @@ public:
 		Ar << KinematicDisabledParticles;
 		Ar << DynamicParticles;
 		Ar << DynamicDisabledParticles;
+		// TODO: Add an entry in UObject/ExternalPhysicsCustomObjectVersion.h when adding these back in:
+		//Ar << ClusteredParticles;
+		//Ar << GeometryCollectionParticles;
 
 		Ar.UsingCustomVersion(FExternalPhysicsCustomObjectVersion::GUID);
 		if (Ar.CustomVer(FExternalPhysicsCustomObjectVersion::GUID) >= FExternalPhysicsCustomObjectVersion::AddDynamicKinematicSOA)
@@ -414,6 +529,8 @@ public:
 			AssignIdxHelper(DynamicParticles);
 			AssignIdxHelper(DynamicParticles);
 			AssignIdxHelper(DynamicDisabledParticles);
+			//AssignIdxHelper(ClusteredParticles);
+			//AssignIdxHelper(GeometryCollectionParticles);
 		}
 
 		ensure(ClusteredParticles->Size() == 0);	//not supported yet
@@ -426,18 +543,22 @@ public:
 
 		//todo: update deterministic ID
 
-		UpdateViews();
+		//if (!GeometryCollectionParticles || !GeometryCollectionParticles->Size())
+			UpdateViews();
+		//else
+		//	UpdateGeometryCollectionViews();
 	}
 
 
-	const TParticleView<TGeometryParticles<T, d>>& GetNonDisabledView() const { return NonDisabledView; }
+	const TParticleView<TGeometryParticles<T, d>>& GetNonDisabledView() const { UpdateIfNeeded();  return NonDisabledView; }
 
-	const TParticleView<TPBDRigidParticles<T, d>>& GetNonDisabledDynamicView() const { return NonDisabledDynamicView; }
+	const TParticleView<TPBDRigidParticles<T, d>>& GetNonDisabledDynamicView() const { UpdateIfNeeded(); return NonDisabledDynamicView; }
 
-	const TParticleView<TPBDRigidParticles<T, d>>& GetActiveParticlesView() const { return ActiveParticlesView; }
-	TParticleView<TPBDRigidParticles<T, d>>& GetActiveParticlesView() { return ActiveParticlesView; }
+	const TParticleView<TPBDRigidParticles<T, d>>& GetActiveParticlesView() const { UpdateIfNeeded(); return ActiveParticlesView; }
+	TParticleView<TPBDRigidParticles<T, d>>& GetActiveParticlesView() { UpdateIfNeeded(); return ActiveParticlesView; }
+
+	const TParticleView<TGeometryParticles<T, d>>& GetAllParticlesView() const { UpdateIfNeeded(); return AllParticlesView; }
 
-	const TParticleView<TGeometryParticles<T, d>>& GetAllParticlesView() const { return AllParticlesView; }
 
 	const TParticleView<TKinematicGeometryParticles<T, d>>& GetActiveKinematicParticlesView() const { return ActiveKinematicParticlesView; }
 	TParticleView<TKinematicGeometryParticles<T, d>>& GetActiveKinematicParticlesView() { return ActiveKinematicParticlesView; }
@@ -454,13 +575,23 @@ public:
 	const TPBDGeometryCollectionParticles<T, d>& GetGeometryCollectionParticles() const { return *GeometryCollectionParticles; }
 	TPBDGeometryCollectionParticles<T, d>& GetGeometryCollectionParticles() { return *GeometryCollectionParticles; }
 
-	const TParticleView<TPBDGeometryCollectionParticles<T, d>>& GetActiveGeometryCollectionParticlesView() const { return ActiveGeometryCollectionParticlesView; }
-	TParticleView<TPBDGeometryCollectionParticles<T, d>>& GetActiveGeometryCollectionParticlesView() { return ActiveGeometryCollectionParticlesView; }
+	void UpdateIfNeeded() const 
+	{
+		if (bGeometryCollectionDirty) 
+		{
+			TPBDRigidsSOAs<T, d>* NCThis = const_cast<TPBDRigidsSOAs<T, d>*>(this);
+			NCThis->UpdateGeometryCollectionViews();
+		}
+	}
 
 	/**
 	 * Update which particle arrays geometry collection particles are in based on 
 	 * their object state (static, kinematic, dynamic, sleeping) and their disabled 
 	 * state.
+	 *
+	 * The reason for specializing this function for geometry collections is for 
+	 * scalability.  That is, we try to process many geometry collection particles
+	 * at a time, rather than one by one.
 	 */
 	void UpdateGeometryCollectionViews(const bool ForceUpdateViews=false)
 	{
@@ -471,102 +602,89 @@ public:
 			return;
 		}
 
-		int32 AIdx = 0, SIdx = 0, KIdx = 0, DIdx = 0; // Active, Static, Kinematic, Dynamic Index
+		int32 StaticIdx = 0, KinematicIdx = 0, SleepingIdx = 0, DynamicIdx = 0;
+		int32 ActiveIdx = 0, DisabledIdx = 0;
 
 		for(int32 PIdx = 0; PIdx < (int32)GeometryCollectionParticles->Size(); PIdx++)
 		{
-			TPBDGeometryCollectionParticleHandle<T, d>* Handle = GeometryCollectionParticles->Handle(PIdx);
+			const TPBDGeometryCollectionParticleHandle<T, d>* Handle = 
+				GeometryCollectionParticles->Handle(PIdx);
 			if (!Handle)
 				continue;
 
-			// If the particle is disabled we treat it as static, but for no reason 
-			// other than immediate convenience.
-			const Chaos::EObjectStateType State = Handle->Disabled() ? Chaos::EObjectStateType::Static : Handle->ObjectState();
+			const bool bDisabled = Handle->Disabled();
+			if (bDisabled)
+				continue;
 
 			// Count the number of particles in each state.
+			const Chaos::EObjectStateType State = 
+				Handle->Sleeping() ? Chaos::EObjectStateType::Sleeping : Handle->ObjectState();
 			switch (State)
 			{
 			case Chaos::EObjectStateType::Static:
-				SIdx++;
-				AIdx += (int32)(!Handle->Disabled()); // If Disabled == true, then it's not active
+				StaticIdx++;
 				break;
-
 			case Chaos::EObjectStateType::Kinematic:
-				KIdx++;
-				AIdx++;
+				KinematicIdx++;
 				break;
-
-			case Chaos::EObjectStateType::Sleeping: // Sleeping is a modified dynamic state
-				DIdx++;
+			case Chaos::EObjectStateType::Sleeping:
+				SleepingIdx++;
 				break;
-
 			case Chaos::EObjectStateType::Dynamic:
-				DIdx++;
-				AIdx++;
+				DynamicIdx++;
 				break;
-
 			default:
 				break;
 			};
 		}
 
 		// Compare with the previous array sizes, and resize if needed.
-		bool Changed = 
-			ActiveGeometryCollectionArray.Num() != AIdx ||
-			StaticGeometryCollectionArray.Num() != SIdx || 
-			KinematicGeometryCollectionArray.Num() != KIdx || 
-			DynamicGeometryCollectionArray.Num() != DIdx;
+		bool Changed =
+			StaticGeometryCollectionArray.Num() != StaticIdx ||
+			KinematicGeometryCollectionArray.Num() != KinematicIdx ||
+			SleepingGeometryCollectionArray.Num() != SleepingIdx ||
+			DynamicGeometryCollectionArray.Num() != DynamicIdx;
 		if (Changed)
 		{
-			ActiveGeometryCollectionArray.SetNumUninitialized(AIdx);
-			StaticGeometryCollectionArray.SetNumUninitialized(SIdx);
-			KinematicGeometryCollectionArray.SetNumUninitialized(KIdx);
-			DynamicGeometryCollectionArray.SetNumUninitialized(DIdx);
+			StaticGeometryCollectionArray.SetNumUninitialized(StaticIdx);
+			KinematicGeometryCollectionArray.SetNumUninitialized(KinematicIdx);
+			SleepingGeometryCollectionArray.SetNumUninitialized(SleepingIdx);
+			DynamicGeometryCollectionArray.SetNumUninitialized(DynamicIdx);
 		}
 
 		// (Re)populate the arrays, making note if any prior entires differ from the current.
-		AIdx = SIdx = KIdx = DIdx = 0;
+		StaticIdx = KinematicIdx = SleepingIdx = DynamicIdx = 0;
 		for(int32 PIdx = 0; PIdx < (int32)GeometryCollectionParticles->Size(); PIdx++)
 		{
-			TPBDGeometryCollectionParticleHandle<T, d>* Handle = GeometryCollectionParticles->Handle(PIdx);
+			TPBDGeometryCollectionParticleHandle<T, d>* Handle = 
+				GeometryCollectionParticles->Handle(PIdx);
 			if (!Handle)
 				continue;
 
+			const bool bDisabled = Handle->Disabled();
+			if (bDisabled)
+				continue;
+
 			const Chaos::EObjectStateType State = 
-				Handle->Disabled() ? Chaos::EObjectStateType::Static : Handle->ObjectState();
+				Handle->Sleeping() ? Chaos::EObjectStateType::Sleeping : Handle->ObjectState();
 			switch (State)
 			{
 			case Chaos::EObjectStateType::Static:
-				Changed |= StaticGeometryCollectionArray[SIdx] != Handle;
-				StaticGeometryCollectionArray[SIdx++] = Handle;
-				if (!Handle->Disabled())
-				{
-					Changed |= ActiveGeometryCollectionArray[AIdx] != Handle;
-					ActiveGeometryCollectionArray[AIdx++] = Handle;
-				}
+				Changed |= StaticGeometryCollectionArray[StaticIdx] != Handle;
+				StaticGeometryCollectionArray[StaticIdx++] = Handle;
 				break;
-
 			case Chaos::EObjectStateType::Kinematic:
-				Changed |= KinematicGeometryCollectionArray[KIdx] != Handle;
-				KinematicGeometryCollectionArray[KIdx++] = Handle;
-
-				Changed |= ActiveGeometryCollectionArray[AIdx] != Handle;
-				ActiveGeometryCollectionArray[AIdx++] = Handle;
+				Changed |= KinematicGeometryCollectionArray[KinematicIdx] != Handle;
+				KinematicGeometryCollectionArray[KinematicIdx++] = Handle;
 				break;
-
-			case Chaos::EObjectStateType::Sleeping: // Sleeping is a modified dynamic state
-				Changed |= DynamicGeometryCollectionArray[DIdx] != Handle;
-				DynamicGeometryCollectionArray[DIdx++] = Handle;
+			case Chaos::EObjectStateType::Sleeping:
+				Changed |= SleepingGeometryCollectionArray[SleepingIdx] != Handle;
+				SleepingGeometryCollectionArray[SleepingIdx++] = Handle;
 				break;
-
 			case Chaos::EObjectStateType::Dynamic:
-				Changed |= DynamicGeometryCollectionArray[DIdx] != Handle;
-				DynamicGeometryCollectionArray[DIdx++] = Handle;
-
-				Changed |= ActiveGeometryCollectionArray[AIdx] != Handle;
-				ActiveGeometryCollectionArray[AIdx++] = Handle;
+				Changed |= DynamicGeometryCollectionArray[DynamicIdx] != Handle;
+				DynamicGeometryCollectionArray[DynamicIdx++] = Handle;
 				break;
-
 			default:
 				break;
 			};
@@ -576,6 +694,7 @@ public:
 		{
 			UpdateViews();
 		}
+		bGeometryCollectionDirty = false;
 	}
 
 	//TEMP: only needed while clustering code continues to use direct indices
@@ -702,9 +821,10 @@ private:
 				KinematicParticles.Get(), 
 				DynamicParticles.Get(),
 				DynamicKinematicParticles.Get(),
-				{&NonDisabledClusteredArray}, 
+				{&NonDisabledClusteredArray},
 				{&StaticGeometryCollectionArray},
 				{&KinematicGeometryCollectionArray},
+				{&SleepingGeometryCollectionArray},
 				{&DynamicGeometryCollectionArray}
 			};
 			NonDisabledView = MakeParticleView(MoveTemp(TmpArray));
@@ -714,6 +834,7 @@ private:
 			{ 
 				DynamicParticles.Get(), 
 				{&NonDisabledClusteredArray}, 
+				{&SleepingGeometryCollectionArray},
 				{&DynamicGeometryCollectionArray}
 			};
 			NonDisabledDynamicView = MakeParticleView(MoveTemp(TmpArray));
@@ -722,8 +843,10 @@ private:
 			TArray<TSOAView<TPBDRigidParticles<T, d>>> TmpArray = 
 			{ 
 				{&ActiveParticlesArray},
-				{&NonDisabledClusteredArray}, // Ryan!
-				{&ActiveGeometryCollectionArray}
+				{&NonDisabledClusteredArray},
+				{&StaticGeometryCollectionArray},
+				{&KinematicGeometryCollectionArray},
+				{&DynamicGeometryCollectionArray}
 			};
 			ActiveParticlesView = MakeParticleView(MoveTemp(TmpArray));
 		}
@@ -752,7 +875,12 @@ private:
 			ActiveKinematicParticlesView = MakeParticleView(MoveTemp(TmpArray));
 		}
 		{
-			TArray<TSOAView<TPBDGeometryCollectionParticles<T, d>>> TmpArray = { {&ActiveGeometryCollectionArray} };
+			TArray<TSOAView<TPBDGeometryCollectionParticles<T, d>>> TmpArray = 
+			{ 
+				{&StaticGeometryCollectionArray},
+				{&KinematicGeometryCollectionArray},
+				{&DynamicGeometryCollectionArray}
+			};
 			ActiveGeometryCollectionParticlesView = MakeParticleView(MoveTemp(TmpArray));
 		}
 	}
@@ -772,12 +900,13 @@ private:
 
 	TUniquePtr<TPBDGeometryCollectionParticles<T, d>> GeometryCollectionParticles;
 
-	TMap<TPBDGeometryCollectionParticleHandle<T, d>*, int32> ActiveGeometryCollectionToIndex;
-	TArray<TPBDGeometryCollectionParticleHandle<T, d>*> ActiveGeometryCollectionArray;
-
+	// Geometry collection particle state is controlled via their disabled state and assigned 
+	// EObjectStateType, and are shuffled into these corresponding arrays in UpdateGeometryCollectionViews().
 	TArray<TPBDGeometryCollectionParticleHandle<T, d>*> StaticGeometryCollectionArray;
 	TArray<TPBDGeometryCollectionParticleHandle<T, d>*> KinematicGeometryCollectionArray;
+	TArray<TPBDGeometryCollectionParticleHandle<T, d>*> SleepingGeometryCollectionArray;
 	TArray<TPBDGeometryCollectionParticleHandle<T, d>*> DynamicGeometryCollectionArray;
+	bool bGeometryCollectionDirty;
 
 	//Utility structures for maintaining an Active particles view
 	TMap<TPBDRigidParticleHandle<T, d>*, int32> ActiveParticlesToIndex;
diff --git a/Engine/Source/Runtime/Experimental/ChaosSolvers/Private/PhysicsProxy/GeometryCollectionPhysicsProxy.cpp b/Engine/Source/Runtime/Experimental/ChaosSolvers/Private/PhysicsProxy/GeometryCollectionPhysicsProxy.cpp
--- a/Engine/Source/Runtime/Experimental/ChaosSolvers/Private/PhysicsProxy/GeometryCollectionPhysicsProxy.cpp
+++ b/Engine/Source/Runtime/Experimental/ChaosSolvers/Private/PhysicsProxy/GeometryCollectionPhysicsProxy.cpp
@@ -889,10 +889,10 @@ void FGeometryCollectionPhysicsProxy::InitializeBodiesPT(
 				}
 			}
 
-			// We've likely chagned the state of leaf nodes, which are geometry
+			// We've likely changed the state of leaf nodes, which are geometry
 			// collection particles.  Update which particle views they belong in,
 			// as well as views of clustered particles.
-			Particles.UpdateGeometryCollectionViews(true);
+			Particles.UpdateGeometryCollectionViews(true); 
 
 			// Set cluster connectivity.  TPBDRigidClustering::CreateClusterParticle() 
 			// will optionally do this, but we switch that functionality off in BuildClusters().
@@ -2100,7 +2100,11 @@ void FGeometryCollectionPhysicsProxy::PushToPhysicsState(const Chaos::FParticleD
 
 			Handle->SetX(ParticleTransform.GetTranslation());
 			Handle->SetR(ParticleTransform.GetRotation());
-			Handle->SetDisabled(GState->DisabledStates[Idx]);
+
+			// Particle disabled states are initialized on the physics thread, so 
+			// we don't want to pull this from the game thread
+
+			//Handle->SetDisabled(GState->DisabledStates[Idx]); ryan - fix this!
 
 			Handle->SetLinearEtherDrag(Parameters.LinearEtherDrag);
 			Handle->SetAngularEtherDrag(Parameters.AngularEtherDrag);
@@ -2158,6 +2162,7 @@ void FGeometryCollectionPhysicsProxy::PushToPhysicsState(const Chaos::FParticleD
 		}
 	}
 }
+
 void FGeometryCollectionPhysicsProxy::BufferPhysicsResults()
 {
 	/**
