Added TProperties to particles. Nothing sets remote data yet so existing behavior is the same

diff --git a/Engine/Source/Programs/HeadlessChaos/Private/HeadlessChaosTestDirtyParticle.cpp b/Engine/Source/Programs/HeadlessChaos/Private/HeadlessChaosTestDirtyParticle.cpp
--- a/Engine/Source/Programs/HeadlessChaos/Private/HeadlessChaosTestDirtyParticle.cpp
+++ b/Engine/Source/Programs/HeadlessChaos/Private/HeadlessChaosTestDirtyParticle.cpp
@@ -17,159 +17,53 @@ namespace ChaosTest
 
 	GTEST_TEST(DirtyParticleTests,Basic)
 	{
-#if 0
-		struct FSimpleParticleRemote
-		{
-			FSimpleParticleRemote() : Manager(nullptr){}
-
-			FDirtyPropertiesManager* Manager;
-			TRemoteParticleProperty<FVec3> X;
-			TRemoteParticleProperty<FReal> InvM;
-			TRemoteParticleProperty<TSharedPtr<FImplicitObject,ESPMode::ThreadSafe>> Geometry;
-
-			const FVec3& GetX(FDirtyPropertiesManager& Manager) const { return X.Read(*Manager); }
-			const FReal GetInvM(FDirtyPropertiesManager& Manager) const { return InvM.Read(*Manager); }
-			const TSharedPtr<FImplicitObject,ESPMode::ThreadSafe>& GetGeometry(FDirtyPropertiesManager& Manager) const { return Geometry.Read(*Manager); }
-			
-			void Clear()
-			{
-				if(Manager)
-				{
-					X.Clear(*Manager);
-					InvM.Clear(*Manager);
-					Geometry.Clear(*Manager);
-				}
-			}
-
-			~FSimpleParticleRemote()
-			{
-				Clear();
-			}
-		};
-
-		struct FSimpleParticle
-		{
-			FSimpleParticle() : Manager(nullptr){}
-			const FVec3& GetX() const{ return X.Read(); }
-			void SetX(const FVec3& InX){ X.Write(InX,Remote.X,ToFlush,Remote.Dirty, Remote.Manager); }
-
-			const FReal GetInvM() const { return InvM.Read(); }
-			void SetInvM(FReal InInvM){ InvM.Write(InInvM,Remote.InvM, ToFlush, Remote.Dirty, Remote.Manager); }
-
-			const TSharedPtr<FImplicitObject,ESPMode::ThreadSafe>& GetGeometry() const { return Geometry.Read(); }
-			void SetGeometry(const TSharedPtr<FImplicitObject,ESPMode::ThreadSafe>& InGeom){ Geometry.Write(InGeom, Remote.Geometry, ToFlush, Remote.Dirty, Remote.Manager); }
-
-			void Flush()
-			{
-				ensure(Remote.Manager);
-				if(ToFlush.IsDirty())
-				{
-					if(ToFlush.IsDirty(EParticleFlags::X))
-					{
-						X.Flush(Remote.X,ToFlush,*Remote.Manager);
-					}
-
-					if(ToFlush.IsDirty(EParticleFlags::InvM))
-					{
-						InvM.Flush(Remote.InvM,ToFlush, *Remote.Manager);
-					}
-
-					if(ToFlush.IsDirty(EParticleFlags::Geometry))
-					{
-						Geometry.Flush(Remote.Geometry,ToFlush,*Remote.Manager);
-					}
-				}
-
-				ensure(ToFlush.IsClean());
-			}
-
-			bool IsDirty() const { return Remote.Dirty.IsDirty(); }
-
-			void StealRemote(FSimpleParticleRemote& Dst)
-			{
-				Flush();
-				Swap(Dst,Remote);
-			}
-
-			TParticleProperty<FVec3,EParticleFlags::X> X;
-			TParticleProperty<FReal,EParticleFlags::InvM> InvM;
-			TParticleProperty<TSharedPtr<FImplicitObject, ESPMode::ThreadSafe>,EParticleFlags::Geometry> Geometry;
-			FParticleDirtyFlags ToFlush;
-
-			FSimpleParticleRemote Remote;
-		};
-
-		FSimpleParticle P1;
-
-		P1.SetX(FVec3(1,1,1));
-		P1.SetInvM(3);
-
-		EXPECT_EQ(P1.GetX(),FVec3(1,1,1));
-		EXPECT_EQ(P1.GetInvM(),3);
-
-		EXPECT_TRUE(P1.ToFlush.IsDirty(EParticleFlags::X));
-		EXPECT_TRUE(P1.ToFlush.IsDirty(EParticleFlags::InvM));
-		EXPECT_FALSE(P1.ToFlush.IsDirty(EParticleFlags::Geometry));
+		auto Particle = TGeometryParticle<FReal,3>::CreateParticle();
+		Particle->SetX(FVec3(1,1,1));
 
+		TSharedPtr<FImplicitObject,ESPMode::ThreadSafe> Ptr(new TSphere<FReal,3>(FVec3(0),0));
+		const auto RawPtr = Ptr.Get();
+		TWeakPtr<FImplicitObject,ESPMode::ThreadSafe> WeakPtr(Ptr);
 		FDirtyPropertiesManager Manager;
-		P1.Remote.Manager = &Manager;
-
-		FSimpleParticleRemote Remote;
-		P1.StealRemote(Remote);
 
-		EXPECT_TRUE(Remote.Dirty.IsDirty(EParticleFlags::X));
-		EXPECT_TRUE(Remote.Dirty.IsDirty(EParticleFlags::X));
-		EXPECT_FALSE(Remote.Dirty.IsDirty(EParticleFlags::Geometry));
+		{
+			FParticlePropertiesData RemoteData;
+			RemoteData.Manager = &Manager;
 
-		EXPECT_EQ(Remote.GetX(Manager),FVec3(1,1,1));
-		EXPECT_EQ(Remote.GetInvM(Manager),3);
+			Particle->SetRemoteData(&RemoteData);
 
-		Remote.Clear();
+			//then on PushToPhysics we'd do a pointer swap and use RemoteData internally
 
-		/*
-		FDirtyProperties DirtyParticle1;
-		DirtyParticle1.WriteX(DirtyPropertiesPool, FVec3(1,1,1));
-		DirtyParticle1.WriteInvM(DirtyPropertiesPool, 3);
+			EXPECT_TRUE(RemoteData.HasX());
+			EXPECT_EQ(RemoteData.GetX(),FVec3(1,1,1));
+			EXPECT_FALSE(RemoteData.HasInvM());	//was never set so it's false
 
-		FDirtyProperties DirtyParticle2;
-		DirtyParticle2.WriteX(DirtyPropertiesPool, FVec3(2,1,1));
+			Particle->SetX(FVec3(2,1,1));
+			EXPECT_EQ(RemoteData.GetX(),FVec3(2,1,1));	//remote is set so immediate change
 
-		//properties are dirty
-		EXPECT_TRUE(DirtyParticle1.IsDirty(EParticleFlags::X));
-		EXPECT_TRUE(DirtyParticle1.IsDirty(EParticleFlags::InvM));
+			//make sure we are not leaking shared ptrs
+			Particle->SetGeometry(Ptr);
+			Ptr = nullptr;
 
-		EXPECT_TRUE(DirtyParticle2.IsDirty(EParticleFlags::X));
+			EXPECT_TRUE(WeakPtr.IsValid());	//still around because particle is holding on to it
+			EXPECT_TRUE(RemoteData.HasGeometry());
+			EXPECT_EQ(RemoteData.GetGeometry().Get(),RawPtr);
 
-		//untouched properties are clean
-		EXPECT_FALSE(DirtyParticle2.IsDirty(EParticleFlags::InvM));
+			Particle->SetRemoteData(nullptr);	//disconnect remote data so we can pretend we are freeing things on GT without affecting PT
 
-		//values were saved
-		EXPECT_EQ(DirtyParticle1.ReadX(DirtyPropertiesPool), FVec3(1,1,1));
-		EXPECT_EQ(DirtyParticle1.ReadInvM(DirtyPropertiesPool), 3);
-		EXPECT_EQ(DirtyParticle2.ReadX(DirtyPropertiesPool), FVec3(2,1,1));
-		
-		EXPECT_EQ(DirtyParticle1.PopX(DirtyPropertiesPool),FVec3(1,1,1));
-		EXPECT_FALSE(DirtyParticle1.IsDirty(EParticleFlags::X));
+			Particle->SetGeometry(Ptr);	//free geometry on GT side
 
-		//make sure we are not leaking shared ptrs
-		TSharedPtr<FImplicitObject,ESPMode::ThreadSafe> Ptr(new TSphere<FReal,3>(FVec3(0), 0));
-		DirtyParticle1.WriteGeometry(DirtyPropertiesPool, Ptr);
-		TWeakPtr<FImplicitObject,ESPMode::ThreadSafe> WeakPtr(Ptr);
-		Ptr = nullptr;
+			//geometry still on PT side
+			EXPECT_TRUE(RemoteData.HasGeometry());
+			EXPECT_EQ(RemoteData.GetGeometry().Get(),RawPtr);
 
-		EXPECT_TRUE(WeakPtr.IsValid());	//still around because dirty pool is holding on to it
-		{
-			TSharedPtr<FImplicitObject,ESPMode::ThreadSafe> Geom = DirtyParticle1.PopGeometry(DirtyPropertiesPool);
-			EXPECT_TRUE(WeakPtr.IsValid());	//Popped the geometry but still holding on to it
+			EXPECT_TRUE(WeakPtr.IsValid());	//still around because particle is holding on to it
 		}
-
-		EXPECT_FALSE(WeakPtr.IsValid());	//Finished with popped geometry so shared ptr goes away
 		
+		//remote data is gone so geometry shared ptr is freed, even though pool is still around (i.e. it was removed from pool)
+		EXPECT_FALSE(WeakPtr.IsValid());
+
+
 
-		//If we haven't popped everything we call this
-		DirtyParticle1.Clean(DirtyPropertiesPool);
-		DirtyParticle2.Clean(DirtyPropertiesPool);*/
-#endif
 	}
 }
 
diff --git a/Engine/Source/Runtime/Experimental/Chaos/Private/Chaos/ParticleHandle.cpp b/Engine/Source/Runtime/Experimental/Chaos/Private/Chaos/ParticleHandle.cpp
--- a/Engine/Source/Runtime/Experimental/Chaos/Private/Chaos/ParticleHandle.cpp
+++ b/Engine/Source/Runtime/Experimental/Chaos/Private/Chaos/ParticleHandle.cpp
@@ -26,10 +26,11 @@ namespace Chaos
 			}
 		}
 
-		if (MGeometry)
+		auto& Geometry = MGeometry.Read();
+		if (Geometry)
 		{
 			int32 CurrentShapeIndex = INDEX_NONE;
-			if (const auto* Union = MGeometry->template GetObject<FImplicitObjectUnion>())
+			if (const auto* Union = Geometry->template GetObject<FImplicitObjectUnion>())
 			{
 				for (const TUniquePtr<FImplicitObject>& ImplicitObject : Union->GetObjects())
 				{
@@ -51,15 +52,15 @@ namespace Chaos
 			}
 			else 
 			{
-				if (const FImplicitObject* ImplicitChildObject = Utilities::ImplicitChildHelper(MGeometry.Get()))
+				if (const FImplicitObject* ImplicitChildObject = Utilities::ImplicitChildHelper(Geometry.Get()))
 				{
-					if (ImplicitShapeMap.Contains(MGeometry.Get()))
+					if (ImplicitShapeMap.Contains(Geometry.Get()))
 					{
-						ImplicitShapeMap.Add(ImplicitChildObject, CopyTemp(ImplicitShapeMap[MGeometry.Get()]));
+						ImplicitShapeMap.Add(ImplicitChildObject, CopyTemp(ImplicitShapeMap[Geometry.Get()]));
 					}
 					else if (ImplicitShapeMap.Contains(ImplicitChildObject))
 					{
-						ImplicitShapeMap.Add(MGeometry.Get(), CopyTemp(ImplicitShapeMap[ImplicitChildObject]));
+						ImplicitShapeMap.Add(Geometry.Get(), CopyTemp(ImplicitShapeMap[ImplicitChildObject]));
 					}
 				}
 			}
diff --git a/Engine/Source/Runtime/Experimental/Chaos/Public/Chaos/ParticleDirtyFlags.h b/Engine/Source/Runtime/Experimental/Chaos/Public/Chaos/ParticleDirtyFlags.h
--- a/Engine/Source/Runtime/Experimental/Chaos/Public/Chaos/ParticleDirtyFlags.h
+++ b/Engine/Source/Runtime/Experimental/Chaos/Public/Chaos/ParticleDirtyFlags.h
@@ -3,6 +3,8 @@
 #pragma once
 
 #include "Chaos/Core.h"
+#include "Framework/PhysicsProxyBase.h"
+#include "Chaos/ChaosArchive.h"
 
 class FName;
 
@@ -56,6 +58,17 @@ struct TParticlePropertyTrait<Type>\
 	};
 #undef PARTICLE_PROPERTY
 
+	constexpr EParticleFlags ParticlePropToFlag(EParticleProperty Prop)
+	{
+		switch(Prop)
+		{
+			#define PARTICLE_PROPERTY(PropName, Type) case PropName: return EParticleFlags::PropName;
+#include "ParticleProperties.inl"
+#undef PARTICLE_PROPERTY
+		default: return (EParticleFlags)0;
+		}
+	}
+
 	class FParticleDirtyFlags
 	{
 	public:
@@ -261,40 +274,147 @@ public:
 			Idx.bHasEntry = true;
 		}
 	}
+
+	bool IsSet() const
+	{
+		return Idx.bHasEntry;
+	}
 private:
 	FDirtyIdx Idx;
 
 	TRemoteParticleProperty<T>& operator=(const TRemoteParticleProperty<T>& Rhs){}
 };
 
-template <typename T, EParticleFlags Flag>
+struct FParticlePropertiesData
+{
+	template <typename T, EParticleProperty PropertyIdx>
+	TRemoteParticleProperty<T>& GetProperty()
+	{
+		switch(PropertyIdx)
+		{
+#define PARTICLE_PROPERTY(PropName, Type) case EParticleProperty::PropName: return (TRemoteParticleProperty<T>&) PropName;
+#include "ParticleProperties.inl"
+#undef PARTICLE_PROPERTY
+		default: check(false);
+		}
+
+		static TRemoteParticleProperty<T> Error;
+		return Error;
+	}
+
+	FDirtyPropertiesManager* Manager;
+
+	~FParticlePropertiesData()
+	{
+		if(Manager)
+		{
+#define PARTICLE_PROPERTY(PropName, Type) PropName.Clear(*Manager);
+#include "ParticleProperties.inl"
+#undef PARTICLE_PROPERTY
+		}
+	}
+
+#define PARTICLE_PROPERTY(PropName, Type)\
+Type const & Get##PropName() { return PropName.Read(*Manager); };\
+bool Has##PropName() const { return PropName.IsSet(); };
+
+#include "ParticleProperties.inl"
+#undef PARTICLE_PROPERTY
+
+private:
+#define PARTICLE_PROPERTY(PropName, Type) TRemoteParticleProperty<Type> PropName;
+#include "ParticleProperties.inl"
+#undef PARTICLE_PROPERTY
+
+};
+
+template <typename T, EParticleProperty PropName>
 class TParticleProperty
 {
 public:
+	TParticleProperty() = default;
+	TParticleProperty(const T& Val)
+		: Property(Val)
+	{
+	}
+
+	//we don't support this because it could lead to bugs with values not being properly written to remote
+	TParticleProperty(const TParticleProperty<T,PropName>& Rhs) = delete;
+
 	const T& Read() const { return Property; }
-	void Write(const T& Val, TRemoteParticleProperty<T>& Remote, FParticleDirtyFlags& PendingFlush, FParticleDirtyFlags& Dirty, FDirtyPropertiesManager* Manager)
+	void Write(const T& Val, bool bInvalidate, FParticlePropertiesData* Remote, FParticleDirtyFlags& Dirty, IPhysicsProxyBase* Proxy)
 	{
 		Property = Val;
-		if(Manager)
-		{
-			Remote.Write(*Manager, Val);
-			PendingFlush.MarkClean(Flag);
-		}
-		else
+		WriteImp(bInvalidate, Remote, Dirty, Proxy);
+	}
+
+	template <typename Lambda>
+	void Modify(bool bInvalidate,FParticlePropertiesData* Remote,FParticleDirtyFlags& Dirty,IPhysicsProxyBase* Proxy,const Lambda& LambdaFunc)
+	{
+		LambdaFunc(Property);
+		WriteImp(bInvalidate, Remote, Dirty, Proxy);
+	}
+
+	void InitialFlush(FParticlePropertiesData& Remote,FParticleDirtyFlags& Dirty, IPhysicsProxyBase* Proxy)
+	{
+		if(Dirty.IsDirty(PropertyFlag))
 		{
-			PendingFlush.MarkDirty(Flag);
+			WriteToRemote(Remote,Proxy);
 		}
-
-		Dirty.MarkDirty(Flag);
 	}
 
-	void Flush(TRemoteParticleProperty<T>& Remote,FParticleDirtyFlags& PendingFlush,FDirtyPropertiesManager& Manager)
+	void Serialize(FChaosArchive& Ar)
 	{
-		Remote.Write(Manager, Property);
-		PendingFlush.MarkClean(Flag);
+		Ar << Property;
 	}
 
 private:
 	T Property;
+	static constexpr EParticleFlags PropertyFlag = ParticlePropToFlag(PropName);
+
+	void WriteToRemote(FParticlePropertiesData& Remote, IPhysicsProxyBase* Proxy)
+	{
+		Remote.GetProperty<T,PropName>().Write(*Remote.Manager, Property);
+		if(Proxy)
+		{
+			if(FPhysicsSolverBase* PhysicsSolverBase = Proxy->GetSolver<FPhysicsSolverBase>())
+			{
+				PhysicsSolverBase->AddDirtyProxy(Proxy);
+			}
+		}
+	}
+
+	void WriteImp(bool bInvalidate, FParticlePropertiesData* Remote, FParticleDirtyFlags& Dirty, IPhysicsProxyBase* Proxy)
+	{
+		if(bInvalidate)
+		{
+			Dirty.MarkDirty(PropertyFlag);
+			if(Remote)
+			{
+				WriteToRemote(*Remote,Proxy);
+
+			}
+
+			//for now we have to dirty proxy in here because rest of pipeline is still old
+			//todo: remove this once remote data is set, writing to remote should take care of proxy dirtying
+			if(Proxy)
+			{
+				if(FPhysicsSolverBase* PhysicsSolverBase = Proxy->GetSolver<FPhysicsSolverBase>())
+				{
+					PhysicsSolverBase->AddDirtyProxy(Proxy);
+				}
+			}
+		}
+	}
+
+	//we don't support this because it could lead to bugs with values not being properly written to remote
+	TParticleProperty<T,PropName>& operator=(const TParticleProperty<T,PropName>& Rhs) {}
 };
+
+template <typename T, EParticleProperty PropName>
+FChaosArchive& operator<<(FChaosArchive& Ar, TParticleProperty<T, PropName>& Prop)
+{
+	Prop.Serialize(Ar);
+	return Ar;
+}
 }
diff --git a/Engine/Source/Runtime/Experimental/Chaos/Public/Chaos/ParticleHandle.h b/Engine/Source/Runtime/Experimental/Chaos/Public/Chaos/ParticleHandle.h
--- a/Engine/Source/Runtime/Experimental/Chaos/Public/Chaos/ParticleHandle.h
+++ b/Engine/Source/Runtime/Experimental/Chaos/Public/Chaos/ParticleHandle.h
@@ -7,7 +7,6 @@
 #include "Chaos/ParticleHandleFwd.h"
 #include "Chaos/ParticleIterator.h"
 #include "Chaos/ParticleDirtyFlags.h"
-#include "Chaos/Framework/PhysicsProxyBase.h"
 #include "ChaosCheck.h"
 
 
@@ -1348,6 +1347,7 @@ protected:
 
 	TGeometryParticle(const TGeometryParticleParameters<T, d>& StaticParams = TGeometryParticleParameters<T, d>())
 		: MUserData(nullptr)
+		, RemoteData(nullptr)
 	{
 		Type = EParticleType::Static;
 		Proxy = nullptr;
@@ -1389,30 +1389,28 @@ public:
 
 	virtual bool IsParticleValid() const
 	{
-		return MGeometry && MGeometry->IsValidGeometry();	//todo: if we want support for sample particles without geometry we need to adjust this
+		auto& Geometry = MGeometry.Read();
+		return Geometry && Geometry->IsValidGeometry();	//todo: if we want support for sample particles without geometry we need to adjust this
 	}
 
 	static TGeometryParticle<T, d>* SerializationFactory(FChaosArchive& Ar, TGeometryParticle<T, d>* Serializable);
 
-	const TVector<T, d>& X() const { return this->MX; }
+	const TVector<T, d>& X() const { return MX.Read(); }
 	void SetX(const TVector<T, d>& InX, bool bInvalidate = true)
 	{
-		this->MarkDirty(EParticleFlags::X, bInvalidate);
-		this->MX = InX;
+		MX.Write(InX, bInvalidate, RemoteData, MDirtyFlags, Proxy);
 	}
 
-	FUniqueIdx UniqueIdx() const { return this->MUniqueIdx; }
+	FUniqueIdx UniqueIdx() const { return MUniqueIdx.Read(); }
 	void SetUniqueIdx(const FUniqueIdx UniqueIdx, bool bInvalidate = true)
 	{
-		this->MarkDirty(EParticleFlags::UniqueIdx,bInvalidate);
-		this->MUniqueIdx = UniqueIdx;
+		MUniqueIdx.Write(UniqueIdx,bInvalidate,RemoteData,MDirtyFlags,Proxy);
 	}
 
-	const TRotation<T, d>& R() const { return this->MR; }
+	const TRotation<T, d>& R() const { return MR.Read(); }
 	void SetR(const TRotation<T, d>& InR, bool bInvalidate = true)
 	{
-		this->MarkDirty(EParticleFlags::R, bInvalidate);
-		this->MR = InR;
+		MR.Write(InR,bInvalidate,RemoteData,MDirtyFlags,Proxy);
 	}
 
 	//todo: geometry should not be owned by particle
@@ -1430,8 +1428,7 @@ public:
 	//       We should replace this with a method for supporting SetGeometry(RawGeometry).
 	void SetGeometry(TSharedPtr<FImplicitObject, ESPMode::ThreadSafe> SharedGeometry)
 	{
-		this->MarkDirty(EParticleFlags::Geometry);
-		this->MGeometry = SharedGeometry;
+		MGeometry.Write(SharedGeometry,/*bInvalidate=*/true,RemoteData,MDirtyFlags,Proxy);
 		UpdateShapesArray();
 	}
 
@@ -1442,30 +1439,28 @@ public:
 		check(false);
 	}
 
-	void* UserData() const { return this->MUserData; }
+	void* UserData() const { return MUserData.Read(); }
 	void SetUserData(void* InUserData)
 	{
-		this->MarkDirty(EParticleFlags::UserData);
-		this->MUserData = InUserData;
+		MUserData.Write(InUserData,/*bInvalidate=*/true,RemoteData,MDirtyFlags,Proxy);
 	}
 
 	void UpdateShapeBounds()
 	{
-		if (MGeometry->HasBoundingBox())
+		if (MGeometry.Read()->HasBoundingBox())
 		{
 			for (auto& Shape : MShapesArray)
 			{
-				Shape->UpdateShapeBounds(FRigidTransform3(MX, MR));
+				Shape->UpdateShapeBounds(FRigidTransform3(MX.Read(), MR.Read()));
 			}
 		}
 	}
 
 #if CHAOS_CHECKED
-	const FName& DebugName() const { return this->MDebugName; }
+	const FName& DebugName() const { return MDebugName.Read(); }
 	void SetDebugName(const FName& InDebugName)
 	{
-		this->MarkDirty(EParticleFlags::DebugName);
-		this->MDebugName = InDebugName;
+		MDebugName.Write(InDebugName,/*bInvalidate=*/true,RemoteData,MDirtyFlags,Proxy);
 	}
 #endif
 
@@ -1477,7 +1472,7 @@ public:
 		MapImplicitShapes();
 	}
 
-	TSerializablePtr<FImplicitObject> Geometry() const { return MakeSerializable(MGeometry); }
+	TSerializablePtr<FImplicitObject> Geometry() const { return MakeSerializable(MGeometry.Read()); }
 
 	const TShapesArray<T,d>& ShapesArray() const { return MShapesArray; }
 
@@ -1496,11 +1491,10 @@ public:
 	const TPBDRigidParticle<T, d>* CastToRigidParticle() const;
 	TPBDRigidParticle<T, d>* CastToRigidParticle();
 
-	FSpatialAccelerationIdx SpatialIdx() const { return MSpatialIdx; }
+	FSpatialAccelerationIdx SpatialIdx() const { return MSpatialIdx.Read(); }
 	void SetSpatialIdx(FSpatialAccelerationIdx Idx)
 	{
-		MarkDirty(EParticleFlags::SpatialIdx);
-		MSpatialIdx = Idx;
+		MSpatialIdx.Write(Idx,/*bInvalidate=*/true,RemoteData,MDirtyFlags,Proxy);
 	}
 
 	void SetShapeCollisionDisable(int32 InShapeIndex, bool bInDisable)
@@ -1537,12 +1531,12 @@ public:
 
 	bool IsDirty() const
 	{
-		return this->MDirtyFlags.IsDirty();
+		return MDirtyFlags.IsDirty();
 	}
 
 	bool IsDirty(const EParticleFlags CheckBits) const
 	{
-		return this->MDirtyFlags.IsDirty(CheckBits);
+		return MDirtyFlags.IsDirty(CheckBits);
 	}
 
 	const FParticleDirtyFlags& DirtyFlags() const
@@ -1550,14 +1544,16 @@ public:
 		return MDirtyFlags;
 	}
 
+	//todo: should this exist?
 	void ClearDirtyFlags()
 	{
-		this->MDirtyFlags.Clear();
+		MDirtyFlags.Clear();
 	}
 
+	//todo: should this exist?
 	void MarkClean(const EParticleFlags CleanBits)
 	{
-		this->MDirtyFlags.MarkClean(CleanBits);
+		MDirtyFlags.MarkClean(CleanBits);
 	}
 
 
@@ -1591,17 +1587,23 @@ public:
 	//       inside the particle, but not setting the shape data through it.
 	void MarkShapeSimDataDirty() { MarkDirty(EParticleFlags::ShapeSimData); }
 
+	void SetRemoteData(FParticlePropertiesData* InRemoteData)
+	{
+		RemoteData = InRemoteData;
+		FlushToRemoteData();
+	}
+
 private:
-	TVector<T, d> MX;
-	FUniqueIdx MUniqueIdx;
-	TRotation<T, d> MR;
-	TSharedPtr<FImplicitObject, ESPMode::ThreadSafe> MGeometry;	//TODO: geometry should live in bodysetup
+	TParticleProperty<TVector<T, d>, EParticleProperty::X> MX;
+	TParticleProperty<FUniqueIdx, EParticleProperty::UniqueIdx> MUniqueIdx;
+	TParticleProperty<TRotation<T, d>, EParticleProperty::R> MR;
+	TParticleProperty<TSharedPtr<FImplicitObject, ESPMode::ThreadSafe>, EParticleProperty::Geometry> MGeometry;	//TODO: geometry should live in bodysetup
 	TShapesArray<T,d> MShapesArray;
 	TMap<const FImplicitObject*, int32> ImplicitShapeMap;
-	FSpatialAccelerationIdx MSpatialIdx;
+	TParticleProperty<FSpatialAccelerationIdx, EParticleProperty::SpatialIdx> MSpatialIdx;
 
 	// Pointer to some arbitrary data associated with the particle, but not used by Chaos. External systems may use this for whatever.
-	void* MUserData;
+	TParticleProperty<void*, EParticleProperty::UserData> MUserData;
 
 public:
 	// Ryan: FGeometryCollectionPhysicsProxy needs access to GeometrySharedLowLevel(), 
@@ -1613,7 +1615,7 @@ public:
 	// when the geometry is being copied from GT to PT.
 	TSharedPtr<FImplicitObject, ESPMode::ThreadSafe> GeometrySharedLowLevel() const
 	{
-		return MGeometry;
+		return MGeometry.Read();
 	}
 private:
 
@@ -1622,6 +1624,8 @@ private:
 #endif
 
 protected:
+	FParticlePropertiesData* RemoteData;
+
 	EParticleType Type;
 	FParticleDirtyFlags MDirtyFlags;
 
@@ -1629,10 +1633,18 @@ protected:
 
 	void UpdateShapesArray()
 	{
-		UpdateShapesArrayFromGeometry<T, d>(MShapesArray, MakeSerializable(MGeometry), FRigidTransform3(X(), R()));
+		UpdateShapesArrayFromGeometry<T, d>(MShapesArray, MakeSerializable(MGeometry.Read()), FRigidTransform3(X(), R()));
 		MapImplicitShapes();
 	}
 
+	void FlushToRemoteData()
+	{
+		if(RemoteData && MDirtyFlags.IsDirty())
+		{
+			MX.InitialFlush(*RemoteData,MDirtyFlags,Proxy);
+		}
+	}
+
 	void MapImplicitShapes();
 };
 
@@ -1782,32 +1794,28 @@ public:
 		Ar << MW;
 	}
 
-	const TVector<T, d>& V() const { return MV; }
+	const TVector<T, d>& V() const { return MV.Read(); }
 	void SetV(const TVector<T, d>& InV, bool bInvalidate = true)
 	{
-		this->MarkDirty(EParticleFlags::V, bInvalidate);
-		this->MV = InV;
+		MV.Write(InV,bInvalidate,this->RemoteData,this->MDirtyFlags,this->Proxy);
 	}
 
-	const TVector<T, d>& W() const { return MW; }
+	const TVector<T, d>& W() const { return MW.Read(); }
 	void SetW(const TVector<T, d>& InW, bool bInvalidate = true)
 	{
-		this->MarkDirty(EParticleFlags::W, bInvalidate);
-		this->MW = InW;
+		MW.Write(InW,bInvalidate,this->RemoteData,this->MDirtyFlags,this->Proxy);
 	}
 
-	const TVector<T, d>& CenterOfMass() const { return MCenterOfMass; }
+	const TVector<T, d>& CenterOfMass() const { return MCenterOfMass.Read(); }
 	void SetCenterOfMass(const TVector<T, d>& InCenterOfMass, bool bInvalidate = true)
 	{
-		this->MarkDirty(EParticleFlags::CenterOfMass, bInvalidate);
-		this->MCenterOfMass = InCenterOfMass;
+		MCenterOfMass.Write(InCenterOfMass,bInvalidate,this->RemoteData,this->MDirtyFlags,this->Proxy);
 	}
 	
-	const TRotation<T, d>& RotationOfMass() const { return MRotationOfMass; }
+	const TRotation<T, d>& RotationOfMass() const { return MRotationOfMass.Read(); }
 	void SetRotationOfMass(const TRotation<T, d>& InRotationOfMass, bool bInvalidate = true)
 	{
-		this->MarkDirty(EParticleFlags::RotationOfMass, bInvalidate);
-		this->MRotationOfMass = InRotationOfMass;
+		MRotationOfMass.Write(InRotationOfMass,bInvalidate,this->RemoteData,this->MDirtyFlags,this->Proxy);
 	}
 
 	EObjectStateType ObjectState() const;
@@ -1818,10 +1826,10 @@ public:
 	}
 
 private:
-	TVector<T, d> MV;
-	TVector<T, d> MW;
-	TVector<T, d> MCenterOfMass;
-	TRotation<T, d> MRotationOfMass;
+	TParticleProperty<TVector<T, d>, EParticleProperty::V> MV;
+	TParticleProperty<TVector<T, d>, EParticleProperty::W> MW;
+	TParticleProperty<TVector<T, d>, EParticleProperty::CenterOfMass> MCenterOfMass;
+	TParticleProperty<TRotation<T, d>, EParticleProperty::RotationOfMass> MRotationOfMass;
 };
 
 template <typename T, int d>
@@ -1934,142 +1942,122 @@ public:
 
 	const TUniquePtr<TBVHParticles<T, d>>& CollisionParticles() const { return MCollisionParticles; }
 
-	int32 CollisionGroup() const { return MCollisionGroup; }
+	int32 CollisionGroup() const { return MCollisionGroup.Read(); }
 	void SetCollisionGroup(const int32 InCollisionGroup)
 	{
-		this->MarkDirty(EParticleFlags::CollisionGroup);
-		this->MCollisionGroup = InCollisionGroup;
+		MCollisionGroup.Write(InCollisionGroup, /*bInvalidate=*/true,this->RemoteData,this->MDirtyFlags,this->Proxy);
 	}
 
-	bool Disabled() const { return MDisabled; }
+	bool Disabled() const { return MDisabled.Read(); }
 	void SetDisabled(const bool InDisabled)
 	{
-		this->MarkDirty(EParticleFlags::CollisionGroup);
-		this->MDisabled = InDisabled;
+		MDisabled.Write(InDisabled, /*bInvalidate=*/true,this->RemoteData,this->MDirtyFlags,this->Proxy);
 	}
 
-	bool IsGravityEnabled() const { return MGravityEnabled; }
+	bool IsGravityEnabled() const { return MGravityEnabled.Read(); }
 	void SetGravityEnabled(const bool InGravityEnabled)
 	{
-		this->MarkDirty(EParticleFlags::GravityEnabled);
-		this->MGravityEnabled = InGravityEnabled;
+		MGravityEnabled.Write(InGravityEnabled, /*bInvalidate=*/true,this->RemoteData,this->MDirtyFlags,this->Proxy);
 	}
 
+	//todo: remove this
 	bool IsInitialized() const { return MInitialized; }
 	void SetInitialized(const bool InInitialized)
 	{
 		this->MInitialized = InInitialized;
 	}
 
-	// Named to match signature of TPBDRigidParticleHandle, as both are used in templated functions.
-	// See its comment for details.
-	bool& SetDisabledLowLevel() { return MDisabled; }
-
-	const TVector<T, d>& PreV() const { return MPreV; }
+	const TVector<T, d>& PreV() const { return MPreV.Read(); }
 	void SetPreV(const TVector<T, d>& InPreV)
 	{
-		this->MarkDirty(EParticleFlags::PreV);
-		this->MPreV = InPreV;
+		MPreV.Write(InPreV, /*bInvalidate=*/true,this->RemoteData,this->MDirtyFlags,this->Proxy);
 	}
 
-	const TVector<T, d>& PreW() const { return MPreW; }
+	const TVector<T, d>& PreW() const { return MPreW.Read(); }
 	void SetPreW(const TVector<T, d>& InPreW)
 	{
-		this->MarkDirty(EParticleFlags::PreW);
-		this->MPreW = InPreW;
+		MPreW.Write(InPreW, /*bInvalidate=*/true,this->RemoteData,this->MDirtyFlags,this->Proxy);
 	}
 
-	const TVector<T, d>& P() const { return MP; }
+	const TVector<T, d>& P() const { return MP.Read(); }
 	void SetP(const TVector<T, d>& InP)
 	{
-		this->MarkDirty(EParticleFlags::P);
-		this->MP = InP;
+		MP.Write(InP, /*bInvalidate=*/true,this->RemoteData,this->MDirtyFlags,this->Proxy);
 	}
 
-	const TRotation<T, d>& Q() const { return MQ; }
+	const TRotation<T, d>& Q() const { return MQ.Read(); }
 	void SetQ(const TRotation<T, d>& InQ)
 	{
-		this->MarkDirty(EParticleFlags::Q);
-		this->MQ = InQ;
+		MQ.Write(InQ, /*bInvalidate=*/true,this->RemoteData,this->MDirtyFlags,this->Proxy);
 	}
 
-	const TVector<T, d>& F() const { return MF; }
+	const TVector<T, d>& F() const { return MF.Read(); }
 	void SetF(const TVector<T, d>& InF)
 	{
 		//question: should we do this check? only adding because we clear forces after removing from dirty list, but this marks dirty
-		if(InF != this->MF)
+		if(InF != MF.Read())
 		{
-			this->MarkDirty(EParticleFlags::F);
-			this->MF = InF;
+			MF.Write(InF, /*bInvalidate=*/true,this->RemoteData,this->MDirtyFlags,this->Proxy);
 		}
 		
 	}
 
-	const TVector<T, d>& Torque() const { return MTorque; }
+	const TVector<T, d>& Torque() const { return MTorque.Read(); }
 	void SetTorque(const TVector<T, d>& InTorque)
 	{
 		//question: should we do this check? only adding because we clear forces after removing from dirty list, but this marks dirty
-		if(InTorque != this->MTorque)
+		if(InTorque != MTorque.Read())
 		{
-			this->MarkDirty(EParticleFlags::Torque);
-			this->MTorque = InTorque;
+			MTorque.Write(InTorque, /*bInvalidate=*/true,this->RemoteData,this->MDirtyFlags,this->Proxy);
 		}
 	}
 
-	const TVector<T, d>& LinearImpulse() const { return MLinearImpulse; }
+	const TVector<T, d>& LinearImpulse() const { return MLinearImpulse.Read(); }
 	void SetLinearImpulse(const TVector<T, d>& InLinearImpulse, bool bInvalidate = true)
 	{
-		this->MarkDirty(EParticleFlags::LinearImpulse, bInvalidate);
-		this->MLinearImpulse = InLinearImpulse;
+		MLinearImpulse.Write(InLinearImpulse, bInvalidate,this->RemoteData,this->MDirtyFlags,this->Proxy);
 	}
 
-	const TVector<T, d>& AngularImpulse() const { return MAngularImpulse; }
+	const TVector<T, d>& AngularImpulse() const { return MAngularImpulse.Read(); }
 	void SetAngularImpulse(const TVector<T, d>& InAngularImpulse, bool bInvalidate = true)
 	{
-		this->MarkDirty(EParticleFlags::AngularImpulse, bInvalidate);
-		this->MAngularImpulse = InAngularImpulse;
+		MAngularImpulse.Write(InAngularImpulse, bInvalidate,this->RemoteData,this->MDirtyFlags,this->Proxy);
 	}
 
-	const PMatrix<T, d, d>& I() const { return MI; }
+	const PMatrix<T, d, d>& I() const { return MI.Read(); }
 	void SetI(const PMatrix<T, d, d>& InI)
 	{
-		this->MarkDirty(EParticleFlags::I);
-		this->MI = InI;
+		MI.Write(InI, /*bInvalidate=*/true,this->RemoteData,this->MDirtyFlags,this->Proxy);
 	}
 
-	const PMatrix<T, d, d>& InvI() const { return MInvI; }
+	const PMatrix<T, d, d>& InvI() const { return MInvI.Read(); }
 	void SetInvI(const PMatrix<T, d, d>& InInvI)
 	{
-		this->MarkDirty(EParticleFlags::InvI);
-		this->MInvI = InInvI;
+		MInvI.Write(InInvI, /*bInvalidate=*/true,this->RemoteData,this->MDirtyFlags,this->Proxy);
 	}
 
-	T M() const { return MM; }
+	T M() const { return MM.Read(); }
 	void SetM(const T& InM)
 	{
-		this->MarkDirty(EParticleFlags::M);
-		this->MM = InM;
+		MM.Write(InM, /*bInvalidate=*/true,this->RemoteData,this->MDirtyFlags,this->Proxy);
 	}
 
-	T InvM() const { return MInvM; }
+	T InvM() const { return MInvM.Read(); }
 	void SetInvM(const T& InInvM)
 	{
-		this->MarkDirty(EParticleFlags::InvM);
-		this->MInvM = InInvM;
+		MInvM.Write(InInvM, /*bInvalidate=*/true,this->RemoteData,this->MDirtyFlags,this->Proxy);
 	}
 
-	T LinearEtherDrag() const { return MLinearEtherDrag; }
+	T LinearEtherDrag() const { return MLinearEtherDrag.Read(); }
 	void SetLinearEtherDrag(const T& InLinearEtherDrag)
 	{
-		this->MarkDirty(EParticleFlags::LinearEtherDrag);
-		this->MLinearEtherDrag = InLinearEtherDrag;
+		MLinearEtherDrag.Write(InLinearEtherDrag, /*bInvalidate=*/true,this->RemoteData,this->MDirtyFlags,this->Proxy);
 	}
 
-	T AngularEtherDrag() const { return MAngularEtherDrag; }
+	T AngularEtherDrag() const { return MAngularEtherDrag.Read(); }
 	void SetAngularEtherDrag(const T& InAngularEtherDrag)
 	{
-		this->MarkDirty(EParticleFlags::AngularEtherDrag);
-		this->MAngularEtherDrag = InAngularEtherDrag;
+		MAngularEtherDrag.Write(InAngularEtherDrag, /*bInvalidate=*/true,this->RemoteData,this->MDirtyFlags,this->Proxy);
 	}
 
 	int32 Island() const { return MIsland; }
@@ -2086,15 +2074,14 @@ public:
 		this->MToBeRemovedOnFracture = bToBeRemovedOnFracture;
 	}
 
-	EObjectStateType ObjectState() const { return MObjectState; }
+	EObjectStateType ObjectState() const { return MObjectState.Read(); }
 	void SetObjectState(const EObjectStateType InState, bool bAllowEvents = false)
 	{
-		if (bAllowEvents && MObjectState != EObjectStateType::Dynamic && InState == EObjectStateType::Dynamic)
+		if (bAllowEvents && MObjectState.Read() != EObjectStateType::Dynamic && InState == EObjectStateType::Dynamic)
 		{
 			MAwakeEvent |= true;
 		}
-		MObjectState = InState;
-		this->MarkDirty(EParticleFlags::ObjectState);
+		MObjectState.Write(InState, /*bInvalidate=*/true,this->RemoteData,this->MDirtyFlags,this->Proxy);
 	}
 
 	void ClearEvents() { MAwakeEvent = false; }
@@ -2107,27 +2094,27 @@ public:
 
 
 private:
-	TRotation<T, d> MQ;
-	TVector<T, d> MPreV;
-	TVector<T, d> MPreW;
-	TVector<T, d> MP;
-	TVector<T, d> MF;
-	TVector<T, d> MTorque;
-	TVector<T, d> MLinearImpulse;
-	TVector<T, d> MAngularImpulse;
-	PMatrix<T, d, d> MI;
-	PMatrix<T, d, d> MInvI;
+	TParticleProperty<TRotation<T, d>, EParticleProperty::Q> MQ;
+	TParticleProperty<TVector<T, d>, EParticleProperty::PreV> MPreV;
+	TParticleProperty<TVector<T, d>, EParticleProperty::PreW> MPreW;
+	TParticleProperty<TVector<T, d>, EParticleProperty::P> MP;
+	TParticleProperty<TVector<T, d>, EParticleProperty::F> MF;
+	TParticleProperty<TVector<T, d>, EParticleProperty::Torque> MTorque;
+	TParticleProperty<TVector<T, d>, EParticleProperty::LinearImpulse> MLinearImpulse;
+	TParticleProperty<TVector<T, d>, EParticleProperty::AngularImpulse> MAngularImpulse;
+	TParticleProperty<PMatrix<T, d, d>, EParticleProperty::I> MI;
+	TParticleProperty<PMatrix<T, d, d>, EParticleProperty::InvI> MInvI;
 	TUniquePtr<TBVHParticles<T, d>> MCollisionParticles;
-	T MM;
-	T MInvM;
-	T MLinearEtherDrag;
-	T MAngularEtherDrag;
+	TParticleProperty<T, EParticleProperty::M> MM;
+	TParticleProperty<T, EParticleProperty::InvM> MInvM;
+	TParticleProperty<T, EParticleProperty::LinearEtherDrag> MLinearEtherDrag;
+	TParticleProperty<T, EParticleProperty::AngularEtherDrag> MAngularEtherDrag;
 	int32 MIsland;
-	int32 MCollisionGroup;
-	EObjectStateType MObjectState;
-	bool MDisabled;
+	TParticleProperty<int32, EParticleProperty::CollisionGroup> MCollisionGroup;
+	TParticleProperty<EObjectStateType, EParticleProperty::ObjectState> MObjectState;
+	TParticleProperty<bool, EParticleProperty::Disabled> MDisabled;
 	bool MToBeRemovedOnFracture;
-	bool MGravityEnabled;
+	TParticleProperty<bool, EParticleProperty::GravityEnabled> MGravityEnabled;
 	bool MInitialized;
 	bool MAwakeEvent;
 };
