Chaos 4.25 :

-Removed TGeometryParticle support in the GeometryCollectionPhysicsProxy

diff --git a/Engine/Source/Runtime/Experimental/ChaosSolvers/Private/PBDRigidsSolver.cpp b/Engine/Source/Runtime/Experimental/ChaosSolvers/Private/PBDRigidsSolver.cpp
--- a/Engine/Source/Runtime/Experimental/ChaosSolvers/Private/PBDRigidsSolver.cpp
+++ b/Engine/Source/Runtime/Experimental/ChaosSolvers/Private/PBDRigidsSolver.cpp
@@ -766,7 +766,7 @@ namespace Chaos
 					((FGeometryParticlePhysicsProxy*)(Proxy))->BufferPhysicsResults();
 					break;
 				case Chaos::EParticleType::GeometryCollection:
-					ActiveGC.AddUnique((FGeometryCollectionPhysicsProxy*)(ActiveObject.GTGeometryParticle()->Proxy));
+					ActiveGC.AddUnique((FGeometryCollectionPhysicsProxy*)(Proxy));
 					break;
 				case Chaos::EParticleType::Clustered:
 					ActiveGC.AddUnique((FGeometryCollectionPhysicsProxy*)(Proxy));
@@ -807,7 +807,7 @@ namespace Chaos
 					((FGeometryParticlePhysicsProxy*)(Proxy))->FlipBuffer();
 					break;
 				case Chaos::EParticleType::GeometryCollection:
-					ActiveGC.AddUnique((FGeometryCollectionPhysicsProxy*)(ActiveObject.GTGeometryParticle()->Proxy));
+					ActiveGC.AddUnique((FGeometryCollectionPhysicsProxy*)(Proxy));
 					break;
 				case Chaos::EParticleType::Clustered:
 					ActiveGC.AddUnique((FGeometryCollectionPhysicsProxy*)(Proxy));
@@ -851,7 +851,7 @@ namespace Chaos
 					((FGeometryParticlePhysicsProxy*)(Proxy))->PullFromPhysicsState();
 					break;
 				case Chaos::EParticleType::GeometryCollection:
-					ActiveGC.AddUnique((FGeometryCollectionPhysicsProxy*)(ActiveObject.GTGeometryParticle()->Proxy));
+					ActiveGC.AddUnique((FGeometryCollectionPhysicsProxy*)(Proxy));
 					break;
 				case Chaos::EParticleType::Clustered:
 					ActiveGC.AddUnique((FGeometryCollectionPhysicsProxy*)(Proxy));
diff --git a/Engine/Source/Runtime/Experimental/ChaosSolvers/Private/PhysicsProxy/GeometryCollectionPhysicsProxy.cpp b/Engine/Source/Runtime/Experimental/ChaosSolvers/Private/PhysicsProxy/GeometryCollectionPhysicsProxy.cpp
--- a/Engine/Source/Runtime/Experimental/ChaosSolvers/Private/PhysicsProxy/GeometryCollectionPhysicsProxy.cpp
+++ b/Engine/Source/Runtime/Experimental/ChaosSolvers/Private/PhysicsProxy/GeometryCollectionPhysicsProxy.cpp
@@ -58,7 +58,6 @@ FName FGeometryCollectionPhysicsProxy::ImplicitsAttribute("Implicits");
 FName FGeometryCollectionPhysicsProxy::SharedImplicitsAttribute("SharedImplicits");
 FName FGeometryCollectionPhysicsProxy::SolverParticleHandlesAttribute("SolverParticleHandles");
 FName FGeometryCollectionPhysicsProxy::SolverClusterHandlesAttribute("SolverClusterHandles");
-FName FGeometryCollectionPhysicsProxy::GTGeometryParticleAttribute("GTGeometryParticles");
 
 DEFINE_LOG_CATEGORY_STATIC(UGCC_LOG, Error, All);
 
@@ -550,64 +549,6 @@ void FGeometryCollectionPhysicsProxy::InitializeBodiesGT()
 		NumParticles = GetTransformGroupSize(); // SimulatableParticles.Count(true);
 		BaseParticleIndex = 0; // Are we always zero indexed now?
 
-		// Create game thread particles...
-		for (int32 Idx = 0; Idx < SimulatableParticles.Num(); ++Idx)
-		{
-			if (SimulatableParticles[Idx])
-			{
-				// Clustering data only lives on the physics thread, so there are
-				// no game thread facing API for clustering.  Therefore, we don't 
-				// construct clustered particle instances (only clustered particleS,
-				// emphasis on the "S").
-				
-				TUniquePtr<Chaos::TPBDGeometryCollectionParticle<float, 3>> Particle = 
-					Chaos::TPBDGeometryCollectionParticle<float, 3>::CreateParticle();
-
-				GTGeometryParticles[Idx] = MoveTemp(Particle); // Particle is now undefined
-				Chaos::TPBDRigidParticle<float, 3>* GTParticle = 
-					GTGeometryParticles[Idx]->CastToRigidParticle();
-
-				FTransform ParticleTransform = MassToLocal[Idx] * Transform[Idx] * Parameters.WorldTransform;
-				if (GTParticle->ObjectState() == Chaos::EObjectStateType::Dynamic)
-				{
-					// Calling SetX(), SetR(), etc will mark the particle dirty,
-					// which causes this proxy class to be put on a dirty list.
-					// That causes our PushToPhysicsState() to be called prematurely, 
-					// so we need to guard against that.
-					//GTParticle->SetX(ParticleTransform.GetTranslation(), false); 
-					//GTParticle->SetR(ParticleTransform.GetRotation().GetNormalized(), false);
-					GTParticle->SetIsland(INDEX_NONE); // doesn't mark dirty
-					GTParticle->Proxy = this;
-
-					// 
-					// AddGeometry
-					//
-
-					TArray<TUniquePtr<Chaos::FImplicitObject>> Geoms;
-					Chaos::TShapesArray<float, 3> Shapes;
-
-					if (Implicits[Idx])
-					{
-						GTParticle->SetGeometry(Implicits[Idx]); // shared pointer version
-					}
-					if (Simplicials[Idx])
-					{
-						// Not sure what to do with particles at this point.
-					}
-				}
-			}
-/*			else
-			{
-				// This is likely a cluster parent.  Create a base particle, until 
-				// we have a reason to include more data.
-
-				TUniquePtr<Chaos::TGeometryParticle<float, 3>> Particle = 
-					Chaos::TGeometryParticle<float, 3>::CreateParticle();
-				GTGeometryParticles[Idx] = MoveTemp(Particle); // Particle is now undefined
-			}
-*/
-		}
-
 		InitializedState = Parameters.InitializationState;
 	}
 }
@@ -666,13 +607,13 @@ void FGeometryCollectionPhysicsProxy::InitializeBodiesPT(
 		{
 			if (SimulatableParticles[Idx] && Children[Idx].Num() == 0) // Leaf node
 			{
-				check(GTGeometryParticles[Idx]);
-				Chaos::TPBDRigidParticle<float, 3>* GTParticle = GTGeometryParticles[Idx]->CastToRigidParticle();
+				// todo: Unblocked read access of game thread data on the physics thread.
+
 				Chaos::TPBDGeometryCollectionParticleHandle<float, 3>* Handle = Handles[NextIdx++];
 
-				if (GTParticle->ObjectState() == Chaos::EObjectStateType::Dynamic)
+				if (DynamicState[ Idx ] == (int32)EObjectStateTypeEnum::Chaos_Object_Dynamic)
 				{
-					Handle->GTGeometryParticle() = GTParticle;
+					Handle->GTGeometryParticle() = nullptr; // GeometryCollections do not support Game Thread TGeometryParticle
 					RigidsSolver->AddParticleToProxy(Handle, this);
 				}
 
@@ -878,10 +819,6 @@ void FGeometryCollectionPhysicsProxy::InitializeBodiesPT(
 							RigidChildrenTransformGroupIndex, 
 							CreationParameters);
 
-					Chaos::TGeometryParticle<float, 3>* GTParticle = 
-						GTGeometryParticles[TransformGroupIndex].Get();
-					Handle->GTGeometryParticle() = GTParticle;
-
 					SolverClusterHandles[TransformGroupIndex] = Handle;
 					SolverParticleHandles[TransformGroupIndex] = Handle;
 					HandleToTransformGroupIndex.Add(Handle, TransformGroupIndex);
@@ -1222,8 +1159,6 @@ void FGeometryCollectionPhysicsProxy::CreateDynamicAttributes()
 	// variable, so it's still defined in the header.
 	//DynamicCollection->AddExternalAttribute("RigidBodyID", FTransformCollection::TransformGroup, RigidBodyID);
 
-	// Holds TUniquePtr<Chaos::TPBDGeometryCollectionParticle<float,3>> for leaf nodes (non parents).
-	DynamicCollection->AddExternalAttribute(GTGeometryParticleAttribute, FTransformCollection::TransformGroup, GTGeometryParticles);
 	// Holds Chaos::TPBDGeometryCollectionParticleHandle<float, 3>* for all nodes
 	DynamicCollection->AddExternalAttribute(SolverParticleHandlesAttribute, FTransformCollection::TransformGroup, SolverParticleHandles);
 	// Holds Chaos::TPBDRigidClusteredParticleHandle<float, 3>* for cluster parent nodes
@@ -1980,35 +1915,23 @@ void FGeometryCollectionPhysicsProxy::BufferGameState()
 		Buffer.ShapeSimData.SetNum(NumParticles);
 		Buffer.ShapeQueryData.SetNum(NumParticles);
 
-		// To be done once:
-		int32 Idx = 0;
-		for (auto& GTParticle : GTGeometryParticles)
+		for (int32 Idx = 0; Idx < SimulatableParticles.Num(); Idx++)
 		{
-			if (GTParticle)
-			{
-				Buffer.DisabledStates[Idx] = false;
-				Buffer.SharedGeometry[Idx] = GTParticle->GeometrySharedLowLevel();
+			Buffer.DisabledStates[Idx] = !SimulatableParticles[Idx];
+			Buffer.SharedGeometry[Idx] = Implicits[Idx];
 
-				for (auto& Shape : GTParticle->ShapesArray())
-				{
-					Buffer.ShapeQueryData[Idx].Add(Shape->QueryData);
-					Buffer.ShapeSimData[Idx].Add(Shape->SimData);
 
-					// Ryan - Hard code everything collides with everything:
-					const int32 Idx1 = Buffer.ShapeSimData[Idx].Num() - 1;
-					Buffer.ShapeSimData[Idx][Idx1].Word1 = 0xFFFF; // this body channel
-					Buffer.ShapeSimData[Idx][Idx1].Word3 = 0xFFFF; // collision candidate channels
-				}
-			}
-			else
+			// Hard coded to collide all
+			// todo(chaos) : Pull shape data from the interface. 
+			for (int32 Gdx = 0; Gdx<Buffer.SharedGeometry.Num();Gdx++)
 			{
-				Buffer.DisabledStates[Idx] = true;
-				// Buffer.SpatialIdx[Idx] = ???;
-				//Buffer.UniqueIdx[Idx].Idx = INDEX_NONE;
-				//Buffer.UserData[Idx] = nullptr;
-				//Buffer.SharedGeometry[Idx] = nullptr;
+				int32 Idx1 = Buffer.ShapeQueryData[Idx].Add(FCollisionFilterData());
+				int32 Idx2 = Buffer.ShapeSimData[Idx].Add(FCollisionFilterData());
+				ensure(Idx1 == Idx2);
+
+				Buffer.ShapeSimData[Idx][Idx1].Word1 = 0xFFFF; // this body channel
+				Buffer.ShapeSimData[Idx][Idx1].Word3 = 0xFFFF; // collision candidate channels
 			}
-			++Idx;
 		}
 	}
 
@@ -2029,21 +1952,6 @@ void FGeometryCollectionPhysicsProxy::BufferGameState()
 	GeometryCollectionAlgo::GlobalMatrices(
 		DynamicCollection->Transform, DynamicCollection->Parent, Buffer.GlobalTransforms);
 	
-	int32 Idx = 0;
-	for (auto& GTParticle : GTGeometryParticles)
-	{
-		if (GTParticle)
-		{
-			if (auto* GTRigidParticle = GTParticle->CastToRigidParticle())
-			{
-				Buffer.DisabledStates[Idx] = GTRigidParticle->Disabled();
-			}
-
-			//const Chaos::TShapesArray<float,3>& ShapesArray = GTParticle->ShapesArray();
-		}
-		++Idx;
-	}
-	
 	// Make the buffer available to the consumer.
 	GameToPhysInterchange->FlipProducer(); 
 	
@@ -2460,22 +2368,6 @@ void FGeometryCollectionPhysicsProxy::PullFromPhysicsState()
 			CacheSyncFunc(TR);
 		}
 	}
-/*
-	for (int32 TransformIndex = 0; TransformIndex < NumParticles; ++TransformIndex)
-	{
-		const int32 ParticleIndex = BaseParticleIndex + TransformIndex;
-		//Chaos::TPBDRigidParticleHandle<float, 3>& Handle = *SolverParticleHandles[ParticleIndex];
-		if (Chaos::TGeometryParticle<float, 3>* Particle = GTGeometryParticles[ParticleIndex].Get())
-		{
-			const FTransform& Xf = TR.Transforms[ParticleIndex];
-			Particle->SetX(Xf.GetTranslation());
-			Particle->SetR(Xf.GetRotation());
-
-			GTDynamicCollection->Transform[TransformIndex] = Xf;
-		}
-	}
-*/
-//	UpdateGeometryCollection(TR);
 }
 
 void FGeometryCollectionPhysicsProxy::UpdateGeometryCollection(
@@ -2510,7 +2402,7 @@ void FGeometryCollectionPhysicsProxy::UpdateGeometryCollection(
 	const int32 TransformSize = GetTransformGroupSize();
 	for (int32 TransformGroupIndex = 0; TransformGroupIndex < TransformSize; ++TransformGroupIndex)
 	{
-		if (!GTGeometryParticles[TransformGroupIndex])
+		if (!SimulatableParticles[TransformGroupIndex])
 			continue;
 
 		auto* Handle = SolverParticleHandles[TransformGroupIndex];
diff --git a/Engine/Source/Runtime/Experimental/ChaosSolvers/Public/PhysicsProxy/GeometryCollectionPhysicsProxy.h b/Engine/Source/Runtime/Experimental/ChaosSolvers/Public/PhysicsProxy/GeometryCollectionPhysicsProxy.h
--- a/Engine/Source/Runtime/Experimental/ChaosSolvers/Public/PhysicsProxy/GeometryCollectionPhysicsProxy.h
+++ b/Engine/Source/Runtime/Experimental/ChaosSolvers/Public/PhysicsProxy/GeometryCollectionPhysicsProxy.h
@@ -117,7 +117,6 @@ public:
 	static FName SharedImplicitsAttribute;
 	static FName SolverParticleHandlesAttribute;
 	static FName SolverClusterHandlesAttribute;
-	static FName GTGeometryParticleAttribute;
 
 	typedef FCollisionStructureManager::FSimplicial FSimplicial;
 
@@ -355,7 +354,6 @@ private:
 
 	TManagedArray<Chaos::TPBDRigidClusteredParticleHandle<float, 3>*> SolverParticleHandles;
 	TManagedArray<Chaos::TPBDRigidClusteredParticleHandle<float, 3>*> SolverClusterHandles;
-	TManagedArray<TUniquePtr<Chaos::TGeometryParticle<float, 3>>> GTGeometryParticles;
 
 	TMap<Chaos::TPBDRigidParticleHandle<float, 3>*, int32> HandleToTransformGroupIndex;
 
