EPA edge case fix: Caused vehicles to clip through heightfields

#[fyi] Michael.Lentine, Satchit.Subramanian, Ori.Cohen, Steven.Barnett
#[review] [REDACTED_HANDLE]

diff --git a/Engine/Source/Runtime/Experimental/Chaos/Private/Chaos/HeightField.cpp b/Engine/Source/Runtime/Experimental/Chaos/Private/Chaos/HeightField.cpp
--- a/Engine/Source/Runtime/Experimental/Chaos/Private/Chaos/HeightField.cpp
+++ b/Engine/Source/Runtime/Experimental/Chaos/Private/Chaos/HeightField.cpp
@@ -1259,7 +1259,7 @@ namespace Chaos
 				return true;
 			}
 
-			return LocalContactPhi<0.f;
+			return false;
 		};
 
 		bool bResult = false;
diff --git a/Engine/Source/Runtime/Experimental/Chaos/Public/Chaos/EPA.h b/Engine/Source/Runtime/Experimental/Chaos/Public/Chaos/EPA.h
--- a/Engine/Source/Runtime/Experimental/Chaos/Public/Chaos/EPA.h
+++ b/Engine/Source/Runtime/Experimental/Chaos/Public/Chaos/EPA.h
@@ -330,12 +330,11 @@ void ComputeEPAResults(const TVec3<T>* VertsA, const TVec3<T>* VertsB, const TEP
 	else
 	{
 		OutDir /= OutPenetration;
-	}
-
-	if (Entry.Distance < 0)
-	{
-		//The origin is on the outside, so the direction is reversed
-		OutDir = -OutDir;
+		if (Entry.Distance < 0)
+		{
+			//The origin is on the outside, so the direction is reversed
+			OutDir = -OutDir;
+		}
 	}
 
 	OutA = TVec3<T>(0);
