Reimplemented solver actors for Chaos geometry collections.
 - Reimplemented TaskGraph tick mode to correctly handle multiple worlds (this involves ticking global command queues outside of world tick and mapping worlds to solver lists for each world tick so we can correctly handle multiple worlds)
 - Audited global commands that had an explicit solver capture and moved them to solver commands
 - Updated headless tests to new API for solver flags

diff --git a/Engine/Plugins/Experimental/GeometryCollectionPlugin/Source/GeometryCollectionExampleCore/Private/GeometryCollection/GeometryCollectionExampleClustering.cpp b/Engine/Plugins/Experimental/GeometryCollectionPlugin/Source/GeometryCollectionExampleCore/Private/GeometryCollection/GeometryCollectionExampleClustering.cpp
--- a/Engine/Plugins/Experimental/GeometryCollectionPlugin/Source/GeometryCollectionExampleCore/Private/GeometryCollection/GeometryCollectionExampleClustering.cpp
+++ b/Engine/Plugins/Experimental/GeometryCollectionPlugin/Source/GeometryCollectionExampleCore/Private/GeometryCollection/GeometryCollectionExampleClustering.cpp
@@ -83,7 +83,7 @@ namespace GeometryCollectionExample
 		FGeometryCollectionPhysicsProxy* PhysObject = new FGeometryCollectionPhysicsProxy(nullptr, DynamicCollection.Get(), InitFunc, nullptr, nullptr);
 		PhysObject->Initialize();
 
-		Chaos::FPBDRigidsSolver* Solver = FChaosSolversModule::GetModule()->CreateSolver(true);
+		Chaos::FPBDRigidsSolver* Solver = FChaosSolversModule::GetModule()->CreateSolver(nullptr, ESolverFlags::Standalone);
 #if CHAOS_PARTICLEHANDLE_TODO
 		Solver->RegisterObject(PhysObject);
 #endif
@@ -175,7 +175,7 @@ namespace GeometryCollectionExample
 		FGeometryCollectionPhysicsProxy* PhysObject = new FGeometryCollectionPhysicsProxy(nullptr, DynamicCollection.Get(), InitFunc, nullptr, nullptr);
 		PhysObject->Initialize();
 
-		Chaos::FPBDRigidsSolver* Solver = FChaosSolversModule::GetModule()->CreateSolver(true);
+		Chaos::FPBDRigidsSolver* Solver = FChaosSolversModule::GetModule()->CreateSolver(nullptr, ESolverFlags::Standalone);
 #if CHAOS_PARTICLEHANDLE_TODO
 		Solver->RegisterObject(PhysObject);
 #endif
@@ -320,7 +320,7 @@ namespace GeometryCollectionExample
 		PhysObject->Initialize();
 		PhysObject->SetCollisionParticlesPerObjectFraction(1.0);
 
-		Chaos::FPBDRigidsSolver* Solver = FChaosSolversModule::GetModule()->CreateSolver(true);
+		Chaos::FPBDRigidsSolver* Solver = FChaosSolversModule::GetModule()->CreateSolver(nullptr, ESolverFlags::Standalone);
 #if CHAOS_PARTICLEHANDLE_TODO
 		Solver->RegisterObject(PhysObject);
 #endif
@@ -458,7 +458,7 @@ namespace GeometryCollectionExample
 		FGeometryCollectionPhysicsProxy* PhysObject = new FGeometryCollectionPhysicsProxy(nullptr, DynamicCollection.Get(), InitFunc, nullptr, nullptr);
 		PhysObject->Initialize();
 
-		Chaos::FPBDRigidsSolver* Solver = FChaosSolversModule::GetModule()->CreateSolver(true);
+		Chaos::FPBDRigidsSolver* Solver = FChaosSolversModule::GetModule()->CreateSolver(nullptr, ESolverFlags::Standalone);
 #if CHAOS_PARTICLEHANDLE_TODO
 		Solver->RegisterObject(PhysObject);
 #endif
@@ -588,7 +588,7 @@ namespace GeometryCollectionExample
 		FGeometryCollectionPhysicsProxy* PhysObject = new FGeometryCollectionPhysicsProxy(nullptr, DynamicCollection.Get(), InitFunc, nullptr, nullptr);;
 		PhysObject->Initialize();
 
-		Chaos::FPBDRigidsSolver* Solver = FChaosSolversModule::GetModule()->CreateSolver(true);
+		Chaos::FPBDRigidsSolver* Solver = FChaosSolversModule::GetModule()->CreateSolver(nullptr, ESolverFlags::Standalone);
 #if CHAOS_PARTICLEHANDLE_TODO
 		Solver->RegisterObject(PhysObject);
 #endif
@@ -743,7 +743,7 @@ namespace GeometryCollectionExample
 		FGeometryCollectionPhysicsProxy* PhysObject = new FGeometryCollectionPhysicsProxy(nullptr, DynamicCollection.Get(), InitFunc, nullptr, nullptr);;
 		PhysObject->Initialize();
 
-		Chaos::FPBDRigidsSolver* Solver = FChaosSolversModule::GetModule()->CreateSolver(true);
+		Chaos::FPBDRigidsSolver* Solver = FChaosSolversModule::GetModule()->CreateSolver(nullptr, ESolverFlags::Standalone);
 #if CHAOS_PARTICLEHANDLE_TODO
 		Solver->RegisterObject(PhysObject);
 #endif
@@ -874,7 +874,7 @@ namespace GeometryCollectionExample
 		FGeometryCollectionPhysicsProxy* PhysObject = new FGeometryCollectionPhysicsProxy(nullptr, DynamicCollection.Get(), InitFunc, nullptr, nullptr);;
 		PhysObject->Initialize();
 
-		Chaos::FPBDRigidsSolver* Solver = FChaosSolversModule::GetModule()->CreateSolver(true);
+		Chaos::FPBDRigidsSolver* Solver = FChaosSolversModule::GetModule()->CreateSolver(nullptr, ESolverFlags::Standalone);
 #if TODO_REIMPLEMENT_RIGID_CLUSTERING
 		Chaos::FPBDRigidsSolver::FClusteringType & Clustering = Solver->GetRigidClustering();
 #endif
@@ -1112,7 +1112,7 @@ namespace GeometryCollectionExample
 		FGeometryCollectionPhysicsProxy* PhysObject = new FGeometryCollectionPhysicsProxy(nullptr, DynamicCollection.Get(), InitFunc, nullptr, nullptr);;
 		PhysObject->Initialize();
 
-		Chaos::FPBDRigidsSolver* Solver = FChaosSolversModule::GetModule()->CreateSolver(true);
+		Chaos::FPBDRigidsSolver* Solver = FChaosSolversModule::GetModule()->CreateSolver(nullptr, ESolverFlags::Standalone);
 #if TODO_REIMPLEMENT_RIGID_CLUSTERING
 		Chaos::FPBDRigidsSolver::FClusteringType & Clustering = Solver->GetRigidClustering();
 
@@ -1331,7 +1331,7 @@ namespace GeometryCollectionExample
 		PhysObject2->Initialize();
 
 
-		Chaos::FPBDRigidsSolver* Solver = FChaosSolversModule::GetModule()->CreateSolver(true);
+		Chaos::FPBDRigidsSolver* Solver = FChaosSolversModule::GetModule()->CreateSolver(nullptr, ESolverFlags::Standalone);
 #if CHAOS_PARTICLEHANDLE_TODO
 		Solver->RegisterObject(PhysObject);
 		Solver->RegisterObject(PhysObject2);
@@ -1464,7 +1464,7 @@ namespace GeometryCollectionExample
 		PhysObject2->Initialize();
 
 
-		Chaos::FPBDRigidsSolver* Solver = FChaosSolversModule::GetModule()->CreateSolver(true);
+		Chaos::FPBDRigidsSolver* Solver = FChaosSolversModule::GetModule()->CreateSolver(nullptr, ESolverFlags::Standalone);
 #if TODO_REIMPLEMENT_RIGID_CLUSTERING
 		Chaos::FPBDRigidsSolver::FClusteringType & Clustering = Solver->GetRigidClustering();
 #endif
@@ -1600,7 +1600,7 @@ namespace GeometryCollectionExample
 		FGeometryCollectionPhysicsProxy* PhysObject2 = new FGeometryCollectionPhysicsProxy(nullptr, DynamicCollection2.Get(), InitFunc2, nullptr, nullptr);;
 		PhysObject2->Initialize();
 
-		Chaos::FPBDRigidsSolver* Solver = FChaosSolversModule::GetModule()->CreateSolver(true);
+		Chaos::FPBDRigidsSolver* Solver = FChaosSolversModule::GetModule()->CreateSolver(nullptr, ESolverFlags::Standalone);
 #if TODO_REIMPLEMENT_RIGID_CLUSTERING
 		Chaos::FPBDRigidsSolver::FClusteringType & Clustering = Solver->GetRigidClustering();
 #endif
@@ -1710,7 +1710,7 @@ namespace GeometryCollectionExample
 		FGeometryCollectionPhysicsProxy* PhysObject = new FGeometryCollectionPhysicsProxy(nullptr, DynamicCollection.Get(), InitFunc, nullptr, nullptr);;
 		PhysObject->Initialize();
 
-		Chaos::FPBDRigidsSolver* Solver = FChaosSolversModule::GetModule()->CreateSolver(true);
+		Chaos::FPBDRigidsSolver* Solver = FChaosSolversModule::GetModule()->CreateSolver(nullptr, ESolverFlags::Standalone);
 #if CHAOS_PARTICLEHANDLE_TODO
 		Solver->RegisterObject(PhysObject);
 #endif
@@ -1800,7 +1800,7 @@ namespace GeometryCollectionExample
 		PhysObject->Initialize();
 
 
-		Chaos::FPBDRigidsSolver* Solver = FChaosSolversModule::GetModule()->CreateSolver(true);
+		Chaos::FPBDRigidsSolver* Solver = FChaosSolversModule::GetModule()->CreateSolver(nullptr, ESolverFlags::Standalone);
 #if TODO_REIMPLEMENT_RIGID_CLUSTERING
 		Chaos::FPBDRigidsSolver::FClusteringType & Clustering = Solver->GetRigidClustering();
 #endif
@@ -1907,7 +1907,7 @@ namespace GeometryCollectionExample
 		FFieldSystemPhysicsProxy* FieldObject = new FFieldSystemPhysicsProxy(nullptr);
 
 		FGeometryCollectionPhysicsProxy* PhysObject = new FGeometryCollectionPhysicsProxy(nullptr, DynamicCollection.Get(), InitFunc, nullptr, nullptr);
-		Chaos::FPBDRigidsSolver* Solver = FChaosSolversModule::GetModule()->CreateSolver(true);
+		Chaos::FPBDRigidsSolver* Solver = FChaosSolversModule::GetModule()->CreateSolver(nullptr, ESolverFlags::Standalone);
 #if TODO_REIMPLEMENT_RIGID_CLUSTERING
 		Chaos::FPBDRigidsSolver::FClusteringType & Clustering = Solver->GetRigidClustering();
 #endif
@@ -1985,7 +1985,7 @@ namespace GeometryCollectionExample
 		const TManagedArray<FSimplicialPointer> & Simplicials = Object->RestCollection->template GetAttribute<FSimplicialPointer>(FGeometryCollectionPhysicsProxy::SimplicialsAttribute, FTransformCollection::TransformGroup);
 
 
-		Chaos::FPBDRigidsSolver* Solver = FChaosSolversModule::GetModule()->CreateSolver(true);
+		Chaos::FPBDRigidsSolver* Solver = FChaosSolversModule::GetModule()->CreateSolver(nullptr, ESolverFlags::Standalone);
 		Solver->SetHasFloor(false);
 		Solver->SetEnabled(true);
 #if CHAOS_PARTICLEHANDLE_TODO
diff --git a/Engine/Plugins/Experimental/GeometryCollectionPlugin/Source/GeometryCollectionExampleCore/Private/GeometryCollection/GeometryCollectionExampleEvents.cpp b/Engine/Plugins/Experimental/GeometryCollectionPlugin/Source/GeometryCollectionExampleCore/Private/GeometryCollection/GeometryCollectionExampleEvents.cpp
--- a/Engine/Plugins/Experimental/GeometryCollectionPlugin/Source/GeometryCollectionExampleCore/Private/GeometryCollection/GeometryCollectionExampleEvents.cpp
+++ b/Engine/Plugins/Experimental/GeometryCollectionPlugin/Source/GeometryCollectionExampleCore/Private/GeometryCollection/GeometryCollectionExampleEvents.cpp
@@ -83,7 +83,7 @@ namespace GeometryCollectionExample
 	void Event_Handler()
 	{
 		Chaos::FEventManager EventManager(Chaos::EMultiBufferMode::Single);
-		Chaos::FPBDRigidsSolver* Solver = FChaosSolversModule::GetModule()->CreateSolver(true);
+		Chaos::FPBDRigidsSolver* Solver = FChaosSolversModule::GetModule()->CreateSolver(nullptr, ESolverFlags::Standalone);
 
 		MyEventHandler HandlerTest(EventManager);
 		MyEventHandler AnotherHandlerTest(EventManager);
diff --git a/Engine/Plugins/Experimental/GeometryCollectionPlugin/Source/GeometryCollectionExampleCore/Private/GeometryCollection/GeometryCollectionExampleSimulation.cpp b/Engine/Plugins/Experimental/GeometryCollectionPlugin/Source/GeometryCollectionExampleCore/Private/GeometryCollection/GeometryCollectionExampleSimulation.cpp
--- a/Engine/Plugins/Experimental/GeometryCollectionPlugin/Source/GeometryCollectionExampleCore/Private/GeometryCollection/GeometryCollectionExampleSimulation.cpp
+++ b/Engine/Plugins/Experimental/GeometryCollectionPlugin/Source/GeometryCollectionExampleCore/Private/GeometryCollection/GeometryCollectionExampleSimulation.cpp
@@ -58,7 +58,7 @@ namespace GeometryCollectionExample
 		// DynamicCollection pointers.  Calls FGeometryCollectionPhysicsProx::Initialize().
 		FGeometryCollectionPhysicsProxy* PhysObject = RigidBodySetup(PhysicalMaterial, RestCollection, GTDynamicCollection);
 
-		Chaos::FPBDRigidsSolver* Solver = FChaosSolversModule::GetModule()->CreateSolver(true);
+		Chaos::FPBDRigidsSolver* Solver = FChaosSolversModule::GetModule()->CreateSolver(nullptr, ESolverFlags::Standalone);
 		Solver->RegisterObject(PhysObject);
 		Solver->SetHasFloor(false);
 		Solver->SetEnabled(true);
@@ -113,7 +113,7 @@ namespace GeometryCollectionExample
 
 		FGeometryCollectionPhysicsProxy* PhysObject = RigidBodySetup(PhysicalMaterial, RestCollection, DynamicCollection, CustomFunc);
 
-		Chaos::FPBDRigidsSolver* Solver = FChaosSolversModule::GetModule()->CreateSolver(true);
+		Chaos::FPBDRigidsSolver* Solver = FChaosSolversModule::GetModule()->CreateSolver(nullptr, ESolverFlags::Standalone);
 #if CHAOS_PARTICLEHANDLE_TODO
 		Solver->RegisterObject(PhysObject);
 #endif
@@ -168,7 +168,7 @@ namespace GeometryCollectionExample
 
 		FGeometryCollectionPhysicsProxy* PhysObject = RigidBodySetup(PhysicalMaterial, RestCollection, DynamicCollection, CustomFunc);
 
-		Chaos::FPBDRigidsSolver* Solver = FChaosSolversModule::GetModule()->CreateSolver(true);
+		Chaos::FPBDRigidsSolver* Solver = FChaosSolversModule::GetModule()->CreateSolver(nullptr, ESolverFlags::Standalone);
 #if CHAOS_PARTICLEHANDLE_TODO
 		Solver->RegisterObject(PhysObject);
 #endif
@@ -223,7 +223,7 @@ namespace GeometryCollectionExample
 
 		FGeometryCollectionPhysicsProxy* PhysObject = RigidBodySetup(PhysicalMaterial, RestCollection, DynamicCollection, CustomFunc);
 
-		Chaos::FPBDRigidsSolver* Solver = FChaosSolversModule::GetModule()->CreateSolver(true);
+		Chaos::FPBDRigidsSolver* Solver = FChaosSolversModule::GetModule()->CreateSolver(nullptr, ESolverFlags::Standalone);
 #if CHAOS_PARTICLEHANDLE_TODO
 		Solver->RegisterObject(PhysObject);
 #endif
@@ -271,7 +271,7 @@ namespace GeometryCollectionExample
 
 		FGeometryCollectionPhysicsProxy* PhysObject = RigidBodySetup(PhysicalMaterial, RestCollection, DynamicCollection);
 
-		Chaos::FPBDRigidsSolver* Solver = FChaosSolversModule::GetModule()->CreateSolver(true);
+		Chaos::FPBDRigidsSolver* Solver = FChaosSolversModule::GetModule()->CreateSolver(nullptr, ESolverFlags::Standalone);
 #if CHAOS_PARTICLEHANDLE_TODO
 		Solver->RegisterObject(PhysObject);
 #endif
@@ -332,7 +332,7 @@ namespace GeometryCollectionExample
 		ObjectType[0] = (int32)EObjectStateTypeEnum::Chaos_Object_Sleeping;
 		ObjectType[1] = (int32)EObjectStateTypeEnum::Chaos_Object_Dynamic;
 
-		Chaos::FPBDRigidsSolver* Solver = FChaosSolversModule::GetModule()->CreateSolver(true);
+		Chaos::FPBDRigidsSolver* Solver = FChaosSolversModule::GetModule()->CreateSolver(nullptr, ESolverFlags::Standalone);
 #if CHAOS_PARTICLEHANDLE_TODO
 		Solver->RegisterObject(PhysObject);
 #endif
@@ -390,7 +390,7 @@ namespace GeometryCollectionExample
 		FGeometryCollectionPhysicsProxy* PhysObject = RigidBodySetup(PhysicalMaterial, RestCollection, DynamicCollection, CustomFunc);
 		PhysObject->SetCollisionParticlesPerObjectFraction( 1.0 );
 
-		Chaos::FPBDRigidsSolver* Solver = FChaosSolversModule::GetModule()->CreateSolver(true);
+		Chaos::FPBDRigidsSolver* Solver = FChaosSolversModule::GetModule()->CreateSolver(nullptr, ESolverFlags::Standalone);
 #if CHAOS_PARTICLEHANDLE_TODO
 		Solver->RegisterObject(PhysObject);
 #endif
diff --git a/Engine/Plugins/Experimental/GeometryCollectionPlugin/Source/GeometryCollectionExampleCore/Private/GeometryCollection/GeometryCollectionExampleSimulationField.cpp b/Engine/Plugins/Experimental/GeometryCollectionPlugin/Source/GeometryCollectionExampleCore/Private/GeometryCollection/GeometryCollectionExampleSimulationField.cpp
--- a/Engine/Plugins/Experimental/GeometryCollectionPlugin/Source/GeometryCollectionExampleCore/Private/GeometryCollection/GeometryCollectionExampleSimulationField.cpp
+++ b/Engine/Plugins/Experimental/GeometryCollectionPlugin/Source/GeometryCollectionExampleCore/Private/GeometryCollection/GeometryCollectionExampleSimulationField.cpp
@@ -52,7 +52,7 @@ namespace GeometryCollectionExample
 
 		FGeometryCollectionPhysicsProxy* PhysObject = RigidBodySetup(PhysicalMaterial, RestCollection, DynamicCollection);
 
-		Chaos::FPBDRigidsSolver* Solver = Module->CreateSolver(true);
+		Chaos::FPBDRigidsSolver* Solver = Module->CreateSolver(nullptr, ESolverFlags::Standalone);
 		Solver->RegisterObject(PhysObject);
 		Solver->SetHasFloor(false);
 		Solver->SetIsFloorAnalytic(true);
@@ -140,7 +140,7 @@ namespace GeometryCollectionExample
 		FFieldSystemPhysicsProxy* FieldObject = new FFieldSystemPhysicsProxy(nullptr);
 		FGeometryCollectionPhysicsProxy* PhysObject = RigidBodySetup(PhysicalMaterial, RestCollection, DynamicCollection, CustomFunc);
 
-		Chaos::FPBDRigidsSolver* Solver = FChaosSolversModule::GetModule()->CreateSolver(true);
+		Chaos::FPBDRigidsSolver* Solver = FChaosSolversModule::GetModule()->CreateSolver(nullptr, ESolverFlags::Standalone);
 		Solver->RegisterObject(PhysObject);
 		Solver->RegisterObject(FieldObject);
 		Solver->SetHasFloor(false);
@@ -226,7 +226,7 @@ namespace GeometryCollectionExample
 
 		FFieldSystemPhysicsProxy* FieldObject = new FFieldSystemPhysicsProxy(nullptr);
 
-		Chaos::FPBDRigidsSolver* Solver = FChaosSolversModule::GetModule()->CreateSolver(true);
+		Chaos::FPBDRigidsSolver* Solver = FChaosSolversModule::GetModule()->CreateSolver(nullptr, ESolverFlags::Standalone);
 		Solver->RegisterObject(PhysObject);
 		Solver->RegisterObject(FieldObject);
 		Solver->SetHasFloor(false);
@@ -313,7 +313,7 @@ namespace GeometryCollectionExample
 
 		FFieldSystemPhysicsProxy* FieldObject = new FFieldSystemPhysicsProxy(nullptr);
 
-		Chaos::FPBDRigidsSolver* Solver = FChaosSolversModule::GetModule()->CreateSolver(true);
+		Chaos::FPBDRigidsSolver* Solver = FChaosSolversModule::GetModule()->CreateSolver(nullptr, ESolverFlags::Standalone);
 		Solver->RegisterObject(PhysObject);
 		Solver->RegisterObject(FieldObject);
 		Solver->SetHasFloor(false);
@@ -394,7 +394,7 @@ namespace GeometryCollectionExample
 
 		FFieldSystemPhysicsProxy* FieldObject = new FFieldSystemPhysicsProxy(nullptr);
 
-		Chaos::FPBDRigidsSolver* Solver = FChaosSolversModule::GetModule()->CreateSolver(true);
+		Chaos::FPBDRigidsSolver* Solver = FChaosSolversModule::GetModule()->CreateSolver(nullptr, ESolverFlags::Standalone);
 		Solver->RegisterObject(PhysObject);
 		Solver->RegisterObject(FieldObject);
 		Solver->SetHasFloor(false);
@@ -485,7 +485,7 @@ namespace GeometryCollectionExample
 
 		FFieldSystemPhysicsProxy* FieldObject = new FFieldSystemPhysicsProxy(nullptr);
 
-		Chaos::FPBDRigidsSolver* Solver = FChaosSolversModule::GetModule()->CreateSolver(true);
+		Chaos::FPBDRigidsSolver* Solver = FChaosSolversModule::GetModule()->CreateSolver(nullptr, ESolverFlags::Standalone);
 		Solver->RegisterObject(PhysObject);
 		Solver->RegisterObject(FieldObject);
 		Solver->SetHasFloor(false);
@@ -555,7 +555,7 @@ namespace GeometryCollectionExample
 
 		FFieldSystemPhysicsProxy* FieldObject = new FFieldSystemPhysicsProxy(nullptr);
 
-		Chaos::FPBDRigidsSolver* Solver = FChaosSolversModule::GetModule()->CreateSolver(true);
+		Chaos::FPBDRigidsSolver* Solver = FChaosSolversModule::GetModule()->CreateSolver(nullptr, ESolverFlags::Standalone);
 		Solver->RegisterObject(PhysObject);
 		Solver->RegisterObject(FieldObject);
 		Solver->SetHasFloor(false);
@@ -637,7 +637,7 @@ namespace GeometryCollectionExample
 		//
 		FGeometryCollectionPhysicsProxy* PhysObject = RigidBodySetup(PhysicalMaterial, RestCollection, DynamicCollection);
 
-		Chaos::FPBDRigidsSolver* Solver = FChaosSolversModule::GetModule()->CreateSolver(true);
+		Chaos::FPBDRigidsSolver* Solver = FChaosSolversModule::GetModule()->CreateSolver(nullptr, ESolverFlags::Standalone);
 		Solver->RegisterObject(PhysObject);
 		Solver->SetHasFloor(true);
 		Solver->SetEnabled(true);
@@ -713,7 +713,7 @@ namespace GeometryCollectionExample
 		FalloffField->Position = FVector(0.0, 0.0, 0.0);
 		FalloffField->Falloff = EFieldFalloffType::Field_FallOff_None;
 
-		Chaos::FPBDRigidsSolver* Solver = FChaosSolversModule::GetModule()->CreateSolver(true);
+		Chaos::FPBDRigidsSolver* Solver = FChaosSolversModule::GetModule()->CreateSolver(nullptr, ESolverFlags::Standalone);
 		FFieldSystemPhysicsProxy* FieldObject = new FFieldSystemPhysicsProxy(nullptr);
 
 		FGeometryCollectionPhysicsProxy* PhysObject = RigidBodySetup(PhysicalMaterial, RestCollection, DynamicCollection, CustomFunc);
@@ -809,7 +809,7 @@ namespace GeometryCollectionExample
 		FalloffField->Position = FVector(0.0, 0.0, 0.0);
 		FalloffField->Falloff = EFieldFalloffType::Field_FallOff_None;
 
-		Chaos::FPBDRigidsSolver* Solver = FChaosSolversModule::GetModule()->CreateSolver(true);
+		Chaos::FPBDRigidsSolver* Solver = FChaosSolversModule::GetModule()->CreateSolver(nullptr, ESolverFlags::Standalone);
 		FFieldSystemPhysicsProxy* FieldObject = new FFieldSystemPhysicsProxy(nullptr);
 		
 		FGeometryCollectionPhysicsProxy* PhysObject = RigidBodySetup(PhysicalMaterial, RestCollection, DynamicCollection, CustomFunc);
@@ -916,7 +916,7 @@ namespace GeometryCollectionExample
 		FalloffField->Position = FVector(350.0, 0.0, 0.0);
 		FalloffField->Falloff = EFieldFalloffType::Field_FallOff_None;
 
-		Chaos::FPBDRigidsSolver* Solver = FChaosSolversModule::GetModule()->CreateSolver(true);
+		Chaos::FPBDRigidsSolver* Solver = FChaosSolversModule::GetModule()->CreateSolver(nullptr, ESolverFlags::Standalone);
 		FFieldSystemPhysicsProxy* FieldObject = new FFieldSystemPhysicsProxy(nullptr);
 		
 		FGeometryCollectionPhysicsProxy* PhysObject = RigidBodySetup(PhysicalMaterial, RestCollection, DynamicCollection, CustomFunc);
@@ -1025,7 +1025,7 @@ namespace GeometryCollectionExample
 		FalloffField->Position = FVector(0.0, 0.0, 0.0);
 		FalloffField->Falloff = EFieldFalloffType::Field_FallOff_None;
 
-		Chaos::FPBDRigidsSolver* Solver = FChaosSolversModule::GetModule()->CreateSolver(true);
+		Chaos::FPBDRigidsSolver* Solver = FChaosSolversModule::GetModule()->CreateSolver(nullptr, ESolverFlags::Standalone);
 		FFieldSystemPhysicsProxy* FieldObject = new FFieldSystemPhysicsProxy(nullptr);
 
 		FGeometryCollectionPhysicsProxy* PhysObject = RigidBodySetup(PhysicalMaterial, RestCollection, DynamicCollection, CustomFunc);
diff --git a/Engine/Plugins/Experimental/GeometryCollectionPlugin/Source/GeometryCollectionExampleCore/Private/GeometryCollection/GeometryCollectionExampleSimulationSolver.cpp b/Engine/Plugins/Experimental/GeometryCollectionPlugin/Source/GeometryCollectionExampleCore/Private/GeometryCollection/GeometryCollectionExampleSimulationSolver.cpp
--- a/Engine/Plugins/Experimental/GeometryCollectionPlugin/Source/GeometryCollectionExampleCore/Private/GeometryCollection/GeometryCollectionExampleSimulationSolver.cpp
+++ b/Engine/Plugins/Experimental/GeometryCollectionPlugin/Source/GeometryCollectionExampleCore/Private/GeometryCollection/GeometryCollectionExampleSimulationSolver.cpp
@@ -33,7 +33,7 @@ namespace GeometryCollectionExample
 	template<class T>
 	void Solver_AdvanceNoObjects()
 	{
-		Chaos::FPBDRigidsSolver* Solver = FChaosSolversModule::GetModule()->CreateSolver(true);
+		Chaos::FPBDRigidsSolver* Solver = FChaosSolversModule::GetModule()->CreateSolver(nullptr, ESolverFlags::Standalone);
 		Solver->SetHasFloor(false);
 		Solver->SetEnabled(true);
 		Solver->AdvanceSolverBy(1 / 24.);
@@ -67,7 +67,8 @@ namespace GeometryCollectionExample
 		FGeometryCollectionPhysicsProxy* PhysObject = new FGeometryCollectionPhysicsProxy(nullptr, DynamicCollection.Get(), InitFunc, nullptr, nullptr);;
 		PhysObject->Initialize();
 
-		Chaos::FPBDRigidsSolver* Solver = FChaosSolversModule::GetModule()->CreateSolver(true);
+		Chaos::FPBDRigidsSolver* Solver = FChaosSolversModule::GetModule()->CreateSolver(nullptr, ESolverFlags::Standalone);
+
 		Solver->RegisterObject(PhysObject);
 
 		Solver->SetHasFloor(false);
@@ -124,7 +125,7 @@ namespace GeometryCollectionExample
 		FGeometryCollectionPhysicsProxy* PhysObject = new FGeometryCollectionPhysicsProxy(nullptr, DynamicCollection.Get(), InitFunc, nullptr, nullptr);
 		PhysObject->Initialize();
 
-		Chaos::FPBDRigidsSolver* Solver = FChaosSolversModule::GetModule()->CreateSolver(true);
+		Chaos::FPBDRigidsSolver* Solver = FChaosSolversModule::GetModule()->CreateSolver(nullptr, ESolverFlags::Standalone);
 #if CHAOS_PARTICLEHANDLE_TODO
 		Solver->RegisterObject(PhysObject);
 #endif
@@ -155,7 +156,7 @@ namespace GeometryCollectionExample
 		TUniquePtr<Chaos::FChaosPhysicsMaterial> PhysicalMaterial = MakeUnique<Chaos::FChaosPhysicsMaterial>();
 		InitMaterialToZero(PhysicalMaterial);
 
-		Chaos::FPBDRigidsSolver* Solver = FChaosSolversModule::GetModule()->CreateSolver(true);
+		Chaos::FPBDRigidsSolver* Solver = FChaosSolversModule::GetModule()->CreateSolver(nullptr, ESolverFlags::Standalone);
 		Solver->SetHasFloor(true);
 		Solver->SetEnabled(true);
 		TArray<TSharedPtr<FGeometryCollection> > RestArray;
@@ -216,7 +217,7 @@ namespace GeometryCollectionExample
 		PhysicalMaterial = MakeUnique<Chaos::FChaosPhysicsMaterial>();
 		InitMaterialToZero(PhysicalMaterial);
 
-		*SolverInOut = FChaosSolversModule::GetModule()->CreateSolver(true);
+		*SolverInOut = FChaosSolversModule::GetModule()->CreateSolver(nullptr, ESolverFlags::Standalone);
 		Chaos::FPBDRigidsSolver* Solver = *SolverInOut;
 		Solver->SetHasFloor(true);
 		Solver->SetEnabled(true);
diff --git a/Engine/Plugins/Experimental/GeometryCollectionPlugin/Source/GeometryCollectionExampleCore/Private/GeometryCollection/GeometryCollectionExampleSimulationStreaming.cpp b/Engine/Plugins/Experimental/GeometryCollectionPlugin/Source/GeometryCollectionExampleCore/Private/GeometryCollection/GeometryCollectionExampleSimulationStreaming.cpp
--- a/Engine/Plugins/Experimental/GeometryCollectionPlugin/Source/GeometryCollectionExampleCore/Private/GeometryCollection/GeometryCollectionExampleSimulationStreaming.cpp
+++ b/Engine/Plugins/Experimental/GeometryCollectionPlugin/Source/GeometryCollectionExampleCore/Private/GeometryCollection/GeometryCollectionExampleSimulationStreaming.cpp
@@ -26,7 +26,7 @@ namespace GeometryCollectionExample
 	template<class T>
 	void RigidBodies_Streaming_StartSolverEmpty()
 	{
-		Chaos::FPBDRigidsSolver* Solver = FChaosSolversModule::GetModule()->CreateSolver(true);
+		Chaos::FPBDRigidsSolver* Solver = FChaosSolversModule::GetModule()->CreateSolver(nullptr, ESolverFlags::Standalone);
 		Solver->SetHasFloor(false);
 		Solver->SetEnabled(true);
 
@@ -66,7 +66,7 @@ namespace GeometryCollectionExample
 		typename SimulationObjects<T>::FParameters P;
 		P.CollisionGroup = -1;
 
-		Chaos::FPBDRigidsSolver* Solver = FChaosSolversModule::GetModule()->CreateSolver(true);
+		Chaos::FPBDRigidsSolver* Solver = FChaosSolversModule::GetModule()->CreateSolver(nullptr, ESolverFlags::Standalone);
 		Solver->SetHasFloor(false);
 		Solver->SetEnabled(true);
 		Solver->AdvanceSolverBy(1 / 24.);
@@ -115,7 +115,7 @@ namespace GeometryCollectionExample
 		P.EnableClustering = true;
 		P.ClusterGroupIndex = 1;
 
-		Chaos::FPBDRigidsSolver* Solver = FChaosSolversModule::GetModule()->CreateSolver(true);
+		Chaos::FPBDRigidsSolver* Solver = FChaosSolversModule::GetModule()->CreateSolver(nullptr, ESolverFlags::Standalone);
 		Solver->SetHasFloor(false);
 		Solver->SetEnabled(true);
 		Solver->AdvanceSolverBy(1 / 24.);
diff --git a/Engine/Source/Editor/UnrealEd/Private/EditorEngine.cpp b/Engine/Source/Editor/UnrealEd/Private/EditorEngine.cpp
--- a/Engine/Source/Editor/UnrealEd/Private/EditorEngine.cpp
+++ b/Engine/Source/Editor/UnrealEd/Private/EditorEngine.cpp
@@ -220,6 +220,11 @@
 #include "StudioAnalytics.h"
 #include "Engine/LevelScriptActor.h"
 
+#if WITH_CHAOS
+#include "ChaosSolversModule.h"
+#endif
+
+
 DEFINE_LOG_CATEGORY_STATIC(LogEditor, Log, All);
 
 #define LOCTEXT_NAMESPACE "UnrealEd.Editor"
@@ -1489,6 +1494,15 @@ void UEditorEngine::Tick( float DeltaSeconds, bool bIdleMode )
 	bool bAWorldTicked = false;
 	ELevelTick TickType = IsRealtime ? LEVELTICK_ViewportsOnly : LEVELTICK_TimeOnly;
 
+#if WITH_CHAOS
+	// Before we begin ticking any of our worlds, dispatch the global command lists and queues for physics
+	FChaosSolversModule* ChaosModule = FChaosSolversModule::GetModule();
+	if(ensure(ChaosModule))
+	{
+		ChaosModule->DispatchGlobalCommands();
+	}
+#endif
+
 	if( bShouldTickEditorWorld )
 	{ 
 		//EditorContext.World()->FXSystem->Resume();
diff --git a/Engine/Source/Editor/UnrealEd/UnrealEd.Build.cs b/Engine/Source/Editor/UnrealEd/UnrealEd.Build.cs
--- a/Engine/Source/Editor/UnrealEd/UnrealEd.Build.cs
+++ b/Engine/Source/Editor/UnrealEd/UnrealEd.Build.cs
@@ -326,6 +326,11 @@ public class UnrealEd : ModuleRules
 
 		SetupModulePhysicsSupport(Target);
 
+		if(Target.bCompileChaos && Target.bUseChaos)
+        {
+			PrivateDependencyModuleNames.Add("ChaosSolvers");
+        }
+
 		if (Target.bCompileRecast)
 		{
 			PrivateDependencyModuleNames.Add("Navmesh");
diff --git a/Engine/Source/Programs/HeadlessChaos/Private/HeadlessChaosTestSolverCommandList.cpp b/Engine/Source/Programs/HeadlessChaos/Private/HeadlessChaosTestSolverCommandList.cpp
--- a/Engine/Source/Programs/HeadlessChaos/Private/HeadlessChaosTestSolverCommandList.cpp
+++ b/Engine/Source/Programs/HeadlessChaos/Private/HeadlessChaosTestSolverCommandList.cpp
@@ -31,7 +31,7 @@ namespace ChaosTest {
 		Module->ChangeThreadingMode(EChaosThreadingMode::DedicatedThread);
 
 		// Make a solver
-		Chaos::FPhysicsSolver* Solver = FChaosSolversModule::GetModule()->CreateSolver(true);
+		Chaos::FPhysicsSolver* Solver = FChaosSolversModule::GetModule()->CreateSolver(nullptr, ESolverFlags::Standalone);
 		Solver->SetEnabled(true);
 
 		// Make a particle
diff --git a/Engine/Source/Programs/HeadlessChaos/Private/HeadlessChaosTestSolverProxies.cpp b/Engine/Source/Programs/HeadlessChaos/Private/HeadlessChaosTestSolverProxies.cpp
--- a/Engine/Source/Programs/HeadlessChaos/Private/HeadlessChaosTestSolverProxies.cpp
+++ b/Engine/Source/Programs/HeadlessChaos/Private/HeadlessChaosTestSolverProxies.cpp
@@ -30,7 +30,7 @@ namespace ChaosTest {
 		Module->ChangeThreadingMode(EChaosThreadingMode::SingleThread);
 
 		// Make a solver
-		Chaos::FPhysicsSolver* Solver = Module->CreateSolver(true);
+		Chaos::FPhysicsSolver* Solver = Module->CreateSolver(nullptr, ESolverFlags::Standalone);
 		Solver->SetEnabled(true);
 
 		// Make a particle
@@ -89,7 +89,7 @@ namespace ChaosTest {
 		Module->ChangeThreadingMode(EChaosThreadingMode::DedicatedThread);
 
 		// Make a solver
-		Chaos::FPhysicsSolver* Solver = Module->CreateSolver(true);
+		Chaos::FPhysicsSolver* Solver = Module->CreateSolver(nullptr, ESolverFlags::Standalone);
 		Solver->SetEnabled(true);
 
 		// Make a particle
@@ -145,7 +145,7 @@ namespace ChaosTest {
 		Module->ChangeThreadingMode(EChaosThreadingMode::SingleThread);
 
 		// Make a solver
-		Chaos::FPhysicsSolver* Solver = Module->CreateSolver(true);
+		Chaos::FPhysicsSolver* Solver = Module->CreateSolver(nullptr, ESolverFlags::Standalone);
 		Solver->SetEnabled(true);
 
 		// Make a particle
diff --git a/Engine/Source/Runtime/Engine/Private/GameEngine.cpp b/Engine/Source/Runtime/Engine/Private/GameEngine.cpp
--- a/Engine/Source/Runtime/Engine/Private/GameEngine.cpp
+++ b/Engine/Source/Runtime/Engine/Private/GameEngine.cpp
@@ -72,6 +72,10 @@
 	#include "IMediaModule.h"
 #endif
 
+#if WITH_CHAOS
+#include "ChaosSolversModule.h"
+#endif
+
 CSV_DECLARE_CATEGORY_MODULE_EXTERN(CORE_API, Basic);
 
 ENGINE_API bool GDisallowNetworkTravel = false;
@@ -1631,6 +1635,15 @@ void UGameEngine::Tick( float DeltaSeconds, bool bIdleMode )
 		FStudioAnalytics::Tick(DeltaSeconds);
 	}
 
+#if WITH_CHAOS
+	// Before we begin ticking any of our worlds, dispatch the global command lists and queues for physics
+	FChaosSolversModule* ChaosModule = FChaosSolversModule::GetModule();
+	if(ensure(ChaosModule))
+	{
+		ChaosModule->DispatchGlobalCommands();
+	}
+#endif
+
 	// -----------------------------------------------------
 	// Begin ticking worlds
 	// -----------------------------------------------------
diff --git a/Engine/Source/Runtime/Engine/Private/PhysicsEngine/Experimental/PhysScene_Chaos.cpp b/Engine/Source/Runtime/Engine/Private/PhysicsEngine/Experimental/PhysScene_Chaos.cpp
--- a/Engine/Source/Runtime/Engine/Private/PhysicsEngine/Experimental/PhysScene_Chaos.cpp
+++ b/Engine/Source/Runtime/Engine/Private/PhysicsEngine/Experimental/PhysScene_Chaos.cpp
@@ -211,7 +211,7 @@ private:
 
 		FSpacialDebugDraw DrawInterface(WorldPtr);
 
-		const TArray<FPhysicsSolver*>& Solvers = ChaosModule->GetSolvers();
+		const TArray<FPhysicsSolver*>& Solvers = ChaosModule->GetAllSolvers();
 
 		for(FPhysicsSolver* Solver : Solvers)
 		{
@@ -390,7 +390,9 @@ FPhysScene_Chaos::FPhysScene_Chaos(AActor* InSolverActor
 	ChaosModule = FModuleManager::Get().GetModulePtr<FChaosSolversModule>("ChaosSolvers");
 	check(ChaosModule);
 
-	SceneSolver = ChaosModule->CreateSolver(false
+	UWorld* WorldPtr = SolverActor.Get() ? SolverActor->GetWorld() : nullptr;
+
+	SceneSolver = ChaosModule->CreateSolver(WorldPtr
 #if CHAOS_CHECKED
 	, DebugName
 #endif
@@ -566,10 +568,10 @@ void AddPhysicsProxy(ObjectType* InObject, Chaos::FPhysicsSolver* InSolver, Chao
 	if(/*bDedicatedThread && */InDispatcher)
 	{
 		// Pass the proxy off to the physics thread
-		InDispatcher->EnqueueCommandImmediate([InObject, InSolver](Chaos::FPersistentPhysicsTask* PhysThread)
+		InDispatcher->EnqueueCommandImmediate(InSolver, [InObject](Chaos::FPhysicsSolver* InnerSolver)
 		{
 #if CHAOS_PARTICLEHANDLE_TODO
-			InSolver->RegisterObject(InObject);
+			InnerSolver->RegisterObject(InObject);
 #endif
 		});
 	}
@@ -635,19 +637,21 @@ void FPhysScene_Chaos::AddObject(UPrimitiveComponent* Component, FFieldSystemPhy
 
 	if (Chaos::IDispatcher* Dispatcher = GetDispatcher())
 	{
-		for (Chaos::FPhysicsSolver* Solver : ChaosModule->GetSolvers())
+		TArray<Chaos::FPhysicsSolver*> WorldSolverList = ChaosModule->GetSolversMutable(Component->GetWorld());
+
+		for(Chaos::FPhysicsSolver* Solver : WorldSolverList)
 		{
-			if (true || Solver->HasActiveParticles())
+			if(true || Solver->HasActiveParticles())
 			{
 				Solver->RegisterObject(InObject);
 
-				if (/*bDedicatedThread && */Dispatcher)
+				if(/*bDedicatedThread && */Dispatcher)
 				{
 					// Pass the proxy off to the physics thread
 					Dispatcher->EnqueueCommandImmediate([InObject, InSolver = Solver](Chaos::FPersistentPhysicsTask* PhysThread)
-					{
-						InSolver->RegisterObject(InObject);
-					});
+						{
+							InSolver->RegisterObject(InObject);
+						});
 				}
 			}
 		}
@@ -839,19 +843,21 @@ void FPhysScene_Chaos::RemoveObject(FFieldSystemPhysicsProxy* InObject)
 
 	if (Chaos::IDispatcher* Dispatcher = GetDispatcher())
 	{
-		for (Chaos::FPhysicsSolver* Solver : ChaosModule->GetSolvers())
+		TArray<Chaos::FPhysicsSolver*> SolverList = ChaosModule->GetSolversMutable(CurrSceneSolver->GetOwner());
+		
+		for(Chaos::FPhysicsSolver* Solver : SolverList)
 		{
-			if (true || Solver->HasActiveParticles())
+			if(true || Solver->HasActiveParticles())
 			{
 				Solver->RegisterObject(InObject);
 
-				if (/*bDedicatedThread && */Dispatcher)
+				if(/*bDedicatedThread && */Dispatcher)
 				{
 					// Pass the proxy off to the physics thread
 					Dispatcher->EnqueueCommandImmediate([InObject, InSolver = Solver](Chaos::FPersistentPhysicsTask* PhysThread)
-					{
-						InSolver->UnregisterObject(InObject);
-					});
+						{
+							InSolver->UnregisterObject(InObject);
+						});
 				}
 			}
 		}
@@ -1876,9 +1882,15 @@ void FPhysScene_ChaosInterface::StartFrame()
 			OnPhysScenePreTick.Broadcast(this, Dt);
 			OnPhysSceneStep.Broadcast(this, Dt);
 
+			TArray<Chaos::FPhysicsSolver*> SolverList;
+			SolverModule->GetSolversMutable(GetOwningWorld(), SolverList);
+			
 			if (FPhysicsSolver* Solver = GetSolver())
 			{
 				Solver->PushPhysicsState(Dispatcher);
+
+				// Make sure our solver is in the list
+				SolverList.AddUnique(Solver);
 			}
 
 			FGraphEventRef SimulationCompleteEvent = FGraphEvent::CreateGraphEvent();
@@ -1886,7 +1898,7 @@ void FPhysScene_ChaosInterface::StartFrame()
 			// Need to fire off a parallel task to handle running physics commands and
 			// ticking the scene while the engine continues on until TG_EndPhysics
 			// (this should happen in TG_StartPhysics)
-			PhysicsTickTask = TGraphTask<FPhysicsTickTask>::CreateTask(nullptr, ENamedThreads::GameThread).ConstructAndDispatchWhenReady(SimulationCompleteEvent, GetSolver(), Dt);
+			PhysicsTickTask = TGraphTask<FPhysicsTickTask>::CreateTask(nullptr, ENamedThreads::GameThread).ConstructAndDispatchWhenReady(SimulationCompleteEvent, SolverList, Dt);
 
 			// Setup post simulate tasks
 			if (PhysicsTickTask.GetReference())
@@ -1980,12 +1992,16 @@ void FPhysScene_ChaosInterface::EndFrame(ULineBatchComponent* InLineBatcher)
 		//flush queue so we can merge the two threads
 		Dispatcher->Execute();
 
-		// Make a list of solvers to process. This is just our solver if we have one,
-		// and if not then it's all solvers in the solvers module.
-		const TArray<FPhysicsSolver*>& SolverList
-			= (GetSolver() == nullptr)
-			? SolverModule->GetSolvers()
-			: [&]() { TArray<FPhysicsSolver*> Solvers; Solvers.Init(GetSolver(), 1); return Solvers; }();
+		// Make a list of solvers to process. This is a list of all solvers registered to our world
+		// And our internal base scene solver.
+		TArray<Chaos::FPhysicsSolver*> SolverList;
+		SolverModule->GetSolversMutable(GetOwningWorld(), SolverList);
+
+		if(FPhysicsSolver* Solver = GetSolver())
+		{
+			// Make sure our solver is in the list
+			SolverList.AddUnique(Solver);
+		}
 
 		// flush solver queues
 		for (FPhysicsSolver* Solver : SolverList)
@@ -2234,29 +2250,40 @@ void FPhysScene_ChaosInterface::CompleteSceneSimulation(ENamedThreads::Type Curr
 
 		check(Module);
 
-		const TArray<FPhysicsSolver*>& SolverList = Module->GetSolvers();
+		TArray<FPhysicsSolver*> SolverList = Module->GetSolversMutable(GetOwningWorld());
+		FPhysicsSolver* SceneSolver = GetSolver();
+
 		TArray<FPhysicsSolver*> ActiveSolvers;
-		ActiveSolvers.Reserve(SolverList.Num());
 
-		// #BG calculate active solver list once as we dispatch our first task
-		for(FPhysicsSolver* Solver : SolverList)
+		if(SolverList.Num() > 0)
 		{
-			if(Solver->HasActiveParticles())
+			ActiveSolvers.Reserve(SolverList.Num());
+
+			// #BG calculate active solver list once as we dispatch our first task
+			for(FPhysicsSolver* Solver : SolverList)
 			{
-				ActiveSolvers.Add(Solver);
+				if(Solver->HasActiveParticles())
+				{
+					ActiveSolvers.Add(Solver);
+				}
 			}
 		}
 
+		if(SceneSolver && SceneSolver->HasActiveParticles())
+		{
+			ActiveSolvers.AddUnique(SceneSolver);
+		}
+
 		const int32 NumActiveSolvers = ActiveSolvers.Num();
 
 		PhysicsParallelFor(NumActiveSolvers, [&](int32 Index)
-		{
-			FPhysicsSolver* Solver = ActiveSolvers[Index];
+			{
+				FPhysicsSolver* Solver = ActiveSolvers[Index];
 
-			Solver->GetActiveParticlesBuffer()->CaptureSolverData(Solver);
-			Solver->BufferPhysicsResults();
-			Solver->FlipBuffers();
-		});
+				Solver->GetActiveParticlesBuffer()->CaptureSolverData(Solver);
+				Solver->BufferPhysicsResults();
+				Solver->FlipBuffers();
+			});
 	}
 }
 
diff --git a/Engine/Source/Runtime/Experimental/Chaos/Private/ChaosStats.cpp b/Engine/Source/Runtime/Experimental/Chaos/Private/ChaosStats.cpp
--- a/Engine/Source/Runtime/Experimental/Chaos/Private/ChaosStats.cpp
+++ b/Engine/Source/Runtime/Experimental/Chaos/Private/ChaosStats.cpp
@@ -10,6 +10,7 @@ DEFINE_STAT(STAT_IntegrateSolver);
 DEFINE_STAT(STAT_SyncProxies);
 DEFINE_STAT(STAT_PhysCommands);
 DEFINE_STAT(STAT_TaskCommands);
+DEFINE_STAT(STAT_WaitGlobalCommands);
 DEFINE_STAT(STAT_KinematicUpdate);
 DEFINE_STAT(STAT_BeginFrame);
 DEFINE_STAT(STAT_EndFrame);
diff --git a/Engine/Source/Runtime/Experimental/Chaos/Public/Chaos/Framework/PhysicsSolverBase.h b/Engine/Source/Runtime/Experimental/Chaos/Public/Chaos/Framework/PhysicsSolverBase.h
--- a/Engine/Source/Runtime/Experimental/Chaos/Public/Chaos/Framework/PhysicsSolverBase.h
+++ b/Engine/Source/Runtime/Experimental/Chaos/Public/Chaos/Framework/PhysicsSolverBase.h
@@ -5,6 +5,7 @@
 #include "Chaos/Matrix.h"
 #include "Misc/ScopeLock.h"
 
+class UObject;
 class IPhysicsProxyBase; // WTF - not in Chaos?
 
 namespace Chaos
@@ -14,10 +15,6 @@ namespace Chaos
 	{
 	public:
 
-		FPhysicsSolverBase(const EMultiBufferMode BufferingModeIn)
-			: BufferMode(BufferingModeIn)
-		{}
-		
 		void ChangeBufferMode(EMultiBufferMode InBufferMode);
 
 		void AddDirtyProxy(IPhysicsProxyBase * ProxyBaseIn)
@@ -29,6 +26,17 @@ namespace Chaos
 			DirtyProxiesSet.Remove(ProxyBaseIn);
 		}
 
+		const UObject* GetOwner() const
+		{ 
+			return Owner; 
+		}
+
+		void SetOwner(const UObject* InOwner)
+		{
+			Owner = InOwner;
+		}
+
+
 #if CHAOS_CHECKED
 		void SetDebugName(const FName& Name)
 		{
@@ -42,13 +50,38 @@ namespace Chaos
 #endif
 
 	protected:
-			EMultiBufferMode BufferMode;
 
-			// Input Proxy Map
-			TSet< IPhysicsProxyBase *> DirtyProxiesSet;
+		/** Protected construction so callers still have to go through the module to create new instances */
+		FPhysicsSolverBase(const EMultiBufferMode BufferingModeIn, UObject* InOwner)
+			: BufferMode(BufferingModeIn)
+			, Owner(InOwner)
+		{}
+
+		/** Only allow construction with valid parameters as well as restricting to module construction */
+		FPhysicsSolverBase() = delete;
+		FPhysicsSolverBase(const FPhysicsSolverBase& InCopy) = delete;
+		FPhysicsSolverBase(FPhysicsSolverBase&& InSteal) = delete;
+		FPhysicsSolverBase& operator =(const FPhysicsSolverBase& InCopy) = delete;
+		FPhysicsSolverBase& operator =(FPhysicsSolverBase&& InSteal) = delete;
+
+		/** Mode that the results buffers should be set to (single, double, triple) */
+		EMultiBufferMode BufferMode;
+
+		// Input Proxy Map
+		TSet< IPhysicsProxyBase *> DirtyProxiesSet;
 
 #if CHAOS_CHECKED
-			FName DebugName;
+		FName DebugName;
 #endif
+
+	private:
+
+		/** 
+		 * Ptr to the engine object that is counted as the owner of this solver.
+		 * Never used internally beyond how the solver is stored and accessed through the solver module.
+		 * Nullptr owner means the solver is global or standalone.
+		 * @see FChaosSolversModule::CreateSolver
+		 */
+		const UObject* Owner = nullptr;
 	};
 }
diff --git a/Engine/Source/Runtime/Experimental/Chaos/Public/ChaosStats.h b/Engine/Source/Runtime/Experimental/Chaos/Public/ChaosStats.h
--- a/Engine/Source/Runtime/Experimental/Chaos/Public/ChaosStats.h
+++ b/Engine/Source/Runtime/Experimental/Chaos/Public/ChaosStats.h
@@ -21,6 +21,7 @@ DECLARE_CYCLE_STAT_EXTERN(TEXT("Integrate Solver"), STAT_IntegrateSolver, STATGR
 DECLARE_CYCLE_STAT_EXTERN(TEXT("Sync Physics Proxies"), STAT_SyncProxies, STATGROUP_Chaos, CHAOS_API);
 DECLARE_CYCLE_STAT_EXTERN(TEXT("Handle Physics Commands"), STAT_PhysCommands, STATGROUP_Chaos, CHAOS_API);
 DECLARE_CYCLE_STAT_EXTERN(TEXT("Handle Task Commands"), STAT_TaskCommands, STATGROUP_Chaos, CHAOS_API);
+DECLARE_CYCLE_STAT_EXTERN(TEXT("Wait for previous global commands"), STAT_WaitGlobalCommands, STATGROUP_Chaos, CHAOS_API);
 DECLARE_CYCLE_STAT_EXTERN(TEXT("Kinematic Particle Update"), STAT_KinematicUpdate, STATGROUP_Chaos, CHAOS_API);
 DECLARE_CYCLE_STAT_EXTERN(TEXT("Begin Frame"), STAT_BeginFrame, STATGROUP_Chaos, CHAOS_API);
 DECLARE_CYCLE_STAT_EXTERN(TEXT("End Frame"), STAT_EndFrame, STATGROUP_Chaos, CHAOS_API);
diff --git a/Engine/Source/Runtime/Experimental/ChaosSolverEngine/Private/Chaos/ChaosSolverActor.cpp b/Engine/Source/Runtime/Experimental/ChaosSolverEngine/Private/Chaos/ChaosSolverActor.cpp
--- a/Engine/Source/Runtime/Experimental/ChaosSolverEngine/Private/Chaos/ChaosSolverActor.cpp
+++ b/Engine/Source/Runtime/Experimental/ChaosSolverEngine/Private/Chaos/ChaosSolverActor.cpp
@@ -210,6 +210,7 @@ AChaosSolverActor::AChaosSolverActor(const FObjectInitializer& ObjectInitializer
 #endif
 	));
 	Solver = PhysScene->GetSolver();
+
 	// Ticking setup for collision/breaking notifies
 	PrimaryActorTick.TickGroup = TG_PostPhysics;
 	PrimaryActorTick.bCanEverTick = true;
@@ -261,6 +262,11 @@ AChaosSolverActor::AChaosSolverActor(const FObjectInitializer& ObjectInitializer
 	GameplayEventDispatcherComponent = ObjectInitializer.CreateDefaultSubobject<UChaosGameplayEventDispatcher>(this, TEXT("GameplayEventDispatcher"));
 }
 
+void AChaosSolverActor::PreInitializeComponents()
+{
+	Super::PreInitializeComponents();
+}
+
 void AChaosSolverActor::BeginPlay()
 {
 	Super::BeginPlay();
@@ -363,6 +369,16 @@ void AChaosSolverActor::PostRegisterAllComponents()
 #endif
 }
 
+void AChaosSolverActor::PostDuplicate(EDuplicateMode::Type DuplicateMode)
+{
+	Super::PostDuplicate(DuplicateMode);
+
+	if(FChaosSolversModule* Module = FChaosSolversModule::GetModule())
+	{
+		Module->MigrateSolver(GetSolver(), GetWorld());
+	}
+}
+
 void AChaosSolverActor::SetAsCurrentWorldSolver()
 {
 	UWorld* const W = GetWorld();
diff --git a/Engine/Source/Runtime/Experimental/ChaosSolverEngine/Public/Chaos/ChaosSolverActor.h b/Engine/Source/Runtime/Experimental/ChaosSolverEngine/Public/Chaos/ChaosSolverActor.h
--- a/Engine/Source/Runtime/Experimental/ChaosSolverEngine/Public/Chaos/ChaosSolverActor.h
+++ b/Engine/Source/Runtime/Experimental/ChaosSolverEngine/Public/Chaos/ChaosSolverActor.h
@@ -169,7 +169,7 @@ public:
 	UFUNCTION(BlueprintCallable, Category = "ChaosPhysics")
 	void SetAsCurrentWorldSolver();
 
-	/** Controles whether the solver is able to simulate particles it controls */
+	/** Controls whether the solver is able to simulate particles it controls */
 	UFUNCTION(BlueprintCallable, Category = "ChaosPhysics")
 	virtual void SetSolverActive(bool bActive);
 
@@ -189,6 +189,8 @@ public:
 #endif
 	virtual void PostRegisterAllComponents() override;
 	
+	virtual void PostDuplicate(EDuplicateMode::Type DuplicateMode) override;
+	virtual void PreInitializeComponents() override;
 	virtual void BeginPlay() override;
 	virtual void EndPlay(const EEndPlayReason::Type ReasonEnd) override;
 
diff --git a/Engine/Source/Runtime/Experimental/ChaosSolvers/Private/ChaosSolversModule.cpp b/Engine/Source/Runtime/Experimental/ChaosSolvers/Private/ChaosSolversModule.cpp
--- a/Engine/Source/Runtime/Experimental/ChaosSolvers/Private/ChaosSolversModule.cpp
+++ b/Engine/Source/Runtime/Experimental/ChaosSolvers/Private/ChaosSolversModule.cpp
@@ -20,6 +20,7 @@
 #include "Framework/DispatcherImpl.h"
 #include "Misc/App.h"
 #include "Chaos/PhysicalMaterials.h"
+#include "Framework/PhysicsTickTask.h"
 
 TAutoConsoleVariable<int32> CVarChaosThreadEnabled(
 	TEXT("p.Chaos.DedicatedThreadEnabled"),
@@ -300,7 +301,7 @@ void FChaosSolversModule::ShutdownThreadingMode()
 			// need to handle this for task graph dispatchers
 			Dispatcher->Execute();
 
-			for(FPBDRigidsSolver* Solver : Solvers)
+			for(FPBDRigidsSolver* Solver : AllSolvers)
 			{
 				IDispatcher::FSolverCommand Command;
 				while(Solver->CommandQueue.Dequeue(Command))
@@ -462,7 +463,20 @@ void FChaosSolversModule::SyncTask(bool bForceBlockingSync /*= false*/)
 	UpdateStats();
 }
 
-Chaos::FPhysicsSolver* FChaosSolversModule::CreateSolver(bool bStandalone /*= false*/
+Chaos::FPhysicsSolver* FChaosSolversModule::CreateSolver(UObject* InOwner
+#if CHAOS_CHECKED
+	, const FName& DebugName
+#endif
+)
+{
+	return CreateSolver(InOwner, ESolverFlags::None
+#if CHAOS_CHECKED
+		, DebugName
+#endif
+	);
+}
+
+Chaos::FPhysicsSolver* FChaosSolversModule::CreateSolver(UObject* InOwner, ESolverFlags InFlags
 #if CHAOS_CHECKED
 	, const FName& DebugName
 #endif
@@ -478,12 +492,16 @@ Chaos::FPhysicsSolver* FChaosSolversModule::CreateSolver(bool bStandalone /*= fa
 		SolverBufferMode = GetDesiredBufferingMode();
 	}
 
-	Solvers.Add(new Chaos::FPhysicsSolver(SolverBufferMode));
-	Chaos::FPhysicsSolver* NewSolver = Solvers.Last();
+	AllSolvers.Add(new Chaos::FPhysicsSolver(SolverBufferMode, InOwner));
+	Chaos::FPhysicsSolver* NewSolver = AllSolvers.Last();
+
+	// Add The solver to the owner list
+	TArray<Chaos::FPhysicsSolver*>& OwnerSolverList = SolverMap.FindOrAdd(InOwner);
+	OwnerSolverList.Add(NewSolver);
 
 #if CHAOS_CHECKED
-    // Add solver nubmer to solver name
-	const FName NewDebugName = *FString::Printf(TEXT("%s (%d)"), DebugName == NAME_None ? TEXT("Solver") : *DebugName.ToString(), Solvers.Num() - 1);
+    // Add solver number to solver name
+	const FName NewDebugName = *FString::Printf(TEXT("%s (%d)"), DebugName == NAME_None ? TEXT("Solver") : *DebugName.ToString(), AllSolvers.Num() - 1);
 	NewSolver->SetDebugName(NewDebugName);
 #endif
 
@@ -498,7 +516,7 @@ Chaos::FPhysicsSolver* FChaosSolversModule::CreateSolver(bool bStandalone /*= fa
 		NewSolver->QueryMaterialLock.WriteUnlock();
 	}
 
-	if(!bStandalone && IsPersistentTaskRunning() && Dispatcher)
+	if(!(InFlags & ESolverFlags::Standalone) && IsPersistentTaskRunning() && Dispatcher)
 	{
 		// Need to let the thread know there's a new solver to care about
 		Dispatcher->EnqueueCommandImmediate([NewSolver](Chaos::FPersistentPhysicsTask* PhysThread)
@@ -511,6 +529,31 @@ Chaos::FPhysicsSolver* FChaosSolversModule::CreateSolver(bool bStandalone /*= fa
 	return NewSolver;
 }
 
+void FChaosSolversModule::MigrateSolver(Chaos::FPhysicsSolver* InSolver, const UObject* InNewOwner)
+{
+	checkSlow(IsInGameThread());
+
+	if(InSolver->GetOwner() == InNewOwner)
+	{
+		// No migration
+		return;
+	}
+
+	if(const UObject* CurrentOwner = InSolver->GetOwner())
+	{
+		TArray<Chaos::FPhysicsSolver*>* OwnerList = SolverMap.Find(CurrentOwner);
+		ensure(OwnerList->Remove(InSolver));
+	}
+
+	InSolver->SetOwner(InNewOwner);
+
+	if(InNewOwner)
+	{
+		TArray<Chaos::FPhysicsSolver*>& OwnerList = SolverMap.FindOrAdd(InNewOwner);
+		OwnerList.Add(InSolver);
+	}
+}
+
 UClass* FChaosSolversModule::GetSolverActorClass() const
 {
 	check(SolverActorClassProvider);
@@ -541,8 +584,12 @@ void FChaosSolversModule::DestroySolver(Chaos::FPhysicsSolver* InSolver)
 
 	FChaosScopeSolverLock SolverScopeLock;
 
-	if(Solvers.Remove(InSolver) > 0)
+	if(AllSolvers.Remove(InSolver) > 0)
 	{
+		// Also remove from the owner list
+		TArray<Chaos::FPhysicsSolver*> OwnerList = GetSolversMutable(InSolver->GetOwner());
+		ensureMsgf(OwnerList.Remove(InSolver), TEXT("Removed a solver from the global list but not an owner list."));
+
 		if(Dispatcher)
 		{
 			Dispatcher->EnqueueCommandImmediate([InSolver](Chaos::FPersistentPhysicsTask* PhysThread)
@@ -566,6 +613,45 @@ void FChaosSolversModule::DestroySolver(Chaos::FPhysicsSolver* InSolver)
 	}
 }
 
+const TArray<Chaos::FPhysicsSolver*>& FChaosSolversModule::GetAllSolvers() const
+{
+	return AllSolvers;
+}
+
+TArray<const Chaos::FPhysicsSolver*> FChaosSolversModule::GetSolvers(const UObject* InOwner) const
+{
+	// Can't just return the ptr from TMap::Find here as it's likely these solver lists will be sent to
+	// the physics thread. In which case the solver pointers are stable but the container pointers are not
+	TArray<const Chaos::FPhysicsSolver*> Solvers;
+	GetSolvers(InOwner, Solvers);
+	return Solvers;
+}
+
+TArray<Chaos::FPhysicsSolver*> FChaosSolversModule::GetSolversMutable(const UObject* InOwner)
+{
+	// Can't just return the ptr from TMap::Find here as it's likely these solver lists will be sent to
+	// the physics thread. In which case the solver pointers are stable but the container pointers are not
+	TArray<Chaos::FPhysicsSolver*> Solvers;
+	GetSolversMutable(InOwner, Solvers);
+	return Solvers;
+}
+
+void FChaosSolversModule::GetSolvers(const UObject* InOwner, TArray<const Chaos::FPhysicsSolver*>& OutSolvers) const
+{
+	if(const TArray<Chaos::FPhysicsSolver*>* Solvers = SolverMap.Find(InOwner))
+	{
+		OutSolvers.Append(*Solvers);
+	}
+}
+
+void FChaosSolversModule::GetSolversMutable(const UObject* InOwner, TArray<Chaos::FPhysicsSolver*>& OutSolvers)
+{
+	if(TArray<Chaos::FPhysicsSolver*>* Solvers = SolverMap.Find(InOwner))
+	{
+		OutSolvers.Append(*Solvers);
+	}
+}
+
 TAutoConsoleVariable<FString> DumpHier_ElementBuckets(
 	TEXT("p.Chaos.DumpHierElementBuckets"),
 	TEXT("1,4,8,16,32,64,128,256,512"),
@@ -591,10 +677,10 @@ void FChaosSolversModule::DumpHierarchyStats(int32* OutOptMaxCellElements)
 	TArray<int32> BucketCounts;
 	BucketCounts.AddZeroed(NumBuckets);
 
-	const int32 NumSolvers = Solvers.Num();
+	const int32 NumSolvers = AllSolvers.Num();
 	for(int32 SolverIndex = 0; SolverIndex < NumSolvers; ++SolverIndex)
 	{
-		Chaos::FPhysicsSolver* Solver = Solvers[SolverIndex];
+		Chaos::FPhysicsSolver* Solver = AllSolvers[SolverIndex];
 #if TODO_REIMPLEMENT_SPATIAL_ACCELERATION_ACCESS
 		if(const Chaos::ISpatialAcceleration<float,3>* SpatialAcceleration = Solver->GetSpatialAcceleration())
 		{
@@ -801,11 +887,11 @@ void FChaosSolversModule::UpdateStats()
 	FRAMEPRO_CUSTOM_STAT("Chaos_Thread_NumAllocatedParticles", PerSolverStats.NumAllocatedParticles, "ChaosThread", "Particles", FRAMEPRO_COLOUR(255,255,255));
 	FRAMEPRO_CUSTOM_STAT("Chaos_Thread_NumPaticleIslands", PerSolverStats.NumParticleIslands, "ChaosThread", "Islands", FRAMEPRO_COLOUR(255,255,255));
 
-	const int32 NumSolvers = Solvers.Num();
+	const int32 NumSolvers = AllSolvers.Num();
 #if 0
 	for(int32 SolverIndex = 0; SolverIndex < NumSolvers; ++SolverIndex)
 	{
-		Chaos::PBDRigidsSolver* Solver = Solvers[SolverIndex];
+		Chaos::PBDRigidsSolver* Solver = AllSolvers[SolverIndex];
 
 		const Chaos::TBoundingVolume<Chaos::TPBDRigidParticles<float, 3>, float, 3>* Hierarchy = Solver->GetSpatialAcceleration();
 
@@ -863,7 +949,7 @@ bool FChaosSolversModule::ShouldStepSolver(int32& InOutSingleStepCounter) const
 
 void FChaosSolversModule::ChangeBufferMode(Chaos::EMultiBufferMode BufferMode)
 {
-	for (Chaos::FPhysicsSolver* Solver : Solvers)
+	for (Chaos::FPhysicsSolver* Solver : AllSolvers)
 	{
 		if (Dispatcher)
 		{
@@ -909,7 +995,7 @@ void FChaosSolversModule::OnUpdateMaterial(Chaos::FMaterialHandle InHandle)
 
 	if(ensure(Material))
 	{
-		for(Chaos::FPhysicsSolver* Solver : Solvers)
+		for(Chaos::FPhysicsSolver* Solver : AllSolvers)
 		{
 			// Send a copy of the material to each solver
 			Dispatcher->EnqueueCommandImmediate(Solver, [InHandle, MaterialCopy = *Material](Chaos::FPhysicsSolver* InSolver)
@@ -929,7 +1015,7 @@ void FChaosSolversModule::OnCreateMaterial(Chaos::FMaterialHandle InHandle)
 
 	if(ensure(Material))
 	{
-		for(Chaos::FPhysicsSolver* Solver : Solvers)
+		for(Chaos::FPhysicsSolver* Solver : AllSolvers)
 		{
 			// Send a copy of the material to each solver
 			Dispatcher->EnqueueCommandImmediate(Solver, [InHandle, MaterialCopy = *Material](Chaos::FPhysicsSolver* InSolver)
@@ -942,8 +1028,6 @@ void FChaosSolversModule::OnCreateMaterial(Chaos::FMaterialHandle InHandle)
 
 void FChaosSolversModule::OnDestroyMaterial(Chaos::FMaterialHandle InHandle)
 {
-	// @todo(bgallagher)
-	//check(Dispatcher);
 	if (!ensure(Dispatcher))
 	{
 		return;
@@ -954,7 +1038,7 @@ void FChaosSolversModule::OnDestroyMaterial(Chaos::FMaterialHandle InHandle)
 
 	if(ensure(Material))
 	{
-		for(Chaos::FPhysicsSolver* Solver : Solvers)
+		for(Chaos::FPhysicsSolver* Solver : AllSolvers)
 		{
 			// Notify each solver
 			Dispatcher->EnqueueCommandImmediate(Solver, [InHandle](Chaos::FPhysicsSolver* InSolver)
@@ -974,7 +1058,7 @@ void FChaosSolversModule::OnUpdateMaterialMask(Chaos::FMaterialMaskHandle InHand
 
 	if (ensure(MaterialMask))
 	{
-		for (Chaos::FPhysicsSolver* Solver : Solvers)
+		for (Chaos::FPhysicsSolver* Solver : AllSolvers)
 		{
 			// Send a copy of the material to each solver
 			Dispatcher->EnqueueCommandImmediate(Solver, [InHandle, MaterialMaskCopy = *MaterialMask](Chaos::FPhysicsSolver* InSolver)
@@ -994,7 +1078,7 @@ void FChaosSolversModule::OnCreateMaterialMask(Chaos::FMaterialMaskHandle InHand
 
 	if (ensure(MaterialMask))
 	{
-		for (Chaos::FPhysicsSolver* Solver : Solvers)
+		for (Chaos::FPhysicsSolver* Solver : AllSolvers)
 		{
 			// Send a copy of the material to each solver
 			Dispatcher->EnqueueCommandImmediate(Solver, [InHandle, MaterialMaskCopy = *MaterialMask](Chaos::FPhysicsSolver* InSolver)
@@ -1007,8 +1091,6 @@ void FChaosSolversModule::OnCreateMaterialMask(Chaos::FMaterialMaskHandle InHand
 
 void FChaosSolversModule::OnDestroyMaterialMask(Chaos::FMaterialMaskHandle InHandle)
 {
-	// @todo(bgallagher)
-	//check(Dispatcher);
 	if (!ensure(Dispatcher))
 	{
 		return;
@@ -1019,7 +1101,7 @@ void FChaosSolversModule::OnDestroyMaterialMask(Chaos::FMaterialMaskHandle InHan
 
 	if (ensure(MaterialMask))
 	{
-		for (Chaos::FPhysicsSolver* Solver : Solvers)
+		for (Chaos::FPhysicsSolver* Solver : AllSolvers)
 		{
 			// Notify each solver
 			Dispatcher->EnqueueCommandImmediate(Solver, [InHandle](Chaos::FPhysicsSolver* InSolver)
@@ -1030,6 +1112,36 @@ void FChaosSolversModule::OnDestroyMaterialMask(Chaos::FMaterialMaskHandle InHan
 	}
 }
 
+void FChaosSolversModule::DispatchGlobalCommands()
+{
+	if(IsValidRef(GlobalCommandTaskEventRef))
+	{
+		if(!GlobalCommandTaskEventRef->IsComplete())
+		{
+			SCOPE_CYCLE_COUNTER(STAT_WaitGlobalCommands);
+			FTaskGraphInterface::Get().WaitUntilTaskCompletes(GlobalCommandTaskEventRef);
+
+			UE_LOG(LogChaos, Warning, TEXT("Forced to wait on global physics commands from the previous frame."));
+		}
+
+		// Clear out our reference once the task has completed
+		GlobalCommandTaskEventRef.SafeRelease();
+	}
+
+	check(!IsValidRef(GlobalCommandTaskEventRef));
+
+	// Fire off commands for this grame
+	GlobalCommandTaskEventRef = TGraphTask<FPhysicsCommandsTask>::CreateTask().ConstructAndDispatchWhenReady();
+}
+
+void FChaosSolversModule::GetSolverUpdatePrerequisites(FGraphEventArray& InPrerequisiteContainer)
+{
+	if(IsValidRef(GlobalCommandTaskEventRef))
+	{
+		InPrerequisiteContainer.Add(GlobalCommandTaskEventRef);
+	}
+}
+
 const IChaosSettingsProvider& FChaosSolversModule::GetSettingsProvider() const
 {
 	return SettingsProvider ? *SettingsProvider : Chaos::GDefaultChaosSettings;
diff --git a/Engine/Source/Runtime/Experimental/ChaosSolvers/Private/Framework/PersistentTask.cpp b/Engine/Source/Runtime/Experimental/ChaosSolvers/Private/Framework/PersistentTask.cpp
--- a/Engine/Source/Runtime/Experimental/ChaosSolvers/Private/Framework/PersistentTask.cpp
+++ b/Engine/Source/Runtime/Experimental/ChaosSolvers/Private/Framework/PersistentTask.cpp
@@ -54,7 +54,9 @@ namespace Chaos
 		// Capture solver states from the module by copying the current state. The module
 		// will inject any new solvers with a command.
 		FChaosSolversModule& ChaosModule = FModuleManager::Get().GetModuleChecked<FChaosSolversModule>("ChaosSolvers");
-		Solvers = ChaosModule.GetSolvers();
+		
+		// Might not be what we want here - maybe just all owned solvers? Or non-standalone
+		Solvers = ChaosModule.GetAllSolvers();
 
 #if CHAOS_DEBUG_SUBSTEP
 		// Prepare the debug substepping tasks for all existing solvers
diff --git a/Engine/Source/Runtime/Experimental/ChaosSolvers/Private/Framework/PhysicsTickTask.cpp b/Engine/Source/Runtime/Experimental/ChaosSolvers/Private/Framework/PhysicsTickTask.cpp
--- a/Engine/Source/Runtime/Experimental/ChaosSolvers/Private/Framework/PhysicsTickTask.cpp
+++ b/Engine/Source/Runtime/Experimental/ChaosSolvers/Private/Framework/PhysicsTickTask.cpp
@@ -20,13 +20,48 @@ FAutoConsoleTaskPriority CPrio_FPhysicsTickTask(
 FPhysicsTickTask::FPhysicsTickTask(FGraphEventRef& InCompletionEvent, Chaos::FPhysicsSolver* InPhysicsSolver, float InDt)
 	: CompletionEvent(InCompletionEvent)
 	, Module(nullptr)
-	, PhysicsSolver(InPhysicsSolver)
 	, Dt(InDt)
 {
+	check(IsInGameThread());
+
 	Module = FChaosSolversModule::GetModule();
 
 	check(Module);
 	checkSlow(Module->GetDispatcher() && Module->GetDispatcher()->GetMode() == EChaosThreadingMode::TaskGraph);
+
+	if(InPhysicsSolver)
+	{
+		SolverList.Reset(1);
+		SolverList.Add(InPhysicsSolver);
+	}
+	else
+	{
+		SolverList = Module->GetAllSolvers();
+	}
+
+	Module->GetSolverUpdatePrerequisites(SolverTaskPrerequisites);
+}
+
+FPhysicsTickTask::FPhysicsTickTask(FGraphEventRef& InCompletionEvent, const TArray<Chaos::FPhysicsSolver*>& InSolverList, float InDt)
+	: CompletionEvent(InCompletionEvent)
+	, Module(nullptr)
+	, Dt(InDt)
+{
+	check(IsInGameThread());
+
+	Module = FChaosSolversModule::GetModule();
+
+	check(Module);
+	checkSlow(Module->GetDispatcher() && Module->GetDispatcher()->GetMode() == EChaosThreadingMode::TaskGraph);
+
+	SolverList = InSolverList;
+	Module->GetSolverUpdatePrerequisites(SolverTaskPrerequisites);
+}
+
+FPhysicsTickTask::FPhysicsTickTask(FGraphEventRef& InCompletionEvent, float InDt)
+	: FPhysicsTickTask(InCompletionEvent, nullptr, InDt)
+{
+
 }
 
 TStatId FPhysicsTickTask::GetStatId() const
@@ -48,16 +83,6 @@ void FPhysicsTickTask::DoTask(ENamedThreads::Type CurrentThread, const FGraphEve
 {
 	using namespace Chaos;
 
-	// The command task runs the two global command queues prior to us running the 
-	// per-solver commands and the solver advance
-	FGraphEventRef CommandsTask = TGraphTask<FPhysicsCommandsTask>::CreateTask().ConstructAndDispatchWhenReady();
-
-	// Otherwise get a full list of solvers from the solvers module.
-	const TArray<FPhysicsSolver*>& SolverList
-		= (PhysicsSolver == nullptr)
-		? Module->GetSolvers()
-		: [&]() { TArray<FPhysicsSolver*> Solvers; Solvers.Init(PhysicsSolver, 1); return Solvers; }();
-
 
 	// List of active solvers (assume all are active for single alloc)
 	TArray<FPhysicsSolver*> ActiveSolvers;
@@ -73,12 +98,8 @@ void FPhysicsTickTask::DoTask(ENamedThreads::Type CurrentThread, const FGraphEve
 
 	const int32 NumActiveSolvers = ActiveSolvers.Num();
 
-	// Prereqs for the solver task to run
-	FGraphEventArray SolverTaskPrerequisites;
 	// Prereqs for the final completion task to run (collection of all the solver tasks)
 	FGraphEventArray CompletionTaskPrerequisites;
-	// Solver tasks have to depend on the global command queues running before them
-	SolverTaskPrerequisites.Add(CommandsTask);
 
 	// For each solver spawn a new solver advance task (which will run the per-solver command buffer and then advance the solver)
 	// Record the task reference as a prerequisite for the completion
diff --git a/Engine/Source/Runtime/Experimental/ChaosSolvers/Private/PBDRigidsSolver.cpp b/Engine/Source/Runtime/Experimental/ChaosSolvers/Private/PBDRigidsSolver.cpp
--- a/Engine/Source/Runtime/Experimental/ChaosSolvers/Private/PBDRigidsSolver.cpp
+++ b/Engine/Source/Runtime/Experimental/ChaosSolvers/Private/PBDRigidsSolver.cpp
@@ -177,11 +177,8 @@ namespace Chaos
 		TSharedPtr<FEvent> PrevEvent, CurrentEvent;
 	};
 
-
-
-
-	FPBDRigidsSolver::FPBDRigidsSolver(const EMultiBufferMode BufferingModeIn)
-		: Super(BufferingModeIn)
+	FPBDRigidsSolver::FPBDRigidsSolver(const EMultiBufferMode BufferingModeIn, UObject* InOwner)
+		: Super(BufferingModeIn, InOwner)
 		, CurrentFrame(0)
 		, MTime(0.0)
 		, MLastDt(0.0)
diff --git a/Engine/Source/Runtime/Experimental/ChaosSolvers/Public/ChaosSolversModule.h b/Engine/Source/Runtime/Experimental/ChaosSolvers/Public/ChaosSolversModule.h
--- a/Engine/Source/Runtime/Experimental/ChaosSolvers/Public/ChaosSolversModule.h
+++ b/Engine/Source/Runtime/Experimental/ChaosSolvers/Public/ChaosSolversModule.h
@@ -79,6 +79,13 @@ private:
 
 };
 
+enum class ESolverFlags : uint8
+{
+	None = 0,
+	Standalone = 1 << 0
+};
+ENUM_CLASS_FLAGS(ESolverFlags);
+
 class CHAOSSOLVERS_API FChaosSolversModule : public IModuleInterface
 {
 public:
@@ -149,14 +156,22 @@ public:
 	 * Should be called from the game thread to create a new solver. After creation, non-standalone solvers
 	 * are dispatched to the physics thread automatically if it is available
 	 *
+	 * @param InOwner Ptr to some owning UObject. The module doesn't use this but allows calling code to organize solver ownership
 	 * @param bStandalone Whether the solver is standalone (not sent to physics thread - updating left to caller)
 	 */
-	Chaos::FPhysicsSolver* CreateSolver(bool bStandalone = false
+	Chaos::FPhysicsSolver* CreateSolver(UObject* InOwner
+#if CHAOS_CHECKED
+		, const FName& DebugName = NAME_None
+#endif
+	);
+	Chaos::FPhysicsSolver* CreateSolver(UObject* InOwner, ESolverFlags InFlags
 #if CHAOS_CHECKED
 		, const FName& DebugName = NAME_None
 #endif
 	);
 
+	void MigrateSolver(Chaos::FPhysicsSolver* InSolver, const UObject* InNewOwner);
+
 	Chaos::EMultiBufferMode GetBufferModeFromThreadingModel(Chaos::EThreadingMode ThreadingMode) const
 	{
 		Chaos::EMultiBufferMode SolverBufferMode = Chaos::EMultiBufferMode::Single;
@@ -222,11 +237,19 @@ public:
 	 */
 	void DestroySolver(Chaos::FPhysicsSolver* InState);
 
+	/** Retrieve the master list of solvers. This contains all owned, unowned and standalone solvers */
+	const TArray<Chaos::FPhysicsSolver*>& GetAllSolvers() const;
+
 	/**
 	 * Read access to the current solver-state objects, be aware which thread owns this data when
 	 * attempting to use this. Physics thread will query when spinning up to get current world state
 	 */
-	const TArray<Chaos::FPhysicsSolver*>& GetSolvers() const { return Solvers; }
+	TArray<const Chaos::FPhysicsSolver*> GetSolvers(const UObject* InOwner) const;
+	void GetSolvers(const UObject* InOwner, TArray<const Chaos::FPhysicsSolver*>& OutSolvers) const;
+
+	/** Non-const access to the solver array is private - only the module should ever modify the storage lists */
+	TArray<Chaos::FPhysicsSolver*> GetSolversMutable(const UObject* InOwner);
+	void GetSolversMutable(const UObject* InOwner, TArray<Chaos::FPhysicsSolver*>& OutSolvers);
 
 	/**
 	 * Outputs statistics for the solver hierarchies. Currently engine calls into this
@@ -298,6 +321,12 @@ public:
 	void OnCreateMaterialMask(Chaos::FMaterialMaskHandle InHandle);
 	void OnDestroyMaterialMask(Chaos::FMaterialMaskHandle InHandle);
 
+	/** Calls from engine tick notifying how the frame is updated */
+	void DispatchGlobalCommands();
+
+	/** Gets a list of pending prerequisites required before proceeding with a solver update */
+	void GetSolverUpdatePrerequisites(FGraphEventArray& InPrerequisiteContainer);
+
 private:
 
 	/** Safe method for always getting a settings provider (from the external caller or an internal default) */
@@ -322,14 +351,12 @@ private:
 	// Core delegate signaling app shutdown, clean up and spin down threads before exit.
 	FDelegateHandle PreExitHandle;
 
-	// Allocated storage for solvers and proxies. Existing on the module makes it easier for hand off in multi threaded mode.
-	// To actually use a solver, call CreateSolver to receive one of these and use it to hold the solver. In the event
-	// of switching to multi threaded mode these will be handed over to the other thread.
-	//
-	// Where these objects are valid for interaction depends on the current threading mode. Use IsPersistentTaskRunning to
-	// check whether the physics thread owns these before manipulating. When adding/removing solver or proxy items in
-	// multi threaded mode the physics thread must also be notified of the change.
-	TArray<Chaos::FPhysicsSolver*> Solvers;
+	// Flat list of every solver the module has created.
+	TArray<Chaos::FPhysicsSolver*> AllSolvers;
+
+	// Map of solver owners to the solvers they own. Nullptr is a valid 'owner' in that it's a map
+	// key for unowned, standalone solvers
+	TMap<const UObject*, TArray<Chaos::FPhysicsSolver*>> SolverMap;
 
 	// Lock for the above list to ensure we don't delete solvers out from underneath other threads
 	// or mess up the solvers array during use.
@@ -344,6 +371,9 @@ private:
 	/** SolverActorClass is require to be this class or a child thereof */
 	UClass* SolverActorRequiredBaseClass;
 
+	/** Reference to an in progress or completed task to run global and thread commands (a prerequisite for any solver ticking) */
+	FGraphEventRef GlobalCommandTaskEventRef;
+
 #if STATS
 	// Stored stats from the physics thread
 	float AverageUpdateTime;
diff --git a/Engine/Source/Runtime/Experimental/ChaosSolvers/Public/Framework/PhysicsTickTask.h b/Engine/Source/Runtime/Experimental/ChaosSolvers/Public/Framework/PhysicsTickTask.h
--- a/Engine/Source/Runtime/Experimental/ChaosSolvers/Public/Framework/PhysicsTickTask.h
+++ b/Engine/Source/Runtime/Experimental/ChaosSolvers/Public/Framework/PhysicsTickTask.h
@@ -4,31 +4,53 @@
 
 
 #include "Async/TaskGraphInterfaces.h"
-#include "Dispatcher.h"
 #include "Chaos/Declares.h"
 
 class FChaosSolversModule;
 
+namespace Chaos
+{
+	class IDispatcher;
+}
+
 /**
  * Task responsible for handling a full frame update for physics under the TaskGraph threading
- * mode.
+ * mode. Multiple ways to construct this for various situations depending on the subset
+ * of solvers the caller wishes to update.
  */
 class CHAOSSOLVERS_API FPhysicsTickTask
 {
 public:
 
+	/** Construct a task that will tick all solvers in the solver module */
+	FPhysicsTickTask(FGraphEventRef& InCompletionEvent, float InDt);
+	/** Construct a task that will tick the provided solver (or all solvers in the module if nullptr passed) */
 	FPhysicsTickTask(FGraphEventRef& InCompletionEvent, Chaos::FPhysicsSolver* InPhysicsSolver, float InDt);
+	/** Construct a task to tick the provided list of solvers */
+	FPhysicsTickTask(FGraphEventRef& InCompletionEvent, const TArray<Chaos::FPhysicsSolver*>& InSolverList, float InDt);
 
+	/** Task API */
 	TStatId GetStatId() const;
 	static ENamedThreads::Type GetDesiredThread();
 	static ESubsequentsMode::Type GetSubsequentsMode();
 	void DoTask(ENamedThreads::Type CurrentThread, const FGraphEventRef& MyCompletionGraphEvent);
+	/** End Task API */
 
 private:
 
+	/** Any prerequisites when this task was constructed for the solver ticks to obey */
+	FGraphEventArray SolverTaskPrerequisites;
+
+	/** An event to dispatch once completed to signal the calling thread */
 	FGraphEventRef CompletionEvent;
+
+	/** Solver module containing master solver lists */
 	FChaosSolversModule* Module;
-	Chaos::FPhysicsSolver* PhysicsSolver;
+
+	/** The solvers this task will tick */
+	TArray<Chaos::FPhysicsSolver*> SolverList;
+
+	/** Delta time for the solver tick */
 	float Dt;
 };
 
diff --git a/Engine/Source/Runtime/Experimental/ChaosSolvers/Public/PBDRigidsSolver.h b/Engine/Source/Runtime/Experimental/ChaosSolvers/Public/PBDRigidsSolver.h
--- a/Engine/Source/Runtime/Experimental/ChaosSolvers/Public/PBDRigidsSolver.h
+++ b/Engine/Source/Runtime/Experimental/ChaosSolvers/Public/PBDRigidsSolver.h
@@ -66,7 +66,7 @@ namespace Chaos
 	class CHAOSSOLVERS_API FPBDRigidsSolver : public FPhysicsSolverBase
 	{
 
-		FPBDRigidsSolver(const EMultiBufferMode BufferingModeIn);
+		FPBDRigidsSolver(const EMultiBufferMode BufferingModeIn, UObject* InOwner);
 
 	public:
 
diff --git a/Engine/Source/Runtime/Experimental/FieldSystem/Source/FieldSystemEngine/Private/Field/FieldSystemComponent.cpp b/Engine/Source/Runtime/Experimental/FieldSystem/Source/FieldSystemEngine/Private/Field/FieldSystemComponent.cpp
--- a/Engine/Source/Runtime/Experimental/FieldSystem/Source/FieldSystemEngine/Private/Field/FieldSystemComponent.cpp
+++ b/Engine/Source/Runtime/Experimental/FieldSystem/Source/FieldSystemEngine/Private/Field/FieldSystemComponent.cpp
@@ -140,33 +140,41 @@ void UFieldSystemComponent::DispatchCommand(const FFieldSystemCommand& InCommand
 		checkSlow(PhysicsDispatcher); // Should always have one of these
 
 		// Assemble a list of compatible solvers
-		TArray<Chaos::FPhysicsSolver*> SolverList;
+		TArray<Chaos::FPhysicsSolver*> SupportedSolverList;
 		if(SupportedSolvers.Num() > 0)
 		{
 			for(TSoftObjectPtr<AChaosSolverActor>& SolverActorPtr : SupportedSolvers)
 			{
 				if(AChaosSolverActor* CurrActor = SolverActorPtr.Get())
 				{
-					SolverList.Add(CurrActor->GetSolver());
+					SupportedSolverList.Add(CurrActor->GetSolver());
 				}
 			}
 		}
 
-		PhysicsDispatcher->EnqueueCommandImmediate([PhysicsProxy = this->PhysicsProxy, NewCommand = InCommand, ChaosModule = this->ChaosModule, SolverList]()
-		{
-			const int32 NumFilterSolvers = SolverList.Num();
-			const TArray<Chaos::FPhysicsSolver*>& Solvers = ChaosModule->GetSolvers();
+		TArray<Chaos::FPhysicsSolver*> WorldSolverList;
+		WorldSolverList = ChaosModule->GetSolversMutable(GetWorld());
 
-			for(Chaos::FPhysicsSolver* Solver : Solvers)
+		// #BGTODO Currently all commands will end up actually executing a frame late. That's because this command has to be logged as a global command
+		// so we don't end up with multiple solver threads writing to the proxy. We need a better way to buffer up multi-solver commands so they can be
+		// executed in parallel and then move those commands to the respective solver queues to fix the frame delay.
+		if(WorldSolverList.Num() > 0)
+		{
+			PhysicsDispatcher->EnqueueCommandImmediate([PhysicsProxy = this->PhysicsProxy, NewCommand = InCommand, ChaosModule = this->ChaosModule, SupportedSolverList, WorldSolvers = MoveTemp(WorldSolverList)]()
 			{
-				const bool bSolverValid = NumFilterSolvers == 0 || SolverList.Contains(Solver);
+				const int32 NumFilterSolvers = SupportedSolverList.Num();
 
-				if(Solver->Enabled() && Solver->HasActiveParticles() && bSolverValid)
+				for(Chaos::FPhysicsSolver* Solver : WorldSolvers)
 				{
-					PhysicsProxy->BufferCommand(Solver, NewCommand);
+					const bool bSolverValid = NumFilterSolvers == 0 || SupportedSolverList.Contains(Solver);
+
+					if(Solver->Enabled() && Solver->HasActiveParticles() && bSolverValid)
+					{
+						PhysicsProxy->BufferCommand(Solver, NewCommand);
+					}
 				}
-			}
-		});
+			});
+		}
 	}
 }
 
diff --git a/Engine/Source/Runtime/Experimental/GeometryCollectionEngine/Private/GeometryCollection/GeometryCollectionParticlesData.cpp b/Engine/Source/Runtime/Experimental/GeometryCollectionEngine/Private/GeometryCollection/GeometryCollectionParticlesData.cpp
--- a/Engine/Source/Runtime/Experimental/GeometryCollectionEngine/Private/GeometryCollection/GeometryCollectionParticlesData.cpp
+++ b/Engine/Source/Runtime/Experimental/GeometryCollectionEngine/Private/GeometryCollection/GeometryCollectionParticlesData.cpp
@@ -27,7 +27,7 @@ TGeometryCollectionParticlesData<T, d>::TGeometryCollectionParticlesData()
 
 template<class T, int d>
 //void TGeometryCollectionParticlesData<T, d>::Sync(const Chaos::FPhysicsSolver* Solver, const TManagedArray<int32>& RigidBodyIds)
-void TGeometryCollectionParticlesData<T, d>::Sync(const Chaos::FPhysicsSolver* Solver, const TManagedArray<FGuid>& RigidBodyIds)
+void TGeometryCollectionParticlesData<T, d>::Sync(Chaos::FPhysicsSolver* Solver, const TManagedArray<FGuid>& RigidBodyIds)
 {
 	// No point in calling twice the sync function within the same frame
 	if (!ensureMsgf(SyncFrame != GFrameCounter, TEXT("Sync should not happen twice during the same tick.")))
@@ -54,7 +54,7 @@ void TGeometryCollectionParticlesData<T, d>::Sync(const Chaos::FPhysicsSolver* S
 			Chaos::IDispatcher* const PhysicsDispatcher = ChaosModule->GetDispatcher();
 			check(PhysicsDispatcher);
 
-			PhysicsDispatcher->EnqueueCommandImmediate([this, Solver, &RigidBodyIds]()
+			PhysicsDispatcher->EnqueueCommandImmediate(Solver, [this, &RigidBodyIds](Chaos::FPhysicsSolver* InSolver)
 			{
 				// Iterate through all data
 				FData& Data = BufferedData.GetPhysicsDataForWrite();
@@ -63,7 +63,7 @@ void TGeometryCollectionParticlesData<T, d>::Sync(const Chaos::FPhysicsSolver* S
 					// Only sync required infos
 					if (Data.RequiredDataFlags[DataIndex])
 					{
-						Data.Copy(EGeometryCollectionParticlesData(DataIndex), Solver, RigidBodyIds);
+						Data.Copy(EGeometryCollectionParticlesData(DataIndex), InSolver, RigidBodyIds);
 					}
 					else
 					{
diff --git a/Engine/Source/Runtime/Experimental/GeometryCollectionEngine/Public/GeometryCollection/GeometryCollectionParticlesData.h b/Engine/Source/Runtime/Experimental/GeometryCollectionEngine/Public/GeometryCollection/GeometryCollectionParticlesData.h
--- a/Engine/Source/Runtime/Experimental/GeometryCollectionEngine/Public/GeometryCollection/GeometryCollectionParticlesData.h
+++ b/Engine/Source/Runtime/Experimental/GeometryCollectionEngine/Public/GeometryCollection/GeometryCollectionParticlesData.h
@@ -79,7 +79,7 @@ public:
 
 	/** Copy the data of the specified set of particles/rigid body ids to this object. */
 	//void Sync(const Chaos::FPhysicsSolver* Solver, const TManagedArray<int32>& RigidBodyIds);
-	void Sync(const Chaos::FPhysicsSolver* Solver, const TManagedArray<FGuid>& RigidBodyIds);
+	void Sync(Chaos::FPhysicsSolver* Solver, const TManagedArray<FGuid>& RigidBodyIds);
 
 	const Chaos::TVector<T, d>               & GetX                     (int32 Index) const { check(HasSyncedData(EGeometryCollectionParticlesData::X                     )); return BufferedData.GetGameDataForRead().X                     [Index]; }
 	const Chaos::TRotation<T, d>             & GetR                     (int32 Index) const { check(HasSyncedData(EGeometryCollectionParticlesData::R                     )); return BufferedData.GetGameDataForRead().R                     [Index]; }
