Get clusters simulating and breaking properly.  (Cluster bodies still have incorrect initial transform)

fyi: Brice.Criswell, Michael.Lentine, Devon.Penney, Jack.Oakman

diff --git a/Engine/Source/Runtime/Experimental/ChaosSolvers/Private/PhysicsProxy/GeometryCollectionPhysicsProxy.cpp b/Engine/Source/Runtime/Experimental/ChaosSolvers/Private/PhysicsProxy/GeometryCollectionPhysicsProxy.cpp
--- a/Engine/Source/Runtime/Experimental/ChaosSolvers/Private/PhysicsProxy/GeometryCollectionPhysicsProxy.cpp
+++ b/Engine/Source/Runtime/Experimental/ChaosSolvers/Private/PhysicsProxy/GeometryCollectionPhysicsProxy.cpp
@@ -816,6 +816,13 @@ void FGeometryCollectionPhysicsProxy::InitializeBodiesPT(
 							RigidChildrenTransformGroupIndex, 
 							CreationParameters);
 
+					int32 RigidChildrenIdx = 0;
+					for (const int32 ChildIndex : Children[TransformGroupIndex])
+					{
+						SolverClusterID[ChildIndex] = RigidChildren[RigidChildrenIdx++]->CastToClustered()->ClusterIds().Id;
+					}
+					SolverClusterID[TransformGroupIndex] = Handle->ClusterIds().Id;
+
 					SolverClusterHandles[TransformGroupIndex] = Handle;
 					SolverParticleHandles[TransformGroupIndex] = Handle;
 					HandleToTransformGroupIndex.Add(Handle, TransformGroupIndex);
@@ -993,7 +1000,9 @@ FGeometryCollectionPhysicsProxy::BuildClusters(
 		}
 	}
 
-	//const FTransform ParticleTM(Particles.R(NewSolverClusterID), Particles.X(NewSolverClusterID));	//in theory we should be computing this and passing in to avoid inertia computation at runtime. If we do this we must account for leaf particles that have already been created in world space
+	// In theory we should be computing this and passing in to avoid inertia computation at 
+	// runtime. If we do this we must account for leaf particles that have already been created in world space
+	//const FTransform ParticleTM(Particles.R(NewSolverClusterID), Particles.X(NewSolverClusterID));	
 	MassToLocal[CollectionClusterIndex] = FTransform::Identity;
 
 	check(Parameters.RestCollection);
@@ -1205,9 +1214,8 @@ void FGeometryCollectionPhysicsProxy::CreateDynamicAttributes()
 			// Update the rest collection with a copy of the shared pointer.
 			(*RestCollectionSharedImplicits)[Index] = Implicits[Index];
 		}
-
-		SolverClusterID[Index] = INDEX_NONE;
 	}
+	SolverClusterID.Fill(nullptr);
 
 	// Clean up now invalid unique pointers.
 	if (RestCollectionSharedImplicits)
@@ -1936,14 +1944,14 @@ void FGeometryCollectionPhysicsProxy::BufferGameState()
 	const FGeometryCollection* RestCollection = Parameters.RestCollection;
 
 	//Buffer.TransformIndex.Init(RestCollection->TransformIndex);
-	Buffer.BoneMap.Init(RestCollection->BoneMap);
-	Buffer.Parent.Init(RestCollection->Parent);
-	Buffer.Children.Init(RestCollection->Children);
-	Buffer.SimulationType.Init(RestCollection->SimulationType);
-
-	Buffer.DynamicState.Init(DynamicCollection->DynamicState);
-	Buffer.Mass.Init(RestCollection->GetAttribute<float>("Mass", FTransformCollection::TransformGroup));
-	Buffer.InertiaTensor.Init(RestCollection->GetAttribute<FVector>("InertiaTensor", FTransformCollection::TransformGroup));
+//	Buffer.BoneMap.Init(RestCollection->BoneMap);
+//	Buffer.Parent.Init(RestCollection->Parent);
+//	Buffer.Children.Init(RestCollection->Children);
+//	Buffer.SimulationType.Init(RestCollection->SimulationType);
+
+//	Buffer.DynamicState.Init(DynamicCollection->DynamicState);
+//	Buffer.Mass.Init(RestCollection->GetAttribute<float>("Mass", FTransformCollection::TransformGroup));
+//	Buffer.InertiaTensor.Init(RestCollection->GetAttribute<FVector>("InertiaTensor", FTransformCollection::TransformGroup));
 	Buffer.Transforms.Init(DynamicCollection->Transform);
 
 	GeometryCollectionAlgo::GlobalMatrices(
@@ -2008,7 +2016,6 @@ void FGeometryCollectionPhysicsProxy::PushToPhysicsState(const Chaos::FParticleD
 
 			// Particle disabled states are initialized on the physics thread, so 
 			// we don't want to pull this from the game thread
-
 			//Handle->SetDisabled(GState->DisabledStates[Idx]); ryan - fix this!
 
 			Handle->SetLinearEtherDrag(Parameters.LinearEtherDrag);
@@ -2110,11 +2117,16 @@ void FGeometryCollectionPhysicsProxy::BufferPhysicsResults()
 	TargetResults.IsObjectLoading = IsObjectLoading;
 	
 	const FTransform& ActorToWorld = Parameters.WorldTransform;
-	// It's not worth shrinking the EndFrameUnparentingBuffer array, at least until the solver supports deleting bodies.
-	if (EndFrameUnparentingBuffer.Num() < NumParticles)
+
+	const int32 TransfromGroupSize = GetTransformGroupSize();
+	if (EndFrameUnparentingBuffer.Num() < TransfromGroupSize)
 	{
-		EndFrameUnparentingBuffer.Init(INDEX_NONE, NumParticles);
+		// It's not worth shrinking the EndFrameUnparentingBuffer array, 
+		// at least until the solver supports deleting bodies.
+		EndFrameUnparentingBuffer.Init(INDEX_NONE, TransfromGroupSize);
 	}
+	TManagedArray<int32>& Parent = PTDynamicCollection.Parent;
+	TManagedArray<TSet<int32>>& Children = PTDynamicCollection.Children;
 
 	{ // tmp scope
 		SCOPE_CYCLE_COUNTER(STAT_CalcParticleToWorld);
@@ -2139,7 +2151,7 @@ void FGeometryCollectionPhysicsProxy::BufferPhysicsResults()
 		//const TArrayCollectionArray<FMultiChildProxyId>& MultiChildProxyIdArray = GetSolver()->GetRigidClustering().GetMultiChildProxyIdArray();
 		//const TArrayCollectionArray<TUniquePtr<TMultiChildProxyData<float, 3>>>& MultiChildProxyDataArray = GetSolver()->GetRigidClustering().GetMultiChildProxyDataArray();
 
-		for (int32 TransformIndex = 0; TransformIndex < NumParticles; ++TransformIndex)
+		for (int32 TransformIndex = 0; TransformIndex < TransfromGroupSize; ++TransformIndex)
 		{
 			Chaos::TPBDRigidClusteredParticleHandle<float, 3>* Handle = SolverParticleHandles[TransformIndex];
 			if (!Handle)
@@ -2197,43 +2209,43 @@ void FGeometryCollectionPhysicsProxy::BufferPhysicsResults()
 					MassToLocal[TransformIndex].GetRelativeTransformReverse(ParticleToWorld).GetRelativeTransform(ActorToWorld);
 				TargetResults.Transforms[TransformIndex].NormalizeRotation();
 
-
-#if 0 // TODO Ryan
-				// Force all enabled rigid bodies out of the transform hierarchy
-				if (TargetResults.Parent[TransformIndex] != INDEX_NONE)
+				// If the parent of this NON DISABLED body is set to anything other than INDEX_NONE,
+				// then it was just unparented, likely either by rigid clustering or by fields.  We
+				// need to force all such enabled rigid bodies out of the transform hierarchy.
+				if (Parent[TransformIndex] != INDEX_NONE)
 				{
+					const int32 ParentIndex = Parent[TransformIndex];
 					// Children in the hierarchy are stored in a TSet, which is not thread safe.  
 					// So we retain indices to remove afterwards.
-					EndFrameUnparentingBuffer[TransformIndex] = TargetResults.Parent[TransformIndex];
+					EndFrameUnparentingBuffer[TransformIndex] = ParentIndex;
 				}
-#endif
-#if 0 // TODO Ryan
+
 				// When a leaf node rigid body is removed from a cluster, the rigid
 				// body will become active and needs its clusterID updated.  This just
 				// syncs the clusterID all the time.
-				CollectionClusterID[TransformGroupIndex] = ClusterID[RigidBodyIndex].Id;
-#endif
+				TPBDRigidParticleHandle<float, 3>* ClusterParentId = Handle->ClusterIds().Id;
+				SolverClusterID[TransformIndex] = ClusterParentId;
 			}
 			else // Handle->Disabled()
 			{
-#if 0 // TODO Ryan
 				// The rigid body parent cluster has changed within the solver, and its
 				// parent body is not tracked within the geometry collection. So we need to
 				// pull the rigid bodies out of the transform hierarchy, and just drive
 				// the positions directly from the solvers cluster particle. 
-				if (CollectionClusterID[TransformGroupIndex] != ClusterID[RigidBodyIndex].Id)
+				TPBDRigidParticleHandle<float, 3>* ClusterParentId = Handle->ClusterIds().Id;
+				if(SolverClusterID[TransformIndex] != ClusterParentId)
 				{
 					// Force all driven rigid bodies out of the transform hierarchy
-					if (Parent[TransformGroupIndex] != INDEX_NONE)
+					if (Parent[TransformIndex] != INDEX_NONE)
 					{
-						const int32 ParentIndex = Parent[TransformGroupIndex];
+						const int32 ParentIndex = Parent[TransformIndex];
 						// Children in the hierarchy are stored in a TSet, which is not thread safe.  So we retain
 						// indices to remove afterwards.
-						EndFrameUnparentingBuffer[TransformGroupIndex] = ParentIndex;
+						EndFrameUnparentingBuffer[TransformIndex] = ParentIndex;
 					}
-					CollectionClusterID[TransformGroupIndex] = ClusterID[RigidBodyIndex].Id;
+					SolverClusterID[TransformIndex] = Handle->ClusterIds().Id;
 				}
-#endif
+
 				// Disabled rigid bodies that have valid cluster parents, and have been re-indexed by the
 				// solver (as in, They were re-clustered outside of the geometry collection), These clusters 
 				// will need to be rendered based on the clusters position. 
@@ -2249,9 +2261,7 @@ void FGeometryCollectionPhysicsProxy::BufferPhysicsResults()
 							const FTransform ClusterChildToWorld = ChildToParent * FTransform(ParentHandle->R(), ParentHandle->X());
 							if (Parameters.IsCacheRecording())
 							{
-								//Particles.X(RigidBodyIndex) = ClusterChildToWorld.GetTranslation();
 								Handle->SetX(ClusterChildToWorld.GetTranslation());
-								//Particles.R(RigidBodyIndex) = ClusterChildToWorld.GetRotation();
 								Handle->SetR(ClusterChildToWorld.GetRotation());
 							}
 							// GeomToActor = ActorToWorld.Inv() * ClusterChildToWorld * MassToLocal.Inv();
@@ -2266,6 +2276,21 @@ void FGeometryCollectionPhysicsProxy::BufferPhysicsResults()
 		} // end for
 	} // tmp scope
 
+	for (int32 TransformGroupIndex = 0; TransformGroupIndex < TransfromGroupSize; TransformGroupIndex++)
+	{
+		const int32 ParentIndex = EndFrameUnparentingBuffer[TransformGroupIndex];
+		if (ParentIndex >= 0)
+		{
+			// We reuse EndFrameUnparentingBuffer potentially without reinitialization, 
+			// so reset this index to INDEX_NONE before it gets paged out.
+			EndFrameUnparentingBuffer[TransformGroupIndex] = INDEX_NONE;
+
+			Children[ParentIndex].Remove(TransformGroupIndex);
+			Parent[TransformGroupIndex] = INDEX_NONE;
+		}
+		TargetResults.Parent[TransformGroupIndex] = Parent[TransformGroupIndex];
+	}
+
 	// If object is dynamic, compute global matrices	
 	if (IsObjectDynamic || TargetResults.GlobalTransforms.Num() == 0)
 	{
@@ -2273,43 +2298,22 @@ void FGeometryCollectionPhysicsProxy::BufferPhysicsResults()
 		check(TargetResults.Transforms.Num() == TargetResults.Parent.Num());
 		GeometryCollectionAlgo::GlobalMatrices(
 			TargetResults.Transforms, TargetResults.Parent, TargetResults.GlobalTransforms);
-	}
 
-	// compute world bounds
-	// #note: this is a loose bounds based on the circumscribed box of a bounding sphere for the geometry.		
-	if (IsObjectDynamic || TargetResults.WorldBounds.GetSphere().W < 1e-5)
-	{
+		// Compute world bounds.  This is a loose bounds based on the circumscribed box 
+		// of a bounding sphere for the geometry.		
 		SCOPE_CYCLE_COUNTER(STAT_CalcGlobalGCBounds);
 		FBox BoundingBox(ForceInit);
-		//const FMatrix& ActorToWorld = Parameters.WorldTransform.ToMatrixWithScale();
 		for (int i = 0; i < ValidGeometryBoundingBoxes.Num(); ++i)
 		{
 			BoundingBox += ValidGeometryBoundingBoxes[i].TransformBy(
 				TargetResults.GlobalTransforms[ValidGeometryTransformIndices[i]] * ActorToWorld);
 		}
-		TargetResults.WorldBounds = FBoxSphereBounds(BoundingBox);		
-	}
-
-#if 0 // TODO Ryan
-	for (int32 TransformGroupIndex = 0; TransformGroupIndex < TransformSize; TransformGroupIndex++)
-	{
-		const int32 ParentIndex = EndFrameUnparentingBuffer[TransformGroupIndex];
-		if (ParentIndex >= 0)
-		{
-			// We reuse EndFrameUnparentingBuffer potentially without reinitialization, so reset this index to -1 before it gets paged out.
-			EndFrameUnparentingBuffer[TransformGroupIndex] = -1;
-
-			Children[ParentIndex].Remove(TransformGroupIndex);
-			Parent[TransformGroupIndex] = INDEX_NONE;
-		}
+		TargetResults.WorldBounds = FBoxSphereBounds(BoundingBox);
 	}
-#endif
 
 	// We pulled state out of particles.  Push results into PTDynamicCollection, 
 	// so that it has an updated view.
 	UpdateGeometryCollection(TargetResults);
-
-//#endif // TODO_REIMPLEMENT_GET_RIGID_PARTICLES
 }
 
 void FGeometryCollectionPhysicsProxy::FlipBuffer()
diff --git a/Engine/Source/Runtime/Experimental/ChaosSolvers/Public/PhysicsProxy/GeometryCollectionPhysicsProxy.h b/Engine/Source/Runtime/Experimental/ChaosSolvers/Public/PhysicsProxy/GeometryCollectionPhysicsProxy.h
--- a/Engine/Source/Runtime/Experimental/ChaosSolvers/Public/PhysicsProxy/GeometryCollectionPhysicsProxy.h
+++ b/Engine/Source/Runtime/Experimental/ChaosSolvers/Public/PhysicsProxy/GeometryCollectionPhysicsProxy.h
@@ -352,21 +352,23 @@ private:
 	//TUniquePtr<FGeometryDynamicCollection> PTDynamicCollection;
 	FGeometryDynamicCollection PTDynamicCollection;
 
+	TMap<Chaos::TPBDRigidParticleHandle<float, 3>*, int32> HandleToTransformGroupIndex;
+
 	TManagedArray<Chaos::TPBDRigidClusteredParticleHandle<float, 3>*> SolverParticleHandles;
 	TManagedArray<Chaos::TPBDRigidClusteredParticleHandle<float, 3>*> SolverClusterHandles;
-
-	TMap<Chaos::TPBDRigidParticleHandle<float, 3>*, int32> HandleToTransformGroupIndex;
+	TManagedArray<Chaos::TPBDRigidParticleHandle<float, 3>*> SolverClusterID; // Rename to ClusterParent?
 
 	TManagedArray<FTransform> MassToLocal;
 	TManagedArray<int32> CollisionMask;
 	TManagedArray<int32> CollisionStructureID;
 	TManagedArray<int32> RigidBodyID; // Deprecated.  Not added to dynamic collection.
-	TManagedArray<int32> SolverClusterID;
 	TManagedArray<FVector> InitialAngularVelocity;
 	TManagedArray<FVector> InitialLinearVelocity;
 	TManagedArray<bool> SimulatableParticles;
+
 	TManagedArray<TUniquePtr<FSimplicial> > Simplicials; // FSimplicial = Chaos::TBVHParticles<float,3>
 	TManagedArray<TSharedPtr<Chaos::FImplicitObject, ESPMode::ThreadSafe>> Implicits;
+
 	TArray<int32> EndFrameUnparentingBuffer;
 
 	// This is a subset of the geometry group that are used in the transform hierarchy to represent geometry
