CIS Fix - Add CHAOS_API to the accessor in the same manner as the other accessors in the code file (although FORCEINLINE would also work here) to avoid Clang Debug linker errors.

diff --git a/Engine/Source/Runtime/Experimental/Chaos/Public/Chaos/GeometryParticles.h b/Engine/Source/Runtime/Experimental/Chaos/Public/Chaos/GeometryParticles.h
--- a/Engine/Source/Runtime/Experimental/Chaos/Public/Chaos/GeometryParticles.h
+++ b/Engine/Source/Runtime/Experimental/Chaos/Public/Chaos/GeometryParticles.h
@@ -237,7 +237,7 @@ namespace Chaos
 
 		CHAOS_API const TUniquePtr<FImplicitObject>& DynamicGeometry(const int32 Index) const { return MDynamicGeometry[Index]; }
 
-		const TSharedPtr<FImplicitObject, ESPMode::ThreadSafe>& SharedGeometry(const int32 Index) const { return MSharedGeometry[Index]; }
+		CHAOS_API const TSharedPtr<FImplicitObject, ESPMode::ThreadSafe>& SharedGeometry(const int32 Index) const { return MSharedGeometry[Index]; }
 
 		CHAOS_API const TShapesArray<T, d>& ShapesArray(const int32 Index) const { return MShapesArray[Index]; }
 
