Fixed two serialization breaks causing Chaos geometry collections to fail to serialize. Added a ChaosVersionString to Chaos/Core.h that will be used in future as the core Chaos data version that all dependent builders and processes should rely on. As Chaos data changes we'll rotate that key and correctly invalidate external builders too.

diff --git a/Engine/Source/Runtime/Engine/Private/PhysicsEngine/Experimental/ChaosDerivedData.cpp b/Engine/Source/Runtime/Engine/Private/PhysicsEngine/Experimental/ChaosDerivedData.cpp
--- a/Engine/Source/Runtime/Engine/Private/PhysicsEngine/Experimental/ChaosDerivedData.cpp
+++ b/Engine/Source/Runtime/Engine/Private/PhysicsEngine/Experimental/ChaosDerivedData.cpp
@@ -4,6 +4,7 @@
 
 #if WITH_CHAOS
 
+#include "Chaos/Core.h"
 #include "Chaos/CollisionConvexMesh.h"
 #include "Chaos/ChaosArchive.h"
 #include "ChaosDerivedDataUtil.h"
@@ -24,7 +25,11 @@ const TCHAR* FChaosDerivedDataCooker::GetPluginName() const
 
 const TCHAR* FChaosDerivedDataCooker::GetVersionString() const
 {
-	return TEXT("99F5C9095ED541A5A8A8434EB38D8303");
+	// As changing our DDC version will most likely affect any external callers that rely on Chaos types
+	// for their own DDC or serialized data - change Chaos::ChaosVersionString in Chaos/Core.h to bump our
+	// Chaos data version. Callers can also rely on that version in their builders and avoid bad serialization
+	// when basic Chaos data changes
+	return Chaos::ChaosVersionString;
 }
 
 FString FChaosDerivedDataCooker::GetPluginSpecificCacheKeySuffix() const
diff --git a/Engine/Source/Runtime/Experimental/ChaosCore/Public/Chaos/Core.h b/Engine/Source/Runtime/Experimental/ChaosCore/Public/Chaos/Core.h
--- a/Engine/Source/Runtime/Experimental/ChaosCore/Public/Chaos/Core.h
+++ b/Engine/Source/Runtime/Experimental/ChaosCore/Public/Chaos/Core.h
@@ -6,6 +6,7 @@
 #include "Chaos/Matrix.h"
 #include "Chaos/Rotation.h"
 #include "Chaos/Transform.h"
+#include "Containers/UnrealString.h"
 
 namespace Chaos
 {
@@ -24,4 +25,7 @@ namespace Chaos
 
 	template <typename T>
 	using TVec3 = TVector<T, 3>;
+
+	/** Core version string for Chaos data. Any DDC builder dependent on Chaos for serialization should depend on this version */
+	static const FString ChaosVersionString = TEXT("C247EE00C33148EA87C9F53715E83429");
 }
\ No newline at end of file
diff --git a/Engine/Source/Runtime/Experimental/GeometryCollectionEngine/Private/GeometryCollection/DerivedDataGeometryCollectionCooker.cpp b/Engine/Source/Runtime/Experimental/GeometryCollectionEngine/Private/GeometryCollection/DerivedDataGeometryCollectionCooker.cpp
--- a/Engine/Source/Runtime/Experimental/GeometryCollectionEngine/Private/GeometryCollection/DerivedDataGeometryCollectionCooker.cpp
+++ b/Engine/Source/Runtime/Experimental/GeometryCollectionEngine/Private/GeometryCollection/DerivedDataGeometryCollectionCooker.cpp
@@ -1,5 +1,7 @@
 // Copyright Epic Games, Inc. All Rights Reserved.
+
 #include "GeometryCollection/DerivedDataGeometryCollectionCooker.h"
+#include "Chaos/Core.h"
 #include "GeometryCollection/GeometryCollectionObject.h"
 #include "GeometryCollection/GeometryCollection.h"
 #include "Serialization/MemoryWriter.h"
@@ -88,7 +90,7 @@ const TCHAR* FDerivedDataGeometryCollectionCooker::GetVersionString() const
 FString FDerivedDataGeometryCollectionCooker::GetPluginSpecificCacheKeySuffix() const
 {
 
-	return FString::Printf(TEXT("%s_%s_%d"), *GeometryCollection.GetIdGuid().ToString(), *GeometryCollection.GetStateGuid().ToString(), FDestructionObjectVersion::Type::LatestVersion);
+	return FString::Printf(TEXT("%s_%s_%s_%d"), *Chaos::ChaosVersionString, *GeometryCollection.GetIdGuid().ToString(), *GeometryCollection.GetStateGuid().ToString(), FDestructionObjectVersion::Type::LatestVersion);
 }
 
 
