Chaos: Re-implement Dynamic State initialization for Geometry Collection

diff --git a/Engine/Source/Runtime/Experimental/ChaosSolvers/Private/PhysicsProxy/GeometryCollectionPhysicsProxy.cpp b/Engine/Source/Runtime/Experimental/ChaosSolvers/Private/PhysicsProxy/GeometryCollectionPhysicsProxy.cpp
--- a/Engine/Source/Runtime/Experimental/ChaosSolvers/Private/PhysicsProxy/GeometryCollectionPhysicsProxy.cpp
+++ b/Engine/Source/Runtime/Experimental/ChaosSolvers/Private/PhysicsProxy/GeometryCollectionPhysicsProxy.cpp
@@ -29,10 +29,6 @@
 #define TODO_REIMPLEMENT_INIT_COMMANDS 0
 #endif
 
-#ifndef TODO_REIMPLEMENT_KINEMATICS
-#define TODO_REIMPLEMENT_KINEMATICS 0
-#endif
-
 #ifndef TODO_REIMPLEMENT_FRACTURE
 #define TODO_REIMPLEMENT_FRACTURE 0
 #endif
@@ -742,9 +738,8 @@ void FGeometryCollectionPhysicsProxy::InitializeBodiesPT(
 			}
 		}
 
-#if TODO_REIMPLEMENT_KINEMATICS
-		InitializeKinematics(Particles, DynamicState);
-#endif // TODO_REIMPLEMENT_KINEMATICS
+		InitializeKinematics(DynamicState);
+
 #if TODO_REIMPLEMENT_FRACTURE
 		InitializeRemoveOnFracture(Particles, DynamicState);
 #endif // TODO_REIMPLEMENT_FRACTURE
@@ -1280,26 +1275,29 @@ void FGeometryCollectionPhysicsProxy::CreateDynamicAttributes()
 	}
 }
 
-void FGeometryCollectionPhysicsProxy::InitializeKinematics(FParticlesType& Particles, const TManagedArray<int32>& DynamicState)
+void FGeometryCollectionPhysicsProxy::InitializeKinematics(const TManagedArray<int32>& DynamicState)
 {
 	if(Parameters.DynamicCollection)
 	{
-		for(int TransformGroupIndex = 0; TransformGroupIndex < RigidBodyID.Num(); TransformGroupIndex++)
+
+		for (int32 TransformGroupIndex = 0; TransformGroupIndex < GetTransformGroupSize(); ++TransformGroupIndex)
 		{
-			if(RigidBodyID[TransformGroupIndex] != INDEX_NONE)
+			if (Chaos::TPBDRigidParticleHandle<float, 3>* Handle = SolverParticleHandles[TransformGroupIndex])
 			{
-				int32 RigidBodyIndex = RigidBodyID[TransformGroupIndex];
-				if(DynamicState[TransformGroupIndex] == (uint8)EObjectStateTypeEnum::Chaos_Object_Kinematic)
+				// this is being initialized to Static, can't change from static to sleeping
+				Handle->SetObjectState(Chaos::EObjectStateType::Dynamic); 
+
+				if(DynamicState[TransformGroupIndex] == (int32)EObjectStateTypeEnum::Chaos_Object_Kinematic)
 				{
-					Particles.SetObjectState(RigidBodyIndex, Chaos::EObjectStateType::Kinematic);
+					Handle->SetObjectState( Chaos::EObjectStateType::Kinematic);
 				}
-				else if(DynamicState[TransformGroupIndex] == (uint8)EObjectStateTypeEnum::Chaos_Object_Static)
+				else if(DynamicState[TransformGroupIndex] == (int32)EObjectStateTypeEnum::Chaos_Object_Static)
 				{
-					Particles.SetObjectState(RigidBodyIndex, Chaos::EObjectStateType::Static);
+					Handle->SetObjectState(Chaos::EObjectStateType::Static);
 				}
-				else if(DynamicState[TransformGroupIndex] == (uint8)EObjectStateTypeEnum::Chaos_Object_Sleeping)
+				else if(DynamicState[TransformGroupIndex] == (int32)EObjectStateTypeEnum::Chaos_Object_Sleeping)
 				{
-					Particles.SetObjectState(RigidBodyIndex, Chaos::EObjectStateType::Sleeping);
+					Handle->SetObjectState(Chaos::EObjectStateType::Sleeping);
 				}
 			}
 		}
diff --git a/Engine/Source/Runtime/Experimental/ChaosSolvers/Public/PhysicsProxy/GeometryCollectionPhysicsProxy.h b/Engine/Source/Runtime/Experimental/ChaosSolvers/Public/PhysicsProxy/GeometryCollectionPhysicsProxy.h
--- a/Engine/Source/Runtime/Experimental/ChaosSolvers/Public/PhysicsProxy/GeometryCollectionPhysicsProxy.h
+++ b/Engine/Source/Runtime/Experimental/ChaosSolvers/Public/PhysicsProxy/GeometryCollectionPhysicsProxy.h
@@ -299,7 +299,7 @@ protected:
 		int32 TransformIndex) const;
 
 	void CreateDynamicAttributes();
-	void InitializeKinematics(FParticlesType& Particles, const TManagedArray<int32>& ObjectType);
+	void InitializeKinematics(const TManagedArray<int32>& ObjectType);
 	void InitializeRemoveOnFracture(FParticlesType& Particles, const TManagedArray<int32>& DynamicState);
 	void PushKinematicStateToSolver(FParticlesType& Particles);
 
