Reimplemented solver actors for Chaos geometry collections.
 - Reimplemented TaskGraph tick mode to correctly handle multiple worlds (this involves ticking global command queues outside of world tick and mapping worlds to solver lists for each world tick so we can correctly handle multiple worlds)
 - Audited global commands that had an explicit solver capture and moved them to solver commands
 - Updated headless tests to new API for solver flags

diff --git a/Engine/Source/Runtime/Experimental/ChaosSolvers/Private/ChaosSolversModule.cpp b/Engine/Source/Runtime/Experimental/ChaosSolvers/Private/ChaosSolversModule.cpp
--- a/Engine/Source/Runtime/Experimental/ChaosSolvers/Private/ChaosSolversModule.cpp
+++ b/Engine/Source/Runtime/Experimental/ChaosSolvers/Private/ChaosSolversModule.cpp
@@ -1114,6 +1114,12 @@ void FChaosSolversModule::OnDestroyMaterialMask(Chaos::FMaterialMaskHandle InHan
 
 void FChaosSolversModule::DispatchGlobalCommands()
 {
+	if(Dispatcher->GetMode() == Chaos::EThreadingMode::SingleThread)
+	{
+		// Single thread commands run as they are enqueued so no need to proceed here
+		return;
+	}
+
 	if(IsValidRef(GlobalCommandTaskEventRef))
 	{
 		if(!GlobalCommandTaskEventRef->IsComplete())
