GeometryCollectionTest

- Enabled Kinematics Test

diff --git a/Engine/Source/Runtime/Experimental/ChaosSolvers/Private/PhysicsProxy/GeometryCollectionPhysicsProxy.cpp b/Engine/Source/Runtime/Experimental/ChaosSolvers/Private/PhysicsProxy/GeometryCollectionPhysicsProxy.cpp
--- a/Engine/Source/Runtime/Experimental/ChaosSolvers/Private/PhysicsProxy/GeometryCollectionPhysicsProxy.cpp
+++ b/Engine/Source/Runtime/Experimental/ChaosSolvers/Private/PhysicsProxy/GeometryCollectionPhysicsProxy.cpp
@@ -673,7 +673,7 @@ void FGeometryCollectionPhysicsProxy::InitializeBodiesPT(
 #endif // TODO_REIMPLEMENT_SOLVER_SETTINGS_ACCESSORS
 					//RigidBodyIndex, 
 					WorldTransform, 
-					(uint8)Parameters.ObjectType, 
+					(uint8)DynamicState[TransformGroupIndex], 
 					CollisionGroup[TransformGroupIndex]);
 
 				if (Parameters.EnableClustering)
@@ -738,8 +738,6 @@ void FGeometryCollectionPhysicsProxy::InitializeBodiesPT(
 			}
 		}
 
-		InitializeKinematics(DynamicState);
-
 #if TODO_REIMPLEMENT_FRACTURE
 		InitializeRemoveOnFracture(Particles, DynamicState);
 #endif // TODO_REIMPLEMENT_FRACTURE
@@ -1275,34 +1273,6 @@ void FGeometryCollectionPhysicsProxy::CreateDynamicAttributes()
 	}
 }
 
-void FGeometryCollectionPhysicsProxy::InitializeKinematics(const TManagedArray<int32>& DynamicState)
-{
-	if(Parameters.DynamicCollection)
-	{
-
-		for (int32 TransformGroupIndex = 0; TransformGroupIndex < GetTransformGroupSize(); ++TransformGroupIndex)
-		{
-			if (Chaos::TPBDRigidParticleHandle<float, 3>* Handle = SolverParticleHandles[TransformGroupIndex])
-			{
-				// this is being initialized to Static, can't change from static to sleeping
-				Handle->SetObjectState(Chaos::EObjectStateType::Dynamic); 
-
-				if(DynamicState[TransformGroupIndex] == (int32)EObjectStateTypeEnum::Chaos_Object_Kinematic)
-				{
-					Handle->SetObjectState( Chaos::EObjectStateType::Kinematic);
-				}
-				else if(DynamicState[TransformGroupIndex] == (int32)EObjectStateTypeEnum::Chaos_Object_Static)
-				{
-					Handle->SetObjectState(Chaos::EObjectStateType::Static);
-				}
-				else if(DynamicState[TransformGroupIndex] == (int32)EObjectStateTypeEnum::Chaos_Object_Sleeping)
-				{
-					Handle->SetObjectState(Chaos::EObjectStateType::Sleeping);
-				}
-			}
-		}
-	}
-}
 
 void FGeometryCollectionPhysicsProxy::InitializeRemoveOnFracture(FParticlesType& Particles, const TManagedArray<int32>& DynamicState)
 {
@@ -2371,6 +2341,7 @@ void FGeometryCollectionPhysicsProxy::PullFromPhysicsState()
 			{
 				GTDynamicCollection->Transform[TmIndex] = TR.Transforms[TmIndex];
 			}
+			GTDynamicCollection->DynamicState[TmIndex] = TR.DynamicState[TmIndex];
 		}
 
 		//question: why do we need this? Sleeping objects will always have to update GPU
@@ -4211,7 +4182,31 @@ void FGeometryCollectionPhysicsProxy::CreateRigidBodyCallback(FParticlesType& Pa
 			}
 		});
 
-		InitializeKinematics(Particles, DynamicState);
+		// InitializeKinematics
+		if (Parameters.DynamicCollection)
+		{
+			for (int32 TransformGroupIndex = 0; TransformGroupIndex < GetTransformGroupSize(); ++TransformGroupIndex)
+			{
+				if (Chaos::TPBDRigidParticleHandle<float, 3>* Handle = SolverParticleHandles[TransformGroupIndex])
+				{
+					// this is being initialized to Static, can't change from static to sleeping
+					Handle->SetObjectState(Chaos::EObjectStateType::Dynamic);
+
+					if (DynamicState[TransformGroupIndex] == (int32)EObjectStateTypeEnum::Chaos_Object_Kinematic)
+					{
+						Handle->SetObjectState(Chaos::EObjectStateType::Kinematic);
+					}
+					else if (DynamicState[TransformGroupIndex] == (int32)EObjectStateTypeEnum::Chaos_Object_Static)
+					{
+						Handle->SetObjectState(Chaos::EObjectStateType::Static);
+					}
+					else if (DynamicState[TransformGroupIndex] == (int32)EObjectStateTypeEnum::Chaos_Object_Sleeping)
+					{
+						Handle->SetObjectState(Chaos::EObjectStateType::Sleeping);
+					}
+				}
+			}
+		}
 
 		InitializeRemoveOnFracture(Particles, DynamicState);
 
diff --git a/Engine/Source/Runtime/Experimental/ChaosSolvers/Public/PhysicsProxy/GeometryCollectionPhysicsProxy.h b/Engine/Source/Runtime/Experimental/ChaosSolvers/Public/PhysicsProxy/GeometryCollectionPhysicsProxy.h
--- a/Engine/Source/Runtime/Experimental/ChaosSolvers/Public/PhysicsProxy/GeometryCollectionPhysicsProxy.h
+++ b/Engine/Source/Runtime/Experimental/ChaosSolvers/Public/PhysicsProxy/GeometryCollectionPhysicsProxy.h
@@ -299,7 +299,6 @@ protected:
 		int32 TransformIndex) const;
 
 	void CreateDynamicAttributes();
-	void InitializeKinematics(const TManagedArray<int32>& ObjectType);
 	void InitializeRemoveOnFracture(FParticlesType& Particles, const TManagedArray<int32>& DynamicState);
 	void PushKinematicStateToSolver(FParticlesType& Particles);
 
