GeometryCollectionTest

- Testing updates.

diff --git a/Engine/Source/Runtime/Experimental/GeometryCollectionCore/Private/GeometryCollection/TransformCollection.cpp b/Engine/Source/Runtime/Experimental/GeometryCollectionCore/Private/GeometryCollection/TransformCollection.cpp
--- a/Engine/Source/Runtime/Experimental/GeometryCollectionCore/Private/GeometryCollection/TransformCollection.cpp
+++ b/Engine/Source/Runtime/Experimental/GeometryCollectionCore/Private/GeometryCollection/TransformCollection.cpp
@@ -53,6 +53,16 @@ void FTransformCollection::Serialize(Chaos::FChaosArchive& Ar)
 
 }
 
+FTransformCollection FTransformCollection::SingleTransform(const FTransform& TransformRoot)
+{
+	FTransformCollection TransformCollection;
+	TransformCollection.AddElements(1, FTransformCollection::TransformGroup);
+	TransformCollection.Transform[0] = TransformRoot;
+	TransformCollection.Parent[0] = Invalid;
+	return TransformCollection;
+}
+
+
 int32 FTransformCollection::AppendTransform(const FTransformCollection & Element, const FTransform& TransformRoot)
 {
 	check(Element.NumElements(FTransformCollection::TransformGroup) > 0);
diff --git a/Engine/Source/Runtime/Experimental/GeometryCollectionCore/Public/GeometryCollection/TransformCollection.h b/Engine/Source/Runtime/Experimental/GeometryCollectionCore/Public/GeometryCollection/TransformCollection.h
--- a/Engine/Source/Runtime/Experimental/GeometryCollectionCore/Public/GeometryCollection/TransformCollection.h
+++ b/Engine/Source/Runtime/Experimental/GeometryCollectionCore/Public/GeometryCollection/TransformCollection.h
@@ -53,9 +53,15 @@ public:
 	/** Serialize */
 	virtual void Serialize(Chaos::FChaosArchive& Ar) override;
 
+	/*
+	* SingleTransform:
+	*   Create a single transform.
+	*/
+	static FTransformCollection SingleTransform(const FTransform& TransformRoot = FTransform::Identity);
+
 	/*
 	* AppendTransform:
-	*   Append a single transform at the end of the collection without
+	*   Append a transform at the end of the collection without
 	*   parenting. 
 	*/
 	int32 AppendTransform(const FTransformCollection & GeometryCollection, const FTransform& TransformRoot = FTransform::Identity);
