Fix FConvex::PhiWithNormal to be more accurate.

diff --git a/Engine/Source/Runtime/Experimental/Chaos/Public/Chaos/Convex.h b/Engine/Source/Runtime/Experimental/Chaos/Public/Chaos/Convex.h
--- a/Engine/Source/Runtime/Experimental/Chaos/Public/Chaos/Convex.h
+++ b/Engine/Source/Runtime/Experimental/Chaos/Public/Chaos/Convex.h
@@ -79,6 +79,48 @@ namespace Chaos
 		}
 
 		virtual FReal PhiWithNormal(const FVec3& x, FVec3& Normal) const override
+		{
+			float Phi = PhiWithNormalInternal(x, Normal);
+			if (Phi > 0)
+			{
+				//Outside convex, so test against bounding box - this is done to avoid
+				//inaccurate results given by PhiWithNormalInternal when x is far outside
+
+				FVec3 SnappedPosition, BoundingNormal;
+				const float BoundingPhi = LocalBoundingBox.PhiWithNormal(x, BoundingNormal);
+				if (BoundingPhi <= 0)
+				{
+					//Inside bounding box - snap to convex
+					SnappedPosition = x - Phi * Normal;
+				}
+				else
+				{
+					//Snap to bounding box, and test convex
+					SnappedPosition = x - BoundingPhi * BoundingNormal;
+					Phi = PhiWithNormalInternal(SnappedPosition, Normal);
+					SnappedPosition -= Phi * Normal;
+				}
+
+				//one final snap to ensure we're on the surface
+				Phi = PhiWithNormalInternal(SnappedPosition, Normal);
+				SnappedPosition -= Phi * Normal;
+
+				//Return phi/normal based on distance from original position to snapped position
+				const FVec3 Difference = x - SnappedPosition;
+				Phi = Difference.Size();
+				if (CHAOS_ENSURE(Phi > TNumericLimits<float>::Min())) //Phi shouldn't be 0 here since we only enter this block if x was outside the convex
+				{
+					Normal = (Difference) / Phi;
+				}
+				else
+				{
+					Normal = FVector::ForwardVector;
+				}
+			}
+			return Phi;
+		}
+
+		FReal PhiWithNormalInternal(const FVec3& x, FVec3& Normal) const
 		{
 			const int32 NumPlanes = Planes.Num();
 			if (NumPlanes == 0)
