Revive GeometryCollection test RigidBodies_SingleSphereCollidingWithSolverFloor and fix RigidBodies_SingleBodyCollidingWithGroundPlane to correctly specify Box as it's implicit type.  Add a utility function for making a geometry collections from and arbitrary mesh.  Add a test utility function that generates a geometry collection containing a sphere.

diff --git a/Engine/Source/Runtime/Experimental/GeometryCollectionCore/Public/GeometryCollection/GeometryCollectionUtility.h b/Engine/Source/Runtime/Experimental/GeometryCollectionCore/Public/GeometryCollection/GeometryCollectionUtility.h
--- a/Engine/Source/Runtime/Experimental/GeometryCollectionCore/Public/GeometryCollection/GeometryCollectionUtility.h
+++ b/Engine/Source/Runtime/Experimental/GeometryCollectionCore/Public/GeometryCollection/GeometryCollectionUtility.h
@@ -9,6 +9,89 @@
 
 namespace GeometryCollection 
 {
+	/****
+	* MakeMeshElement
+	*   Utility to create an arbitrary triangulated mesh using the FGeometryCollection format.
+	*/
+	template <class TV3_PTS, class TV3_NORM, class TV2, class TV_INT3>
+	TSharedPtr<FGeometryCollection> MakeMeshElement(
+		const TArray<TV3_PTS>& PointsIn,
+		const TArray<TV3_NORM>& NormalsIn,
+		const TArray<TV_INT3>& TrianglesIn,
+		const TArray<TV2>& UVsIn,
+		const FTransform& Xf,
+		const FTransform& GeoXf = FTransform::Identity,
+		const int NumberOfMaterials = 2)
+	{
+		FGeometryCollection* RestCollection = new FGeometryCollection();
+		RestCollection->AddElements(PointsIn.Num(), FGeometryCollection::VerticesGroup);
+		RestCollection->AddElements(TrianglesIn.Num(), FGeometryCollection::FacesGroup);
+		RestCollection->AddElements(1, FGeometryCollection::TransformGroup);
+
+		TManagedArray<FVector>& Vertices = RestCollection->Vertex;
+		TManagedArray<FVector>& Normals = RestCollection->Normal;
+		TManagedArray<FVector>& TangentU = RestCollection->TangentU;
+		TManagedArray<FVector>& TangentV = RestCollection->TangentV;
+		TManagedArray<FVector2D>& UVs = RestCollection->UV;
+		TManagedArray<FLinearColor>& Colors = RestCollection->Color;
+		TManagedArray<FIntVector>& Indices = RestCollection->Indices;
+		TManagedArray<bool>& Visible = RestCollection->Visible;
+		TManagedArray<int32>& MaterialIndex = RestCollection->MaterialIndex;
+		TManagedArray<int32>& MaterialID = RestCollection->MaterialID;
+		TManagedArray<FTransform>& Transform = RestCollection->Transform;
+
+		// Set particle info
+		Transform[0] = Xf;
+		Transform[0].NormalizeRotation();
+
+		// Set vertex info
+		for (int32 Idx = 0; Idx < PointsIn.Num(); ++Idx)
+		{
+			Vertices[Idx] = GeoXf.TransformPosition(FVector(PointsIn[Idx][0], PointsIn[Idx][1], PointsIn[Idx][2])); // transform points by GeoXf
+			Normals[Idx] = NormalsIn.Num() > Idx ? FVector(NormalsIn[Idx][0], NormalsIn[Idx][1], NormalsIn[Idx][2]) : FVector(0);
+			UVs[Idx] = UVsIn.Num() > Idx ? FVector2D(UVsIn[Idx][0], UVsIn[Idx][1]) : FVector2D(0);
+			Colors[Idx] = FLinearColor::White;
+		}
+
+		// Set face info
+		const int NumberOfEachMaterial = TrianglesIn.Num() / NumberOfMaterials;
+		for (int32 Idx = 0; Idx < TrianglesIn.Num(); ++Idx)
+		{
+			const auto& Tri = TrianglesIn[Idx];
+			Indices[Idx] = FIntVector(Tri[0], Tri[1], Tri[2]);
+
+			Visible[Idx] = true;
+			MaterialIndex[Idx] = Idx;
+			MaterialID[Idx] = Idx / NumberOfEachMaterial;
+
+			for (int32 Axis = 0; Axis < 3; ++Axis)
+			{
+				const FVector& Normal = Normals[Tri[Axis]];
+				const FVector Edge = (Vertices[Tri[(Axis + 1) % 3]] - Vertices[Tri[Axis]]);
+				TangentU[Tri[Axis]] = (Edge ^ Normal).GetSafeNormal();
+				TangentV[Tri[Axis]] = (Normal ^ TangentU[Tri[Axis]]).GetSafeNormal();
+			}
+		}
+
+		// GeometryGroup
+		GeometryCollection::AddGeometryProperties(RestCollection);
+
+		// Add the materail sections to simulate NumberOfMaterials on the object
+		TManagedArray<FGeometryCollectionSection>& Sections = RestCollection->Sections;
+
+		// the first 6 indices are material 0
+		int FirstElement = RestCollection->AddElements(NumberOfMaterials, FGeometryCollection::MaterialGroup);
+		for (int Element = 0; Element < NumberOfMaterials; Element++)
+		{
+			Sections[Element].MaterialID = Element;
+			Sections[Element].FirstIndex = (Element * NumberOfEachMaterial) * 3;
+			Sections[Element].NumTriangles = NumberOfEachMaterial;
+			Sections[Element].MinVertexIndex = 0;
+			Sections[Element].MaxVertexIndex = Vertices.Num() - 1;
+		}
+
+		return TSharedPtr<FGeometryCollection>(RestCollection);
+	}
 
 	/****
 	* MakeCubeElement
