fix chaos checked.

diff --git a/Engine/Source/Runtime/Experimental/Chaos/Public/Chaos/ParticleHandle.h b/Engine/Source/Runtime/Experimental/Chaos/Public/Chaos/ParticleHandle.h
--- a/Engine/Source/Runtime/Experimental/Chaos/Public/Chaos/ParticleHandle.h
+++ b/Engine/Source/Runtime/Experimental/Chaos/Public/Chaos/ParticleHandle.h
@@ -1620,7 +1620,7 @@ public:
 private:
 
 #if CHAOS_CHECKED
-	FName MDebugName;
+	TParticleProperty<FName, EParticleProperty::DebugName> MDebugName;
 #endif
 
 protected:
