Chaos : Destruction

- Removed InitializeBodieGT

diff --git a/Engine/Source/Runtime/Experimental/ChaosSolvers/Private/PBDRigidsSolver.cpp b/Engine/Source/Runtime/Experimental/ChaosSolvers/Private/PBDRigidsSolver.cpp
--- a/Engine/Source/Runtime/Experimental/ChaosSolvers/Private/PBDRigidsSolver.cpp
+++ b/Engine/Source/Runtime/Experimental/ChaosSolvers/Private/PBDRigidsSolver.cpp
@@ -430,7 +430,6 @@ namespace Chaos
 		GeometryCollectionPhysicsProxies.AddUnique(InProxy);
 		InProxy->SetSolver(this);
 		InProxy->Initialize();
-		InProxy->InitializeBodiesGT();
 		InProxy->NewData(); // Buffers data on the proxy.
 		FParticlesType* InParticles = &GetParticles();
 
diff --git a/Engine/Source/Runtime/Experimental/ChaosSolvers/Private/PhysicsProxy/GeometryCollectionPhysicsProxy.cpp b/Engine/Source/Runtime/Experimental/ChaosSolvers/Private/PhysicsProxy/GeometryCollectionPhysicsProxy.cpp
--- a/Engine/Source/Runtime/Experimental/ChaosSolvers/Private/PhysicsProxy/GeometryCollectionPhysicsProxy.cpp
+++ b/Engine/Source/Runtime/Experimental/ChaosSolvers/Private/PhysicsProxy/GeometryCollectionPhysicsProxy.cpp
@@ -530,38 +530,14 @@ void FGeometryCollectionPhysicsProxy::Initialize()
 		//Results.Get(1).WorldBounds = FBoxSphereBounds(BoundingBox);
 	}
 
-	InitializedState = ESimulationInitializationState::Unintialized;
-}
-
-void FGeometryCollectionPhysicsProxy::InitializeBodiesGT()
-{
-	const FGeometryCollection* RestCollection = Parameters.RestCollection;
-	FGeometryDynamicCollection* DynamicCollection = Parameters.DynamicCollection;
-
-	if (Parameters.Simulating &&
-		((InitializedState == ESimulationInitializationState::Unintialized) ||
-		(InitializedState == ESimulationInitializationState::Activated)))
-	{
-		const TManagedArray<int32>& TransformIndex = RestCollection->TransformIndex;
-		const TManagedArray<int32>& BoneMap = RestCollection->BoneMap;
-		const TManagedArray<int32>& Parent = RestCollection->Parent;
-		const TManagedArray<TSet<int32>>& Children = RestCollection->Children;
-		const TManagedArray<int32>& SimulationType = RestCollection->SimulationType;
-		const TManagedArray<FVector>& Vertex = RestCollection->Vertex;
-		const TManagedArray<int32>& DynamicState = DynamicCollection->DynamicState;
-		const TManagedArray<int32>& CollisionGroup = DynamicCollection->CollisionGroup;
-		const TManagedArray<float>& Mass = RestCollection->GetAttribute<float>("Mass", FTransformCollection::TransformGroup);
-		const TManagedArray<FVector>& InertiaTensor = RestCollection->GetAttribute<FVector>("InertiaTensor", FTransformCollection::TransformGroup);
+	//TArray<FTransform> Transform;
+	//GeometryCollectionAlgo::GlobalMatrices(DynamicCollection->Transform, DynamicCollection->Parent, Transform);
+	//check(DynamicCollection->Transform.Num() == Transform.Num());
 
-		TArray<FTransform> Transform;
-		GeometryCollectionAlgo::GlobalMatrices(DynamicCollection->Transform, DynamicCollection->Parent, Transform);
-		check(DynamicCollection->Transform.Num() == Transform.Num());
+	NumParticles = GetTransformGroupSize(); // SimulatableParticles.Count(true);
+	BaseParticleIndex = 0; // Are we always zero indexed now?
 
-		NumParticles = GetTransformGroupSize(); // SimulatableParticles.Count(true);
-		BaseParticleIndex = 0; // Are we always zero indexed now?
-
-		InitializedState = Parameters.InitializationState;
-	}
+	InitializedState = ESimulationInitializationState::Unintialized;
 }
 
 int32 ReportTooManyChildrenNum = -1;
diff --git a/Engine/Source/Runtime/Experimental/ChaosSolvers/Public/PhysicsProxy/GeometryCollectionPhysicsProxy.h b/Engine/Source/Runtime/Experimental/ChaosSolvers/Public/PhysicsProxy/GeometryCollectionPhysicsProxy.h
--- a/Engine/Source/Runtime/Experimental/ChaosSolvers/Public/PhysicsProxy/GeometryCollectionPhysicsProxy.h
+++ b/Engine/Source/Runtime/Experimental/ChaosSolvers/Public/PhysicsProxy/GeometryCollectionPhysicsProxy.h
@@ -149,13 +149,6 @@ public:
 	void Initialize();
 	void Reset() { InitializedState = ESimulationInitializationState::Unintialized; }
 
-	/** 
-	 * Initialize rigid body particles on the game thread. 
-	 *
-	 * Called by \c FPBDRigidsSolver::RegisterObject().
-	 */
-	void InitializeBodiesGT();
-
 	/** 
 	 * Finish initialization on the physics thread. 
 	 *
