Correct indexing in GetConstraint() function

diff --git a/Engine/Source/Runtime/Experimental/Chaos/Private/Chaos/PBDCollisionConstraints.cpp b/Engine/Source/Runtime/Experimental/Chaos/Private/Chaos/PBDCollisionConstraints.cpp
--- a/Engine/Source/Runtime/Experimental/Chaos/Private/Chaos/PBDCollisionConstraints.cpp
+++ b/Engine/Source/Runtime/Experimental/Chaos/Private/Chaos/PBDCollisionConstraints.cpp
@@ -591,6 +591,27 @@ namespace Chaos
 		return bNeedsAnotherIteration;
 	}
 
+	template<typename T, int d>
+	const typename TPBDCollisionConstraints<T, d>::FConstraintBase& TPBDCollisionConstraints<T, d>::GetConstraint(int32 Index) const
+	{
+		check(Index < NumConstraints());
+		
+		if (Index < PointConstraints.Num())
+		{
+			return PointConstraints[Index];
+		}
+		Index -= PointConstraints.Num();
+
+		if (Index < SweptPointConstraints.Num())
+		{
+			return SweptPointConstraints[Index];
+		}
+		Index -= SweptPointConstraints.Num();
+
+		{
+			return IterativeConstraints[Index];
+		}
+	}
 
 
 	template class TAccelerationStructureHandle<float, 3>;
diff --git a/Engine/Source/Runtime/Experimental/Chaos/Public/Chaos/PBDCollisionConstraints.h b/Engine/Source/Runtime/Experimental/Chaos/Public/Chaos/PBDCollisionConstraints.h
--- a/Engine/Source/Runtime/Experimental/Chaos/Public/Chaos/PBDCollisionConstraints.h
+++ b/Engine/Source/Runtime/Experimental/Chaos/Public/Chaos/PBDCollisionConstraints.h
@@ -284,17 +284,7 @@ public:
 		return Handles;
 	}
 
-	const FConstraintBase& GetConstraint(int32 Index) const 
-	{
-		check(Index < NumConstraints());
-		
-		if (Index < PointConstraints.Num())
-		{
-			return PointConstraints[Index];
-		}
-		
-		return IterativeConstraints[Index - PointConstraints.Num()];
-	}
+	const FConstraintBase& GetConstraint(int32 Index) const;
 
 
 protected:
