Temporarily disable CCD to enable valet playtesting. Please don't drive over the speedlimit ;)

diff --git a/Engine/Source/Runtime/Experimental/Chaos/Private/Chaos/CollisionResolution.cpp b/Engine/Source/Runtime/Experimental/Chaos/Private/Chaos/CollisionResolution.cpp
--- a/Engine/Source/Runtime/Experimental/Chaos/Private/Chaos/CollisionResolution.cpp
+++ b/Engine/Source/Runtime/Experimental/Chaos/Private/Chaos/CollisionResolution.cpp
@@ -31,7 +31,7 @@ DECLARE_CYCLE_STAT(TEXT("Collisions::GJK"), STAT_Collisions_GJK, STATGROUP_Chaos
 
 //PRAGMA_DISABLE_OPTIMIZATION
 
-float CCDEnableThresholdBoundsScale = 0.4f;
+float CCDEnableThresholdBoundsScale = -1.f;//0.4f;
 FAutoConsoleVariableRef  CVarCCDEnableThresholdBoundsScale(TEXT("p.Chaos.CCD.EnableThresholdBoundsScale"), CCDEnableThresholdBoundsScale , TEXT("CCD is used when object position is changing > smallest bound's extent * BoundsScale. 0 will always Use CCD. Values < 0 disables CCD."));
 
 float CCDAllowedDepthBoundsScale = 0.05f;
diff --git a/Engine/Source/Runtime/Experimental/Chaos/Private/Chaos/PBDRigidsEvolutionGBF.cpp b/Engine/Source/Runtime/Experimental/Chaos/Private/Chaos/PBDRigidsEvolutionGBF.cpp
--- a/Engine/Source/Runtime/Experimental/Chaos/Private/Chaos/PBDRigidsEvolutionGBF.cpp
+++ b/Engine/Source/Runtime/Experimental/Chaos/Private/Chaos/PBDRigidsEvolutionGBF.cpp
@@ -52,7 +52,7 @@ float BoundsThicknessVelocityMultiplier = 2.0f;
 FAutoConsoleVariableRef CVarBoundsThickness(TEXT("p.CollisionBoundsThickness"), BoundsThickness, TEXT("Collision inflation for speculative contact generation.[def:0.0]"));
 FAutoConsoleVariableRef CVarBoundsThicknessVelocityMultiplier(TEXT("p.CollisionBoundsVelocityInflation"), BoundsThicknessVelocityMultiplier, TEXT("Collision velocity inflation for speculatibe contact generation.[def:2.0]"));
 
-float HackCCD_EnableThreshold = 0.3f;
+float HackCCD_EnableThreshold = -1.f;//0.3f;
 FAutoConsoleVariableRef CVarHackCCDVelThreshold(TEXT("p.Chaos.CCD.EnableThreshold"), HackCCD_EnableThreshold, TEXT("If distance moved is greater than this times the minimum object dimension, use CCD"));
 
 float HackCCD_DepthThreshold = 0.05f;
