Remove destroyed particles from GameThreadActiveParticlesBuffer. Prevents crash if GT destroys partilcle during CompleteSceneSimulation.

diff --git a/Engine/Source/Runtime/Experimental/Chaos/Public/Chaos/Framework/MultiBufferResource.h b/Engine/Source/Runtime/Experimental/Chaos/Public/Chaos/Framework/MultiBufferResource.h
--- a/Engine/Source/Runtime/Experimental/Chaos/Public/Chaos/Framework/MultiBufferResource.h
+++ b/Engine/Source/Runtime/Experimental/Chaos/Public/Chaos/Framework/MultiBufferResource.h
@@ -23,6 +23,10 @@ namespace Chaos
 
 		virtual EMultiBufferMode GetBufferMode() = 0;
 		virtual ResourceType* AccessProducerBuffer() = 0;
+
+		//NOTE: these changes may not make it to producer side, it is meant for resource management not to pass information back
+		virtual ResourceType* GetConsumerBufferMutable() = 0;
+
 		virtual const ResourceType* GetProducerBuffer() const = 0;
 		virtual const ResourceType* GetConsumerBuffer() const = 0;
 		virtual void FlipProducer() = 0;
@@ -42,6 +46,7 @@ namespace Chaos
 
 		virtual EMultiBufferMode GetBufferMode() override { return EMultiBufferMode::Single; }
 		virtual ResourceType* AccessProducerBuffer() override { return &Data; }
+		virtual ResourceType* GetConsumerBufferMutable() override {return &Data;}
 		virtual const ResourceType* GetProducerBuffer() const override { return &Data; }
 		virtual const ResourceType* GetConsumerBuffer() const override { return &Data; }
 
@@ -68,6 +73,7 @@ namespace Chaos
 
 		virtual EMultiBufferMode GetBufferMode() override { return EMultiBufferMode::Double; }
 		virtual ResourceType* AccessProducerBuffer() override { return Data_Producer; }
+		virtual ResourceType* GetConsumerBufferMutable() override {return Data_Consumer;}
 		virtual const ResourceType* GetProducerBuffer() const override { return Data_Producer; }
 		virtual const ResourceType* GetConsumerBuffer() const override { return Data_Consumer; }
 
@@ -110,6 +116,8 @@ namespace Chaos
 
 		virtual EMultiBufferMode GetBufferMode() override { return EMultiBufferMode::Triple; }
 		virtual ResourceType* AccessProducerBuffer() override { return &Data[GetWriteIndex()]; }
+		virtual ResourceType* GetConsumerBufferMutable() override {return &Data[GetReadIndex()];}
+
 		virtual const ResourceType* GetProducerBuffer() const override { return &Data[GetWriteIndex()]; }
 		virtual const ResourceType* GetConsumerBuffer() const override { return &Data[GetReadIndex()]; }
 
@@ -209,6 +217,12 @@ namespace Chaos
 			return ProducerThreadBuffer->Value.Get();
 		}
 
+		virtual ResourceType* GetConsumerBufferMutable() override
+		{
+			check(false);
+			return nullptr;
+		}
+
 		/**
 		 * Get the current producer buffer.
 		 *
diff --git a/Engine/Source/Runtime/Experimental/ChaosSolvers/Private/PBDRigidActiveParticlesBuffer.cpp b/Engine/Source/Runtime/Experimental/ChaosSolvers/Private/PBDRigidActiveParticlesBuffer.cpp
--- a/Engine/Source/Runtime/Experimental/ChaosSolvers/Private/PBDRigidActiveParticlesBuffer.cpp
+++ b/Engine/Source/Runtime/Experimental/ChaosSolvers/Private/PBDRigidActiveParticlesBuffer.cpp
@@ -19,6 +19,14 @@ namespace Chaos
 		ResourceOutLock.WriteUnlock();
 	}
 
+	void FPBDRigidActiveParticlesBuffer::RemoveActiveParticleFromConsumerBuffer(TGeometryParticle<FReal, 3>* Particle)
+	{
+		ResourceOutLock.WriteLock();
+		auto& ActiveGameThreadParticles = SolverDataOut->GetConsumerBufferMutable()->ActiveGameThreadParticles;
+		ActiveGameThreadParticles.RemoveSingleSwap(Particle);
+		ResourceOutLock.WriteUnlock();
+	}
+
 	void FPBDRigidActiveParticlesBuffer::BufferPhysicsResults(FPBDRigidsSolver* Solver)
 	{
 		auto& ActiveGameThreadParticles = SolverDataOut->AccessProducerBuffer()->ActiveGameThreadParticles;
diff --git a/Engine/Source/Runtime/Experimental/ChaosSolvers/Private/PBDRigidsSolver.cpp b/Engine/Source/Runtime/Experimental/ChaosSolvers/Private/PBDRigidsSolver.cpp
--- a/Engine/Source/Runtime/Experimental/ChaosSolvers/Private/PBDRigidsSolver.cpp
+++ b/Engine/Source/Runtime/Experimental/ChaosSolvers/Private/PBDRigidsSolver.cpp
@@ -415,6 +415,10 @@ namespace Chaos
 				delete Proxy;
 			}
 
+			// Remove game thread particle from ActiveGameThreadParticles so we won't crash when pulling physics state
+			// if this particle was deleted after buffering results. 
+			Solver->GetActiveParticlesBuffer()->RemoveActiveParticleFromConsumerBuffer(Handle->GTGeometryParticle());
+
 			Solver->MParticleToProxy.Remove(Handle);
 
 			// Use the handle to destroy the particle data
diff --git a/Engine/Source/Runtime/Experimental/ChaosSolvers/Public/PBDRigidActiveParticlesBuffer.h b/Engine/Source/Runtime/Experimental/ChaosSolvers/Public/PBDRigidActiveParticlesBuffer.h
--- a/Engine/Source/Runtime/Experimental/ChaosSolvers/Public/PBDRigidActiveParticlesBuffer.h
+++ b/Engine/Source/Runtime/Experimental/ChaosSolvers/Public/PBDRigidActiveParticlesBuffer.h
@@ -29,6 +29,8 @@ namespace Chaos
 		FPBDRigidActiveParticlesBuffer(const Chaos::EMultiBufferMode& InBufferMode);
 
 		void CaptureSolverData(FPBDRigidsSolver* Solver);
+
+		void RemoveActiveParticleFromConsumerBuffer(TGeometryParticle<FReal, 3>* Particle);
 	
 	private:
 		const FPBDRigidActiveParticlesBufferOut* GetSolverOutData() const
