Chaos : Destruction
-Removed cross collection clustering initialization support.

diff --git a/Engine/Source/Runtime/Engine/Private/PhysicsEngine/Experimental/PhysScene_Chaos.cpp b/Engine/Source/Runtime/Engine/Private/PhysicsEngine/Experimental/PhysScene_Chaos.cpp
--- a/Engine/Source/Runtime/Engine/Private/PhysicsEngine/Experimental/PhysScene_Chaos.cpp
+++ b/Engine/Source/Runtime/Engine/Private/PhysicsEngine/Experimental/PhysScene_Chaos.cpp
@@ -559,61 +559,25 @@ Chaos::IDispatcher* FPhysScene_Chaos::GetDispatcher() const
 template<typename ObjectType>
 void AddPhysicsProxy(ObjectType* InObject, Chaos::FPhysicsSolver* InSolver, Chaos::IDispatcher* InDispatcher)
 {
-	check(IsInGameThread() && InSolver);
-	//const bool bDedicatedThread = ChaosModule->IsPersistentTaskEnabled();
-
-	InObject->SetSolver(InSolver);
-	InObject->Initialize();
-
-	if(/*bDedicatedThread && */InDispatcher)
-	{
-		// Pass the proxy off to the physics thread
-		InDispatcher->EnqueueCommandImmediate(InSolver, [InObject](Chaos::FPhysicsSolver* InnerSolver)
-		{
-#if CHAOS_PARTICLEHANDLE_TODO
-			InnerSolver->RegisterObject(InObject);
-#endif
-		});
-	}
+	ensure(false);
 }
 
 void FPhysScene_Chaos::AddObject(UPrimitiveComponent* Component, FSkeletalMeshPhysicsProxy* InObject)
 {
 	AddToComponentMaps(Component, InObject);
-
 	ensure(false);
-#if 0
-	Chaos::FPhysicsSolver* Solver = GetSolver();
-	Solver->RegisterObject(InObject);
-
-	AddPhysicsProxy(InObject, Solver, GetDispatcher());
-#endif
 }
 
 void FPhysScene_Chaos::AddObject(UPrimitiveComponent* Component, FStaticMeshPhysicsProxy* InObject)
 {
 	AddToComponentMaps(Component, InObject);
-
 	ensure(false);
-#if 0
-	Chaos::FPhysicsSolver* Solver = GetSolver();
-	Solver->RegisterObject(InObject);
-
-	AddPhysicsProxy(InObject, Solver, GetDispatcher());
-#endif
 }
 
 void FPhysScene_Chaos::AddObject(UPrimitiveComponent* Component, FGeometryParticlePhysicsProxy* InObject)
 {
 	AddToComponentMaps(Component, InObject);
-
 	ensure(false);
-#if 0
-	Chaos::FPhysicsSolver* Solver = GetSolver();
-	Solver->RegisterObject(InObject);
-
-	AddPhysicsProxy(InObject, Solver, GetDispatcher());
-#endif
 }
 
 void FPhysScene_Chaos::AddObject(UPrimitiveComponent* Component, FGeometryCollectionPhysicsProxy* InObject)
@@ -622,8 +586,6 @@ void FPhysScene_Chaos::AddObject(UPrimitiveComponent* Component, FGeometryCollec
 
 	Chaos::FPhysicsSolver* Solver = GetSolver();
 	Solver->RegisterObject(InObject);
-
-	AddPhysicsProxy(InObject, Solver, GetDispatcher());
 }
 
 void FPhysScene_Chaos::AddObject(UPrimitiveComponent* Component, FFieldSystemPhysicsProxy* InObject)
diff --git a/Engine/Source/Runtime/Experimental/ChaosSolvers/Private/PhysicsProxy/GeometryCollectionPhysicsProxy.cpp b/Engine/Source/Runtime/Experimental/ChaosSolvers/Private/PhysicsProxy/GeometryCollectionPhysicsProxy.cpp
--- a/Engine/Source/Runtime/Experimental/ChaosSolvers/Private/PhysicsProxy/GeometryCollectionPhysicsProxy.cpp
+++ b/Engine/Source/Runtime/Experimental/ChaosSolvers/Private/PhysicsProxy/GeometryCollectionPhysicsProxy.cpp
@@ -361,18 +361,17 @@ void PopulateSimulatedParticle(
 
 FGeometryCollectionPhysicsProxy::FGeometryCollectionPhysicsProxy(
 	UObject* InOwner, 
-	FGeometryDynamicCollection* InDynamicCollection, 
+	const FSimulationParameters & SimulationParameters,
 	FInitFunc InInitFunc, 
 	FCacheSyncFunc InCacheSyncFunc, 
 	FFinalSyncFunc InFinalSyncFunc,
 	const Chaos::EMultiBufferMode BufferMode)
 	: Base(InOwner)
-	//, InitializedState(ESimulationInitializationState::Unintialized)
-	, InitializedState(ESimulationInitializationState::Created)
+	, Parameters(SimulationParameters)
 	, IsObjectDynamic(false)
 	, IsObjectLoading(true)
 
-	, GTDynamicCollection(InDynamicCollection)
+	, GTDynamicCollection(SimulationParameters.DynamicCollection)
 	//, PTDynamicCollection(nullptr)
 
 	, BaseParticleIndex(INDEX_NONE)
@@ -401,9 +400,6 @@ FAutoConsoleVariableRef CVarReportHighParticleFraction(TEXT("p.gc.ReportHighPart
 
 void FGeometryCollectionPhysicsProxy::Initialize()
 {
-	if (InitializedState != ESimulationInitializationState::Created)
-		return;
-
 	// Old proxy init
 	check(IsInGameThread());
 
@@ -439,7 +435,10 @@ void FGeometryCollectionPhysicsProxy::Initialize()
 	//});
 
 	// Back to engine for setup from components
-	InitFunc(Parameters);
+	if (InitFunc)
+	{
+		InitFunc(Parameters);
+	}
 
 	PTDynamicCollection.SyncAllGroups(*Parameters.RestCollection);
 	CollisionParticlesPerObjectFraction = 
@@ -478,7 +477,6 @@ void FGeometryCollectionPhysicsProxy::Initialize()
 	CreateDynamicAttributes();
 
 	ProxySimDuration = 0.0f;
-//	InitializedState = Parameters.InitializationState;
 
 	// Old proxy init
 	RecordedTracks.Records.Reset();
@@ -536,8 +534,6 @@ void FGeometryCollectionPhysicsProxy::Initialize()
 
 	NumParticles = GetTransformGroupSize(); // SimulatableParticles.Count(true);
 	BaseParticleIndex = 0; // Are we always zero indexed now?
-
-	InitializedState = ESimulationInitializationState::Unintialized;
 }
 
 int32 ReportTooManyChildrenNum = -1;
@@ -550,9 +546,7 @@ void FGeometryCollectionPhysicsProxy::InitializeBodiesPT(
 	const FGeometryCollection* RestCollection = Parameters.RestCollection;
 	FGeometryDynamicCollection* DynamicCollection = Parameters.DynamicCollection;
 
-	if (Parameters.Simulating && 
-		((InitializedState == ESimulationInitializationState::Unintialized) || 
-		 (InitializedState == ESimulationInitializationState::Activated)))
+	if (Parameters.Simulating)
 	{
 		const TManagedArray<int32>& TransformIndex = RestCollection->TransformIndex;
 		const TManagedArray<int32>& BoneMap = RestCollection->BoneMap;
@@ -847,63 +841,8 @@ void FGeometryCollectionPhysicsProxy::InitializeBodiesPT(
 		}
 #endif // TODO_REIMPLEMENT_RIGID_CACHING
 
-		if (InitializedState == ESimulationInitializationState::Activated)
-		{
-			//
-			//  Activated bodies has already been called so we are good to go.
-			//
-			InitializedState = ESimulationInitializationState::Initialized;
-
-			if (Parameters.EnableClustering && Parameters.ClusterGroupIndex)
-			{
-				GetSolver()->GetEvolution()->GetRigidClustering().IncrementPendingClusterCounter(Parameters.ClusterGroupIndex);
-				GetSolver()->GetEvolution()->GetRigidClustering().DecrementPendingClusterCounter(Parameters.ClusterGroupIndex);
-			}
-		}
-		else if (InitializedState == ESimulationInitializationState::Unintialized)
-		{
-			//
-			//  Activated bodies has not been called, so we are waiting 
-			//  to become active. Deactivate all bodies, and wait for
-			//  ActivatedBodies to be called, and defer the cluster initialization
-			//
-			InitializedState = ESimulationInitializationState::Created;
-			for (int32 TransformGroupIndex = 0; TransformGroupIndex < NumTransforms; TransformGroupIndex++)
-			{
-				if(Chaos::TPBDRigidParticleHandle<float, 3>* Handle = SolverParticleHandles[TransformGroupIndex])
-				//const int32 RigidBodyIndex = RigidBodyID[TransformGroupIndex];
-				//if (RigidBodyIndex != INDEX_NONE)
-				{
-#define TODO_REIMPLEMENT_DEFERRED_ACTIVATION 0
-#if TODO_REIMPLEMENT_DEFERRED_ACTIVATION
-					if (!Handle->Disabled())
-					//if (!Particles.Disabled(RigidBodyIndex))
-					{
-						PendingActivationList.Add(TransformGroupIndex);
-						RigidsSolver->GetEvolution()->DisableParticle(Handle);
-						//GetSolver()->GetEvolution()->DisableParticle(RigidBodyIndex);
-					}
-#endif
-				}
-			}
-
-			//
-			//  Clustering needs to advertise its group id to the cluster so
-			//  that the group is not initialized before all the bodies are
-			//  loaded and created. 
-			//
-			if (Parameters.EnableClustering && Parameters.ClusterGroupIndex )
-			{
-				GetSolver()->GetEvolution()->GetRigidClustering().IncrementPendingClusterCounter(
-					Parameters.ClusterGroupIndex);
-			}
-		}
-		else
-		{
-			// unknown initialization state in creation callback
-			ensure(false);
-		}
 	} // end if simulating...
+
 }
 
 int32 ReportNoLevelsetCluster = 0;
@@ -1940,11 +1879,6 @@ void FGeometryCollectionPhysicsProxy::PushToPhysicsState(const Chaos::FParticleD
 	 * the solver
 	 */
 
-	if (InitializedState == ESimulationInitializationState::Unintialized)
-	{
-		return;
-	}
-
 	// The other proxies do a deep copy into dynamic transient memory, passed into
 	// this function via InData.  Instead of that, we're using persistent memory
 	// in a buffer to facilitate communicating between threads.
@@ -2334,10 +2268,6 @@ void FGeometryCollectionPhysicsProxy::UpdateGeometryCollection(
 	const FGeometryCollectionResults& GCResults,
 	FGeometryDynamicCollection* Collection)
 {
-	if (InitializedState != ESimulationInitializationState::Initialized)
-	{
-		return;
-	}
 	if (!Collection)
 	{
 		Collection = Parameters.DynamicCollection;
@@ -3341,8 +3271,6 @@ void FGeometryCollectionPhysicsProxy::ProcessCommands(FParticlesType& Particles,
 
 void FGeometryCollectionPhysicsProxy::FieldForcesUpdateCallback(Chaos::FPhysicsSolver* InSolver, FParticlesType& Particles, Chaos::TArrayCollectionArray<FVector> & Force, Chaos::TArrayCollectionArray<FVector> & Torque, const float Time)
 {
-	if (InitializedState == ESimulationInitializationState::Initialized)
-	{
 		if (Commands.Num())
 		{
 			// @todo: This seems like a waste if we just want to get everything
@@ -3404,7 +3332,6 @@ void FGeometryCollectionPhysicsProxy::FieldForcesUpdateCallback(Chaos::FPhysicsS
 
 			}
 		}
-	}
 }
 
 
@@ -3416,9 +3343,9 @@ void FGeometryCollectionPhysicsProxy::FieldForcesUpdateCallback(Chaos::FPhysicsS
 void FGeometryCollectionPhysicsProxy::UpdateKinematicBodiesCallback(const FParticlesType& Particles, const float Dt, const float Time, FKinematicProxy& Proxy)
 {
 	check(false); // not currently called
-	if (InitializedState == ESimulationInitializationState::Initialized)
-	{
-		SCOPE_CYCLE_COUNTER(STAT_KinematicUpdate);
+#if TODO_REIMPLEMENT_KINEMATIC_PROXY
+
+	SCOPE_CYCLE_COUNTER(STAT_KinematicUpdate);
 		FGeometryDynamicCollection* Collection = Parameters.DynamicCollection;
 		check(Collection);
 
@@ -3429,7 +3356,6 @@ void FGeometryCollectionPhysicsProxy::UpdateKinematicBodiesCallback(const FParti
 			return;
 		}
 
-#if TODO_REIMPLEMENT_KINEMATIC_PROXY
 		bool bFirst = !Proxy.Ids.Num();
 		if (bFirst)
 		{
@@ -3551,12 +3477,12 @@ void FGeometryCollectionPhysicsProxy::UpdateKinematicBodiesCallback(const FParti
 			// #BGallagher Handle new inactives. If it's a cluster parent and it's fully disabled we'll need to decluster it here.
 		}
 #endif
-	}
 }
 
 void FGeometryCollectionPhysicsProxy::StartFrameCallback(const float Dt, const float Time)
 {
 	SCOPE_CYCLE_COUNTER(STAT_GeomBeginFrame);
+#if TODO_REIMPLEMENT_GET_RIGID_PARTICLES
 	if (InitializedState == ESimulationInitializationState::Initialized)
 	{
 		// Reverse playback only plays back what we just recorded.  So, the condition 
@@ -3564,7 +3490,6 @@ void FGeometryCollectionPhysicsProxy::StartFrameCallback(const float Dt, const f
 		const bool bIsReverseCachePlaying = Parameters.IsCacheRecording() && Parameters.ReverseCacheBeginTime != 0 && Time > Parameters.ReverseCacheBeginTime;
 		if (Parameters.IsCachePlaying() || bIsReverseCachePlaying)
 		{
-#if TODO_REIMPLEMENT_GET_RIGID_PARTICLES
 			// Update the enabled/disabled state for kinematic particles for the upcoming frame
 			Chaos::FPhysicsSolver* ThisSolver = GetSolver();
 			Chaos::FPhysicsSolver::FParticlesType& Particles = ThisSolver->GetRigidParticles();
@@ -3788,14 +3713,15 @@ void FGeometryCollectionPhysicsProxy::StartFrameCallback(const float Dt, const f
 			{
 				ThisSolver->InitializeFromParticleData(0);
 			}
-#endif
 		}
 	}
+#endif
 }
 
 void FGeometryCollectionPhysicsProxy::EndFrameCallback(const float EndFrame)
 {
 	check(false); // not currently called
+#if TODO_REIMPLEMENT_GET_RIGID_PARTICLES
 	if (InitializedState == ESimulationInitializationState::Initialized)
 	{
 		FGeometryDynamicCollection* Collection = Parameters.DynamicCollection;
@@ -3805,7 +3731,6 @@ void FGeometryCollectionPhysicsProxy::EndFrameCallback(const float EndFrame)
 
 		if(Collection->HasAttribute("RigidBodyID", FGeometryCollection::TransformGroup))
 		{
-#if TODO_REIMPLEMENT_GET_RIGID_PARTICLES
 			//
 			//  Update transforms for the simulated transforms
 			//
@@ -4032,7 +3957,6 @@ void FGeometryCollectionPhysicsProxy::EndFrameCallback(const float EndFrame)
 					}
 				}
 			}
-#endif
 		}
 
 #if TODO_REIMPLEMENT_RIGID_CLUSTERING
@@ -4042,11 +3966,14 @@ void FGeometryCollectionPhysicsProxy::EndFrameCallback(const float EndFrame)
 #endif
 	}
 	Commands.Empty();
+#endif
+
 }
 
 void FGeometryCollectionPhysicsProxy::CreateRigidBodyCallback(FParticlesType& Particles)
 {
 	check(false); // This function is deprecated and should never be called.
+#if TODO_REIMPLEMENT_RIGID_CLUSTERING
 
 	const FGeometryCollection* RestCollection = Parameters.RestCollection;
 	FGeometryDynamicCollection* DynamicCollection = Parameters.DynamicCollection;
@@ -4056,7 +3983,6 @@ void FGeometryCollectionPhysicsProxy::CreateRigidBodyCallback(FParticlesType& Pa
 		((InitializedState == ESimulationInitializationState::Unintialized) || 
 		 (InitializedState == ESimulationInitializationState::Activated)))
 	{
-#if TODO_REIMPLEMENT_RIGID_CLUSTERING
 		Chaos::FPhysicsSolver* ThisSolver = GetSolver();
 		Chaos::TArrayCollectionArray<int32>& ClusterGroupIndex = ThisSolver->GetRigidClustering().GetClusterGroupIndexArray();
 		Chaos::TArrayCollectionArray<float>& StrainArray = ThisSolver->GetRigidClustering().GetStrainArray();
@@ -4325,73 +4251,27 @@ void FGeometryCollectionPhysicsProxy::CreateRigidBodyCallback(FParticlesType& Pa
 			// unknown initialization state in creation callback
 			ensure(false);
 		}
-#endif
 	}
-}
-
-void FGeometryCollectionPhysicsProxy::ActivateBodies()
-{
-	if (Parameters.Simulating)
-	{
-		if (InitializedState == ESimulationInitializationState::Created)
-		{
-#if TODO_REIMPLEMENT_RIGID_CLUSTERING
-			int32 ParentIndex = INDEX_NONE;
-
-			if (Parameters.EnableClustering && Parameters.ClusterGroupIndex)
-			{
-				auto& Clustering = GetSolver()->GetEvolution()->GetRigidClustering();
-				Clustering.DecrementPendingClusterCounter(Parameters.ClusterGroupIndex);
-				ParentIndex = Parameters.ClusterGroupIndex;
-			}
-#endif
-#if TODO_REIMPLEMENT_DEFERRED_ACTIVATION
-			Chaos::FPhysicsSolver::FParticlesType & Particles = GetSolver()->GetRigidParticles();
-			for (uint32 TransformGroupIndex : PendingActivationList)
-			{
-//				int32 RigidBodyIndex = RigidBodyID[TransformGroupIndex];
-//				checkSlow(RigidBodyIndex != INDEX_NONE);
-				if(Chaos::TPBDRigidParticleHandle<float, 3>* Handle = SolverParticleHandles[TransformGroupIndex])
-//				if (Particles.Disabled(RigidBodyIndex))
-				{
-					GetSolver()->GetEvolution()->EnableParticle(RigidBodyIndex, ParentIndex);
-				}
-			}
 #endif
-			PendingActivationList.Reset(0);
-
-			InitializedState = ESimulationInitializationState::Initialized;
-		}
-		else if (InitializedState == ESimulationInitializationState::Unintialized)
-		{
-			InitializedState = ESimulationInitializationState::Activated;
-		}
-		else
-		{
-			// unknown initialization state in activate bodies
-			ensure(false);
-		}
-	}
 }
 
-void FGeometryCollectionPhysicsProxy::BindParticleCallbackMapping(Chaos::TArrayCollectionArray<PhysicsProxyWrapper> & PhysicsProxyReverseMap, Chaos::TArrayCollectionArray<int32> & ParticleIDReverseMap)
+
+void FGeometryCollectionPhysicsProxy::BindParticleCallbackMapping(Chaos::TArrayCollectionArray<PhysicsProxyWrapper>& PhysicsProxyReverseMap, Chaos::TArrayCollectionArray<int32>& ParticleIDReverseMap)
 {
-	if (InitializedState== ESimulationInitializationState::Initialized)
+	for (int32 Index = 0; Index < RigidBodyID.Num(); Index++)
 	{
-		for (int32 Index = 0; Index < RigidBodyID.Num(); Index++)
+		const int32 RigidBodyIndex = RigidBodyID[Index];
+		if (RigidBodyIndex != INDEX_NONE)
 		{
-			const int32 RigidBodyIndex = RigidBodyID[Index];
-			if (RigidBodyIndex != INDEX_NONE)
-			{
-				PhysicsProxyReverseMap[RigidBodyIndex] = {this, EPhysicsProxyType::GeometryCollectionType};
-				ParticleIDReverseMap[RigidBodyIndex] = Index;
-			}
+			PhysicsProxyReverseMap[RigidBodyIndex] = { this, EPhysicsProxyType::GeometryCollectionType };
+			ParticleIDReverseMap[RigidBodyIndex] = Index;
 		}
 	}
 }
 
 void FGeometryCollectionPhysicsProxy::PushKinematicStateToSolver(FParticlesType& Particles)
 {
+#if TODO_REIMPLEMENT_RIGID_CLUSTERING
 	if (InitializedState == ESimulationInitializationState::Initialized)
 	{
 		FGeometryDynamicCollection* Collection = Parameters.DynamicCollection;
@@ -4399,7 +4279,6 @@ void FGeometryCollectionPhysicsProxy::PushKinematicStateToSolver(FParticlesType&
 		{
 			TManagedArray<int32>& DynamicState = Collection->GetAttribute<int32>("DynamicState", FGeometryCollection::TransformGroup);
 
-#if TODO_REIMPLEMENT_RIGID_CLUSTERING
 			TSet<int32> ClustersToUpdate;
 			const Chaos::TArrayCollectionArray<Chaos::ClusterId>& ClusterID = GetSolver()->GetRigidClustering().GetClusterIdsArray();
 
@@ -4481,44 +4360,41 @@ void FGeometryCollectionPhysicsProxy::PushKinematicStateToSolver(FParticlesType&
 					}
 				}
 			}
-#endif
 		}
 	}
+#endif
 }
 
 void FGeometryCollectionPhysicsProxy::ParameterUpdateCallback(FParticlesType& Particles, const float Time)
 {
-	if (InitializedState == ESimulationInitializationState::Initialized)
-	{
-		FGeometryDynamicCollection* Collection = Parameters.DynamicCollection;
-		check(Collection);
+	FGeometryDynamicCollection* Collection = Parameters.DynamicCollection;
+	check(Collection);
 
-		if (Collection->Transform.Num())
-		{
-			ProcessCommands(Particles, Time);
+	if (Collection->Transform.Num())
+	{
+		ProcessCommands(Particles, Time);
 
 
-			if (Parameters.RecordedTrack)
+		if (Parameters.RecordedTrack)
+		{
+			float ReverseTime = Parameters.RecordedTrack->GetLastTime() - Time + Parameters.ReverseCacheBeginTime;
+			// @todo(mlentine): We shouldn't need to do this every frame
+			if (Parameters.IsCacheRecording() && Time > Parameters.ReverseCacheBeginTime&& Parameters.ReverseCacheBeginTime != 0 && Parameters.RecordedTrack->IsTimeValid(ReverseTime))
 			{
-				float ReverseTime = Parameters.RecordedTrack->GetLastTime() - Time + Parameters.ReverseCacheBeginTime;
-				// @todo(mlentine): We shouldn't need to do this every frame
-				if (Parameters.IsCacheRecording() && Time > Parameters.ReverseCacheBeginTime && Parameters.ReverseCacheBeginTime != 0 && Parameters.RecordedTrack->IsTimeValid(ReverseTime))
+				for (int32 Index = 0; Index < RigidBodyID.Num(); Index++)
 				{
-					for (int32 Index = 0; Index < RigidBodyID.Num(); Index++)
-					{
-						int32 RigidBodyIndex = RigidBodyID[Index];
+					int32 RigidBodyIndex = RigidBodyID[Index];
 
-						// Check index, will be invalid for cluster parents.
-						if (RigidBodyIndex != INDEX_NONE)
-						{
-							Particles.InvM(RigidBodyIndex) = 0.f;
-							Particles.InvI(RigidBodyIndex) = Chaos::PMatrix<float, 3, 3>(0.f);
-						}
+					// Check index, will be invalid for cluster parents.
+					if (RigidBodyIndex != INDEX_NONE)
+					{
+						Particles.InvM(RigidBodyIndex) = 0.f;
+						Particles.InvI(RigidBodyIndex) = Chaos::PMatrix<float, 3, 3>(0.f);
 					}
 				}
 			}
-			/* @question : Should we tell the solver the mass has changed ? */
 		}
+		/* @question : Should we tell the solver the mass has changed ? */
 	}
 }
 
diff --git a/Engine/Source/Runtime/Experimental/ChaosSolvers/Public/PhysicsProxy/GeometryCollectionPhysicsProxy.h b/Engine/Source/Runtime/Experimental/ChaosSolvers/Public/PhysicsProxy/GeometryCollectionPhysicsProxy.h
--- a/Engine/Source/Runtime/Experimental/ChaosSolvers/Public/PhysicsProxy/GeometryCollectionPhysicsProxy.h
+++ b/Engine/Source/Runtime/Experimental/ChaosSolvers/Public/PhysicsProxy/GeometryCollectionPhysicsProxy.h
@@ -135,7 +135,7 @@ public:
 	 */
 	FGeometryCollectionPhysicsProxy(
 		UObject* InOwner, 
-		FGeometryDynamicCollection* InDynamicCollection, 
+		const FSimulationParameters& SimulationParameters,
 		FInitFunc InInitFunc, 
 		FCacheSyncFunc InCacheSyncFunc, 
 		FFinalSyncFunc InFinalSyncFunc  ,
@@ -147,7 +147,7 @@ public:
 	 * and prepare for simulation.
 	 */
 	void Initialize();
-	void Reset() { InitializedState = ESimulationInitializationState::Unintialized; }
+	void Reset() { }
 
 	/** 
 	 * Finish initialization on the physics thread. 
@@ -240,7 +240,6 @@ public:
 	void DisableCollisionsCallback(TSet<TTuple<int32, int32>>& InPairs);
 	void AddForceCallback(FParticlesType& InParticles, const float InDt, const int32 InIndex);
 	void FieldForcesUpdateCallback(Chaos::FPhysicsSolver* InSolver, FParticlesType& Particles, Chaos::TArrayCollectionArray<FVector> & Force, Chaos::TArrayCollectionArray<FVector> & Torque, const float Time);
-	void ActivateBodies();
 
 protected:
 	/**
@@ -332,7 +331,6 @@ private:
 private:
 	FSimulationParameters Parameters;
 
-	ESimulationInitializationState InitializedState;
 	// Records current dynamic state
 	bool IsObjectDynamic;
 	// Indicate when loaded
diff --git a/Engine/Source/Runtime/Experimental/GeometryCollectionEngine/Private/GeometryCollection/GeometryCollectionComponent.cpp b/Engine/Source/Runtime/Experimental/GeometryCollectionEngine/Private/GeometryCollection/GeometryCollectionComponent.cpp
--- a/Engine/Source/Runtime/Experimental/GeometryCollectionEngine/Private/GeometryCollection/GeometryCollectionComponent.cpp
+++ b/Engine/Source/Runtime/Experimental/GeometryCollectionEngine/Private/GeometryCollection/GeometryCollectionComponent.cpp
@@ -218,31 +218,6 @@ void UGeometryCollectionComponent::BeginPlay()
 	// ---------- });
 	//////////////////////////////////////////////////////////////////////////
 
-	if (PhysicsProxy != nullptr)
-	{
-		FChaosSolversModule* ChaosModule = FModuleManager::Get().GetModulePtr<FChaosSolversModule>("ChaosSolvers");
-		if (ChaosModule != nullptr)
-		{
-			Chaos::FPhysicsSolver* Solver = nullptr;
-			if (ChaosSolverActor != nullptr)
-			{
-				Solver = ChaosSolverActor->GetSolver();
-			}
-			else if (FPhysScene_Chaos* Scene = GetPhysicsScene())
-			{
-				Solver = Scene->GetSolver();
-			}
-			if (Solver != nullptr)
-			{
-				ChaosModule->GetDispatcher()->EnqueueCommandImmediate(Solver, [InPhysicsProxy = PhysicsProxy](Chaos::FPhysicsSolver* InSolver)
-				{
-					InPhysicsProxy->ActivateBodies();
-				});
-			}
-		}
-	}
-	InitializationState = ESimulationInitializationState::Activated;
-
 	// default current cache time
 	CurrentCacheTime = MAX_flt;
 }
@@ -1153,6 +1128,51 @@ void UGeometryCollectionComponent::OnCreatePhysicsState()
 				ChaosMaterial->SleepingAngularThreshold = PhysicalMaterial->SleepingAngularVelocityThreshold;
 			}
 
+			FSimulationParameters SimulationParameters;
+			{
+#if !(UE_BUILD_SHIPPING || UE_BUILD_TEST)
+				SimulationParameters.Name = GetPathName();
+#endif
+				RestCollection->GetSharedSimulationParams(SimulationParameters.Shared);
+				SimulationParameters.RestCollection = RestCollection->GetGeometryCollection().Get();
+				SimulationParameters.Simulating = Simulating;
+				SimulationParameters.EnableClustering = EnableClustering;
+				SimulationParameters.ClusterGroupIndex = EnableClustering ? ClusterGroupIndex : 0;
+				SimulationParameters.MaxClusterLevel = MaxClusterLevel;
+				SimulationParameters.DamageThreshold = DamageThreshold;
+				SimulationParameters.ClusterConnectionMethod = (Chaos::FClusterCreationParameters<float>::EConnectionMethod)ClusterConnectionType;
+				SimulationParameters.CollisionGroup = CollisionGroup;
+				SimulationParameters.CollisionSampleFraction = CollisionSampleFraction;
+				SimulationParameters.LinearEtherDrag = LinearEtherDrag;
+				SimulationParameters.AngularEtherDrag = AngularEtherDrag;
+				SimulationParameters.InitialVelocityType = InitialVelocityType;
+				SimulationParameters.InitialLinearVelocity = InitialLinearVelocity;
+				SimulationParameters.InitialAngularVelocity = InitialAngularVelocity;
+				SimulationParameters.bClearCache = true;
+				SimulationParameters.ObjectType = ObjectType;
+				SimulationParameters.CacheType = CacheParameters.CacheMode;
+				SimulationParameters.ReverseCacheBeginTime = CacheParameters.ReverseCacheBeginTime;
+				SimulationParameters.CollisionData.SaveCollisionData = CacheParameters.SaveCollisionData;
+				SimulationParameters.CollisionData.DoGenerateCollisionData = CacheParameters.DoGenerateCollisionData;
+				SimulationParameters.CollisionData.CollisionDataSizeMax = CacheParameters.CollisionDataSizeMax;
+				SimulationParameters.CollisionData.DoCollisionDataSpatialHash = CacheParameters.DoCollisionDataSpatialHash;
+				SimulationParameters.CollisionData.CollisionDataSpatialHashRadius = CacheParameters.CollisionDataSpatialHashRadius;
+				SimulationParameters.CollisionData.MaxCollisionPerCell = CacheParameters.MaxCollisionPerCell;
+				SimulationParameters.BreakingData.SaveBreakingData = CacheParameters.SaveBreakingData;
+				SimulationParameters.BreakingData.DoGenerateBreakingData = CacheParameters.DoGenerateBreakingData;
+				SimulationParameters.BreakingData.BreakingDataSizeMax = CacheParameters.BreakingDataSizeMax;
+				SimulationParameters.BreakingData.DoBreakingDataSpatialHash = CacheParameters.DoBreakingDataSpatialHash;
+				SimulationParameters.BreakingData.BreakingDataSpatialHashRadius = CacheParameters.BreakingDataSpatialHashRadius;
+				SimulationParameters.BreakingData.MaxBreakingPerCell = CacheParameters.MaxBreakingPerCell;
+				SimulationParameters.TrailingData.SaveTrailingData = CacheParameters.SaveTrailingData;
+				SimulationParameters.TrailingData.DoGenerateTrailingData = CacheParameters.DoGenerateTrailingData;
+				SimulationParameters.TrailingData.TrailingDataSizeMax = CacheParameters.TrailingDataSizeMax;
+				SimulationParameters.TrailingData.TrailingMinSpeedThreshold = CacheParameters.TrailingMinSpeedThreshold;
+				SimulationParameters.TrailingData.TrailingMinVolumeThreshold = CacheParameters.TrailingMinVolumeThreshold;
+				SimulationParameters.RemoveOnFractureEnabled = SimulationParameters.Shared.RemoveOnFractureIndices.Num() > 0;
+			}
+
+
 			//
 			// Called from FGeometryCollectionPhysicsProxy::Initialize()
 			//
@@ -1161,48 +1181,7 @@ void UGeometryCollectionComponent::OnCreatePhysicsState()
 #if !(UE_BUILD_SHIPPING || UE_BUILD_TEST)
 				InParams.Name = GetPathName();
 #endif
-				InParams.RestCollection = RestCollection->GetGeometryCollection().Get();
 				GetInitializationCommands(InParams.InitializationCommands);
-				InParams.InitializationState = InitializationState;
-				InParams.Simulating = Simulating;
-				InParams.WorldTransform = GetComponentToWorld();
-				RestCollection->GetSharedSimulationParams(InParams.Shared);
-				InParams.EnableClustering = EnableClustering;
-				InParams.ClusterGroupIndex = EnableClustering?ClusterGroupIndex:0;
-				InParams.MaxClusterLevel = MaxClusterLevel;
-				InParams.DamageThreshold = DamageThreshold;
-				InParams.ClusterConnectionMethod = (Chaos::FClusterCreationParameters<float>::EConnectionMethod)ClusterConnectionType;
-				InParams.CollisionGroup = CollisionGroup;
-				InParams.CollisionSampleFraction = CollisionSampleFraction;
-				InParams.LinearEtherDrag = LinearEtherDrag;
-				InParams.AngularEtherDrag = AngularEtherDrag;
-				InParams.InitialVelocityType = InitialVelocityType;
-				InParams.InitialLinearVelocity = InitialLinearVelocity;
-				InParams.InitialAngularVelocity = InitialAngularVelocity;
-				InParams.bClearCache = true;
-				InParams.ObjectType = ObjectType;
-				InParams.PhysicalMaterial = MakeSerializable(ChaosMaterial);
-				InParams.CacheType = CacheParameters.CacheMode;
-				InParams.ReverseCacheBeginTime = CacheParameters.ReverseCacheBeginTime;
-				InParams.CollisionData.SaveCollisionData = CacheParameters.SaveCollisionData;
-				InParams.CollisionData.DoGenerateCollisionData = CacheParameters.DoGenerateCollisionData;
-				InParams.CollisionData.CollisionDataSizeMax = CacheParameters.CollisionDataSizeMax;
-				InParams.CollisionData.DoCollisionDataSpatialHash = CacheParameters.DoCollisionDataSpatialHash;
-				InParams.CollisionData.CollisionDataSpatialHashRadius = CacheParameters.CollisionDataSpatialHashRadius;
-				InParams.CollisionData.MaxCollisionPerCell = CacheParameters.MaxCollisionPerCell;
-				InParams.BreakingData.SaveBreakingData = CacheParameters.SaveBreakingData;
-				InParams.BreakingData.DoGenerateBreakingData = CacheParameters.DoGenerateBreakingData;
-				InParams.BreakingData.BreakingDataSizeMax = CacheParameters.BreakingDataSizeMax;
-				InParams.BreakingData.DoBreakingDataSpatialHash = CacheParameters.DoBreakingDataSpatialHash;
-				InParams.BreakingData.BreakingDataSpatialHashRadius = CacheParameters.BreakingDataSpatialHashRadius;
-				InParams.BreakingData.MaxBreakingPerCell = CacheParameters.MaxBreakingPerCell;
-				InParams.TrailingData.SaveTrailingData = CacheParameters.SaveTrailingData;
-				InParams.TrailingData.DoGenerateTrailingData = CacheParameters.DoGenerateTrailingData;
-				InParams.TrailingData.TrailingDataSizeMax = CacheParameters.TrailingDataSizeMax;
-				InParams.TrailingData.TrailingMinSpeedThreshold = CacheParameters.TrailingMinSpeedThreshold;
-				InParams.TrailingData.TrailingMinVolumeThreshold = CacheParameters.TrailingMinVolumeThreshold;
-				InParams.RemoveOnFractureEnabled = InParams.Shared.RemoveOnFractureIndices.Num() > 0;
-
 				UGeometryCollectionCache* Cache = CacheParameters.TargetCache;
 				if(Cache && CacheParameters.CacheMode != EGeometryCollectionCacheType::None)
 				{
@@ -1366,7 +1345,8 @@ void UGeometryCollectionComponent::OnCreatePhysicsState()
 			}
 			// end temporary 
 
-			PhysicsProxy = new FGeometryCollectionPhysicsProxy(this, DynamicCollection.Get(), InitFunc, CacheSyncFunc, FinalSyncFunc);
+			SimulationParameters.DynamicCollection = DynamicCollection.Get();
+			PhysicsProxy = new FGeometryCollectionPhysicsProxy(this, SimulationParameters, InitFunc, CacheSyncFunc, FinalSyncFunc);
 			FPhysScene_Chaos* Scene = GetPhysicsScene();
 			Scene->AddObject(this, PhysicsProxy);
 
diff --git a/Engine/Source/Runtime/Experimental/GeometryCollectionSimulationCore/Public/GeometryCollection/GeometryCollectionSimulationCoreTypes.h b/Engine/Source/Runtime/Experimental/GeometryCollectionSimulationCore/Public/GeometryCollection/GeometryCollectionSimulationCoreTypes.h
--- a/Engine/Source/Runtime/Experimental/GeometryCollectionSimulationCore/Public/GeometryCollection/GeometryCollectionSimulationCoreTypes.h
+++ b/Engine/Source/Runtime/Experimental/GeometryCollectionSimulationCore/Public/GeometryCollection/GeometryCollectionSimulationCoreTypes.h
@@ -222,7 +222,6 @@ struct FSimulationParameters
 		, CacheBeginTime(0.0f)
 		, ReverseCacheBeginTime(0.0f)
 		, bClearCache(false)
-		, InitializationState(ESimulationInitializationState::Unintialized)
 		, RemoveOnFractureEnabled(false)
 	{}
 
@@ -255,7 +254,6 @@ struct FSimulationParameters
 		, BreakingData(Other.BreakingData)
 		, TrailingData(Other.TrailingData)
 		, Shared(Other.Shared)
-		, InitializationState(Other.InitializationState)
 		, RemoveOnFractureEnabled(false)
 	{}
 
@@ -312,7 +310,5 @@ struct FSimulationParameters
 
 	FSharedSimulationParameters Shared;
 
-	ESimulationInitializationState InitializationState;
-
 	bool RemoveOnFractureEnabled;
 };
