Force particles that have a swept constraint to always update their manifolds during first Apply iteration.
Re-enable continuous collisoin detection.

diff --git a/Engine/Source/Runtime/Experimental/Chaos/Private/Chaos/CollisionResolution.cpp b/Engine/Source/Runtime/Experimental/Chaos/Private/Chaos/CollisionResolution.cpp
--- a/Engine/Source/Runtime/Experimental/Chaos/Private/Chaos/CollisionResolution.cpp
+++ b/Engine/Source/Runtime/Experimental/Chaos/Private/Chaos/CollisionResolution.cpp
@@ -31,7 +31,7 @@ DECLARE_CYCLE_STAT(TEXT("Collisions::GJK"), STAT_Collisions_GJK, STATGROUP_Chaos
 
 //PRAGMA_DISABLE_OPTIMIZATION
 
-float CCDEnableThresholdBoundsScale = -1.f;//0.4f;
+float CCDEnableThresholdBoundsScale = 0.4f;
 FAutoConsoleVariableRef  CVarCCDEnableThresholdBoundsScale(TEXT("p.Chaos.CCD.EnableThresholdBoundsScale"), CCDEnableThresholdBoundsScale , TEXT("CCD is used when object position is changing > smallest bound's extent * BoundsScale. 0 will always Use CCD. Values < 0 disables CCD."));
 
 float CCDAllowedDepthBoundsScale = 0.05f;
diff --git a/Engine/Source/Runtime/Experimental/Chaos/Private/Chaos/PBDCollisionConstraints.cpp b/Engine/Source/Runtime/Experimental/Chaos/Private/Chaos/PBDCollisionConstraints.cpp
--- a/Engine/Source/Runtime/Experimental/Chaos/Private/Chaos/PBDCollisionConstraints.cpp
+++ b/Engine/Source/Runtime/Experimental/Chaos/Private/Chaos/PBDCollisionConstraints.cpp
@@ -212,6 +212,10 @@ namespace Chaos
 			}
 		}
 
+		// Swept constraint may drastically change particle position during Apply, other constraint's
+		// manifolds with particle may be invalidated, so we should update manifolds with this particle even on first apply iteration.
+		AddParticleAlwaysUpdateManifold(InConstraint.Particle[0]);
+
 		UpdateConstraintMaterialProperties(SweptPointConstraints[Idx]);
 	}
 
@@ -272,6 +276,7 @@ namespace Chaos
 		PointConstraints.Reset();
 		SweptPointConstraints.Reset();
 		IterativeConstraints.Reset();
+		ParticlesAlwaysUpdatingManifold.Reset();
 		Handles.Reset();
 #endif
 
@@ -456,30 +461,45 @@ namespace Chaos
 		}
 	}
 
+	template<typename T, int d>
+	void TPBDCollisionConstraints<T, d>::AddParticleAlwaysUpdateManifold(const TGeometryParticleHandle<T, d>* Particle)
+	{
+		ParticlesAlwaysUpdatingManifold.Add(Particle);
+	}
+
+	template<typename T, int d>
+	bool TPBDCollisionConstraints<T, d>::ShouldParticleAlwaysUpdateManifold(const TGeometryParticleHandle<T, d>* Particle) const
+	{
+		return ParticlesAlwaysUpdatingManifold.Contains(Particle);
+	}
+
 	template<typename T, int d>
 	bool TPBDCollisionConstraints<T, d>::Apply(const T Dt, const int32 Iterations, const int32 NumIterations)
 	{
 		SCOPE_CYCLE_COUNTER(STAT_Collisions_Apply);
 
 		bool bNeedsAnotherIteration = false;
+		bool bAlwaysUpdateManifold = false;
 		if (MApplyPairIterations > 0)
 		{
 			const Collisions::TContactParticleParameters<T> ParticleParameters = { MCullDistance, MShapePadding, &MCollided };
-			const Collisions::TContactIterationParameters<T> IterationParameters = { Dt, Iterations, NumIterations, MApplyPairIterations, ApplyType, &bNeedsAnotherIteration };
+			const Collisions::TContactIterationParameters<T> IterationParameters = { Dt, Iterations, NumIterations, MApplyPairIterations, ApplyType, &bNeedsAnotherIteration, &bAlwaysUpdateManifold };
 
 			for (FPointContactConstraint& Contact : PointConstraints)
 			{
 				Collisions::ApplySinglePoint(Contact, IterationParameters, ParticleParameters);
 			}
 
-			for (FSweptPointContactConstraint& Contact : SweptPointConstraints)
+			for (FMultiPointContactConstraint& Contact : IterativeConstraints)
 			{
-				Collisions::Apply(Contact, IterationParameters, ParticleParameters);
+				Collisions::ApplyMultiPoint(Contact, IterationParameters, ParticleParameters);
 			}
 
-			for (FMultiPointContactConstraint& Contact : IterativeConstraints)
+			// Swept apply may significantly change particle position, invalidating other constraint's manifolds.
+			// We don't update manifolds on first apply iteration, so make sure we apply swept constraints last.
+			for (FSweptPointContactConstraint& Contact : SweptPointConstraints)
 			{
-				Collisions::ApplyMultiPoint(Contact, IterationParameters, ParticleParameters);
+				Collisions::Apply(Contact, IterationParameters, ParticleParameters);
 			}
 		}
 
@@ -541,9 +561,11 @@ namespace Chaos
 				FConstraintContainerHandle* ConstraintHandle = InConstraintHandles[ConstraintHandleIndex];
 				check(ConstraintHandle != nullptr);
 
+				TVector<const TGeometryParticleHandle<T, d>*, 2> ConstrainedParticles = ConstraintHandle->GetConstrainedParticles();
+				bool bAlwaysUpdateManifold = ShouldParticleAlwaysUpdateManifold(ConstrainedParticles[0]) || ShouldParticleAlwaysUpdateManifold(ConstrainedParticles[1]);
 				bool bNeedsAnotherIteration = false;
 				Collisions::TContactParticleParameters<T> ParticleParameters = { MCullDistance, MShapePadding, &MCollided };
-				Collisions::TContactIterationParameters<T> IterationParameters = { Dt, Iterations, NumIterations, MApplyPairIterations, ApplyType, &bNeedsAnotherIteration };
+				Collisions::TContactIterationParameters<T> IterationParameters = { Dt, Iterations, NumIterations, MApplyPairIterations, ApplyType, &bNeedsAnotherIteration, &bAlwaysUpdateManifold };
 				Collisions::Apply(ConstraintHandle->GetContact(), IterationParameters, ParticleParameters);
 
 				if (bNeedsAnotherIteration)
diff --git a/Engine/Source/Runtime/Experimental/Chaos/Private/Chaos/PBDCollisionConstraintsContact.cpp b/Engine/Source/Runtime/Experimental/Chaos/Private/Chaos/PBDCollisionConstraintsContact.cpp
--- a/Engine/Source/Runtime/Experimental/Chaos/Private/Chaos/PBDCollisionConstraintsContact.cpp
+++ b/Engine/Source/Runtime/Experimental/Chaos/Private/Chaos/PBDCollisionConstraintsContact.cpp
@@ -342,7 +342,7 @@ namespace Chaos
 				// Collision is already up-to-date on first iteration (we either just detected it, or updated it in DetectCollisions)
 				// @todo(ccaulfield): this is not great - try to do something nicer like a dirty flag on the constraint?
 				// In particular it is not right if the Collisions are not the first constraints to be solved...
-				const bool bNeedCollisionUpdate = (PairIt > 0) || (IterationParameters.Iteration > 0);
+				const bool bNeedCollisionUpdate = (PairIt > 0) || (IterationParameters.Iteration > 0) || (*IterationParameters.AlwaysUpdateManifold);
 				if (bNeedCollisionUpdate)
 				{
 					Collisions::Update(Constraint, ParticleParameters.CullDistance);
@@ -400,8 +400,8 @@ namespace Chaos
 			const FReal PartialDT = Constraint.TimeOfImpact * IterationParameters.Dt;
 			const FReal RemainingDT = (1 - Constraint.TimeOfImpact) * IterationParameters.Dt;
 			const int32 FakeIteration = FMath::Max(IterationParameters.Iteration, 1); // Force Apply to update constraint, as other constraints could've changed P
-			const FContactIterationParameters IterationParametersPartialDT{ PartialDT, FakeIteration, IterationParameters.NumIterations, IterationParameters.NumPairIterations, IterationParameters.ApplyType, IterationParameters.NeedsAnotherIteration };
-			const FContactIterationParameters IterationParametersRemainingDT{ RemainingDT, FakeIteration, IterationParameters.NumIterations, IterationParameters.NumPairIterations, IterationParameters.ApplyType, IterationParameters.NeedsAnotherIteration };
+			const FContactIterationParameters IterationParametersPartialDT{ PartialDT, FakeIteration, IterationParameters.NumIterations, IterationParameters.NumPairIterations, IterationParameters.ApplyType, IterationParameters.NeedsAnotherIteration, IterationParameters.AlwaysUpdateManifold };
+			const FContactIterationParameters IterationParametersRemainingDT{ RemainingDT, FakeIteration, IterationParameters.NumIterations, IterationParameters.NumPairIterations, IterationParameters.ApplyType, IterationParameters.NeedsAnotherIteration, IterationParameters.AlwaysUpdateManifold };
 
 			// Rewind P to TOI and Apply
 			Particle0->P() = FMath::Lerp(Particle0->X(), Particle0->P(), Constraint.TimeOfImpact);
diff --git a/Engine/Source/Runtime/Experimental/Chaos/Public/Chaos/PBDCollisionConstraints.h b/Engine/Source/Runtime/Experimental/Chaos/Public/Chaos/PBDCollisionConstraints.h
--- a/Engine/Source/Runtime/Experimental/Chaos/Public/Chaos/PBDCollisionConstraints.h
+++ b/Engine/Source/Runtime/Experimental/Chaos/Public/Chaos/PBDCollisionConstraints.h
@@ -165,6 +165,14 @@ public:
 	*/
 	void UpdateManifolds(T Dt);
 
+	/*
+	* Add particle to list of particles that will always update their constraint manifold.
+	*/
+	void AddParticleAlwaysUpdateManifold(const TGeometryParticleHandle<T, d>* Particle);
+
+	bool ShouldParticleAlwaysUpdateManifold(const TGeometryParticleHandle<T, d>* Particle) const;
+
+
 	//
 	// General Rule API
 	//
@@ -301,6 +309,8 @@ private:
 	TArray<FSweptPointContactConstraint> SweptPointConstraints;
 	TArray<FMultiPointContactConstraint> IterativeConstraints;
 
+	TSet<const TGeometryParticleHandle<T, d>*> ParticlesAlwaysUpdatingManifold;
+
 #if CHAOS_COLLISION_PERSISTENCE_ENABLED
 	TMap< FConstraintContainerHandleKey, FConstraintContainerHandle* > Manifolds;
 #endif
diff --git a/Engine/Source/Runtime/Experimental/Chaos/Public/Chaos/PBDCollisionConstraintsContact.h b/Engine/Source/Runtime/Experimental/Chaos/Public/Chaos/PBDCollisionConstraintsContact.h
--- a/Engine/Source/Runtime/Experimental/Chaos/Public/Chaos/PBDCollisionConstraintsContact.h
+++ b/Engine/Source/Runtime/Experimental/Chaos/Public/Chaos/PBDCollisionConstraintsContact.h
@@ -28,6 +28,7 @@ namespace Chaos
 			const int32 NumPairIterations;
 			const ECollisionApplyType ApplyType;	// @todo(chaos): a better way to customize the collision solver
 			bool* NeedsAnotherIteration;
+			bool* AlwaysUpdateManifold;
 		};
 		using FContactIterationParameters = TContactIterationParameters<FReal>;
 
