GeometryCollectionTest - reimplementing field tests

diff --git a/Engine/Source/Runtime/Experimental/ChaosSolvers/Private/PhysicsProxy/GeometryCollectionPhysicsProxy.cpp b/Engine/Source/Runtime/Experimental/ChaosSolvers/Private/PhysicsProxy/GeometryCollectionPhysicsProxy.cpp
--- a/Engine/Source/Runtime/Experimental/ChaosSolvers/Private/PhysicsProxy/GeometryCollectionPhysicsProxy.cpp
+++ b/Engine/Source/Runtime/Experimental/ChaosSolvers/Private/PhysicsProxy/GeometryCollectionPhysicsProxy.cpp
@@ -673,7 +673,6 @@ void FGeometryCollectionPhysicsProxy::InitializeBodiesPT(
 		// After population, the states of each particle could have changed
 		Particles.UpdateGeometryCollectionViews();
 
-#if TODO_REIMPLEMENT_INIT_COMMANDS
 		for (FFieldSystemCommand& Cmd : Parameters.InitializationCommands)
 		{
 			if (Cmd.MetaData.Contains(FFieldSystemMetaData::EMetaType::ECommandData_ProcessingResolution)) 
@@ -686,8 +685,7 @@ void FGeometryCollectionPhysicsProxy::InitializeBodiesPT(
 			Commands.Add(Cmd);
 		}
 		Parameters.InitializationCommands.Empty();
-		ProcessCommands(Particles, GetSolver()->GetSolverTime());
-#endif // TODO_REIMPLEMENT_INIT_COMMANDS
+		ProcessCommands(Particles.GetGeometryCollectionParticles(), GetSolver()->GetSolverTime());
 
 		if (Parameters.InitialVelocityType == EInitialVelocityTypeEnum::Chaos_Initial_Velocity_User_Defined)
 		{
