Add default restitution.

diff --git a/Engine/Source/Runtime/Experimental/Chaos/Private/Chaos/PBDCollisionConstraints.cpp b/Engine/Source/Runtime/Experimental/Chaos/Private/Chaos/PBDCollisionConstraints.cpp
--- a/Engine/Source/Runtime/Experimental/Chaos/Private/Chaos/PBDCollisionConstraints.cpp
+++ b/Engine/Source/Runtime/Experimental/Chaos/Private/Chaos/PBDCollisionConstraints.cpp
@@ -56,6 +56,9 @@ namespace Chaos
 	float DefaultCollisionFriction = 0;
 	FAutoConsoleVariableRef CVarDefaultCollisionFriction(TEXT("p.DefaultCollisionFriction"), DefaultCollisionFriction, TEXT("Collision friction default value if no materials are found."));
 
+	float DefaultCollisionRestitution = 0;
+	FAutoConsoleVariableRef CVarDefaultCollisionRestitution(TEXT("p.DefaultCollisionRestitution"), DefaultCollisionRestitution, TEXT("Collision restitution default value if no materials are found."));
+
 	DECLARE_CYCLE_STAT(TEXT("Collisions::Reset"), STAT_Collisions_Reset, STATGROUP_ChaosCollision);
 	DECLARE_CYCLE_STAT(TEXT("Collisions::UpdatePointConstraints"), STAT_Collisions_UpdatePointConstraints, STATGROUP_ChaosCollision);
 	DECLARE_CYCLE_STAT(TEXT("Collisions::UpdateManifoldConstraints"), STAT_Collisions_UpdateManifoldConstraints, STATGROUP_ChaosCollision);
@@ -150,7 +153,7 @@ namespace Chaos
 		else
 		{
 			Contact.Friction = DefaultCollisionFriction;
-			Contact.Restitution = 0;
+			Contact.Restitution = DefaultCollisionRestitution;
 		}
 		Contact.AngularFriction = MAngularFriction;
 
