Chaos Collision: Box vs Convex excessive jittering workaround -> don't use iterative manifolds

diff --git a/Engine/Source/Runtime/Experimental/Chaos/Private/Chaos/CollisionResolution.cpp b/Engine/Source/Runtime/Experimental/Chaos/Private/Chaos/CollisionResolution.cpp
--- a/Engine/Source/Runtime/Experimental/Chaos/Private/Chaos/CollisionResolution.cpp
+++ b/Engine/Source/Runtime/Experimental/Chaos/Private/Chaos/CollisionResolution.cpp
@@ -1605,7 +1605,14 @@ namespace Chaos
 		{
 			TRigidTransform<T, d> ParticleImplicit0TM = Transform0.GetRelativeTransform(Collisions::GetTransform(Particle0));
 			TRigidTransform<T, d> ParticleImplicit1TM = Transform1.GetRelativeTransform(Collisions::GetTransform(Particle1));
-			if (Chaos_Collision_UseManifolds)
+
+
+			EImplicitObjectType Implicit0Type = Particle0->Geometry()->GetType();
+			EImplicitObjectType Implicit1Type = Particle1->Geometry()->GetType();
+
+			// Note: This TBox check is a temporary workaround to avoid jitter in cases of Box vs Convex; investigation ongoing
+			// We need to improve iterative manifolds for this case
+			if (Chaos_Collision_UseManifolds && Implicit0Type != TBox<FReal, 3>::StaticType() && Implicit1Type != TBox<FReal, 3>::StaticType())
 			{
 				FRigidBodyMultiPointContactConstraint Constraint = FRigidBodyMultiPointContactConstraint(Particle0, Implicit0, ParticleImplicit0TM, Particle1, Implicit1, ParticleImplicit1TM, EContactShapesType::ConvexConvex);
 				UpdateConvexConvexManifold(Constraint, Transform0, Transform1, CullDistance);
