Chaos : Destruction

- Moved the DynamicCollections into ChaosSolvers

diff --git a/Engine/Source/Runtime/Experimental/ChaosSolvers/Private/GeometryCollectionProxyData.cpp b/Engine/Source/Runtime/Experimental/ChaosSolvers/Private/GeometryCollectionProxyData.cpp
new file mode 100644
--- /dev/null
+++ b/Engine/Source/Runtime/Experimental/ChaosSolvers/Private/GeometryCollectionProxyData.cpp
@@ -0,0 +1,50 @@
+// Copyright Epic Games, Inc. All Rights Reserved.
+
+/*=============================================================================
+GeometryCollectionProxyData.cpp: 
+=============================================================================*/
+
+#include "GeometryCollectionProxyData.h"
+#include "GeometryCollection/GeometryCollection.h"
+#include "GeometryCollection/TransformCollection.h"
+
+/*
+* FTransformDynamicCollection (FManagedArrayCollection)
+*/
+
+FTransformDynamicCollection::FTransformDynamicCollection()
+	: FManagedArrayCollection()
+{
+	Construct();
+}
+
+
+void FTransformDynamicCollection::Construct()
+{
+	FManagedArrayCollection::FConstructionParameters TransformDependency(FTransformCollection::TransformGroup);
+
+	// Transform Group
+	AddExternalAttribute<FTransform>(FTransformCollection::TransformAttribute, FTransformCollection::TransformGroup, Transform);
+	AddExternalAttribute<int32>(FTransformCollection::ParentAttribute, FTransformCollection::TransformGroup, Parent);
+	AddExternalAttribute<TSet<int32>>(FTransformCollection::ChildrenAttribute, FTransformCollection::TransformGroup, Children);
+	AddExternalAttribute<int32>(FGeometryCollection::SimulationTypeAttribute, FTransformCollection::TransformGroup, SimulationType);
+	AddExternalAttribute<int32>(FGeometryCollection::StatusFlagsAttribute, FTransformCollection::TransformGroup, StatusFlags);
+}
+
+/*
+* FGeometryDynamicCollection (FTransformDynamicCollection)
+*/
+
+const FName FGeometryDynamicCollection::ActiveAttribute("Active");
+const FName FGeometryDynamicCollection::CollisionGroupAttribute("CollisionGroup");
+const FName FGeometryDynamicCollection::DynamicStateAttribute("DynamicState");
+
+
+FGeometryDynamicCollection::FGeometryDynamicCollection()
+	: FTransformDynamicCollection()
+{
+	// Transform Group
+	AddExternalAttribute<int32>(FGeometryDynamicCollection::DynamicStateAttribute, FTransformCollection::TransformGroup, DynamicState);
+	AddExternalAttribute<bool>(FGeometryDynamicCollection::ActiveAttribute, FTransformCollection::TransformGroup, Active);
+	AddExternalAttribute<int32>(FGeometryDynamicCollection::CollisionGroupAttribute, FTransformCollection::TransformGroup, CollisionGroup);
+}
diff --git a/Engine/Source/Runtime/Experimental/ChaosSolvers/Public/GeometryCollectionProxyData.h b/Engine/Source/Runtime/Experimental/ChaosSolvers/Public/GeometryCollectionProxyData.h
new file mode 100644
--- /dev/null
+++ b/Engine/Source/Runtime/Experimental/ChaosSolvers/Public/GeometryCollectionProxyData.h
@@ -0,0 +1,71 @@
+/// Copyright Epic Games, Inc. All Rights Reserved.
+
+#pragma once
+
+#include "CoreMinimal.h"
+#include "Misc/Crc.h"
+
+
+/*
+* Managed arrays for simulation data used by the GeometryCollectionProxy
+*/
+
+/**
+* FTransformDynamicCollection (FManagedArrayCollection)
+*
+* Stores per instance data for transforms and hierarchy information
+*/
+class CHAOSSOLVERS_API FTransformDynamicCollection : public FManagedArrayCollection
+{
+public:
+	typedef FManagedArrayCollection Super;
+
+	FTransformDynamicCollection();
+	FTransformDynamicCollection(FTransformDynamicCollection&) = delete;
+	FTransformDynamicCollection& operator=(const FTransformDynamicCollection&) = delete;
+	FTransformDynamicCollection(FTransformDynamicCollection&&) = delete;
+	FTransformDynamicCollection& operator=(FTransformDynamicCollection&&) = delete;
+
+	// Transform Group
+	TManagedArray<FTransform>   Transform;
+	TManagedArray<int32>        Parent;
+	TManagedArray<TSet<int32>>  Children;
+	TManagedArray<int32>        SimulationType;
+	TManagedArray<int32>        StatusFlags;
+
+protected:
+
+	/** Construct */
+	void Construct();
+};
+
+
+/**
+* FGeometryDynamicCollection (FTransformDynamicCollection)
+*
+* Stores per instance data for simulation level information
+*/
+
+class CHAOSSOLVERS_API FGeometryDynamicCollection : public FTransformDynamicCollection
+{
+
+public:
+	FGeometryDynamicCollection();
+	FGeometryDynamicCollection(FGeometryDynamicCollection&) = delete;
+	FGeometryDynamicCollection& operator=(const FGeometryDynamicCollection&) = delete;
+	FGeometryDynamicCollection(FGeometryDynamicCollection&&) = delete;
+	FGeometryDynamicCollection& operator=(FGeometryDynamicCollection&&) = delete;
+
+
+	typedef FTransformDynamicCollection Super;
+
+	static const FName ActiveAttribute;
+	static const FName CollisionGroupAttribute;
+	static const FName DynamicStateAttribute;
+
+
+	// Transform Group
+	TManagedArray<bool> Active;
+	TManagedArray<int32> CollisionGroup;
+	TManagedArray<int32> DynamicState;
+};
\ No newline at end of file
diff --git a/Engine/Source/Runtime/Experimental/ChaosSolvers/Public/PhysicsProxy/GeometryCollectionPhysicsProxy.h b/Engine/Source/Runtime/Experimental/ChaosSolvers/Public/PhysicsProxy/GeometryCollectionPhysicsProxy.h
--- a/Engine/Source/Runtime/Experimental/ChaosSolvers/Public/PhysicsProxy/GeometryCollectionPhysicsProxy.h
+++ b/Engine/Source/Runtime/Experimental/ChaosSolvers/Public/PhysicsProxy/GeometryCollectionPhysicsProxy.h
@@ -3,7 +3,7 @@
 #pragma once
 
 #include "GeometryCollection/GeometryCollection.h"
-#include "GeometryCollection/GeometryDynamicCollection.h"
+#include "GeometryCollectionProxyData.h"
 
 #include "Chaos/Framework/PhysicsProxy.h"
 #include "GeometryCollection/ManagedArray.h"
diff --git a/Engine/Source/Runtime/Experimental/GeometryCollectionCore/Private/GeometryCollection/GeometryCollectionAlgo.cpp b/Engine/Source/Runtime/Experimental/GeometryCollectionCore/Private/GeometryCollection/GeometryCollectionAlgo.cpp
--- a/Engine/Source/Runtime/Experimental/GeometryCollectionCore/Private/GeometryCollection/GeometryCollectionAlgo.cpp
+++ b/Engine/Source/Runtime/Experimental/GeometryCollectionCore/Private/GeometryCollection/GeometryCollectionAlgo.cpp
@@ -7,7 +7,6 @@
 #include "GeometryCollection/GeometryCollection.h"
 #include "GeometryCollection/RecordedTransformTrack.h"
 #include "Async/ParallelFor.h"
-#include "GeometryCollection/GeometryDynamicCollection.h"
 
 DEFINE_LOG_CATEGORY_STATIC(GeometryCollectionAlgoLog, Log, All);
 
diff --git a/Engine/Source/Runtime/Experimental/GeometryCollectionCore/Private/GeometryCollection/GeometryDynamicCollection.cpp b/Engine/Source/Runtime/Experimental/GeometryCollectionCore/Private/GeometryCollection/GeometryDynamicCollection.cpp
deleted file mode 100644
diff --git a/Engine/Source/Runtime/Experimental/GeometryCollectionCore/Private/GeometryCollection/TransformDynamicCollection.cpp b/Engine/Source/Runtime/Experimental/GeometryCollectionCore/Private/GeometryCollection/TransformDynamicCollection.cpp
deleted file mode 100644
diff --git a/Engine/Source/Runtime/Experimental/GeometryCollectionCore/Public/GeometryCollection/GeometryDynamicCollection.h b/Engine/Source/Runtime/Experimental/GeometryCollectionCore/Public/GeometryCollection/GeometryDynamicCollection.h
deleted file mode 100644
diff --git a/Engine/Source/Runtime/Experimental/GeometryCollectionCore/Public/GeometryCollection/TransformDynamicCollection.h b/Engine/Source/Runtime/Experimental/GeometryCollectionCore/Public/GeometryCollection/TransformDynamicCollection.h
deleted file mode 100644
diff --git a/Engine/Source/Runtime/Experimental/GeometryCollectionEngine/Public/GeometryCollection/GeometryCollectionComponent.h b/Engine/Source/Runtime/Experimental/GeometryCollectionEngine/Public/GeometryCollection/GeometryCollectionComponent.h
--- a/Engine/Source/Runtime/Experimental/GeometryCollectionEngine/Public/GeometryCollection/GeometryCollectionComponent.h
+++ b/Engine/Source/Runtime/Experimental/GeometryCollectionEngine/Public/GeometryCollection/GeometryCollectionComponent.h
@@ -15,7 +15,7 @@
 #include "GeometryCollection/GeometryCollectionSimulationTypes.h"
 #include "GeometryCollection/GeometryCollectionSimulationCoreTypes.h"
 #include "Physics/Experimental/PhysScene_Chaos.h"
-#include "GeometryCollection/GeometryDynamicCollection.h"
+#include "GeometryCollectionProxyData.h"
 #include "GeometryCollectionObject.h"
 #include "GeometryCollectionEditorSelection.h"
 #include "GeometryCollection/RecordedTransformTrack.h"
diff --git a/Engine/Source/Runtime/Experimental/GeometryCollectionEngine/Public/GeometryCollection/GeometryCollectionEngineTypes.h b/Engine/Source/Runtime/Experimental/GeometryCollectionEngine/Public/GeometryCollection/GeometryCollectionEngineTypes.h
--- a/Engine/Source/Runtime/Experimental/GeometryCollectionEngine/Public/GeometryCollection/GeometryCollectionEngineTypes.h
+++ b/Engine/Source/Runtime/Experimental/GeometryCollectionEngine/Public/GeometryCollection/GeometryCollectionEngineTypes.h
@@ -4,7 +4,7 @@
 
 
 #include "UObject/ObjectMacros.h"
-#include "GeometryCollection/GeometryDynamicCollection.h"
+#include "GeometryCollectionProxyData.h"
 #include "GeometryCollection/GeometryCollection.h"
 
 UENUM()
