Chaos :

-Fix for CIS

diff --git a/Engine/Source/Runtime/Experimental/ChaosSolvers/Public/GeometryCollectionProxyData.h b/Engine/Source/Runtime/Experimental/ChaosSolvers/Public/GeometryCollectionProxyData.h
--- a/Engine/Source/Runtime/Experimental/ChaosSolvers/Public/GeometryCollectionProxyData.h
+++ b/Engine/Source/Runtime/Experimental/ChaosSolvers/Public/GeometryCollectionProxyData.h
@@ -4,7 +4,7 @@
 
 #include "CoreMinimal.h"
 #include "Misc/Crc.h"
-
+#include "GeometryCollection/ManagedArrayCollection.h"
 
 /*
 * Managed arrays for simulation data used by the GeometryCollectionProxy
