Change contact modification callback to be non-const.

diff --git a/Engine/Source/Runtime/Experimental/Chaos/Private/Chaos/PBDCollisionConstraints.cpp b/Engine/Source/Runtime/Experimental/Chaos/Private/Chaos/PBDCollisionConstraints.cpp
--- a/Engine/Source/Runtime/Experimental/Chaos/Private/Chaos/PBDCollisionConstraints.cpp
+++ b/Engine/Source/Runtime/Experimental/Chaos/Private/Chaos/PBDCollisionConstraints.cpp
@@ -288,7 +288,7 @@ namespace Chaos
 	}
 
 	template<typename T, int d>
-	void TPBDCollisionConstraints<T, d>::ApplyCollisionModifier(const TFunction<ECollisionModifierResult(const FConstraintContainerHandle* Handle)>& CollisionModifier)
+	void TPBDCollisionConstraints<T, d>::ApplyCollisionModifier(const TFunction<ECollisionModifierResult(FConstraintContainerHandle* Handle)>& CollisionModifier)
 	{
 		if (CollisionModifier)
 		{
diff --git a/Engine/Source/Runtime/Experimental/Chaos/Public/Chaos/PBDCollisionConstraints.h b/Engine/Source/Runtime/Experimental/Chaos/Public/Chaos/PBDCollisionConstraints.h
--- a/Engine/Source/Runtime/Experimental/Chaos/Public/Chaos/PBDCollisionConstraints.h
+++ b/Engine/Source/Runtime/Experimental/Chaos/Public/Chaos/PBDCollisionConstraints.h
@@ -61,7 +61,7 @@ template<typename T, int d>
 using TRigidBodyContactConstraintsPostApplyPushOutCallback = TFunction<void(const T Dt, const TArray<TPBDCollisionConstraintHandle<T, d>*>&, bool)>;
 
 template <typename T, int d>
-using TCollisionModifierCallback = TFunction<ECollisionModifierResult(const TPBDCollisionConstraintHandle<T, d>*)>;
+using TCollisionModifierCallback = TFunction<ECollisionModifierResult(TPBDCollisionConstraintHandle<T, d>*)>;
 
 /**
  * A container and solver for collision constraints.
