Allow per particle sleep linear and angular velocity threshold.

diff --git a/Engine/Source/Runtime/Engine/Classes/PhysicalMaterials/PhysicalMaterial.h b/Engine/Source/Runtime/Engine/Classes/PhysicalMaterials/PhysicalMaterial.h
--- a/Engine/Source/Runtime/Engine/Classes/PhysicalMaterials/PhysicalMaterial.h
+++ b/Engine/Source/Runtime/Engine/Classes/PhysicalMaterials/PhysicalMaterial.h
@@ -87,6 +87,18 @@ class ENGINE_API UPhysicalMaterial : public UObject
 	UPROPERTY(EditAnywhere, BlueprintReadOnly, Category = PhysicalMaterial, meta=(ClampMin=0))
 	float Density;
 
+	/**  How low the linear velocity can be before solver puts body to sleep. */
+	UPROPERTY(EditAnywhere, BlueprintReadOnly, Category = PhysicalMaterial, meta = (ClampMin = 0))
+	float SleepLinearVelocityThreshold;
+
+	/** How low the angular velocity can be before solver puts body to sleep. */
+	UPROPERTY(EditAnywhere, BlueprintReadOnly, Category = PhysicalMaterial, meta = (ClampMin = 0))
+	float SleepAngularVelocityThreshold;
+
+	/** How many ticks we can be under thresholds for before solver puts body to sleep. */
+	UPROPERTY(EditAnywhere, BlueprintReadOnly, Category = PhysicalMaterial, meta = (ClampMin = 0))
+	int32 SleepCounterThreshold;
+
 	/** 
 	 *	Used to adjust the way that mass increases as objects get larger. This is applied to the mass as calculated based on a 'solid' object. 
 	 *	In actuality, larger objects do not tend to be solid, and become more like 'shells' (e.g. a car is not a solid piece of metal).
diff --git a/Engine/Source/Runtime/Engine/Private/PhysicsEngine/Experimental/PhysInterface_Chaos.cpp b/Engine/Source/Runtime/Engine/Private/PhysicsEngine/Experimental/PhysInterface_Chaos.cpp
--- a/Engine/Source/Runtime/Engine/Private/PhysicsEngine/Experimental/PhysInterface_Chaos.cpp
+++ b/Engine/Source/Runtime/Engine/Private/PhysicsEngine/Experimental/PhysInterface_Chaos.cpp
@@ -316,6 +316,9 @@ void FPhysInterface_Chaos::UpdateMaterial(FPhysicsMaterialHandle& InHandle, UPhy
 	{
 		Material->Friction = InMaterial->Friction;
 		Material->Restitution = InMaterial->Restitution;
+		Material->SleepingLinearThreshold = InMaterial->SleepLinearVelocityThreshold;
+		Material->SleepingAngularThreshold = InMaterial->SleepAngularVelocityThreshold;
+		Material->SleepCounterThreshold = InMaterial->SleepCounterThreshold;
 	}
 
 	Chaos::FPhysicalMaterialManager::Get().UpdateMaterial(InHandle);
diff --git a/Engine/Source/Runtime/Engine/Private/PhysicsEngine/PhysicalMaterial.cpp b/Engine/Source/Runtime/Engine/Private/PhysicsEngine/PhysicalMaterial.cpp
--- a/Engine/Source/Runtime/Engine/Private/PhysicsEngine/PhysicalMaterial.cpp
+++ b/Engine/Source/Runtime/Engine/Private/PhysicsEngine/PhysicalMaterial.cpp
@@ -28,6 +28,9 @@ UPhysicalMaterial::UPhysicalMaterial(const FObjectInitializer& ObjectInitializer
 	Restitution = 0.3f;
 	RaiseMassToPower = 0.75f;
 	Density = 1.0f;
+	SleepLinearVelocityThreshold = 0.001f;
+	SleepAngularVelocityThreshold = 0.0087f;
+	SleepCounterThreshold = 0;
 	DestructibleDamageThresholdScale = 1.0f;
 	TireFrictionScale = 1.0f;
 	bOverrideFrictionCombineMode = false;
diff --git a/Engine/Source/Runtime/Experimental/Chaos/Private/Chaos/PBDConstraintGraph.cpp b/Engine/Source/Runtime/Experimental/Chaos/Private/Chaos/PBDConstraintGraph.cpp
--- a/Engine/Source/Runtime/Experimental/Chaos/Private/Chaos/PBDConstraintGraph.cpp
+++ b/Engine/Source/Runtime/Experimental/Chaos/Private/Chaos/PBDConstraintGraph.cpp
@@ -1,7 +1,6 @@
 // Copyright Epic Games, Inc. All Rights Reserved.
 
 #include "Chaos/PBDConstraintGraph.h"
-
 #include "ProfilingDebugging/ScopedTimers.h"
 #include "ChaosStats.h"
 #include "ChaosLog.h"
@@ -592,7 +591,7 @@ void FPBDConstraintGraph::ComputeIsland(const int32 InNode, const int32 Island,
 	}
 }
 
-bool FPBDConstraintGraph::SleepInactive(const int32 Island, const TArrayCollectionArray<TSerializablePtr<FChaosPhysicsMaterial>>& PerParticleMaterialAttributes)
+bool FPBDConstraintGraph::SleepInactive(const int32 Island, const TArrayCollectionArray<TSerializablePtr<FChaosPhysicsMaterial>>& PerParticleMaterialAttributes, THandleArray<FChaosPhysicsMaterial>& SolverPhysicsMaterials)
 {
 	FReal LinearSleepingThreshold = FLT_MAX;
 	FReal AngularSleepingThreshold = FLT_MAX;
@@ -625,13 +624,32 @@ bool FPBDConstraintGraph::SleepInactive(const int32 Island, const TArrayCollecti
 				M += PBDRigid->M();
 				V += PBDRigid->V() * PBDRigid->M();
 
+				bool bThresholdsSet = false;
 				if (TSerializablePtr<FChaosPhysicsMaterial> PhysicsMaterial = Particle->AuxilaryValue(PerParticleMaterialAttributes))
 				{
 					LinearSleepingThreshold = FMath::Min(LinearSleepingThreshold, PhysicsMaterial->SleepingLinearThreshold);
 					AngularSleepingThreshold = FMath::Min(AngularSleepingThreshold, PhysicsMaterial->SleepingAngularThreshold);
 					SleepCounterThreshold = FMath::Max(SleepCounterThreshold, PhysicsMaterial->SleepCounterThreshold);
+					bThresholdsSet = true;
+				}
+				else if (PBDRigid->ShapesArray().Num())
+				{
+					if (TPerShapeData<FReal, 3>* PerShapeData = PBDRigid->ShapesArray()[0].Get())
+					{
+						if (PerShapeData->Materials.Num())
+						{
+							if (FChaosPhysicsMaterial* Material = SolverPhysicsMaterials.Get(PBDRigid->ShapesArray()[0].Get()->Materials[0].InnerHandle))
+							{
+								LinearSleepingThreshold = FMath::Min(LinearSleepingThreshold, Material->SleepingLinearThreshold);
+								AngularSleepingThreshold = FMath::Min(AngularSleepingThreshold, Material->SleepingAngularThreshold);
+								SleepCounterThreshold = FMath::Max(SleepCounterThreshold, Material->SleepCounterThreshold);
+								bThresholdsSet = true;
+							}
+						}
+					}
 				}
-				else
+
+				if (!bThresholdsSet)
 				{
 					LinearSleepingThreshold = FMath::Min(LinearSleepingThreshold, ChaosSolverCollisionDefaultLinearSleepThresholdCVar);
 					AngularSleepingThreshold = FMath::Min(AngularSleepingThreshold, ChaosSolverCollisionDefaultAngularSleepThresholdCVar);
diff --git a/Engine/Source/Runtime/Experimental/Chaos/Private/Chaos/PBDRigidsEvolution.cpp b/Engine/Source/Runtime/Experimental/Chaos/Private/Chaos/PBDRigidsEvolution.cpp
--- a/Engine/Source/Runtime/Experimental/Chaos/Private/Chaos/PBDRigidsEvolution.cpp
+++ b/Engine/Source/Runtime/Experimental/Chaos/Private/Chaos/PBDRigidsEvolution.cpp
@@ -125,8 +125,9 @@ namespace Chaos
 		}
 	};
 
-	FPBDRigidsEvolutionBase::FPBDRigidsEvolutionBase(TPBDRigidsSOAs<FReal, 3>& InParticles, int32 InNumIterations, int32 InNumPushOutIterations, bool InIsSingleThreaded)
-		: Particles(InParticles)
+	FPBDRigidsEvolutionBase::FPBDRigidsEvolutionBase(TPBDRigidsSOAs<FReal, 3>& InParticles, THandleArray<FChaosPhysicsMaterial>& InSolverPhysicsMaterials, int32 InNumIterations, int32 InNumPushOutIterations, bool InIsSingleThreaded)
+	    : Particles(InParticles)
+		, SolverPhysicsMaterials(InSolverPhysicsMaterials)
 		, bExternalReady(false)
 		, bIsSingleThreaded(InIsSingleThreaded)
 		, NumIterations(InNumIterations)
diff --git a/Engine/Source/Runtime/Experimental/Chaos/Private/Chaos/PBDRigidsEvolutionGBF.cpp b/Engine/Source/Runtime/Experimental/Chaos/Private/Chaos/PBDRigidsEvolutionGBF.cpp
--- a/Engine/Source/Runtime/Experimental/Chaos/Private/Chaos/PBDRigidsEvolutionGBF.cpp
+++ b/Engine/Source/Runtime/Experimental/Chaos/Private/Chaos/PBDRigidsEvolutionGBF.cpp
@@ -392,7 +392,7 @@ void FPBDRigidsEvolutionGBF::AdvanceOneTimeStep(const FReal Dt, const FReal Step
 			}
 
 			// Turn off if not moving
-			SleepedIslands[Island] = GetConstraintGraph().SleepInactive(Island, PhysicsMaterials);
+			SleepedIslands[Island] = GetConstraintGraph().SleepInactive(Island, PhysicsMaterials, SolverPhysicsMaterials);
 		});
 	}
 
@@ -421,8 +421,8 @@ void FPBDRigidsEvolutionGBF::AdvanceOneTimeStep(const FReal Dt, const FReal Step
 	ParticleUpdatePosition(Particles.GetActiveParticlesView(), Dt);
 }
 
-FPBDRigidsEvolutionGBF::FPBDRigidsEvolutionGBF(TPBDRigidsSOAs<FReal, 3>& InParticles, int32 InNumIterations, int32 InNumPushoutIterations, bool InIsSingleThreaded)
-	: Base(InParticles, InNumIterations, InNumPushoutIterations, InIsSingleThreaded)
+FPBDRigidsEvolutionGBF::FPBDRigidsEvolutionGBF(TPBDRigidsSOAs<FReal, 3>& InParticles, THandleArray<FChaosPhysicsMaterial>& SolverPhysicsMaterials, int32 InNumIterations, int32 InNumPushoutIterations, bool InIsSingleThreaded)
+	: Base(InParticles, SolverPhysicsMaterials, InNumIterations, InNumPushoutIterations, InIsSingleThreaded)
 	, Clustering(*this, Particles.GetClusteredParticles())
 	, CollisionConstraints(InParticles, Collided, PhysicsMaterials, DefaultNumPairIterations, DefaultNumPushOutPairIterations)
 	, CollisionRule(CollisionConstraints)
diff --git a/Engine/Source/Runtime/Experimental/Chaos/Public/Chaos/PBDConstraintGraph.h b/Engine/Source/Runtime/Experimental/Chaos/Public/Chaos/PBDConstraintGraph.h
--- a/Engine/Source/Runtime/Experimental/Chaos/Public/Chaos/PBDConstraintGraph.h
+++ b/Engine/Source/Runtime/Experimental/Chaos/Public/Chaos/PBDConstraintGraph.h
@@ -91,7 +91,7 @@ namespace Chaos
 		/**
 		 * Put particles in inactive islands to sleep.
 		 */
-		bool SleepInactive(const int32 Island, const TArrayCollectionArray<TSerializablePtr<FChaosPhysicsMaterial>>& PerParticleMaterialAttributes);
+		bool SleepInactive(const int32 Island, const TArrayCollectionArray<TSerializablePtr<FChaosPhysicsMaterial>>& PerParticleMaterialAttributes, THandleArray<FChaosPhysicsMaterial>& SolverPhysicsMaterials);
 
 		/**
 		 * Wake all particles in an Island.
diff --git a/Engine/Source/Runtime/Experimental/Chaos/Public/Chaos/PBDRigidsEvolution.h b/Engine/Source/Runtime/Experimental/Chaos/Public/Chaos/PBDRigidsEvolution.h
--- a/Engine/Source/Runtime/Experimental/Chaos/Public/Chaos/PBDRigidsEvolution.h
+++ b/Engine/Source/Runtime/Experimental/Chaos/Public/Chaos/PBDRigidsEvolution.h
@@ -226,7 +226,7 @@ class FPBDRigidsEvolutionBase
 
 	friend void ChaosTest::TestPendingSpatialDataHandlePointerConflict();
 
-	CHAOS_API FPBDRigidsEvolutionBase(TPBDRigidsSOAs<FReal, 3>& InParticles, int32 InNumIterations = 1, int32 InNumPushOutIterations = 1, bool InIsSingleThreaded = false);
+	CHAOS_API FPBDRigidsEvolutionBase(TPBDRigidsSOAs<FReal, 3>& InParticles, THandleArray<FChaosPhysicsMaterial>& InSolverPhysicsMaterials, int32 InNumIterations = 1, int32 InNumPushOutIterations = 1, bool InIsSingleThreaded = false);
 	CHAOS_API virtual ~FPBDRigidsEvolutionBase();
 
 	CHAOS_API TArray<TGeometryParticleHandle<FReal, 3>*> CreateStaticParticles(int32 NumParticles, const FUniqueIdx* ExistingIndices = nullptr, const TGeometryParticleParameters<FReal, 3>& Params = TGeometryParticleParameters<FReal, 3>())
@@ -701,6 +701,7 @@ protected:
 	TArrayCollectionArray<bool> Collided;
 
 	TPBDRigidsSOAs<FReal, 3>& Particles;
+	THandleArray<FChaosPhysicsMaterial>& SolverPhysicsMaterials;
 	TUniquePtr<FAccelerationStructure> InternalAcceleration;
 	TUniquePtr<FAccelerationStructure> AsyncInternalAcceleration;
 	TUniquePtr<FAccelerationStructure> AsyncExternalAcceleration;
diff --git a/Engine/Source/Runtime/Experimental/Chaos/Public/Chaos/PBDRigidsEvolutionGBF.h b/Engine/Source/Runtime/Experimental/Chaos/Public/Chaos/PBDRigidsEvolutionGBF.h
--- a/Engine/Source/Runtime/Experimental/Chaos/Public/Chaos/PBDRigidsEvolutionGBF.h
+++ b/Engine/Source/Runtime/Experimental/Chaos/Public/Chaos/PBDRigidsEvolutionGBF.h
@@ -48,7 +48,7 @@ namespace Chaos
 		static constexpr int32 DefaultNumPushOutIterations = 5;
 		static constexpr int32 DefaultNumPushOutPairIterations = 2;
 
-		CHAOS_API FPBDRigidsEvolutionGBF(TPBDRigidsSOAs<FReal, 3>& InParticles, int32 InNumIterations = DefaultNumIterations, int32 InNumPushoutIterations = DefaultNumPushOutIterations, bool InIsSingleThreaded = false);
+		CHAOS_API FPBDRigidsEvolutionGBF(TPBDRigidsSOAs<FReal, 3>& InParticles, THandleArray<FChaosPhysicsMaterial>& SolverPhysicsMaterials, int32 InNumIterations = DefaultNumIterations, int32 InNumPushoutIterations = DefaultNumPushOutIterations, bool InIsSingleThreaded = false);
 		CHAOS_API ~FPBDRigidsEvolutionGBF() {}
 
 		void SetPostIntegrateCallback(const FPBDRigidsEvolutionCallback& Cb)
diff --git a/Engine/Source/Runtime/Experimental/Chaos/Public/Chaos/ParticleHandle.h b/Engine/Source/Runtime/Experimental/Chaos/Public/Chaos/ParticleHandle.h
--- a/Engine/Source/Runtime/Experimental/Chaos/Public/Chaos/ParticleHandle.h
+++ b/Engine/Source/Runtime/Experimental/Chaos/Public/Chaos/ParticleHandle.h
@@ -1691,12 +1691,14 @@ public:
 		CollisionTraceType.Empty(Shapes.Num());
 		ShapeSimData.Empty(Shapes.Num());
 		ShapeQueryData.Empty(Shapes.Num());
+		ShapeMaterials.Empty(Shapes.Num());
 		for (const TUniquePtr<TPerShapeData<T, d>>& ShapePtr : Shapes)
 		{
 			ShapeCollisionDisableFlags.Add(ShapePtr->bDisable);
 			CollisionTraceType.Add(ShapePtr->CollisionTraceType);
 			ShapeSimData.Add(ShapePtr->SimData);
 			ShapeQueryData.Add(ShapePtr->QueryData);
+			ShapeMaterials.Add(ShapePtr->Materials);
 		}
 	}
 
@@ -1714,6 +1716,7 @@ public:
 		CollisionTraceType.Reset();
 		ShapeSimData.Reset();
 		ShapeQueryData.Reset();
+		ShapeMaterials.Reset();
 #if CHAOS_CHECKED
 		DebugName = NAME_None;
 #endif
@@ -1737,12 +1740,14 @@ public:
 			CollisionTraceType.Empty(Shapes.Num());
 			ShapeSimData.Empty(Shapes.Num());
 			ShapeQueryData.Empty(Shapes.Num());
+			ShapeMaterials.Empty(Shapes.Num());
 			for (const TUniquePtr<TPerShapeData<T, d>>& ShapePtr : Shapes)
 			{
 				ShapeCollisionDisableFlags.Add(ShapePtr->bDisable);
 				CollisionTraceType.Add(ShapePtr->CollisionTraceType);
 				ShapeSimData.Add(ShapePtr->SimData);
 				ShapeQueryData.Add(ShapePtr->QueryData);
+				ShapeMaterials.Add(ShapePtr->Materials);
 			}
 	}
 	TVector<T, d> X;
@@ -1756,6 +1761,7 @@ public:
 	TArray<EChaosCollisionTraceFlag> CollisionTraceType;
 	TArray < FCollisionFilterData > ShapeSimData;
 	TArray < FCollisionFilterData > ShapeQueryData;
+	TArray < TArray<FMaterialHandle> > ShapeMaterials;
 #if CHAOS_CHECKED
 	FName DebugName;
 #endif
diff --git a/Engine/Source/Runtime/Experimental/ChaosSolvers/Private/PBDRigidsSolver.cpp b/Engine/Source/Runtime/Experimental/ChaosSolvers/Private/PBDRigidsSolver.cpp
--- a/Engine/Source/Runtime/Experimental/ChaosSolvers/Private/PBDRigidsSolver.cpp
+++ b/Engine/Source/Runtime/Experimental/ChaosSolvers/Private/PBDRigidsSolver.cpp
@@ -189,7 +189,7 @@ namespace Chaos
 		, bHasFloor(true)
 		, bIsFloorAnalytic(false)
 		, FloorHeight(0.f)
-		, MEvolution(new FPBDRigidsEvolution(Particles, ChaosSolverCollisionDefaultIterationsCVar, ChaosSolverCollisionDefaultPushoutIterationsCVar, BufferingModeIn == Chaos::EMultiBufferMode::Single))
+		, MEvolution(new FPBDRigidsEvolution(Particles, SimMaterials, ChaosSolverCollisionDefaultIterationsCVar, ChaosSolverCollisionDefaultPushoutIterationsCVar, BufferingModeIn == Chaos::EMultiBufferMode::Single))
 		, MEventManager(new FEventManager(BufferingModeIn))
 		, MSolverEventFilters(new FSolverEventFilters())
 		, MActiveParticlesBuffer(new FActiveParticlesBuffer(BufferingModeIn, BufferingModeIn == Chaos::EMultiBufferMode::Single))
@@ -518,7 +518,7 @@ namespace Chaos
 		MMaxDeltaTime = 1.f;
 		MMinDeltaTime = 1.e-10f;
 		MMaxSubSteps = 1;
-		MEvolution = TUniquePtr<FPBDRigidsEvolution>(new FPBDRigidsEvolution(Particles, ChaosSolverCollisionDefaultIterationsCVar, ChaosSolverCollisionDefaultPushoutIterationsCVar, BufferMode == EMultiBufferMode::Single)); 
+		MEvolution = TUniquePtr<FPBDRigidsEvolution>(new FPBDRigidsEvolution(Particles, SimMaterials, ChaosSolverCollisionDefaultIterationsCVar, ChaosSolverCollisionDefaultPushoutIterationsCVar, BufferMode == EMultiBufferMode::Single)); 
 
 		FEventDefaults::RegisterSystemEvents(*GetEventManager());
 	}
diff --git a/Engine/Source/Runtime/Experimental/ChaosSolvers/Private/PhysicsProxy/SingleParticlePhysicsProxy.cpp b/Engine/Source/Runtime/Experimental/ChaosSolvers/Private/PhysicsProxy/SingleParticlePhysicsProxy.cpp
--- a/Engine/Source/Runtime/Experimental/ChaosSolvers/Private/PhysicsProxy/SingleParticlePhysicsProxy.cpp
+++ b/Engine/Source/Runtime/Experimental/ChaosSolvers/Private/PhysicsProxy/SingleParticlePhysicsProxy.cpp
@@ -96,6 +96,13 @@ void FSingleParticlePhysicsProxy<Chaos::TGeometryParticle<float, 3>>::PushToPhys
 			}
 		}
 
+		{
+			int32 CurrShape = 0;
+			for (const TUniquePtr<Chaos::TPerShapeData<Chaos::FReal, 3>>& Shape : RigidHandle->ShapesArray())
+			{
+				Shape->Materials = Data->ShapeMaterials[CurrShape++];
+			}
+		}
 	}
 }
 
@@ -206,6 +213,14 @@ void FSingleParticlePhysicsProxy<Chaos::TKinematicGeometryParticle<float, 3>>::P
 			}
 		}
 
+		{
+			int32 CurrShape = 0;
+			for (const TUniquePtr<Chaos::TPerShapeData<Chaos::FReal, 3>>& Shape : RigidHandle->ShapesArray())
+			{
+				Shape->Materials = Data->ShapeMaterials[CurrShape++];
+			}
+		}
+
 
 	}
 }
@@ -397,6 +412,14 @@ void FSingleParticlePhysicsProxy<Chaos::TPBDRigidParticle<float, 3>>::PushToPhys
 			}
 		}
 
+		{
+			int32 CurrShape = 0;
+			for (const TUniquePtr<Chaos::TPerShapeData<Chaos::FReal, 3>>& Shape : RigidHandle->ShapesArray())
+			{
+				Shape->Materials = Data->ShapeMaterials[CurrShape++];
+			}
+		}
+
 		if (bDynamicPropertyUpdated)
 		{
 			if (bInitialized && RigidHandle->ObjectState() == Chaos::EObjectStateType::Sleeping)
diff --git a/Engine/Source/Runtime/PhysicsCore/Private/PhysTestSerializer.cpp b/Engine/Source/Runtime/PhysicsCore/Private/PhysTestSerializer.cpp
--- a/Engine/Source/Runtime/PhysicsCore/Private/PhysTestSerializer.cpp
+++ b/Engine/Source/Runtime/PhysicsCore/Private/PhysTestSerializer.cpp
@@ -339,11 +339,11 @@ void FPhysTestSerializer::CreateChaosData()
 			++Idx;
 		}
 
-		ChaosEvolution = MakeUnique<FPBDRigidsEvolutionGBF>(Particles);
+		ChaosEvolution = MakeUnique<FPBDRigidsEvolutionGBF>(Particles, PhysicalMaterials);
 	}
 	else
 	{
-		ChaosEvolution = MakeUnique<FPBDRigidsEvolutionGBF>(Particles);
+		ChaosEvolution = MakeUnique<FPBDRigidsEvolutionGBF>(Particles, PhysicalMaterials);
 
 		FMemoryReader Ar(Data);
 		FChaosArchive ChaosAr(Ar);
diff --git a/Engine/Source/Runtime/PhysicsCore/Public/PhysTestSerializer.h b/Engine/Source/Runtime/PhysicsCore/Public/PhysTestSerializer.h
--- a/Engine/Source/Runtime/PhysicsCore/Public/PhysTestSerializer.h
+++ b/Engine/Source/Runtime/PhysicsCore/Public/PhysTestSerializer.h
@@ -142,6 +142,7 @@ private:
 
 	TUniquePtr<Chaos::FPBDRigidsEvolutionGBF> ChaosEvolution;
 	Chaos::TPBDRigidsSOAs<float, 3> Particles;
+	Chaos::THandleArray<Chaos::FChaosPhysicsMaterial> PhysicalMaterials;
 	TArray <TUniquePtr<Chaos::TGeometryParticle<float, 3>>> GTParticles;
 
 	TUniquePtr<Chaos::FChaosArchiveContext> ChaosContext;
