Fix GeometryCollectionTest RigidBodies_Initialize_ParticleImplictCollisionGeometry.  Fix level set rasterization issues from geometry collections.  Don't remove "small" triangles from surfaces to be used for level set rasterization.  Add support for a gridded box geometry collection test structure, for a  better approximation for the typical case of level set geometry.  Add support for optionally removing geometricly degenerate triangles from a Chaos::TTriangleMesh.

diff --git a/Engine/Source/Runtime/Experimental/Chaos/Private/Chaos/TriangleMesh.cpp b/Engine/Source/Runtime/Experimental/Chaos/Private/Chaos/TriangleMesh.cpp
--- a/Engine/Source/Runtime/Experimental/Chaos/Private/Chaos/TriangleMesh.cpp
+++ b/Engine/Source/Runtime/Experimental/Chaos/Private/Chaos/TriangleMesh.cpp
@@ -21,9 +21,9 @@ TTriangleMesh<T>::TTriangleMesh()
 {}
 
 template<class T>
-TTriangleMesh<T>::TTriangleMesh(TArray<TVector<int32, 3>>&& Elements, const int32 StartIdx, const int32 EndIdx)
+TTriangleMesh<T>::TTriangleMesh(TArray<TVector<int32, 3>>&& Elements, const int32 StartIdx, const int32 EndIdx, const bool CullDegenerateElements)
 {
-	Init(Elements, StartIdx, EndIdx);
+	Init(Elements, StartIdx, EndIdx, CullDegenerateElements);
 }
 
 template<class T>
@@ -39,40 +39,51 @@ TTriangleMesh<T>::~TTriangleMesh()
 {}
 
 template<class T>
-void TTriangleMesh<T>::Init(TArray<TVector<int32, 3>>&& Elements, const int32 StartIdx, const int32 EndIdx)
+void TTriangleMesh<T>::Init(TArray<TVector<int32, 3>>&& Elements, const int32 StartIdx, const int32 EndIdx, const bool CullDegenerateElements)
 {
 	MElements = MoveTemp(Elements);
 	MStartIdx = 0;
 	MNumIndices = 0;
-	InitHelper(StartIdx, EndIdx);
+	InitHelper(StartIdx, EndIdx, CullDegenerateElements);
 }
 
 template<class T>
-void TTriangleMesh<T>::Init(const TArray<TVector<int32, 3>>& Elements, const int32 StartIdx, const int32 EndIdx)
+void TTriangleMesh<T>::Init(const TArray<TVector<int32, 3>>& Elements, const int32 StartIdx, const int32 EndIdx, const bool CullDegenerateElements)
 {
 	MElements = Elements;
 	MStartIdx = 0;
 	MNumIndices = 0;
-	InitHelper(StartIdx, EndIdx);
+	InitHelper(StartIdx, EndIdx, CullDegenerateElements);
 }
 
 template<class T>
-void TTriangleMesh<T>::InitHelper(const int32 StartIdx, const int32 EndIdx)
+void TTriangleMesh<T>::InitHelper(const int32 StartIdx, const int32 EndIdx, const bool CullDegenerateElements)
 {
 	if (MElements.Num())
 	{
-		MStartIdx = MElements[0][0];
-		int32 MaxIdx = MElements[0][0];
-		for (int i = 0; i < MElements.Num(); ++i)
+		MStartIdx = MElements[MElements.Num()-1][0];
+		int32 MaxIdx = MStartIdx;
+		for (int i = MElements.Num()-1; i >= 0 ; --i)
 		{
 			for (int Axis = 0; Axis < 3; ++Axis)
 			{
 				MStartIdx = FMath::Min(MStartIdx, MElements[i][Axis]);
 				MaxIdx = FMath::Max(MaxIdx, MElements[i][Axis]);
 			}
-			check(MElements[i][0] != MElements[i][1]);
-			check(MElements[i][0] != MElements[i][2]);
-			check(MElements[i][1] != MElements[i][2]);
+			if (CullDegenerateElements)
+			{
+				if (MElements[i][0] == MElements[i][1] ||
+					MElements[i][0] == MElements[i][2] ||
+					MElements[i][1] == MElements[i][2])
+				{
+					// It's possible that the order of the triangles might be important.
+					// RemoveAtSwap() changes the order of the array.  I figure that if
+					// you're up for CullDegenerateElements, then triangle reordering is
+					// fair game.
+					MElements.RemoveAtSwap(i);
+				}
+
+			}
 		}
 		// This assumes vertices are contiguous in the vertex buffer. Assumption is held throughout TTriangleMesh
 		MNumIndices = MaxIdx - MStartIdx + 1;
diff --git a/Engine/Source/Runtime/Experimental/Chaos/Public/Chaos/TriangleMesh.h b/Engine/Source/Runtime/Experimental/Chaos/Public/Chaos/TriangleMesh.h
--- a/Engine/Source/Runtime/Experimental/Chaos/Public/Chaos/TriangleMesh.h
+++ b/Engine/Source/Runtime/Experimental/Chaos/Public/Chaos/TriangleMesh.h
@@ -15,13 +15,19 @@ namespace Chaos
 	{
 	public:
 		TTriangleMesh();
-		TTriangleMesh(TArray<TVector<int32, 3>>&& Elements, const int32 StartIdx = 0, const int32 EndIdx = -1);
+		TTriangleMesh(TArray<TVector<int32, 3>>&& Elements, const int32 StartIdx = 0, const int32 EndIdx = -1, const bool CullDegenerateElements=true);
 		TTriangleMesh(const TTriangleMesh& Other) = delete;
 		TTriangleMesh(TTriangleMesh&& Other);
 		~TTriangleMesh();
 
-		void Init(TArray<TVector<int32, 3>>&& Elements, const int32 StartIdx = 0, const int32 EndIdx = -1);
-		void Init(const TArray<TVector<int32, 3>>& Elements, const int32 StartIdx = 0, const int32 EndIdx = -1);
+		/**
+		 * Initialize the \c TTriangleMesh.
+		 *
+		 *	\p CullDegenerateElements removes faces with degenerate indices, and 
+		 *	will change the order of \c MElements.
+		 */
+		void Init(TArray<TVector<int32, 3>>&& Elements, const int32 StartIdx = 0, const int32 EndIdx = -1, const bool CullDegenerateElements=true);
+		void Init(const TArray<TVector<int32, 3>>& Elements, const int32 StartIdx = 0, const int32 EndIdx = -1, const bool CullDegenerateElements=true);
 
 		/**
 		 * Returns the closed interval of the smallest vertex index used by 
@@ -233,7 +239,7 @@ namespace Chaos
 		}
 
 	private:
-		void InitHelper(const int32 StartIdx, const int32 EndIdx);
+		void InitHelper(const int32 StartIdx, const int32 EndIdx, const bool CullDegenerateElements=true);
 
 		int32 GlobalToLocal(int32 GlobalIdx) const
 		{
diff --git a/Engine/Source/Runtime/Experimental/ChaosSolvers/Private/PhysicsProxy/GeometryCollectionPhysicsProxy.cpp b/Engine/Source/Runtime/Experimental/ChaosSolvers/Private/PhysicsProxy/GeometryCollectionPhysicsProxy.cpp
--- a/Engine/Source/Runtime/Experimental/ChaosSolvers/Private/PhysicsProxy/GeometryCollectionPhysicsProxy.cpp
+++ b/Engine/Source/Runtime/Experimental/ChaosSolvers/Private/PhysicsProxy/GeometryCollectionPhysicsProxy.cpp
@@ -103,41 +103,39 @@ bool IsMultithreaded()
 	return false;
 }
 
-Chaos::TTriangleMesh<float>* CreateTriangleMesh(int32 FaceCount, int32 VertexOffset, int32 StartIndex, const TManagedArray<FVector>& Vertex, const TManagedArray<bool>& Visible, const TManagedArray<FIntVector>& Indices, TSet<int32>& VertsAdded)
+Chaos::TTriangleMesh<float>* CreateTriangleMesh(
+	const int32 FaceCount, 
+	const int32 VertexOffset, 
+	const int32 StartIndex, 
+	const TManagedArray<FVector>& Vertex, 
+	const TManagedArray<bool>& Visible, 
+	const TManagedArray<FIntVector>& Indices, 
+	TSet<int32>& VertsAdded)
 {
 	TArray<Chaos::TVector<int32, 3>> Faces;
 	{
 		Faces.Reserve(FaceCount);
-		for(int j = 0; j < FaceCount; j++)
+		for (int j = 0; j < FaceCount; j++)
 		{
-			if(!Visible[j + StartIndex])
+			if (!Visible[j + StartIndex])
 			{
 				continue;
 			}
 
-			// @todo: This should never happen but seems to so we need to make sure these faces are not counted
-			if(Indices[j + StartIndex].X == Indices[j + StartIndex].Y || Indices[j + StartIndex].Z == Indices[j + StartIndex].Y || Indices[j + StartIndex].X == Indices[j + StartIndex].Z)
-			{
-				continue;
-			}
+			// Note: This function used to cull small triangles.  As one of the purposes 
+			// of the tri mesh this function creates is for level set rasterization, we 
+			// don't want to do that.  Keep the mesh intact, which hopefully is water tight.
 
-			//make sure triangle is not degenerate (above only checks indices, we need to check for co-linear etc...)
-			const Chaos::TVector<float, 3>& X = Vertex[Indices[j + StartIndex].X];
-			const Chaos::TVector<float, 3>& Y = Vertex[Indices[j + StartIndex].Y];
-			const Chaos::TVector<float, 3>& Z = Vertex[Indices[j + StartIndex].Z];
-			const Chaos::TVector<float, 3> Cross = Chaos::TVector<float, 3>::CrossProduct(Z - X, Y - X);
-			if(Cross.SizeSquared() >= 1e-2)
+			const FIntVector& Tri = Indices[j + StartIndex];
+			Faces.Add(Chaos::TVector<int32, 3>(Tri.X, Tri.Y, Tri.Z));
+			for (int Axis = 0; Axis < 3; ++Axis)
 			{
-				Faces.Add(Chaos::TVector<int32, 3>(Indices[j + StartIndex].X, Indices[j + StartIndex].Y, Indices[j + StartIndex].Z));
-				for(int Axis = 0; Axis < 3; ++Axis)
-				{
-					VertsAdded.Add(Indices[j + StartIndex][Axis]);
-				}
+				VertsAdded.Add(Tri[Axis]);
 			}
 		}
 	}
 
-	return new Chaos::TTriangleMesh<float>(MoveTemp(Faces));
+	return new Chaos::TTriangleMesh<float>(MoveTemp(Faces)); // Culls geometrically degenerate faces
 }
 
 TArray<int32> ComputeTransformToGeometryMap(const FGeometryCollection& Collection)
@@ -301,24 +299,9 @@ void PopulateSimulatedParticle(
 			{
 				CollisionParticles->X(VertexIndex) = Simplicial->X(VertexIndex);
 
-
-#if 0 
-				// Ryan - this is crashing for Jack with stack:
-				/*
-				Assertion failed : (Index >= 0)& (Index < ArrayNum)[File:E:\4.25\Engine\Source\Runtime\Core\Public\Containers / Array.h][Line:674] Array index out of bounds : 4 from an array of size 4
-					ChaosDestructionDemoEditor_Core!FDebug::OptionallyLogFormattedEnsureMessageReturningFalseImpl()[e:\4.25\engine\source\runtime\core\private\misc\assertionmacros.cpp:491]
-					ChaosDestructionDemoEditor_ChaosSolvers!<lambda_9b5be2529a6e7fa96c17dd5bd4cd785f>::operator()()[e:\4.25\engine\source\runtime\experimental\chaossolvers\private\physicsproxy\geometrycollectionphysicsproxy.cpp:310]
-					ChaosDestructionDemoEditor_ChaosSolvers!PopulateSimulatedParticle()[e:\4.25\engine\source\runtime\experimental\chaossolvers\private\physicsproxy\geometrycollectionphysicsproxy.cpp:310]
-					ChaosDestructionDemoEditor_ChaosSolvers!FGeometryCollectionPhysicsProxy::InitializeSharedCollisionStructures()[e:\4.25\engine\source\runtime\experimental\chaossolvers\private\physicsproxy\geometrycollectionphysicsproxy.cpp:2825]
-					ChaosDestructionDemoEditor_GeometryCollectionEngine!FDerivedDataGeometryCollectionCooker::Build()[e:\4.25\engine\source\runtime\experimental\geometrycollectionengine\private\geometrycollection\deriveddatageometrycollectioncooker.cpp:37]
-					ChaosDestructionDemoEditor_DerivedDataCache!FDerivedDataCache::FBuildAsyncWorker::DoWork()
-				*/
-				// Removing for now.
-
 				// Make sure the collision particles are at least in the domain 
 				// of the implicit shape.
 				ensure(ImplicitShapeDomain.Contains(CollisionParticles->X(VertexIndex)));
-#endif
 			}
 		}
 
@@ -2621,6 +2604,7 @@ void FGeometryCollectionPhysicsProxy::InitializeSharedCollisionStructures(
 			//  Build the simplicial for the rest collection. This will be used later in the DynamicCollection to 
 			//  populate the collision structures of the simulation. 
 			//
+			check(TriMesh);
 			Chaos::TBVHParticles<float, 3> * Simplicial = 
 				FCollisionStructureManager::NewSimplicial(
 					MassSpaceParticles, 
diff --git a/Engine/Source/Runtime/Experimental/GeometryCollectionSimulationCore/Private/GeometryCollection/GeometryCollectionCollisionStructureManager.cpp b/Engine/Source/Runtime/Experimental/GeometryCollectionSimulationCore/Private/GeometryCollection/GeometryCollectionCollisionStructureManager.cpp
--- a/Engine/Source/Runtime/Experimental/GeometryCollectionSimulationCore/Private/GeometryCollection/GeometryCollectionCollisionStructureManager.cpp
+++ b/Engine/Source/Runtime/Experimental/GeometryCollectionSimulationCore/Private/GeometryCollection/GeometryCollectionCollisionStructureManager.cpp
@@ -249,8 +249,6 @@ FCollisionStructureManager::NewImplicitLevelset(
 	{
 		return nullptr;
 	}
-	HalfExtents *= CollisionObjectReduction / 100.f;
-	const float MinExtent = FMath::Min(HalfExtents[0], FMath::Min(HalfExtents[1], HalfExtents[2]));
 	Chaos::TLevelSet<float, 3>* LevelSet = NewLevelset(ErrorReporter, MeshParticles, TriMesh, CollisionBounds, MinRes, MaxRes, CollisionType);
 	if (LevelSet)
 	{
@@ -266,11 +264,17 @@ FCollisionStructureManager::NewImplicitLevelset(
 					"    percentage filled: %g%%"),
 					DomainVolume, FilledVolume, FilledVolume / DomainVolume * 100.0));
 		}
+		HalfExtents *= CollisionObjectReduction / 100.f;
+		const float MinExtent = FMath::Min(HalfExtents[0], FMath::Min(HalfExtents[1], HalfExtents[2]));
 		if (MinExtent > 0)
 		{
 			LevelSet->Shrink(MinExtent);
 		}
 	}
+	else
+	{
+		ErrorReporter.ReportError(TEXT("Level set rasterization failed."));
+	}
 	return LevelSet;
 }
 
