Chaos : Destruction

- Removed the multiple objects buffered on the Proxy and moved all relevant input data for the physics thread into the FGeometryDynamicCollection

diff --git a/Engine/Source/Runtime/Experimental/Chaos/Public/Chaos/CollisionFilterData.h b/Engine/Source/Runtime/Experimental/Chaos/Public/Chaos/CollisionFilterData.h
--- a/Engine/Source/Runtime/Experimental/Chaos/Public/Chaos/CollisionFilterData.h
+++ b/Engine/Source/Runtime/Experimental/Chaos/Public/Chaos/CollisionFilterData.h
@@ -4,7 +4,7 @@
 #include "ChaosArchive.h"
 #include "UObject/ExternalPhysicsCustomObjectVersion.h"
 
-struct FCollisionFilterData
+struct CHAOS_API FCollisionFilterData
 {
 	uint32 Word0;
 	uint32 Word1;
diff --git a/Engine/Source/Runtime/Experimental/ChaosSolvers/Private/GeometryCollectionProxyData.cpp b/Engine/Source/Runtime/Experimental/ChaosSolvers/Private/GeometryCollectionProxyData.cpp
--- a/Engine/Source/Runtime/Experimental/ChaosSolvers/Private/GeometryCollectionProxyData.cpp
+++ b/Engine/Source/Runtime/Experimental/ChaosSolvers/Private/GeometryCollectionProxyData.cpp
@@ -37,14 +37,39 @@ void FTransformDynamicCollection::Construct()
 
 const FName FGeometryDynamicCollection::ActiveAttribute("Active");
 const FName FGeometryDynamicCollection::CollisionGroupAttribute("CollisionGroup");
+const FName FGeometryDynamicCollection::CollisionMaskAttribute("CollisionMask");
 const FName FGeometryDynamicCollection::DynamicStateAttribute("DynamicState");
-
+const FName FGeometryDynamicCollection::ImplicitsAttribute("Implicits");
+const FName FGeometryDynamicCollection::ShapesQueryDataAttribute("ShapesQueryData");
+const FName FGeometryDynamicCollection::ShapesSimDataAttribute("ShapesSimData");
+const FName FGeometryDynamicCollection::SharedGeometryAttribute("SharedGeometry");
+const FName FGeometryDynamicCollection::SimplicialsAttribute("CollisionParticles");
+const FName FGeometryDynamicCollection::SharedImplicitsAttribute("SharedImplicits");
+const FName FGeometryDynamicCollection::SolverParticleHandlesAttribute("SolverParticleHandles");
+const FName FGeometryDynamicCollection::SolverClusterHandlesAttribute("SolverClusterHandles");
 
 FGeometryDynamicCollection::FGeometryDynamicCollection()
 	: FTransformDynamicCollection()
 {
 	// Transform Group
-	AddExternalAttribute<int32>(FGeometryDynamicCollection::DynamicStateAttribute, FTransformCollection::TransformGroup, DynamicState);
 	AddExternalAttribute<bool>(FGeometryDynamicCollection::ActiveAttribute, FTransformCollection::TransformGroup, Active);
 	AddExternalAttribute<int32>(FGeometryDynamicCollection::CollisionGroupAttribute, FTransformCollection::TransformGroup, CollisionGroup);
+	AddExternalAttribute<int32>(FGeometryDynamicCollection::CollisionMaskAttribute, FTransformCollection::TransformGroup, CollisionMask);
+	AddExternalAttribute("CollisionStructureID", FTransformCollection::TransformGroup, CollisionStructureID);
+	AddExternalAttribute<int32>(FGeometryDynamicCollection::DynamicStateAttribute, FTransformCollection::TransformGroup, DynamicState);
+	AddExternalAttribute(ImplicitsAttribute, FTransformCollection::TransformGroup, Implicits);
+	AddExternalAttribute("InitialAngularVelocity", FTransformCollection::TransformGroup, InitialAngularVelocity);
+	AddExternalAttribute("InitialLinearVelocity", FTransformCollection::TransformGroup, InitialLinearVelocity);
+	AddExternalAttribute("MassToLocal", FTransformCollection::TransformGroup, MassToLocal);
+	//AddExternalAttribute(ShapesQueryDataAttribute, FTransformCollection::TransformGroup, ShapeQueryData);
+	//AddExternalAttribute(ShapesSimDataAttribute, FTransformCollection::TransformGroup, ShapeSimData);
+	//AddExternalAttribute(SharedGeometryAttribute, FTransformCollection::TransformGroup, SharedGeometry);
+	//AddExternalAttribute(SimplicialsAttribute, FTransformCollection::TransformGroup, Simplicials);
+	AddExternalAttribute("SimulatableParticles", FGeometryCollection::TransformGroup, SimulatableParticles);
+	AddExternalAttribute(SolverClusterHandlesAttribute, FTransformCollection::TransformGroup, SolverClusterHandles);
+	AddExternalAttribute("SolverClusterID", FTransformCollection::TransformGroup, SolverClusterID);
+	AddExternalAttribute(SolverParticleHandlesAttribute, FTransformCollection::TransformGroup, SolverParticleHandles);
+
+
+
 }
diff --git a/Engine/Source/Runtime/Experimental/ChaosSolvers/Private/PhysicsProxy/GeometryCollectionPhysicsProxy.cpp b/Engine/Source/Runtime/Experimental/ChaosSolvers/Private/PhysicsProxy/GeometryCollectionPhysicsProxy.cpp
--- a/Engine/Source/Runtime/Experimental/ChaosSolvers/Private/PhysicsProxy/GeometryCollectionPhysicsProxy.cpp
+++ b/Engine/Source/Runtime/Experimental/ChaosSolvers/Private/PhysicsProxy/GeometryCollectionPhysicsProxy.cpp
@@ -49,11 +49,6 @@ FAutoConsoleVariableRef CVarGeometryCollectionDisableGravity(
 	DisableGeometryCollectionGravity,
 	TEXT("Disable gravity for geometry collections"));
 
-FName FGeometryCollectionPhysicsProxy::SimplicialsAttribute("CollisionParticles");
-FName FGeometryCollectionPhysicsProxy::ImplicitsAttribute("Implicits");
-FName FGeometryCollectionPhysicsProxy::SharedImplicitsAttribute("SharedImplicits");
-FName FGeometryCollectionPhysicsProxy::SolverParticleHandlesAttribute("SolverParticleHandles");
-FName FGeometryCollectionPhysicsProxy::SolverClusterHandlesAttribute("SolverClusterHandles");
 
 DEFINE_LOG_CATEGORY_STATIC(UGCC_LOG, Error, All);
 
@@ -354,9 +349,6 @@ FGeometryCollectionPhysicsProxy::FGeometryCollectionPhysicsProxy(
 	, IsObjectDynamic(false)
 	, IsObjectLoading(true)
 
-	, GTDynamicCollection(SimulationParameters.DynamicCollection)
-	//, PTDynamicCollection(nullptr)
-
 	, BaseParticleIndex(INDEX_NONE)
 	, NumParticles(INDEX_NONE)
 	, ProxySimDuration(0.0f)
@@ -367,9 +359,6 @@ FGeometryCollectionPhysicsProxy::FGeometryCollectionPhysicsProxy(
 
 	, LastSyncCountGT(MAX_uint32)
 	, CollisionParticlesPerObjectFraction(CollisionParticlesPerObjectFractionDefault)
-
-	, PhysToGameInterchange(MakeUnique<Chaos::FGuardedTripleBuffer<FGeometryCollectionResults>>())
-	, GameToPhysInterchange(MakeUnique<Chaos::FGuardedTripleBuffer<FGeometryCollectionResults>>())
 {
 	// We rely on a guarded buffer.
 	check(BufferMode == Chaos::EMultiBufferMode::TripleGuarded);
@@ -386,47 +375,37 @@ void FGeometryCollectionPhysicsProxy::Initialize()
 	// Old proxy init
 	check(IsInGameThread());
 
-	PTDynamicCollection.CopyAttribute(*GTDynamicCollection, FTransformCollection::TransformAttribute, FTransformCollection::TransformGroup);
-	PTDynamicCollection.CopyAttribute(*GTDynamicCollection, FTransformCollection::ParentAttribute, FTransformCollection::TransformGroup);
-	PTDynamicCollection.CopyAttribute(*GTDynamicCollection, FTransformCollection::ChildrenAttribute, FTransformCollection::TransformGroup);
-	PTDynamicCollection.CopyAttribute(*GTDynamicCollection, FGeometryCollection::SimulationTypeAttribute, FTransformCollection::TransformGroup);
-	PTDynamicCollection.CopyAttribute(*GTDynamicCollection, FGeometryCollection::StatusFlagsAttribute, FTransformCollection::TransformGroup);
-	PTDynamicCollection.CopyAttribute(*GTDynamicCollection, FGeometryDynamicCollection::DynamicStateAttribute, FTransformCollection::TransformGroup);
-	PTDynamicCollection.CopyAttribute(*GTDynamicCollection, FGeometryDynamicCollection::CollisionGroupAttribute, FTransformCollection::TransformGroup);
-	PTDynamicCollection.CopyAttribute(*GTDynamicCollection, FGeometryDynamicCollection::ActiveAttribute, FTransformCollection::TransformGroup);
 
-	//todo(benn.g): having to spell out which attribute to copy kind of sucks. We don't want a full copy since GT can have a lot of gameplay attributes
-	//question: should we copy dynamic state? 
-	//response: Seems like we will just end up copying everything, otherwise the simulation collection will be missing attributes.
+	//
+	// Game thread initilization. 
+	//
+	//  1) Create a input buffer to store all game thread side data. 
+	//  2) Populate the buffer with the necessary data.
+	//  3) Deep copy the data to the other buffers. 
+	//
+	FGeometryDynamicCollection& DynamicCollection = *GameToPhysInterchange.AccessProducerBuffer();
 
-	/*
-	//todo(ocohen): avoid copies when not multi threaded
-	if(IsMultithreaded())
-	{
-		PTDynamicCollection.CopyAttribute(*GTDynamicCollection, FTransformCollection::TransformAttribute, FGeometryCollection::TransformGroup);
-	}*/
 
-	// Replace with normal funcs in this class
-	//Callbacks->SetUpdateRecordedStateFunction([this](float SolverTime, const TManagedArray<int32> & RigidBodyID, const TManagedArray<FGeometryCollectionBoneNode>& Hierarchy, const FSolverCallbacks::FParticlesType& Particles, const FSolverCallbacks::FCollisionConstraintsType& CollisionRule)
-	//{
-	//	UpdateRecordedState(SolverTime, RigidBodyID, Hierarchy, Particles, CollisionRule);
-	//});
+	InitializeDynamicCollection(DynamicCollection, *Parameters.RestCollection, Parameters);
 
-	//Callbacks->SetCommitRecordedStateFunction([this](FRecordedTransformTrack& InTrack)
-	//{
-	//	InTrack = FRecordedTransformTrack::ProcessRawRecordedData(RecordedTracks);
-	//});
+	NumParticles = GetTransformGroupSize(); // SimulatableParticles.Count(true);
+	BaseParticleIndex = 0; // Are we always zero indexed now?
 
-	// Back to engine for setup from components
+	//	
+	//  Give clients the opportunity to update the parameters before the simualtion is setup. 
+	//
+	Parameters.DynamicCollection = &DynamicCollection;
 	if (InitFunc)
 	{
 		InitFunc(Parameters);
 	}
+	Parameters.DynamicCollection = nullptr; // safty setup to prevent usage of the Parameters.DynamicCollection. 
 
-	PTDynamicCollection.SyncAllGroups(*Parameters.RestCollection);
-	CollisionParticlesPerObjectFraction = 
-		Parameters.CollisionSampleFraction * CollisionParticlesPerObjectFractionDefault;
 
+	//
+	// Collision vertices down sampling validation.  
+	//
+	CollisionParticlesPerObjectFraction = Parameters.CollisionSampleFraction * CollisionParticlesPerObjectFractionDefault;
 	if (ReportHighParticleFraction > 0)
 	{
 		for (const FSharedSimulationSizeSpecificData& Data : Parameters.Shared.SizeSpecificData)
@@ -439,45 +418,38 @@ void FGeometryCollectionPhysicsProxy::Initialize()
 		}
 	}
 
-	// Setup proxy parameters
-
-	// Ryan - I don't know why this happens here.  Previously, Parameters.DynamicCollection
-	// was set to the client level GTDynamicCollection, at least in the case of RigidBodiesFallingUnderGravity().
-	// So why override it to the physics thread dynamic collection???
-	//check(!Parameters.DynamicCollection);
-	Parameters.DynamicCollection = &PTDynamicCollection;
-
-	// Old callbacks init
 
-	//if(Parameters.bClearCache)
-	//{
-	//	if(ResetAnimationCacheCallback)
-	//	{
-	//		ResetAnimationCacheCallback();
-	//	}
-	//}
+#if TODO_REIMPLEMENT_RIGID_CACHING
 
-	CreateDynamicAttributes();
+	// todo : Replace with normal funcs in this class
+	Callbacks->SetUpdateRecordedStateFunction([this](float SolverTime, const TManagedArray<int32>& RigidBodyID, const TManagedArray<FGeometryCollectionBoneNode>& Hierarchy, const FSolverCallbacks::FParticlesType& Particles, const FSolverCallbacks::FCollisionConstraintsType& CollisionRule)
+		{
+			UpdateRecordedState(SolverTime, RigidBodyID, Hierarchy, Particles, CollisionRule);
+		});
 
-	ProxySimDuration = 0.0f;
+	Callbacks->SetCommitRecordedStateFunction([this](FRecordedTransformTrack& InTrack)
+		{
+			InTrack = FRecordedTransformTrack::ProcessRawRecordedData(RecordedTracks);
+		});
 
 	// Old proxy init
 	RecordedTracks.Records.Reset();
-	
-	// Setup double buffer data
-	//Results.Get(0).Transforms.Init(PTDynamicCollection.Transform);
-	//Results.Get(1).Transforms.Init(PTDynamicCollection.Transform);
-	
-	LastSyncCountGT = 0;
 
-	// Initialize global transforms
+	if (Parameters.bClearCache)
+	{
+		if (ResetAnimationCacheCallback)
+		{
+			ResetAnimationCacheCallback();
+		}
+	}
 
-	TArray<FMatrix> &TmpGlobalTransforms = GameToPhysInterchange->AccessProducerBuffer()->GlobalTransforms;		
-	//TArray<FMatrix> &TmpGlobalTransforms = Results.Get(0).GlobalTransforms;
-	GeometryCollectionAlgo::GlobalMatrices(PTDynamicCollection.Transform, PTDynamicCollection.Parent, TmpGlobalTransforms);
-	//Results.Get(1).GlobalTransforms = TmpGlobalTransforms;
+	ProxySimDuration = T(0);
+	LastSyncCountGT = 0;
 
 
+#endif
+
+	/*
 	// Initialize data for faster bound calculations
 	// precompute data used for bounds calculation
 	{
@@ -510,15 +482,153 @@ void FGeometryCollectionPhysicsProxy::Initialize()
 		//Results.Get(0).WorldBounds = FBoxSphereBounds(BoundingBox);
 		//Results.Get(1).WorldBounds = FBoxSphereBounds(BoundingBox);
 	}
+	*/
+	GameToPhysInterchange.FlipProducer();
+	FGeometryDynamicCollection& ProducerCollection = *GameToPhysInterchange.AccessProducerBuffer();
+	ProducerCollection.CopyMatchingAttributesFrom(DynamicCollection);
+}
 
-	//TArray<FTransform> Transform;
-	//GeometryCollectionAlgo::GlobalMatrices(DynamicCollection->Transform, DynamicCollection->Parent, Transform);
-	//check(DynamicCollection->Transform.Num() == Transform.Num());
+void FGeometryCollectionPhysicsProxy::InitializeDynamicCollection(FGeometryDynamicCollection& DynamicCollection, const FGeometryCollection& RestCollection, const FSimulationParameters& Params)
+{
+	// 
+	// This function will use the rest collection to populate the dynamic collection. 
+	//
 
-	NumParticles = GetTransformGroupSize(); // SimulatableParticles.Count(true);
-	BaseParticleIndex = 0; // Are we always zero indexed now?
+	DynamicCollection.CopyMatchingAttributesFrom(RestCollection);
+
+	
+	DynamicCollection.SolverClusterID.Fill(nullptr); 
+
+	//
+	// User defined initial velocities need to be populated. 
+	//
+	{
+		if (Params.InitialVelocityType == EInitialVelocityTypeEnum::Chaos_Initial_Velocity_User_Defined)
+		{
+			DynamicCollection.InitialLinearVelocity.Fill(Params.InitialLinearVelocity);
+			DynamicCollection.InitialAngularVelocity.Fill(Params.InitialAngularVelocity);
+		}
+	}
+
+
+	//
+	// The RestCollection may have been authored with TUniquePtr<FImplicitObject> under the name 
+	// FGeometryDynamicCollection::ImplicitsAttribute, as they originally were.  As the particle system now uses TSharedPtr,
+	// we're now authoring TSharedPtr<FImplicitObject> under the name FGeometryDynamicCollection::SharedImplicitsAttribute.
+	{
+		FGeometryCollection& NCRestCollection = const_cast<FGeometryCollection&>(RestCollection);
+		TManagedArray<TUniquePtr<Chaos::FImplicitObject>>* UniqueRestImplicits =
+			NCRestCollection.FindAttribute<TUniquePtr<Chaos::FImplicitObject>>(
+				FGeometryDynamicCollection::ImplicitsAttribute, FTransformCollection::TransformGroup);
+		TManagedArray<TSharedPtr<Chaos::FImplicitObject, ESPMode::ThreadSafe>>* SharedRestImplicits =
+			NCRestCollection.FindAttribute<TSharedPtr<Chaos::FImplicitObject, ESPMode::ThreadSafe>>(
+				FGeometryDynamicCollection::FGeometryDynamicCollection::SharedImplicitsAttribute, FTransformCollection::TransformGroup);
+
+		// If we have unique pointers to geometry, we need to convert them to shared pointers.
+		TManagedArray<TSharedPtr<Chaos::FImplicitObject, ESPMode::ThreadSafe>>* RestCollectionSharedImplicits = nullptr;
+		if (!SharedRestImplicits)
+		{
+			RestCollectionSharedImplicits =
+				&NCRestCollection.AddAttribute<TSharedPtr<Chaos::FImplicitObject, ESPMode::ThreadSafe>>(
+					FGeometryDynamicCollection::SharedImplicitsAttribute, FTransformCollection::TransformGroup);
+		}
+
+		for (int32 Index = DynamicCollection.NumElements(FTransformCollection::TransformGroup) - 1; 0 <= Index; Index--)
+		{
+			//Implicits[Index] = Chaos::MakeSerializable(RestImplicits[Index]);
+			if (SharedRestImplicits)
+			{
+				DynamicCollection.Implicits[Index] = (*SharedRestImplicits)[Index];
+			}
+			else if (UniqueRestImplicits && RestCollectionSharedImplicits) // RestCollectionSharedImplicits is redundant, but static analysis build complains without it
+			{
+				// Steal the memory from the unique ptr.
+				DynamicCollection.Implicits[Index] = TSharedPtr<Chaos::FImplicitObject, ESPMode::ThreadSafe>((*UniqueRestImplicits)[Index].Release());
+				// Update the rest collection with a copy of the shared pointer.
+				(*RestCollectionSharedImplicits)[Index] = DynamicCollection.Implicits[Index];
+			}
+		}
+
+		// Clean up now invalid unique pointers.
+		if (RestCollectionSharedImplicits)
+		{
+			NCRestCollection.RemoveAttribute(FGeometryDynamicCollection::ImplicitsAttribute, FTransformCollection::TransformGroup);
+		}
+	}
+
+
+	// process simplicials
+	{
+		/*
+		if (RestCollection.HasAttribute(DynamicCollection.SimplicialsAttribute, FTransformCollection::TransformGroup))
+		{
+			const auto& RestSimplicials = RestCollection.GetAttribute<TUniquePtr<FSimplicial>>(
+				DynamicCollection.SimplicialsAttribute, FTransformCollection::TransformGroup);
+			for (int32 Index = DynamicCollection.NumElements(FTransformCollection::TransformGroup) - 1; 0 <= Index; Index--)
+			{
+				DynamicCollection.Simplicials[Index] = TSharedPtr<FSimplicial>(RestSimplicials[Index] ? RestSimplicials[Index]->NewCopy() : nullptr);
+			}
+		}
+		else
+		{
+			for (int32 Index = DynamicCollection.NumElements(FTransformCollection::TransformGroup) - 1; 0 <= Index; Index--)
+			{
+				DynamicCollection.Simplicials[Index] = TSharedPtr<FSimplicial>(nullptr);
+			}
+		}
+		*/
+	}
+
+
+	// Process Activity
+	{
+		const int32 NumTransforms = DynamicCollection.SimulatableParticles.Num();
+		if (RestCollection.HasAttribute(FGeometryCollection::SimulatableParticlesAttribute, FTransformCollection::TransformGroup))
+		{
+			// When the rest collection has been pre configured with simulation data, use that to determine it ability to simulate. 
+			const TManagedArray<bool>& ActiveRestParticles =
+				RestCollection.GetAttribute<bool>(
+					FGeometryCollection::SimulatableParticlesAttribute, FTransformCollection::TransformGroup);
+			check(ActiveRestParticles.Num() == DynamicCollection.Active.Num());
+			check(ActiveRestParticles.Num() == DynamicCollection.SimulatableParticles.Num());
+			for (int32 TransformIdx = 0; TransformIdx < NumTransforms; TransformIdx++)
+			{
+				DynamicCollection.SimulatableParticles[TransformIdx] = DynamicCollection.Active[TransformIdx] && ActiveRestParticles[TransformIdx];
+			}
+		}
+		else
+		{
+			// If no simulation data is available then default to the simulation of just the rigid geometry.
+			for (int32 TransformIdx = 0; TransformIdx < NumTransforms; TransformIdx++)
+			{
+				if (DynamicCollection.Children[TransformIdx].Num())
+				{
+					DynamicCollection.SimulatableParticles[TransformIdx] = false;
+				}
+				else
+				{
+					DynamicCollection.SimulatableParticles[TransformIdx] = DynamicCollection.Active[TransformIdx];
+				}
+			}
+		}
+	}
+
+	// Process Transforms
+	{
+		//TArray<FTransform> Transform;
+		//GeometryCollectionAlgo::GlobalMatrices(DynamicCollection->Transform, DynamicCollection->Parent, Transform);
+		//check(DynamicCollection->Transform.Num() == Transform.Num());
+
+		// Initialize global transforms
+		//TArray<FMatrix>& TmpGlobalTransforms = GameToPhysInterchange.AccessProducerBuffer()->GlobalTransforms;
+		//TArray<FMatrix> &TmpGlobalTransforms = Results.Get(0).GlobalTransforms;
+		//GeometryCollectionAlgo::GlobalMatrices(PTDynamicCollection.Transform, PTDynamicCollection.Parent, TmpGlobalTransforms);
+		//Results.Get(1).GlobalTransforms = TmpGlobalTransforms;
+	}
 }
 
+
+
 int32 ReportTooManyChildrenNum = -1;
 FAutoConsoleVariableRef CVarReportTooManyChildrenNum(TEXT("p.ReportTooManyChildrenNum"), ReportTooManyChildrenNum, TEXT("Issue warning if more than this many children exist in a single cluster"));
 
@@ -526,6 +636,9 @@ void FGeometryCollectionPhysicsProxy::InitializeBodiesPT(
 	Chaos::FPBDRigidsSolver* RigidsSolver,
 	Chaos::FPBDRigidsSolver::FParticlesType& Particles)
 {
+	/*
+	@todo break_everything
+
 	const FGeometryCollection* RestCollection = Parameters.RestCollection;
 	FGeometryDynamicCollection* DynamicCollection = Parameters.DynamicCollection;
 
@@ -823,7 +936,7 @@ void FGeometryCollectionPhysicsProxy::InitializeBodiesPT(
 #endif // TODO_REIMPLEMENT_RIGID_CACHING
 
 	} // end if simulating...
-
+	*/
 }
 
 int32 ReportNoLevelsetCluster = 0;
@@ -845,17 +958,19 @@ FGeometryCollectionPhysicsProxy::BuildClusters(
 	check(ChildHandles.Num() != 0);
 
 	check(Parameters.DynamicCollection);
-	FGeometryDynamicCollection* DynamicCollection = Parameters.DynamicCollection;
-	TManagedArray<int32>& DynamicState = DynamicCollection->DynamicState;
-	TManagedArray<int32>& ParentIndex = DynamicCollection->Parent;
-	TManagedArray<TSet<int32>>& Children = DynamicCollection->Children;
-	TManagedArray<FTransform>& Transform = DynamicCollection->Transform;
+	FGeometryDynamicCollection& DynamicCollection = *GameToPhysInterchange.AccessProducerBuffer();
+	TManagedArray<int32>& DynamicState = DynamicCollection.DynamicState;
+	TManagedArray<int32>& ParentIndex = DynamicCollection.Parent;
+	TManagedArray<TSet<int32>>& Children = DynamicCollection.Children;
+	TManagedArray<FTransform>& Transform = DynamicCollection.Transform;
+	TManagedArray<FTransform>& MassToLocal = DynamicCollection.MassToLocal;
+	//TManagedArray<TSharedPtr<FCollisionStructureManager::FSimplicial> >& Simplicials = DynamicCollection.Simplicials;
+	TManagedArray<TSharedPtr<Chaos::FImplicitObject, ESPMode::ThreadSafe>>& Implicits = DynamicCollection.Implicits;
+	TManagedArray<Chaos::TPBDRigidClusteredParticleHandle<float, 3>*>& SolverClusterHandles = DynamicCollection.SolverClusterHandles;
 
 	//If we are a root particle use the world transform, otherwise set the relative transform
-	const FTransform CollectionSpaceTransform = 
-		GeometryCollectionAlgo::GlobalMatrix(Transform, ParentIndex, CollectionClusterIndex);
-	const Chaos::TRigidTransform<float, 3> ParticleTM = 
-		MassToLocal[CollectionClusterIndex] * CollectionSpaceTransform * Parameters.WorldTransform;
+	const FTransform CollectionSpaceTransform = GeometryCollectionAlgo::GlobalMatrix(Transform, ParentIndex, CollectionClusterIndex);
+	const Chaos::TRigidTransform<float, 3> ParticleTM = MassToLocal[CollectionClusterIndex] * CollectionSpaceTransform * Parameters.WorldTransform;
 
 	//create new cluster particle
 	//The reason we need to pass in a mass orientation override is as follows:
@@ -865,7 +980,7 @@ FGeometryCollectionPhysicsProxy::BuildClusters(
 	//Because of this we need to know how to rotate collision particles and geometry to match with original computation. If it was just geometry we could transform it before passing, but we need collision particles as well
 	Chaos::FClusterCreationParameters<float> ClusterCreationParameters = ClusterParameters;
 	ClusterCreationParameters.bGenerateConnectionGraph = false;
-	ClusterCreationParameters.CollisionParticles = Simplicials[CollectionClusterIndex].Release();
+	// fix... ClusterCreationParameters.CollisionParticles = Simplicials[CollectionClusterIndex];
 	ClusterCreationParameters.ConnectionMethod = Parameters.ClusterConnectionMethod;
 	if (ClusterCreationParameters.CollisionParticles)
 	{
@@ -948,7 +1063,7 @@ FGeometryCollectionPhysicsProxy::BuildClusters(
 	GetSolver()->GetEvolution()->SetPhysicsMaterial(Parent, Parameters.PhysicalMaterial);
 
 	const FTransform ParentTransform = 
-		GeometryCollectionAlgo::GlobalMatrix(DynamicCollection->Transform, DynamicCollection->Parent, CollectionClusterIndex);
+		GeometryCollectionAlgo::GlobalMatrix(DynamicCollection.Transform, DynamicCollection.Parent, CollectionClusterIndex);
 
 	int32 MinCollisionGroup = INT_MAX;
 	for(int32 Idx=0; Idx < ChildHandles.Num(); Idx++)
@@ -1036,142 +1151,22 @@ void FGeometryCollectionPhysicsProxy::ContiguousIndices(TArray<ContextIndex>& Ar
 	}
 }
 
-int32 FGeometryCollectionPhysicsProxy::CalculateHierarchyLevel(
-	const FGeometryDynamicCollection* GeometryCollection, 
-	int32 TransformIndex) const
+int32 FGeometryCollectionPhysicsProxy::CalculateHierarchyLevel(const FGeometryDynamicCollection& GeometryCollection, int32 TransformIndex) const
 {
 	int32 Level = 0;
-	while (GeometryCollection->Parent[TransformIndex] != -1)
+	while (GeometryCollection.Parent[TransformIndex] != -1)
 	{
-		TransformIndex = GeometryCollection->Parent[TransformIndex];
+		TransformIndex = GeometryCollection.Parent[TransformIndex];
 		Level++;
 	}
 	return Level;
 }
 
-void FGeometryCollectionPhysicsProxy::CreateDynamicAttributes()
-{
-	const FGeometryCollection& RestCollection = *Parameters.RestCollection;
-	FGeometryDynamicCollection* DynamicCollection = Parameters.DynamicCollection;
-
-	DynamicCollection->AddExternalAttribute("SimulatableParticles", FGeometryCollection::TransformGroup, SimulatableParticles);
-	DynamicCollection->AddExternalAttribute("SolverClusterID", FTransformCollection::TransformGroup, SolverClusterID);
-	DynamicCollection->AddExternalAttribute("MassToLocal", FTransformCollection::TransformGroup, MassToLocal);
-	DynamicCollection->CopyAttribute(RestCollection, "MassToLocal", FTransformCollection::TransformGroup);
-	DynamicCollection->AddExternalAttribute("CollisionStructureID", FTransformCollection::TransformGroup, CollisionStructureID);
-
-	// TODO: Move these to the geometry group?
-	DynamicCollection->AddExternalAttribute(SimplicialsAttribute, FTransformCollection::TransformGroup, Simplicials);
-	DynamicCollection->AddExternalAttribute(ImplicitsAttribute, FTransformCollection::TransformGroup, Implicits);
-
-	// We don't have access to rigid body ids anymore, so we aren't using this currently.
-	// There is still a bunch of legacy code in this file that builds against the RigidBodyID
-	// variable, so it's still defined in the header.
-	//DynamicCollection->AddExternalAttribute("RigidBodyID", FTransformCollection::TransformGroup, RigidBodyID);
-
-	// Holds Chaos::TPBDGeometryCollectionParticleHandle<float, 3>* for all nodes
-	DynamicCollection->AddExternalAttribute(SolverParticleHandlesAttribute, FTransformCollection::TransformGroup, SolverParticleHandles);
-	// Holds Chaos::TPBDRigidClusteredParticleHandle<float, 3>* for cluster parent nodes
-	DynamicCollection->AddExternalAttribute(SolverClusterHandlesAttribute, FTransformCollection::TransformGroup, SolverClusterHandles);
-
-	if (Parameters.InitialVelocityType == EInitialVelocityTypeEnum::Chaos_Initial_Velocity_User_Defined)
-	{
-		DynamicCollection->AddExternalAttribute("InitialLinearVelocity", FTransformCollection::TransformGroup, InitialLinearVelocity);
-		DynamicCollection->AddExternalAttribute("InitialAngularVelocity", FTransformCollection::TransformGroup, InitialAngularVelocity);
-		InitialLinearVelocity.Fill(Parameters.InitialLinearVelocity);
-		InitialAngularVelocity.Fill(Parameters.InitialAngularVelocity);
-	}
-
-	// The RestCollection may have been authored with TUniquePtr<FImplicitObject> under the name 
-	// ImplicitsAttribute, as they originally were.  As the particle system now uses TSharedPtr,
-	// we're now authoring TSharedPtr<FImplicitObject> under the name SharedImplicitsAttribute.
-	FGeometryCollection& NCRestCollection = const_cast<FGeometryCollection&>(*Parameters.RestCollection);
-	TManagedArray<TUniquePtr<Chaos::FImplicitObject>>* UniqueRestImplicits = 
-		NCRestCollection.FindAttribute<TUniquePtr<Chaos::FImplicitObject>>(
-			ImplicitsAttribute, FTransformCollection::TransformGroup);
-	TManagedArray<TSharedPtr<Chaos::FImplicitObject, ESPMode::ThreadSafe>>* SharedRestImplicits = 
-		NCRestCollection.FindAttribute<TSharedPtr<Chaos::FImplicitObject, ESPMode::ThreadSafe>>(
-			SharedImplicitsAttribute, FTransformCollection::TransformGroup);
-	
-	// If we have unique pointers to geometry, we need to convert them to shared pointers.
-	TManagedArray<TSharedPtr<Chaos::FImplicitObject, ESPMode::ThreadSafe>>* RestCollectionSharedImplicits = nullptr;
-	if (!SharedRestImplicits)
-	{
-		RestCollectionSharedImplicits = 
-			&NCRestCollection.AddAttribute<TSharedPtr<Chaos::FImplicitObject, ESPMode::ThreadSafe>>(
-				SharedImplicitsAttribute, FTransformCollection::TransformGroup);
-	}
-
-	for (int32 Index = DynamicCollection->NumElements(FTransformCollection::TransformGroup) - 1; 0 <= Index; Index--)
-	{
-		Simplicials[Index] = TUniquePtr<FSimplicial>(nullptr);
-
-		//Implicits[Index] = Chaos::MakeSerializable(RestImplicits[Index]);
-		if (SharedRestImplicits)
-		{
-			Implicits[Index] = (*SharedRestImplicits)[Index];
-		}
-		else if(UniqueRestImplicits && RestCollectionSharedImplicits) // RestCollectionSharedImplicits is redundant, but static analysis build complains without it
-		{
-			// Steal the memory from the unique ptr.
-			Implicits[Index] = TSharedPtr<Chaos::FImplicitObject, ESPMode::ThreadSafe>((*UniqueRestImplicits)[Index].Release());
-			// Update the rest collection with a copy of the shared pointer.
-			(*RestCollectionSharedImplicits)[Index] = Implicits[Index];
-		}
-	}
-	SolverClusterID.Fill(nullptr);
-
-	// Clean up now invalid unique pointers.
-	if (RestCollectionSharedImplicits)
-	{
-		NCRestCollection.RemoveAttribute(ImplicitsAttribute, FTransformCollection::TransformGroup);
-	}
-
-	if (Parameters.RestCollection->HasAttribute(SimplicialsAttribute, FTransformCollection::TransformGroup))
-	{
-		const auto& RestSimplicials = Parameters.RestCollection->GetAttribute<TUniquePtr<FSimplicial>>(
-			SimplicialsAttribute, FTransformCollection::TransformGroup);
-		for (int32 Index = DynamicCollection->NumElements(FTransformCollection::TransformGroup) - 1; 0 <= Index; Index--)
-		{
-			Simplicials[Index] = TUniquePtr<FSimplicial>(RestSimplicials[Index] ? RestSimplicials[Index]->NewCopy() : nullptr);
-		}
-	}
-
-	// Merge Active Flags.
-	const int32 NumTransforms = SimulatableParticles.Num();
-	if(Parameters.RestCollection->HasAttribute(FGeometryCollection::SimulatableParticlesAttribute, FTransformCollection::TransformGroup))
-	{
-		// When the rest collection has been pre configured with simulation data, use that to determine it ability to simulate. 
-		const TManagedArray<bool>& ActiveRestParticles = 
-			Parameters.RestCollection->GetAttribute<bool>(
-				FGeometryCollection::SimulatableParticlesAttribute, FTransformCollection::TransformGroup);
-		check(ActiveRestParticles.Num() == DynamicCollection->Active.Num());
-		check(ActiveRestParticles.Num() == SimulatableParticles.Num());
-		for(int32 TransformIdx = 0; TransformIdx < NumTransforms; TransformIdx++)
-		{
-			SimulatableParticles[TransformIdx] = DynamicCollection->Active[TransformIdx] && ActiveRestParticles[TransformIdx];
-		}
-	}
-	else
-	{
-		// If no simulation data is available then default to the simulation of just the rigid geometry.
-		for(int32 TransformIdx = 0; TransformIdx < NumTransforms; TransformIdx++)
-		{
-			if(DynamicCollection->Children[TransformIdx].Num())
-			{
-				SimulatableParticles[TransformIdx] = false;
-			}
-			else
-			{
-				SimulatableParticles[TransformIdx] = DynamicCollection->Active[TransformIdx];
-			}
-		}
-	}
-}
-
 
 void FGeometryCollectionPhysicsProxy::InitializeRemoveOnFracture(FParticlesType& Particles, const TManagedArray<int32>& DynamicState)
 {
+	/*
+	@todo break everything
 	if (Parameters.DynamicCollection && Parameters.RemoveOnFractureEnabled)
 	{
 	//	TManagedArray<FGeometryCollectionBoneNode>& Hierarchy = Parameters.DynamicCollection->BoneHierarchy;
@@ -1189,6 +1184,7 @@ void FGeometryCollectionPhysicsProxy::InitializeRemoveOnFracture(FParticlesType&
 			}
 		}
 	}
+	*/
 }
 
 void FGeometryCollectionPhysicsProxy::AddCollisionToCollisionData(FRecordedFrame* ExistingFrame, const FParticlesType& Particles, const Chaos::TPBDCollisionConstraints<float, 3>::FPointContactConstraint& Constraint)
@@ -1718,7 +1714,9 @@ void FGeometryCollectionPhysicsProxy::UpdateRecordedState(float SolverTime, cons
 
 void FGeometryCollectionPhysicsProxy::OnRemoveFromSolver(Chaos::FPBDRigidsSolver *RBDSolver)
 {
-	for (auto* Handle : SolverClusterHandles)
+	const FGeometryDynamicCollection& DynamicCollection = *GameToPhysInterchange.GetConsumerBuffer();
+
+	for (const Chaos::TPBDRigidClusteredParticleHandle<float, 3>* Handle : DynamicCollection.SolverClusterHandles)
 	{
 		RBDSolver->RemoveParticleToProxy(Handle);
 	}
@@ -1764,6 +1762,7 @@ void FGeometryCollectionPhysicsProxy::SyncBeforeDestroy()
 
 void FGeometryCollectionPhysicsProxy::BufferGameState() 
 {
+	//
 	// NewData() redirects to this function.
 	/**
 	 * CONTEXT: GAMETHREAD
@@ -1779,14 +1778,14 @@ void FGeometryCollectionPhysicsProxy::BufferGameState()
 	// I don't think we need to allocate memory to do this.  I think the other 
 	// proxies get away with that because they're dealing with a low number of 
 	// bodies; which is not the case here. Rather, let's just use a triple buffer.
-
+/*
 	if (!Parameters.Simulating || 
 		!ensure(NumParticles != INDEX_NONE)) // Make sure InitBodiesGT() has been called!
 	{
 		return;
 	}
 
-	FGeometryCollectionResults& Buffer = *GameToPhysInterchange->AccessProducerBuffer();
+	FGeometryCollectionResults& Buffer = *GameToPhysInterchange.AccessProducerBuffer();
 	if (Buffer.NumParticlesAdded != NumParticles)
 	{
 		Buffer.InitArrays(*GTDynamicCollection);
@@ -1837,7 +1836,7 @@ void FGeometryCollectionPhysicsProxy::BufferGameState()
 		DynamicCollection->Transform, DynamicCollection->Parent, Buffer.GlobalTransforms);
 	
 	// Make the buffer available to the consumer.
-	GameToPhysInterchange->FlipProducer(); 
+	GameToPhysInterchange.FlipProducer(); 
 	
 	// If we wanted to emulate the other proxies, we could augment the 
 	// FStubGeometryCollectionData class to carry along the consumer buffer:
@@ -1849,6 +1848,7 @@ void FGeometryCollectionPhysicsProxy::BufferGameState()
 	// class, as I believe all these operations are happening with the same 
 	// instance of this class.  So skip all that, and just return null.
 	return;
+	*/
 }
 
 void FGeometryCollectionPhysicsProxy::PushToPhysicsState(const Chaos::FParticleData *InData)
@@ -1860,11 +1860,13 @@ void FGeometryCollectionPhysicsProxy::PushToPhysicsState(const Chaos::FParticleD
 	 * the solver
 	 */
 
+	/*
+	@todo break everything : move to the init
 	// The other proxies do a deep copy into dynamic transient memory, passed into
 	// this function via InData.  Instead of that, we're using persistent memory
 	// in a buffer to facilitate communicating between threads.
 
-	const FGeometryCollectionResults* GState = GameToPhysInterchange->GetConsumerBuffer();
+	const FGeometryCollectionResults* GState = GameToPhysInterchange.GetConsumerBuffer();
 	if (!GState)
 	{
 		// This will happen if GameToPhysInterchange hasn't been updated since
@@ -1947,6 +1949,7 @@ void FGeometryCollectionPhysicsProxy::PushToPhysicsState(const Chaos::FParticleD
 			}
 		}
 	}
+	*/
 }
 
 void FGeometryCollectionPhysicsProxy::BufferPhysicsResults()
@@ -1959,8 +1962,10 @@ void FGeometryCollectionPhysicsProxy::BufferPhysicsResults()
 	using namespace Chaos;
 	SCOPE_CYCLE_COUNTER(STAT_CacheResultGeomCollection);
 
+	/*
+	@todo break everything
 	//FGeometryCollectionResults& TargetResults = Results.GetPhysicsDataForWrite();
-	FGeometryCollectionResults& TargetResults = *PhysToGameInterchange->AccessProducerBuffer();
+	FGeometryCollectionResults& TargetResults = *PhysToGameInterchange.AccessProducerBuffer();
 
 	if (TargetResults.NumTransformGroup() != PTDynamicCollection.NumElements(FGeometryCollection::TransformGroup))
 	{
@@ -2006,20 +2011,20 @@ void FGeometryCollectionPhysicsProxy::BufferPhysicsResults()
 		SCOPE_CYCLE_COUNTER(STAT_CalcParticleToWorld);
 		// Fill particle to world transforms
 //		TargetResults.ParticleToWorldTransforms.SetNum(NumParticles);
-/*		for (int32 TransformIndex = 0; TransformIndex < NumParticles; ++TransformIndex)
-		{
-			//only update roots and first children
-			if (TargetResults.Parent[TransformIndex] == INDEX_NONE)
-			{
-				if (auto* Handle = SolverParticleHandles[TransformIndex])
-				{
-					const FTransform Xf(Handle->R(), Handle->X());
-					//const int32 ParticleIndex = BaseParticleIndex + TransformIndex;
-					TargetResults.ParticleToWorldTransforms[TransformIndex] = Xf;
-				}
-			}
-		}
-*/
+		//for (int32 TransformIndex = 0; TransformIndex < NumParticles; ++TransformIndex)
+		//{
+		//	//only update roots and first children
+		//	if (TargetResults.Parent[TransformIndex] == INDEX_NONE)
+		//	{
+		//		if (auto* Handle = SolverParticleHandles[TransformIndex])
+		//		{
+		//			const FTransform Xf(Handle->R(), Handle->X());
+		//			//const int32 ParticleIndex = BaseParticleIndex + TransformIndex;
+		//			TargetResults.ParticleToWorldTransforms[TransformIndex] = Xf;
+		//		}
+		//	}
+		//}
+
 		//const TArrayCollectionArray<Chaos::ClusterId>& ClusterID = GetSolver()->GetRigidClustering().GetClusterIdsArray();
 		//const TArrayCollectionArray<Chaos::TRigidTransform<float, 3>>& ClusterChildToParentMap = GetSolver()->GetRigidClustering().GetChildToParentMap();
 		//const TArrayCollectionArray<FMultiChildProxyId>& MultiChildProxyIdArray = GetSolver()->GetRigidClustering().GetMultiChildProxyIdArray();
@@ -2188,6 +2193,7 @@ void FGeometryCollectionPhysicsProxy::BufferPhysicsResults()
 	// We pulled state out of particles.  Push results into PTDynamicCollection, 
 	// so that it has an updated view.
 	UpdateGeometryCollection(TargetResults);
+	*/
 }
 
 void FGeometryCollectionPhysicsProxy::FlipBuffer()
@@ -2199,7 +2205,7 @@ void FGeometryCollectionPhysicsProxy::FlipBuffer()
 	 * the data
 	 */
 
-	PhysToGameInterchange->FlipProducer();
+	PhysToGameInterchange.FlipProducer();
 }
 
 // Called from FPhysScene_ChaosInterface::SyncBodies(), NOT the solver.
@@ -2214,8 +2220,9 @@ void FGeometryCollectionPhysicsProxy::PullFromPhysicsState()
 	 * Note: A read lock will have been acquired for this - so the physics thread won't force a buffer flip while this
 	 * sync is ongoing
 	 */
+	FGeometryDynamicCollection& DynamicCollection = *GameToPhysInterchange.AccessProducerBuffer();
 
-	const FGeometryCollectionResults* TargetResultPtr = PhysToGameInterchange->GetConsumerBuffer();
+	const FGeometryCollectionResults* TargetResultPtr = PhysToGameInterchange.GetConsumerBuffer();
 	if (!TargetResultPtr)
 		return;
 	// Remove const-ness as we're going to do the ExchangeArrays() thing.
@@ -2223,20 +2230,20 @@ void FGeometryCollectionPhysicsProxy::PullFromPhysicsState()
 
 	// We should never be changing the number of entries, this would break other 
 	// attributes in the transform group.
-	const int32 NumTransforms = GTDynamicCollection->Transform.Num();
+	const int32 NumTransforms = DynamicCollection.Transform.Num();
 	if (ensure(NumTransforms == TR.Transforms.Num()))	
 	{
 		for(int32 TmIndex = 0; TmIndex < NumTransforms; ++TmIndex)
 		{
 			if(!TR.DisabledStates[TmIndex])
 			{
-				GTDynamicCollection->Transform[TmIndex] = TR.Transforms[TmIndex];
+				DynamicCollection.Transform[TmIndex] = TR.Transforms[TmIndex];
 			}
-			GTDynamicCollection->DynamicState[TmIndex] = TR.DynamicState[TmIndex];
+			DynamicCollection.DynamicState[TmIndex] = TR.DynamicState[TmIndex];
 		}
 
 		//question: why do we need this? Sleeping objects will always have to update GPU
-		GTDynamicCollection->MakeDirty();
+		DynamicCollection.MakeDirty();
 
 		if (CacheSyncFunc)
 		{
@@ -2273,10 +2280,10 @@ void FGeometryCollectionPhysicsProxy::UpdateGeometryCollection(
 	const int32 TransformSize = GetTransformGroupSize();
 	for (int32 TransformGroupIndex = 0; TransformGroupIndex < TransformSize; ++TransformGroupIndex)
 	{
-		if (!SimulatableParticles[TransformGroupIndex])
+		if (!Collection->SimulatableParticles[TransformGroupIndex])
 			continue;
 
-		auto* Handle = SolverParticleHandles[TransformGroupIndex];
+		auto* Handle = Collection->SolverParticleHandles[TransformGroupIndex];
 
 		// dynamic state is also updated by the solver during field interaction. 
 		if(!Handle->Sleeping())
@@ -2389,10 +2396,10 @@ void FGeometryCollectionPhysicsProxy::InitializeSharedCollisionStructures(
 			TEXT("Mass"), FTransformCollection::TransformGroup);
 	TManagedArray<TUniquePtr<FSimplicial>>& CollectionSimplicials = 
 		RestCollection.AddAttribute<TUniquePtr<FSimplicial>>(
-			SimplicialsAttribute, FTransformCollection::TransformGroup);
+			FGeometryDynamicCollection::SimplicialsAttribute, FTransformCollection::TransformGroup);
 	TManagedArray<TSharedPtr<Chaos::FImplicitObject, ESPMode::ThreadSafe>>& CollectionImplicits = 
 		RestCollection.AddAttribute<TSharedPtr<Chaos::FImplicitObject, ESPMode::ThreadSafe>>(
-			SharedImplicitsAttribute, FTransformCollection::TransformGroup);
+			FGeometryDynamicCollection::SharedImplicitsAttribute, FTransformCollection::TransformGroup);
 
 	FTransform IdentityXf(FQuat::Identity, FVector(0));
 	IdentityXf.NormalizeRotation();
@@ -3120,7 +3127,7 @@ void FGeometryCollectionPhysicsProxy::ProcessCommands(FParticlesType& Particles,
 								Command.MetaData
 							};
 
-							TArrayView<FVector> ResultsView(&(InitialLinearVelocity[0]), InitialLinearVelocity.Num());
+							TArrayView<FVector> ResultsView(&(Collection->InitialLinearVelocity[0]), Collection->InitialLinearVelocity.Num());
 							static_cast<const FFieldNode<FVector> *>(Command.RootNode.Get())->Evaluate(Context, ResultsView);
 						}
 					}
@@ -3149,7 +3156,7 @@ void FGeometryCollectionPhysicsProxy::ProcessCommands(FParticlesType& Particles,
 								Command.MetaData
 							};
 
-							TArrayView<FVector> ResultsView(&(InitialAngularVelocity[0]), InitialAngularVelocity.Num());
+							TArrayView<FVector> ResultsView(&(Collection->InitialAngularVelocity[0]), Collection->InitialAngularVelocity.Num());
 							static_cast<const FFieldNode<FVector> *>(Command.RootNode.Get())->Evaluate(Context, ResultsView);
 						}
 					}
@@ -3164,6 +3171,9 @@ void FGeometryCollectionPhysicsProxy::ProcessCommands(FParticlesType& Particles,
 	{
 		Chaos::FPhysicsSolver* CurrentSolver = GetSolver();
 
+		/*
+		@todo break everything
+
 		//  Generate a Index mapping between the rigid body indices and 
 		//  the particle indices. This allows the geometry collection to
 		//  evaluate only its own particles. 
@@ -3248,12 +3258,17 @@ void FGeometryCollectionPhysicsProxy::ProcessCommands(FParticlesType& Particles,
 				Commands.RemoveAt(CommandIndex);
 			}
 		}
+		*/
 	}
 }
 
 void FGeometryCollectionPhysicsProxy::FieldForcesUpdateCallback(Chaos::FPhysicsSolver* InSolver, FParticlesType& Particles, Chaos::TArrayCollectionArray<FVector> & Force, Chaos::TArrayCollectionArray<FVector> & Torque, const float Time)
 {
-		if (Commands.Num())
+	/*
+
+	@todo break everything : re-enable fields without rigid body ids. 
+
+	if (Commands.Num())
 		{
 			// @todo: This seems like a waste if we just want to get everything
 			int32 Counter = 0;
@@ -3314,6 +3329,7 @@ void FGeometryCollectionPhysicsProxy::FieldForcesUpdateCallback(Chaos::FPhysicsS
 
 			}
 		}
+		*/
 }
 
 
@@ -4240,6 +4256,7 @@ void FGeometryCollectionPhysicsProxy::CreateRigidBodyCallback(FParticlesType& Pa
 
 void FGeometryCollectionPhysicsProxy::BindParticleCallbackMapping(Chaos::TArrayCollectionArray<PhysicsProxyWrapper>& PhysicsProxyReverseMap, Chaos::TArrayCollectionArray<int32>& ParticleIDReverseMap)
 {
+	/*
 	for (int32 Index = 0; Index < RigidBodyID.Num(); Index++)
 	{
 		const int32 RigidBodyIndex = RigidBodyID[Index];
@@ -4249,6 +4266,7 @@ void FGeometryCollectionPhysicsProxy::BindParticleCallbackMapping(Chaos::TArrayC
 			ParticleIDReverseMap[RigidBodyIndex] = Index;
 		}
 	}
+	*/
 }
 
 void FGeometryCollectionPhysicsProxy::PushKinematicStateToSolver(FParticlesType& Particles)
@@ -4357,6 +4375,7 @@ void FGeometryCollectionPhysicsProxy::ParameterUpdateCallback(FParticlesType& Pa
 		ProcessCommands(Particles, Time);
 
 
+#ifndef TODO_REIMPLEMENT_RIGID_CACHING
 		if (Parameters.RecordedTrack)
 		{
 			float ReverseTime = Parameters.RecordedTrack->GetLastTime() - Time + Parameters.ReverseCacheBeginTime;
@@ -4376,7 +4395,7 @@ void FGeometryCollectionPhysicsProxy::ParameterUpdateCallback(FParticlesType& Pa
 				}
 			}
 		}
-		/* @question : Should we tell the solver the mass has changed ? */
+#endif
 	}
 }
 
diff --git a/Engine/Source/Runtime/Experimental/ChaosSolvers/Public/GeometryCollectionProxyData.h b/Engine/Source/Runtime/Experimental/ChaosSolvers/Public/GeometryCollectionProxyData.h
--- a/Engine/Source/Runtime/Experimental/ChaosSolvers/Public/GeometryCollectionProxyData.h
+++ b/Engine/Source/Runtime/Experimental/ChaosSolvers/Public/GeometryCollectionProxyData.h
@@ -4,7 +4,9 @@
 
 #include "CoreMinimal.h"
 #include "Misc/Crc.h"
+
 #include "GeometryCollection/ManagedArrayCollection.h"
+#include "GeometryCollection/GeometryCollectionCollisionStructureManager.h"
 
 /*
 * Managed arrays for simulation data used by the GeometryCollectionProxy
@@ -56,16 +58,39 @@ public:
 	FGeometryDynamicCollection(FGeometryDynamicCollection&&) = delete;
 	FGeometryDynamicCollection& operator=(FGeometryDynamicCollection&&) = delete;
 
-
 	typedef FTransformDynamicCollection Super;
 
 	static const FName ActiveAttribute;
 	static const FName CollisionGroupAttribute;
+	static const FName CollisionMaskAttribute;
 	static const FName DynamicStateAttribute;
-
+	static const FName ImplicitsAttribute;
+	static const FName ShapesQueryDataAttribute;
+	static const FName ShapesSimDataAttribute;
+	static const FName SharedGeometryAttribute;
+	static const FName SharedImplicitsAttribute;
+	static const FName SimplicialsAttribute;
+	static const FName SolverParticleHandlesAttribute;
+	static const FName SolverClusterHandlesAttribute;
 
 	// Transform Group
 	TManagedArray<bool> Active;
 	TManagedArray<int32> CollisionGroup;
+	TManagedArray<int32> CollisionMask;
+	TManagedArray<int32> CollisionStructureID;
 	TManagedArray<int32> DynamicState;
+	TManagedArray<TSharedPtr<Chaos::FImplicitObject, ESPMode::ThreadSafe>> Implicits;
+	TManagedArray<FVector> InitialAngularVelocity;
+	TManagedArray<FVector> InitialLinearVelocity;
+	TManagedArray<FTransform> MassToLocal;
+	//TManagedArray<TArray<FCollisionFilterData>> ShapeQueryData;
+	//TManagedArray<TArray<FCollisionFilterData>> ShapeSimData;
+	//TManagedArray<TSharedPtr<Chaos::FImplicitObject, ESPMode::ThreadSafe>> SharedGeometry;
+	//TManagedArray<TSharedPtr<FCollisionStructureManager::FSimplicial> > Simplicials;
+	TManagedArray<bool> SimulatableParticles;
+	TManagedArray<Chaos::TPBDRigidClusteredParticleHandle<float, 3>*> SolverClusterHandles;
+	TManagedArray<Chaos::TPBDRigidParticleHandle<float, 3>*> SolverClusterID; // Rename to ClusterParent?
+	TManagedArray<Chaos::TPBDRigidClusteredParticleHandle<float, 3>*> SolverParticleHandles;
+
+
 };
\ No newline at end of file
diff --git a/Engine/Source/Runtime/Experimental/ChaosSolvers/Public/PhysicsProxy/GeometryCollectionPhysicsProxy.h b/Engine/Source/Runtime/Experimental/ChaosSolvers/Public/PhysicsProxy/GeometryCollectionPhysicsProxy.h
--- a/Engine/Source/Runtime/Experimental/ChaosSolvers/Public/PhysicsProxy/GeometryCollectionPhysicsProxy.h
+++ b/Engine/Source/Runtime/Experimental/ChaosSolvers/Public/PhysicsProxy/GeometryCollectionPhysicsProxy.h
@@ -70,9 +70,7 @@ public:
 
 	TManagedArray<int32> ClusterId;
 
-	TArray<TSharedPtr<Chaos::FImplicitObject, ESPMode::ThreadSafe>> SharedGeometry;
-	TArray<TArray<FCollisionFilterData>> ShapeSimData;
-	TArray<TArray<FCollisionFilterData>> ShapeQueryData;
+
 
 	bool IsObjectDynamic;
 	bool IsObjectLoading;
@@ -111,12 +109,7 @@ class CHAOSSOLVERS_API FGeometryCollectionPhysicsProxy : public TPhysicsProxy<FG
 	typedef TPhysicsProxy<FGeometryCollectionPhysicsProxy, FStubGeometryCollectionData> Base;
 
 public:
-	// collection attributes
-	static FName SimplicialsAttribute;
-	static FName ImplicitsAttribute;
-	static FName SharedImplicitsAttribute;
-	static FName SolverParticleHandlesAttribute;
-	static FName SolverClusterHandlesAttribute;
+
 
 	typedef FCollisionStructureManager::FSimplicial FSimplicial;
 
@@ -211,10 +204,10 @@ public:
 	{CollisionParticlesPerObjectFraction = CollisionParticlesPerObjectFractionIn;}
 
 	TManagedArray<Chaos::TPBDRigidClusteredParticleHandle<float, 3>*>& GetSolverParticleHandles() 
-	{ return SolverParticleHandles; }
+	{ return GameToPhysInterchange.AccessProducerBuffer()->SolverParticleHandles; }
 
 	const FGeometryCollectionResults* GetConsumerResultsGT() const 
-	{ return PhysToGameInterchange ? PhysToGameInterchange->PeekConsumerBuffer() : nullptr; }
+	{ return PhysToGameInterchange.PeekConsumerBuffer(); }
 
 	/** Enqueue a field \p Command to be processed by \c ProcessCommands() or 
 	 * \c FieldForcesUpdateCallback(). 
@@ -286,11 +279,9 @@ protected:
 	 * Traverses the parents of \p TransformIndex in \p GeometryCollection, counting
 	 * the number of levels until the next parent is \c INDEX_NONE.
 	 */
-	int32 CalculateHierarchyLevel(
-		const FGeometryDynamicCollection* GeometryCollection, 
-		int32 TransformIndex) const;
+	int32 CalculateHierarchyLevel(const FGeometryDynamicCollection& GeometryCollection, int32 TransformIndex) const;
 
-	void CreateDynamicAttributes();
+	static void InitializeDynamicCollection(FGeometryDynamicCollection& DynamicCollection, const FGeometryCollection& RestCollection, const FSimulationParameters& Params);
 	void InitializeRemoveOnFracture(FParticlesType& Particles, const TManagedArray<int32>& DynamicState);
 	void PushKinematicStateToSolver(FParticlesType& Particles);
 
@@ -336,28 +327,10 @@ private:
 	// Indicate when loaded
 	bool IsObjectLoading;
 
-	// Dynamic collection on the game thread - used to populate the simulated collection
-	FGeometryDynamicCollection* GTDynamicCollection;
-	// Duplicated dynamic collection for use on the physics thread, copied to the game thread on sync
-	//TUniquePtr<FGeometryDynamicCollection> PTDynamicCollection;
-	FGeometryDynamicCollection PTDynamicCollection;
-
 	TMap<Chaos::TPBDRigidParticleHandle<float, 3>*, int32> HandleToTransformGroupIndex;
 
-	TManagedArray<Chaos::TPBDRigidClusteredParticleHandle<float, 3>*> SolverParticleHandles;
-	TManagedArray<Chaos::TPBDRigidClusteredParticleHandle<float, 3>*> SolverClusterHandles;
-	TManagedArray<Chaos::TPBDRigidParticleHandle<float, 3>*> SolverClusterID; // Rename to ClusterParent?
 
-	TManagedArray<FTransform> MassToLocal;
-	TManagedArray<int32> CollisionMask;
-	TManagedArray<int32> CollisionStructureID;
-	TManagedArray<int32> RigidBodyID; // Deprecated.  Not added to dynamic collection.
-	TManagedArray<FVector> InitialAngularVelocity;
-	TManagedArray<FVector> InitialLinearVelocity;
-	TManagedArray<bool> SimulatableParticles;
 
-	TManagedArray<TUniquePtr<FSimplicial> > Simplicials; // FSimplicial = Chaos::TBVHParticles<float,3>
-	TManagedArray<TSharedPtr<Chaos::FImplicitObject, ESPMode::ThreadSafe>> Implicits;
 
 	TArray<int32> EndFrameUnparentingBuffer;
 
@@ -375,7 +348,6 @@ private:
 
 	// Index of the first particles for this collection in the larger particle array
 	int32 BaseParticleIndex;
-	// Number of particles added by this collection
 	int32 NumParticles;
 
 	// Time since this object started simulating
@@ -399,10 +371,7 @@ private:
 	// Per object collision fraction.
 	float CollisionParticlesPerObjectFraction;
 
-	// Double buffer of geom collection result data - TODO (Ryan): deprecated?
-	//Chaos::TBufferedData<FGeometryCollectionResults> Results;
-
-	// TODO (Ryan) - Currently this is using triple buffers for game-physics and 
+	// Currently this is using triple buffers for game-physics and 
 	// physics-game thread communication, but not for any reason other than this 
 	// is the only implementation we currently have of a guarded buffer - a buffer 
 	// that tracks it's own state, rather than having other mechanisms determine 
@@ -411,8 +380,8 @@ private:
 	// currently managing the exchange is built upon.  However, I believe that 
 	// logic locks, and the triple buffer would enable a decoupled lock-free 
 	// paradigm, at least for this component of the handshake.
-	TUniquePtr<Chaos::FGuardedTripleBuffer<FGeometryCollectionResults>> PhysToGameInterchange;
-	TUniquePtr<Chaos::FGuardedTripleBuffer<FGeometryCollectionResults>> GameToPhysInterchange;
+	Chaos::FGuardedTripleBuffer<FGeometryCollectionResults> PhysToGameInterchange;
+	Chaos::FDoubleBuffer<FGeometryDynamicCollection> GameToPhysInterchange;
 };
 
 CHAOSSOLVERS_API void BuildSimulationData(Chaos::FErrorReporter& ErrorReporter, FGeometryCollection& GeometryCollection, const FSharedSimulationParameters& SharedParams);
diff --git a/Engine/Source/Runtime/Experimental/GeometryCollectionCore/Private/GeometryCollection/ManagedArrayCollection.cpp b/Engine/Source/Runtime/Experimental/GeometryCollectionCore/Private/GeometryCollection/ManagedArrayCollection.cpp
--- a/Engine/Source/Runtime/Experimental/GeometryCollectionCore/Private/GeometryCollection/ManagedArrayCollection.cpp
+++ b/Engine/Source/Runtime/Experimental/GeometryCollectionCore/Private/GeometryCollection/ManagedArrayCollection.cpp
@@ -19,7 +19,7 @@ static const FName GuidName("GUID");
 void FManagedArrayCollection::AddGroup(FName Group)
 {
 	ensure(!GroupInfo.Contains(Group));
-	FGroupInfo info {
+	FGroupInfo info{
 		0
 	};
 	GroupInfo.Add(Group, info);
@@ -77,7 +77,7 @@ TArray<FName> FManagedArrayCollection::GroupNames() const
 bool FManagedArrayCollection::HasAttribute(FName Name, FName Group) const
 {
 	bool bReturnValue = false;
-	for (const TTuple<FKeyType,FValueType>& Entry : Map)
+	for (const TTuple<FKeyType, FValueType>& Entry : Map)
 	{
 		if (Entry.Key.Get<0>() == Name && Entry.Key.Get<1>() == Group)
 		{
@@ -120,10 +120,10 @@ void FManagedArrayCollection::GenerateGuids(FName Group, int32 StartIdx)
 	// we don't need these at runtime in any case, so if we need in the editor later, make sure this is in-editor only
 	if (GIsEditor)
 	{
-	 	for (int32 Idx = StartIdx; Idx < Guids.Num(); ++Idx)
-	 	{
-	 		Guids[Idx] = FGuid::NewGuid();
-	 	}
+		for (int32 Idx = StartIdx; Idx < Guids.Num(); ++Idx)
+		{
+			Guids[Idx] = FGuid::NewGuid();
+		}
 	}
 }
 
@@ -177,7 +177,7 @@ void FManagedArrayCollection::RemoveGroup(FName Group)
 			Entry.Value.GroupIndexDependency = "";
 		}
 	}
-	for (const FName & AttrName : DelList)
+	for (const FName& AttrName : DelList)
 	{
 		Map.Remove(FManagedArrayCollection::MakeMapKey(AttrName, Group));
 	}
@@ -194,7 +194,7 @@ void FManagedArrayCollection::Resize(int32 Size, FName Group)
 		return;
 	}
 
-	ensureMsgf(Size > NumElements(Group),TEXT("Use RemoveElements to shrink a group."));
+	ensureMsgf(Size > NumElements(Group), TEXT("Use RemoveElements to shrink a group."));
 	const int32 StartSize = GroupInfo[Group].Size;
 	for (TTuple<FKeyType, FValueType>& Entry : Map)
 	{
@@ -246,20 +246,6 @@ void FManagedArrayCollection::ReorderElements(FName Group, const TArray<int32>&
 	}
 }
 
-#if 0	//not needed until we support per instance serialization
-void FManagedArrayCollection::SwapElements(int32 Index1, int32 Index2, FName Group)
-{
-	//todo(ocohen): handle dependent groups
-	for (TTuple<FKeyType, FValueType>& Entry : Map)
-	{
-		if (Entry.Key.Get<1>() == Group)
-		{
-			Entry.Value.Value->Swap(Index1, Index2);
-		}
-	}
-}
-#endif
-
 void FManagedArrayCollection::SetDependency(FName Name, FName Group, FName DependencyGroup)
 {
 	ensure(HasAttribute(Name, Group));
@@ -282,83 +268,38 @@ void FManagedArrayCollection::RemoveDependencyFor(FName Group)
 	}
 }
 
-void FManagedArrayCollection::SyncGroupSizeAndOrder(const FManagedArrayCollection& MasterCollection, FName Group)
+void FManagedArrayCollection::SyncGroupSizeFrom(const FManagedArrayCollection& InCollection, FName Group)
 {
 	if (!HasGroup(Group))
 	{
 		AddGroup(Group);
 	}
 
-	//For now we ignore order and just sync size. Ordering is needed for saving out per instance changes. Will add this soon
-	const int32 GroupSize = MasterCollection.GroupInfo[Group].Size;
-	Resize(GroupSize, Group);
-
-#if 0
-	//We want to lock the group and ensure that the master collection's guids match ours in order and size.
-	//Existing entries with guids must be preserved and simply re-ordered
-	//Entries with mismatching guids will be deleted
-	//New entries will be created to ensure the size is the same.
-
-	const TManagedArray<FGuid>& MasterGuids = MasterCollection.GetAttribute<FGuid>(GuidName, Group);
-	TManagedArray<FGuid>& SlaveGuids = GetAttribute<FGuid>(GuidName, Group);
-	
-	//remove guids that are not in the master list
-	TArray<int32> SortedEntriesToDelete;
-	SortedEntriesToDelete.Reserve(SlaveGuids.Num());
-	for (int32 Idx = 0; Idx < SlaveGuids.Num(); ++Idx)
+	Resize(InCollection.GroupInfo[Group].Size, Group);
+}
+
+void FManagedArrayCollection::CopyMatchingAttributesFrom(const FManagedArrayCollection& InCollection)
+{
+	for (const auto& Pair : InCollection.GroupInfo)
 	{
-		if (!MasterGuids.Contains(SlaveGuids[Idx]))	//comment: set may help, but it doesn't take TManagedArray directly
-		{
-			SortedEntriesToDelete.Add(Idx);
-		}
+		SyncGroupSizeFrom(InCollection, Pair.Key);
 	}
-	RemoveElements(Group, SortedEntriesToDelete);
-	//Slave guids are now a subset of the master guids, but they can be out of order. So we must re-order and add new guids
-	for (int32 Idx = 0; Idx < MasterGuids.Num(); ++Idx)
+	for (TTuple<FKeyType, FValueType>& Entry : Map)
 	{
-		int32 FoundSlaveIdx = -1;
-		{
-			for (int32 SlaveIdx = 0; SlaveIdx < SlaveGuids.Num(); ++SlaveIdx)
-			{
-				if (MasterGuids[Idx] == SlaveGuids[SlaveIdx])
-				{
-					FoundSlaveIdx = SlaveIdx;
-					break;
-				}
-			}
-		}
-
-		if (FoundSlaveIdx == -1)
-		{
-			//no existing entry, create a new one and set its guid
-			FoundSlaveIdx = NumElements(Group);
-			AddElements(1, Group);
-			SlaveGuids[FoundSlaveIdx] = MasterGuids[Idx];
-		}
-
-		if(FoundSlaveIdx != Idx)
+		if (InCollection.HasAttribute(Entry.Key.Get<0>(), Entry.Key.Get<1>()))
 		{
-			//entry exists, just needs to be re-ordered
-			SwapElements(FoundSlaveIdx, Idx, Group);
+			CopyAttribute(InCollection, Entry.Key.Get<0>(), Entry.Key.Get<1>());
 		}
 	}
-#endif
-}
 
-void FManagedArrayCollection::SyncAllGroups(const FManagedArrayCollection& MasterCollection)
-{
-	for (const auto& Pair : MasterCollection.GroupInfo)
-	{
-		SyncGroupSizeAndOrder(MasterCollection, Pair.Key);
-	}
 }
 
-void FManagedArrayCollection::CopyAttribute(const FManagedArrayCollection& MasterCollection, FName Name, FName Group)
+void FManagedArrayCollection::CopyAttribute(const FManagedArrayCollection& InCollection, FName Name, FName Group)
 {
-	SyncGroupSizeAndOrder(MasterCollection, Group);
+	SyncGroupSizeFrom(InCollection, Group);
 	FKeyType Key = FManagedArrayCollection::MakeMapKey(Name, Group);
-	
-	const FValueType& OriginalValue = MasterCollection.Map[Key];
+
+	const FValueType& OriginalValue = InCollection.Map[Key];
 	const FValueType& DestValue = Map[Key];	//todo(ocohen): API assumes an AddAttribute is called before copy is done. It'd be nice to handle the case where AddAttribute was not done first
 	check(OriginalValue.ArrayType == DestValue.ArrayType);
 	DestValue.Value->Init(*OriginalValue.Value);
@@ -406,7 +347,7 @@ bool FManagedArrayCollection::HasCycle(FName NewGroup, FName DependencyGroup)
 FString FManagedArrayCollection::ToString() const
 {
 	FString Buffer("");
-	for (FName GroupName : GroupNames()) 
+	for (FName GroupName : GroupNames())
 	{
 		Buffer += GroupName.ToString() + "\n";
 		for (FName AttributeName : AttributeNames(GroupName))
@@ -414,11 +355,11 @@ FString FManagedArrayCollection::ToString() const
 			FKeyType Key = FManagedArrayCollection::MakeMapKey(AttributeName, GroupName);
 			const FValueType& Value = Map[Key];
 
-			const void * PointerAddress = static_cast<const void*>(Value.Value);
+			const void* PointerAddress = static_cast<const void*>(Value.Value);
 			std::stringstream AddressStream;
 			AddressStream << PointerAddress;
 
-			Buffer += GroupName.ToString() + ":" + AttributeName.ToString() + " ["+ FString(AddressStream.str().c_str()) +"]\n";
+			Buffer += GroupName.ToString() + ":" + AttributeName.ToString() + " [" + FString(AddressStream.str().c_str()) + "]\n";
 		}
 	}
 	return Buffer;
@@ -430,7 +371,7 @@ void FManagedArrayCollection::Serialize(Chaos::FChaosArchive& Ar)
 	int Version = 4;
 	Ar << Version;
 
-	if(Ar.IsLoading())
+	if (Ar.IsLoading())
 	{
 		//We can't serialize entire tmap in place because we may have new groups. todo(ocohen): baked data should be simpler since all entries exist
 		TMap< FName, FGroupInfo> TmpGroupInfo;
@@ -530,6 +471,6 @@ FArchive& operator<<(FArchive& Ar, FManagedArrayCollection::FValueType& ValueIn)
 		//todo(ocohen): need a better way to enforce this
 		ValueIn.Value->Serialize(static_cast<Chaos::FChaosArchive&>(Ar));
 	}
-	
+
 	return Ar;
 }
\ No newline at end of file
diff --git a/Engine/Source/Runtime/Experimental/GeometryCollectionCore/Public/GeometryCollection/ManagedArrayCollection.h b/Engine/Source/Runtime/Experimental/GeometryCollectionCore/Public/GeometryCollection/ManagedArrayCollection.h
--- a/Engine/Source/Runtime/Experimental/GeometryCollectionCore/Public/GeometryCollection/ManagedArrayCollection.h
+++ b/Engine/Source/Runtime/Experimental/GeometryCollectionCore/Public/GeometryCollection/ManagedArrayCollection.h
@@ -44,8 +44,9 @@ class GEOMETRYCOLLECTIONCORE_API FManagedArrayCollection
 public:
 
 	FManagedArrayCollection();
-	FManagedArrayCollection(const FManagedArrayCollection&) = delete;
 	virtual ~FManagedArrayCollection() {}
+
+	FManagedArrayCollection(const FManagedArrayCollection&) = delete;
 	FManagedArrayCollection& operator=(const FManagedArrayCollection&) = delete;
 	FManagedArrayCollection(FManagedArrayCollection&&) = default;
 	FManagedArrayCollection& operator=(FManagedArrayCollection&&) = default;
@@ -264,20 +265,13 @@ public:
 	* @param Name - The name of the attribute
 	* @param Group - The group that manages the attribute
 	*/
-	void CopyAttribute(const FManagedArrayCollection& MasterCollection, FName Name, FName Group);
-
-	/**
-	* Size and order a group so that it matches the group found in the master collection.
-	* @param MasterCollection - The collection we are ordering our group against. The GUID attribute allows us to keep data that has been previously synced. Otherwise data is lost
-	* @param Group - The group that manages the attribute
-	*/
-	void SyncGroupSizeAndOrder(const FManagedArrayCollection& MasterCollection, FName Group);
+	void CopyAttribute(const FManagedArrayCollection& InCollection, FName Name, FName Group);
 
 	/**
-	* Sync all master groups with this collection. This is a utility to easily sync rest and dynamic collections
-	* @param MasterCollection - All groups from this collection found in the master will be sized and ordered accordingly
+	* Copy attributes that match the input collection. This is a utility to easily sync collections
+	* @param InCollection - All groups from this collection found in the input will be sized accordingly
 	*/
-	void SyncAllGroups(const FManagedArrayCollection& MasterCollection);
+	void CopyMatchingAttributesFrom(const FManagedArrayCollection& InCollection);
 
 	/**
 	* Number of elements in a group
@@ -418,6 +412,14 @@ private:
 
 protected:
 
+	/**
+	* Size and order a group so that it matches the group found in the input collection.
+	* @param InCollection - The collection we are ordering our group against. 
+	* @param Group - The group that manages the attribute
+	*/
+	void SyncGroupSizeFrom(const FManagedArrayCollection& InCollection, FName Group);
+
+
 	/**
 	*   operator<<(FGroupInfo)
 	*/
