Chaos :

- Removed dead end code.

diff --git a/Engine/Source/Runtime/Experimental/ChaosSolvers/Private/PhysicsProxy/GeometryCollectionPhysicsProxy.cpp b/Engine/Source/Runtime/Experimental/ChaosSolvers/Private/PhysicsProxy/GeometryCollectionPhysicsProxy.cpp
--- a/Engine/Source/Runtime/Experimental/ChaosSolvers/Private/PhysicsProxy/GeometryCollectionPhysicsProxy.cpp
+++ b/Engine/Source/Runtime/Experimental/ChaosSolvers/Private/PhysicsProxy/GeometryCollectionPhysicsProxy.cpp
@@ -1187,531 +1187,6 @@ void FGeometryCollectionPhysicsProxy::InitializeRemoveOnFracture(FParticlesType&
 	*/
 }
 
-void FGeometryCollectionPhysicsProxy::AddCollisionToCollisionData(FRecordedFrame* ExistingFrame, const FParticlesType& Particles, const Chaos::TPBDCollisionConstraints<float, 3>::FPointContactConstraint& Constraint)
-{
-#if TODO_REIMPLEMENT_RIGID_CLUSTERING
-	const Chaos::TArrayCollectionArray<Chaos::ClusterId>& ClusterIdsArray = GetSolver()->GetRigidClustering().GetClusterIdsArray();
-	const TMap<uint32, TUniquePtr<TArray<uint32>>>& ParentToChildrenMap = GetSolver()->GetRigidClustering().GetChildrenMap();
-	const Chaos::TArrayCollectionArray<PhysicsProxyWrapper>& PhysicsProxyReverseMapping = GetSolver()->GetPhysicsProxyReverseMapping();
-
-	int32 NewIdx = ExistingFrame->Collisions.Add(FSolverCollisionData());
-	FSolverCollisionData& Collision = ExistingFrame->Collisions[NewIdx];
-
-	Collision.Location = Constraint.Location;
-	Collision.AccumulatedImpulse = Constraint.AccumulatedImpulse;
-	Collision.Normal = Constraint.Normal;
-	Collision.Velocity1 = Particles.V(Constraint.ParticleIndex);
-	Collision.Velocity2 = Particles.V(Constraint.LevelsetIndex);
-	Collision.AngularVelocity1 = Particles.W(Constraint.ParticleIndex);
-	Collision.AngularVelocity2 = Particles.W(Constraint.LevelsetIndex);
-	Collision.Mass1 = Particles.M(Constraint.ParticleIndex);
-	Collision.Mass2 = Particles.M(Constraint.LevelsetIndex);
-	Collision.ParticleIndex = Constraint.ParticleIndex;
-	Collision.LevelsetIndex = Constraint.LevelsetIndex;
-
-	// If Constraint.ParticleIndex is a cluster store an index for a mesh in this cluster
-	if (ClusterIdsArray[Constraint.ParticleIndex].NumChildren > 0)
-	{
-		int32 ParticleIndexMesh = GetSolver()->GetParticleIndexMesh(ParentToChildrenMap, Constraint.ParticleIndex);
-		ensure(ParticleIndexMesh != INDEX_NONE);
-		Collision.ParticleIndexMesh = ParticleIndexMesh;
-	}
-	// If Constraint.LevelsetIndex is a cluster store an index for a mesh in this cluster
-	if (ClusterIdsArray[Constraint.LevelsetIndex].NumChildren > 0)
-	{
-		int32 LevelsetIndexMesh = GetSolver()->GetParticleIndexMesh(ParentToChildrenMap, Constraint.LevelsetIndex);
-		ensure(LevelsetIndexMesh != INDEX_NONE);
-		Collision.LevelsetIndexMesh = LevelsetIndexMesh;
-	}
-#endif
-}
-
-void FGeometryCollectionPhysicsProxy::UpdateCollisionData(const FParticlesType& Particles, const FCollisionConstraintsType& CollisionRule, FRecordedFrame* ExistingFrame)
-{
-	ExistingFrame->Collisions.Reset();
-
-	if (Parameters.CollisionData.SaveCollisionData && ExistingFrame->Timestamp > 0.f && Parameters.CollisionData.CollisionDataSizeMax > 0)
-	{
-#if TODO_REIMPLEMENT_PHYSICS_PROXY_REVERSE_MAPPING
-		const TArray<Chaos::TPBDCollisionConstraints<float, 3>::FRigidBodyContactConstraint>& AllConstraintsArray = CollisionRule.GetAllConstraints();
-		if (AllConstraintsArray.Num() > 0)
-		{
-			const Chaos::TArrayCollectionArray<PhysicsProxyWrapper>& PhysicsProxyReverseMapping = GetSolver()->GetPhysicsProxyReverseMapping();
-
-			TArray<Chaos::TPBDCollisionConstraints<float, 3>::FRigidBodyContactConstraint> ConstraintsArray;
-			ConstraintsArray.SetNumUninitialized(AllConstraintsArray.Num());
-
-			FBox BoundingBox(ForceInitToZero);
-			int32 NumConstraints = 0;
-			for (int32 Idx = 0; Idx < AllConstraintsArray.Num(); ++Idx)
-			{
-				// Check if the collision is for this PhysicsProxy
-				IPhysicsProxyBase* PhysicsProxy = PhysicsProxyReverseMapping[AllConstraintsArray[Idx].ParticleIndex].PhysicsProxy;
-				if (PhysicsProxy == this)
-				{
-					if (ensure(!AllConstraintsArray[Idx].AccumulatedImpulse.ContainsNaN()))
-					{
-						if (!AllConstraintsArray[Idx].AccumulatedImpulse.IsZero())
-						{
-							if (ensure(!AllConstraintsArray[Idx].Location.ContainsNaN() &&
-								!AllConstraintsArray[Idx].Normal.ContainsNaN()) &&
-								!Particles.V(AllConstraintsArray[Idx].ParticleIndex).ContainsNaN() &&
-								!Particles.V(AllConstraintsArray[Idx].LevelsetIndex).ContainsNaN() &&
-								!Particles.W(AllConstraintsArray[Idx].ParticleIndex).ContainsNaN() &&
-								!Particles.W(AllConstraintsArray[Idx].LevelsetIndex).ContainsNaN())
-							{
-								BoundingBox += AllConstraintsArray[Idx].Location;
-
-								ConstraintsArray[NumConstraints] = AllConstraintsArray[Idx];
-								NumConstraints++;
-							}
-						}
-					}
-				}
-			}
-
-			ConstraintsArray.SetNum(NumConstraints);
-
-			if (ConstraintsArray.Num() > 0)
-			{
-				if (Parameters.CollisionData.DoCollisionDataSpatialHash &&
-					Parameters.CollisionData.CollisionDataSpatialHashRadius > 0.f &&
-					ConstraintsArray.Num() > 1 &&
-					(BoundingBox.GetExtent().X > 0.f || BoundingBox.GetExtent().Y > 0.f || BoundingBox.GetExtent().Z > 0.f))
-				{
-					// Validate Parameters.CollisionData.CollisionDataSpatialHashRadius
-					// CellSize must be smaller than the smallest bbox extent
-					float SpatialHashRadius = Parameters.CollisionData.CollisionDataSpatialHashRadius;
-					TArray<float> ExtentArray = { BoundingBox.GetExtent().X , BoundingBox.GetExtent().Y, BoundingBox.GetExtent().Z };
-					ExtentArray.Sort();
-					if (ExtentArray[0] != 0.f)
-					{
-						if (2.f * SpatialHashRadius > ExtentArray[0])
-						{
-							SpatialHashRadius = 0.5f * ExtentArray[0];
-						}
-					}
-					else
-					{
-						if (ExtentArray[1] != 0.f)
-						{
-							if (2.f * SpatialHashRadius > ExtentArray[1])
-							{
-								SpatialHashRadius = 0.5f * ExtentArray[1];
-							}
-						}
-						else
-						{
-							if (2.f * SpatialHashRadius > ExtentArray[2])
-							{
-								SpatialHashRadius = 0.5f * ExtentArray[2];
-							}
-						}
-					}
-
-					// Spatial hash the constraints
-					TMultiMap<int32, int32> HashTableMap;
-					ComputeHashTable(ConstraintsArray, BoundingBox, HashTableMap, SpatialHashRadius);
-
-					TArray<int32> UsedCellsArray;
-					HashTableMap.GetKeys(UsedCellsArray);
-
-					int32 NumCollisionsThisFrame = 0;
-					for (int32 IdxCell = 0; IdxCell < UsedCellsArray.Num(); ++IdxCell)
-					{
-						TArray<int32> ConstraintsInCellArray;
-						HashTableMap.MultiFind(UsedCellsArray[IdxCell], ConstraintsInCellArray);
-
-						int32 NumConstraintsToGetFromCell = FMath::Min(Parameters.CollisionData.MaxCollisionPerCell, ConstraintsInCellArray.Num());
-						for (int32 IdxConstraint = 0; IdxConstraint < NumConstraintsToGetFromCell; ++IdxConstraint)
-						{
-							AddCollisionToCollisionData(ExistingFrame, Particles, ConstraintsArray[ConstraintsInCellArray[IdxConstraint]]);
-						}
-					}
-
-					if (ExistingFrame->Collisions.Num() > Parameters.CollisionData.CollisionDataSizeMax)
-					{
-						TArray<FSolverCollisionData> CollisionsArray1;
-
-						float FInc = (float)ExistingFrame->Collisions.Num() / (float)Parameters.CollisionData.CollisionDataSizeMax;
-
-						CollisionsArray1.SetNumUninitialized(Parameters.CollisionData.CollisionDataSizeMax);
-						for (int32 IdxCollision = 0; IdxCollision < Parameters.CollisionData.CollisionDataSizeMax; ++IdxCollision)
-						{
-							int32 NewIdx = FMath::FloorToInt((float)IdxCollision * FInc);
-							CollisionsArray1[IdxCollision] = ExistingFrame->Collisions[NewIdx];
-						}
-
-						ExistingFrame->Collisions.SetNumUninitialized(Parameters.CollisionData.CollisionDataSizeMax);
-						for (int32 IdxCollision = 0; IdxCollision < Parameters.CollisionData.CollisionDataSizeMax; ++IdxCollision)
-						{
-							ExistingFrame->Collisions[IdxCollision] = CollisionsArray1[IdxCollision];
-						}
-					}
-				}
-				else
-				{
-					if (ConstraintsArray.Num() <= Parameters.CollisionData.CollisionDataSizeMax)
-					{
-						for (int32 IdxConstraint = 0; IdxConstraint < ConstraintsArray.Num(); ++IdxConstraint)
-						{
-							AddCollisionToCollisionData(ExistingFrame, Particles, ConstraintsArray[IdxConstraint]);
-						}
-					}
-					else
-					{
-						float FInc = (float)ConstraintsArray.Num() / (float)Parameters.CollisionData.CollisionDataSizeMax;
-
-						for (int32 IdxConstraint = 0; IdxConstraint < Parameters.CollisionData.CollisionDataSizeMax; ++IdxConstraint)
-						{
-							int32 Idx = FMath::FloorToInt((float)IdxConstraint * FInc);
-
-							AddCollisionToCollisionData(ExistingFrame, Particles, ConstraintsArray[Idx]);
-						}
-					}
-				}
-			}
-		}
-#endif
-	}
-}
-
-void FGeometryCollectionPhysicsProxy::AddBreakingToBreakingData(FRecordedFrame* ExistingFrame, const FParticlesType& Particles, const Chaos::TBreakingData<float, 3>& Breaking)
-{
-#if TODO_REIMPLEMENT_RIGID_CLUSTERING
-	const Chaos::TArrayCollectionArray<Chaos::ClusterId>& ClusterIdsArray = GetSolver()->GetRigidClustering().GetClusterIdsArray();
-	const TMap<uint32, TUniquePtr<TArray<uint32>>>& ParentToChildrenMap = GetSolver()->GetRigidClustering().GetChildrenMap();
-	const Chaos::TArrayCollectionArray<PhysicsProxyWrapper>& PhysicsProxyReverseMapping = GetSolver()->GetPhysicsProxyReverseMapping();
-
-	int32 NewIdx = ExistingFrame->Breakings.Add(FSolverBreakingData());
-	FSolverBreakingData& NewBreaking = ExistingFrame->Breakings[NewIdx];
-
-	NewBreaking.Location = Breaking.Location;
-	NewBreaking.Velocity = Particles.V(Breaking.ParticleIndex);
-	NewBreaking.AngularVelocity = Particles.W(Breaking.ParticleIndex);
-	NewBreaking.Mass = Particles.M(Breaking.ParticleIndex);
-	NewBreaking.ParticleIndex = Breaking.ParticleIndex;
-
-	// If Breaking.ParticleIndex is a cluster store an index for a mesh in this cluster
-	if (ClusterIdsArray[Breaking.ParticleIndex].NumChildren > 0)
-	{
-		int32 ParticleIndexMesh = GetSolver()->GetParticleIndexMesh(ParentToChildrenMap, Breaking.ParticleIndex);
-		ensure(ParticleIndexMesh != INDEX_NONE);
-		NewBreaking.ParticleIndexMesh = ParticleIndexMesh;
-	}
-#endif
-}
-
-void FGeometryCollectionPhysicsProxy::UpdateBreakingData(const FParticlesType& Particles, FRecordedFrame* ExistingFrame)
-{
-	ExistingFrame->Breakings.Reset();
-
-	if (Parameters.BreakingData.SaveBreakingData && ExistingFrame->Timestamp > 0.f && Parameters.BreakingData.BreakingDataSizeMax > 0)
-	{
-#if TODO_REIMPLEMENT_GETALLCLUSTERBREAKINGS
-		const TArray<Chaos::TBreakingData<float, 3>>& AllBreakingsArray = GetSolver()->GetAllClusterBreakings();
-		if (AllBreakingsArray.Num() > 0)
-		{
-			const Chaos::TArrayCollectionArray<PhysicsProxyWrapper>& PhysicsProxyReverseMapping = GetSolver()->GetPhysicsProxyReverseMapping();
-
-			TArray<Chaos::TBreakingData<float, 3>> BreakingsArray;
-			BreakingsArray.SetNumUninitialized(AllBreakingsArray.Num());
-
-			FBox BoundingBox(ForceInitToZero);
-			int32 NumBreakings = 0;
-			for (int32 Idx = 0; Idx < AllBreakingsArray.Num(); ++Idx)
-			{
-				// Check if the breaking is for this PhysicsProxy
-				IPhysicsProxyBase* PhysicsProxy = PhysicsProxyReverseMapping[AllBreakingsArray[Idx].ParticleIndex].PhysicsProxy;
-				if (PhysicsProxy == this)
-				{
-					if (ensure(!AllBreakingsArray[Idx].Location.ContainsNaN() &&
-						!Particles.V(AllBreakingsArray[Idx].ParticleIndex).ContainsNaN() &&
-						!Particles.W(AllBreakingsArray[Idx].ParticleIndex).ContainsNaN()))
-					{
-						BoundingBox += AllBreakingsArray[Idx].Location;
-
-						BreakingsArray[NumBreakings] = AllBreakingsArray[Idx];
-						NumBreakings++;
-					}
-				}
-			}
-
-			BreakingsArray.SetNum(NumBreakings);
-
-			if (BreakingsArray.Num() > 0)
-			{
-				if (Parameters.BreakingData.DoBreakingDataSpatialHash &&
-					Parameters.BreakingData.BreakingDataSpatialHashRadius > 0.f &&
-					BreakingsArray.Num() > 1 &&
-					(BoundingBox.GetExtent().X > 0.f || BoundingBox.GetExtent().Y > 0.f || BoundingBox.GetExtent().Z > 0.f))
-				{
-					// Validate Parameters.CollisionData.CollisionDataSpatialHashRadius
-					// CellSize must be smaller than the smallest bbox extent
-					float SpatialHashRadius = Parameters.BreakingData.BreakingDataSpatialHashRadius;
-					TArray<float> ExtentArray = { BoundingBox.GetExtent().X , BoundingBox.GetExtent().Y, BoundingBox.GetExtent().Z };
-					ExtentArray.Sort();
-					if (ExtentArray[0] != 0.f)
-					{
-						if (2.f * SpatialHashRadius > ExtentArray[0])
-						{
-							SpatialHashRadius = 0.5f * ExtentArray[0];
-						}
-					}
-					else
-					{
-						if (ExtentArray[1] != 0.f)
-						{
-							if (2.f * SpatialHashRadius > ExtentArray[1])
-							{
-								SpatialHashRadius = 0.5f * ExtentArray[1];
-							}
-						}
-						else
-						{
-							if (2.f * SpatialHashRadius > ExtentArray[2])
-							{
-								SpatialHashRadius = 0.5f * ExtentArray[2];
-							}
-						}
-					}
-
-					// Spatial hash the constraints
-					TMultiMap<int32, int32> HashTableMap;
-					ComputeHashTable(BreakingsArray, BoundingBox, HashTableMap, SpatialHashRadius);
-
-					TArray<int32> UsedCellsArray;
-					HashTableMap.GetKeys(UsedCellsArray);
-
-					for (int32 IdxCell = 0; IdxCell < UsedCellsArray.Num(); ++IdxCell)
-					{
-						TArray<int32> BreakingsInCellArray;
-						HashTableMap.MultiFind(UsedCellsArray[IdxCell], BreakingsInCellArray);
-
-						int32 NumBreakingsToGetFromCell = FMath::Min(Parameters.BreakingData.MaxBreakingPerCell, BreakingsInCellArray.Num());
-						for (int32 IdxBreaking = 0; IdxBreaking < NumBreakingsToGetFromCell; ++IdxBreaking)
-						{
-							AddBreakingToBreakingData(ExistingFrame, Particles, BreakingsArray[BreakingsInCellArray[IdxBreaking]]);
-						}
-					}
-
-					if (ExistingFrame->Breakings.Num() > Parameters.BreakingData.BreakingDataSizeMax)
-					{
-						TArray<FSolverBreakingData> BreakingsArray1;
-
-						float FInc = (float)ExistingFrame->Breakings.Num() / (float)Parameters.BreakingData.BreakingDataSizeMax;
-
-						BreakingsArray1.SetNumUninitialized(Parameters.BreakingData.BreakingDataSizeMax);
-						for (int32 IdxBreaking = 0; IdxBreaking < Parameters.BreakingData.BreakingDataSizeMax; ++IdxBreaking)
-						{
-							int32 NewIdx = FMath::FloorToInt((float)IdxBreaking * FInc);
-							BreakingsArray1[IdxBreaking] = ExistingFrame->Breakings[NewIdx];
-						}
-
-						ExistingFrame->Breakings.SetNumUninitialized(Parameters.BreakingData.BreakingDataSizeMax);
-						for (int32 IdxBreaking = 0; IdxBreaking < Parameters.BreakingData.BreakingDataSizeMax; ++IdxBreaking)
-						{
-							ExistingFrame->Breakings[IdxBreaking] = BreakingsArray1[IdxBreaking];
-						}
-					}
-				}
-				else
-				{
-					if (BreakingsArray.Num() <= Parameters.BreakingData.BreakingDataSizeMax)
-					{
-						for (int32 IdxBreaking = 0; IdxBreaking < BreakingsArray.Num(); ++IdxBreaking)
-						{
-							AddBreakingToBreakingData(ExistingFrame, Particles, BreakingsArray[IdxBreaking]);
-						}
-					}
-					else
-					{
-						float FInc = (float)BreakingsArray.Num() / (float)Parameters.BreakingData.BreakingDataSizeMax;
-
-						for (int32 IdxBreaking = 0; IdxBreaking < Parameters.BreakingData.BreakingDataSizeMax; ++IdxBreaking)
-						{
-							int32 Idx = FMath::FloorToInt((float)IdxBreaking * FInc);
-
-							AddBreakingToBreakingData(ExistingFrame, Particles, BreakingsArray[Idx]);
-						}
-					}
-				}
-			}
-		}
-#endif
-	}
-}
-
-void FGeometryCollectionPhysicsProxy::UpdateTrailingData(const FParticlesType& Particles, FRecordedFrame* ExistingFrame)
-{
-	ExistingFrame->Trailings.Reset();
-
-	if (Parameters.TrailingData.SaveTrailingData && ExistingFrame->Timestamp > 0.f && Parameters.TrailingData.TrailingDataSizeMax > 0)
-	{
-		const float TrailingMinSpeedThresholdSquared = Parameters.TrailingData.TrailingMinSpeedThreshold * Parameters.TrailingData.TrailingMinSpeedThreshold;
-
-		// Find previous frame trailing data
-		int32 LastKey = RecordedTracks.FindLastKeyBefore(ExistingFrame->Timestamp);
-		if (LastKey != INDEX_NONE)
-		{
-			FRecordedFrame& PrevFrame = RecordedTracks.Records[LastKey];
-			if (PrevFrame.Trailings.Num() > 0)
-			{
-				for (FSolverTrailingData& Trailing : PrevFrame.Trailings)
-				{
-					if (Particles.ObjectState(Trailing.ParticleIndex) != Chaos::EObjectStateType::Sleeping &&
-						!Particles.Disabled(Trailing.ParticleIndex) &&
-						Particles.V(Trailing.ParticleIndex).SizeSquared() >= TrailingMinSpeedThresholdSquared)
-					{
-						ExistingFrame->Trailings.Add(Trailing);
-					}
-				}
-			}
-		}
-
-		if (ExistingFrame->Trailings.Num() < Parameters.TrailingData.TrailingDataSizeMax)
-		{
-#if TODO_REIMPLEMENT_PHYSICS_PROXY_REVERSE_MAPPING
-			const Chaos::TArrayCollectionArray<PhysicsProxyWrapper>& PhysicsProxyReverseMapping = GetSolver()->GetPhysicsProxyReverseMapping();
-			const Chaos::TArrayCollectionArray<Chaos::ClusterId>& ClusterIdsArray = GetSolver()->GetRigidClustering().GetClusterIdsArray();
-			const TMap<uint32, TUniquePtr<TArray<uint32>>>& ParentToChildrenMap = GetSolver()->GetRigidClustering().GetChildrenMap();
-
-			TArray<FSolverTrailingData> AllTrailings;
-			for (int32 IdxParticle = 0; IdxParticle < (int32)Particles.Size(); ++IdxParticle)
-			{
-				// Check if the particle is for this PhysicsProxy
-				IPhysicsProxyBase* PhysicsProxy = PhysicsProxyReverseMapping[IdxParticle].PhysicsProxy;
-				if (PhysicsProxy == this)
-				{
-					if (ensure(FMath::IsFinite(Particles.InvM(IdxParticle))))
-					{
-						if (Particles.ObjectState(IdxParticle) != Chaos::EObjectStateType::Sleeping &&
-							!Particles.Disabled(IdxParticle) &&
-							Particles.InvM(IdxParticle) != 0.f)
-						{
-							if (Particles.Geometry(IdxParticle) && 
-								Particles.Geometry(IdxParticle)->HasBoundingBox())
-							{
-								if (ensure(!Particles.X(IdxParticle).ContainsNaN()) &&
-									!Particles.V(IdxParticle).ContainsNaN() &&
-									!Particles.W(IdxParticle).ContainsNaN() &&
-									FMath::IsFinite(Particles.M(IdxParticle)))
-								{
-									Chaos::TAABB<float, 3> BoundingBox = Particles.Geometry(IdxParticle)->BoundingBox();
-									Chaos::TVector<float, 3> Extents = BoundingBox.Extents();
-									float Volume = Extents[0] * Extents[1] * Extents[2];
-									float SpeedSquared = Particles.V(IdxParticle).SizeSquared();
-
-									if (SpeedSquared >= TrailingMinSpeedThresholdSquared &&
-										Volume > Parameters.TrailingData.TrailingMinVolumeThreshold)
-									{
-										FSolverTrailingData TrailingData(Particles.X(IdxParticle),
-											Particles.V(IdxParticle),
-											Particles.W(IdxParticle),
-											Particles.M(IdxParticle),
-											IdxParticle,
-											INDEX_NONE);
-
-										// If IdxParticle is a cluster store an index for a mesh in this cluster
-										if (ClusterIdsArray[IdxParticle].NumChildren > 0)
-										{
-											int32 ParticleIndexMesh = GetSolver()->GetParticleIndexMesh(ParentToChildrenMap, IdxParticle);
-											ensure(ParticleIndexMesh != INDEX_NONE);
-											TrailingData.ParticleIndexMesh = ParticleIndexMesh;
-										}
-
-										FSetElementId Id = ExistingFrame->Trailings.FindId(TrailingData);
-										if (!Id.IsValidId())
-										{
-											AllTrailings.Add(TrailingData);
-										}
-										else
-										{
-											ExistingFrame->Trailings[Id].Location = Particles.X(IdxParticle);
-											ExistingFrame->Trailings[Id].Velocity = Particles.V(IdxParticle);
-											ExistingFrame->Trailings[Id].AngularVelocity = Particles.W(IdxParticle);
-										}
-									}
-								}
-							}
-						}
-					}
-				}
-			}
-
-			if (AllTrailings.Num() > 0)
-			{
-				int32 NumTrailingsToAdd = Parameters.TrailingData.TrailingDataSizeMax - ExistingFrame->Trailings.Num();
-
-				if (AllTrailings.Num() <= NumTrailingsToAdd)
-				{
-					for (int32 IdxTrailing = 0; IdxTrailing < AllTrailings.Num(); ++IdxTrailing)
-					{
-						ExistingFrame->Trailings.Add(AllTrailings[IdxTrailing]);
-					}
-				}
-				else
-				{
-					float FInc = (float)AllTrailings.Num() / (float)(NumTrailingsToAdd);
-
-					for (int32 IdxTrailing = 0; IdxTrailing < NumTrailingsToAdd; ++IdxTrailing)
-					{
-						int32 Idx = FMath::FloorToInt((float)IdxTrailing * FInc);
-						ExistingFrame->Trailings.Add(AllTrailings[Idx]);
-					}
-				}
-			}
-#endif
-		}
-	}
-}
-
-void FGeometryCollectionPhysicsProxy::UpdateRecordedState(float SolverTime, const TManagedArray<int32>& InRigidBodyID, const TManagedArray<int32>& InCollectionClusterID, const Chaos::TArrayCollectionArray<bool>& InInternalCluster, const FParticlesType& Particles, const FCollisionConstraintsType& CollisionRule)
-{
-	FRecordedFrame* ExistingFrame = RecordedTracks.FindRecordedFrame(SolverTime);
-
-	if(!ExistingFrame)
-	{
-		ExistingFrame = &InsertRecordedFrame(RecordedTracks, SolverTime);
-	}
-
-	ExistingFrame->Reset(InRigidBodyID.Num());
-	ExistingFrame->Timestamp = SolverTime;
-
-	// Collision
-	UpdateCollisionData(Particles, CollisionRule, ExistingFrame);
-
-	// Breaking
-	UpdateBreakingData(Particles, ExistingFrame);
-
-	// Trailing
-	UpdateTrailingData(Particles, ExistingFrame);
-
-	ParallelFor(InRigidBodyID.Num(), [&](int32 Index)
-	{
-		const int32 ExternalIndex = InRigidBodyID[Index];
-
-		if(ExternalIndex >= 0)
-		{
-			FTransform& NewTransform = ExistingFrame->Transforms[Index];
-
-			NewTransform.SetTranslation(Particles.X(ExternalIndex));
-			NewTransform.SetRotation(Particles.R(ExternalIndex));
-			NewTransform.SetScale3D(FVector(1.0f));
-
-			int32 ClusterParentIndex = InCollectionClusterID[Index];
-			if (ClusterParentIndex != INDEX_NONE && InInternalCluster[ClusterParentIndex])
-			{
-				ExistingFrame->DisabledFlags[Index] = Particles.Disabled(ClusterParentIndex);
-			}
-			else
-			{
-				ExistingFrame->DisabledFlags[Index] = Particles.Disabled(ExternalIndex);
-			}
-		}
-	});
-}
-
 void FGeometryCollectionPhysicsProxy::OnRemoveFromSolver(Chaos::FPBDRigidsSolver *RBDSolver)
 {
 	const FGeometryDynamicCollection& DynamicCollection = *GameToPhysInterchange.GetConsumerBuffer();
@@ -2996,56 +2471,6 @@ void BuildSimulationData(Chaos::FErrorReporter& ErrorReporter, FGeometryCollecti
 	FGeometryCollectionPhysicsProxy::InitializeSharedCollisionStructures(ErrorReporter, GeometryCollection, SharedParams);
 }
 
-void FGeometryCollectionPhysicsProxy::MergeRecordedTracks(const FRecordedTransformTrack& A, const FRecordedTransformTrack& B, FRecordedTransformTrack& Target)
-{
-	const int32 NumAKeys = A.Records.Num();
-	const int32 NumBKeys = B.Records.Num();
-
-	if(NumAKeys == 0)
-	{
-		Target = B;
-		return;
-	}
-
-	if(NumBKeys == 0)
-	{
-		Target = A;
-		return;
-	}
-
-	// We have to copy the tracks to a local cache here because Target could point at either A or B.
-	FRecordedTransformTrack TempMergedTrack = A;
-
-	// Expand to hold all the keys
-	TempMergedTrack.Records.Reserve(NumAKeys + NumBKeys);
-
-	// Insert B frames into the merged set
-	for(int32 BKeyIndex = 0; BKeyIndex < NumBKeys; ++BKeyIndex)
-	{
-		const FRecordedFrame& BFrame = B.Records[BKeyIndex];
-		int32 KeyBefore = TempMergedTrack.FindLastKeyBefore(BFrame.Timestamp);
-
-		TempMergedTrack.Records.Insert(BFrame, KeyBefore + 1);
-	}
-
-	// Copy to target
-	Target = TempMergedTrack;
-}
-
-FRecordedFrame& FGeometryCollectionPhysicsProxy::InsertRecordedFrame(FRecordedTransformTrack& InTrack, float InTime)
-{
-	// Can't just throw on the end, might need to insert
-	const int32 BeforeNewIndex = InTrack.FindLastKeyBefore(InTime);
-
-	if(BeforeNewIndex == InTrack.Records.Num() - 1)
-	{
-		InTrack.Records.AddDefaulted();
-		return InTrack.Records.Last();
-	}
-
-	const int32 NewRecordIndex = InTrack.Records.Insert(FRecordedFrame(), BeforeNewIndex + 1);
-	return InTrack.Records[NewRecordIndex];
-}
 
 
 //==============================================================================
@@ -3100,7 +2525,7 @@ void FGeometryCollectionPhysicsProxy::ProcessCommands(FParticlesType& Particles,
 						TArrayView<int32> DynamicStateView(&(Collection->DynamicState[0]), Collection->DynamicState.Num());
 						static_cast<const FFieldNode<int32> *>(Command.RootNode.Get())->Evaluate(Context, DynamicStateView);
 
-						PushKinematicStateToSolver(Particles);
+						//PushKinematicStateToSolver(Particles);
 					}
 				}
 				Commands.RemoveAt(CommandIndex);
@@ -3333,1079 +2758,16 @@ void FGeometryCollectionPhysicsProxy::FieldForcesUpdateCallback(Chaos::FPhysicsS
 }
 
 
-//==============================================================================
-// DEPRECATED
-//==============================================================================
 
 
-void FGeometryCollectionPhysicsProxy::UpdateKinematicBodiesCallback(const FParticlesType& Particles, const float Dt, const float Time, FKinematicProxy& Proxy)
+void FGeometryCollectionPhysicsProxy::ParameterUpdateCallback(FParticlesType& Particles, const float Time)
 {
-	check(false); // not currently called
-#if TODO_REIMPLEMENT_KINEMATIC_PROXY
-
-	SCOPE_CYCLE_COUNTER(STAT_KinematicUpdate);
-		FGeometryDynamicCollection* Collection = Parameters.DynamicCollection;
-		check(Collection);
-
-		bool bIsCachePlaying = Parameters.IsCachePlaying();
-		bool bIsReverseCachePlaying = Parameters.IsCacheRecording() && Parameters.ReverseCacheBeginTime != 0 && Parameters.ReverseCacheBeginTime < Time;
-		if ((!bIsCachePlaying && !bIsReverseCachePlaying) || !Parameters.RecordedTrack)
-		{
-			return;
-		}
-
-		bool bFirst = !Proxy.Ids.Num();
-		if (bFirst)
-		{
-			Proxy.Position.Reset(RigidBodyID.Num());
-			Proxy.Rotation.Reset(RigidBodyID.Num());
-			Proxy.NextPosition.Reset(RigidBodyID.Num());
-			Proxy.NextRotation.Reset(RigidBodyID.Num());
-
-			Proxy.Position.AddUninitialized(RigidBodyID.Num());
-			Proxy.Rotation.AddUninitialized(RigidBodyID.Num());
-			Proxy.NextPosition.AddUninitialized(RigidBodyID.Num());
-			Proxy.NextRotation.AddUninitialized(RigidBodyID.Num());
-
-			for (int32 i = 0; i < RigidBodyID.Num(); ++i)
-			{
-				Proxy.Ids.Add(RigidBodyID[i]);
-
-				// Initialise to rest state
-				const int32 RbId = Proxy.Ids.Last();
-				Proxy.Position[i] = RbId != INDEX_NONE ? Particles.X(RbId) : FVector::ZeroVector;
-				Proxy.Rotation[i] = RbId != INDEX_NONE ? Particles.R(RbId) : FQuat::Identity;
-				Proxy.NextPosition[i] = Proxy.Position[i];
-				Proxy.NextRotation[i] = Proxy.Rotation[i];
-			}
-		}
-
-		if (bIsCachePlaying && !bIsReverseCachePlaying && (Time < Parameters.CacheBeginTime || !Parameters.RecordedTrack->IsTimeValid(Time)))
-		{
-			return;
-		}
-
-		float ReverseTime = Parameters.RecordedTrack->GetLastTime() - Time + Parameters.ReverseCacheBeginTime;
-		if (bIsReverseCachePlaying && !Parameters.RecordedTrack->IsTimeValid(ReverseTime))
-		{
-			return;
-		}
-
-		const FRecordedFrame* FirstFrame = nullptr;
-		const FRecordedFrame* SecondFrame = nullptr;
-		float PlaybackTime = bIsReverseCachePlaying ? ReverseTime : Time;
-		Parameters.RecordedTrack->GetFramesForTime(PlaybackTime, FirstFrame, SecondFrame);
-
-		if (FirstFrame && !SecondFrame)
-		{
-			// Only one frame to take information from (simpler case)
-			const int32 NumActives = FirstFrame->TransformIndices.Num();
-
-			// Actives
-			Chaos::PhysicsParallelFor(NumActives, [&](int32 Index)
-			{
-				const int32 InternalIndex = FirstFrame->TransformIndices[Index];
-				if (InternalIndex >= RigidBodyID.Num() || InternalIndex < 0)
-				{
-					UE_LOG(UGCC_LOG, Error, 
-						TEXT("%s: Cache index %d out of range: [%d, %d).  Regenerate the cache."),
-						*Parameters.Name, InternalIndex, 0, RigidBodyID.Num());
-					return;
-				}
-				const int32 ExternalIndex = RigidBodyID[InternalIndex];
-
-				if (ExternalIndex != INDEX_NONE && Particles.InvM(ExternalIndex) == 0.0f && !Particles.Disabled(ExternalIndex))
-				{
-					const FTransform& ParticleTransform = FirstFrame->Transforms[Index];
-					Proxy.Position[InternalIndex] = Particles.X(ExternalIndex);
-					Proxy.Rotation[InternalIndex] = Particles.R(ExternalIndex);
-					Proxy.NextPosition[InternalIndex] = ParticleTransform.GetTranslation();
-					Proxy.NextRotation[InternalIndex] = ParticleTransform.GetRotation();
-				}
-			});
-		}
-		else if (FirstFrame && SecondFrame)
-		{
-			// Both frames valid, second frame has all the indices we need
-			const int32 NumActives = SecondFrame->TransformIndices.Num();
-
-			const float Alpha = (PlaybackTime - FirstFrame->Timestamp) / (SecondFrame->Timestamp - FirstFrame->Timestamp);
-			check(0 <= Alpha && Alpha <= 1.0f);
-
-			Chaos::PhysicsParallelFor(NumActives, [&](int32 Index)
-			{
-				const int32 InternalIndex = SecondFrame->TransformIndices[Index];
-				if (InternalIndex >= RigidBodyID.Num() || InternalIndex < 0)
-				{
-					UE_LOG(UGCC_LOG, Error, 
-						TEXT("%s: Cache index %d out of range: [%d, %d).  Regenerate the cache."),
-						*Parameters.Name, InternalIndex, 0, RigidBodyID.Num());
-					return;
-				}
-				const int32 ExternalIndex = RigidBodyID[InternalIndex];
-				const int32 PreviousIndexSlot = Index < SecondFrame->PreviousTransformIndices.Num() ? SecondFrame->PreviousTransformIndices[Index] : INDEX_NONE;
-
-				if (ExternalIndex != INDEX_NONE && Particles.InvM(ExternalIndex) == 0.0f && !Particles.Disabled(ExternalIndex))
-				{
-					if (PreviousIndexSlot != INDEX_NONE)
-					{
-						Proxy.Position[InternalIndex] = Proxy.NextPosition[InternalIndex];
-						Proxy.Rotation[InternalIndex] = Proxy.NextRotation[InternalIndex];
-
-						FTransform BlendedTM;
-						BlendedTM.Blend(FirstFrame->Transforms[PreviousIndexSlot], SecondFrame->Transforms[Index], Alpha);
-
-						Proxy.NextPosition[InternalIndex] = BlendedTM.GetTranslation();
-						Proxy.NextRotation[InternalIndex] = BlendedTM.GetRotation();
-					}
-					else
-					{
-						// NewActive case
-						Proxy.Position[InternalIndex] = Proxy.NextPosition[InternalIndex];
-						Proxy.Rotation[InternalIndex] = Proxy.NextRotation[InternalIndex];
-
-						FTransform BlendedTM;
-						BlendedTM.Blend(FTransform(Particles.R(ExternalIndex), Particles.X(ExternalIndex), FVector::OneVector), SecondFrame->Transforms[Index], Alpha);
-
-						Proxy.NextPosition[InternalIndex] = BlendedTM.GetTranslation();
-						Proxy.NextRotation[InternalIndex] = BlendedTM.GetRotation();
-					}
-				}
-			});
-			// #BGallagher Handle new inactives. If it's a cluster parent and it's fully disabled we'll need to decluster it here.
-		}
-#endif
-}
-
-void FGeometryCollectionPhysicsProxy::StartFrameCallback(const float Dt, const float Time)
-{
-	SCOPE_CYCLE_COUNTER(STAT_GeomBeginFrame);
-#if TODO_REIMPLEMENT_GET_RIGID_PARTICLES
-	if (InitializedState == ESimulationInitializationState::Initialized)
-	{
-		// Reverse playback only plays back what we just recorded.  So, the condition 
-		// is, "are we in record mode, but ready to play back what we've recorded?"
-		const bool bIsReverseCachePlaying = Parameters.IsCacheRecording() && Parameters.ReverseCacheBeginTime != 0 && Time > Parameters.ReverseCacheBeginTime;
-		if (Parameters.IsCachePlaying() || bIsReverseCachePlaying)
-		{
-			// Update the enabled/disabled state for kinematic particles for the upcoming frame
-			Chaos::FPhysicsSolver* ThisSolver = GetSolver();
-			Chaos::FPhysicsSolver::FParticlesType& Particles = ThisSolver->GetRigidParticles();
-
-			if (!Parameters.RecordedTrack)
-			{
-				if (!ensure(Parameters.CacheType == EGeometryCollectionCacheType::Record))
-				{
-					return;
-				}
-				Parameters.RecordedTrack = new FRecordedTransformTrack();
-				Parameters.bOwnsTrack = true;
-			}
-			if (Parameters.bClearCache && bIsReverseCachePlaying)
-			{
-				check(!CommitRecordedStateCallback);
-				*const_cast<FRecordedTransformTrack*>(Parameters.RecordedTrack) = FRecordedTransformTrack::ProcessRawRecordedData(RecordedTracks);
-				Parameters.bClearCache = false;
-			}
-
-			bool bParticlesUpdated = false;
-
-			const float ThisFrameTime = bIsReverseCachePlaying ? (Parameters.RecordedTrack->GetLastTime() - Time + Parameters.ReverseCacheBeginTime) : Time;
-			if (!Parameters.RecordedTrack->IsTimeValid(ThisFrameTime))
-			{
-				// Invalid cache time, nothing to update
-				return;
-			}
-
-			FCriticalSection ParticleUpdateLock;
-			const int32 NumMappings = RigidBodyID.Num();
-			Chaos::PhysicsParallelFor(NumMappings, [&](int32 InternalParticleIndex)
-				//for(int32 Index = 0; Index < NumMappings; ++Index)
-			{
-				const int32 ExternalParticleIndex = RigidBodyID[InternalParticleIndex];
-
-				if (ExternalParticleIndex == INDEX_NONE)
-				{
-					return;
-				}
-
-				if (Particles.InvM(ExternalParticleIndex) != 0)
-				{
-					return;
-				}
-
-				// We need to check a window of Now - Dt to Now and see if we ever activated in that time.
-				// This is for short activations because if we miss one then the playback will be incorrect
-				const bool bShouldBeDisabled = !Parameters.RecordedTrack->GetWasActiveInWindow(InternalParticleIndex, ThisFrameTime, bIsReverseCachePlaying ? (ThisFrameTime - Dt) : (ThisFrameTime + Dt));
-				const bool bDisabledNow = Particles.Disabled(ExternalParticleIndex);
-				if (bDisabledNow != bShouldBeDisabled)
-				{
-					Particles.SetDisabledLowLevel(ExternalParticleIndex, bShouldBeDisabled);
-					ParticleUpdateLock.Lock();
-					if (!bShouldBeDisabled)
-					{
-						Particles.SetObjectState(ExternalParticleIndex, Chaos::EObjectStateType::Kinematic);
-						if (!ensure(ThisSolver->ActiveIndices().Find(ExternalParticleIndex) == nullptr))
-						{
-							bParticlesUpdated = true;
-						}
-						else if(!bParticlesUpdated)
-						{
-							ThisSolver->NonDisabledIndices().Add(ExternalParticleIndex);
-							ThisSolver->ActiveIndices().Add(ExternalParticleIndex);
-						}
-					}
-					else
-					{
-						if (!ensure(ThisSolver->ActiveIndices().Find(ExternalParticleIndex) != nullptr))
-						{
-							bParticlesUpdated = true;
-						}
-						else if(!bParticlesUpdated)
-						{
-							ThisSolver->NonDisabledIndices().Remove(ExternalParticleIndex);
-							ThisSolver->ActiveIndices().Remove(ExternalParticleIndex);
-						}
-					}
-					ParticleUpdateLock.Unlock();
-				}
-
-			});
-
-			// Do not add collisions if reverse
-			if (!bIsReverseCachePlaying)
-			{
-				const FRecordedFrame* RecordedFrame = Parameters.RecordedTrack->FindRecordedFrame(ThisFrameTime);
-				if (RecordedFrame == nullptr)
-				{
-					const int32 Index = Parameters.RecordedTrack->FindLastKeyBefore(ThisFrameTime);
-					if (Index > 0 && Index < Parameters.RecordedTrack->Records.Num())
-					{
-						RecordedFrame = &Parameters.RecordedTrack->Records[Index];
-					}
-				}
-
-				if (RecordedFrame)
-				{
-					// Collisions
-					if (Parameters.CollisionData.DoGenerateCollisionData && ThisFrameTime > 0.f && Parameters.CollisionData.CollisionDataSizeMax > 0)
-					{
-						if (RecordedFrame->Collisions.Num() > 0)
-						{
-							Chaos::FPhysicsSolver::FCollisionDataArray& AllCollisionsDataArray = ThisSolver->GetAllCollisionsDataArray();
-							TMap<IPhysicsProxyBase*, TArray<int32>>& AllCollisionsIndicesByPhysicsProxy = ThisSolver->GetAllCollisionsIndicesByPhysicsProxy();
-
-							if (!AllCollisionsIndicesByPhysicsProxy.Contains(this))
-							{
-								AllCollisionsIndicesByPhysicsProxy.Add(this, TArray<int32>());
-							}
-
-							for (int32 Idx = 0; Idx < RecordedFrame->Collisions.Num(); ++Idx)
-							{
-								// Check if the particle is still kinematic
-								if (RecordedFrame->Collisions[Idx].ParticleIndex < 0 ||
-									(RecordedFrame->Collisions[Idx].ParticleIndex >= 0 && RecordedFrame->Collisions[Idx].ParticleIndex < static_cast<int32>(Particles.Size()) && Particles.ObjectState(RecordedFrame->Collisions[Idx].ParticleIndex) == Chaos::EObjectStateType::Kinematic))
-								{
-									const int32 NewIdx = AllCollisionsDataArray.Add(Chaos::TCollisionData<float, 3>());
-									Chaos::TCollisionData<float, 3>& AllCollisionsDataArrayItem = AllCollisionsDataArray[NewIdx];
-
-									AllCollisionsDataArrayItem.Location = RecordedFrame->Collisions[Idx].Location;
-									AllCollisionsDataArrayItem.AccumulatedImpulse = RecordedFrame->Collisions[Idx].AccumulatedImpulse;
-									AllCollisionsDataArrayItem.Normal = RecordedFrame->Collisions[Idx].Normal;
-									AllCollisionsDataArrayItem.Velocity1 = RecordedFrame->Collisions[Idx].Velocity1;
-									AllCollisionsDataArrayItem.Velocity2 = RecordedFrame->Collisions[Idx].Velocity2;
-									AllCollisionsDataArrayItem.AngularVelocity1 = RecordedFrame->Collisions[Idx].AngularVelocity1;
-									AllCollisionsDataArrayItem.AngularVelocity2 = RecordedFrame->Collisions[Idx].AngularVelocity2;
-									AllCollisionsDataArrayItem.Mass1 = RecordedFrame->Collisions[Idx].Mass1;
-									AllCollisionsDataArrayItem.Mass2 = RecordedFrame->Collisions[Idx].Mass2;
-#if TODO_CONVERT_GEOMETRY_COLLECTION_PARTICLE_INDICES_TO_PARTICLE_POINTERS
-									AllCollisionsDataArrayItem.ParticleIndex = RecordedFrame->Collisions[Idx].ParticleIndex;
-#endif
-									AllCollisionsDataArrayItem.LevelsetIndex = RecordedFrame->Collisions[Idx].LevelsetIndex;
-									AllCollisionsDataArrayItem.ParticleIndexMesh = RecordedFrame->Collisions[Idx].ParticleIndexMesh;
-									AllCollisionsDataArrayItem.LevelsetIndexMesh = RecordedFrame->Collisions[Idx].LevelsetIndexMesh;
-
-									AllCollisionsIndicesByPhysicsProxy[this].Add(NewIdx);
-								}
-							}
-						}
-					}
-
-					// Breaking
-					if (Parameters.BreakingData.DoGenerateBreakingData && ThisFrameTime > 0.f && Parameters.BreakingData.BreakingDataSizeMax > 0)
-					{
-						if (RecordedFrame->Breakings.Num() > 0)
-						{
-							Chaos::FPhysicsSolver::FBreakingDataArray& AllBreakingsDataArray = ThisSolver->GetAllBreakingsDataArray();
-							TMap<IPhysicsProxyBase*, TArray<int32>>& AllBreakingsIndicesByPhysicsProxy = ThisSolver->GetAllBreakingsIndicesByPhysicsProxy();
-
-							if (!AllBreakingsIndicesByPhysicsProxy.Contains(this))
-							{
-								AllBreakingsIndicesByPhysicsProxy.Add(this, TArray<int32>());
-							}
-
-							for (int32 Idx = 0; Idx < RecordedFrame->Breakings.Num(); ++Idx)
-							{
-								// Check if the particle is still kinematic
-								if (RecordedFrame->Breakings[Idx].ParticleIndex < 0 ||
-									(RecordedFrame->Breakings[Idx].ParticleIndex >= 0 && RecordedFrame->Breakings[Idx].ParticleIndex < static_cast<int32>(Particles.Size()) && Particles.ObjectState(RecordedFrame->Breakings[Idx].ParticleIndex) == Chaos::EObjectStateType::Kinematic))
-								{
-									const int32 NewIdx = AllBreakingsDataArray.Add(Chaos::TBreakingData<float, 3>());
-									Chaos::TBreakingData<float, 3>& AllBreakingsDataArrayItem = AllBreakingsDataArray[NewIdx];
-
-									AllBreakingsDataArrayItem.Location = RecordedFrame->Breakings[Idx].Location;
-									AllBreakingsDataArrayItem.Velocity = RecordedFrame->Breakings[Idx].Velocity;
-									AllBreakingsDataArrayItem.AngularVelocity = RecordedFrame->Breakings[Idx].AngularVelocity;
-									AllBreakingsDataArrayItem.Mass = RecordedFrame->Breakings[Idx].Mass;
-#if TODO_CONVERT_GEOMETRY_COLLECTION_PARTICLE_INDICES_TO_PARTICLE_POINTERS
-									AllBreakingsDataArrayItem.ParticleIndex = RecordedFrame->Breakings[Idx].ParticleIndex;
-#endif
-									AllBreakingsDataArrayItem.ParticleIndexMesh = RecordedFrame->Breakings[Idx].ParticleIndexMesh;
-
-									AllBreakingsIndicesByPhysicsProxy[this].Add(NewIdx);
-								}
-							}
-						}
-					}
-
-					// Trailing
-					if (Parameters.TrailingData.DoGenerateTrailingData && ThisFrameTime > 0.f && Parameters.TrailingData.TrailingDataSizeMax > 0)
-					{
-						if (RecordedFrame->Trailings.Num() > 0)
-						{
-							Chaos::FPhysicsSolver::FTrailingDataArray& AllTrailingsDataArray = ThisSolver->GetAllTrailingsDataArray();
-							TMap<IPhysicsProxyBase*, TArray<int32>>& AllTrailingsIndicesByPhysicsProxy = ThisSolver->GetAllTrailingsIndicesByPhysicsProxy();
-
-							if (!AllTrailingsIndicesByPhysicsProxy.Contains(this))
-							{
-								AllTrailingsIndicesByPhysicsProxy.Add(this, TArray<int32>());
-							}
-
-							for (const FSolverTrailingData& Trailing : RecordedFrame->Trailings)
-							{
-								// Check if the particle is still kinematic
-								if (Trailing.ParticleIndex < 0 ||
-									(Trailing.ParticleIndex >= 0 && Trailing.ParticleIndex < static_cast<int32>(Particles.Size()) && Particles.ObjectState(Trailing.ParticleIndex) == Chaos::EObjectStateType::Kinematic))
-								{
-									const int32 NewIdx = AllTrailingsDataArray.Add(Chaos::TTrailingData<float, 3>());
-									Chaos::TTrailingData<float, 3>& AllTrailingsDataArrayItem = AllTrailingsDataArray[NewIdx];
-
-									AllTrailingsDataArrayItem.Location = Trailing.Location;
-									AllTrailingsDataArrayItem.Velocity = Trailing.Velocity;
-									AllTrailingsDataArrayItem.AngularVelocity = Trailing.AngularVelocity;
-									AllTrailingsDataArrayItem.Mass = Trailing.Mass;
-#if TODO_CONVERT_GEOMETRY_COLLECTION_PARTICLE_INDICES_TO_PARTICLE_POINTERS
-									AllTrailingsDataArrayItem.ParticleIndex = Trailing.ParticleIndex;
-#endif
-									AllTrailingsDataArrayItem.ParticleIndexMesh = Trailing.ParticleIndexMesh;
-
-									AllTrailingsIndicesByPhysicsProxy[this].Add(NewIdx);
-								}
-							}
-						}
-					}
-				}
-			}
-
-			if (bParticlesUpdated)
-			{
-				ThisSolver->InitializeFromParticleData(0);
-			}
-		}
-	}
-#endif
-}
-
-void FGeometryCollectionPhysicsProxy::EndFrameCallback(const float EndFrame)
-{
-	check(false); // not currently called
-#if TODO_REIMPLEMENT_GET_RIGID_PARTICLES
-	if (InitializedState == ESimulationInitializationState::Initialized)
-	{
-		FGeometryDynamicCollection* Collection = Parameters.DynamicCollection;
-		check(Collection);
-
-		ProxySimDuration += EndFrame;
-
-		if(Collection->HasAttribute("RigidBodyID", FGeometryCollection::TransformGroup))
-		{
-			//
-			//  Update transforms for the simulated transforms
-			//
-			TManagedArray<int32>& CollectionClusterID = SolverClusterID;
-			TManagedArray<FTransform>& Transform = Collection->Transform;
-			TManagedArray<int32>& Parent = Collection->Parent;
-			TManagedArray<TSet<int32>>& Children = Collection->Children;
-			TManagedArray<int32>& SimulationType = Collection->SimulationType;
-
-			TManagedArray<int32>& DynamicState = Collection->GetAttribute<int32>("DynamicState", FGeometryCollection::TransformGroup);
-
-			Chaos::TPBDRigidParticles<float, 3>& Particles = GetSolver()->GetRigidParticles();
-			const Chaos::TPBDCollisionConstraints<float, 3>& CollisionRule = GetSolver()->GetCollisionConstraints();
-			const Chaos::TArrayCollectionArray<Chaos::ClusterId>& ClusterID = GetSolver()->GetRigidClustering().GetClusterIdsArray();
-			const Chaos::TArrayCollectionArray<Chaos::TRigidTransform<float, 3>>& ClusterChildToParentMap = GetSolver()->GetRigidClustering().GetChildToParentMap();
-			const Chaos::TArrayCollectionArray<bool>& InternalCluster = GetSolver()->GetRigidClustering().GetInternalClusterArray();
-
-			//Particles X and R are aligned with center of mass and inertia principal axes.
-			//Renderer doesn't know about this and simply does ActorToWorld * GeomToActor * LocalSpaceVerts
-			//In proper math multiplication order:
-			//ParticleToWorld = ActorToWorld * GeomToActor * MassToLocal
-			//GeomToWorld = ActorToWorld * GeomToActor
-			//=> GeomToWorld = ParticleToWorld * MassToLocal.Inv()
-			//=> GeomToActor = ActorToWorld.Inv() * ParticleToWorld * MassToLocal.Inv()
-			const int32 TransformSize = Collection->NumElements(FGeometryCollection::TransformGroup);
-			const FTransform& ActorToWorld = Parameters.WorldTransform;
-
-			// It's not worth shrinking the EndFrameUnparentingBuffer array, at least until the solver supports deleting bodies.
-			if (EndFrameUnparentingBuffer.Num() < TransformSize)
-			{
-				EndFrameUnparentingBuffer.Init(-1, TransformSize);
-			}
-
-			{
-				for (int32 TransformGroupIndex = 0; TransformGroupIndex < TransformSize; ++TransformGroupIndex)
-				{
-					const int32 RigidBodyIndex = RigidBodyID[TransformGroupIndex];
-					if (RigidBodyIndex != INDEX_NONE)
-					{
-						// Update the transform and parent hierarchy of the active rigid bodies. Active bodies can be either
-						// rigid geometry defined from the leaf nodes of the collection, or cluster bodies that drive an entire
-						// branch of the hierarchy within the GeometryCollection.
-						// - Active bodies are directly driven from the global position of the corresponding
-						//   rigid bodies within the solver ( cases where RigidBodyID[TransformGroupIndex] is not disabled ). 
-						// - Deactivated bodies are driven from the transforms of their active parents. However the solver can
-						//   take ownership of the parents during the simulation, so it might be necessary to force deactivated
-						//   bodies out of the collections hierarchy during the simulation.  
-						if (!Particles.Disabled(RigidBodyID[TransformGroupIndex]))
-						{
-							// Update the transform of the active body. The active body can be either a single rigid
-							// or a collection of rigidly attached geometries (Clustering). The cluster is represented as a
-							// single transform in the GeometryCollection, and all children are stored in the local space
-							// of the parent cluster.
-							// ... When setting cluster transforms it is expected that the MassToLocal is identity.
-							//     Cluster initialization will set the vertices in the MassSpace of the rigid body.
-							// ... When setting individual rigid bodies that are not clustered, the MassToLocal will be 
-							//     non-Identity, and will reflect the difference between the geometric center of the geometry
-							//     and that corresponding rigid bodies center of mass. 
-							const FTransform ParticleToWorld(Particles.R(RigidBodyIndex), Particles.X(RigidBodyIndex));
-							// GeomToActor = ActorToWorld.Inv() * ParticleToWorld * MassToLocal.Inv();
-							Transform[TransformGroupIndex] = MassToLocal[TransformGroupIndex].GetRelativeTransformReverse(ParticleToWorld).GetRelativeTransform(ActorToWorld);
-							Transform[TransformGroupIndex].NormalizeRotation();
-
-							// dynamic state is also updated by the solver during field interaction. 
-							if (!Particles.Sleeping(RigidBodyIndex))
-							{
-								if (Particles.ObjectState(RigidBodyIndex) == Chaos::EObjectStateType::Kinematic)
-								{
-									DynamicState[TransformGroupIndex] = (int)EObjectStateTypeEnum::Chaos_Object_Kinematic;
-								}
-								else if (Particles.ObjectState(RigidBodyIndex) == Chaos::EObjectStateType::Static)
-								{
-									DynamicState[TransformGroupIndex] = (int)EObjectStateTypeEnum::Chaos_Object_Static;
-								}
-								else
-								{
-									DynamicState[TransformGroupIndex] = (int)EObjectStateTypeEnum::Chaos_Object_Dynamic;
-								}
-							}
-
-							// Force all enabled rigid bodies out of the transform hierarchy
-							if (Parent[TransformGroupIndex] != INDEX_NONE)
-							{
-								const int32 ParentIndex = Parent[TransformGroupIndex];
-								// Children in the hierarchy are stored in a TSet, which is not thread safe.  So we retain
-								// indices to remove afterwards.
-								EndFrameUnparentingBuffer[TransformGroupIndex] = ParentIndex;
-							}
-
-							// When a leaf node rigid body is removed from a cluster the rigid
-							// body will become active and needs its clusterID updated. This just
-							// syncs the clusterID all the time. 
-							CollectionClusterID[TransformGroupIndex] = ClusterID[RigidBodyIndex].Id;
-						}
-						else if (Particles.Disabled(RigidBodyIndex))
-						{
-							// dynamic state is also updated by the solver during field interaction. 
-							if (!Particles.Sleeping(RigidBodyIndex))
-							{
-								if (Particles.ObjectState(RigidBodyIndex) == Chaos::EObjectStateType::Kinematic)
-								{
-									DynamicState[TransformGroupIndex] = (int)EObjectStateTypeEnum::Chaos_Object_Kinematic;
-								}
-								else if (Particles.ObjectState(RigidBodyIndex) == Chaos::EObjectStateType::Static)
-								{
-									DynamicState[TransformGroupIndex] = (int)EObjectStateTypeEnum::Chaos_Object_Static;
-								}
-								else
-								{
-									DynamicState[TransformGroupIndex] = (int)EObjectStateTypeEnum::Chaos_Object_Dynamic;
-								}
-							}
-
-							// The rigid body parent cluster has changed within the solver, and its
-							// parent body is not tracked within the geometry collection. So we need to
-							// pull the rigid bodies out of the transform hierarchy, and just drive
-							// the positions directly from the solvers cluster particle. 
-							if (CollectionClusterID[TransformGroupIndex] != ClusterID[RigidBodyIndex].Id)
-							{
-								// Force all driven rigid bodies out of the transform hierarchy
-								if (Parent[TransformGroupIndex] != INDEX_NONE)
-								{
-									const int32 ParentIndex = Parent[TransformGroupIndex];
-									// Children in the hierarchy are stored in a TSet, which is not thread safe.  So we retain
-									// indices to remove afterwards.
-									EndFrameUnparentingBuffer[TransformGroupIndex] = ParentIndex;
-								}
-								CollectionClusterID[TransformGroupIndex] = ClusterID[RigidBodyIndex].Id;
-							}
-
-							// Disabled rigid bodies that have valid cluster parents, and have been re-indexed by the
-							// solver ( As in, They were re-clustered outside of the geometry collection), These clusters 
-							// will need to be rendered based on the clusters position. 
-							const int32 ClusterParentIndex = CollectionClusterID[TransformGroupIndex];
-							if (ClusterParentIndex != INDEX_NONE)
-							{
-								if (InternalCluster[ClusterParentIndex])
-								{
-									const FTransform ClusterChildToWorld = ClusterChildToParentMap[RigidBodyIndex] * FTransform(Particles.R(ClusterParentIndex), Particles.X(ClusterParentIndex));
-									if (Parameters.IsCacheRecording())
-									{
-										Particles.X(RigidBodyIndex) = ClusterChildToWorld.GetTranslation();
-										Particles.R(RigidBodyIndex) = ClusterChildToWorld.GetRotation();
-									}
-									// GeomToActor = ActorToWorld.Inv() * ClusterChildToWorld * MassToLocal.Inv();
-									Transform[TransformGroupIndex] = MassToLocal[TransformGroupIndex].GetRelativeTransformReverse(ClusterChildToWorld).GetRelativeTransform(ActorToWorld);
-									Transform[TransformGroupIndex].NormalizeRotation();
-								}
-							}
-						}
-					}
-				}
-			}
-			for (int32 TransformGroupIndex = 0; TransformGroupIndex < TransformSize; TransformGroupIndex++)
-			{
-				const int32 ParentIndex = EndFrameUnparentingBuffer[TransformGroupIndex];
-				if (ParentIndex >= 0)
-				{
-					// We reuse EndFrameUnparentingBuffer potentially without reinitialization, so reset this index to -1 before it gets paged out.
-					EndFrameUnparentingBuffer[TransformGroupIndex] = -1;
-
-					Children[ParentIndex].Remove(TransformGroupIndex);
-					Parent[TransformGroupIndex] = INDEX_NONE;
-				}
-			}
-
-			//
-			//  Set rest cache on simulated object.
-			//
-			if (Parameters.IsCacheRecording())
-			{
-				check(!UpdateRecordedStateCallback);
-				UpdateRecordedState(ProxySimDuration, RigidBodyID, CollectionClusterID, InternalCluster, Particles, CollisionRule);
-			}
-
-			// one way trigger from non-simulating to simulating
-			if (!IsObjectDynamic)
-			{
-				const Chaos::TArrayCollectionArray<bool>& ExternalID = GetSolver()->GetRigidClustering().GetInternalClusterArray();
-
-				for (int32 TransformGroupIndex = 0; TransformGroupIndex < TransformSize; TransformGroupIndex++)
-				{
-					const int32 RigidBodyIndex = RigidBodyID[TransformGroupIndex];
-					if (RigidBodyIndex != INDEX_NONE)
-					{
-						bool HasBeenRemoved = (PTDynamicCollection.StatusFlags[TransformGroupIndex] & FGeometryCollection::FS_RemoveOnFracture) && Particles.Disabled(RigidBodyIndex) && ClusterID[RigidBodyIndex].Id == INDEX_NONE;
-
-						if (HasBeenRemoved || !Particles.Disabled(RigidBodyIndex))
-						{
-							if (DynamicState[TransformGroupIndex] != (int32)EObjectStateTypeEnum::Chaos_Object_Static)
-							{
-								IsObjectDynamic = true;
-								break;
-							}
-						}
-						else
-						{
-							const int32 ParentClusterID = ClusterID[RigidBodyIndex].Id;
-							if (ParentClusterID != INDEX_NONE && ExternalID[ParentClusterID]
-								&& Particles.ObjectState(ParentClusterID) != Chaos::EObjectStateType::Static)
-							{
-								IsObjectDynamic = true;
-								break;
-							}
-						}
-					}
-				}
-			}
-
-			// Can't change visibility right now so setting scale to zero instead
-			// only process if enabled on this object (ones with glass), most will skip over
-			if (Parameters.RemoveOnFractureEnabled && IsObjectDynamic)
-			{
-				for (int32 TransformGroupIndex = 0; TransformGroupIndex < TransformSize; TransformGroupIndex++)
-				{
-					const int32 RigidBodyIndex = RigidBodyID[TransformGroupIndex];
-					if (RigidBodyIndex != INDEX_NONE)
-					{
-						if ((PTDynamicCollection.StatusFlags[TransformGroupIndex] & FGeometryCollection::FS_RemoveOnFracture)
-							&& Particles.Disabled(RigidBodyIndex) && ClusterID[RigidBodyIndex].Id == INDEX_NONE)
-						{
-							PTDynamicCollection.Transform[TransformGroupIndex].SetScale3D(FVector::ZeroVector);
-						}
-					}
-				}
-			}
-		}
-
-#if TODO_REIMPLEMENT_RIGID_CLUSTERING
-		{
-			IsObjectLoading = GetSolver()->GetRigidClustering().NumberOfPendingClusters() != 0;
-		}
-#endif
-	}
-	Commands.Empty();
-#endif
-
-}
-
-void FGeometryCollectionPhysicsProxy::CreateRigidBodyCallback(FParticlesType& Particles)
-{
-	check(false); // This function is deprecated and should never be called.
-#if TODO_REIMPLEMENT_RIGID_CLUSTERING
-
-	const FGeometryCollection* RestCollection = Parameters.RestCollection;
-	FGeometryDynamicCollection* DynamicCollection = Parameters.DynamicCollection;
-	check(DynamicCollection);
-	
-	if (Parameters.Simulating && 
-		((InitializedState == ESimulationInitializationState::Unintialized) || 
-		 (InitializedState == ESimulationInitializationState::Activated)))
-	{
-		Chaos::FPhysicsSolver* ThisSolver = GetSolver();
-		Chaos::TArrayCollectionArray<int32>& ClusterGroupIndex = ThisSolver->GetRigidClustering().GetClusterGroupIndexArray();
-		Chaos::TArrayCollectionArray<float>& StrainArray = ThisSolver->GetRigidClustering().GetStrainArray();
-		const float StrainDefault = Parameters.DamageThreshold.Num() ? Parameters.DamageThreshold[0] : 0;
-
-		const TManagedArray<int32>& TransformIndex = RestCollection->TransformIndex;
-		const TManagedArray<int32>& BoneMap = RestCollection->BoneMap;
-		const TManagedArray<int32>& Parent = RestCollection->Parent;
-		const TManagedArray<TSet<int32>>& Children = RestCollection->Children;
-		const TManagedArray<int32>& SimulationType = RestCollection->SimulationType;
-		const TManagedArray<FVector>& Vertex = RestCollection->Vertex;
-		const TManagedArray<int32>& DynamicState = DynamicCollection->DynamicState;
-		const TManagedArray<int32>& CollisionGroup = DynamicCollection->CollisionGroup;
-		const TManagedArray<float>& Mass = RestCollection->GetAttribute<float>("Mass", FTransformCollection::TransformGroup);
-		const TManagedArray<FVector>& InertiaTensor = RestCollection->GetAttribute<FVector>("InertiaTensor", FTransformCollection::TransformGroup);
-
-		TArray<FTransform> Transform;
-		GeometryCollectionAlgo::GlobalMatrices(DynamicCollection->Transform, DynamicCollection->Parent, Transform);
-		check(DynamicCollection->Transform.Num() == Transform.Num());
-
-		// count particles to add
-		int NumSimulatedParticles = 0;
-		for (int32 Index = 0; Index < SimulatableParticles.Num(); Index++)
-		{
-			if (SimulatableParticles[Index] == true)
-			{
-				NumSimulatedParticles++;
-			}
-		}
-
-		// Add entries into simulation array
-		int NumRigids = Particles.Size();
-		BaseParticleIndex = NumRigids;
-		Particles.AddParticles(NumSimulatedParticles);
-		for (int32 Index = 0, NextId = 0; Index < SimulatableParticles.Num(); Index++)
-		{
-			if (SimulatableParticles[Index] == true)
-			{
-				RigidBodyID[Index] = NumRigids + NextId++;
-			}
-		}
-
-		// Add the rigid bodies
-		//for (int32 Index = 0; Index < TransformIndex.Num(); Index++)
-		const int32 NumGeometries = DynamicCollection->NumElements(FGeometryCollection::GeometryGroup);
-		ParallelFor(NumGeometries, [&](int32 GeometryIndex)
-		{
-			const int32 TransformGroupIndex = TransformIndex[GeometryIndex];
-			if (SimulatableParticles[TransformGroupIndex] == true)
-			{
-				const int32 RigidBodyIndex = RigidBodyID[TransformGroupIndex];
-				const FTransform WorldTransform = MassToLocal[TransformGroupIndex] * Transform[TransformGroupIndex] * Parameters.WorldTransform;
-
-				PopulateSimulatedParticle(Particles, Parameters.Shared, nullptr,// Simplicials[TransformGroupIndex].Get(),
-					Implicits[TransformGroupIndex], ThisSolver->GetMassScale() * Mass[TransformGroupIndex], ThisSolver->GetMassScale() * InertiaTensor[TransformGroupIndex],
-					RigidBodyIndex, WorldTransform, (uint8)EObjectStateTypeEnum::Chaos_Object_Dynamic, CollisionGroup[TransformGroupIndex]);
-				ClusterGroupIndex[RigidBodyIndex] = Parameters.ClusterGroupIndex;
-				Particles.CollisionParticles(RigidBodyIndex).Reset(Simplicials[TransformGroupIndex].Release());
-				if (Particles.CollisionParticles(RigidBodyIndex))
-				{
-					int32 NumCollisionParticles = Particles.CollisionParticles(RigidBodyIndex)->Size();
-					int32 CollisionParticlesSize = FMath::Max(0, FMath::Min(int(NumCollisionParticles*CollisionParticlesPerObjectFraction), NumCollisionParticles));
-					Particles.CollisionParticles(RigidBodyIndex)->Resize(CollisionParticlesSize);
-				}
-
-				StrainArray[RigidBodyIndex] = StrainDefault;
-				GetSolver()->SetPhysicsMaterial(RigidBodyIndex, Parameters.PhysicalMaterial);
-			}
-
-		});
-
-		for (FFieldSystemCommand & Cmd : Parameters.InitializationCommands)
-		{
-			if (Cmd.MetaData.Contains(FFieldSystemMetaData::EMetaType::ECommandData_ProcessingResolution)) Cmd.MetaData.Remove(FFieldSystemMetaData::EMetaType::ECommandData_ProcessingResolution);
-			FFieldSystemMetaDataProcessingResolution* ResolutionData = new FFieldSystemMetaDataProcessingResolution(EFieldResolutionType::Field_Resolution_Maximum);
-			Cmd.MetaData.Add(FFieldSystemMetaData::EMetaType::ECommandData_ProcessingResolution, TUniquePtr< FFieldSystemMetaDataProcessingResolution >(ResolutionData));
-			Commands.Add(Cmd);
-		}
-		Parameters.InitializationCommands.Empty();
-		ProcessCommands(Particles, GetSolver()->GetSolverTime());
-
-		ParallelFor(NumGeometries, [&](int32 GeometryIndex)
-		{
-			const int32 TransformGroupIndex = TransformIndex[GeometryIndex];
-			if (SimulatableParticles[TransformGroupIndex] == true)
-			{
-				const int32 RigidBodyIndex = RigidBodyID[TransformGroupIndex];
-
-				if (Parameters.InitialVelocityType == EInitialVelocityTypeEnum::Chaos_Initial_Velocity_User_Defined)
-				{
-					InitialLinearVelocity[TransformGroupIndex] = Parameters.InitialLinearVelocity;
-					InitialAngularVelocity[TransformGroupIndex] = Parameters.InitialAngularVelocity;
-					if (DynamicState[TransformGroupIndex] == (int32)EObjectStateTypeEnum::Chaos_Object_Dynamic)
-					{
-						Particles.V(RigidBodyIndex) = Parameters.InitialLinearVelocity;
-						Particles.W(RigidBodyIndex) = Parameters.InitialAngularVelocity;
-					}
-				}
-			}
-		});
-
-		// InitializeKinematics
-		if (Parameters.DynamicCollection)
-		{
-			for (int32 TransformGroupIndex = 0; TransformGroupIndex < GetTransformGroupSize(); ++TransformGroupIndex)
-			{
-				if (Chaos::TPBDRigidParticleHandle<float, 3>* Handle = SolverParticleHandles[TransformGroupIndex])
-				{
-					// this is being initialized to Static, can't change from static to sleeping
-					Handle->SetObjectState(Chaos::EObjectStateType::Dynamic);
-
-					if (DynamicState[TransformGroupIndex] == (int32)EObjectStateTypeEnum::Chaos_Object_Kinematic)
-					{
-						Handle->SetObjectState(Chaos::EObjectStateType::Kinematic);
-					}
-					else if (DynamicState[TransformGroupIndex] == (int32)EObjectStateTypeEnum::Chaos_Object_Static)
-					{
-						Handle->SetObjectState(Chaos::EObjectStateType::Static);
-					}
-					else if (DynamicState[TransformGroupIndex] == (int32)EObjectStateTypeEnum::Chaos_Object_Sleeping)
-					{
-						Handle->SetObjectState(Chaos::EObjectStateType::Sleeping);
-					}
-				}
-			}
-		}
-
-		InitializeRemoveOnFracture(Particles, DynamicState);
-
-		// #BG Temporary - don't cluster when playing back. Needs to be changed when kinematics are per-proxy to support
-		// kinematic to dynamic transition for clusters.
-		if (Parameters.EnableClustering)// && Parameters.CacheType != EGeometryCollectionCacheType::Play)
-		{
-			const TArray<int32> RecursiveOrder = ComputeRecursiveOrder(*RestCollection);
-
-			// num clusters
-			uint32 NumClusters=0;
-			TArray<bool> SubTreeContainsSimulatableParticle;
-			SubTreeContainsSimulatableParticle.SetNum(RecursiveOrder.Num());
-			for (const int32 TransformGroupIndex : RecursiveOrder)
-			{
-				if (Children[TransformGroupIndex].Num() > 0)
-				{
-					SubTreeContainsSimulatableParticle[TransformGroupIndex] = false;
-
-					TArray<uint32> RigidChildren, CollectionChildren;
-					for (const int32 ChildIndex : Children[TransformGroupIndex])
-					{
-						if(SubTreeContainsSimulatableParticle[ChildIndex])
-						{
-							NumClusters++;
-							SubTreeContainsSimulatableParticle[TransformGroupIndex] = true;
-							break;
-						}
-					}
-				}
-				else
-				{
-					SubTreeContainsSimulatableParticle[TransformGroupIndex] = (RigidBodyID[TransformGroupIndex] != INDEX_NONE);
-				}
-			}
-
-			const int32 ClusterStartIndex = Particles.Size();
-			Particles.AddParticles(NumClusters);
-
-			int32 ClusterRigidBodyId = ClusterStartIndex;
-			for (const int32 TransformGroupIndex : RecursiveOrder)
-			{
-				TArray<uint32> RigidChildren, CollectionChildren;
-				for (const int32 ChildIndex : Children[TransformGroupIndex])
-				{
-					if (RigidBodyID[ChildIndex] != INDEX_NONE)
-					{
-						RigidChildren.Add(RigidBodyID[ChildIndex]);
-						CollectionChildren.Add(ChildIndex);
-					}
-				}
-				if (RigidChildren.Num())
-				{
-					if (ReportTooManyChildrenNum >= 0 && RigidChildren.Num() > ReportTooManyChildrenNum)
-					{
-						UE_LOG(LogChaos, Warning, TEXT("Too many children (%d) in a single cluster:%s"), RigidChildren.Num(), *Parameters.Name);
-					}
-
-					RigidBodyID[TransformGroupIndex] = ClusterRigidBodyId;
-					Chaos::FClusterCreationParameters<float> CreationParameters;
-					CreationParameters.RigidBodyIndex = RigidBodyID[TransformGroupIndex];
-					BuildClusters(TransformGroupIndex, CollectionChildren, RigidChildren, CreationParameters);
-					ClusterRigidBodyId++;
-				}
-			}
-		}
-
-		NumParticles = Particles.Size() - BaseParticleIndex;
-
-		// Set Connectivity
-		const int32 NumTransforms = DynamicCollection->NumElements(FGeometryCollection::TransformGroup);
-		ParallelFor(NumTransforms, [&](int32 TransformGroupIndex)
-		{
-			if (Children[TransformGroupIndex].Num() > 0)
-			{
-				if (RigidBodyID[TransformGroupIndex] != INDEX_NONE)
-				{
-					GetSolver()->GetRigidClustering().GenerateConnectionGraph(RigidBodyID[TransformGroupIndex]);
-				}
-			}
-		});
-
-		// If we're recording and want to start immediately caching then we should cache the rest state
-		if (Parameters.IsCacheRecording() && Parameters.CacheBeginTime == 0.0f)
-		{
-			if (UpdateRecordedStateCallback)
-			{
-				UpdateRecordedStateCallback(0.0f, RigidBodyID, Particles, GetSolver()->GetCollisionConstraints());
-			}
-		}
-
-
-		if (InitializedState == ESimulationInitializationState::Activated)
-		{
-			//
-			//  Activated bodies has already been called so we are good to go.
-			//
-			InitializedState = ESimulationInitializationState::Initialized;
-
-			if (Parameters.EnableClustering && Parameters.ClusterGroupIndex)
-			{
-				GetSolver()->GetRigidClustering().IncrementPendingClusterCounter(Parameters.ClusterGroupIndex);
-				GetSolver()->GetRigidClustering().DecrementPendingClusterCounter(Parameters.ClusterGroupIndex);
-			}
-		}
-		else if (InitializedState == ESimulationInitializationState::Unintialized)
-		{
-
-			//
-			//  Activated bodies has not been called, se we are waiting 
-			//  to become active. Deactivate all bodies, and wait for
-			//  ActivatedBodies to be called, and defer the cluster initialization
-			//
-			InitializedState = ESimulationInitializationState::Created;
-			for (int32 TransformGroupIndex = 0; TransformGroupIndex < NumTransforms; TransformGroupIndex++)
-			{
-				const int32 RigidBodyIndex = RigidBodyID[TransformGroupIndex];
-				if (RigidBodyIndex != INDEX_NONE)
-				{
-					if (!Particles.Disabled(RigidBodyIndex))
-					{
-						PendingActivationList.Add(TransformGroupIndex);
-						GetSolver()->GetEvolution()->DisableParticle(RigidBodyIndex);
-					}
-				}
-			}
-
-			//
-			//  Clustering needs to advertise its group id to the cluster so
-			//  that the group is not initialized before all the bodies are
-			//  loaded and created. 
-			//
-			if (Parameters.EnableClustering && Parameters.ClusterGroupIndex )
-			{
-				GetSolver()->GetRigidClustering().IncrementPendingClusterCounter(Parameters.ClusterGroupIndex);
-			}
-		}
-		else
-		{
-			// unknown initialization state in creation callback
-			ensure(false);
-		}
-	}
-#endif
-}
-
-
-void FGeometryCollectionPhysicsProxy::BindParticleCallbackMapping(Chaos::TArrayCollectionArray<PhysicsProxyWrapper>& PhysicsProxyReverseMap, Chaos::TArrayCollectionArray<int32>& ParticleIDReverseMap)
-{
-	/*
-	for (int32 Index = 0; Index < RigidBodyID.Num(); Index++)
-	{
-		const int32 RigidBodyIndex = RigidBodyID[Index];
-		if (RigidBodyIndex != INDEX_NONE)
-		{
-			PhysicsProxyReverseMap[RigidBodyIndex] = { this, EPhysicsProxyType::GeometryCollectionType };
-			ParticleIDReverseMap[RigidBodyIndex] = Index;
-		}
-	}
-	*/
-}
-
-void FGeometryCollectionPhysicsProxy::PushKinematicStateToSolver(FParticlesType& Particles)
-{
-#if TODO_REIMPLEMENT_RIGID_CLUSTERING
-	if (InitializedState == ESimulationInitializationState::Initialized)
-	{
-		FGeometryDynamicCollection* Collection = Parameters.DynamicCollection;
-		if (Collection->Transform.Num())
-		{
-			TManagedArray<int32>& DynamicState = Collection->GetAttribute<int32>("DynamicState", FGeometryCollection::TransformGroup);
-
-			TSet<int32> ClustersToUpdate;
-			const Chaos::TArrayCollectionArray<Chaos::ClusterId>& ClusterID = GetSolver()->GetRigidClustering().GetClusterIdsArray();
-
-			for (int32 TransformGroupIndex = 0; TransformGroupIndex < DynamicState.Num(); TransformGroupIndex++)
-			{
-				if (RigidBodyID[TransformGroupIndex] != INDEX_NONE)
-				{
-					int32 RigidBodyIndex = RigidBodyID[TransformGroupIndex];
-					if (DynamicState[TransformGroupIndex] == (int32)EObjectStateTypeEnum::Chaos_Object_Dynamic
-						&& (Particles.ObjectState(RigidBodyIndex) == Chaos::EObjectStateType::Kinematic || Particles.ObjectState(RigidBodyIndex) == Chaos::EObjectStateType::Static)
-						&& FLT_EPSILON < Particles.M(RigidBodyIndex))
-					{
-						Particles.SetObjectState(RigidBodyIndex, Chaos::EObjectStateType::Dynamic);
-
-						if (Parameters.InitialVelocityType == EInitialVelocityTypeEnum::Chaos_Initial_Velocity_User_Defined)
-						{
-							Particles.V(RigidBodyIndex) = Chaos::TVector<float, 3>(InitialLinearVelocity[TransformGroupIndex]);
-							Particles.W(RigidBodyIndex) = Chaos::TVector<float, 3>(InitialAngularVelocity[TransformGroupIndex]);
-						}
-
-						if (!Particles.Disabled(RigidBodyIndex) && Particles.Sleeping(RigidBodyIndex))
-						{
-							Particles.SetSleeping(RigidBodyIndex, false);
-							GetSolver()->ActiveIndices().Add(RigidBodyIndex);
-						}
-						else
-						{
-							if (ClusterID[RigidBodyIndex].Id != INDEX_NONE)
-							{
-								int32 ActiveClusterIndex = GetSolver()->GetRigidClustering().GetActiveClusterIndex(RigidBodyIndex);
-								if (ActiveClusterIndex != INDEX_NONE)
-								{
-									ClustersToUpdate.Add(ActiveClusterIndex);
-								}
-							}
-						}
-					}
-					else if ((DynamicState[TransformGroupIndex] == (int32)EObjectStateTypeEnum::Chaos_Object_Kinematic)
-						&& (Particles.ObjectState(RigidBodyIndex) == Chaos::EObjectStateType::Dynamic)
-						&& FLT_EPSILON < Particles.M(RigidBodyIndex))
-					{
-						Particles.SetObjectState(RigidBodyIndex, Chaos::EObjectStateType::Kinematic);
-
-						if (ClusterID[RigidBodyIndex].Id != INDEX_NONE)
-						{
-							int32 ActiveClusterIndex = GetSolver()->GetRigidClustering().GetActiveClusterIndex(RigidBodyIndex);
-							if (ActiveClusterIndex != INDEX_NONE)
-							{
-								ClustersToUpdate.Add(ActiveClusterIndex);
-							}
-						}
-					}
-					else if ((DynamicState[TransformGroupIndex] == (int32)EObjectStateTypeEnum::Chaos_Object_Static)
-						&& (Particles.ObjectState(RigidBodyIndex) == Chaos::EObjectStateType::Dynamic)
-						&& FLT_EPSILON < Particles.M(RigidBodyIndex))
-					{
-						Particles.SetObjectState(RigidBodyIndex, Chaos::EObjectStateType::Static);
-
-						if (ClusterID[RigidBodyIndex].Id != INDEX_NONE)
-						{
-							int32 ActiveClusterIndex = GetSolver()->GetRigidClustering().GetActiveClusterIndex(RigidBodyIndex);
-							if (ActiveClusterIndex != INDEX_NONE)
-							{
-								ClustersToUpdate.Add(ActiveClusterIndex);
-							}
-						}
-					}
-					else if ((DynamicState[TransformGroupIndex] == (int32)EObjectStateTypeEnum::Chaos_Object_Sleeping)
-						&& (Particles.ObjectState(RigidBodyIndex) == Chaos::EObjectStateType::Dynamic))
-					{
-						Particles.SetObjectState(RigidBodyIndex, Chaos::EObjectStateType::Sleeping);
-						GetSolver()->ActiveIndices().Remove(RigidBodyIndex);
-					}
-					else if ((DynamicState[TransformGroupIndex] == (int32)EObjectStateTypeEnum::Chaos_Object_Dynamic)
-						&& (Particles.ObjectState(RigidBodyIndex) == Chaos::EObjectStateType::Sleeping))
-					{
-						Particles.SetObjectState(RigidBodyIndex, Chaos::EObjectStateType::Dynamic);
-						GetSolver()->ActiveIndices().Add(RigidBodyIndex);
-					}
-				}
-			}
-		}
-	}
-#endif
-}
-
-void FGeometryCollectionPhysicsProxy::ParameterUpdateCallback(FParticlesType& Particles, const float Time)
-{
-	FGeometryDynamicCollection* Collection = Parameters.DynamicCollection;
-	check(Collection);
+	FGeometryDynamicCollection* Collection = Parameters.DynamicCollection;
+	check(Collection);
 
 	if (Collection->Transform.Num())
 	{
 		ProcessCommands(Particles, Time);
-
-
-#ifndef TODO_REIMPLEMENT_RIGID_CACHING
-		if (Parameters.RecordedTrack)
-		{
-			float ReverseTime = Parameters.RecordedTrack->GetLastTime() - Time + Parameters.ReverseCacheBeginTime;
-			// @todo(mlentine): We shouldn't need to do this every frame
-			if (Parameters.IsCacheRecording() && Time > Parameters.ReverseCacheBeginTime&& Parameters.ReverseCacheBeginTime != 0 && Parameters.RecordedTrack->IsTimeValid(ReverseTime))
-			{
-				for (int32 Index = 0; Index < RigidBodyID.Num(); Index++)
-				{
-					int32 RigidBodyIndex = RigidBodyID[Index];
-
-					// Check index, will be invalid for cluster parents.
-					if (RigidBodyIndex != INDEX_NONE)
-					{
-						Particles.InvM(RigidBodyIndex) = 0.f;
-						Particles.InvI(RigidBodyIndex) = Chaos::PMatrix<float, 3, 3>(0.f);
-					}
-				}
-			}
-		}
-#endif
 	}
 }
 
-void FGeometryCollectionPhysicsProxy::DisableCollisionsCallback(TSet<TTuple<int32, int32>>& InPairs)
-{
-
-}
-
-void FGeometryCollectionPhysicsProxy::AddForceCallback(FParticlesType& InParticles, const float InDt, const int32 InIndex)
-{
-	// gravity forces managed directly on the solver for now
-}
-
diff --git a/Engine/Source/Runtime/Experimental/ChaosSolvers/Private/PhysicsProxy/GeometryCollectionPhysicsProxyOLD.cpp b/Engine/Source/Runtime/Experimental/ChaosSolvers/Private/PhysicsProxy/GeometryCollectionPhysicsProxyOLD.cpp
new file mode 100644
--- /dev/null
+++ b/Engine/Source/Runtime/Experimental/ChaosSolvers/Private/PhysicsProxy/GeometryCollectionPhysicsProxyOLD.cpp
@@ -0,0 +1,4413 @@
+// Copyright Epic Games, Inc. All Rights Reserved.
+#ifdef __USING_OLD_STUB_CODE__
+
+#include "PhysicsProxy/GeometryCollectionPhysicsProxy.h"
+
+#include "PhysicsSolver.h"
+#include "ChaosStats.h"
+#include "Chaos/ArrayCollectionArray.h"
+#include "Chaos/BVHParticles.h"
+#include "Chaos/Transform.h"
+#include "Chaos/ParallelFor.h"
+#include "Chaos/Particles.h"
+#include "Chaos/TriangleMesh.h"
+#include "Chaos/MassProperties.h"
+#include "ChaosSolversModule.h"
+#include "Chaos/PBDCollisionConstraintsUtil.h"
+#include "Chaos/PerParticleGravity.h"
+#include "Chaos/ImplicitObject.h"
+#include "Chaos/Serializable.h"
+#include "Chaos/ErrorReporter.h"
+#include "Chaos/PBDRigidClustering.h"
+#include "GeometryCollection/GeometryCollectionAlgo.h"
+#include "GeometryCollection/GeometryCollectionUtility.h"
+#include "GeometryCollection/GeometryCollectionSimulationTypes.h"
+#include "GeometryCollection/ManagedArrayCollection.h"
+#include "Modules/ModuleManager.h"
+
+#ifndef TODO_REIMPLEMENT_INIT_COMMANDS
+#define TODO_REIMPLEMENT_INIT_COMMANDS 0
+#endif
+
+#ifndef TODO_REIMPLEMENT_FRACTURE
+#define TODO_REIMPLEMENT_FRACTURE 0
+#endif
+
+#ifndef TODO_REIMPLEMENT_RIGID_CACHING
+#define TODO_REIMPLEMENT_RIGID_CACHING 0
+#endif
+
+float CollisionParticlesPerObjectFractionDefault = 1.0f;
+FAutoConsoleVariableRef CVarCollisionParticlesPerObjectFractionDefault(
+	TEXT("p.CollisionParticlesPerObjectFractionDefault"), 
+	CollisionParticlesPerObjectFractionDefault, 
+	TEXT("Fraction of verts"));
+
+bool DisableGeometryCollectionGravity = false;
+FAutoConsoleVariableRef CVarGeometryCollectionDisableGravity(
+	TEXT("p.GeometryCollectionDisableGravity"),
+	DisableGeometryCollectionGravity,
+	TEXT("Disable gravity for geometry collections"));
+
+
+DEFINE_LOG_CATEGORY_STATIC(UGCC_LOG, Error, All);
+
+//==============================================================================
+// FGeometryCollectionResults
+//==============================================================================
+
+FGeometryCollectionResults::FGeometryCollectionResults()
+	: BaseIndex(0)
+	, NumParticlesAdded(0)
+	, IsObjectDynamic(false)
+	, IsObjectLoading(false)
+	, WorldBounds(ForceInit)
+{}
+
+void FGeometryCollectionResults::Reset()
+{
+	BaseIndex = 0;
+	NumParticlesAdded = 0;
+	DisabledStates.SetNum(0);
+	//Transforms.Resize(0);
+	GlobalTransforms.SetNum(0);
+	//RigidBodyIds.Resize(0);
+	ParticleToWorldTransforms.SetNum(0);
+	//Level.Resize(0);
+	//Parent.Resize(0);
+	//Children.Resize(0);
+	//SimulationType.Resize(0);
+	//StatusFlags.Resize(0);
+	IsObjectDynamic = false;
+	IsObjectLoading = false;
+	WorldBounds = FBoxSphereBounds(ForceInit);
+}
+
+//==============================================================================
+// FGeometryCollectionPhysicsProxy helper functions
+//==============================================================================
+
+bool IsMultithreaded()
+{
+	if(FChaosSolversModule* Module = FModuleManager::Get().GetModulePtr<FChaosSolversModule>("ChaosSolvers"))
+	{
+		return Module->GetDispatcher() && 
+			Module->GetDispatcher()->GetMode() == Chaos::EThreadingMode::DedicatedThread && 
+			Module->IsPersistentTaskRunning();
+	}
+	return false;
+}
+
+Chaos::TTriangleMesh<float>* CreateTriangleMesh(
+	const int32 FaceCount, 
+	const int32 VertexOffset, 
+	const int32 StartIndex, 
+	const TManagedArray<FVector>& Vertex, 
+	const TManagedArray<bool>& Visible, 
+	const TManagedArray<FIntVector>& Indices, 
+	TSet<int32>& VertsAdded)
+{
+	TArray<Chaos::TVector<int32, 3>> Faces;
+	{
+		Faces.Reserve(FaceCount);
+		for (int j = 0; j < FaceCount; j++)
+		{
+			if (!Visible[j + StartIndex])
+			{
+				continue;
+			}
+
+			// Note: This function used to cull small triangles.  As one of the purposes 
+			// of the tri mesh this function creates is for level set rasterization, we 
+			// don't want to do that.  Keep the mesh intact, which hopefully is water tight.
+
+			const FIntVector& Tri = Indices[j + StartIndex];
+			Faces.Add(Chaos::TVector<int32, 3>(Tri.X, Tri.Y, Tri.Z));
+			for (int Axis = 0; Axis < 3; ++Axis)
+			{
+				VertsAdded.Add(Tri[Axis]);
+			}
+		}
+	}
+
+	return new Chaos::TTriangleMesh<float>(MoveTemp(Faces)); // Culls geometrically degenerate faces
+}
+
+TArray<int32> ComputeTransformToGeometryMap(const FGeometryCollection& Collection)
+{
+	const int32 NumTransforms = Collection.NumElements(FGeometryCollection::TransformGroup);
+	const int32 NumGeometries = Collection.NumElements(FGeometryCollection::GeometryGroup);
+	const TManagedArray<int32>& TransformIndex = Collection.TransformIndex;
+
+	TArray<int32> TransformToGeometryMap;
+	TransformToGeometryMap.AddUninitialized(NumTransforms);
+	for(int32 GeometryIndex = 0; GeometryIndex < NumGeometries; ++GeometryIndex)
+	{
+		const int32 TransformGroupIndex = TransformIndex[GeometryIndex];
+		TransformToGeometryMap[TransformGroupIndex] = GeometryIndex;
+	}
+
+	return TransformToGeometryMap;
+}
+
+//Computes the order of transform indices so that children in a tree always appear before their parents. Handles forests
+TArray<int32> ComputeRecursiveOrder(const FGeometryCollection& Collection)
+{
+	const int32 NumTransforms = Collection.NumElements(FGeometryCollection::TransformGroup);
+	const TManagedArray<int32>& Parent = Collection.Parent;
+	const TManagedArray<TSet<int32>>& Children = Collection.Children;
+
+	//traverse cluster hierarchy in depth first and record order
+	struct FClusterProcessing
+	{
+		int32 TransformGroupIndex;
+		enum
+		{
+			None,
+			VisitingChildren
+		} State;
+
+		FClusterProcessing(int32 InIndex) : TransformGroupIndex(InIndex), State(None) {};
+	};
+
+	TArray<FClusterProcessing> ClustersToProcess;
+	//enqueue all roots
+	for(int32 TransformGroupIndex = 0; TransformGroupIndex < NumTransforms; TransformGroupIndex++)
+	{
+		if(Parent[TransformGroupIndex] == FGeometryCollection::Invalid && Children[TransformGroupIndex].Num() > 0)
+		{
+			ClustersToProcess.Emplace(TransformGroupIndex);
+		}
+	}
+
+	TArray<int32> TransformOrder;
+	TransformOrder.Reserve(NumTransforms);
+
+	while(ClustersToProcess.Num())
+	{
+		FClusterProcessing CurCluster = ClustersToProcess.Pop();
+		const int32 ClusterTransformIdx = CurCluster.TransformGroupIndex;
+		if(CurCluster.State == FClusterProcessing::VisitingChildren)
+		{
+			//children already visited
+			TransformOrder.Add(ClusterTransformIdx);
+		}
+		else
+		{
+			if(Children[ClusterTransformIdx].Num())
+			{
+				CurCluster.State = FClusterProcessing::VisitingChildren;
+				ClustersToProcess.Add(CurCluster);
+
+				//order of children doesn't matter as long as all children appear before parent
+				for(int32 ChildIdx : Children[ClusterTransformIdx])
+				{
+					ClustersToProcess.Emplace(ChildIdx);
+				}
+			}
+			else
+			{
+				TransformOrder.Add(ClusterTransformIdx);
+			}
+		}
+	}
+
+	return TransformOrder;
+}
+
+DECLARE_CYCLE_STAT(TEXT("FGeometryCollectionPhysicsProxy::PopulateSimulatedParticle"), STAT_PopulateSimulatedParticle, STATGROUP_Chaos);
+void PopulateSimulatedParticle(
+	Chaos::TPBDRigidParticleHandle<float,3>* Handle,
+	const FSharedSimulationParameters& SharedParams,
+	const FCollisionStructureManager::FSimplicial* Simplicial,
+	TSharedPtr<Chaos::FImplicitObject, ESPMode::ThreadSafe> Implicit,
+	const float MassIn,
+	const FVector& InertiaTensorVec,
+	const FTransform& WorldTransform, 
+	const uint8 DynamicState, 
+	const int16 CollisionGroup)
+{
+	SCOPE_CYCLE_COUNTER(STAT_PopulateSimulatedParticle);
+
+	Handle->SetDisabledLowLevel(false);
+
+	Handle->SetX(WorldTransform.GetTranslation());
+	Handle->SetV(Chaos::TVector<float, 3>(0.f));
+	Handle->SetR(WorldTransform.GetRotation().GetNormalized());
+	Handle->SetW(Chaos::TVector<float, 3>(0.f));
+	Handle->SetP(Handle->X());
+	Handle->SetQ(Handle->R());
+	Handle->SetIsland(INDEX_NONE);
+
+	//todo: If mass is too small use the right inertia
+	ensureMsgf(MassIn >= SharedParams.MinimumMassClamp, 
+		TEXT("Mass smaller than minimum mass clamp. Too late to change"));
+	ensureMsgf(Handle->InvM() > SMALL_NUMBER, 
+		TEXT("Object mass is too large. Too late to change"));
+	ensureMsgf(InertiaTensorVec[0] > SMALL_NUMBER && InertiaTensorVec[1] > SMALL_NUMBER && InertiaTensorVec[2] > SMALL_NUMBER, 
+		TEXT("Inertia tensor is too small. Too late to change"));
+
+	Handle->SetM(MassIn);
+	Handle->SetInvM(1.0f / MassIn);
+	if (FMath::IsNaN(InertiaTensorVec[0]) || FMath::IsNaN(InertiaTensorVec[1]) || FMath::IsNaN(InertiaTensorVec[2]) ||
+		InertiaTensorVec[0] < SMALL_NUMBER || InertiaTensorVec[1] < SMALL_NUMBER || InertiaTensorVec[2] < SMALL_NUMBER)
+	{
+		Handle->SetI(Chaos::PMatrix<float, 3, 3>(1.f, 1.f, 1.f));
+	}
+	else
+	{
+		Handle->SetI(Chaos::PMatrix<float, 3, 3>(InertiaTensorVec[0], InertiaTensorVec[1], InertiaTensorVec[2]));
+		Handle->SetInvI(Chaos::PMatrix<float, 3, 3>(1.0f / InertiaTensorVec[0], 1.0f / InertiaTensorVec[1], 1.0f / InertiaTensorVec[2]));
+	}
+
+	// for validation set the body to dynamic and check the inverse masses. 
+	Handle->SetObjectStateLowLevel(Chaos::EObjectStateType::Dynamic);
+
+	const Chaos::PMatrix<float,3,3>& InvI = Handle->InvI();
+	if (!(InvI.M[0][0] > SMALL_NUMBER && 
+		  InvI.M[1][1] > SMALL_NUMBER && 
+		  InvI.M[2][2] > SMALL_NUMBER))
+	{
+		UE_LOG(LogChaos, Warning, TEXT("Inertia tensor is too large. Too late to change"));
+	}
+
+	Handle->SetCollisionGroup(CollisionGroup);
+
+	if (Implicit)	//todo(ocohen): this is only needed for cases where clusters have no proxy. Kind of gross though, should refactor
+	{
+		Handle->SetSharedGeometry(Implicit);
+	}
+
+	Handle->CollisionParticlesInitIfNeeded();
+	if (Simplicial)
+	{
+		TUniquePtr<Chaos::TBVHParticles<float, 3>>& CollisionParticles = Handle->CollisionParticles();
+		if (Simplicial->Size())
+		{
+			const Chaos::TAABB<float,3> ImplicitShapeDomain = 
+				Implicit && Implicit->GetType() == Chaos::ImplicitObjectType::LevelSet && Implicit->HasBoundingBox() ? 
+				Implicit->BoundingBox() : Chaos::TAABB<float,3>::FullAABB();
+
+			CollisionParticles->Resize(0);
+			CollisionParticles->AddParticles(Simplicial->Size());
+			for (int32 VertexIndex = 0; VertexIndex < (int32)Simplicial->Size(); ++VertexIndex)
+			{
+				CollisionParticles->X(VertexIndex) = Simplicial->X(VertexIndex);
+
+				// Make sure the collision particles are at least in the domain 
+				// of the implicit shape.
+				ensure(ImplicitShapeDomain.Contains(CollisionParticles->X(VertexIndex)));
+			}
+		}
+
+		// @todo(remove): IF there is no simplicial we should not be forcing one. 
+		if (!CollisionParticles->Size())
+		{
+			CollisionParticles->AddParticles(1);
+			CollisionParticles->X(0) = Chaos::TVector<float, 3>(0);
+		}
+		CollisionParticles->UpdateAccelerationStructures();
+	}
+
+	//
+	//  Manage Object State
+	//
+
+	// Only sleep if we're not replaying a simulation
+	// #BG TODO If this becomes an issue, recorded tracks should track awake state as well as transforms
+	if (DynamicState == (uint8)EObjectStateTypeEnum::Chaos_Object_Sleeping)
+	{
+		Handle->SetObjectStateLowLevel(Chaos::EObjectStateType::Sleeping);
+	}
+	else if (DynamicState == (uint8)EObjectStateTypeEnum::Chaos_Object_Kinematic)
+	{
+		Handle->SetObjectStateLowLevel(Chaos::EObjectStateType::Kinematic);
+	}
+	else if (DynamicState == (uint8)EObjectStateTypeEnum::Chaos_Object_Static)
+	{
+		Handle->SetObjectStateLowLevel(Chaos::EObjectStateType::Static);
+	}
+	else
+	{
+		Handle->SetObjectStateLowLevel(Chaos::EObjectStateType::Dynamic);
+	}
+}
+
+//==============================================================================
+// FGeometryCollectionPhysicsProxy
+//==============================================================================
+
+FGeometryCollectionPhysicsProxy::FGeometryCollectionPhysicsProxy(
+	UObject* InOwner, 
+	const FSimulationParameters & SimulationParameters,
+	FInitFunc InInitFunc, 
+	FCacheSyncFunc InCacheSyncFunc, 
+	FFinalSyncFunc InFinalSyncFunc,
+	const Chaos::EMultiBufferMode BufferMode)
+	: Base(InOwner)
+	, Parameters(SimulationParameters)
+	, IsObjectDynamic(false)
+	, IsObjectLoading(true)
+
+	, BaseParticleIndex(INDEX_NONE)
+	, NumParticles(INDEX_NONE)
+	, ProxySimDuration(0.0f)
+
+	, InitFunc(InInitFunc)
+	, CacheSyncFunc(InCacheSyncFunc)
+	, FinalSyncFunc(InFinalSyncFunc)
+
+	, LastSyncCountGT(MAX_uint32)
+	, CollisionParticlesPerObjectFraction(CollisionParticlesPerObjectFractionDefault)
+{
+	// We rely on a guarded buffer.
+	check(BufferMode == Chaos::EMultiBufferMode::TripleGuarded);
+}
+
+FGeometryCollectionPhysicsProxy::~FGeometryCollectionPhysicsProxy()
+{}
+
+float ReportHighParticleFraction = -1.f;
+FAutoConsoleVariableRef CVarReportHighParticleFraction(TEXT("p.gc.ReportHighParticleFraction"), ReportHighParticleFraction, TEXT("Report any objects with particle fraction above this threshold"));
+
+void FGeometryCollectionPhysicsProxy::Initialize()
+{
+	// Old proxy init
+	check(IsInGameThread());
+
+
+	//
+	// Game thread initilization. 
+	//
+	//  1) Create a input buffer to store all game thread side data. 
+	//  2) Populate the buffer with the necessary data.
+	//  3) Deep copy the data to the other buffers. 
+	//
+	FGeometryDynamicCollection& DynamicCollection = *GameToPhysInterchange.AccessProducerBuffer();
+
+
+	InitializeDynamicCollection(DynamicCollection, *Parameters.RestCollection, Parameters);
+
+	NumParticles = GetTransformGroupSize(); // SimulatableParticles.Count(true);
+	BaseParticleIndex = 0; // Are we always zero indexed now?
+
+	//	
+	//  Give clients the opportunity to update the parameters before the simualtion is setup. 
+	//
+	Parameters.DynamicCollection = &DynamicCollection;
+	if (InitFunc)
+	{
+		InitFunc(Parameters);
+	}
+	Parameters.DynamicCollection = nullptr; // safty setup to prevent usage of the Parameters.DynamicCollection. 
+
+
+	//
+	// Collision vertices down sampling validation.  
+	//
+	CollisionParticlesPerObjectFraction = Parameters.CollisionSampleFraction * CollisionParticlesPerObjectFractionDefault;
+	if (ReportHighParticleFraction > 0)
+	{
+		for (const FSharedSimulationSizeSpecificData& Data : Parameters.Shared.SizeSpecificData)
+		{
+			if (Data.CollisionParticlesFraction >= ReportHighParticleFraction)
+			{
+				ensureMsgf(false, TEXT("Collection with small particle fraction"));
+				UE_LOG(LogChaos, Warning, TEXT("Collection with small particle fraction(%f):%s"), Data.CollisionParticlesFraction, *Parameters.Name);
+			}
+		}
+	}
+
+
+#if TODO_REIMPLEMENT_RIGID_CACHING
+
+	// todo : Replace with normal funcs in this class
+	Callbacks->SetUpdateRecordedStateFunction([this](float SolverTime, const TManagedArray<int32>& RigidBodyID, const TManagedArray<FGeometryCollectionBoneNode>& Hierarchy, const FSolverCallbacks::FParticlesType& Particles, const FSolverCallbacks::FCollisionConstraintsType& CollisionRule)
+		{
+			UpdateRecordedState(SolverTime, RigidBodyID, Hierarchy, Particles, CollisionRule);
+		});
+
+	Callbacks->SetCommitRecordedStateFunction([this](FRecordedTransformTrack& InTrack)
+		{
+			InTrack = FRecordedTransformTrack::ProcessRawRecordedData(RecordedTracks);
+		});
+
+	// Old proxy init
+	RecordedTracks.Records.Reset();
+
+	if (Parameters.bClearCache)
+	{
+		if (ResetAnimationCacheCallback)
+		{
+			ResetAnimationCacheCallback();
+		}
+	}
+
+	ProxySimDuration = T(0);
+	LastSyncCountGT = 0;
+
+
+#endif
+
+	/*
+	// Initialize data for faster bound calculations
+	// precompute data used for bounds calculation
+	{
+		const TManagedArray<FBox>& BoundingBoxes = Parameters.RestCollection->BoundingBox;
+		const TManagedArray<int32>& TransformIndices = Parameters.RestCollection->TransformIndex;
+
+		const int32 NumBoxes = BoundingBoxes.Num();
+
+		ValidGeometryBoundingBoxes.Reset();
+		ValidGeometryTransformIndices.Reset();
+		for (int32 BoxIdx = 0; BoxIdx < NumBoxes; ++BoxIdx)
+		{
+			const int32 CurrTransformIndex = TransformIndices[BoxIdx];
+
+			if (Parameters.RestCollection->IsGeometry(CurrTransformIndex))
+			{
+				ValidGeometryBoundingBoxes.Add(BoundingBoxes[BoxIdx]);
+				ValidGeometryTransformIndices.Add(CurrTransformIndex);
+			}
+		}
+
+		FBox BoundingBox(ForceInit);
+		const FMatrix& ActorToWorld = Parameters.WorldTransform.ToMatrixWithScale();
+		for (int i = 0; i < ValidGeometryBoundingBoxes.Num(); ++i)
+		{
+			BoundingBox += ValidGeometryBoundingBoxes[i].TransformBy(
+				TmpGlobalTransforms[ValidGeometryTransformIndices[i]] * ActorToWorld);
+		}
+
+		//Results.Get(0).WorldBounds = FBoxSphereBounds(BoundingBox);
+		//Results.Get(1).WorldBounds = FBoxSphereBounds(BoundingBox);
+	}
+	*/
+	GameToPhysInterchange.FlipProducer();
+	FGeometryDynamicCollection& ProducerCollection = *GameToPhysInterchange.AccessProducerBuffer();
+	ProducerCollection.CopyMatchingAttributesFrom(DynamicCollection);
+}
+
+void FGeometryCollectionPhysicsProxy::InitializeDynamicCollection(FGeometryDynamicCollection& DynamicCollection, const FGeometryCollection& RestCollection, const FSimulationParameters& Params)
+{
+	// 
+	// This function will use the rest collection to populate the dynamic collection. 
+	//
+
+	DynamicCollection.CopyMatchingAttributesFrom(RestCollection);
+
+	
+	DynamicCollection.SolverClusterID.Fill(nullptr); 
+
+	//
+	// User defined initial velocities need to be populated. 
+	//
+	{
+		if (Params.InitialVelocityType == EInitialVelocityTypeEnum::Chaos_Initial_Velocity_User_Defined)
+		{
+			DynamicCollection.InitialLinearVelocity.Fill(Params.InitialLinearVelocity);
+			DynamicCollection.InitialAngularVelocity.Fill(Params.InitialAngularVelocity);
+		}
+	}
+
+
+	//
+	// The RestCollection may have been authored with TUniquePtr<FImplicitObject> under the name 
+	// FGeometryDynamicCollection::ImplicitsAttribute, as they originally were.  As the particle system now uses TSharedPtr,
+	// we're now authoring TSharedPtr<FImplicitObject> under the name FGeometryDynamicCollection::SharedImplicitsAttribute.
+	{
+		FGeometryCollection& NCRestCollection = const_cast<FGeometryCollection&>(RestCollection);
+		TManagedArray<TUniquePtr<Chaos::FImplicitObject>>* UniqueRestImplicits =
+			NCRestCollection.FindAttribute<TUniquePtr<Chaos::FImplicitObject>>(
+				FGeometryDynamicCollection::ImplicitsAttribute, FTransformCollection::TransformGroup);
+		TManagedArray<TSharedPtr<Chaos::FImplicitObject, ESPMode::ThreadSafe>>* SharedRestImplicits =
+			NCRestCollection.FindAttribute<TSharedPtr<Chaos::FImplicitObject, ESPMode::ThreadSafe>>(
+				FGeometryDynamicCollection::FGeometryDynamicCollection::SharedImplicitsAttribute, FTransformCollection::TransformGroup);
+
+		// If we have unique pointers to geometry, we need to convert them to shared pointers.
+		TManagedArray<TSharedPtr<Chaos::FImplicitObject, ESPMode::ThreadSafe>>* RestCollectionSharedImplicits = nullptr;
+		if (!SharedRestImplicits)
+		{
+			RestCollectionSharedImplicits =
+				&NCRestCollection.AddAttribute<TSharedPtr<Chaos::FImplicitObject, ESPMode::ThreadSafe>>(
+					FGeometryDynamicCollection::SharedImplicitsAttribute, FTransformCollection::TransformGroup);
+		}
+
+		for (int32 Index = DynamicCollection.NumElements(FTransformCollection::TransformGroup) - 1; 0 <= Index; Index--)
+		{
+			//Implicits[Index] = Chaos::MakeSerializable(RestImplicits[Index]);
+			if (SharedRestImplicits)
+			{
+				DynamicCollection.Implicits[Index] = (*SharedRestImplicits)[Index];
+			}
+			else if (UniqueRestImplicits && RestCollectionSharedImplicits) // RestCollectionSharedImplicits is redundant, but static analysis build complains without it
+			{
+				// Steal the memory from the unique ptr.
+				DynamicCollection.Implicits[Index] = TSharedPtr<Chaos::FImplicitObject, ESPMode::ThreadSafe>((*UniqueRestImplicits)[Index].Release());
+				// Update the rest collection with a copy of the shared pointer.
+				(*RestCollectionSharedImplicits)[Index] = DynamicCollection.Implicits[Index];
+			}
+		}
+
+		// Clean up now invalid unique pointers.
+		if (RestCollectionSharedImplicits)
+		{
+			NCRestCollection.RemoveAttribute(FGeometryDynamicCollection::ImplicitsAttribute, FTransformCollection::TransformGroup);
+		}
+	}
+
+
+	// process simplicials
+	{
+		/*
+		if (RestCollection.HasAttribute(DynamicCollection.SimplicialsAttribute, FTransformCollection::TransformGroup))
+		{
+			const auto& RestSimplicials = RestCollection.GetAttribute<TUniquePtr<FSimplicial>>(
+				DynamicCollection.SimplicialsAttribute, FTransformCollection::TransformGroup);
+			for (int32 Index = DynamicCollection.NumElements(FTransformCollection::TransformGroup) - 1; 0 <= Index; Index--)
+			{
+				DynamicCollection.Simplicials[Index] = TSharedPtr<FSimplicial>(RestSimplicials[Index] ? RestSimplicials[Index]->NewCopy() : nullptr);
+			}
+		}
+		else
+		{
+			for (int32 Index = DynamicCollection.NumElements(FTransformCollection::TransformGroup) - 1; 0 <= Index; Index--)
+			{
+				DynamicCollection.Simplicials[Index] = TSharedPtr<FSimplicial>(nullptr);
+			}
+		}
+		*/
+	}
+
+
+	// Process Activity
+	{
+		const int32 NumTransforms = DynamicCollection.SimulatableParticles.Num();
+		if (RestCollection.HasAttribute(FGeometryCollection::SimulatableParticlesAttribute, FTransformCollection::TransformGroup))
+		{
+			// When the rest collection has been pre configured with simulation data, use that to determine it ability to simulate. 
+			const TManagedArray<bool>& ActiveRestParticles =
+				RestCollection.GetAttribute<bool>(
+					FGeometryCollection::SimulatableParticlesAttribute, FTransformCollection::TransformGroup);
+			check(ActiveRestParticles.Num() == DynamicCollection.Active.Num());
+			check(ActiveRestParticles.Num() == DynamicCollection.SimulatableParticles.Num());
+			for (int32 TransformIdx = 0; TransformIdx < NumTransforms; TransformIdx++)
+			{
+				DynamicCollection.SimulatableParticles[TransformIdx] = DynamicCollection.Active[TransformIdx] && ActiveRestParticles[TransformIdx];
+			}
+		}
+		else
+		{
+			// If no simulation data is available then default to the simulation of just the rigid geometry.
+			for (int32 TransformIdx = 0; TransformIdx < NumTransforms; TransformIdx++)
+			{
+				if (DynamicCollection.Children[TransformIdx].Num())
+				{
+					DynamicCollection.SimulatableParticles[TransformIdx] = false;
+				}
+				else
+				{
+					DynamicCollection.SimulatableParticles[TransformIdx] = DynamicCollection.Active[TransformIdx];
+				}
+			}
+		}
+	}
+
+	// Process Transforms
+	{
+		//TArray<FTransform> Transform;
+		//GeometryCollectionAlgo::GlobalMatrices(DynamicCollection->Transform, DynamicCollection->Parent, Transform);
+		//check(DynamicCollection->Transform.Num() == Transform.Num());
+
+		// Initialize global transforms
+		//TArray<FMatrix>& TmpGlobalTransforms = GameToPhysInterchange.AccessProducerBuffer()->GlobalTransforms;
+		//TArray<FMatrix> &TmpGlobalTransforms = Results.Get(0).GlobalTransforms;
+		//GeometryCollectionAlgo::GlobalMatrices(PTDynamicCollection.Transform, PTDynamicCollection.Parent, TmpGlobalTransforms);
+		//Results.Get(1).GlobalTransforms = TmpGlobalTransforms;
+	}
+}
+
+
+
+int32 ReportTooManyChildrenNum = -1;
+FAutoConsoleVariableRef CVarReportTooManyChildrenNum(TEXT("p.ReportTooManyChildrenNum"), ReportTooManyChildrenNum, TEXT("Issue warning if more than this many children exist in a single cluster"));
+
+void FGeometryCollectionPhysicsProxy::InitializeBodiesPT(
+	Chaos::FPBDRigidsSolver* RigidsSolver,
+	Chaos::FPBDRigidsSolver::FParticlesType& Particles)
+{
+	/*
+	@todo break_everything
+
+	const FGeometryCollection* RestCollection = Parameters.RestCollection;
+	FGeometryDynamicCollection* DynamicCollection = Parameters.DynamicCollection;
+
+	if (Parameters.Simulating)
+	{
+		const TManagedArray<int32>& TransformIndex = RestCollection->TransformIndex;
+		const TManagedArray<int32>& BoneMap = RestCollection->BoneMap;
+		const TManagedArray<int32>& Parent = RestCollection->Parent;
+		const TManagedArray<TSet<int32>>& Children = RestCollection->Children;
+		const TManagedArray<int32>& SimulationType = RestCollection->SimulationType;
+		const TManagedArray<FVector>& Vertex = RestCollection->Vertex;
+		const TManagedArray<int32>& DynamicState = DynamicCollection->DynamicState;
+		const TManagedArray<int32>& CollisionGroup = DynamicCollection->CollisionGroup;
+		const TManagedArray<float>& Mass = RestCollection->GetAttribute<float>("Mass", FTransformCollection::TransformGroup);
+		const TManagedArray<FVector>& InertiaTensor = RestCollection->GetAttribute<FVector>("InertiaTensor", FTransformCollection::TransformGroup);
+
+		const int32 NumTransforms = GetTransformGroupSize();
+
+		TArray<FTransform> Transform;
+		GeometryCollectionAlgo::GlobalMatrices(
+			DynamicCollection->Transform, DynamicCollection->Parent, Transform);
+		check(DynamicCollection->Transform.Num() == Transform.Num());
+
+		const int NumRigids = 0; // ryan - Since we're doing SOA, we start at zero?
+		BaseParticleIndex = NumRigids;
+
+		// Count geometry collection leaf node particles to add
+		int NumSimulatedParticles = 0;
+		int NumLeafNodes = 0;
+		for (int32 Idx = 0; Idx < SimulatableParticles.Num(); ++Idx)
+		{
+			NumSimulatedParticles += SimulatableParticles[Idx];
+			NumLeafNodes += (SimulatableParticles[Idx] && Children[Idx].Num() == 0);
+		}
+
+		// Add entries into simulation array
+		RigidsSolver->GetEvolution()->ReserveParticles(NumSimulatedParticles);
+		TArray<Chaos::TPBDGeometryCollectionParticleHandle<float, 3>*> Handles = 
+			RigidsSolver->GetEvolution()->CreateGeometryCollectionParticles(NumLeafNodes);
+
+		int32 NextIdx = 0;
+		for(int32 Idx=0; Idx < SimulatableParticles.Num(); ++Idx)
+		{
+			if (SimulatableParticles[Idx] && Children[Idx].Num() == 0) // Leaf node
+			{
+				// todo: Unblocked read access of game thread data on the physics thread.
+
+				Chaos::TPBDGeometryCollectionParticleHandle<float, 3>* Handle = Handles[NextIdx++];
+
+				Handle->GTGeometryParticle() = nullptr; // GeometryCollections do not support Game Thread TGeometryParticle
+				RigidsSolver->AddParticleToProxy(Handle, this);
+
+				SolverParticleHandles[Idx] = Handle;
+				HandleToTransformGroupIndex.Add(Handle, Idx);
+
+				check(SolverParticleHandles[Idx]->GetParticleType() == Handle->GetParticleType());
+				RigidsSolver->GetEvolution()->CreateParticle(Handle);
+			}
+		}
+
+		const float StrainDefault = Parameters.DamageThreshold.Num() ? Parameters.DamageThreshold[0] : 0;
+		// Add the rigid bodies
+
+		// Iterating over the geometry group is a fast way of skipping everything that's
+		// not a leaf node, as each geometry has a transform index, which is a shortcut
+		// for the case when there's a 1-to-1 mapping between transforms and geometries.
+		// At the point that we start supporting instancing, this assumption will no longer
+		// hold, and those reverse mappints will be INDEX_NONE.
+
+		const int32 NumGeometries = DynamicCollection->NumElements(FGeometryCollection::GeometryGroup);
+		ParallelFor(NumGeometries, [&](int32 GeometryIndex)
+		//for (int32 GeometryIndex = 0; GeometryIndex < NumGeometries; ++GeometryIndex)
+		{
+			const int32 TransformGroupIndex = TransformIndex[GeometryIndex];
+			if (Chaos::TPBDRigidClusteredParticleHandle<float, 3>* Handle = SolverParticleHandles[TransformGroupIndex])
+			{
+				// TODO: We probably want to take better advantage of this solver
+				// feature, rather than doing it all here ourselves?
+				Handle->SetCenterOfMass(Chaos::TVector<float,3>(0.0f));
+				Handle->SetRotationOfMass(Chaos::TRotation<float,3>::FromIdentity());
+			
+				// Mass space -> Composed parent space -> world
+				const FTransform WorldTransform = 
+					MassToLocal[TransformGroupIndex] * Transform[TransformGroupIndex] * Parameters.WorldTransform;
+
+				PopulateSimulatedParticle(
+					Handle,
+					//Particles, 
+					Parameters.Shared, 
+					Simplicials[TransformGroupIndex].Get(),
+					Implicits[TransformGroupIndex], 
+#if TODO_REIMPLEMENT_SOLVER_SETTINGS_ACCESSORS
+					RigidsSolver->GetMassScale() * Mass[TransformGroupIndex], 
+					RigidsSolver->GetMassScale() * InertiaTensor[TransformGroupIndex],
+#else // TODO_REIMPLEMENT_SOLVER_SETTINGS_ACCESSORS
+					Mass[TransformGroupIndex], 
+					InertiaTensor[TransformGroupIndex],
+#endif // TODO_REIMPLEMENT_SOLVER_SETTINGS_ACCESSORS
+					//RigidBodyIndex, 
+					WorldTransform, 
+					(uint8)DynamicState[TransformGroupIndex], 
+					CollisionGroup[TransformGroupIndex]);
+
+				if (Parameters.EnableClustering)
+				{
+					Handle->SetClusterGroupIndex(Parameters.ClusterGroupIndex);
+					Handle->SetStrain(StrainDefault);
+				}
+
+				TUniquePtr<Chaos::TBVHParticles<float, 3>>& CollisionParticles = Handle->CollisionParticles();
+				CollisionParticles.Reset(Simplicials[TransformGroupIndex].Release()); // Steals!
+				//Particles.CollisionParticles(RigidBodyIndex).Reset(Simplicials[TransformGroupIndex].Release());
+				if (CollisionParticles)
+				//if (Particles.CollisionParticles(RigidBodyIndex))
+				{
+					int32 NumCollisionParticles = CollisionParticles->Size();
+					//int32 NumCollisionParticles = Particles.CollisionParticles(RigidBodyIndex)->Size();
+					int32 CollisionParticlesSize = FMath::Max(0, FMath::Min(int(NumCollisionParticles*CollisionParticlesPerObjectFraction), NumCollisionParticles));
+					CollisionParticles->Resize(CollisionParticlesSize); // Truncates!
+					//Particles.CollisionParticles(RigidBodyIndex)->Resize(CollisionParticlesSize);
+				}
+
+				RigidsSolver->GetEvolution()->SetPhysicsMaterial(Handle, Parameters.PhysicalMaterial);
+				//GetSolver()->SetPhysicsMaterial(RigidBodyIndex, Parameters.PhysicalMaterial);
+			}
+		});
+
+		// After population, the states of each particle could have changed
+		Particles.UpdateGeometryCollectionViews();
+
+		for (FFieldSystemCommand& Cmd : Parameters.InitializationCommands)
+		{
+			if (Cmd.MetaData.Contains(FFieldSystemMetaData::EMetaType::ECommandData_ProcessingResolution)) 
+				Cmd.MetaData.Remove(FFieldSystemMetaData::EMetaType::ECommandData_ProcessingResolution);
+			FFieldSystemMetaDataProcessingResolution* ResolutionData = 
+				new FFieldSystemMetaDataProcessingResolution(EFieldResolutionType::Field_Resolution_Maximum);
+			Cmd.MetaData.Add(
+				FFieldSystemMetaData::EMetaType::ECommandData_ProcessingResolution, 
+				TUniquePtr< FFieldSystemMetaDataProcessingResolution >(ResolutionData));
+			Commands.Add(Cmd);
+		}
+		Parameters.InitializationCommands.Empty();
+		ProcessCommands(Particles.GetGeometryCollectionParticles(), GetSolver()->GetSolverTime());
+
+		if (Parameters.InitialVelocityType == EInitialVelocityTypeEnum::Chaos_Initial_Velocity_User_Defined)
+		{
+			// A previous implementation of this went wide on this loop.  The general 
+			// rule of thumb for parallelization is that each thread needs at least
+			// 1000 operations in order to overcome the expense of threading.  I don't
+			// think that's generally going to be the case here...
+			for (int32 TransformGroupIndex = 0; TransformGroupIndex < GetTransformGroupSize(); ++TransformGroupIndex)
+			{
+				if (Chaos::TPBDRigidParticleHandle<float, 3>* Handle = SolverParticleHandles[TransformGroupIndex])
+				{
+					if (DynamicState[TransformGroupIndex] == (int32)EObjectStateTypeEnum::Chaos_Object_Dynamic)
+					{
+						Handle->SetV(InitialLinearVelocity[TransformGroupIndex]);
+						Handle->SetW(InitialAngularVelocity[TransformGroupIndex]);
+					}
+				}
+			}
+		}
+
+#if TODO_REIMPLEMENT_FRACTURE
+		InitializeRemoveOnFracture(Particles, DynamicState);
+#endif // TODO_REIMPLEMENT_FRACTURE
+
+		// #BG Temporary - don't cluster when playing back. Needs to be changed when kinematics are per-proxy to support
+		// kinematic to dynamic transition for clusters.
+		if (Parameters.EnableClustering)// && Parameters.CacheType != EGeometryCollectionCacheType::Play)
+		{
+			// "RecursiveOrder" means bottom up - children come before their parents.
+			const TArray<int32> RecursiveOrder = ComputeRecursiveOrder(*RestCollection);
+
+			// Propagate simulated particle flags up the hierarchy from children 
+			// to their parents, grandparents, etc...
+			TArray<bool> SubTreeContainsSimulatableParticle;
+			SubTreeContainsSimulatableParticle.SetNumZeroed(RecursiveOrder.Num());
+			for (const int32 TransformGroupIndex : RecursiveOrder)
+			{
+				if (Children[TransformGroupIndex].Num() == 0)
+				{
+					// Leaf node
+					SubTreeContainsSimulatableParticle[TransformGroupIndex] =
+						SolverParticleHandles[TransformGroupIndex] != nullptr;
+				}
+				else
+				{
+					// Cluster parent
+					const TSet<int32>& ChildIndices = Children[TransformGroupIndex];
+					for (const int32 ChildIndex : ChildIndices)
+					{
+						if(SubTreeContainsSimulatableParticle[ChildIndex])
+						{
+							SubTreeContainsSimulatableParticle[TransformGroupIndex] = true;
+							break;
+						}
+					}
+				}
+			}
+
+			TArray<Chaos::TPBDRigidClusteredParticleHandle<float, 3>*> ClusterHandles;
+			// Ryan - It'd be better to batch allocate cluster particles ahead of time,
+			// but if ClusterHandles is empty, then new particles will be allocated
+			// on the fly by TPBDRigidClustering::CreateClusterParticle(), which 
+			// needs to work before this does...
+			//ClusterHandles = GetSolver()->GetEvolution()->CreateClusteredParticles(NumClusters);
+
+			int32 ClusterHandlesIndex = 0;
+			TArray<Chaos::TPBDRigidParticleHandle<float, 3>*> RigidChildren;
+			TArray<int32> RigidChildrenTransformGroupIndex;
+			for (const int32 TransformGroupIndex : RecursiveOrder)
+			{
+				// Don't construct particles for branches of the hierarchy that  
+				// don't contain any simulated particles.
+				if (!SubTreeContainsSimulatableParticle[TransformGroupIndex])
+				{
+					continue;
+				}
+
+				RigidChildren.Reset(Children.Num());
+				RigidChildrenTransformGroupIndex.Reset(Children.Num());
+				for (const int32 ChildIndex : Children[TransformGroupIndex])
+				{
+					if (Chaos::TPBDRigidClusteredParticleHandle<float, 3>* Handle = SolverParticleHandles[ChildIndex])
+					{
+						RigidChildren.Add(Handle);
+						RigidChildrenTransformGroupIndex.Add(ChildIndex);
+					}
+				}
+				if (RigidChildren.Num())
+				{
+					if (ReportTooManyChildrenNum >= 0 && RigidChildren.Num() > ReportTooManyChildrenNum)
+					{
+						UE_LOG(LogChaos, Warning, TEXT("Too many children (%d) in a single cluster:%s"), 
+							RigidChildren.Num(), *Parameters.Name);
+					}
+
+					Chaos::FClusterCreationParameters<float> CreationParameters;
+					CreationParameters.ClusterParticleHandle = 
+						ClusterHandles.Num() ? ClusterHandles[ClusterHandlesIndex++] : nullptr;
+
+					Chaos::TPBDRigidClusteredParticleHandle<float, 3>* Handle =
+						BuildClusters(
+							TransformGroupIndex, 
+							RigidChildren, 
+							RigidChildrenTransformGroupIndex, 
+							CreationParameters);
+
+					int32 RigidChildrenIdx = 0;
+					for (const int32 ChildIndex : Children[TransformGroupIndex])
+					{
+						SolverClusterID[ChildIndex] = RigidChildren[RigidChildrenIdx++]->CastToClustered()->ClusterIds().Id;
+					}
+					SolverClusterID[TransformGroupIndex] = Handle->ClusterIds().Id;
+
+					SolverClusterHandles[TransformGroupIndex] = Handle;
+					SolverParticleHandles[TransformGroupIndex] = Handle;
+					HandleToTransformGroupIndex.Add(Handle, TransformGroupIndex);
+					RigidsSolver->AddParticleToProxy(Handle, this);
+				}
+			}
+
+			// We've likely changed the state of leaf nodes, which are geometry
+			// collection particles.  Update which particle views they belong in,
+			// as well as views of clustered particles.
+			Particles.UpdateGeometryCollectionViews(true); 
+
+			// Set cluster connectivity.  TPBDRigidClustering::CreateClusterParticle() 
+			// will optionally do this, but we switch that functionality off in BuildClusters().
+			ParallelFor(NumTransforms, [&](int32 TransformGroupIndex)
+			{
+				const TManagedArray<TSet<int32>>& RestChildren = RestCollection->Children;
+				if (RestChildren[TransformGroupIndex].Num() > 0)
+				{
+					if (SolverClusterHandles[TransformGroupIndex])
+					{
+						GetSolver()->GetEvolution()->GetRigidClustering().GenerateConnectionGraph(
+							SolverClusterHandles[TransformGroupIndex]);
+					}
+				}
+			});
+		} // end if EnableClustering
+ 
+
+#if TODO_REIMPLEMENT_RIGID_CACHING
+		// If we're recording and want to start immediately caching then we should cache the rest state
+		if (Parameters.IsCacheRecording() && Parameters.CacheBeginTime == 0.0f)
+		{
+			if (UpdateRecordedStateCallback)
+			{
+				UpdateRecordedStateCallback(0.0f, RigidBodyID, Particles, GetSolver()->GetCollisionConstraints());
+			}
+		}
+#endif // TODO_REIMPLEMENT_RIGID_CACHING
+
+	} // end if simulating...
+	*/
+}
+
+int32 ReportNoLevelsetCluster = 0;
+FAutoConsoleVariableRef CVarReportNoLevelsetCluster(TEXT("p.gc.ReportNoLevelsetCluster"), ReportNoLevelsetCluster, TEXT("Report any cluster objects without levelsets"));
+
+DECLARE_CYCLE_STAT(TEXT("FGeometryCollectionPhysicsProxy::BuildClusters"), STAT_BuildClusters, STATGROUP_Chaos);
+DECLARE_CYCLE_STAT(TEXT("FGeometryCollectionPhysicsProxy::BuildClusters:GlobalMatrices"), STAT_BuildClustersGlobalMatrices, STATGROUP_Chaos);
+
+Chaos::TPBDRigidClusteredParticleHandle<float, 3>* 
+FGeometryCollectionPhysicsProxy::BuildClusters(
+	const uint32 CollectionClusterIndex, // TransformGroupIndex
+	TArray<Chaos::TPBDRigidParticleHandle<float,3>*>& ChildHandles,
+	const TArray<int32>& ChildTransformGroupIndices,
+	const Chaos::FClusterCreationParameters<float> & ClusterParameters)
+{
+	SCOPE_CYCLE_COUNTER(STAT_BuildClusters);
+
+	check(CollectionClusterIndex != INDEX_NONE);
+	check(ChildHandles.Num() != 0);
+
+	check(Parameters.DynamicCollection);
+	FGeometryDynamicCollection& DynamicCollection = *GameToPhysInterchange.AccessProducerBuffer();
+	TManagedArray<int32>& DynamicState = DynamicCollection.DynamicState;
+	TManagedArray<int32>& ParentIndex = DynamicCollection.Parent;
+	TManagedArray<TSet<int32>>& Children = DynamicCollection.Children;
+	TManagedArray<FTransform>& Transform = DynamicCollection.Transform;
+	TManagedArray<FTransform>& MassToLocal = DynamicCollection.MassToLocal;
+	//TManagedArray<TSharedPtr<FCollisionStructureManager::FSimplicial> >& Simplicials = DynamicCollection.Simplicials;
+	TManagedArray<TSharedPtr<Chaos::FImplicitObject, ESPMode::ThreadSafe>>& Implicits = DynamicCollection.Implicits;
+	TManagedArray<Chaos::TPBDRigidClusteredParticleHandle<float, 3>*>& SolverClusterHandles = DynamicCollection.SolverClusterHandles;
+
+	//If we are a root particle use the world transform, otherwise set the relative transform
+	const FTransform CollectionSpaceTransform = GeometryCollectionAlgo::GlobalMatrix(Transform, ParentIndex, CollectionClusterIndex);
+	const Chaos::TRigidTransform<float, 3> ParticleTM = MassToLocal[CollectionClusterIndex] * CollectionSpaceTransform * Parameters.WorldTransform;
+
+	//create new cluster particle
+	//The reason we need to pass in a mass orientation override is as follows:
+	//Consider a pillar made up of many boxes along the Y-axis. In this configuration we could generate a proxy pillar along the Y with identity rotation.
+	//Now if we instantiate the pillar and rotate it so that it is along the X-axis, we would still like to use the same pillar proxy.
+	//Since the mass orientation is computed in world space in both cases we'd end up with a diagonal inertia matrix and identity rotation that looks like this: [big, small, big] or [small, big, big].
+	//Because of this we need to know how to rotate collision particles and geometry to match with original computation. If it was just geometry we could transform it before passing, but we need collision particles as well
+	Chaos::FClusterCreationParameters<float> ClusterCreationParameters = ClusterParameters;
+	ClusterCreationParameters.bGenerateConnectionGraph = false;
+	// fix... ClusterCreationParameters.CollisionParticles = Simplicials[CollectionClusterIndex];
+	ClusterCreationParameters.ConnectionMethod = Parameters.ClusterConnectionMethod;
+	if (ClusterCreationParameters.CollisionParticles)
+	{
+		const int32 NumCollisionParticles = ClusterCreationParameters.CollisionParticles->Size();
+		const int32 ClampedCollisionParticlesSize = 
+			FMath::Max(0, FMath::Min(int(NumCollisionParticles * CollisionParticlesPerObjectFraction), NumCollisionParticles));
+		ClusterCreationParameters.CollisionParticles->Resize(ClampedCollisionParticlesSize);
+	}
+	TArray<Chaos::TPBDRigidParticleHandle<float, 3>*> ChildHandlesCopy(ChildHandles);
+
+	// Construct an active cluster particle, disable children, derive M and I from children:
+	Chaos::TPBDRigidClusteredParticleHandle<float, 3>* Parent = 
+		GetSolver()->GetEvolution()->GetRigidClustering().CreateClusterParticle(
+			Parameters.ClusterGroupIndex, 
+			MoveTemp(ChildHandlesCopy),
+			ClusterCreationParameters,
+			Implicits[CollectionClusterIndex], // union from children if null
+			&ParticleTM);
+
+	if (ReportNoLevelsetCluster && 
+		Parent->DynamicGeometry())
+	{
+		//ensureMsgf(false, TEXT("Union object generated for cluster"));
+		UE_LOG(LogChaos, Warning, TEXT("Union object generated for cluster:%s"), *Parameters.Name);
+	}
+
+	if (Parent->InvM() == 0.0)
+	{
+		if (Parent->ObjectState() == Chaos::EObjectStateType::Static)
+		{
+			DynamicState[CollectionClusterIndex] = (uint8)EObjectStateTypeEnum::Chaos_Object_Static;
+		}
+		else //if (Particles.ObjectState(NewSolverClusterID) == Chaos::EObjectStateType::Kinematic)
+		{
+			DynamicState[CollectionClusterIndex] = (uint8)EObjectStateTypeEnum::Chaos_Object_Kinematic;
+		}
+	}
+
+	// In theory we should be computing this and passing in to avoid inertia computation at 
+	// runtime. If we do this we must account for leaf particles that have already been created in world space
+	//const FTransform ParticleTM(Particles.R(NewSolverClusterID), Particles.X(NewSolverClusterID));	
+	MassToLocal[CollectionClusterIndex] = FTransform::Identity;
+
+	check(Parameters.RestCollection);
+	const TManagedArray<float>& Mass = 
+		Parameters.RestCollection->GetAttribute<float>("Mass", FTransformCollection::TransformGroup);
+	const TManagedArray<FVector>& InertiaTensor = 
+		Parameters.RestCollection->GetAttribute<FVector>("InertiaTensor", FTransformCollection::TransformGroup);
+
+	PopulateSimulatedParticle(
+		Parent,
+		Parameters.Shared, 
+		nullptr, // CollisionParticles is optionally set from CreateClusterParticle()
+		nullptr, // Parent->Geometry() ? Parent->Geometry() : Implicits[CollectionClusterIndex], 
+#if TODO_REIMPLEMENT_SOLVER_SETTINGS_ACCESSORS
+		// Note: if the children had this scaling baked in, then cluster parents
+		// will too, and we shouldn't apply the scaling to the parents twice!
+		ThisSolver->GetMassScale() * Mass[CollectionClusterIndex], 
+		ThisSolver->GetMassScale() * InertiaTensor[CollectionClusterIndex],
+#else
+		// We computed M and I from children in CreateClusterParticle().
+		Parent->M() > 0.0 ? Parent->M() : Mass[CollectionClusterIndex], 
+		Parent->I().GetDiagonal() != Chaos::TVector<float,3>(0.0) ? Parent->I().GetDiagonal() : InertiaTensor[CollectionClusterIndex],
+#endif
+		ParticleTM, 
+		(uint8)DynamicState[CollectionClusterIndex], 
+		0); // CollisionGroup
+
+	// two-way mapping
+	SolverClusterHandles[CollectionClusterIndex] = Parent;
+
+	const int32 NumThresholds = Parameters.DamageThreshold.Num();
+	const int32 Level = FMath::Clamp(CalculateHierarchyLevel(DynamicCollection, CollectionClusterIndex), 0, INT_MAX);
+	const float DefaultDamage = NumThresholds > 0 ? Parameters.DamageThreshold[NumThresholds - 1] : 0.f;
+	float Damage = Level < NumThresholds ? Parameters.DamageThreshold[Level] : DefaultDamage;
+	if(Level >= Parameters.MaxClusterLevel) 
+		Damage = FLT_MAX;
+	Parent->SetStrains(Damage);
+
+	GetSolver()->GetEvolution()->SetPhysicsMaterial(Parent, Parameters.PhysicalMaterial);
+
+	const FTransform ParentTransform = 
+		GeometryCollectionAlgo::GlobalMatrix(DynamicCollection.Transform, DynamicCollection.Parent, CollectionClusterIndex);
+
+	int32 MinCollisionGroup = INT_MAX;
+	for(int32 Idx=0; Idx < ChildHandles.Num(); Idx++)
+	{
+		Chaos::TPBDRigidParticleHandle<float, 3>* Child = ChildHandles[Idx];
+		if (Chaos::TPBDRigidClusteredParticleHandle<float, 3>* ClusteredChild = Child->CastToClustered())
+		{
+			ClusteredChild->SetStrains(Damage);
+		}
+
+		const int32 ChildTransformGroupIndex = ChildTransformGroupIndices[Idx];
+		SolverClusterHandles[ChildTransformGroupIndex] = Parent;
+
+		const FTransform ChildTransform(Child->R(), Child->X());
+		if(Children[ChildTransformGroupIndex].Num()) // clustered local transform
+		{
+			Transform[ChildTransformGroupIndex] = ChildTransform.GetRelativeTransform(ParticleTM);
+		}
+		else // rigid local transform
+		{
+			const FTransform RestTransform = 
+				Parameters.RestCollection->Transform[ChildTransformGroupIndex] * ParentTransform * Parameters.WorldTransform;
+			Transform[ChildTransformGroupIndex] = RestTransform.GetRelativeTransform(ParticleTM);
+		}
+		Transform[ChildTransformGroupIndex].NormalizeRotation();
+
+		MinCollisionGroup = FMath::Min(Child->CollisionGroup(), MinCollisionGroup);
+	}
+	Parent->SetCollisionGroup(MinCollisionGroup);
+	return Parent;
+}
+
+void FGeometryCollectionPhysicsProxy::ContiguousIndices(TArray<ContextIndex>& Array, const Chaos::FPhysicsSolver* RigidSolver, EFieldResolutionType ResolutionType, bool bForce = true)
+{
+	check(false); // Not currently called
+	if (bForce)
+	{
+#if TODO_REIMPLEMENT_GET_RIGID_PARTICLES
+		const Chaos::FPhysicsSolver::FParticlesType & Particles = RigidSolver->GetRigidParticles();
+		if (ResolutionType == EFieldResolutionType::Field_Resolution_Minimal)
+		{
+			const Chaos::TArrayCollectionArray<Chaos::ClusterId>& ClusterIdArray = RigidSolver->GetRigidClustering().GetClusterIdsArray();
+
+
+			//  Generate a Index mapping between the rigid body indices and 
+			//  the particle indices. This allows the geometry collection to
+			//  evaluate only its own ACTIVE particles + ClusterChildren
+			int32 NumIndices = 0;
+			Array.SetNumUninitialized(RigidBodyID.Num());
+			for (int32 i = 0; i < RigidBodyID.Num(); i++)
+			{
+				const int32 RigidBodyIndex = RigidBodyID[i];
+				if (RigidBodyIndex != INDEX_NONE && !Particles.Disabled(RigidBodyIndex)) // active bodies
+				{
+					Array[NumIndices] = { RigidBodyID[i],i };
+					NumIndices++;
+				}
+				if (ClusterIdArray[RigidBodyIndex].Id != INDEX_NONE && !Particles.Disabled(ClusterIdArray[RigidBodyIndex].Id)) // children
+				{
+					Array[NumIndices] = { RigidBodyID[i],i };
+					NumIndices++;
+				}
+			}
+			Array.SetNum(NumIndices);
+		}
+		else if (ResolutionType == EFieldResolutionType::Field_Resolution_Maximum)
+		{
+			//  Generate a Index mapping between the rigid body indices and 
+			//  the particle indices. This allows the geometry collection to
+			//  evaluate only its own particles. 
+			int32 NumIndices = 0;
+			Array.SetNumUninitialized(RigidBodyID.Num());
+			for (int32 i = 0; i < RigidBodyID.Num(); i++)
+			{
+				const int32 RigidBodyIndex = RigidBodyID[i];
+				if (RigidBodyIndex != INDEX_NONE)
+				{
+					Array[NumIndices] = { RigidBodyIndex, i };
+					NumIndices++;
+				}
+			}
+			Array.SetNum(NumIndices);
+		}
+#endif
+	}
+}
+
+int32 FGeometryCollectionPhysicsProxy::CalculateHierarchyLevel(const FGeometryDynamicCollection& GeometryCollection, int32 TransformIndex) const
+{
+	int32 Level = 0;
+	while (GeometryCollection.Parent[TransformIndex] != -1)
+	{
+		TransformIndex = GeometryCollection.Parent[TransformIndex];
+		Level++;
+	}
+	return Level;
+}
+
+
+void FGeometryCollectionPhysicsProxy::InitializeRemoveOnFracture(FParticlesType& Particles, const TManagedArray<int32>& DynamicState)
+{
+	/*
+	@todo break everything
+	if (Parameters.DynamicCollection && Parameters.RemoveOnFractureEnabled)
+	{
+	//	TManagedArray<FGeometryCollectionBoneNode>& Hierarchy = Parameters.DynamicCollection->BoneHierarchy;
+
+		for (int TransformGroupIndex = 0; TransformGroupIndex < RigidBodyID.Num(); TransformGroupIndex++)
+		{
+			if (RigidBodyID[TransformGroupIndex] != INDEX_NONE)
+			{
+				int32 RigidBodyIndex = RigidBodyID[TransformGroupIndex];
+
+				if (Parameters.DynamicCollection->StatusFlags[TransformGroupIndex] & FGeometryCollection::FS_RemoveOnFracture)
+				{
+					Particles.ToBeRemovedOnFracture(RigidBodyIndex) = true;
+				}
+			}
+		}
+	}
+	*/
+}
+
+void FGeometryCollectionPhysicsProxy::AddCollisionToCollisionData(FRecordedFrame* ExistingFrame, const FParticlesType& Particles, const Chaos::TPBDCollisionConstraints<float, 3>::FPointContactConstraint& Constraint)
+{
+#if TODO_REIMPLEMENT_RIGID_CLUSTERING
+	const Chaos::TArrayCollectionArray<Chaos::ClusterId>& ClusterIdsArray = GetSolver()->GetRigidClustering().GetClusterIdsArray();
+	const TMap<uint32, TUniquePtr<TArray<uint32>>>& ParentToChildrenMap = GetSolver()->GetRigidClustering().GetChildrenMap();
+	const Chaos::TArrayCollectionArray<PhysicsProxyWrapper>& PhysicsProxyReverseMapping = GetSolver()->GetPhysicsProxyReverseMapping();
+
+	int32 NewIdx = ExistingFrame->Collisions.Add(FSolverCollisionData());
+	FSolverCollisionData& Collision = ExistingFrame->Collisions[NewIdx];
+
+	Collision.Location = Constraint.Location;
+	Collision.AccumulatedImpulse = Constraint.AccumulatedImpulse;
+	Collision.Normal = Constraint.Normal;
+	Collision.Velocity1 = Particles.V(Constraint.ParticleIndex);
+	Collision.Velocity2 = Particles.V(Constraint.LevelsetIndex);
+	Collision.AngularVelocity1 = Particles.W(Constraint.ParticleIndex);
+	Collision.AngularVelocity2 = Particles.W(Constraint.LevelsetIndex);
+	Collision.Mass1 = Particles.M(Constraint.ParticleIndex);
+	Collision.Mass2 = Particles.M(Constraint.LevelsetIndex);
+	Collision.ParticleIndex = Constraint.ParticleIndex;
+	Collision.LevelsetIndex = Constraint.LevelsetIndex;
+
+	// If Constraint.ParticleIndex is a cluster store an index for a mesh in this cluster
+	if (ClusterIdsArray[Constraint.ParticleIndex].NumChildren > 0)
+	{
+		int32 ParticleIndexMesh = GetSolver()->GetParticleIndexMesh(ParentToChildrenMap, Constraint.ParticleIndex);
+		ensure(ParticleIndexMesh != INDEX_NONE);
+		Collision.ParticleIndexMesh = ParticleIndexMesh;
+	}
+	// If Constraint.LevelsetIndex is a cluster store an index for a mesh in this cluster
+	if (ClusterIdsArray[Constraint.LevelsetIndex].NumChildren > 0)
+	{
+		int32 LevelsetIndexMesh = GetSolver()->GetParticleIndexMesh(ParentToChildrenMap, Constraint.LevelsetIndex);
+		ensure(LevelsetIndexMesh != INDEX_NONE);
+		Collision.LevelsetIndexMesh = LevelsetIndexMesh;
+	}
+#endif
+}
+
+void FGeometryCollectionPhysicsProxy::UpdateCollisionData(const FParticlesType& Particles, const FCollisionConstraintsType& CollisionRule, FRecordedFrame* ExistingFrame)
+{
+	ExistingFrame->Collisions.Reset();
+
+	if (Parameters.CollisionData.SaveCollisionData && ExistingFrame->Timestamp > 0.f && Parameters.CollisionData.CollisionDataSizeMax > 0)
+	{
+#if TODO_REIMPLEMENT_PHYSICS_PROXY_REVERSE_MAPPING
+		const TArray<Chaos::TPBDCollisionConstraints<float, 3>::FRigidBodyContactConstraint>& AllConstraintsArray = CollisionRule.GetAllConstraints();
+		if (AllConstraintsArray.Num() > 0)
+		{
+			const Chaos::TArrayCollectionArray<PhysicsProxyWrapper>& PhysicsProxyReverseMapping = GetSolver()->GetPhysicsProxyReverseMapping();
+
+			TArray<Chaos::TPBDCollisionConstraints<float, 3>::FRigidBodyContactConstraint> ConstraintsArray;
+			ConstraintsArray.SetNumUninitialized(AllConstraintsArray.Num());
+
+			FBox BoundingBox(ForceInitToZero);
+			int32 NumConstraints = 0;
+			for (int32 Idx = 0; Idx < AllConstraintsArray.Num(); ++Idx)
+			{
+				// Check if the collision is for this PhysicsProxy
+				IPhysicsProxyBase* PhysicsProxy = PhysicsProxyReverseMapping[AllConstraintsArray[Idx].ParticleIndex].PhysicsProxy;
+				if (PhysicsProxy == this)
+				{
+					if (ensure(!AllConstraintsArray[Idx].AccumulatedImpulse.ContainsNaN()))
+					{
+						if (!AllConstraintsArray[Idx].AccumulatedImpulse.IsZero())
+						{
+							if (ensure(!AllConstraintsArray[Idx].Location.ContainsNaN() &&
+								!AllConstraintsArray[Idx].Normal.ContainsNaN()) &&
+								!Particles.V(AllConstraintsArray[Idx].ParticleIndex).ContainsNaN() &&
+								!Particles.V(AllConstraintsArray[Idx].LevelsetIndex).ContainsNaN() &&
+								!Particles.W(AllConstraintsArray[Idx].ParticleIndex).ContainsNaN() &&
+								!Particles.W(AllConstraintsArray[Idx].LevelsetIndex).ContainsNaN())
+							{
+								BoundingBox += AllConstraintsArray[Idx].Location;
+
+								ConstraintsArray[NumConstraints] = AllConstraintsArray[Idx];
+								NumConstraints++;
+							}
+						}
+					}
+				}
+			}
+
+			ConstraintsArray.SetNum(NumConstraints);
+
+			if (ConstraintsArray.Num() > 0)
+			{
+				if (Parameters.CollisionData.DoCollisionDataSpatialHash &&
+					Parameters.CollisionData.CollisionDataSpatialHashRadius > 0.f &&
+					ConstraintsArray.Num() > 1 &&
+					(BoundingBox.GetExtent().X > 0.f || BoundingBox.GetExtent().Y > 0.f || BoundingBox.GetExtent().Z > 0.f))
+				{
+					// Validate Parameters.CollisionData.CollisionDataSpatialHashRadius
+					// CellSize must be smaller than the smallest bbox extent
+					float SpatialHashRadius = Parameters.CollisionData.CollisionDataSpatialHashRadius;
+					TArray<float> ExtentArray = { BoundingBox.GetExtent().X , BoundingBox.GetExtent().Y, BoundingBox.GetExtent().Z };
+					ExtentArray.Sort();
+					if (ExtentArray[0] != 0.f)
+					{
+						if (2.f * SpatialHashRadius > ExtentArray[0])
+						{
+							SpatialHashRadius = 0.5f * ExtentArray[0];
+						}
+					}
+					else
+					{
+						if (ExtentArray[1] != 0.f)
+						{
+							if (2.f * SpatialHashRadius > ExtentArray[1])
+							{
+								SpatialHashRadius = 0.5f * ExtentArray[1];
+							}
+						}
+						else
+						{
+							if (2.f * SpatialHashRadius > ExtentArray[2])
+							{
+								SpatialHashRadius = 0.5f * ExtentArray[2];
+							}
+						}
+					}
+
+					// Spatial hash the constraints
+					TMultiMap<int32, int32> HashTableMap;
+					ComputeHashTable(ConstraintsArray, BoundingBox, HashTableMap, SpatialHashRadius);
+
+					TArray<int32> UsedCellsArray;
+					HashTableMap.GetKeys(UsedCellsArray);
+
+					int32 NumCollisionsThisFrame = 0;
+					for (int32 IdxCell = 0; IdxCell < UsedCellsArray.Num(); ++IdxCell)
+					{
+						TArray<int32> ConstraintsInCellArray;
+						HashTableMap.MultiFind(UsedCellsArray[IdxCell], ConstraintsInCellArray);
+
+						int32 NumConstraintsToGetFromCell = FMath::Min(Parameters.CollisionData.MaxCollisionPerCell, ConstraintsInCellArray.Num());
+						for (int32 IdxConstraint = 0; IdxConstraint < NumConstraintsToGetFromCell; ++IdxConstraint)
+						{
+							AddCollisionToCollisionData(ExistingFrame, Particles, ConstraintsArray[ConstraintsInCellArray[IdxConstraint]]);
+						}
+					}
+
+					if (ExistingFrame->Collisions.Num() > Parameters.CollisionData.CollisionDataSizeMax)
+					{
+						TArray<FSolverCollisionData> CollisionsArray1;
+
+						float FInc = (float)ExistingFrame->Collisions.Num() / (float)Parameters.CollisionData.CollisionDataSizeMax;
+
+						CollisionsArray1.SetNumUninitialized(Parameters.CollisionData.CollisionDataSizeMax);
+						for (int32 IdxCollision = 0; IdxCollision < Parameters.CollisionData.CollisionDataSizeMax; ++IdxCollision)
+						{
+							int32 NewIdx = FMath::FloorToInt((float)IdxCollision * FInc);
+							CollisionsArray1[IdxCollision] = ExistingFrame->Collisions[NewIdx];
+						}
+
+						ExistingFrame->Collisions.SetNumUninitialized(Parameters.CollisionData.CollisionDataSizeMax);
+						for (int32 IdxCollision = 0; IdxCollision < Parameters.CollisionData.CollisionDataSizeMax; ++IdxCollision)
+						{
+							ExistingFrame->Collisions[IdxCollision] = CollisionsArray1[IdxCollision];
+						}
+					}
+				}
+				else
+				{
+					if (ConstraintsArray.Num() <= Parameters.CollisionData.CollisionDataSizeMax)
+					{
+						for (int32 IdxConstraint = 0; IdxConstraint < ConstraintsArray.Num(); ++IdxConstraint)
+						{
+							AddCollisionToCollisionData(ExistingFrame, Particles, ConstraintsArray[IdxConstraint]);
+						}
+					}
+					else
+					{
+						float FInc = (float)ConstraintsArray.Num() / (float)Parameters.CollisionData.CollisionDataSizeMax;
+
+						for (int32 IdxConstraint = 0; IdxConstraint < Parameters.CollisionData.CollisionDataSizeMax; ++IdxConstraint)
+						{
+							int32 Idx = FMath::FloorToInt((float)IdxConstraint * FInc);
+
+							AddCollisionToCollisionData(ExistingFrame, Particles, ConstraintsArray[Idx]);
+						}
+					}
+				}
+			}
+		}
+#endif
+	}
+}
+
+void FGeometryCollectionPhysicsProxy::AddBreakingToBreakingData(FRecordedFrame* ExistingFrame, const FParticlesType& Particles, const Chaos::TBreakingData<float, 3>& Breaking)
+{
+#if TODO_REIMPLEMENT_RIGID_CLUSTERING
+	const Chaos::TArrayCollectionArray<Chaos::ClusterId>& ClusterIdsArray = GetSolver()->GetRigidClustering().GetClusterIdsArray();
+	const TMap<uint32, TUniquePtr<TArray<uint32>>>& ParentToChildrenMap = GetSolver()->GetRigidClustering().GetChildrenMap();
+	const Chaos::TArrayCollectionArray<PhysicsProxyWrapper>& PhysicsProxyReverseMapping = GetSolver()->GetPhysicsProxyReverseMapping();
+
+	int32 NewIdx = ExistingFrame->Breakings.Add(FSolverBreakingData());
+	FSolverBreakingData& NewBreaking = ExistingFrame->Breakings[NewIdx];
+
+	NewBreaking.Location = Breaking.Location;
+	NewBreaking.Velocity = Particles.V(Breaking.ParticleIndex);
+	NewBreaking.AngularVelocity = Particles.W(Breaking.ParticleIndex);
+	NewBreaking.Mass = Particles.M(Breaking.ParticleIndex);
+	NewBreaking.ParticleIndex = Breaking.ParticleIndex;
+
+	// If Breaking.ParticleIndex is a cluster store an index for a mesh in this cluster
+	if (ClusterIdsArray[Breaking.ParticleIndex].NumChildren > 0)
+	{
+		int32 ParticleIndexMesh = GetSolver()->GetParticleIndexMesh(ParentToChildrenMap, Breaking.ParticleIndex);
+		ensure(ParticleIndexMesh != INDEX_NONE);
+		NewBreaking.ParticleIndexMesh = ParticleIndexMesh;
+	}
+#endif
+}
+
+void FGeometryCollectionPhysicsProxy::UpdateBreakingData(const FParticlesType& Particles, FRecordedFrame* ExistingFrame)
+{
+	ExistingFrame->Breakings.Reset();
+
+	if (Parameters.BreakingData.SaveBreakingData && ExistingFrame->Timestamp > 0.f && Parameters.BreakingData.BreakingDataSizeMax > 0)
+	{
+#if TODO_REIMPLEMENT_GETALLCLUSTERBREAKINGS
+		const TArray<Chaos::TBreakingData<float, 3>>& AllBreakingsArray = GetSolver()->GetAllClusterBreakings();
+		if (AllBreakingsArray.Num() > 0)
+		{
+			const Chaos::TArrayCollectionArray<PhysicsProxyWrapper>& PhysicsProxyReverseMapping = GetSolver()->GetPhysicsProxyReverseMapping();
+
+			TArray<Chaos::TBreakingData<float, 3>> BreakingsArray;
+			BreakingsArray.SetNumUninitialized(AllBreakingsArray.Num());
+
+			FBox BoundingBox(ForceInitToZero);
+			int32 NumBreakings = 0;
+			for (int32 Idx = 0; Idx < AllBreakingsArray.Num(); ++Idx)
+			{
+				// Check if the breaking is for this PhysicsProxy
+				IPhysicsProxyBase* PhysicsProxy = PhysicsProxyReverseMapping[AllBreakingsArray[Idx].ParticleIndex].PhysicsProxy;
+				if (PhysicsProxy == this)
+				{
+					if (ensure(!AllBreakingsArray[Idx].Location.ContainsNaN() &&
+						!Particles.V(AllBreakingsArray[Idx].ParticleIndex).ContainsNaN() &&
+						!Particles.W(AllBreakingsArray[Idx].ParticleIndex).ContainsNaN()))
+					{
+						BoundingBox += AllBreakingsArray[Idx].Location;
+
+						BreakingsArray[NumBreakings] = AllBreakingsArray[Idx];
+						NumBreakings++;
+					}
+				}
+			}
+
+			BreakingsArray.SetNum(NumBreakings);
+
+			if (BreakingsArray.Num() > 0)
+			{
+				if (Parameters.BreakingData.DoBreakingDataSpatialHash &&
+					Parameters.BreakingData.BreakingDataSpatialHashRadius > 0.f &&
+					BreakingsArray.Num() > 1 &&
+					(BoundingBox.GetExtent().X > 0.f || BoundingBox.GetExtent().Y > 0.f || BoundingBox.GetExtent().Z > 0.f))
+				{
+					// Validate Parameters.CollisionData.CollisionDataSpatialHashRadius
+					// CellSize must be smaller than the smallest bbox extent
+					float SpatialHashRadius = Parameters.BreakingData.BreakingDataSpatialHashRadius;
+					TArray<float> ExtentArray = { BoundingBox.GetExtent().X , BoundingBox.GetExtent().Y, BoundingBox.GetExtent().Z };
+					ExtentArray.Sort();
+					if (ExtentArray[0] != 0.f)
+					{
+						if (2.f * SpatialHashRadius > ExtentArray[0])
+						{
+							SpatialHashRadius = 0.5f * ExtentArray[0];
+						}
+					}
+					else
+					{
+						if (ExtentArray[1] != 0.f)
+						{
+							if (2.f * SpatialHashRadius > ExtentArray[1])
+							{
+								SpatialHashRadius = 0.5f * ExtentArray[1];
+							}
+						}
+						else
+						{
+							if (2.f * SpatialHashRadius > ExtentArray[2])
+							{
+								SpatialHashRadius = 0.5f * ExtentArray[2];
+							}
+						}
+					}
+
+					// Spatial hash the constraints
+					TMultiMap<int32, int32> HashTableMap;
+					ComputeHashTable(BreakingsArray, BoundingBox, HashTableMap, SpatialHashRadius);
+
+					TArray<int32> UsedCellsArray;
+					HashTableMap.GetKeys(UsedCellsArray);
+
+					for (int32 IdxCell = 0; IdxCell < UsedCellsArray.Num(); ++IdxCell)
+					{
+						TArray<int32> BreakingsInCellArray;
+						HashTableMap.MultiFind(UsedCellsArray[IdxCell], BreakingsInCellArray);
+
+						int32 NumBreakingsToGetFromCell = FMath::Min(Parameters.BreakingData.MaxBreakingPerCell, BreakingsInCellArray.Num());
+						for (int32 IdxBreaking = 0; IdxBreaking < NumBreakingsToGetFromCell; ++IdxBreaking)
+						{
+							AddBreakingToBreakingData(ExistingFrame, Particles, BreakingsArray[BreakingsInCellArray[IdxBreaking]]);
+						}
+					}
+
+					if (ExistingFrame->Breakings.Num() > Parameters.BreakingData.BreakingDataSizeMax)
+					{
+						TArray<FSolverBreakingData> BreakingsArray1;
+
+						float FInc = (float)ExistingFrame->Breakings.Num() / (float)Parameters.BreakingData.BreakingDataSizeMax;
+
+						BreakingsArray1.SetNumUninitialized(Parameters.BreakingData.BreakingDataSizeMax);
+						for (int32 IdxBreaking = 0; IdxBreaking < Parameters.BreakingData.BreakingDataSizeMax; ++IdxBreaking)
+						{
+							int32 NewIdx = FMath::FloorToInt((float)IdxBreaking * FInc);
+							BreakingsArray1[IdxBreaking] = ExistingFrame->Breakings[NewIdx];
+						}
+
+						ExistingFrame->Breakings.SetNumUninitialized(Parameters.BreakingData.BreakingDataSizeMax);
+						for (int32 IdxBreaking = 0; IdxBreaking < Parameters.BreakingData.BreakingDataSizeMax; ++IdxBreaking)
+						{
+							ExistingFrame->Breakings[IdxBreaking] = BreakingsArray1[IdxBreaking];
+						}
+					}
+				}
+				else
+				{
+					if (BreakingsArray.Num() <= Parameters.BreakingData.BreakingDataSizeMax)
+					{
+						for (int32 IdxBreaking = 0; IdxBreaking < BreakingsArray.Num(); ++IdxBreaking)
+						{
+							AddBreakingToBreakingData(ExistingFrame, Particles, BreakingsArray[IdxBreaking]);
+						}
+					}
+					else
+					{
+						float FInc = (float)BreakingsArray.Num() / (float)Parameters.BreakingData.BreakingDataSizeMax;
+
+						for (int32 IdxBreaking = 0; IdxBreaking < Parameters.BreakingData.BreakingDataSizeMax; ++IdxBreaking)
+						{
+							int32 Idx = FMath::FloorToInt((float)IdxBreaking * FInc);
+
+							AddBreakingToBreakingData(ExistingFrame, Particles, BreakingsArray[Idx]);
+						}
+					}
+				}
+			}
+		}
+#endif
+	}
+}
+
+void FGeometryCollectionPhysicsProxy::UpdateTrailingData(const FParticlesType& Particles, FRecordedFrame* ExistingFrame)
+{
+	ExistingFrame->Trailings.Reset();
+
+	if (Parameters.TrailingData.SaveTrailingData && ExistingFrame->Timestamp > 0.f && Parameters.TrailingData.TrailingDataSizeMax > 0)
+	{
+		const float TrailingMinSpeedThresholdSquared = Parameters.TrailingData.TrailingMinSpeedThreshold * Parameters.TrailingData.TrailingMinSpeedThreshold;
+
+		// Find previous frame trailing data
+		int32 LastKey = RecordedTracks.FindLastKeyBefore(ExistingFrame->Timestamp);
+		if (LastKey != INDEX_NONE)
+		{
+			FRecordedFrame& PrevFrame = RecordedTracks.Records[LastKey];
+			if (PrevFrame.Trailings.Num() > 0)
+			{
+				for (FSolverTrailingData& Trailing : PrevFrame.Trailings)
+				{
+					if (Particles.ObjectState(Trailing.ParticleIndex) != Chaos::EObjectStateType::Sleeping &&
+						!Particles.Disabled(Trailing.ParticleIndex) &&
+						Particles.V(Trailing.ParticleIndex).SizeSquared() >= TrailingMinSpeedThresholdSquared)
+					{
+						ExistingFrame->Trailings.Add(Trailing);
+					}
+				}
+			}
+		}
+
+		if (ExistingFrame->Trailings.Num() < Parameters.TrailingData.TrailingDataSizeMax)
+		{
+#if TODO_REIMPLEMENT_PHYSICS_PROXY_REVERSE_MAPPING
+			const Chaos::TArrayCollectionArray<PhysicsProxyWrapper>& PhysicsProxyReverseMapping = GetSolver()->GetPhysicsProxyReverseMapping();
+			const Chaos::TArrayCollectionArray<Chaos::ClusterId>& ClusterIdsArray = GetSolver()->GetRigidClustering().GetClusterIdsArray();
+			const TMap<uint32, TUniquePtr<TArray<uint32>>>& ParentToChildrenMap = GetSolver()->GetRigidClustering().GetChildrenMap();
+
+			TArray<FSolverTrailingData> AllTrailings;
+			for (int32 IdxParticle = 0; IdxParticle < (int32)Particles.Size(); ++IdxParticle)
+			{
+				// Check if the particle is for this PhysicsProxy
+				IPhysicsProxyBase* PhysicsProxy = PhysicsProxyReverseMapping[IdxParticle].PhysicsProxy;
+				if (PhysicsProxy == this)
+				{
+					if (ensure(FMath::IsFinite(Particles.InvM(IdxParticle))))
+					{
+						if (Particles.ObjectState(IdxParticle) != Chaos::EObjectStateType::Sleeping &&
+							!Particles.Disabled(IdxParticle) &&
+							Particles.InvM(IdxParticle) != 0.f)
+						{
+							if (Particles.Geometry(IdxParticle) && 
+								Particles.Geometry(IdxParticle)->HasBoundingBox())
+							{
+								if (ensure(!Particles.X(IdxParticle).ContainsNaN()) &&
+									!Particles.V(IdxParticle).ContainsNaN() &&
+									!Particles.W(IdxParticle).ContainsNaN() &&
+									FMath::IsFinite(Particles.M(IdxParticle)))
+								{
+									Chaos::TAABB<float, 3> BoundingBox = Particles.Geometry(IdxParticle)->BoundingBox();
+									Chaos::TVector<float, 3> Extents = BoundingBox.Extents();
+									float Volume = Extents[0] * Extents[1] * Extents[2];
+									float SpeedSquared = Particles.V(IdxParticle).SizeSquared();
+
+									if (SpeedSquared >= TrailingMinSpeedThresholdSquared &&
+										Volume > Parameters.TrailingData.TrailingMinVolumeThreshold)
+									{
+										FSolverTrailingData TrailingData(Particles.X(IdxParticle),
+											Particles.V(IdxParticle),
+											Particles.W(IdxParticle),
+											Particles.M(IdxParticle),
+											IdxParticle,
+											INDEX_NONE);
+
+										// If IdxParticle is a cluster store an index for a mesh in this cluster
+										if (ClusterIdsArray[IdxParticle].NumChildren > 0)
+										{
+											int32 ParticleIndexMesh = GetSolver()->GetParticleIndexMesh(ParentToChildrenMap, IdxParticle);
+											ensure(ParticleIndexMesh != INDEX_NONE);
+											TrailingData.ParticleIndexMesh = ParticleIndexMesh;
+										}
+
+										FSetElementId Id = ExistingFrame->Trailings.FindId(TrailingData);
+										if (!Id.IsValidId())
+										{
+											AllTrailings.Add(TrailingData);
+										}
+										else
+										{
+											ExistingFrame->Trailings[Id].Location = Particles.X(IdxParticle);
+											ExistingFrame->Trailings[Id].Velocity = Particles.V(IdxParticle);
+											ExistingFrame->Trailings[Id].AngularVelocity = Particles.W(IdxParticle);
+										}
+									}
+								}
+							}
+						}
+					}
+				}
+			}
+
+			if (AllTrailings.Num() > 0)
+			{
+				int32 NumTrailingsToAdd = Parameters.TrailingData.TrailingDataSizeMax - ExistingFrame->Trailings.Num();
+
+				if (AllTrailings.Num() <= NumTrailingsToAdd)
+				{
+					for (int32 IdxTrailing = 0; IdxTrailing < AllTrailings.Num(); ++IdxTrailing)
+					{
+						ExistingFrame->Trailings.Add(AllTrailings[IdxTrailing]);
+					}
+				}
+				else
+				{
+					float FInc = (float)AllTrailings.Num() / (float)(NumTrailingsToAdd);
+
+					for (int32 IdxTrailing = 0; IdxTrailing < NumTrailingsToAdd; ++IdxTrailing)
+					{
+						int32 Idx = FMath::FloorToInt((float)IdxTrailing * FInc);
+						ExistingFrame->Trailings.Add(AllTrailings[Idx]);
+					}
+				}
+			}
+#endif
+		}
+	}
+}
+
+void FGeometryCollectionPhysicsProxy::UpdateRecordedState(float SolverTime, const TManagedArray<int32>& InRigidBodyID, const TManagedArray<int32>& InCollectionClusterID, const Chaos::TArrayCollectionArray<bool>& InInternalCluster, const FParticlesType& Particles, const FCollisionConstraintsType& CollisionRule)
+{
+	FRecordedFrame* ExistingFrame = RecordedTracks.FindRecordedFrame(SolverTime);
+
+	if(!ExistingFrame)
+	{
+		ExistingFrame = &InsertRecordedFrame(RecordedTracks, SolverTime);
+	}
+
+	ExistingFrame->Reset(InRigidBodyID.Num());
+	ExistingFrame->Timestamp = SolverTime;
+
+	// Collision
+	UpdateCollisionData(Particles, CollisionRule, ExistingFrame);
+
+	// Breaking
+	UpdateBreakingData(Particles, ExistingFrame);
+
+	// Trailing
+	UpdateTrailingData(Particles, ExistingFrame);
+
+	ParallelFor(InRigidBodyID.Num(), [&](int32 Index)
+	{
+		const int32 ExternalIndex = InRigidBodyID[Index];
+
+		if(ExternalIndex >= 0)
+		{
+			FTransform& NewTransform = ExistingFrame->Transforms[Index];
+
+			NewTransform.SetTranslation(Particles.X(ExternalIndex));
+			NewTransform.SetRotation(Particles.R(ExternalIndex));
+			NewTransform.SetScale3D(FVector(1.0f));
+
+			int32 ClusterParentIndex = InCollectionClusterID[Index];
+			if (ClusterParentIndex != INDEX_NONE && InInternalCluster[ClusterParentIndex])
+			{
+				ExistingFrame->DisabledFlags[Index] = Particles.Disabled(ClusterParentIndex);
+			}
+			else
+			{
+				ExistingFrame->DisabledFlags[Index] = Particles.Disabled(ExternalIndex);
+			}
+		}
+	});
+}
+
+void FGeometryCollectionPhysicsProxy::OnRemoveFromSolver(Chaos::FPBDRigidsSolver *RBDSolver)
+{
+	const FGeometryDynamicCollection& DynamicCollection = *GameToPhysInterchange.GetConsumerBuffer();
+
+	for (const Chaos::TPBDRigidClusteredParticleHandle<float, 3>* Handle : DynamicCollection.SolverClusterHandles)
+	{
+		RBDSolver->RemoveParticleToProxy(Handle);
+	}
+}
+
+void FGeometryCollectionPhysicsProxy::OnRemoveFromScene()
+{
+#if TODO_REIMPLEMENT_GET_RIGID_PARTICLES
+	// #BG TODO This isn't great - we currently cannot handle things being removed from the solver.
+	// need to refactor how we handle this and actually remove the particles instead of just constantly
+	// growing the array. Currently everything is just tracked by index though so the solver will have
+	// to notify all the proxies that a chunk of data was removed - or use a sparse array (undesireable)
+	Chaos::FPhysicsSolver::FParticlesType& Particles = GetSolver()->GetRigidParticles();
+
+	// #BG TODO Special case here because right now we reset/realloc the evolution per geom component
+	// in endplay which clears this out. That needs to not happen and be based on world shutdown
+	if(Particles.Size() == 0)
+	{
+		return;
+	}
+
+	const int32 Begin = BaseParticleIndex;
+	const int32 Count = NumParticles;
+
+	if (ensure((int32)Particles.Size() > 0 && (Begin + Count) <= (int32)Particles.Size()))
+	{
+		for (int32 ParticleIndex = 0; ParticleIndex < Count; ++ParticleIndex)
+		{
+			GetSolver()->GetEvolution()->DisableParticle(Begin + ParticleIndex);
+			GetSolver()->GetRigidClustering().GetTopLevelClusterParents().Remove(Begin + ParticleIndex);
+		}
+	}
+#endif
+}
+
+void FGeometryCollectionPhysicsProxy::SyncBeforeDestroy()
+{
+	if(FinalSyncFunc)
+	{
+		FinalSyncFunc(RecordedTracks);
+	}
+}
+
+void FGeometryCollectionPhysicsProxy::BufferGameState() 
+{
+	//
+	// NewData() redirects to this function.
+	/**
+	 * CONTEXT: GAMETHREAD
+	 * Returns a new unmanaged allocation of the data saved on the handle, otherwise nullptr
+	 */
+
+	// The other physics proxies allocate memory and returns a copy of game thread
+	// data as a TGeometryParticleData<T,d>, TKinematicGeometryParticleData<T,d>, 
+	// or a TPBDRigidParticleData<T,d>, which are all derived from FParticleData.
+	// That then gets passed to the physics thread via PushToPhysicsState().
+	//
+	// So all we're doing is pushing game thread state to the physics thread.
+	// I don't think we need to allocate memory to do this.  I think the other 
+	// proxies get away with that because they're dealing with a low number of 
+	// bodies; which is not the case here. Rather, let's just use a triple buffer.
+/*
+	if (!Parameters.Simulating || 
+		!ensure(NumParticles != INDEX_NONE)) // Make sure InitBodiesGT() has been called!
+	{
+		return;
+	}
+
+	FGeometryCollectionResults& Buffer = *GameToPhysInterchange.AccessProducerBuffer();
+	if (Buffer.NumParticlesAdded != NumParticles)
+	{
+		Buffer.InitArrays(*GTDynamicCollection);
+
+		// Initialize the buffer
+		Buffer.NumParticlesAdded = NumParticles;
+		Buffer.BaseIndex = BaseParticleIndex;
+
+		Buffer.SharedGeometry.SetNum(NumParticles);
+		Buffer.ShapeSimData.SetNum(NumParticles);
+		Buffer.ShapeQueryData.SetNum(NumParticles);
+
+		for (int32 Idx = 0; Idx < SimulatableParticles.Num(); Idx++)
+		{
+			Buffer.DisabledStates[Idx] = !SimulatableParticles[Idx];
+			Buffer.SharedGeometry[Idx] = Implicits[Idx];
+
+
+			// Hard coded to collide all
+			// todo(chaos) : Pull shape data from the interface. 
+			for (int32 Gdx = 0; Gdx<Buffer.SharedGeometry.Num();Gdx++)
+			{
+				int32 Idx1 = Buffer.ShapeQueryData[Idx].Add(FCollisionFilterData());
+				int32 Idx2 = Buffer.ShapeSimData[Idx].Add(FCollisionFilterData());
+				ensure(Idx1 == Idx2);
+
+				Buffer.ShapeSimData[Idx][Idx1].Word1 = 0xFFFF; // this body channel
+				Buffer.ShapeSimData[Idx][Idx1].Word3 = 0xFFFF; // collision candidate channels
+			}
+		}
+	}
+
+	const FGeometryDynamicCollection* DynamicCollection = Parameters.DynamicCollection;
+	const FGeometryCollection* RestCollection = Parameters.RestCollection;
+
+	//Buffer.TransformIndex.Init(RestCollection->TransformIndex);
+//	Buffer.BoneMap.Init(RestCollection->BoneMap);
+//	Buffer.Parent.Init(RestCollection->Parent);
+//	Buffer.Children.Init(RestCollection->Children);
+//	Buffer.SimulationType.Init(RestCollection->SimulationType);
+
+//	Buffer.DynamicState.Init(DynamicCollection->DynamicState);
+//	Buffer.Mass.Init(RestCollection->GetAttribute<float>("Mass", FTransformCollection::TransformGroup));
+//	Buffer.InertiaTensor.Init(RestCollection->GetAttribute<FVector>("InertiaTensor", FTransformCollection::TransformGroup));
+	Buffer.Transforms.Init(DynamicCollection->Transform);
+
+	GeometryCollectionAlgo::GlobalMatrices(
+		DynamicCollection->Transform, DynamicCollection->Parent, Buffer.GlobalTransforms);
+	
+	// Make the buffer available to the consumer.
+	GameToPhysInterchange.FlipProducer(); 
+	
+	// If we wanted to emulate the other proxies, we could augment the 
+	// FStubGeometryCollectionData class to carry along the consumer buffer:
+	//
+	// return new FStubGeometryCollectionData(GametoPhysInterchange.GetConsumerBuffer());
+	//
+	// However, at this point, it seems pretty unnecessary to do the additional
+	// allocation and deallocation of the transient FStubGeometryCollectionData
+	// class, as I believe all these operations are happening with the same 
+	// instance of this class.  So skip all that, and just return null.
+	return;
+	*/
+}
+
+void FGeometryCollectionPhysicsProxy::PushToPhysicsState(const Chaos::FParticleData *InData)
+{
+	/**
+	 * CONTEXT: GAMETHREAD -> to -> PHYSICSTHREAD
+	 * Called on the game thread when the solver is about to advance forward. This
+	 * callback should Enqueue commands on the PhysicsThread to update the state of
+	 * the solver
+	 */
+
+	/*
+	@todo break everything : move to the init
+	// The other proxies do a deep copy into dynamic transient memory, passed into
+	// this function via InData.  Instead of that, we're using persistent memory
+	// in a buffer to facilitate communicating between threads.
+
+	const FGeometryCollectionResults* GState = GameToPhysInterchange.GetConsumerBuffer();
+	if (!GState)
+	{
+		// This will happen if GameToPhysInterchange hasn't been updated since
+		// the last time GetConsumerBuffer() was called. 
+		return;
+	}
+
+	//const FStubGeometryCollectionData* GCData = static_cast<const FStubGeometryCollectionData*>(InData);
+	//const FGeometryCollectionResults* GCState = GCData->GetStateData();
+
+	// Move the copied game thread data into the handles.
+	for (int32 HandleIdx = 0; HandleIdx < SolverParticleHandles.Num(); ++HandleIdx)
+	{
+		if (Chaos::TPBDRigidParticleHandle<float, 3>* Handle = SolverParticleHandles[HandleIdx])
+		{
+			const int32 Idx = GState->BaseIndex + HandleIdx;
+
+			FTransform ParticleTransform =
+				MassToLocal[Idx] * GState->Transforms[Idx] * Parameters.WorldTransform;
+
+			Handle->SetX(ParticleTransform.GetTranslation());
+			Handle->SetR(ParticleTransform.GetRotation());
+
+			// Particle disabled states are initialized on the physics thread, so 
+			// we don't want to pull this from the game thread
+			//Handle->SetDisabled(GState->DisabledStates[Idx]); ryan - fix this!
+
+			Handle->SetLinearEtherDrag(Parameters.LinearEtherDrag);
+			Handle->SetAngularEtherDrag(Parameters.AngularEtherDrag);
+
+			if (GState->SharedGeometry[Idx] && GState->SharedGeometry[Idx]->HasBoundingBox())
+			{
+				Handle->SetHasBounds(true);
+				Handle->SetLocalBounds(GState->SharedGeometry[Idx]->BoundingBox());
+
+				Chaos::TAABB<float, 3> WorldSpaceBox =
+					GState->SharedGeometry[Idx]->BoundingBox().TransformedAABB(
+						Chaos::TRigidTransform<float, 3>(Handle->X(), Handle->R()));
+				// TODO Ryan - WorldSpaceBox.ThickenSymmetrically(GState->V());
+				Handle->SetWorldSpaceInflatedBounds(WorldSpaceBox);
+			}
+
+			// TODO Ryan - check dirty flags like TGeometryParticleData::DirtyFlags.
+			// In the mean time, we only set once.
+			if (!Handle->SharedGeometry() && GState->SharedGeometry.IsValidIndex(Idx) && GState->SharedGeometry[Idx])
+			{
+				Handle->SetSharedGeometry(GState->SharedGeometry[Idx]);
+			}
+
+			int32 CurrShape = 0;
+			for (const TUniquePtr<Chaos::TPerShapeData<Chaos::FReal, 3>>& Shape : Handle->ShapesArray())
+			{
+				// TODO Ryan - hardwired for now:
+				Shape->bDisable = false;
+				Shape->CollisionTraceType = Chaos::EChaosCollisionTraceFlag::Chaos_CTF_UseSimpleAndComplex; // Chaos_CTF_UseDefault?
+				if (CurrShape < GState->ShapeSimData[Idx].Num())
+				{
+					Shape->SimData = GState->ShapeSimData[Idx][CurrShape];
+					Shape->QueryData = GState->ShapeQueryData[Idx][CurrShape];
+				}
+				++CurrShape;
+			}
+		} // if rigid handle
+	} // for all SolverParticleHandles
+
+	if (DisableGeometryCollectionGravity) // cvar
+	{
+		// Our assumption is that you'd only ever want to wholesale opt geometry 
+		// collections out of gravity for debugging, so we keep this conditional
+		// out of the loop above and on it's own.  This means we can't turn gravity
+		// back on once it's off, but even if we didn't enclose this in an if(),
+		// this function won't be called again unless something dirties the proxy.
+
+		Chaos::TPerParticleGravity<float, 3>& GravityForces = GetSolver()->GetEvolution()->GetGravityForces();
+		for (int32 HandleIdx = 0; HandleIdx < SolverParticleHandles.Num(); ++HandleIdx)
+		{
+			if (Chaos::TPBDRigidParticleHandle<float, 3>* Handle = SolverParticleHandles[HandleIdx])
+			{
+				GravityForces.SetEnabled(*Handle, false);
+			}
+		}
+	}
+	*/
+}
+
+void FGeometryCollectionPhysicsProxy::BufferPhysicsResults()
+{
+	/**
+	 * CONTEXT: PHYSICSTHREAD
+	 * Called per-tick after the simulation has completed. The proxy should cache the results of their
+	 * simulation into the local buffer. 
+	 */
+	using namespace Chaos;
+	SCOPE_CYCLE_COUNTER(STAT_CacheResultGeomCollection);
+
+	/*
+	@todo break everything
+	//FGeometryCollectionResults& TargetResults = Results.GetPhysicsDataForWrite();
+	FGeometryCollectionResults& TargetResults = *PhysToGameInterchange.AccessProducerBuffer();
+
+	if (TargetResults.NumTransformGroup() != PTDynamicCollection.NumElements(FGeometryCollection::TransformGroup))
+	{
+		TargetResults.InitArrays(PTDynamicCollection);
+
+		// Base particle index to calculate index from a global particle index on the game thread
+		TargetResults.BaseIndex = BaseParticleIndex;
+		TargetResults.NumParticlesAdded = NumParticles;
+	}
+
+	{ // tmp scope
+		IsObjectDynamic = false;
+		int32 Idx = 0;
+		for (auto* Handle : SolverParticleHandles)
+		{
+			if (!Handle || Handle->Disabled())
+			{
+				TargetResults.DisabledStates[Idx++] = true;
+				continue;
+			}
+			TargetResults.DisabledStates[Idx++] = false;
+			IsObjectDynamic = true;
+		}
+	} // tmp scope
+
+	// Advertise to game thread
+	TargetResults.IsObjectDynamic = IsObjectDynamic;
+	TargetResults.IsObjectLoading = IsObjectLoading;
+	
+	const FTransform& ActorToWorld = Parameters.WorldTransform;
+
+	const int32 TransfromGroupSize = GetTransformGroupSize();
+	if (EndFrameUnparentingBuffer.Num() < TransfromGroupSize)
+	{
+		// It's not worth shrinking the EndFrameUnparentingBuffer array, 
+		// at least until the solver supports deleting bodies.
+		EndFrameUnparentingBuffer.Init(INDEX_NONE, TransfromGroupSize);
+	}
+	TManagedArray<int32>& Parent = PTDynamicCollection.Parent;
+	TManagedArray<TSet<int32>>& Children = PTDynamicCollection.Children;
+
+	{ // tmp scope
+		SCOPE_CYCLE_COUNTER(STAT_CalcParticleToWorld);
+		// Fill particle to world transforms
+//		TargetResults.ParticleToWorldTransforms.SetNum(NumParticles);
+		//for (int32 TransformIndex = 0; TransformIndex < NumParticles; ++TransformIndex)
+		//{
+		//	//only update roots and first children
+		//	if (TargetResults.Parent[TransformIndex] == INDEX_NONE)
+		//	{
+		//		if (auto* Handle = SolverParticleHandles[TransformIndex])
+		//		{
+		//			const FTransform Xf(Handle->R(), Handle->X());
+		//			//const int32 ParticleIndex = BaseParticleIndex + TransformIndex;
+		//			TargetResults.ParticleToWorldTransforms[TransformIndex] = Xf;
+		//		}
+		//	}
+		//}
+
+		//const TArrayCollectionArray<Chaos::ClusterId>& ClusterID = GetSolver()->GetRigidClustering().GetClusterIdsArray();
+		//const TArrayCollectionArray<Chaos::TRigidTransform<float, 3>>& ClusterChildToParentMap = GetSolver()->GetRigidClustering().GetChildToParentMap();
+		//const TArrayCollectionArray<FMultiChildProxyId>& MultiChildProxyIdArray = GetSolver()->GetRigidClustering().GetMultiChildProxyIdArray();
+		//const TArrayCollectionArray<TUniquePtr<TMultiChildProxyData<float, 3>>>& MultiChildProxyDataArray = GetSolver()->GetRigidClustering().GetMultiChildProxyDataArray();
+
+		for (int32 TransformIndex = 0; TransformIndex < TransfromGroupSize; ++TransformIndex)
+		{
+			Chaos::TPBDRigidClusteredParticleHandle<float, 3>* Handle = SolverParticleHandles[TransformIndex];
+			if (!Handle)
+			{
+				continue;
+			}
+
+			// Dynamic state is also updated by the solver during field interaction.
+			if (!Handle->Sleeping())
+			{
+				const Chaos::EObjectStateType ObjectState = Handle->ObjectState();
+				switch (ObjectState)
+				{
+				case Chaos::EObjectStateType::Kinematic:
+					TargetResults.DynamicState[TransformIndex] = (int)EObjectStateTypeEnum::Chaos_Object_Kinematic;
+					break;
+				case Chaos::EObjectStateType::Static:
+					TargetResults.DynamicState[TransformIndex] = (int)EObjectStateTypeEnum::Chaos_Object_Static;
+					break;
+				case Chaos::EObjectStateType::Sleeping:
+					TargetResults.DynamicState[TransformIndex] = (int)EObjectStateTypeEnum::Chaos_Object_Sleeping;
+					break;
+				case Chaos::EObjectStateType::Dynamic:
+				case Chaos::EObjectStateType::Uninitialized:
+				default:
+					TargetResults.DynamicState[TransformIndex] = (int)EObjectStateTypeEnum::Chaos_Object_Dynamic;
+					break;
+				}
+			}
+
+			// Update the transform and parent hierarchy of the active rigid bodies. Active bodies can be either
+			// rigid geometry defined from the leaf nodes of the collection, or cluster bodies that drive an entire
+			// branch of the hierarchy within the GeometryCollection.
+			// - Active bodies are directly driven from the global position of the corresponding
+			//   rigid bodies within the solver ( cases where RigidBodyID[TransformGroupIndex] is not disabled ). 
+			// - Deactivated bodies are driven from the transforms of their active parents. However the solver can
+			//   take ownership of the parents during the simulation, so it might be necessary to force deactivated
+			//   bodies out of the collections hierarchy during the simulation.  
+			if (!Handle->Disabled())
+			{
+				// Update the transform of the active body. The active body can be either a single rigid
+				// or a collection of rigidly attached geometries (Clustering). The cluster is represented as a
+				// single transform in the GeometryCollection, and all children are stored in the local space
+				// of the parent cluster.
+				// ... When setting cluster transforms it is expected that the MassToLocal is identity.
+				//     Cluster initialization will set the vertices in the MassSpace of the rigid body.
+				// ... When setting individual rigid bodies that are not clustered, the MassToLocal will be 
+				//     non-Identity, and will reflect the difference between the geometric center of the geometry
+				//     and that corresponding rigid bodies center of mass. 
+
+				FTransform& ParticleToWorld = TargetResults.ParticleToWorldTransforms[TransformIndex];
+				ParticleToWorld = FTransform(Handle->R(), Handle->X());
+
+				TargetResults.Transforms[TransformIndex] =
+					MassToLocal[TransformIndex].GetRelativeTransformReverse(ParticleToWorld).GetRelativeTransform(ActorToWorld);
+				TargetResults.Transforms[TransformIndex].NormalizeRotation();
+
+				// If the parent of this NON DISABLED body is set to anything other than INDEX_NONE,
+				// then it was just unparented, likely either by rigid clustering or by fields.  We
+				// need to force all such enabled rigid bodies out of the transform hierarchy.
+				if (Parent[TransformIndex] != INDEX_NONE)
+				{
+					const int32 ParentIndex = Parent[TransformIndex];
+					// Children in the hierarchy are stored in a TSet, which is not thread safe.  
+					// So we retain indices to remove afterwards.
+					EndFrameUnparentingBuffer[TransformIndex] = ParentIndex;
+				}
+
+				// When a leaf node rigid body is removed from a cluster, the rigid
+				// body will become active and needs its clusterID updated.  This just
+				// syncs the clusterID all the time.
+				TPBDRigidParticleHandle<float, 3>* ClusterParentId = Handle->ClusterIds().Id;
+				SolverClusterID[TransformIndex] = ClusterParentId;
+			}
+			else // Handle->Disabled()
+			{
+				// The rigid body parent cluster has changed within the solver, and its
+				// parent body is not tracked within the geometry collection. So we need to
+				// pull the rigid bodies out of the transform hierarchy, and just drive
+				// the positions directly from the solvers cluster particle. 
+				TPBDRigidParticleHandle<float, 3>* ClusterParentId = Handle->ClusterIds().Id;
+				if(SolverClusterID[TransformIndex] != ClusterParentId)
+				{
+					// Force all driven rigid bodies out of the transform hierarchy
+					if (Parent[TransformIndex] != INDEX_NONE)
+					{
+						const int32 ParentIndex = Parent[TransformIndex];
+						// Children in the hierarchy are stored in a TSet, which is not thread safe.  So we retain
+						// indices to remove afterwards.
+						EndFrameUnparentingBuffer[TransformIndex] = ParentIndex;
+					}
+					SolverClusterID[TransformIndex] = Handle->ClusterIds().Id;
+				}
+
+				// Disabled rigid bodies that have valid cluster parents, and have been re-indexed by the
+				// solver (as in, They were re-clustered outside of the geometry collection), These clusters 
+				// will need to be rendered based on the clusters position. 
+
+				if (TPBDRigidParticleHandle<float, 3>* ParentHandle = Handle->ClusterIds().Id)
+				{
+					if (TPBDRigidClusteredParticleHandle<float, 3>* ClusteredParentHandle = ParentHandle->CastToClustered())
+					{
+						if (ClusteredParentHandle->InternalCluster())
+						{
+							//const FTransform ClusterChildToWorld = ClusterChildToParentMap[RigidBodyIndex] * FTransform(Particles.R(ClusterParentIndex), Particles.X(ClusterParentIndex));
+							const TRigidTransform<float, 3>& ChildToParent = Handle->ChildToParent();
+							const FTransform ClusterChildToWorld = ChildToParent * FTransform(ParentHandle->R(), ParentHandle->X());
+							if (Parameters.IsCacheRecording())
+							{
+								Handle->SetX(ClusterChildToWorld.GetTranslation());
+								Handle->SetR(ClusterChildToWorld.GetRotation());
+							}
+							// GeomToActor = ActorToWorld.Inv() * ClusterChildToWorld * MassToLocal.Inv();
+							TargetResults.Transforms[TransformIndex] = 
+								MassToLocal[TransformIndex].GetRelativeTransformReverse(ClusterChildToWorld).GetRelativeTransform(ActorToWorld);
+							TargetResults.Transforms[TransformIndex].NormalizeRotation();
+						}
+					}
+				}
+
+			}// end if
+		} // end for
+	} // tmp scope
+
+	for (int32 TransformGroupIndex = 0; TransformGroupIndex < TransfromGroupSize; TransformGroupIndex++)
+	{
+		const int32 ParentIndex = EndFrameUnparentingBuffer[TransformGroupIndex];
+		if (ParentIndex >= 0)
+		{
+			// We reuse EndFrameUnparentingBuffer potentially without reinitialization, 
+			// so reset this index to INDEX_NONE before it gets paged out.
+			EndFrameUnparentingBuffer[TransformGroupIndex] = INDEX_NONE;
+
+			Children[ParentIndex].Remove(TransformGroupIndex);
+			Parent[TransformGroupIndex] = INDEX_NONE;
+		}
+		TargetResults.Parent[TransformGroupIndex] = Parent[TransformGroupIndex];
+	}
+
+	// If object is dynamic, compute global matrices	
+	if (IsObjectDynamic || TargetResults.GlobalTransforms.Num() == 0)
+	{
+		SCOPE_CYCLE_COUNTER(STAT_CalcGlobalGCMatrices);
+		check(TargetResults.Transforms.Num() == TargetResults.Parent.Num());
+		GeometryCollectionAlgo::GlobalMatrices(
+			TargetResults.Transforms, TargetResults.Parent, TargetResults.GlobalTransforms);
+
+		// Compute world bounds.  This is a loose bounds based on the circumscribed box 
+		// of a bounding sphere for the geometry.		
+		SCOPE_CYCLE_COUNTER(STAT_CalcGlobalGCBounds);
+		FBox BoundingBox(ForceInit);
+		for (int i = 0; i < ValidGeometryBoundingBoxes.Num(); ++i)
+		{
+			BoundingBox += ValidGeometryBoundingBoxes[i].TransformBy(
+				TargetResults.GlobalTransforms[ValidGeometryTransformIndices[i]] * ActorToWorld);
+		}
+		TargetResults.WorldBounds = FBoxSphereBounds(BoundingBox);
+	}
+
+	// We pulled state out of particles.  Push results into PTDynamicCollection, 
+	// so that it has an updated view.
+	UpdateGeometryCollection(TargetResults);
+	*/
+}
+
+void FGeometryCollectionPhysicsProxy::FlipBuffer()
+{
+	/**
+	 * CONTEXT: PHYSICSTHREAD (Write Locked)
+	 * Called by the physics thread to signal that it is safe to perform any double-buffer flips here.
+	 * The physics thread has pre-locked an RW lock for this operation so the game thread won't be reading
+	 * the data
+	 */
+
+	PhysToGameInterchange.FlipProducer();
+}
+
+// Called from FPhysScene_ChaosInterface::SyncBodies(), NOT the solver.
+void FGeometryCollectionPhysicsProxy::PullFromPhysicsState()
+{
+	/**
+	 * CONTEXT: GAMETHREAD (Read Locked)
+	 * Perform a similar operation to Sync, but take the data from a gamethread-safe buffer. This will be called
+	 * from the game thread when it cannot sync to the physics thread. The simulation is very likely to be running
+	 * when this happens so never read any physics thread data here!
+	 *
+	 * Note: A read lock will have been acquired for this - so the physics thread won't force a buffer flip while this
+	 * sync is ongoing
+	 */
+	FGeometryDynamicCollection& DynamicCollection = *GameToPhysInterchange.AccessProducerBuffer();
+
+	const FGeometryCollectionResults* TargetResultPtr = PhysToGameInterchange.GetConsumerBuffer();
+	if (!TargetResultPtr)
+		return;
+	// Remove const-ness as we're going to do the ExchangeArrays() thing.
+	FGeometryCollectionResults& TR = *const_cast<FGeometryCollectionResults*>(TargetResultPtr);
+
+	// We should never be changing the number of entries, this would break other 
+	// attributes in the transform group.
+	const int32 NumTransforms = DynamicCollection.Transform.Num();
+	if (ensure(NumTransforms == TR.Transforms.Num()))	
+	{
+		for(int32 TmIndex = 0; TmIndex < NumTransforms; ++TmIndex)
+		{
+			if(!TR.DisabledStates[TmIndex])
+			{
+				DynamicCollection.Transform[TmIndex] = TR.Transforms[TmIndex];
+			}
+			DynamicCollection.DynamicState[TmIndex] = TR.DynamicState[TmIndex];
+		}
+
+		//question: why do we need this? Sleeping objects will always have to update GPU
+		DynamicCollection.MakeDirty();
+
+		if (CacheSyncFunc)
+		{
+			CacheSyncFunc(TR);
+		}
+	}
+}
+
+void FGeometryCollectionPhysicsProxy::UpdateGeometryCollection(
+	const FGeometryCollectionResults& GCResults,
+	FGeometryDynamicCollection* Collection)
+{
+	if (!Collection)
+	{
+		Collection = Parameters.DynamicCollection;
+	}
+	check(Collection);
+
+	TManagedArray<FTransform>& Transform = Collection->Transform;
+	TManagedArray<int32>& Parent = Collection->Parent;
+	TManagedArray<TSet<int32>>& Children = Collection->Children;
+	TManagedArray<int32>& SimulationType = Collection->SimulationType;
+	TManagedArray<int32>& DynamicState = Collection->GetAttribute<int32>("DynamicState", FGeometryCollection::TransformGroup);
+
+	// Particles X and R are aligned with center of mass and inertia principal axes.
+	// Renderer doesn't know about this and simply does ActorToWorld * GeomToActor * LocalSpaceVerts
+	// In proper math multiplication order:
+	//		ParticleToWorld = ActorToWorld * GeomToActor * MassToLocal
+	//		GeomToWorld = ActorToWorld * GeomToActor
+	//		=> GeomToWorld = ParticleToWorld * MassToLocal.Inv()
+	//		=> GeomToActor = ActorToWorld.Inv() * ParticleToWorld * MassToLocal.Inv()
+
+	const FTransform& ActorToWorld = Parameters.WorldTransform;
+	const int32 TransformSize = GetTransformGroupSize();
+	for (int32 TransformGroupIndex = 0; TransformGroupIndex < TransformSize; ++TransformGroupIndex)
+	{
+		if (!Collection->SimulatableParticles[TransformGroupIndex])
+			continue;
+
+		auto* Handle = Collection->SolverParticleHandles[TransformGroupIndex];
+
+		// dynamic state is also updated by the solver during field interaction. 
+		if(!Handle->Sleeping())
+		{
+			DynamicState[TransformGroupIndex] = GCResults.DynamicState[TransformGroupIndex];
+		}
+
+		// Update the transform and parent hierarchy of the active rigid bodies. Active bodies can be either
+		// rigid geometry defined from the leaf nodes of the collection, or cluster bodies that drive an entire
+		// branch of the hierarchy within the GeometryCollection.
+		// - Active bodies are directly driven from the global position of the corresponding
+		//   rigid bodies within the solver ( cases where RigidBodyID[TransformGroupIndex] is not disabled ). 
+		// - Deactivated bodies are driven from the transforms of their active parents. However the solver can
+		//   take ownership of the parents during the simulation, so it might be necessary to force deactivated
+		//   bodies out of the collections hierarchy during the simulation.  
+
+		// Static analysis build complains that the else clause does the same thing as the if/then,
+		// so remove the conditional for now...
+#if 0
+		if(!GCResults.DisabledStates[TransformGroupIndex])
+		{
+#endif
+			// Update the transform of the active body. The active body can be either a single rigid
+			// or a collection of rigidly attached geometries (Clustering). The cluster is represented as a
+			// single transform in the GeometryCollection, and all children are stored in the local space
+			// of the parent cluster.
+			// ... When setting cluster transforms it is expected that the MassToLocal is identity.
+			//     Cluster initialization will set the vertices in the MassSpace of the rigid body.
+			// ... When setting individual rigid bodies that are not clustered, the MassToLocal will be 
+			//     non-Identity, and will reflect the difference between the geometric center of the geometry
+			//     and that corresponding rigid bodies center of mass. 
+
+			// GeomToActor = ActorToWorld.Inv() * ParticleToWorld * MassToLocal.Inv();
+			Transform[TransformGroupIndex] = GCResults.Transforms[TransformGroupIndex];
+
+			// When a leaf node rigid body is removed from a cluster the rigid
+			// body will become active and needs its clusterID updated. This just
+			// syncs the clusterID all the time. 
+//			CollectionClusterID[TransformGroupIndex] = ClusterID[RigidBodyIndex].Id;
+#if 0
+		}
+		else //if (Particles.Disabled(RigidBodyIndex))
+		{
+			Transform[TransformGroupIndex] = GCResults.Transforms[TransformGroupIndex];
+		}
+#endif
+	}
+}
+
+
+//==============================================================================
+// STATIC SETUP FUNCTIONS
+//==============================================================================
+
+
+int32 FindSizeSpecificIdx(const TArray<FSharedSimulationSizeSpecificData>& SizeSpecificData, const FBox& Bounds)
+{
+	const FVector Extents = Bounds.GetExtent();
+	const float Size = Extents.GetAbsMin();
+	check(SizeSpecificData.Num());
+	int32 UseIdx = 0;
+	float PreSize = FLT_MAX;
+	for (int32 Idx = SizeSpecificData.Num() - 1; Idx >=0 ; --Idx)
+	{
+		ensureMsgf(PreSize >= SizeSpecificData[Idx].MaxSize, TEXT("SizeSpecificData is not sorted"));
+		PreSize = SizeSpecificData[Idx].MaxSize;
+		if (Size < SizeSpecificData[Idx].MaxSize)
+			UseIdx = Idx;
+		else
+			break;
+	}
+	return UseIdx;
+}
+
+void FGeometryCollectionPhysicsProxy::InitializeSharedCollisionStructures(
+	Chaos::FErrorReporter& ErrorReporter, 
+	FGeometryCollection& RestCollection, 
+	const FSharedSimulationParameters& SharedParams)
+{
+	FString BaseErrorPrefix = ErrorReporter.GetPrefix();
+
+	// fracture tools can create an empty GC before appending new geometry
+	if (RestCollection.NumElements(FGeometryCollection::GeometryGroup) == 0)
+		return;
+
+	//TArray<TArray<TArray<int32>>> BoundaryVertexIndices;
+	//GeometryCollectionAlgo::FindOpenBoundaries(&RestCollection, 1e-2, BoundaryVertexIndices);
+	//GeometryCollectionAlgo::TriangulateBoundaries(&RestCollection, BoundaryVertexIndices);
+	//RestCollection.ReindexMaterials();
+
+	using namespace Chaos;
+	const TManagedArray<bool>& Visible = RestCollection.Visible;
+
+	// TransformGroup
+	const TManagedArray<int32>& BoneMap = RestCollection.BoneMap;
+	const TManagedArray<int32>& Parent = RestCollection.Parent;
+	const TManagedArray<TSet<int32>>& Children = RestCollection.Children;
+	const TManagedArray<bool>& CollectionSimulatableParticles = 
+		RestCollection.GetAttribute<bool>(
+			FGeometryCollection::SimulatableParticlesAttribute, FTransformCollection::TransformGroup);
+
+	TManagedArray<FVector>& CollectionInertiaTensor = 
+		RestCollection.AddAttribute<FVector>(
+			TEXT("InertiaTensor"), FTransformCollection::TransformGroup);
+	TManagedArray<FTransform>& CollectionMassToLocal = 
+		RestCollection.AddAttribute<FTransform>(
+			TEXT("MassToLocal"), FTransformCollection::TransformGroup);
+	TManagedArray<float>& CollectionMass = 
+		RestCollection.AddAttribute<float>(
+			TEXT("Mass"), FTransformCollection::TransformGroup);
+	TManagedArray<TUniquePtr<FSimplicial>>& CollectionSimplicials = 
+		RestCollection.AddAttribute<TUniquePtr<FSimplicial>>(
+			FGeometryDynamicCollection::SimplicialsAttribute, FTransformCollection::TransformGroup);
+	TManagedArray<TSharedPtr<Chaos::FImplicitObject, ESPMode::ThreadSafe>>& CollectionImplicits = 
+		RestCollection.AddAttribute<TSharedPtr<Chaos::FImplicitObject, ESPMode::ThreadSafe>>(
+			FGeometryDynamicCollection::SharedImplicitsAttribute, FTransformCollection::TransformGroup);
+
+	FTransform IdentityXf(FQuat::Identity, FVector(0));
+	IdentityXf.NormalizeRotation();
+	CollectionMassToLocal.Fill(IdentityXf);
+
+	// VerticesGroup
+	const TManagedArray<FVector>& Vertex = RestCollection.Vertex;
+
+	// GeometryGroup
+	const TManagedArray<FBox>& BoundingBox = RestCollection.BoundingBox;
+	const TManagedArray<float>& InnerRadius = RestCollection.InnerRadius;
+	const TManagedArray<int32>& VertexCount = RestCollection.VertexCount;
+	const TManagedArray<int32>& VertexStart = RestCollection.VertexStart;
+	const TManagedArray<int32>& FaceCount = RestCollection.FaceCount;
+	const TManagedArray<int32>& FaceStart = RestCollection.FaceStart;
+	const TManagedArray<int32>& TransformIndex = RestCollection.TransformIndex;
+	const TManagedArray<FIntVector>& Indices = RestCollection.Indices;
+
+	TArray<FTransform> CollectionSpaceTransforms;
+	{ // tmp scope
+		const TManagedArray<FTransform>& HierarchyTransform = RestCollection.Transform;
+		GeometryCollectionAlgo::GlobalMatrices(HierarchyTransform, Parent, CollectionSpaceTransforms);
+		check(HierarchyTransform.Num() == CollectionSpaceTransforms.Num());
+	} // tmp scope
+
+	const int32 NumTransforms = CollectionSpaceTransforms.Num();
+	const int32 NumGeometries = RestCollection.NumElements(FGeometryCollection::GeometryGroup);
+
+	TArray<TUniquePtr<TTriangleMesh<float>>> TriangleMeshesArray;	//use to union trimeshes in cluster case
+	TriangleMeshesArray.AddDefaulted(NumTransforms);
+
+	TArray<TSet<int32>> VertsAddedArray;
+	VertsAddedArray.AddDefaulted(NumGeometries);
+
+	TParticles<float, 3> MassSpaceParticles;
+	MassSpaceParticles.AddParticles(Vertex.Num());
+	for (int32 Idx = 0; Idx < Vertex.Num(); ++Idx)
+	{
+		MassSpaceParticles.X(Idx) = Vertex[Idx];	//mass space computation done later down
+	}
+
+	TArray<TMassProperties<float,3>> MassPropertiesArray;
+	MassPropertiesArray.AddUninitialized(NumGeometries);
+
+	TArray<bool> InertiaComputationNeeded;
+	InertiaComputationNeeded.AddUninitialized(NumGeometries);
+
+	const float MinBoundsExtents = 10.f;
+
+	float TotalVolume = 0.f;
+	int32 NumSimulatableParticles = 0;
+
+	// The geometry group has a set of transform indices that maps a geometry index
+	// to a transform index, but only in the case where there is a 1-to-1 mapping 
+	// between the two.  In the event where a geometry is instanced for multiple
+	// transforms, the transform index on the geometry group should be INDEX_NONE.
+	// Otherwise, iterating over the geometry group is a convenient way to iterate
+	// over all the leaves of the hierarchy.
+	check(!TransformIndex.Contains(INDEX_NONE)); // TODO: implement support for instanced bodies
+	for (int32 GeometryIndex = 0; GeometryIndex < NumGeometries; GeometryIndex++)
+	{
+		const int32 TransformGroupIndex = TransformIndex[GeometryIndex];
+		if (CollectionSimulatableParticles[TransformGroupIndex])
+		{
+			NumSimulatableParticles++;
+
+			TUniquePtr<TTriangleMesh<float>> TriMesh(
+				CreateTriangleMesh(
+					FaceCount[GeometryIndex], 
+					VertexStart[GeometryIndex], 
+					FaceStart[GeometryIndex], 
+					Vertex, 
+					Visible, 
+					Indices, 
+					VertsAddedArray[GeometryIndex]));
+
+			TMassProperties<float, 3>& MassProperties = MassPropertiesArray[GeometryIndex];
+
+			//CalculateVolumeAndCenterOfMass(MassSpaceParticles, *TriMesh, MassProperties.Volume, MassProperties.CenterOfMass);
+			//InertiaComputationNeeded[GeometryIndex] = true;
+			//if (MassProperties.Volume < MinVolume)
+			//just assume everything is a box, need better computation later but requires mesh cleanup
+			{
+				//since volume is too small we just use a fallback one
+				InertiaComputationNeeded[GeometryIndex] = false;	
+
+				CollectionMassToLocal[TransformGroupIndex] = 
+					FTransform(TRotation<float, 3>(FQuat(0, 0, 0, 1)), BoundingBox[GeometryIndex].GetCenter());
+
+
+				const TVector<float, 3> Size = BoundingBox[GeometryIndex].GetSize().ComponentMax(FVector(MinBoundsExtents));
+				MassProperties.Volume = Size.Product();	//just fake a box if tiny
+				MassProperties.InertiaTensor = 
+					Chaos::TBox<float, 3>::GetInertiaTensor(
+						1.0f, // Mass
+						Size);
+				/*
+				if (MassProperties.Volume < MinVolume)
+				{
+					MassProperties.Volume = MinVolume;	//for thin shells volume is 0. Just use min volume to avoid divie by 0. Probably needs more thought
+
+					//if all dimensions are tiny we should probably fix content
+					if(!ensureMsgf(BoundingBox[GeometryIndex].GetExtent().GetAbsMin() > MinVolume, TEXT("Geometry too small to simulate. Idx (%d)"), GeometryIndex))
+					{
+						ErrorReporter.ReportError(*FString::Printf(TEXT("Geometry too small to simulate. Idx(%d)"), GeometryIndex));
+						CollectionSimulatableParticles[TransformGroupIndex] = false;	//do not simulate tiny particles
+						ErrorReporter.HandleLatestError();
+					}
+				}*/
+			}
+#if 0
+			else if (MassProperties.Volume > BoundsVolume)
+			{
+				ensure(false);
+				ErrorReporter.ReportError(*FString::Printf(TEXT("Geometry has invalid volume")));
+				ErrorReporter.HandleLatestError();
+
+				//somehow ended up with huge volume, just use bounds volume
+				InertiaComputationNeeded[GeometryIndex] = false;
+				CollectionMassToLocal[TransformGroupIndex] = FTransform(TRotation<float, 3>(FQuat(0, 0, 0, 1)), BoundingBox[GeometryIndex].GetCenter());
+
+				FVector Size = BoundingBox[GeometryIndex].GetSize();
+				FVector SideSquared(Size.X * Size.X, Size.Y * Size.Y, Size.Z * Size.Z);
+				MassProperties.Volume = BoundsVolume;
+				MassProperties.InertiaTensor = PMatrix<float, 3, 3>((SideSquared.Y + SideSquared.Z) / 12.f, (SideSquared.X + SideSquared.Z) / 12.f, (SideSquared.X + SideSquared.Y) / 12.f);
+			}
+#endif
+
+			TotalVolume += MassProperties.Volume;
+			TriangleMeshesArray[TransformGroupIndex] = MoveTemp(TriMesh);
+		}
+	}
+
+	ensureMsgf(TotalVolume > KINDA_SMALL_NUMBER || !NumSimulatableParticles, TEXT("Geometry collection too small"));
+	TotalVolume = FMath::Max(TotalVolume, MinBoundsExtents*MinBoundsExtents*MinBoundsExtents);
+	//User provides us with total mass or density.
+	//Density must be the same for individual parts and the total. Density_i = Density = Mass_i / Volume_i
+	//Total mass must equal sum of individual parts. Mass_i = TotalMass * Volume_i / TotalVolume => Density_i = TotalMass / TotalVolume
+	const float DesiredTotalMass = SharedParams.bMassAsDensity ? SharedParams.Mass * TotalVolume : SharedParams.Mass;
+	const float ClampedTotalMass = FMath::Clamp(DesiredTotalMass, SharedParams.MinimumMassClamp, SharedParams.MaximumMassClamp);
+	const float DesiredDensity = ClampedTotalMass / TotalVolume;
+	TVector<float, 3> MaxChildBounds(1);
+
+	for (int32 GeometryIndex = 0; GeometryIndex < NumGeometries; GeometryIndex++)
+	{
+		const int32 TransformGroupIndex = TransformIndex[GeometryIndex];
+		if (CollectionSimulatableParticles[TransformGroupIndex])
+		{
+			TUniquePtr<TTriangleMesh<float>>& TriMesh = TriangleMeshesArray[TransformGroupIndex];
+			TMassProperties<float, 3>& MassProperties = MassPropertiesArray[GeometryIndex];
+
+			//Must clamp each individual mass regardless of desired density
+			const float Volume_i = MassPropertiesArray[GeometryIndex].Volume;
+			if (DesiredDensity * Volume_i > SharedParams.MaximumMassClamp)
+			{
+				ensure(false);
+			}
+			const float Mass_i = FMath::Max(DesiredDensity * Volume_i, SharedParams.MinimumMassClamp);
+			const float Density_i = Mass_i / Volume_i;
+			CollectionMass[TransformGroupIndex] = Mass_i;
+
+			if (InertiaComputationNeeded[GeometryIndex])
+			{
+				check(false); // Ryan - Currently everything is hard wired to be a box in the loop above, so this should never run.
+				CalculateInertiaAndRotationOfMass(
+					MassSpaceParticles, 
+					TriMesh->GetSurfaceElements(), 
+					Density_i, 
+					MassProperties.CenterOfMass, 
+					MassProperties.InertiaTensor, 
+					MassProperties.RotationOfMass);
+
+				CollectionMassToLocal[TransformGroupIndex] = 
+					FTransform(MassProperties.RotationOfMass, MassProperties.CenterOfMass);
+
+				const TVector<float, 3> DiagonalInertia(
+					MassProperties.InertiaTensor.M[0][0], MassProperties.InertiaTensor.M[1][1], MassProperties.InertiaTensor.M[2][2]);
+				//computation includes mass already
+				CollectionInertiaTensor[TransformGroupIndex] = DiagonalInertia;
+			}
+			else
+			{
+				const TVector<float, 3> DiagonalInertia(
+					MassProperties.InertiaTensor.M[0][0], 
+					MassProperties.InertiaTensor.M[1][1], 
+					MassProperties.InertiaTensor.M[2][2]);
+				//using fallback computation, but adjust for mass
+				CollectionInertiaTensor[TransformGroupIndex] = DiagonalInertia * Mass_i;
+			}
+
+			// Update vertex buffer to be in mass space so that at runtime geometry aligns properly.
+			FBox InstanceBoundingBox(EForceInit::ForceInitToZero);
+			const TSet<int32>& VertsAdded = VertsAddedArray[GeometryIndex];
+			for (int32 VertIdx = VertexStart[GeometryIndex]; VertIdx < VertexStart[GeometryIndex] + VertexCount[GeometryIndex]; ++VertIdx)
+			{
+				if (VertsAdded.Contains(VertIdx))	//only consider verts from the trimesh
+				{
+					MassSpaceParticles.X(VertIdx) = 
+						CollectionMassToLocal[TransformGroupIndex].InverseTransformPosition(
+							MassSpaceParticles.X(VertIdx));
+					InstanceBoundingBox += MassSpaceParticles.X(VertIdx);	//build bounding box for visible verts in mass space
+				}
+			}
+
+			const int32 SizeSpecificIdx = FindSizeSpecificIdx(SharedParams.SizeSpecificData, InstanceBoundingBox);
+			const FSharedSimulationSizeSpecificData& SizeSpecificData = SharedParams.SizeSpecificData[SizeSpecificIdx];
+
+			//
+			//  Build the simplicial for the rest collection. This will be used later in the DynamicCollection to 
+			//  populate the collision structures of the simulation. 
+			//
+			check(TriMesh);
+			Chaos::TBVHParticles<float, 3> * Simplicial = 
+				FCollisionStructureManager::NewSimplicial(
+					MassSpaceParticles, 
+					BoneMap, 
+					SizeSpecificData.CollisionType, 
+					*TriMesh, 
+					SizeSpecificData.CollisionParticlesFraction);
+			CollectionSimplicials[TransformGroupIndex] = TUniquePtr<FSimplicial>(Simplicial);
+			ensureMsgf(CollectionSimplicials[TransformGroupIndex], TEXT("No simplicial representation."));
+
+			if (SizeSpecificData.ImplicitType == EImplicitTypeEnum::Chaos_Implicit_LevelSet)
+			{
+				ErrorReporter.SetPrefix(BaseErrorPrefix + " | Transform Index: " + FString::FromInt(TransformGroupIndex) + " of " + FString::FromInt(TransformIndex.Num()));
+				CollectionImplicits[TransformGroupIndex] = TSharedPtr<Chaos::FImplicitObject, ESPMode::ThreadSafe>(
+					FCollisionStructureManager::NewImplicitLevelset(
+						ErrorReporter,
+						MassSpaceParticles,
+						*TriMesh,
+						InstanceBoundingBox,
+						SizeSpecificData.MinLevelSetResolution,
+						SizeSpecificData.MaxLevelSetResolution,
+						SizeSpecificData.CollisionObjectReductionPercentage,
+						SizeSpecificData.CollisionType));
+			}
+			else if (SizeSpecificData.ImplicitType == EImplicitTypeEnum::Chaos_Implicit_Box)
+			{
+				CollectionImplicits[TransformGroupIndex] = TSharedPtr<Chaos::FImplicitObject, ESPMode::ThreadSafe>(
+					FCollisionStructureManager::NewImplicitBox(
+						InstanceBoundingBox,
+						SizeSpecificData.CollisionObjectReductionPercentage,
+						SizeSpecificData.CollisionType));
+			}
+			else if (SizeSpecificData.ImplicitType == EImplicitTypeEnum::Chaos_Implicit_Sphere)
+			{
+				CollectionImplicits[TransformGroupIndex] = TSharedPtr<Chaos::FImplicitObject, ESPMode::ThreadSafe>(
+					FCollisionStructureManager::NewImplicitSphere(
+						InnerRadius[GeometryIndex],
+						SizeSpecificData.CollisionObjectReductionPercentage,
+						SizeSpecificData.CollisionType));
+			}
+			else if (SizeSpecificData.ImplicitType == EImplicitTypeEnum::Chaos_Implicit_None)
+			{
+				CollectionImplicits[TransformGroupIndex] = nullptr;
+			}
+			else
+			{
+				ensure(false); // unsupported implicit type!
+			}
+
+			if (CollectionImplicits[TransformGroupIndex])
+			{
+				const TVector<float, 3> Extents = 
+					CollectionImplicits[TransformGroupIndex]->BoundingBox().Extents(); // Chaos::TAABB::Extents() is Max - Min
+				MaxChildBounds = MaxChildBounds.ComponentwiseMax(Extents);
+			}
+		}
+	}
+
+	// question: at the moment we always build cluster data in the asset. This 
+	// allows for per instance toggling. Is this needed? It increases memory 
+	// usage for all geometry collection assets.
+	const bool bEnableClustering = true;	
+	if (bEnableClustering)
+	{
+		//Put all children into collection space so we can compute mass properties.
+		TUniquePtr<TPBDRigidClusteredParticles<float, 3>> CollectionSpaceParticles(new TPBDRigidClusteredParticles<float, 3>());
+		CollectionSpaceParticles->AddParticles(NumTransforms);
+
+		// Init to -FLT_MAX for debugging purposes
+		for (int32 Idx = 0; Idx < NumTransforms; Idx++)
+		{
+			CollectionSpaceParticles->X(Idx) = Chaos::TVector<float, 3>(-TNumericLimits<float>::Max());
+		}
+
+		TArray<TUniquePtr<TPBDRigidClusteredParticleHandle<float, 3>>> Handles;
+		Handles.Reserve(NumTransforms);
+		for (int32 Idx = 0; Idx < NumTransforms; Idx++)
+		{
+			Handles.Add(TPBDRigidClusteredParticleHandle<float, 3>::CreateParticleHandle(
+				MakeSerializable(CollectionSpaceParticles), Idx, Idx));
+			Handles[Idx]->SetUniqueIdx(FUniqueIdx(Idx));
+		}
+
+		for (int32 GeometryIdx = 0; GeometryIdx < NumGeometries; ++GeometryIdx)
+		{
+			const int32 TransformGroupIndex = TransformIndex[GeometryIdx];
+			if (CollectionSimulatableParticles[TransformGroupIndex])
+			{
+				FTransform MassToComponent = CollectionMassToLocal[TransformGroupIndex] * CollectionSpaceTransforms[TransformGroupIndex];
+				PopulateSimulatedParticle(
+					Handles[TransformGroupIndex].Get(),
+					SharedParams, 
+					CollectionSimplicials[TransformGroupIndex].Get(),
+					CollectionImplicits[TransformGroupIndex], 
+					CollectionMass[TransformGroupIndex],
+					CollectionInertiaTensor[TransformGroupIndex], 
+					MassToComponent,
+					(uint8)EObjectStateTypeEnum::Chaos_Object_Dynamic, 
+					INDEX_NONE); // CollisionGroup
+			}
+		}
+
+		const TArray<int32> RecursiveOrder = ComputeRecursiveOrder(RestCollection);
+		const TArray<int32> TransformToGeometry = ComputeTransformToGeometryMap(RestCollection);
+
+		TArray<bool> IsClusterSimulated;
+		IsClusterSimulated.Init(false, CollectionSpaceParticles->Size());
+		//build collision structures depth first
+		for (const int32 TransformGroupIndex : RecursiveOrder)
+		{
+			if (Children[TransformGroupIndex].Num())	//only care about clusters at this point
+			{
+				const int32 ClusterTransformIdx = TransformGroupIndex;
+				//update mass 
+				TSet<TPBDRigidParticleHandle<float,3>*> ChildrenIndices;
+				{ // tmp scope
+					ChildrenIndices.Reserve(Children[ClusterTransformIdx].Num());
+					for (int32 ChildIdx : Children[ClusterTransformIdx])
+					{
+						if (CollectionSimulatableParticles[ChildIdx] || IsClusterSimulated[ChildIdx])
+						{
+							ChildrenIndices.Add(Handles[ChildIdx].Get());
+						}
+					}
+					if (!ChildrenIndices.Num())
+					{
+						continue;
+					}
+				} // tmp scope
+
+				//CollectionSimulatableParticles[TransformGroupIndex] = true;
+				IsClusterSimulated[TransformGroupIndex] = true;
+
+				UpdateClusterMassProperties(Handles[ClusterTransformIdx].Get(), ChildrenIndices);	//compute mass properties
+				const FTransform ClusterMassToCollection = 
+					FTransform(CollectionSpaceParticles->R(ClusterTransformIdx), 
+							   CollectionSpaceParticles->X(ClusterTransformIdx));
+
+				// Compute MassToLocal as if the transform hierarchy stays fixed. 
+				// In reality we modify the transform hierarchy so that MassToLocal 
+				// is identity at runtime.
+				CollectionMassToLocal[ClusterTransformIdx] = 
+					ClusterMassToCollection.GetRelativeTransform(
+						CollectionSpaceTransforms[ClusterTransformIdx]);
+
+				//update geometry
+				//merge children meshes and move them into cluster's mass space
+				TArray<TVector<int32, 3>> UnionMeshIndices;
+				int32 BiggestNumElements = 0;
+				{ // tmp scope
+					int32 NumChildIndices = 0;
+					for (TPBDRigidParticleHandle<float, 3>* Child : ChildrenIndices)
+					{
+						const int32 ChildTransformIdx = Child->UniqueIdx().Idx;
+						if (Chaos::TTriangleMesh<float>* ChildMesh = TriangleMeshesArray[ChildTransformIdx].Get())
+						{
+							BiggestNumElements = FMath::Max(BiggestNumElements, ChildMesh->GetNumElements());
+							NumChildIndices += ChildMesh->GetNumElements();
+						}
+					}
+					UnionMeshIndices.Reserve(NumChildIndices);
+				} // tmp scope
+
+				FBox InstanceBoundingBox(EForceInit::ForceInitToZero);
+				{ // tmp scope
+					TSet<int32> VertsAdded;
+					VertsAdded.Reserve(BiggestNumElements);
+					for (TPBDRigidParticleHandle<float, 3>* Child : ChildrenIndices)
+					{
+						const int32 ChildTransformIdx = Child->UniqueIdx().Idx;
+						if (Chaos::TTriangleMesh<float>* ChildMesh = TriangleMeshesArray[ChildTransformIdx].Get())
+						{
+							const TArray<TVector<int32, 3>>& ChildIndices = ChildMesh->GetSurfaceElements();
+							UnionMeshIndices.Append(ChildIndices);
+
+							const FTransform ChildMassToClusterMass =
+								(CollectionMassToLocal[ChildTransformIdx] * CollectionSpaceTransforms[ChildTransformIdx]).GetRelativeTransform(
+									ClusterMassToCollection);
+
+							ChildMesh->GetVertexSet(VertsAdded);
+							for (const int32 VertIdx : VertsAdded)
+							{
+								//Update particles so they are in the cluster's mass space
+								MassSpaceParticles.X(VertIdx) =
+									ChildMassToClusterMass.TransformPosition(MassSpaceParticles.X(VertIdx));
+								InstanceBoundingBox += MassSpaceParticles.X(VertIdx);
+							}
+						}
+					}
+				} // tmp scope
+
+				TUniquePtr<TTriangleMesh<float>> UnionMesh(new TTriangleMesh<float>(MoveTemp(UnionMeshIndices)));
+				const FMatrix& InertiaMatrix = CollectionSpaceParticles->I(ClusterTransformIdx);
+				const FVector InertiaDiagonal(InertiaMatrix.M[0][0], InertiaMatrix.M[1][1], InertiaMatrix.M[2][2]);
+				CollectionInertiaTensor[ClusterTransformIdx] = InertiaDiagonal;
+				CollectionMass[ClusterTransformIdx] = CollectionSpaceParticles->M(ClusterTransformIdx);
+
+				const int32 SizeSpecificIdx = FindSizeSpecificIdx(SharedParams.SizeSpecificData, InstanceBoundingBox);
+				const FSharedSimulationSizeSpecificData& SizeSpecificData = SharedParams.SizeSpecificData[SizeSpecificIdx];
+
+				if (SizeSpecificData.ImplicitType == EImplicitTypeEnum::Chaos_Implicit_LevelSet)
+				{
+					const TVector<float, 3> Scale = 2 * InstanceBoundingBox.GetExtent() / MaxChildBounds;
+					const float ScaleMax = Scale.GetAbsMax();
+					const float ScaleMin = Scale.GetAbsMin();
+
+					float MinResolution = ScaleMin * SizeSpecificData.MinLevelSetResolution;
+					MinResolution = FMath::Clamp<float>(MinResolution, SizeSpecificData.MinLevelSetResolution, SizeSpecificData.MinClusterLevelSetResolution);
+
+					float MaxResolution = ScaleMax * SizeSpecificData.MaxLevelSetResolution;
+					MaxResolution = FMath::Clamp<float>(MaxResolution, SizeSpecificData.MaxLevelSetResolution, SizeSpecificData.MaxClusterLevelSetResolution);
+
+					//don't support non level-set serialization
+					ErrorReporter.SetPrefix(BaseErrorPrefix + " | Cluster Transform Index: " + FString::FromInt(ClusterTransformIdx));
+					CollectionImplicits[ClusterTransformIdx] = TSharedPtr<Chaos::FImplicitObject, ESPMode::ThreadSafe>(
+						FCollisionStructureManager::NewImplicitLevelset(
+							ErrorReporter,
+							MassSpaceParticles,
+							*UnionMesh,
+							InstanceBoundingBox,
+							MinResolution,
+							MaxResolution,
+							SizeSpecificData.CollisionObjectReductionPercentage,
+							SizeSpecificData.CollisionType));
+
+					CollectionSimplicials[ClusterTransformIdx] = TUniquePtr< FSimplicial >(
+						FCollisionStructureManager::NewSimplicial(MassSpaceParticles, *UnionMesh, CollectionImplicits[ClusterTransformIdx].Get(),
+						SharedParams.MaximumCollisionParticleCount));
+				}
+				else if (SizeSpecificData.ImplicitType == EImplicitTypeEnum::Chaos_Implicit_Box)
+				{
+					CollectionImplicits[ClusterTransformIdx] = TSharedPtr<Chaos::FImplicitObject, ESPMode::ThreadSafe>(
+						FCollisionStructureManager::NewImplicitBox(
+							InstanceBoundingBox,
+							SizeSpecificData.CollisionObjectReductionPercentage,
+							SizeSpecificData.CollisionType));
+
+					CollectionSimplicials[ClusterTransformIdx] = TUniquePtr< FSimplicial >(
+						FCollisionStructureManager::NewSimplicial(MassSpaceParticles, *UnionMesh, CollectionImplicits[ClusterTransformIdx].Get(),
+						SharedParams.MaximumCollisionParticleCount));
+				}
+				else if (SizeSpecificData.ImplicitType == EImplicitTypeEnum::Chaos_Implicit_Sphere)
+				{
+ 					CollectionImplicits[ClusterTransformIdx] = TSharedPtr<Chaos::FImplicitObject, ESPMode::ThreadSafe>(
+						FCollisionStructureManager::NewImplicitSphere(
+							InstanceBoundingBox.GetExtent().GetAbsMin() / 2,
+							SizeSpecificData.CollisionObjectReductionPercentage,
+							SizeSpecificData.CollisionType));
+
+					CollectionSimplicials[ClusterTransformIdx] = TUniquePtr< FSimplicial >(
+						FCollisionStructureManager::NewSimplicial(MassSpaceParticles, *UnionMesh, CollectionImplicits[ClusterTransformIdx].Get(),
+						SharedParams.MaximumCollisionParticleCount));
+				}
+				else if(SizeSpecificData.ImplicitType == EImplicitTypeEnum::Chaos_Implicit_Capsule)
+				{
+					ensure(false); // unsupported implicit type
+					CollectionImplicits[ClusterTransformIdx].Reset();
+					CollectionSimplicials[ClusterTransformIdx].Reset();
+				}
+				else // Assume it's a union???
+				{
+					CollectionImplicits[ClusterTransformIdx].Reset();	//union so just set as null
+					CollectionSimplicials[ClusterTransformIdx].Reset();
+				}
+
+				TriangleMeshesArray[ClusterTransformIdx] = MoveTemp(UnionMesh);
+			}
+		}
+
+		InitRemoveOnFracture(RestCollection, SharedParams);
+	}
+}
+
+void FGeometryCollectionPhysicsProxy::InitRemoveOnFracture(FGeometryCollection& RestCollection, const FSharedSimulationParameters& SharedParams)
+{
+	if (SharedParams.RemoveOnFractureIndices.Num() == 0)
+	{
+		return;
+	}
+
+	// Markup Node Hierarchy Status with FS_RemoveOnFracture flags where geometry is ALL glass
+	const int32 NumGeometries = RestCollection.NumElements(FGeometryCollection::GeometryGroup);
+	for (int32 Idx = 0; Idx < NumGeometries; Idx++)
+	{
+		const int32 TransformIndex = RestCollection.TransformIndex[Idx];
+		const int32 Start = RestCollection.FaceStart[Idx];
+		const int32 End = RestCollection.FaceCount[Idx];
+		bool IsToBeRemoved = true;
+		for (int32 Face = Start; Face < Start + End; Face++)
+		{
+			bool FoundMatch = false;
+			for (int32 MaterialIndex : SharedParams.RemoveOnFractureIndices)
+			{
+				if (RestCollection.MaterialID[Face] == MaterialIndex)
+				{
+					FoundMatch = true;
+					break;
+				}
+			}
+			if (!FoundMatch)
+			{
+				IsToBeRemoved = false;
+				break;
+			}
+		}
+		if (IsToBeRemoved)
+		{
+			RestCollection.SetFlags(TransformIndex, FGeometryCollection::FS_RemoveOnFracture);
+		}
+		else
+		{
+			RestCollection.ClearFlags(TransformIndex, FGeometryCollection::FS_RemoveOnFracture);
+		}
+	}
+}
+
+void IdentifySimulatableElements(Chaos::FErrorReporter& ErrorReporter, FGeometryCollection& GeometryCollection)
+{
+	// Determine which collection particles to simulate
+	const TManagedArray<FBox>& BoundingBox = GeometryCollection.BoundingBox;
+	const TManagedArray<int32>& VertexCount = GeometryCollection.VertexCount;
+	const TManagedArray<int32>& TransformIndex = GeometryCollection.TransformIndex;
+	const int32 NumTransforms = GeometryCollection.NumElements(FGeometryCollection::TransformGroup);
+
+	const int32 NumTransformMappings = TransformIndex.Num();
+
+	// Do not simulate hidden geometry
+	TArray<bool> HiddenObject;
+	HiddenObject.Init(true, NumTransforms);
+	const TManagedArray<bool>& Visible = GeometryCollection.Visible;
+	const TManagedArray<int32>& BoneMap = GeometryCollection.BoneMap;
+
+	const TManagedArray<FIntVector>& Indices = GeometryCollection.Indices;
+	int32 PrevObject = -1;
+	bool bContiguous = true;
+	for(int32 i = 0; i < Indices.Num(); i++)
+	{
+		if(Visible[i])
+		{
+			int32 ObjIdx = BoneMap[Indices[i][0]];
+			HiddenObject[ObjIdx] = false;
+			if (!ensureMsgf(ObjIdx >= PrevObject, TEXT("Objects are not contiguous. This breaks assumptions later in the pipeline")))
+			{
+				bContiguous = false;
+			}
+			
+			PrevObject = ObjIdx;
+		}
+	}
+
+	if (!bContiguous)
+	{
+		ErrorReporter.ReportError(TEXT("Objects are not contiguous. This breaks assumptions later in the pipeline"));
+		ErrorReporter.HandleLatestError();
+	}
+
+
+	//For now all simulation data is a non compiled attribute. Not clear what we want for simulated vs kinematic collections
+	TManagedArray<bool>& SimulatableParticles = 
+		GeometryCollection.AddAttribute<bool>(
+			FGeometryCollection::SimulatableParticlesAttribute, FTransformCollection::TransformGroup);
+
+	for(int32 TransformIdx = 0; TransformIdx < NumTransforms; TransformIdx++)
+	{
+		SimulatableParticles[TransformIdx] = false;
+	}
+
+	for(int i = 0; i < NumTransformMappings; i++)
+	{
+		int32 Tdx = TransformIndex[i];
+		checkSlow(0 <= Tdx && Tdx < NumTransforms);
+		if(GeometryCollection.IsGeometry(Tdx) && VertexCount[i] && 0.f < BoundingBox[i].GetSize().SizeSquared() && !HiddenObject[Tdx])
+		{
+			SimulatableParticles[Tdx] = true;
+		}
+	}
+}
+
+void BuildSimulationData(Chaos::FErrorReporter& ErrorReporter, FGeometryCollection& GeometryCollection, const FSharedSimulationParameters& SharedParams)
+{
+	IdentifySimulatableElements(ErrorReporter, GeometryCollection);
+	FGeometryCollectionPhysicsProxy::InitializeSharedCollisionStructures(ErrorReporter, GeometryCollection, SharedParams);
+}
+
+void FGeometryCollectionPhysicsProxy::MergeRecordedTracks(const FRecordedTransformTrack& A, const FRecordedTransformTrack& B, FRecordedTransformTrack& Target)
+{
+	const int32 NumAKeys = A.Records.Num();
+	const int32 NumBKeys = B.Records.Num();
+
+	if(NumAKeys == 0)
+	{
+		Target = B;
+		return;
+	}
+
+	if(NumBKeys == 0)
+	{
+		Target = A;
+		return;
+	}
+
+	// We have to copy the tracks to a local cache here because Target could point at either A or B.
+	FRecordedTransformTrack TempMergedTrack = A;
+
+	// Expand to hold all the keys
+	TempMergedTrack.Records.Reserve(NumAKeys + NumBKeys);
+
+	// Insert B frames into the merged set
+	for(int32 BKeyIndex = 0; BKeyIndex < NumBKeys; ++BKeyIndex)
+	{
+		const FRecordedFrame& BFrame = B.Records[BKeyIndex];
+		int32 KeyBefore = TempMergedTrack.FindLastKeyBefore(BFrame.Timestamp);
+
+		TempMergedTrack.Records.Insert(BFrame, KeyBefore + 1);
+	}
+
+	// Copy to target
+	Target = TempMergedTrack;
+}
+
+FRecordedFrame& FGeometryCollectionPhysicsProxy::InsertRecordedFrame(FRecordedTransformTrack& InTrack, float InTime)
+{
+	// Can't just throw on the end, might need to insert
+	const int32 BeforeNewIndex = InTrack.FindLastKeyBefore(InTime);
+
+	if(BeforeNewIndex == InTrack.Records.Num() - 1)
+	{
+		InTrack.Records.AddDefaulted();
+		return InTrack.Records.Last();
+	}
+
+	const int32 NewRecordIndex = InTrack.Records.Insert(FRecordedFrame(), BeforeNewIndex + 1);
+	return InTrack.Records[NewRecordIndex];
+}
+
+
+//==============================================================================
+// FIELDS
+//==============================================================================
+
+
+void FGeometryCollectionPhysicsProxy::ProcessCommands(FParticlesType& Particles, const float Time)
+{
+	FGeometryDynamicCollection* Collection = Parameters.DynamicCollection;
+	check(Collection);
+
+	// Process Particle-Collection commands
+	if (Commands.Num())
+	{
+		TArray<ContextIndex> IndicesArray;
+		Chaos::FPhysicsSolver* CurrentSolver = GetSolver();
+
+		for (int32 CommandIndex = Commands.Num()-1; 0<=CommandIndex; CommandIndex--)
+		{
+
+			//
+			// Extract command and set metadata
+			//
+			FFieldSystemCommand & Command = Commands[CommandIndex];
+			EFieldResolutionType ResolutionType = EFieldResolutionType::Field_Resolution_Minimal;
+			if (Command.MetaData.Contains(FFieldSystemMetaData::EMetaType::ECommandData_ProcessingResolution))
+			{
+				check(Command.MetaData[FFieldSystemMetaData::EMetaType::ECommandData_ProcessingResolution] != nullptr);
+				ResolutionType = static_cast<FFieldSystemMetaDataProcessingResolution*>(Command.MetaData[FFieldSystemMetaData::EMetaType::ECommandData_ProcessingResolution].Get())->ProcessingResolution;
+			}
+
+			if (Command.TargetAttribute == GetGeometryCollectionPhysicsTypeName(EGeometryCollectionPhysicsTypeEnum::Chaos_DynamicState))
+			{
+				if (ensureMsgf(Command.RootNode->Type() == FFieldNode<int32>::StaticType(),
+					TEXT("Field based evaluation of the simulations 'DynamicState' parameter expects int32 field inputs.")))
+				{
+					FGeometryCollectionPhysicsProxy::ContiguousIndices(IndicesArray, CurrentSolver, ResolutionType, IndicesArray.Num() != Particles.Size());
+					if (IndicesArray.Num())
+					{
+						TArrayView<ContextIndex> IndexView(&(IndicesArray[0]), IndicesArray.Num());
+
+						FVector * xptr = &(Particles.X(0));
+						TArrayView<FVector> SamplesView(xptr, Particles.Size());
+
+						FFieldContext Context{
+							IndexView,
+							SamplesView,
+							Command.MetaData
+						};
+
+						TArrayView<int32> DynamicStateView(&(Collection->DynamicState[0]), Collection->DynamicState.Num());
+						static_cast<const FFieldNode<int32> *>(Command.RootNode.Get())->Evaluate(Context, DynamicStateView);
+
+						PushKinematicStateToSolver(Particles);
+					}
+				}
+				Commands.RemoveAt(CommandIndex);
+			}
+			else if (Command.TargetAttribute == GetGeometryCollectionPhysicsTypeName(EGeometryCollectionPhysicsTypeEnum::Chaos_InitialLinearVelocity))
+			{
+				if (ensureMsgf(Parameters.InitialVelocityType == EInitialVelocityTypeEnum::Chaos_Initial_Velocity_User_Defined,
+					TEXT("Field based evaluation of the simulations 'InitialLinearVelocity' requires the geometry collection be set to User Defined Initial Velocity")))
+				{
+					if (ensureMsgf(Command.RootNode->Type() == FFieldNode<FVector>::StaticType(),
+						TEXT("Field based evaluation of the simulations 'InitialLinearVelocity' parameter expects FVector field inputs.")))
+					{
+						FGeometryCollectionPhysicsProxy::ContiguousIndices(IndicesArray, CurrentSolver, ResolutionType, IndicesArray.Num() != Particles.Size());
+						if (IndicesArray.Num())
+						{
+							TArrayView<ContextIndex> IndexView(&(IndicesArray[0]), IndicesArray.Num());
+
+							FVector * xptr = &(Particles.X(0));
+							TArrayView<FVector> SamplesView(xptr, Particles.Size());
+
+							FFieldContext Context{
+								IndexView,
+								SamplesView,
+								Command.MetaData
+							};
+
+							TArrayView<FVector> ResultsView(&(Collection->InitialLinearVelocity[0]), Collection->InitialLinearVelocity.Num());
+							static_cast<const FFieldNode<FVector> *>(Command.RootNode.Get())->Evaluate(Context, ResultsView);
+						}
+					}
+				}
+				Commands.RemoveAt(CommandIndex);
+			}
+			else if (Command.TargetAttribute == GetGeometryCollectionPhysicsTypeName(EGeometryCollectionPhysicsTypeEnum::Chaos_InitialAngularVelocity))
+			{
+				if (ensureMsgf(Parameters.InitialVelocityType == EInitialVelocityTypeEnum::Chaos_Initial_Velocity_User_Defined,
+					TEXT("Field based evaluation of the simulations 'InitialAngularVelocity' requires the geometry collection be set to User Defined Initial Velocity")))
+				{
+					if (ensureMsgf(Command.RootNode->Type() == FFieldNode<FVector>::StaticType(),
+						TEXT("Field based evaluation of the simulations 'InitialAngularVelocity' parameter expects FVector field inputs.")))
+					{
+						FGeometryCollectionPhysicsProxy::ContiguousIndices(IndicesArray, CurrentSolver, ResolutionType, IndicesArray.Num() != Particles.Size());
+						if (IndicesArray.Num())
+						{
+							TArrayView<ContextIndex> IndexView(&(IndicesArray[0]), IndicesArray.Num());
+
+							FVector * xptr = &(Particles.X(0));
+							TArrayView<FVector> SamplesView(xptr, Particles.Size());
+
+							FFieldContext Context{
+								IndexView,
+								SamplesView,
+								Command.MetaData
+							};
+
+							TArrayView<FVector> ResultsView(&(Collection->InitialAngularVelocity[0]), Collection->InitialAngularVelocity.Num());
+							static_cast<const FFieldNode<FVector> *>(Command.RootNode.Get())->Evaluate(Context, ResultsView);
+						}
+					}
+				}
+				Commands.RemoveAt(CommandIndex);
+			}
+		}
+	}
+
+	// Process Particle-Particle commands
+	if (Commands.Num())
+	{
+		Chaos::FPhysicsSolver* CurrentSolver = GetSolver();
+
+		/*
+		@todo break everything
+
+		//  Generate a Index mapping between the rigid body indices and 
+		//  the particle indices. This allows the geometry collection to
+		//  evaluate only its own particles. 
+		TArray<ContextIndex> IndicesArray;
+		int32 NumIndices = 0;
+		IndicesArray.SetNumUninitialized(RigidBodyID.Num());
+		for (int32 i = 0; i < RigidBodyID.Num(); i++)
+		{
+			if (RigidBodyID[i] != INDEX_NONE)
+			{
+				IndicesArray[NumIndices] = { RigidBodyID[i], RigidBodyID[i] };
+				NumIndices++;
+			}
+		}
+		IndicesArray.SetNum(NumIndices);
+
+		for (int32 CommandIndex = Commands.Num() - 1; 0 <= CommandIndex; CommandIndex--)
+		{
+			FFieldSystemCommand & Command = Commands[CommandIndex];
+			if (Command.TargetAttribute == GetGeometryCollectionPhysicsTypeName(EGeometryCollectionPhysicsTypeEnum::Chaos_LinearVelocity))
+			{
+				if (ensureMsgf(Command.RootNode->Type() == FFieldNode<FVector>::StaticType(),
+					TEXT("Field based evaluation of the simulations 'LinearVelocity' parameter expects FVector field inputs.")))
+				{
+					FVector * xptr = &(Particles.X(0));
+					TArrayView<FVector> SamplesView(xptr, Particles.Size());
+					TArrayView<ContextIndex> IndexView(&(IndicesArray[0]), IndicesArray.Num());
+
+					FFieldContext Context{
+						IndexView,
+						SamplesView,
+						Command.MetaData
+					};
+
+					FVector * vptr = &(Particles.V(0));
+					TArrayView<FVector> ResultsView(vptr, Particles.Size());
+					static_cast<const FFieldNode<FVector> *>(Command.RootNode.Get())->Evaluate(Context, ResultsView);
+				}
+				Commands.RemoveAt(CommandIndex);
+			}
+			else if (Command.TargetAttribute == GetGeometryCollectionPhysicsTypeName(EGeometryCollectionPhysicsTypeEnum::Chaos_AngularVelocity))
+			{
+				if (ensureMsgf(Command.RootNode->Type() == FFieldNode<FVector>::StaticType(),
+					TEXT("Field based evaluation of the simulations 'AngularVelocity' parameter expects FVector field inputs.")))
+				{
+
+					FVector * xptr = &(Particles.X(0));
+					TArrayView<FVector> SamplesView(xptr, Particles.Size());
+					TArrayView<ContextIndex> IndexView(&(IndicesArray[0]), IndicesArray.Num());
+
+					FFieldContext Context{
+						IndexView,
+						SamplesView,
+						Command.MetaData
+					};
+
+					FVector * vptr = &(Particles.W(0));
+					TArrayView<FVector> ResultsView(vptr, Particles.Size());
+					static_cast<const FFieldNode<FVector> *>(Command.RootNode.Get())->Evaluate(Context, ResultsView);
+				}
+				Commands.RemoveAt(CommandIndex);
+			}
+			else if (Command.TargetAttribute == GetGeometryCollectionPhysicsTypeName(EGeometryCollectionPhysicsTypeEnum::Chaos_CollisionGroup))
+			{
+				if (ensureMsgf(Command.RootNode->Type() == FFieldNode<int32>::StaticType(),
+					TEXT("Field based evaluation of the simulations 'CollisionGroup' parameter expects int32 field inputs.")))
+				{
+					FVector * xptr = &(Particles.X(0));
+					TArrayView<FVector> SamplesView(xptr, Particles.Size());
+					TArrayView<ContextIndex> IndexView(&(IndicesArray[0]), IndicesArray.Num());
+
+					FFieldContext Context{
+						IndexView,
+						SamplesView,
+						Command.MetaData
+					};
+
+					int32 * cptr = &(Particles.CollisionGroup(0));
+					TArrayView<int32> ResultsView(cptr, Particles.Size());
+					static_cast<const FFieldNode<int32> *>(Command.RootNode.Get())->Evaluate(Context, ResultsView);
+				}
+				Commands.RemoveAt(CommandIndex);
+			}
+		}
+		*/
+	}
+}
+
+void FGeometryCollectionPhysicsProxy::FieldForcesUpdateCallback(Chaos::FPhysicsSolver* InSolver, FParticlesType& Particles, Chaos::TArrayCollectionArray<FVector> & Force, Chaos::TArrayCollectionArray<FVector> & Torque, const float Time)
+{
+	/*
+
+	@todo break everything : re-enable fields without rigid body ids. 
+
+	if (Commands.Num())
+		{
+			// @todo: This seems like a waste if we just want to get everything
+			int32 Counter = 0;
+			TArray<ContextIndex> IndicesArray;
+			IndicesArray.AddUninitialized(RigidBodyID.Num());
+			for (int32 i = 0; i < RigidBodyID.Num(); i++)
+			{
+				if (RigidBodyID[i] != INDEX_NONE)
+				{
+					IndicesArray[i] = { RigidBodyID[i],RigidBodyID[i] };
+					Counter++;
+				}
+			}
+			IndicesArray.SetNum(Counter, false);
+			TArrayView<ContextIndex> IndexView(&(IndicesArray[0]), IndicesArray.Num());
+
+
+			for (int32 CommandIndex = 0; CommandIndex < Commands.Num(); CommandIndex++)
+			{
+				FFieldSystemCommand & Command = Commands[CommandIndex];
+
+				if (Command.TargetAttribute == GetGeometryCollectionPhysicsTypeName(EGeometryCollectionPhysicsTypeEnum::Chaos_LinearForce))
+				{
+					if (ensureMsgf(Command.RootNode->Type() == FFieldNode<FVector>::StaticType(),
+						TEXT("Field based evaluation of the simulations 'LinearForce' parameter expects FVector field inputs.")))
+					{
+						FVector * tptr = &(Particles.X(0));
+						TArrayView<FVector> SamplesView(tptr, int32(Particles.Size()));
+
+						FFieldContext Context{
+							IndexView,
+							SamplesView,
+							Command.MetaData
+						};
+						TArrayView<FVector> ForceView(&(Force[0]), Force.Num());
+						static_cast<const FFieldNode<FVector> *>(Command.RootNode.Get())->Evaluate(Context, ForceView);
+					}
+					Commands.RemoveAt(CommandIndex);
+				}
+				else if (Command.TargetAttribute == GetGeometryCollectionPhysicsTypeName(EGeometryCollectionPhysicsTypeEnum::Chaos_AngularTorque))
+				{
+					if (ensureMsgf(Command.RootNode->Type() == FFieldNode<FVector>::StaticType(),
+						TEXT("Field based evaluation of the simulations 'AngularTorque' parameter expects FVector field inputs.")))
+					{
+						FVector * tptr = &(Particles.X(0));
+						TArrayView<FVector> SamplesView(tptr, int32(Particles.Size()));
+
+						FFieldContext Context{
+							IndexView,
+							SamplesView,
+							Command.MetaData
+						};
+						TArrayView<FVector> TorqueView(&(Torque[0]), Torque.Num());
+						static_cast<const FFieldNode<FVector> *>(Command.RootNode.Get())->Evaluate(Context, TorqueView);
+					}
+					Commands.RemoveAt(CommandIndex);
+				}
+
+			}
+		}
+		*/
+}
+
+
+//==============================================================================
+// DEPRECATED
+//==============================================================================
+
+
+void FGeometryCollectionPhysicsProxy::UpdateKinematicBodiesCallback(const FParticlesType& Particles, const float Dt, const float Time, FKinematicProxy& Proxy)
+{
+	check(false); // not currently called
+#if TODO_REIMPLEMENT_KINEMATIC_PROXY
+
+	SCOPE_CYCLE_COUNTER(STAT_KinematicUpdate);
+		FGeometryDynamicCollection* Collection = Parameters.DynamicCollection;
+		check(Collection);
+
+		bool bIsCachePlaying = Parameters.IsCachePlaying();
+		bool bIsReverseCachePlaying = Parameters.IsCacheRecording() && Parameters.ReverseCacheBeginTime != 0 && Parameters.ReverseCacheBeginTime < Time;
+		if ((!bIsCachePlaying && !bIsReverseCachePlaying) || !Parameters.RecordedTrack)
+		{
+			return;
+		}
+
+		bool bFirst = !Proxy.Ids.Num();
+		if (bFirst)
+		{
+			Proxy.Position.Reset(RigidBodyID.Num());
+			Proxy.Rotation.Reset(RigidBodyID.Num());
+			Proxy.NextPosition.Reset(RigidBodyID.Num());
+			Proxy.NextRotation.Reset(RigidBodyID.Num());
+
+			Proxy.Position.AddUninitialized(RigidBodyID.Num());
+			Proxy.Rotation.AddUninitialized(RigidBodyID.Num());
+			Proxy.NextPosition.AddUninitialized(RigidBodyID.Num());
+			Proxy.NextRotation.AddUninitialized(RigidBodyID.Num());
+
+			for (int32 i = 0; i < RigidBodyID.Num(); ++i)
+			{
+				Proxy.Ids.Add(RigidBodyID[i]);
+
+				// Initialise to rest state
+				const int32 RbId = Proxy.Ids.Last();
+				Proxy.Position[i] = RbId != INDEX_NONE ? Particles.X(RbId) : FVector::ZeroVector;
+				Proxy.Rotation[i] = RbId != INDEX_NONE ? Particles.R(RbId) : FQuat::Identity;
+				Proxy.NextPosition[i] = Proxy.Position[i];
+				Proxy.NextRotation[i] = Proxy.Rotation[i];
+			}
+		}
+
+		if (bIsCachePlaying && !bIsReverseCachePlaying && (Time < Parameters.CacheBeginTime || !Parameters.RecordedTrack->IsTimeValid(Time)))
+		{
+			return;
+		}
+
+		float ReverseTime = Parameters.RecordedTrack->GetLastTime() - Time + Parameters.ReverseCacheBeginTime;
+		if (bIsReverseCachePlaying && !Parameters.RecordedTrack->IsTimeValid(ReverseTime))
+		{
+			return;
+		}
+
+		const FRecordedFrame* FirstFrame = nullptr;
+		const FRecordedFrame* SecondFrame = nullptr;
+		float PlaybackTime = bIsReverseCachePlaying ? ReverseTime : Time;
+		Parameters.RecordedTrack->GetFramesForTime(PlaybackTime, FirstFrame, SecondFrame);
+
+		if (FirstFrame && !SecondFrame)
+		{
+			// Only one frame to take information from (simpler case)
+			const int32 NumActives = FirstFrame->TransformIndices.Num();
+
+			// Actives
+			Chaos::PhysicsParallelFor(NumActives, [&](int32 Index)
+			{
+				const int32 InternalIndex = FirstFrame->TransformIndices[Index];
+				if (InternalIndex >= RigidBodyID.Num() || InternalIndex < 0)
+				{
+					UE_LOG(UGCC_LOG, Error, 
+						TEXT("%s: Cache index %d out of range: [%d, %d).  Regenerate the cache."),
+						*Parameters.Name, InternalIndex, 0, RigidBodyID.Num());
+					return;
+				}
+				const int32 ExternalIndex = RigidBodyID[InternalIndex];
+
+				if (ExternalIndex != INDEX_NONE && Particles.InvM(ExternalIndex) == 0.0f && !Particles.Disabled(ExternalIndex))
+				{
+					const FTransform& ParticleTransform = FirstFrame->Transforms[Index];
+					Proxy.Position[InternalIndex] = Particles.X(ExternalIndex);
+					Proxy.Rotation[InternalIndex] = Particles.R(ExternalIndex);
+					Proxy.NextPosition[InternalIndex] = ParticleTransform.GetTranslation();
+					Proxy.NextRotation[InternalIndex] = ParticleTransform.GetRotation();
+				}
+			});
+		}
+		else if (FirstFrame && SecondFrame)
+		{
+			// Both frames valid, second frame has all the indices we need
+			const int32 NumActives = SecondFrame->TransformIndices.Num();
+
+			const float Alpha = (PlaybackTime - FirstFrame->Timestamp) / (SecondFrame->Timestamp - FirstFrame->Timestamp);
+			check(0 <= Alpha && Alpha <= 1.0f);
+
+			Chaos::PhysicsParallelFor(NumActives, [&](int32 Index)
+			{
+				const int32 InternalIndex = SecondFrame->TransformIndices[Index];
+				if (InternalIndex >= RigidBodyID.Num() || InternalIndex < 0)
+				{
+					UE_LOG(UGCC_LOG, Error, 
+						TEXT("%s: Cache index %d out of range: [%d, %d).  Regenerate the cache."),
+						*Parameters.Name, InternalIndex, 0, RigidBodyID.Num());
+					return;
+				}
+				const int32 ExternalIndex = RigidBodyID[InternalIndex];
+				const int32 PreviousIndexSlot = Index < SecondFrame->PreviousTransformIndices.Num() ? SecondFrame->PreviousTransformIndices[Index] : INDEX_NONE;
+
+				if (ExternalIndex != INDEX_NONE && Particles.InvM(ExternalIndex) == 0.0f && !Particles.Disabled(ExternalIndex))
+				{
+					if (PreviousIndexSlot != INDEX_NONE)
+					{
+						Proxy.Position[InternalIndex] = Proxy.NextPosition[InternalIndex];
+						Proxy.Rotation[InternalIndex] = Proxy.NextRotation[InternalIndex];
+
+						FTransform BlendedTM;
+						BlendedTM.Blend(FirstFrame->Transforms[PreviousIndexSlot], SecondFrame->Transforms[Index], Alpha);
+
+						Proxy.NextPosition[InternalIndex] = BlendedTM.GetTranslation();
+						Proxy.NextRotation[InternalIndex] = BlendedTM.GetRotation();
+					}
+					else
+					{
+						// NewActive case
+						Proxy.Position[InternalIndex] = Proxy.NextPosition[InternalIndex];
+						Proxy.Rotation[InternalIndex] = Proxy.NextRotation[InternalIndex];
+
+						FTransform BlendedTM;
+						BlendedTM.Blend(FTransform(Particles.R(ExternalIndex), Particles.X(ExternalIndex), FVector::OneVector), SecondFrame->Transforms[Index], Alpha);
+
+						Proxy.NextPosition[InternalIndex] = BlendedTM.GetTranslation();
+						Proxy.NextRotation[InternalIndex] = BlendedTM.GetRotation();
+					}
+				}
+			});
+			// #BGallagher Handle new inactives. If it's a cluster parent and it's fully disabled we'll need to decluster it here.
+		}
+#endif
+}
+
+void FGeometryCollectionPhysicsProxy::StartFrameCallback(const float Dt, const float Time)
+{
+	SCOPE_CYCLE_COUNTER(STAT_GeomBeginFrame);
+#if TODO_REIMPLEMENT_GET_RIGID_PARTICLES
+	if (InitializedState == ESimulationInitializationState::Initialized)
+	{
+		// Reverse playback only plays back what we just recorded.  So, the condition 
+		// is, "are we in record mode, but ready to play back what we've recorded?"
+		const bool bIsReverseCachePlaying = Parameters.IsCacheRecording() && Parameters.ReverseCacheBeginTime != 0 && Time > Parameters.ReverseCacheBeginTime;
+		if (Parameters.IsCachePlaying() || bIsReverseCachePlaying)
+		{
+			// Update the enabled/disabled state for kinematic particles for the upcoming frame
+			Chaos::FPhysicsSolver* ThisSolver = GetSolver();
+			Chaos::FPhysicsSolver::FParticlesType& Particles = ThisSolver->GetRigidParticles();
+
+			if (!Parameters.RecordedTrack)
+			{
+				if (!ensure(Parameters.CacheType == EGeometryCollectionCacheType::Record))
+				{
+					return;
+				}
+				Parameters.RecordedTrack = new FRecordedTransformTrack();
+				Parameters.bOwnsTrack = true;
+			}
+			if (Parameters.bClearCache && bIsReverseCachePlaying)
+			{
+				check(!CommitRecordedStateCallback);
+				*const_cast<FRecordedTransformTrack*>(Parameters.RecordedTrack) = FRecordedTransformTrack::ProcessRawRecordedData(RecordedTracks);
+				Parameters.bClearCache = false;
+			}
+
+			bool bParticlesUpdated = false;
+
+			const float ThisFrameTime = bIsReverseCachePlaying ? (Parameters.RecordedTrack->GetLastTime() - Time + Parameters.ReverseCacheBeginTime) : Time;
+			if (!Parameters.RecordedTrack->IsTimeValid(ThisFrameTime))
+			{
+				// Invalid cache time, nothing to update
+				return;
+			}
+
+			FCriticalSection ParticleUpdateLock;
+			const int32 NumMappings = RigidBodyID.Num();
+			Chaos::PhysicsParallelFor(NumMappings, [&](int32 InternalParticleIndex)
+				//for(int32 Index = 0; Index < NumMappings; ++Index)
+			{
+				const int32 ExternalParticleIndex = RigidBodyID[InternalParticleIndex];
+
+				if (ExternalParticleIndex == INDEX_NONE)
+				{
+					return;
+				}
+
+				if (Particles.InvM(ExternalParticleIndex) != 0)
+				{
+					return;
+				}
+
+				// We need to check a window of Now - Dt to Now and see if we ever activated in that time.
+				// This is for short activations because if we miss one then the playback will be incorrect
+				const bool bShouldBeDisabled = !Parameters.RecordedTrack->GetWasActiveInWindow(InternalParticleIndex, ThisFrameTime, bIsReverseCachePlaying ? (ThisFrameTime - Dt) : (ThisFrameTime + Dt));
+				const bool bDisabledNow = Particles.Disabled(ExternalParticleIndex);
+				if (bDisabledNow != bShouldBeDisabled)
+				{
+					Particles.SetDisabledLowLevel(ExternalParticleIndex, bShouldBeDisabled);
+					ParticleUpdateLock.Lock();
+					if (!bShouldBeDisabled)
+					{
+						Particles.SetObjectState(ExternalParticleIndex, Chaos::EObjectStateType::Kinematic);
+						if (!ensure(ThisSolver->ActiveIndices().Find(ExternalParticleIndex) == nullptr))
+						{
+							bParticlesUpdated = true;
+						}
+						else if(!bParticlesUpdated)
+						{
+							ThisSolver->NonDisabledIndices().Add(ExternalParticleIndex);
+							ThisSolver->ActiveIndices().Add(ExternalParticleIndex);
+						}
+					}
+					else
+					{
+						if (!ensure(ThisSolver->ActiveIndices().Find(ExternalParticleIndex) != nullptr))
+						{
+							bParticlesUpdated = true;
+						}
+						else if(!bParticlesUpdated)
+						{
+							ThisSolver->NonDisabledIndices().Remove(ExternalParticleIndex);
+							ThisSolver->ActiveIndices().Remove(ExternalParticleIndex);
+						}
+					}
+					ParticleUpdateLock.Unlock();
+				}
+
+			});
+
+			// Do not add collisions if reverse
+			if (!bIsReverseCachePlaying)
+			{
+				const FRecordedFrame* RecordedFrame = Parameters.RecordedTrack->FindRecordedFrame(ThisFrameTime);
+				if (RecordedFrame == nullptr)
+				{
+					const int32 Index = Parameters.RecordedTrack->FindLastKeyBefore(ThisFrameTime);
+					if (Index > 0 && Index < Parameters.RecordedTrack->Records.Num())
+					{
+						RecordedFrame = &Parameters.RecordedTrack->Records[Index];
+					}
+				}
+
+				if (RecordedFrame)
+				{
+					// Collisions
+					if (Parameters.CollisionData.DoGenerateCollisionData && ThisFrameTime > 0.f && Parameters.CollisionData.CollisionDataSizeMax > 0)
+					{
+						if (RecordedFrame->Collisions.Num() > 0)
+						{
+							Chaos::FPhysicsSolver::FCollisionDataArray& AllCollisionsDataArray = ThisSolver->GetAllCollisionsDataArray();
+							TMap<IPhysicsProxyBase*, TArray<int32>>& AllCollisionsIndicesByPhysicsProxy = ThisSolver->GetAllCollisionsIndicesByPhysicsProxy();
+
+							if (!AllCollisionsIndicesByPhysicsProxy.Contains(this))
+							{
+								AllCollisionsIndicesByPhysicsProxy.Add(this, TArray<int32>());
+							}
+
+							for (int32 Idx = 0; Idx < RecordedFrame->Collisions.Num(); ++Idx)
+							{
+								// Check if the particle is still kinematic
+								if (RecordedFrame->Collisions[Idx].ParticleIndex < 0 ||
+									(RecordedFrame->Collisions[Idx].ParticleIndex >= 0 && RecordedFrame->Collisions[Idx].ParticleIndex < static_cast<int32>(Particles.Size()) && Particles.ObjectState(RecordedFrame->Collisions[Idx].ParticleIndex) == Chaos::EObjectStateType::Kinematic))
+								{
+									const int32 NewIdx = AllCollisionsDataArray.Add(Chaos::TCollisionData<float, 3>());
+									Chaos::TCollisionData<float, 3>& AllCollisionsDataArrayItem = AllCollisionsDataArray[NewIdx];
+
+									AllCollisionsDataArrayItem.Location = RecordedFrame->Collisions[Idx].Location;
+									AllCollisionsDataArrayItem.AccumulatedImpulse = RecordedFrame->Collisions[Idx].AccumulatedImpulse;
+									AllCollisionsDataArrayItem.Normal = RecordedFrame->Collisions[Idx].Normal;
+									AllCollisionsDataArrayItem.Velocity1 = RecordedFrame->Collisions[Idx].Velocity1;
+									AllCollisionsDataArrayItem.Velocity2 = RecordedFrame->Collisions[Idx].Velocity2;
+									AllCollisionsDataArrayItem.AngularVelocity1 = RecordedFrame->Collisions[Idx].AngularVelocity1;
+									AllCollisionsDataArrayItem.AngularVelocity2 = RecordedFrame->Collisions[Idx].AngularVelocity2;
+									AllCollisionsDataArrayItem.Mass1 = RecordedFrame->Collisions[Idx].Mass1;
+									AllCollisionsDataArrayItem.Mass2 = RecordedFrame->Collisions[Idx].Mass2;
+#if TODO_CONVERT_GEOMETRY_COLLECTION_PARTICLE_INDICES_TO_PARTICLE_POINTERS
+									AllCollisionsDataArrayItem.ParticleIndex = RecordedFrame->Collisions[Idx].ParticleIndex;
+#endif
+									AllCollisionsDataArrayItem.LevelsetIndex = RecordedFrame->Collisions[Idx].LevelsetIndex;
+									AllCollisionsDataArrayItem.ParticleIndexMesh = RecordedFrame->Collisions[Idx].ParticleIndexMesh;
+									AllCollisionsDataArrayItem.LevelsetIndexMesh = RecordedFrame->Collisions[Idx].LevelsetIndexMesh;
+
+									AllCollisionsIndicesByPhysicsProxy[this].Add(NewIdx);
+								}
+							}
+						}
+					}
+
+					// Breaking
+					if (Parameters.BreakingData.DoGenerateBreakingData && ThisFrameTime > 0.f && Parameters.BreakingData.BreakingDataSizeMax > 0)
+					{
+						if (RecordedFrame->Breakings.Num() > 0)
+						{
+							Chaos::FPhysicsSolver::FBreakingDataArray& AllBreakingsDataArray = ThisSolver->GetAllBreakingsDataArray();
+							TMap<IPhysicsProxyBase*, TArray<int32>>& AllBreakingsIndicesByPhysicsProxy = ThisSolver->GetAllBreakingsIndicesByPhysicsProxy();
+
+							if (!AllBreakingsIndicesByPhysicsProxy.Contains(this))
+							{
+								AllBreakingsIndicesByPhysicsProxy.Add(this, TArray<int32>());
+							}
+
+							for (int32 Idx = 0; Idx < RecordedFrame->Breakings.Num(); ++Idx)
+							{
+								// Check if the particle is still kinematic
+								if (RecordedFrame->Breakings[Idx].ParticleIndex < 0 ||
+									(RecordedFrame->Breakings[Idx].ParticleIndex >= 0 && RecordedFrame->Breakings[Idx].ParticleIndex < static_cast<int32>(Particles.Size()) && Particles.ObjectState(RecordedFrame->Breakings[Idx].ParticleIndex) == Chaos::EObjectStateType::Kinematic))
+								{
+									const int32 NewIdx = AllBreakingsDataArray.Add(Chaos::TBreakingData<float, 3>());
+									Chaos::TBreakingData<float, 3>& AllBreakingsDataArrayItem = AllBreakingsDataArray[NewIdx];
+
+									AllBreakingsDataArrayItem.Location = RecordedFrame->Breakings[Idx].Location;
+									AllBreakingsDataArrayItem.Velocity = RecordedFrame->Breakings[Idx].Velocity;
+									AllBreakingsDataArrayItem.AngularVelocity = RecordedFrame->Breakings[Idx].AngularVelocity;
+									AllBreakingsDataArrayItem.Mass = RecordedFrame->Breakings[Idx].Mass;
+#if TODO_CONVERT_GEOMETRY_COLLECTION_PARTICLE_INDICES_TO_PARTICLE_POINTERS
+									AllBreakingsDataArrayItem.ParticleIndex = RecordedFrame->Breakings[Idx].ParticleIndex;
+#endif
+									AllBreakingsDataArrayItem.ParticleIndexMesh = RecordedFrame->Breakings[Idx].ParticleIndexMesh;
+
+									AllBreakingsIndicesByPhysicsProxy[this].Add(NewIdx);
+								}
+							}
+						}
+					}
+
+					// Trailing
+					if (Parameters.TrailingData.DoGenerateTrailingData && ThisFrameTime > 0.f && Parameters.TrailingData.TrailingDataSizeMax > 0)
+					{
+						if (RecordedFrame->Trailings.Num() > 0)
+						{
+							Chaos::FPhysicsSolver::FTrailingDataArray& AllTrailingsDataArray = ThisSolver->GetAllTrailingsDataArray();
+							TMap<IPhysicsProxyBase*, TArray<int32>>& AllTrailingsIndicesByPhysicsProxy = ThisSolver->GetAllTrailingsIndicesByPhysicsProxy();
+
+							if (!AllTrailingsIndicesByPhysicsProxy.Contains(this))
+							{
+								AllTrailingsIndicesByPhysicsProxy.Add(this, TArray<int32>());
+							}
+
+							for (const FSolverTrailingData& Trailing : RecordedFrame->Trailings)
+							{
+								// Check if the particle is still kinematic
+								if (Trailing.ParticleIndex < 0 ||
+									(Trailing.ParticleIndex >= 0 && Trailing.ParticleIndex < static_cast<int32>(Particles.Size()) && Particles.ObjectState(Trailing.ParticleIndex) == Chaos::EObjectStateType::Kinematic))
+								{
+									const int32 NewIdx = AllTrailingsDataArray.Add(Chaos::TTrailingData<float, 3>());
+									Chaos::TTrailingData<float, 3>& AllTrailingsDataArrayItem = AllTrailingsDataArray[NewIdx];
+
+									AllTrailingsDataArrayItem.Location = Trailing.Location;
+									AllTrailingsDataArrayItem.Velocity = Trailing.Velocity;
+									AllTrailingsDataArrayItem.AngularVelocity = Trailing.AngularVelocity;
+									AllTrailingsDataArrayItem.Mass = Trailing.Mass;
+#if TODO_CONVERT_GEOMETRY_COLLECTION_PARTICLE_INDICES_TO_PARTICLE_POINTERS
+									AllTrailingsDataArrayItem.ParticleIndex = Trailing.ParticleIndex;
+#endif
+									AllTrailingsDataArrayItem.ParticleIndexMesh = Trailing.ParticleIndexMesh;
+
+									AllTrailingsIndicesByPhysicsProxy[this].Add(NewIdx);
+								}
+							}
+						}
+					}
+				}
+			}
+
+			if (bParticlesUpdated)
+			{
+				ThisSolver->InitializeFromParticleData(0);
+			}
+		}
+	}
+#endif
+}
+
+void FGeometryCollectionPhysicsProxy::EndFrameCallback(const float EndFrame)
+{
+	check(false); // not currently called
+#if TODO_REIMPLEMENT_GET_RIGID_PARTICLES
+	if (InitializedState == ESimulationInitializationState::Initialized)
+	{
+		FGeometryDynamicCollection* Collection = Parameters.DynamicCollection;
+		check(Collection);
+
+		ProxySimDuration += EndFrame;
+
+		if(Collection->HasAttribute("RigidBodyID", FGeometryCollection::TransformGroup))
+		{
+			//
+			//  Update transforms for the simulated transforms
+			//
+			TManagedArray<int32>& CollectionClusterID = SolverClusterID;
+			TManagedArray<FTransform>& Transform = Collection->Transform;
+			TManagedArray<int32>& Parent = Collection->Parent;
+			TManagedArray<TSet<int32>>& Children = Collection->Children;
+			TManagedArray<int32>& SimulationType = Collection->SimulationType;
+
+			TManagedArray<int32>& DynamicState = Collection->GetAttribute<int32>("DynamicState", FGeometryCollection::TransformGroup);
+
+			Chaos::TPBDRigidParticles<float, 3>& Particles = GetSolver()->GetRigidParticles();
+			const Chaos::TPBDCollisionConstraints<float, 3>& CollisionRule = GetSolver()->GetCollisionConstraints();
+			const Chaos::TArrayCollectionArray<Chaos::ClusterId>& ClusterID = GetSolver()->GetRigidClustering().GetClusterIdsArray();
+			const Chaos::TArrayCollectionArray<Chaos::TRigidTransform<float, 3>>& ClusterChildToParentMap = GetSolver()->GetRigidClustering().GetChildToParentMap();
+			const Chaos::TArrayCollectionArray<bool>& InternalCluster = GetSolver()->GetRigidClustering().GetInternalClusterArray();
+
+			//Particles X and R are aligned with center of mass and inertia principal axes.
+			//Renderer doesn't know about this and simply does ActorToWorld * GeomToActor * LocalSpaceVerts
+			//In proper math multiplication order:
+			//ParticleToWorld = ActorToWorld * GeomToActor * MassToLocal
+			//GeomToWorld = ActorToWorld * GeomToActor
+			//=> GeomToWorld = ParticleToWorld * MassToLocal.Inv()
+			//=> GeomToActor = ActorToWorld.Inv() * ParticleToWorld * MassToLocal.Inv()
+			const int32 TransformSize = Collection->NumElements(FGeometryCollection::TransformGroup);
+			const FTransform& ActorToWorld = Parameters.WorldTransform;
+
+			// It's not worth shrinking the EndFrameUnparentingBuffer array, at least until the solver supports deleting bodies.
+			if (EndFrameUnparentingBuffer.Num() < TransformSize)
+			{
+				EndFrameUnparentingBuffer.Init(-1, TransformSize);
+			}
+
+			{
+				for (int32 TransformGroupIndex = 0; TransformGroupIndex < TransformSize; ++TransformGroupIndex)
+				{
+					const int32 RigidBodyIndex = RigidBodyID[TransformGroupIndex];
+					if (RigidBodyIndex != INDEX_NONE)
+					{
+						// Update the transform and parent hierarchy of the active rigid bodies. Active bodies can be either
+						// rigid geometry defined from the leaf nodes of the collection, or cluster bodies that drive an entire
+						// branch of the hierarchy within the GeometryCollection.
+						// - Active bodies are directly driven from the global position of the corresponding
+						//   rigid bodies within the solver ( cases where RigidBodyID[TransformGroupIndex] is not disabled ). 
+						// - Deactivated bodies are driven from the transforms of their active parents. However the solver can
+						//   take ownership of the parents during the simulation, so it might be necessary to force deactivated
+						//   bodies out of the collections hierarchy during the simulation.  
+						if (!Particles.Disabled(RigidBodyID[TransformGroupIndex]))
+						{
+							// Update the transform of the active body. The active body can be either a single rigid
+							// or a collection of rigidly attached geometries (Clustering). The cluster is represented as a
+							// single transform in the GeometryCollection, and all children are stored in the local space
+							// of the parent cluster.
+							// ... When setting cluster transforms it is expected that the MassToLocal is identity.
+							//     Cluster initialization will set the vertices in the MassSpace of the rigid body.
+							// ... When setting individual rigid bodies that are not clustered, the MassToLocal will be 
+							//     non-Identity, and will reflect the difference between the geometric center of the geometry
+							//     and that corresponding rigid bodies center of mass. 
+							const FTransform ParticleToWorld(Particles.R(RigidBodyIndex), Particles.X(RigidBodyIndex));
+							// GeomToActor = ActorToWorld.Inv() * ParticleToWorld * MassToLocal.Inv();
+							Transform[TransformGroupIndex] = MassToLocal[TransformGroupIndex].GetRelativeTransformReverse(ParticleToWorld).GetRelativeTransform(ActorToWorld);
+							Transform[TransformGroupIndex].NormalizeRotation();
+
+							// dynamic state is also updated by the solver during field interaction. 
+							if (!Particles.Sleeping(RigidBodyIndex))
+							{
+								if (Particles.ObjectState(RigidBodyIndex) == Chaos::EObjectStateType::Kinematic)
+								{
+									DynamicState[TransformGroupIndex] = (int)EObjectStateTypeEnum::Chaos_Object_Kinematic;
+								}
+								else if (Particles.ObjectState(RigidBodyIndex) == Chaos::EObjectStateType::Static)
+								{
+									DynamicState[TransformGroupIndex] = (int)EObjectStateTypeEnum::Chaos_Object_Static;
+								}
+								else
+								{
+									DynamicState[TransformGroupIndex] = (int)EObjectStateTypeEnum::Chaos_Object_Dynamic;
+								}
+							}
+
+							// Force all enabled rigid bodies out of the transform hierarchy
+							if (Parent[TransformGroupIndex] != INDEX_NONE)
+							{
+								const int32 ParentIndex = Parent[TransformGroupIndex];
+								// Children in the hierarchy are stored in a TSet, which is not thread safe.  So we retain
+								// indices to remove afterwards.
+								EndFrameUnparentingBuffer[TransformGroupIndex] = ParentIndex;
+							}
+
+							// When a leaf node rigid body is removed from a cluster the rigid
+							// body will become active and needs its clusterID updated. This just
+							// syncs the clusterID all the time. 
+							CollectionClusterID[TransformGroupIndex] = ClusterID[RigidBodyIndex].Id;
+						}
+						else if (Particles.Disabled(RigidBodyIndex))
+						{
+							// dynamic state is also updated by the solver during field interaction. 
+							if (!Particles.Sleeping(RigidBodyIndex))
+							{
+								if (Particles.ObjectState(RigidBodyIndex) == Chaos::EObjectStateType::Kinematic)
+								{
+									DynamicState[TransformGroupIndex] = (int)EObjectStateTypeEnum::Chaos_Object_Kinematic;
+								}
+								else if (Particles.ObjectState(RigidBodyIndex) == Chaos::EObjectStateType::Static)
+								{
+									DynamicState[TransformGroupIndex] = (int)EObjectStateTypeEnum::Chaos_Object_Static;
+								}
+								else
+								{
+									DynamicState[TransformGroupIndex] = (int)EObjectStateTypeEnum::Chaos_Object_Dynamic;
+								}
+							}
+
+							// The rigid body parent cluster has changed within the solver, and its
+							// parent body is not tracked within the geometry collection. So we need to
+							// pull the rigid bodies out of the transform hierarchy, and just drive
+							// the positions directly from the solvers cluster particle. 
+							if (CollectionClusterID[TransformGroupIndex] != ClusterID[RigidBodyIndex].Id)
+							{
+								// Force all driven rigid bodies out of the transform hierarchy
+								if (Parent[TransformGroupIndex] != INDEX_NONE)
+								{
+									const int32 ParentIndex = Parent[TransformGroupIndex];
+									// Children in the hierarchy are stored in a TSet, which is not thread safe.  So we retain
+									// indices to remove afterwards.
+									EndFrameUnparentingBuffer[TransformGroupIndex] = ParentIndex;
+								}
+								CollectionClusterID[TransformGroupIndex] = ClusterID[RigidBodyIndex].Id;
+							}
+
+							// Disabled rigid bodies that have valid cluster parents, and have been re-indexed by the
+							// solver ( As in, They were re-clustered outside of the geometry collection), These clusters 
+							// will need to be rendered based on the clusters position. 
+							const int32 ClusterParentIndex = CollectionClusterID[TransformGroupIndex];
+							if (ClusterParentIndex != INDEX_NONE)
+							{
+								if (InternalCluster[ClusterParentIndex])
+								{
+									const FTransform ClusterChildToWorld = ClusterChildToParentMap[RigidBodyIndex] * FTransform(Particles.R(ClusterParentIndex), Particles.X(ClusterParentIndex));
+									if (Parameters.IsCacheRecording())
+									{
+										Particles.X(RigidBodyIndex) = ClusterChildToWorld.GetTranslation();
+										Particles.R(RigidBodyIndex) = ClusterChildToWorld.GetRotation();
+									}
+									// GeomToActor = ActorToWorld.Inv() * ClusterChildToWorld * MassToLocal.Inv();
+									Transform[TransformGroupIndex] = MassToLocal[TransformGroupIndex].GetRelativeTransformReverse(ClusterChildToWorld).GetRelativeTransform(ActorToWorld);
+									Transform[TransformGroupIndex].NormalizeRotation();
+								}
+							}
+						}
+					}
+				}
+			}
+			for (int32 TransformGroupIndex = 0; TransformGroupIndex < TransformSize; TransformGroupIndex++)
+			{
+				const int32 ParentIndex = EndFrameUnparentingBuffer[TransformGroupIndex];
+				if (ParentIndex >= 0)
+				{
+					// We reuse EndFrameUnparentingBuffer potentially without reinitialization, so reset this index to -1 before it gets paged out.
+					EndFrameUnparentingBuffer[TransformGroupIndex] = -1;
+
+					Children[ParentIndex].Remove(TransformGroupIndex);
+					Parent[TransformGroupIndex] = INDEX_NONE;
+				}
+			}
+
+			//
+			//  Set rest cache on simulated object.
+			//
+			if (Parameters.IsCacheRecording())
+			{
+				check(!UpdateRecordedStateCallback);
+				UpdateRecordedState(ProxySimDuration, RigidBodyID, CollectionClusterID, InternalCluster, Particles, CollisionRule);
+			}
+
+			// one way trigger from non-simulating to simulating
+			if (!IsObjectDynamic)
+			{
+				const Chaos::TArrayCollectionArray<bool>& ExternalID = GetSolver()->GetRigidClustering().GetInternalClusterArray();
+
+				for (int32 TransformGroupIndex = 0; TransformGroupIndex < TransformSize; TransformGroupIndex++)
+				{
+					const int32 RigidBodyIndex = RigidBodyID[TransformGroupIndex];
+					if (RigidBodyIndex != INDEX_NONE)
+					{
+						bool HasBeenRemoved = (PTDynamicCollection.StatusFlags[TransformGroupIndex] & FGeometryCollection::FS_RemoveOnFracture) && Particles.Disabled(RigidBodyIndex) && ClusterID[RigidBodyIndex].Id == INDEX_NONE;
+
+						if (HasBeenRemoved || !Particles.Disabled(RigidBodyIndex))
+						{
+							if (DynamicState[TransformGroupIndex] != (int32)EObjectStateTypeEnum::Chaos_Object_Static)
+							{
+								IsObjectDynamic = true;
+								break;
+							}
+						}
+						else
+						{
+							const int32 ParentClusterID = ClusterID[RigidBodyIndex].Id;
+							if (ParentClusterID != INDEX_NONE && ExternalID[ParentClusterID]
+								&& Particles.ObjectState(ParentClusterID) != Chaos::EObjectStateType::Static)
+							{
+								IsObjectDynamic = true;
+								break;
+							}
+						}
+					}
+				}
+			}
+
+			// Can't change visibility right now so setting scale to zero instead
+			// only process if enabled on this object (ones with glass), most will skip over
+			if (Parameters.RemoveOnFractureEnabled && IsObjectDynamic)
+			{
+				for (int32 TransformGroupIndex = 0; TransformGroupIndex < TransformSize; TransformGroupIndex++)
+				{
+					const int32 RigidBodyIndex = RigidBodyID[TransformGroupIndex];
+					if (RigidBodyIndex != INDEX_NONE)
+					{
+						if ((PTDynamicCollection.StatusFlags[TransformGroupIndex] & FGeometryCollection::FS_RemoveOnFracture)
+							&& Particles.Disabled(RigidBodyIndex) && ClusterID[RigidBodyIndex].Id == INDEX_NONE)
+						{
+							PTDynamicCollection.Transform[TransformGroupIndex].SetScale3D(FVector::ZeroVector);
+						}
+					}
+				}
+			}
+		}
+
+#if TODO_REIMPLEMENT_RIGID_CLUSTERING
+		{
+			IsObjectLoading = GetSolver()->GetRigidClustering().NumberOfPendingClusters() != 0;
+		}
+#endif
+	}
+	Commands.Empty();
+#endif
+
+}
+
+void FGeometryCollectionPhysicsProxy::CreateRigidBodyCallback(FParticlesType& Particles)
+{
+	check(false); // This function is deprecated and should never be called.
+#if TODO_REIMPLEMENT_RIGID_CLUSTERING
+
+	const FGeometryCollection* RestCollection = Parameters.RestCollection;
+	FGeometryDynamicCollection* DynamicCollection = Parameters.DynamicCollection;
+	check(DynamicCollection);
+	
+	if (Parameters.Simulating && 
+		((InitializedState == ESimulationInitializationState::Unintialized) || 
+		 (InitializedState == ESimulationInitializationState::Activated)))
+	{
+		Chaos::FPhysicsSolver* ThisSolver = GetSolver();
+		Chaos::TArrayCollectionArray<int32>& ClusterGroupIndex = ThisSolver->GetRigidClustering().GetClusterGroupIndexArray();
+		Chaos::TArrayCollectionArray<float>& StrainArray = ThisSolver->GetRigidClustering().GetStrainArray();
+		const float StrainDefault = Parameters.DamageThreshold.Num() ? Parameters.DamageThreshold[0] : 0;
+
+		const TManagedArray<int32>& TransformIndex = RestCollection->TransformIndex;
+		const TManagedArray<int32>& BoneMap = RestCollection->BoneMap;
+		const TManagedArray<int32>& Parent = RestCollection->Parent;
+		const TManagedArray<TSet<int32>>& Children = RestCollection->Children;
+		const TManagedArray<int32>& SimulationType = RestCollection->SimulationType;
+		const TManagedArray<FVector>& Vertex = RestCollection->Vertex;
+		const TManagedArray<int32>& DynamicState = DynamicCollection->DynamicState;
+		const TManagedArray<int32>& CollisionGroup = DynamicCollection->CollisionGroup;
+		const TManagedArray<float>& Mass = RestCollection->GetAttribute<float>("Mass", FTransformCollection::TransformGroup);
+		const TManagedArray<FVector>& InertiaTensor = RestCollection->GetAttribute<FVector>("InertiaTensor", FTransformCollection::TransformGroup);
+
+		TArray<FTransform> Transform;
+		GeometryCollectionAlgo::GlobalMatrices(DynamicCollection->Transform, DynamicCollection->Parent, Transform);
+		check(DynamicCollection->Transform.Num() == Transform.Num());
+
+		// count particles to add
+		int NumSimulatedParticles = 0;
+		for (int32 Index = 0; Index < SimulatableParticles.Num(); Index++)
+		{
+			if (SimulatableParticles[Index] == true)
+			{
+				NumSimulatedParticles++;
+			}
+		}
+
+		// Add entries into simulation array
+		int NumRigids = Particles.Size();
+		BaseParticleIndex = NumRigids;
+		Particles.AddParticles(NumSimulatedParticles);
+		for (int32 Index = 0, NextId = 0; Index < SimulatableParticles.Num(); Index++)
+		{
+			if (SimulatableParticles[Index] == true)
+			{
+				RigidBodyID[Index] = NumRigids + NextId++;
+			}
+		}
+
+		// Add the rigid bodies
+		//for (int32 Index = 0; Index < TransformIndex.Num(); Index++)
+		const int32 NumGeometries = DynamicCollection->NumElements(FGeometryCollection::GeometryGroup);
+		ParallelFor(NumGeometries, [&](int32 GeometryIndex)
+		{
+			const int32 TransformGroupIndex = TransformIndex[GeometryIndex];
+			if (SimulatableParticles[TransformGroupIndex] == true)
+			{
+				const int32 RigidBodyIndex = RigidBodyID[TransformGroupIndex];
+				const FTransform WorldTransform = MassToLocal[TransformGroupIndex] * Transform[TransformGroupIndex] * Parameters.WorldTransform;
+
+				PopulateSimulatedParticle(Particles, Parameters.Shared, nullptr,// Simplicials[TransformGroupIndex].Get(),
+					Implicits[TransformGroupIndex], ThisSolver->GetMassScale() * Mass[TransformGroupIndex], ThisSolver->GetMassScale() * InertiaTensor[TransformGroupIndex],
+					RigidBodyIndex, WorldTransform, (uint8)EObjectStateTypeEnum::Chaos_Object_Dynamic, CollisionGroup[TransformGroupIndex]);
+				ClusterGroupIndex[RigidBodyIndex] = Parameters.ClusterGroupIndex;
+				Particles.CollisionParticles(RigidBodyIndex).Reset(Simplicials[TransformGroupIndex].Release());
+				if (Particles.CollisionParticles(RigidBodyIndex))
+				{
+					int32 NumCollisionParticles = Particles.CollisionParticles(RigidBodyIndex)->Size();
+					int32 CollisionParticlesSize = FMath::Max(0, FMath::Min(int(NumCollisionParticles*CollisionParticlesPerObjectFraction), NumCollisionParticles));
+					Particles.CollisionParticles(RigidBodyIndex)->Resize(CollisionParticlesSize);
+				}
+
+				StrainArray[RigidBodyIndex] = StrainDefault;
+				GetSolver()->SetPhysicsMaterial(RigidBodyIndex, Parameters.PhysicalMaterial);
+			}
+
+		});
+
+		for (FFieldSystemCommand & Cmd : Parameters.InitializationCommands)
+		{
+			if (Cmd.MetaData.Contains(FFieldSystemMetaData::EMetaType::ECommandData_ProcessingResolution)) Cmd.MetaData.Remove(FFieldSystemMetaData::EMetaType::ECommandData_ProcessingResolution);
+			FFieldSystemMetaDataProcessingResolution* ResolutionData = new FFieldSystemMetaDataProcessingResolution(EFieldResolutionType::Field_Resolution_Maximum);
+			Cmd.MetaData.Add(FFieldSystemMetaData::EMetaType::ECommandData_ProcessingResolution, TUniquePtr< FFieldSystemMetaDataProcessingResolution >(ResolutionData));
+			Commands.Add(Cmd);
+		}
+		Parameters.InitializationCommands.Empty();
+		ProcessCommands(Particles, GetSolver()->GetSolverTime());
+
+		ParallelFor(NumGeometries, [&](int32 GeometryIndex)
+		{
+			const int32 TransformGroupIndex = TransformIndex[GeometryIndex];
+			if (SimulatableParticles[TransformGroupIndex] == true)
+			{
+				const int32 RigidBodyIndex = RigidBodyID[TransformGroupIndex];
+
+				if (Parameters.InitialVelocityType == EInitialVelocityTypeEnum::Chaos_Initial_Velocity_User_Defined)
+				{
+					InitialLinearVelocity[TransformGroupIndex] = Parameters.InitialLinearVelocity;
+					InitialAngularVelocity[TransformGroupIndex] = Parameters.InitialAngularVelocity;
+					if (DynamicState[TransformGroupIndex] == (int32)EObjectStateTypeEnum::Chaos_Object_Dynamic)
+					{
+						Particles.V(RigidBodyIndex) = Parameters.InitialLinearVelocity;
+						Particles.W(RigidBodyIndex) = Parameters.InitialAngularVelocity;
+					}
+				}
+			}
+		});
+
+		// InitializeKinematics
+		if (Parameters.DynamicCollection)
+		{
+			for (int32 TransformGroupIndex = 0; TransformGroupIndex < GetTransformGroupSize(); ++TransformGroupIndex)
+			{
+				if (Chaos::TPBDRigidParticleHandle<float, 3>* Handle = SolverParticleHandles[TransformGroupIndex])
+				{
+					// this is being initialized to Static, can't change from static to sleeping
+					Handle->SetObjectState(Chaos::EObjectStateType::Dynamic);
+
+					if (DynamicState[TransformGroupIndex] == (int32)EObjectStateTypeEnum::Chaos_Object_Kinematic)
+					{
+						Handle->SetObjectState(Chaos::EObjectStateType::Kinematic);
+					}
+					else if (DynamicState[TransformGroupIndex] == (int32)EObjectStateTypeEnum::Chaos_Object_Static)
+					{
+						Handle->SetObjectState(Chaos::EObjectStateType::Static);
+					}
+					else if (DynamicState[TransformGroupIndex] == (int32)EObjectStateTypeEnum::Chaos_Object_Sleeping)
+					{
+						Handle->SetObjectState(Chaos::EObjectStateType::Sleeping);
+					}
+				}
+			}
+		}
+
+		InitializeRemoveOnFracture(Particles, DynamicState);
+
+		// #BG Temporary - don't cluster when playing back. Needs to be changed when kinematics are per-proxy to support
+		// kinematic to dynamic transition for clusters.
+		if (Parameters.EnableClustering)// && Parameters.CacheType != EGeometryCollectionCacheType::Play)
+		{
+			const TArray<int32> RecursiveOrder = ComputeRecursiveOrder(*RestCollection);
+
+			// num clusters
+			uint32 NumClusters=0;
+			TArray<bool> SubTreeContainsSimulatableParticle;
+			SubTreeContainsSimulatableParticle.SetNum(RecursiveOrder.Num());
+			for (const int32 TransformGroupIndex : RecursiveOrder)
+			{
+				if (Children[TransformGroupIndex].Num() > 0)
+				{
+					SubTreeContainsSimulatableParticle[TransformGroupIndex] = false;
+
+					TArray<uint32> RigidChildren, CollectionChildren;
+					for (const int32 ChildIndex : Children[TransformGroupIndex])
+					{
+						if(SubTreeContainsSimulatableParticle[ChildIndex])
+						{
+							NumClusters++;
+							SubTreeContainsSimulatableParticle[TransformGroupIndex] = true;
+							break;
+						}
+					}
+				}
+				else
+				{
+					SubTreeContainsSimulatableParticle[TransformGroupIndex] = (RigidBodyID[TransformGroupIndex] != INDEX_NONE);
+				}
+			}
+
+			const int32 ClusterStartIndex = Particles.Size();
+			Particles.AddParticles(NumClusters);
+
+			int32 ClusterRigidBodyId = ClusterStartIndex;
+			for (const int32 TransformGroupIndex : RecursiveOrder)
+			{
+				TArray<uint32> RigidChildren, CollectionChildren;
+				for (const int32 ChildIndex : Children[TransformGroupIndex])
+				{
+					if (RigidBodyID[ChildIndex] != INDEX_NONE)
+					{
+						RigidChildren.Add(RigidBodyID[ChildIndex]);
+						CollectionChildren.Add(ChildIndex);
+					}
+				}
+				if (RigidChildren.Num())
+				{
+					if (ReportTooManyChildrenNum >= 0 && RigidChildren.Num() > ReportTooManyChildrenNum)
+					{
+						UE_LOG(LogChaos, Warning, TEXT("Too many children (%d) in a single cluster:%s"), RigidChildren.Num(), *Parameters.Name);
+					}
+
+					RigidBodyID[TransformGroupIndex] = ClusterRigidBodyId;
+					Chaos::FClusterCreationParameters<float> CreationParameters;
+					CreationParameters.RigidBodyIndex = RigidBodyID[TransformGroupIndex];
+					BuildClusters(TransformGroupIndex, CollectionChildren, RigidChildren, CreationParameters);
+					ClusterRigidBodyId++;
+				}
+			}
+		}
+
+		NumParticles = Particles.Size() - BaseParticleIndex;
+
+		// Set Connectivity
+		const int32 NumTransforms = DynamicCollection->NumElements(FGeometryCollection::TransformGroup);
+		ParallelFor(NumTransforms, [&](int32 TransformGroupIndex)
+		{
+			if (Children[TransformGroupIndex].Num() > 0)
+			{
+				if (RigidBodyID[TransformGroupIndex] != INDEX_NONE)
+				{
+					GetSolver()->GetRigidClustering().GenerateConnectionGraph(RigidBodyID[TransformGroupIndex]);
+				}
+			}
+		});
+
+		// If we're recording and want to start immediately caching then we should cache the rest state
+		if (Parameters.IsCacheRecording() && Parameters.CacheBeginTime == 0.0f)
+		{
+			if (UpdateRecordedStateCallback)
+			{
+				UpdateRecordedStateCallback(0.0f, RigidBodyID, Particles, GetSolver()->GetCollisionConstraints());
+			}
+		}
+
+
+		if (InitializedState == ESimulationInitializationState::Activated)
+		{
+			//
+			//  Activated bodies has already been called so we are good to go.
+			//
+			InitializedState = ESimulationInitializationState::Initialized;
+
+			if (Parameters.EnableClustering && Parameters.ClusterGroupIndex)
+			{
+				GetSolver()->GetRigidClustering().IncrementPendingClusterCounter(Parameters.ClusterGroupIndex);
+				GetSolver()->GetRigidClustering().DecrementPendingClusterCounter(Parameters.ClusterGroupIndex);
+			}
+		}
+		else if (InitializedState == ESimulationInitializationState::Unintialized)
+		{
+
+			//
+			//  Activated bodies has not been called, se we are waiting 
+			//  to become active. Deactivate all bodies, and wait for
+			//  ActivatedBodies to be called, and defer the cluster initialization
+			//
+			InitializedState = ESimulationInitializationState::Created;
+			for (int32 TransformGroupIndex = 0; TransformGroupIndex < NumTransforms; TransformGroupIndex++)
+			{
+				const int32 RigidBodyIndex = RigidBodyID[TransformGroupIndex];
+				if (RigidBodyIndex != INDEX_NONE)
+				{
+					if (!Particles.Disabled(RigidBodyIndex))
+					{
+						PendingActivationList.Add(TransformGroupIndex);
+						GetSolver()->GetEvolution()->DisableParticle(RigidBodyIndex);
+					}
+				}
+			}
+
+			//
+			//  Clustering needs to advertise its group id to the cluster so
+			//  that the group is not initialized before all the bodies are
+			//  loaded and created. 
+			//
+			if (Parameters.EnableClustering && Parameters.ClusterGroupIndex )
+			{
+				GetSolver()->GetRigidClustering().IncrementPendingClusterCounter(Parameters.ClusterGroupIndex);
+			}
+		}
+		else
+		{
+			// unknown initialization state in creation callback
+			ensure(false);
+		}
+	}
+#endif
+}
+
+
+void FGeometryCollectionPhysicsProxy::BindParticleCallbackMapping(Chaos::TArrayCollectionArray<PhysicsProxyWrapper>& PhysicsProxyReverseMap, Chaos::TArrayCollectionArray<int32>& ParticleIDReverseMap)
+{
+	/*
+	for (int32 Index = 0; Index < RigidBodyID.Num(); Index++)
+	{
+		const int32 RigidBodyIndex = RigidBodyID[Index];
+		if (RigidBodyIndex != INDEX_NONE)
+		{
+			PhysicsProxyReverseMap[RigidBodyIndex] = { this, EPhysicsProxyType::GeometryCollectionType };
+			ParticleIDReverseMap[RigidBodyIndex] = Index;
+		}
+	}
+	*/
+}
+
+void FGeometryCollectionPhysicsProxy::PushKinematicStateToSolver(FParticlesType& Particles)
+{
+#if TODO_REIMPLEMENT_RIGID_CLUSTERING
+	if (InitializedState == ESimulationInitializationState::Initialized)
+	{
+		FGeometryDynamicCollection* Collection = Parameters.DynamicCollection;
+		if (Collection->Transform.Num())
+		{
+			TManagedArray<int32>& DynamicState = Collection->GetAttribute<int32>("DynamicState", FGeometryCollection::TransformGroup);
+
+			TSet<int32> ClustersToUpdate;
+			const Chaos::TArrayCollectionArray<Chaos::ClusterId>& ClusterID = GetSolver()->GetRigidClustering().GetClusterIdsArray();
+
+			for (int32 TransformGroupIndex = 0; TransformGroupIndex < DynamicState.Num(); TransformGroupIndex++)
+			{
+				if (RigidBodyID[TransformGroupIndex] != INDEX_NONE)
+				{
+					int32 RigidBodyIndex = RigidBodyID[TransformGroupIndex];
+					if (DynamicState[TransformGroupIndex] == (int32)EObjectStateTypeEnum::Chaos_Object_Dynamic
+						&& (Particles.ObjectState(RigidBodyIndex) == Chaos::EObjectStateType::Kinematic || Particles.ObjectState(RigidBodyIndex) == Chaos::EObjectStateType::Static)
+						&& FLT_EPSILON < Particles.M(RigidBodyIndex))
+					{
+						Particles.SetObjectState(RigidBodyIndex, Chaos::EObjectStateType::Dynamic);
+
+						if (Parameters.InitialVelocityType == EInitialVelocityTypeEnum::Chaos_Initial_Velocity_User_Defined)
+						{
+							Particles.V(RigidBodyIndex) = Chaos::TVector<float, 3>(InitialLinearVelocity[TransformGroupIndex]);
+							Particles.W(RigidBodyIndex) = Chaos::TVector<float, 3>(InitialAngularVelocity[TransformGroupIndex]);
+						}
+
+						if (!Particles.Disabled(RigidBodyIndex) && Particles.Sleeping(RigidBodyIndex))
+						{
+							Particles.SetSleeping(RigidBodyIndex, false);
+							GetSolver()->ActiveIndices().Add(RigidBodyIndex);
+						}
+						else
+						{
+							if (ClusterID[RigidBodyIndex].Id != INDEX_NONE)
+							{
+								int32 ActiveClusterIndex = GetSolver()->GetRigidClustering().GetActiveClusterIndex(RigidBodyIndex);
+								if (ActiveClusterIndex != INDEX_NONE)
+								{
+									ClustersToUpdate.Add(ActiveClusterIndex);
+								}
+							}
+						}
+					}
+					else if ((DynamicState[TransformGroupIndex] == (int32)EObjectStateTypeEnum::Chaos_Object_Kinematic)
+						&& (Particles.ObjectState(RigidBodyIndex) == Chaos::EObjectStateType::Dynamic)
+						&& FLT_EPSILON < Particles.M(RigidBodyIndex))
+					{
+						Particles.SetObjectState(RigidBodyIndex, Chaos::EObjectStateType::Kinematic);
+
+						if (ClusterID[RigidBodyIndex].Id != INDEX_NONE)
+						{
+							int32 ActiveClusterIndex = GetSolver()->GetRigidClustering().GetActiveClusterIndex(RigidBodyIndex);
+							if (ActiveClusterIndex != INDEX_NONE)
+							{
+								ClustersToUpdate.Add(ActiveClusterIndex);
+							}
+						}
+					}
+					else if ((DynamicState[TransformGroupIndex] == (int32)EObjectStateTypeEnum::Chaos_Object_Static)
+						&& (Particles.ObjectState(RigidBodyIndex) == Chaos::EObjectStateType::Dynamic)
+						&& FLT_EPSILON < Particles.M(RigidBodyIndex))
+					{
+						Particles.SetObjectState(RigidBodyIndex, Chaos::EObjectStateType::Static);
+
+						if (ClusterID[RigidBodyIndex].Id != INDEX_NONE)
+						{
+							int32 ActiveClusterIndex = GetSolver()->GetRigidClustering().GetActiveClusterIndex(RigidBodyIndex);
+							if (ActiveClusterIndex != INDEX_NONE)
+							{
+								ClustersToUpdate.Add(ActiveClusterIndex);
+							}
+						}
+					}
+					else if ((DynamicState[TransformGroupIndex] == (int32)EObjectStateTypeEnum::Chaos_Object_Sleeping)
+						&& (Particles.ObjectState(RigidBodyIndex) == Chaos::EObjectStateType::Dynamic))
+					{
+						Particles.SetObjectState(RigidBodyIndex, Chaos::EObjectStateType::Sleeping);
+						GetSolver()->ActiveIndices().Remove(RigidBodyIndex);
+					}
+					else if ((DynamicState[TransformGroupIndex] == (int32)EObjectStateTypeEnum::Chaos_Object_Dynamic)
+						&& (Particles.ObjectState(RigidBodyIndex) == Chaos::EObjectStateType::Sleeping))
+					{
+						Particles.SetObjectState(RigidBodyIndex, Chaos::EObjectStateType::Dynamic);
+						GetSolver()->ActiveIndices().Add(RigidBodyIndex);
+					}
+				}
+			}
+		}
+	}
+#endif
+}
+
+void FGeometryCollectionPhysicsProxy::ParameterUpdateCallback(FParticlesType& Particles, const float Time)
+{
+	FGeometryDynamicCollection* Collection = Parameters.DynamicCollection;
+	check(Collection);
+
+	if (Collection->Transform.Num())
+	{
+		ProcessCommands(Particles, Time);
+
+
+#ifndef TODO_REIMPLEMENT_RIGID_CACHING
+		if (Parameters.RecordedTrack)
+		{
+			float ReverseTime = Parameters.RecordedTrack->GetLastTime() - Time + Parameters.ReverseCacheBeginTime;
+			// @todo(mlentine): We shouldn't need to do this every frame
+			if (Parameters.IsCacheRecording() && Time > Parameters.ReverseCacheBeginTime&& Parameters.ReverseCacheBeginTime != 0 && Parameters.RecordedTrack->IsTimeValid(ReverseTime))
+			{
+				for (int32 Index = 0; Index < RigidBodyID.Num(); Index++)
+				{
+					int32 RigidBodyIndex = RigidBodyID[Index];
+
+					// Check index, will be invalid for cluster parents.
+					if (RigidBodyIndex != INDEX_NONE)
+					{
+						Particles.InvM(RigidBodyIndex) = 0.f;
+						Particles.InvI(RigidBodyIndex) = Chaos::PMatrix<float, 3, 3>(0.f);
+					}
+				}
+			}
+		}
+#endif
+	}
+}
+
+void FGeometryCollectionPhysicsProxy::DisableCollisionsCallback(TSet<TTuple<int32, int32>>& InPairs)
+{
+
+}
+
+void FGeometryCollectionPhysicsProxy::AddForceCallback(FParticlesType& InParticles, const float InDt, const int32 InIndex)
+{
+	// gravity forces managed directly on the solver for now
+}
+#endif
+
+
diff --git a/Engine/Source/Runtime/Experimental/ChaosSolvers/Public/PhysicsProxy/GeometryCollectionPhysicsProxy.h b/Engine/Source/Runtime/Experimental/ChaosSolvers/Public/PhysicsProxy/GeometryCollectionPhysicsProxy.h
--- a/Engine/Source/Runtime/Experimental/ChaosSolvers/Public/PhysicsProxy/GeometryCollectionPhysicsProxy.h
+++ b/Engine/Source/Runtime/Experimental/ChaosSolvers/Public/PhysicsProxy/GeometryCollectionPhysicsProxy.h
@@ -218,21 +218,16 @@ public:
 	static void InitializeSharedCollisionStructures(Chaos::FErrorReporter& ErrorReporter, FGeometryCollection& RestCollection, const FSharedSimulationParameters& SharedParams);
 	static void InitRemoveOnFracture(FGeometryCollection& RestCollection, const FSharedSimulationParameters& SharedParams);
 
-	static void MergeRecordedTracks(const FRecordedTransformTrack& A, const FRecordedTransformTrack& B, FRecordedTransformTrack& Target);
-	static FRecordedFrame& InsertRecordedFrame(FRecordedTransformTrack& InTrack, float InTime);
-
-	// 
-	// DEPRECATED
-	//
-	void UpdateKinematicBodiesCallback(const FParticlesType& InParticles, const float InDt, const float InTime, FKinematicProxy& InKinematicProxy);
-	void StartFrameCallback(const float InDt, const float InTime);
-	void EndFrameCallback(const float InDt);
-	void BindParticleCallbackMapping(Chaos::TArrayCollectionArray<PhysicsProxyWrapper> & PhysicsProxyReverseMap, Chaos::TArrayCollectionArray<int32> & ParticleIDReverseMap);
-	void CreateRigidBodyCallback(FParticlesType& InOutParticles);
+	void FieldForcesUpdateCallback(Chaos::FPhysicsSolver* InSolver, FParticlesType& Particles, Chaos::TArrayCollectionArray<FVector>& Force, Chaos::TArrayCollectionArray<FVector>& Torque, const float Time);
 	void ParameterUpdateCallback(FParticlesType& InParticles, const float InTime);
-	void DisableCollisionsCallback(TSet<TTuple<int32, int32>>& InPairs);
-	void AddForceCallback(FParticlesType& InParticles, const float InDt, const int32 InIndex);
-	void FieldForcesUpdateCallback(Chaos::FPhysicsSolver* InSolver, FParticlesType& Particles, Chaos::TArrayCollectionArray<FVector> & Force, Chaos::TArrayCollectionArray<FVector> & Torque, const float Time);
+
+	void UpdateKinematicBodiesCallback(const FParticlesType& InParticles, const float InDt, const float InTime, FKinematicProxy& InKinematicProxy) {}
+	void StartFrameCallback(const float InDt, const float InTime) {}
+	void EndFrameCallback(const float InDt) {}
+	void BindParticleCallbackMapping(Chaos::TArrayCollectionArray<PhysicsProxyWrapper>& PhysicsProxyReverseMap, Chaos::TArrayCollectionArray<int32>& ParticleIDReverseMap) {}
+	void CreateRigidBodyCallback(FParticlesType& InOutParticles) {}
+	void DisableCollisionsCallback(TSet<TTuple<int32, int32>>& InPairs) {}
+	void AddForceCallback(FParticlesType& InParticles, const float InDt, const int32 InIndex) {}
 
 protected:
 	/**
@@ -283,42 +278,8 @@ protected:
 
 	static void InitializeDynamicCollection(FGeometryDynamicCollection& DynamicCollection, const FGeometryCollection& RestCollection, const FSimulationParameters& Params);
 	void InitializeRemoveOnFracture(FParticlesType& Particles, const TManagedArray<int32>& DynamicState);
-	void PushKinematicStateToSolver(FParticlesType& Particles);
-
 	void ProcessCommands(FParticlesType& Particles, const float Time);
 
-private:
-	void UpdateRecordedState(
-		float SolverTime, 
-		const TManagedArray<int32>& RigidBodyID, 
-		const TManagedArray<int32>& CollectionClusterID, 
-		const Chaos::TArrayCollectionArray<bool>& InternalCluster, 
-		const FParticlesType& Particles, 
-		const FCollisionConstraintsType& CollisionRule);
-
-	void AddCollisionToCollisionData(
-		FRecordedFrame* ExistingFrame, 
-		const FParticlesType& Particles, 
-		const Chaos::TPBDCollisionConstraints<float, 3>::FPointContactConstraint& Constraint);
-
-	void UpdateCollisionData(
-		const FParticlesType& Particles, 
-		const FCollisionConstraintsType& CollisionRule, 
-		FRecordedFrame* ExistingFrame);
-
-	void AddBreakingToBreakingData(
-		FRecordedFrame* ExistingFrame, 
-		const FParticlesType& Particles, 
-		const Chaos::TBreakingData<float, 3>& Breaking);
-
-	void UpdateBreakingData(
-		const FParticlesType& Particles, 
-		FRecordedFrame* ExistingFrame);
-
-	void UpdateTrailingData(
-		const FParticlesType& Particles, 
-		FRecordedFrame* ExistingFrame);
-
 private:
 	FSimulationParameters Parameters;
 
