Chaos :

- Added back Implicits

diff --git a/Engine/Source/Runtime/Experimental/ChaosSolvers/Private/GeometryCollectionProxyData.cpp b/Engine/Source/Runtime/Experimental/ChaosSolvers/Private/GeometryCollectionProxyData.cpp
--- a/Engine/Source/Runtime/Experimental/ChaosSolvers/Private/GeometryCollectionProxyData.cpp
+++ b/Engine/Source/Runtime/Experimental/ChaosSolvers/Private/GeometryCollectionProxyData.cpp
@@ -63,7 +63,7 @@ FGeometryDynamicCollection::FGeometryDynamicCollection()
 	AddExternalAttribute("MassToLocal", FTransformCollection::TransformGroup, MassToLocal);
 	//AddExternalAttribute(ShapesQueryDataAttribute, FTransformCollection::TransformGroup, ShapeQueryData);
 	//AddExternalAttribute(ShapesSimDataAttribute, FTransformCollection::TransformGroup, ShapeSimData);
-	//AddExternalAttribute(SharedGeometryAttribute, FTransformCollection::TransformGroup, SharedGeometry);
+	AddExternalAttribute(SharedGeometryAttribute, FTransformCollection::TransformGroup, SharedGeometry);
 	//AddExternalAttribute(SimplicialsAttribute, FTransformCollection::TransformGroup, Simplicials);
 	AddExternalAttribute("SimulatableParticles", FGeometryCollection::TransformGroup, SimulatableParticles);
 	AddExternalAttribute(SolverClusterHandlesAttribute, FTransformCollection::TransformGroup, SolverClusterHandles);
diff --git a/Engine/Source/Runtime/Experimental/ChaosSolvers/Public/GeometryCollectionProxyData.h b/Engine/Source/Runtime/Experimental/ChaosSolvers/Public/GeometryCollectionProxyData.h
--- a/Engine/Source/Runtime/Experimental/ChaosSolvers/Public/GeometryCollectionProxyData.h
+++ b/Engine/Source/Runtime/Experimental/ChaosSolvers/Public/GeometryCollectionProxyData.h
@@ -85,8 +85,8 @@ public:
 	TManagedArray<FTransform> MassToLocal;
 	//TManagedArray<TArray<FCollisionFilterData>> ShapeQueryData;
 	//TManagedArray<TArray<FCollisionFilterData>> ShapeSimData;
-	//TManagedArray<TSharedPtr<Chaos::FImplicitObject, ESPMode::ThreadSafe>> SharedGeometry;
-	//TManagedArray<TSharedPtr<FCollisionStructureManager::FSimplicial> > Simplicials;
+	TManagedArray<TSharedPtr<Chaos::FImplicitObject, ESPMode::ThreadSafe>> SharedGeometry;
+	//TManagedArray<TUniquePtr<FCollisionStructureManager::FSimplicial>> Simplicials;
 	TManagedArray<bool> SimulatableParticles;
 	TManagedArray<Chaos::TPBDRigidClusteredParticleHandle<float, 3>*> SolverClusterHandles;
 	TManagedArray<Chaos::TPBDRigidParticleHandle<float, 3>*> SolverClusterID; // Rename to ClusterParent?
