Chaos :

- Moved the Handles and Proxy specific particle management arrays out of the DynamicCollection.

diff --git a/Engine/Source/Runtime/Experimental/ChaosSolvers/Private/GeometryCollectionProxyData.cpp b/Engine/Source/Runtime/Experimental/ChaosSolvers/Private/GeometryCollectionProxyData.cpp
--- a/Engine/Source/Runtime/Experimental/ChaosSolvers/Private/GeometryCollectionProxyData.cpp
+++ b/Engine/Source/Runtime/Experimental/ChaosSolvers/Private/GeometryCollectionProxyData.cpp
@@ -45,8 +45,6 @@ const FName FGeometryDynamicCollection::ShapesSimDataAttribute("ShapesSimData");
 const FName FGeometryDynamicCollection::SharedGeometryAttribute("SharedGeometry");
 const FName FGeometryDynamicCollection::SimplicialsAttribute("CollisionParticles");
 const FName FGeometryDynamicCollection::SharedImplicitsAttribute("SharedImplicits");
-const FName FGeometryDynamicCollection::SolverParticleHandlesAttribute("SolverParticleHandles");
-const FName FGeometryDynamicCollection::SolverClusterHandlesAttribute("SolverClusterHandles");
 
 FGeometryDynamicCollection::FGeometryDynamicCollection()
 	: FTransformDynamicCollection()
@@ -66,9 +64,6 @@ FGeometryDynamicCollection::FGeometryDynamicCollection()
 	AddExternalAttribute(SharedGeometryAttribute, FTransformCollection::TransformGroup, SharedGeometry);
 	//AddExternalAttribute(SimplicialsAttribute, FTransformCollection::TransformGroup, Simplicials);
 	AddExternalAttribute("SimulatableParticles", FGeometryCollection::TransformGroup, SimulatableParticles);
-	AddExternalAttribute(SolverClusterHandlesAttribute, FTransformCollection::TransformGroup, SolverClusterHandles);
-	AddExternalAttribute("SolverClusterID", FTransformCollection::TransformGroup, SolverClusterID);
-	AddExternalAttribute(SolverParticleHandlesAttribute, FTransformCollection::TransformGroup, SolverParticleHandles);
 
 
 
diff --git a/Engine/Source/Runtime/Experimental/ChaosSolvers/Private/PBDRigidsSolver.cpp b/Engine/Source/Runtime/Experimental/ChaosSolvers/Private/PBDRigidsSolver.cpp
--- a/Engine/Source/Runtime/Experimental/ChaosSolvers/Private/PBDRigidsSolver.cpp
+++ b/Engine/Source/Runtime/Experimental/ChaosSolvers/Private/PBDRigidsSolver.cpp
@@ -668,20 +668,11 @@ namespace Chaos
 		auto Cmd = [Proxy, Solver](Chaos::FPersistentPhysicsTask* PhysThread)
 		{
 			auto* Evolution = Solver->GetEvolution();
-			TManagedArray<Chaos::TPBDRigidClusteredParticleHandle<float, 3>*>& Handles = 
-				Proxy->GetSolverParticleHandles();
-			for (auto* Handle : Handles)
-			{
-				if (Handle)
-					Evolution->DirtyParticle(*Handle);
-			}
+			TArray<FGeometryCollectionPhysicsProxy::FClusterHandle*>& Handles = Proxy->GetSolverParticleHandles();
+			for (auto* Handle : Handles) if (Handle) Evolution->DirtyParticle(*Handle);
 			Proxy->PushToPhysicsState(nullptr);
 		};
-
-		if (Dispatcher)
-			Dispatcher->EnqueueCommandImmediate(Cmd);
-		else
-			Cmd(nullptr);
+		if (Dispatcher) Dispatcher->EnqueueCommandImmediate(Cmd);  else Cmd(nullptr);
 		Proxy->ClearAccumulatedData();
 		Solver->RemoveDirtyProxy(Proxy);
 	}
diff --git a/Engine/Source/Runtime/Experimental/ChaosSolvers/Private/PhysicsProxy/GeometryCollectionPhysicsProxy.cpp b/Engine/Source/Runtime/Experimental/ChaosSolvers/Private/PhysicsProxy/GeometryCollectionPhysicsProxy.cpp
--- a/Engine/Source/Runtime/Experimental/ChaosSolvers/Private/PhysicsProxy/GeometryCollectionPhysicsProxy.cpp
+++ b/Engine/Source/Runtime/Experimental/ChaosSolvers/Private/PhysicsProxy/GeometryCollectionPhysicsProxy.cpp
@@ -346,18 +346,19 @@ FGeometryCollectionPhysicsProxy::FGeometryCollectionPhysicsProxy(
 	const Chaos::EMultiBufferMode BufferMode)
 	: Base(InOwner)
 	, Parameters(SimulationParameters)
+	, NumParticles(INDEX_NONE)
+	, BaseParticleIndex(INDEX_NONE)
 	, IsObjectDynamic(false)
 	, IsObjectLoading(true)
 
-	, BaseParticleIndex(INDEX_NONE)
-	, NumParticles(INDEX_NONE)
+#if TODO_REIMPLEMENT_RIGID_CACHING
 	, ProxySimDuration(0.0f)
-
+	, LastSyncCountGT(MAX_uint32)
+#endif
 	, InitFunc(InInitFunc)
 	, CacheSyncFunc(InCacheSyncFunc)
 	, FinalSyncFunc(InFinalSyncFunc)
 
-	, LastSyncCountGT(MAX_uint32)
 	, CollisionParticlesPerObjectFraction(CollisionParticlesPerObjectFractionDefault)
 {
 	// We rely on a guarded buffer.
@@ -372,10 +373,8 @@ FAutoConsoleVariableRef CVarReportHighParticleFraction(TEXT("p.gc.ReportHighPart
 
 void FGeometryCollectionPhysicsProxy::Initialize()
 {
-	// Old proxy init
 	check(IsInGameThread());
 
-
 	//
 	// Game thread initilization. 
 	//
@@ -388,8 +387,11 @@ void FGeometryCollectionPhysicsProxy::Initialize()
 
 	InitializeDynamicCollection(DynamicCollection, *Parameters.RestCollection, Parameters);
 
-	NumParticles = GetTransformGroupSize(); // SimulatableParticles.Count(true);
+	NumParticles = DynamicCollection.NumElements(FGeometryCollection::TransformGroup);
 	BaseParticleIndex = 0; // Are we always zero indexed now?
+	SolverClusterID.Init(nullptr, NumParticles);
+	SolverClusterHandles.Init(nullptr, NumParticles);
+	SolverParticleHandles.Init(nullptr, NumParticles);
 
 	//	
 	//  Give clients the opportunity to update the parameters before the simualtion is setup. 
@@ -497,7 +499,6 @@ void FGeometryCollectionPhysicsProxy::InitializeDynamicCollection(FGeometryDynam
 	DynamicCollection.CopyMatchingAttributesFrom(RestCollection);
 
 	
-	DynamicCollection.SolverClusterID.Fill(nullptr); 
 
 	//
 	// User defined initial velocities need to be populated. 
@@ -637,10 +638,9 @@ void FGeometryCollectionPhysicsProxy::InitializeBodiesPT(
 	Chaos::FPBDRigidsSolver::FParticlesType& Particles)
 {
 	/*
-	@todo break_everything
-
+	using namespace Chaos;
 	const FGeometryCollection* RestCollection = Parameters.RestCollection;
-	FGeometryDynamicCollection* DynamicCollection = Parameters.DynamicCollection;
+	const FGeometryDynamicCollection& DynamicCollection = *GameToPhysInterchange.GetConsumerBuffer();
 
 	if (Parameters.Simulating)
 	{
@@ -650,17 +650,16 @@ void FGeometryCollectionPhysicsProxy::InitializeBodiesPT(
 		const TManagedArray<TSet<int32>>& Children = RestCollection->Children;
 		const TManagedArray<int32>& SimulationType = RestCollection->SimulationType;
 		const TManagedArray<FVector>& Vertex = RestCollection->Vertex;
-		const TManagedArray<int32>& DynamicState = DynamicCollection->DynamicState;
-		const TManagedArray<int32>& CollisionGroup = DynamicCollection->CollisionGroup;
 		const TManagedArray<float>& Mass = RestCollection->GetAttribute<float>("Mass", FTransformCollection::TransformGroup);
 		const TManagedArray<FVector>& InertiaTensor = RestCollection->GetAttribute<FVector>("InertiaTensor", FTransformCollection::TransformGroup);
-
+		const TManagedArray<int32>& DynamicState = DynamicCollection.DynamicState;
+		const TManagedArray<int32>& CollisionGroup = DynamicCollection.CollisionGroup;
+		const TManagedArray<bool>& SimulatableParticles = DynamicCollection.SimulatableParticles;
 		const int32 NumTransforms = GetTransformGroupSize();
 
 		TArray<FTransform> Transform;
-		GeometryCollectionAlgo::GlobalMatrices(
-			DynamicCollection->Transform, DynamicCollection->Parent, Transform);
-		check(DynamicCollection->Transform.Num() == Transform.Num());
+		GeometryCollectionAlgo::GlobalMatrices(DynamicCollection.Transform, DynamicCollection.Parent, Transform);
+		check(DynamicCollection.Transform.Num() == Transform.Num());
 
 		const int NumRigids = 0; // ryan - Since we're doing SOA, we start at zero?
 		BaseParticleIndex = NumRigids;
@@ -708,7 +707,7 @@ void FGeometryCollectionPhysicsProxy::InitializeBodiesPT(
 		// At the point that we start supporting instancing, this assumption will no longer
 		// hold, and those reverse mappints will be INDEX_NONE.
 
-		const int32 NumGeometries = DynamicCollection->NumElements(FGeometryCollection::GeometryGroup);
+		const int32 NumGeometries = DynamicCollection.NumElements(FGeometryCollection::GeometryGroup);
 		ParallelFor(NumGeometries, [&](int32 GeometryIndex)
 		//for (int32 GeometryIndex = 0; GeometryIndex < NumGeometries; ++GeometryIndex)
 		{
@@ -966,7 +965,6 @@ FGeometryCollectionPhysicsProxy::BuildClusters(
 	TManagedArray<FTransform>& MassToLocal = DynamicCollection.MassToLocal;
 	//TManagedArray<TSharedPtr<FCollisionStructureManager::FSimplicial> >& Simplicials = DynamicCollection.Simplicials;
 	TManagedArray<TSharedPtr<Chaos::FImplicitObject, ESPMode::ThreadSafe>>& Implicits = DynamicCollection.Implicits;
-	TManagedArray<Chaos::TPBDRigidClusteredParticleHandle<float, 3>*>& SolverClusterHandles = DynamicCollection.SolverClusterHandles;
 
 	//If we are a root particle use the world transform, otherwise set the relative transform
 	const FTransform CollectionSpaceTransform = GeometryCollectionAlgo::GlobalMatrix(Transform, ParentIndex, CollectionClusterIndex);
@@ -1191,7 +1189,7 @@ void FGeometryCollectionPhysicsProxy::OnRemoveFromSolver(Chaos::FPBDRigidsSolver
 {
 	const FGeometryDynamicCollection& DynamicCollection = *GameToPhysInterchange.GetConsumerBuffer();
 
-	for (const Chaos::TPBDRigidClusteredParticleHandle<float, 3>* Handle : DynamicCollection.SolverClusterHandles)
+	for (const FClusterHandle* Handle : SolverClusterHandles)
 	{
 		RBDSolver->RemoveParticleToProxy(Handle);
 	}
@@ -1231,7 +1229,9 @@ void FGeometryCollectionPhysicsProxy::SyncBeforeDestroy()
 {
 	if(FinalSyncFunc)
 	{
+#if TODO_REIMPLEMENT_RIGID_CACHING
 		FinalSyncFunc(RecordedTracks);
+#endif
 	}
 }
 
@@ -1727,15 +1727,10 @@ void FGeometryCollectionPhysicsProxy::PullFromPhysicsState()
 	}
 }
 
-void FGeometryCollectionPhysicsProxy::UpdateGeometryCollection(
-	const FGeometryCollectionResults& GCResults,
-	FGeometryDynamicCollection* Collection)
+void FGeometryCollectionPhysicsProxy::UpdateGeometryCollection( const FGeometryCollectionResults& GCResults, FGeometryDynamicCollection* Collection)
 {
-	if (!Collection)
-	{
-		Collection = Parameters.DynamicCollection;
-	}
-	check(Collection);
+	check(IsInGameThread());
+	if (!Collection) return;
 
 	TManagedArray<FTransform>& Transform = Collection->Transform;
 	TManagedArray<int32>& Parent = Collection->Parent;
@@ -1752,13 +1747,13 @@ void FGeometryCollectionPhysicsProxy::UpdateGeometryCollection(
 	//		=> GeomToActor = ActorToWorld.Inv() * ParticleToWorld * MassToLocal.Inv()
 
 	const FTransform& ActorToWorld = Parameters.WorldTransform;
-	const int32 TransformSize = GetTransformGroupSize();
+	const int32 TransformSize = Collection->NumElements(FGeometryCollection::TransformGroup);
 	for (int32 TransformGroupIndex = 0; TransformGroupIndex < TransformSize; ++TransformGroupIndex)
 	{
 		if (!Collection->SimulatableParticles[TransformGroupIndex])
 			continue;
 
-		auto* Handle = Collection->SolverParticleHandles[TransformGroupIndex];
+		auto* Handle = SolverParticleHandles[TransformGroupIndex];
 
 		// dynamic state is also updated by the solver during field interaction. 
 		if(!Handle->Sleeping())
diff --git a/Engine/Source/Runtime/Experimental/ChaosSolvers/Public/GeometryCollectionProxyData.h b/Engine/Source/Runtime/Experimental/ChaosSolvers/Public/GeometryCollectionProxyData.h
--- a/Engine/Source/Runtime/Experimental/ChaosSolvers/Public/GeometryCollectionProxyData.h
+++ b/Engine/Source/Runtime/Experimental/ChaosSolvers/Public/GeometryCollectionProxyData.h
@@ -70,8 +70,6 @@ public:
 	static const FName SharedGeometryAttribute;
 	static const FName SharedImplicitsAttribute;
 	static const FName SimplicialsAttribute;
-	static const FName SolverParticleHandlesAttribute;
-	static const FName SolverClusterHandlesAttribute;
 
 	// Transform Group
 	TManagedArray<bool> Active;
@@ -88,9 +86,6 @@ public:
 	TManagedArray<TSharedPtr<Chaos::FImplicitObject, ESPMode::ThreadSafe>> SharedGeometry;
 	//TManagedArray<TUniquePtr<FCollisionStructureManager::FSimplicial>> Simplicials;
 	TManagedArray<bool> SimulatableParticles;
-	TManagedArray<Chaos::TPBDRigidClusteredParticleHandle<float, 3>*> SolverClusterHandles;
-	TManagedArray<Chaos::TPBDRigidParticleHandle<float, 3>*> SolverClusterID; // Rename to ClusterParent?
-	TManagedArray<Chaos::TPBDRigidClusteredParticleHandle<float, 3>*> SolverParticleHandles;
 
 
 };
\ No newline at end of file
diff --git a/Engine/Source/Runtime/Experimental/ChaosSolvers/Public/PhysicsProxy/GeometryCollectionPhysicsProxy.h b/Engine/Source/Runtime/Experimental/ChaosSolvers/Public/PhysicsProxy/GeometryCollectionPhysicsProxy.h
--- a/Engine/Source/Runtime/Experimental/ChaosSolvers/Public/PhysicsProxy/GeometryCollectionPhysicsProxy.h
+++ b/Engine/Source/Runtime/Experimental/ChaosSolvers/Public/PhysicsProxy/GeometryCollectionPhysicsProxy.h
@@ -13,6 +13,7 @@
 #include "Chaos/GeometryParticlesfwd.h"
 #include "Chaos/ParticleHandle.h"
 #include "Chaos/ParticleHandleFwd.h"
+#include "Containers/Array.h"
 #include "PBDRigidsSolver.h"
 
 namespace Chaos
@@ -106,12 +107,11 @@ private:
  */
 class CHAOSSOLVERS_API FGeometryCollectionPhysicsProxy : public TPhysicsProxy<FGeometryCollectionPhysicsProxy, FStubGeometryCollectionData>
 {
-	typedef TPhysicsProxy<FGeometryCollectionPhysicsProxy, FStubGeometryCollectionData> Base;
-
 public:
-
-
+	typedef TPhysicsProxy<FGeometryCollectionPhysicsProxy, FStubGeometryCollectionData> Base;
 	typedef FCollisionStructureManager::FSimplicial FSimplicial;
+	typedef Chaos::TPBDRigidParticleHandle<float, 3> FParticleHandle;
+	typedef Chaos::TPBDRigidParticleHandle<float, 3> FClusterHandle;
 
 	/** Proxy publics */
 	using FInitFunc = TFunction<void(FSimulationParameters&)>;
@@ -203,8 +203,7 @@ public:
 	void SetCollisionParticlesPerObjectFraction(float CollisionParticlesPerObjectFractionIn) 
 	{CollisionParticlesPerObjectFraction = CollisionParticlesPerObjectFractionIn;}
 
-	TManagedArray<Chaos::TPBDRigidClusteredParticleHandle<float, 3>*>& GetSolverParticleHandles() 
-	{ return GameToPhysInterchange.AccessProducerBuffer()->SolverParticleHandles; }
+	TArray<FClusterHandle*>& GetSolverParticleHandles() { return SolverParticleHandles; }
 
 	const FGeometryCollectionResults* GetConsumerResultsGT() const 
 	{ return PhysToGameInterchange.PeekConsumerBuffer(); }
@@ -264,12 +263,6 @@ protected:
 		EFieldResolutionType ResolutionType, 
 		bool bForce);
 
-	/** The size of the transform group, as reported by the \c DynamicCollection 
-	 * currently stored in \c Parameters. 
-	 */
-	int32 GetTransformGroupSize() const 
-	{ return Parameters.DynamicCollection ? Parameters.DynamicCollection->NumElements(FGeometryCollection::TransformGroup) : 0; }
-
 	/** 
 	 * Traverses the parents of \p TransformIndex in \p GeometryCollection, counting
 	 * the number of levels until the next parent is \c INDEX_NONE.
@@ -281,26 +274,32 @@ protected:
 	void ProcessCommands(FParticlesType& Particles, const float Time);
 
 private:
-	FSimulationParameters Parameters;
-
-	// Records current dynamic state
-	bool IsObjectDynamic;
-	// Indicate when loaded
-	bool IsObjectLoading;
-
-	TMap<Chaos::TPBDRigidParticleHandle<float, 3>*, int32> HandleToTransformGroupIndex;
 
+	FSimulationParameters Parameters;
+	TArray<FFieldSystemCommand> Commands;
 
+	//
+	//  Proxy State Information
+	//
+	int32 NumParticles;
+	int32 BaseParticleIndex;
+	TArray<FParticleHandle*> SolverClusterID;
+	TArray<FClusterHandle*> SolverClusterHandles; // make a TArray of the base clase with type
+	TArray<FClusterHandle*> SolverParticleHandles;// make a TArray of base class and join with above. 
+	TMap<FParticleHandle*, int32> HandleToTransformGroupIndex;
 
 
+	//
+	// Buffer Results State Information
+	//
+	bool IsObjectDynamic; // Records current dynamic state
+	bool IsObjectLoading; // Indicate when loaded
 	TArray<int32> EndFrameUnparentingBuffer;
-
 	// This is a subset of the geometry group that are used in the transform hierarchy to represent geometry
 	TArray<FBox> ValidGeometryBoundingBoxes;
 	TArray<int32> ValidGeometryTransformIndices;
 
-	TArray<FFieldSystemCommand> Commands;
-
+#ifdef TODO_REIMPLEMENT_RIGID_CACHING
 	TFunction<void(void)> ResetAnimationCacheCallback;
 	TFunction<void(const TArrayView<FTransform> &)> UpdateTransformsCallback;
 	TFunction<void(const int32 & CurrentFrame, const TManagedArray<int32> & RigidBodyID, const TManagedArray<int32>& Level, const TManagedArray<int32>& Parent, const TManagedArray<TSet<int32>>& Children, const TManagedArray<uint32>& SimulationType, const TManagedArray<uint32>& StatusFlags, const FParticlesType& Particles)> UpdateRestStateCallback;
@@ -308,26 +307,22 @@ private:
 	TFunction<void(FRecordedTransformTrack& InTrack)> CommitRecordedStateCallback;
 
 	// Index of the first particles for this collection in the larger particle array
-	int32 BaseParticleIndex;
-	int32 NumParticles;
-
 	// Time since this object started simulating
 	float ProxySimDuration;
 
-	// bodies waiting to activate on initialization
-	TArray<uint32> PendingActivationList;
+	// Sync frame numbers so we don't do many syncs when physics is running behind
+	uint32 LastSyncCountGT;
 
 	// Storage for the recorded frame information when we're caching the geometry component results.
 	// Synced back to the component with SyncBeforeDestroy
 	FRecordedTransformTrack RecordedTracks;
+#endif
 
 	// Functions to handle engine-side events
 	FInitFunc InitFunc;
 	FCacheSyncFunc CacheSyncFunc;
 	FFinalSyncFunc FinalSyncFunc;
 
-	// Sync frame numbers so we don't do many syncs when physics is running behind
-	uint32 LastSyncCountGT;
 
 	// Per object collision fraction.
 	float CollisionParticlesPerObjectFraction;
