Chaos :

- Enabled particle initialization.

diff --git a/Engine/Source/Runtime/Experimental/ChaosSolvers/Private/GeometryCollectionProxyData.cpp b/Engine/Source/Runtime/Experimental/ChaosSolvers/Private/GeometryCollectionProxyData.cpp
--- a/Engine/Source/Runtime/Experimental/ChaosSolvers/Private/GeometryCollectionProxyData.cpp
+++ b/Engine/Source/Runtime/Experimental/ChaosSolvers/Private/GeometryCollectionProxyData.cpp
@@ -44,6 +44,7 @@ const FName FGeometryDynamicCollection::ShapesQueryDataAttribute("ShapesQueryDat
 const FName FGeometryDynamicCollection::ShapesSimDataAttribute("ShapesSimData");
 const FName FGeometryDynamicCollection::SharedGeometryAttribute("SharedGeometry");
 const FName FGeometryDynamicCollection::SimplicialsAttribute("CollisionParticles");
+const FName FGeometryDynamicCollection::SimulatableParticlesAttribute("SimulatableParticlesAttribute");
 const FName FGeometryDynamicCollection::SharedImplicitsAttribute("SharedImplicits");
 
 FGeometryDynamicCollection::FGeometryDynamicCollection()
@@ -63,8 +64,6 @@ FGeometryDynamicCollection::FGeometryDynamicCollection()
 	//AddExternalAttribute(ShapesSimDataAttribute, FTransformCollection::TransformGroup, ShapeSimData);
 	AddExternalAttribute(SharedGeometryAttribute, FTransformCollection::TransformGroup, SharedGeometry);
 	//AddExternalAttribute(SimplicialsAttribute, FTransformCollection::TransformGroup, Simplicials);
-	AddExternalAttribute("SimulatableParticles", FGeometryCollection::TransformGroup, SimulatableParticles);
-
-
+	AddExternalAttribute(SimulatableParticlesAttribute, FGeometryCollection::TransformGroup, SimulatableParticles);
 
 }
diff --git a/Engine/Source/Runtime/Experimental/ChaosSolvers/Private/PhysicsProxy/GeometryCollectionPhysicsProxy.cpp b/Engine/Source/Runtime/Experimental/ChaosSolvers/Private/PhysicsProxy/GeometryCollectionPhysicsProxy.cpp
--- a/Engine/Source/Runtime/Experimental/ChaosSolvers/Private/PhysicsProxy/GeometryCollectionPhysicsProxy.cpp
+++ b/Engine/Source/Runtime/Experimental/ChaosSolvers/Private/PhysicsProxy/GeometryCollectionPhysicsProxy.cpp
@@ -584,20 +584,7 @@ void FGeometryCollectionPhysicsProxy::InitializeDynamicCollection(FGeometryDynam
 	// Process Activity
 	{
 		const int32 NumTransforms = DynamicCollection.SimulatableParticles.Num();
-		if (RestCollection.HasAttribute(FGeometryCollection::SimulatableParticlesAttribute, FTransformCollection::TransformGroup))
-		{
-			// When the rest collection has been pre configured with simulation data, use that to determine it ability to simulate. 
-			const TManagedArray<bool>& ActiveRestParticles =
-				RestCollection.GetAttribute<bool>(
-					FGeometryCollection::SimulatableParticlesAttribute, FTransformCollection::TransformGroup);
-			check(ActiveRestParticles.Num() == DynamicCollection.Active.Num());
-			check(ActiveRestParticles.Num() == DynamicCollection.SimulatableParticles.Num());
-			for (int32 TransformIdx = 0; TransformIdx < NumTransforms; TransformIdx++)
-			{
-				DynamicCollection.SimulatableParticles[TransformIdx] = DynamicCollection.Active[TransformIdx] && ActiveRestParticles[TransformIdx];
-			}
-		}
-		else
+		if (!RestCollection.HasAttribute(FGeometryCollection::SimulatableParticlesAttribute, FTransformCollection::TransformGroup))
 		{
 			// If no simulation data is available then default to the simulation of just the rigid geometry.
 			for (int32 TransformIdx = 0; TransformIdx < NumTransforms; TransformIdx++)
@@ -637,8 +624,6 @@ void FGeometryCollectionPhysicsProxy::InitializeBodiesPT(
 	Chaos::FPBDRigidsSolver* RigidsSolver,
 	Chaos::FPBDRigidsSolver::FParticlesType& Particles)
 {
-	/*
-	using namespace Chaos;
 	const FGeometryCollection* RestCollection = Parameters.RestCollection;
 	const FGeometryDynamicCollection& DynamicCollection = *GameToPhysInterchange.GetConsumerBuffer();
 
@@ -652,10 +637,15 @@ void FGeometryCollectionPhysicsProxy::InitializeBodiesPT(
 		const TManagedArray<FVector>& Vertex = RestCollection->Vertex;
 		const TManagedArray<float>& Mass = RestCollection->GetAttribute<float>("Mass", FTransformCollection::TransformGroup);
 		const TManagedArray<FVector>& InertiaTensor = RestCollection->GetAttribute<FVector>("InertiaTensor", FTransformCollection::TransformGroup);
+
+		const int32 NumTransforms = DynamicCollection.NumElements(FTransformCollection::TransformGroup);
 		const TManagedArray<int32>& DynamicState = DynamicCollection.DynamicState;
 		const TManagedArray<int32>& CollisionGroup = DynamicCollection.CollisionGroup;
 		const TManagedArray<bool>& SimulatableParticles = DynamicCollection.SimulatableParticles;
-		const int32 NumTransforms = GetTransformGroupSize();
+		const TManagedArray<FTransform>& MassToLocal = DynamicCollection.MassToLocal;
+		const TManagedArray<FVector>& InitialAngularVelocity = DynamicCollection.InitialAngularVelocity;
+		const TManagedArray<FVector>& InitialLinearVelocity = DynamicCollection.InitialLinearVelocity;
+		const TManagedArray<FGeometryDynamicCollection::FSharedImplicit>& Implicits = DynamicCollection.SharedGeometry;
 
 		TArray<FTransform> Transform;
 		GeometryCollectionAlgo::GlobalMatrices(DynamicCollection.Transform, DynamicCollection.Parent, Transform);
@@ -712,7 +702,7 @@ void FGeometryCollectionPhysicsProxy::InitializeBodiesPT(
 		//for (int32 GeometryIndex = 0; GeometryIndex < NumGeometries; ++GeometryIndex)
 		{
 			const int32 TransformGroupIndex = TransformIndex[GeometryIndex];
-			if (Chaos::TPBDRigidClusteredParticleHandle<float, 3>* Handle = SolverParticleHandles[TransformGroupIndex])
+			if (FClusterHandle* Handle = SolverParticleHandles[TransformGroupIndex])
 			{
 				// TODO: We probably want to take better advantage of this solver
 				// feature, rather than doing it all here ourselves?
@@ -727,7 +717,7 @@ void FGeometryCollectionPhysicsProxy::InitializeBodiesPT(
 					Handle,
 					//Particles, 
 					Parameters.Shared, 
-					Simplicials[TransformGroupIndex].Get(),
+					nullptr,///Simplicials[TransformGroupIndex].Get(),
 					Implicits[TransformGroupIndex], 
 #if TODO_REIMPLEMENT_SOLVER_SETTINGS_ACCESSORS
 					RigidsSolver->GetMassScale() * Mass[TransformGroupIndex], 
@@ -748,7 +738,7 @@ void FGeometryCollectionPhysicsProxy::InitializeBodiesPT(
 				}
 
 				TUniquePtr<Chaos::TBVHParticles<float, 3>>& CollisionParticles = Handle->CollisionParticles();
-				CollisionParticles.Reset(Simplicials[TransformGroupIndex].Release()); // Steals!
+				// @todo break everything CollisionParticles.Reset(Simplicials[TransformGroupIndex].Release()); // Steals!
 				//Particles.CollisionParticles(RigidBodyIndex).Reset(Simplicials[TransformGroupIndex].Release());
 				if (CollisionParticles)
 				//if (Particles.CollisionParticles(RigidBodyIndex))
@@ -788,7 +778,7 @@ void FGeometryCollectionPhysicsProxy::InitializeBodiesPT(
 			// rule of thumb for parallelization is that each thread needs at least
 			// 1000 operations in order to overcome the expense of threading.  I don't
 			// think that's generally going to be the case here...
-			for (int32 TransformGroupIndex = 0; TransformGroupIndex < GetTransformGroupSize(); ++TransformGroupIndex)
+			for (int32 TransformGroupIndex = 0; TransformGroupIndex < NumTransforms; ++TransformGroupIndex)
 			{
 				if (Chaos::TPBDRigidParticleHandle<float, 3>* Handle = SolverParticleHandles[TransformGroupIndex])
 				{
@@ -935,7 +925,7 @@ void FGeometryCollectionPhysicsProxy::InitializeBodiesPT(
 #endif // TODO_REIMPLEMENT_RIGID_CACHING
 
 	} // end if simulating...
-	*/
+
 }
 
 int32 ReportNoLevelsetCluster = 0;
@@ -1437,9 +1427,8 @@ void FGeometryCollectionPhysicsProxy::BufferPhysicsResults()
 	using namespace Chaos;
 	SCOPE_CYCLE_COUNTER(STAT_CacheResultGeomCollection);
 
-	/*
-	@todo break everything
 	//FGeometryCollectionResults& TargetResults = Results.GetPhysicsDataForWrite();
+	const FGeometryDynamicCollection& PTDynamicCollection = *GameToPhysInterchange.GetConsumerBuffer();
 	FGeometryCollectionResults& TargetResults = *PhysToGameInterchange.AccessProducerBuffer();
 
 	if (TargetResults.NumTransformGroup() != PTDynamicCollection.NumElements(FGeometryCollection::TransformGroup))
@@ -1472,15 +1461,15 @@ void FGeometryCollectionPhysicsProxy::BufferPhysicsResults()
 	
 	const FTransform& ActorToWorld = Parameters.WorldTransform;
 
-	const int32 TransfromGroupSize = GetTransformGroupSize();
+	const int32 TransfromGroupSize = PTDynamicCollection.NumElements(FGeometryCollection::TransformGroup);
 	if (EndFrameUnparentingBuffer.Num() < TransfromGroupSize)
 	{
 		// It's not worth shrinking the EndFrameUnparentingBuffer array, 
 		// at least until the solver supports deleting bodies.
 		EndFrameUnparentingBuffer.Init(INDEX_NONE, TransfromGroupSize);
 	}
-	TManagedArray<int32>& Parent = PTDynamicCollection.Parent;
-	TManagedArray<TSet<int32>>& Children = PTDynamicCollection.Children;
+	const TManagedArray<int32>& Parent = PTDynamicCollection.Parent;
+	const TManagedArray<TSet<int32>>& Children = PTDynamicCollection.Children;
 
 	{ // tmp scope
 		SCOPE_CYCLE_COUNTER(STAT_CalcParticleToWorld);
@@ -1546,6 +1535,8 @@ void FGeometryCollectionPhysicsProxy::BufferPhysicsResults()
 			//   bodies out of the collections hierarchy during the simulation.  
 			if (!Handle->Disabled())
 			{
+				const TManagedArray<FTransform>& MassToLocal = PTDynamicCollection.MassToLocal;
+
 				// Update the transform of the active body. The active body can be either a single rigid
 				// or a collection of rigidly attached geometries (Clustering). The cluster is represented as a
 				// single transform in the GeometryCollection, and all children are stored in the local space
@@ -1610,6 +1601,8 @@ void FGeometryCollectionPhysicsProxy::BufferPhysicsResults()
 					{
 						if (ClusteredParentHandle->InternalCluster())
 						{
+							const TManagedArray<FTransform>& MassToLocal = PTDynamicCollection.MassToLocal;
+
 							//const FTransform ClusterChildToWorld = ClusterChildToParentMap[RigidBodyIndex] * FTransform(Particles.R(ClusterParentIndex), Particles.X(ClusterParentIndex));
 							const TRigidTransform<float, 3>& ChildToParent = Handle->ChildToParent();
 							const FTransform ClusterChildToWorld = ChildToParent * FTransform(ParentHandle->R(), ParentHandle->X());
@@ -1639,8 +1632,9 @@ void FGeometryCollectionPhysicsProxy::BufferPhysicsResults()
 			// so reset this index to INDEX_NONE before it gets paged out.
 			EndFrameUnparentingBuffer[TransformGroupIndex] = INDEX_NONE;
 
-			Children[ParentIndex].Remove(TransformGroupIndex);
-			Parent[TransformGroupIndex] = INDEX_NONE;
+			// todo : this should be on the output buffer -
+			//Children[ParentIndex].Remove(TransformGroupIndex);
+			//Parent[TransformGroupIndex] = INDEX_NONE;
 		}
 		TargetResults.Parent[TransformGroupIndex] = Parent[TransformGroupIndex];
 	}
@@ -1668,7 +1662,6 @@ void FGeometryCollectionPhysicsProxy::BufferPhysicsResults()
 	// We pulled state out of particles.  Push results into PTDynamicCollection, 
 	// so that it has an updated view.
 	UpdateGeometryCollection(TargetResults);
-	*/
 }
 
 void FGeometryCollectionPhysicsProxy::FlipBuffer()
@@ -2475,12 +2468,13 @@ void BuildSimulationData(Chaos::FErrorReporter& ErrorReporter, FGeometryCollecti
 
 void FGeometryCollectionPhysicsProxy::ProcessCommands(FParticlesType& Particles, const float Time)
 {
-	FGeometryDynamicCollection* Collection = Parameters.DynamicCollection;
-	check(Collection);
+	//const FGeometryDynamicCollection& Collection = *GameToPhysInterchange.AccessProducerBuffer();
+	//check(Collection);
 
 	// Process Particle-Collection commands
 	if (Commands.Num())
 	{
+		/*
 		TArray<ContextIndex> IndicesArray;
 		Chaos::FPhysicsSolver* CurrentSolver = GetSolver();
 
@@ -2584,6 +2578,7 @@ void FGeometryCollectionPhysicsProxy::ProcessCommands(FParticlesType& Particles,
 				Commands.RemoveAt(CommandIndex);
 			}
 		}
+		*/
 	}
 
 	// Process Particle-Particle commands
diff --git a/Engine/Source/Runtime/Experimental/ChaosSolvers/Public/GeometryCollectionProxyData.h b/Engine/Source/Runtime/Experimental/ChaosSolvers/Public/GeometryCollectionProxyData.h
--- a/Engine/Source/Runtime/Experimental/ChaosSolvers/Public/GeometryCollectionProxyData.h
+++ b/Engine/Source/Runtime/Experimental/ChaosSolvers/Public/GeometryCollectionProxyData.h
@@ -59,6 +59,7 @@ public:
 	FGeometryDynamicCollection& operator=(FGeometryDynamicCollection&&) = delete;
 
 	typedef FTransformDynamicCollection Super;
+	typedef TSharedPtr<Chaos::FImplicitObject, ESPMode::ThreadSafe> FSharedImplicit;
 
 	static const FName ActiveAttribute;
 	static const FName CollisionGroupAttribute;
@@ -70,6 +71,7 @@ public:
 	static const FName SharedGeometryAttribute;
 	static const FName SharedImplicitsAttribute;
 	static const FName SimplicialsAttribute;
+	static const FName SimulatableParticlesAttribute;
 
 	// Transform Group
 	TManagedArray<bool> Active;
@@ -83,7 +85,7 @@ public:
 	TManagedArray<FTransform> MassToLocal;
 	//TManagedArray<TArray<FCollisionFilterData>> ShapeQueryData;
 	//TManagedArray<TArray<FCollisionFilterData>> ShapeSimData;
-	TManagedArray<TSharedPtr<Chaos::FImplicitObject, ESPMode::ThreadSafe>> SharedGeometry;
+	TManagedArray<FSharedImplicit> SharedGeometry;
 	//TManagedArray<TUniquePtr<FCollisionStructureManager::FSimplicial>> Simplicials;
 	TManagedArray<bool> SimulatableParticles;
 
diff --git a/Engine/Source/Runtime/Experimental/ChaosSolvers/Public/PhysicsProxy/GeometryCollectionPhysicsProxy.h b/Engine/Source/Runtime/Experimental/ChaosSolvers/Public/PhysicsProxy/GeometryCollectionPhysicsProxy.h
--- a/Engine/Source/Runtime/Experimental/ChaosSolvers/Public/PhysicsProxy/GeometryCollectionPhysicsProxy.h
+++ b/Engine/Source/Runtime/Experimental/ChaosSolvers/Public/PhysicsProxy/GeometryCollectionPhysicsProxy.h
@@ -111,7 +111,7 @@ public:
 	typedef TPhysicsProxy<FGeometryCollectionPhysicsProxy, FStubGeometryCollectionData> Base;
 	typedef FCollisionStructureManager::FSimplicial FSimplicial;
 	typedef Chaos::TPBDRigidParticleHandle<float, 3> FParticleHandle;
-	typedef Chaos::TPBDRigidParticleHandle<float, 3> FClusterHandle;
+	typedef Chaos::TPBDRigidClusteredParticleHandle<float, 3> FClusterHandle;
 
 	/** Proxy publics */
 	using FInitFunc = TFunction<void(FSimulationParameters&)>;
@@ -285,7 +285,7 @@ private:
 	int32 BaseParticleIndex;
 	TArray<FParticleHandle*> SolverClusterID;
 	TArray<FClusterHandle*> SolverClusterHandles; // make a TArray of the base clase with type
-	TArray<FClusterHandle*> SolverParticleHandles;// make a TArray of base class and join with above. 
+	TArray<FClusterHandle*> SolverParticleHandles;// make a TArray of base class and join with above
 	TMap<FParticleHandle*, int32> HandleToTransformGroupIndex;
 
 
