Chaos :
- Just removed the input double buffer, in favor of two dedicated DynamicCollections; one on the PhysicsThread and the other on the GameThread.

diff --git a/Engine/Source/Runtime/Experimental/ChaosSolvers/Private/PhysicsProxy/GeometryCollectionPhysicsProxy.cpp b/Engine/Source/Runtime/Experimental/ChaosSolvers/Private/PhysicsProxy/GeometryCollectionPhysicsProxy.cpp
--- a/Engine/Source/Runtime/Experimental/ChaosSolvers/Private/PhysicsProxy/GeometryCollectionPhysicsProxy.cpp
+++ b/Engine/Source/Runtime/Experimental/ChaosSolvers/Private/PhysicsProxy/GeometryCollectionPhysicsProxy.cpp
@@ -382,7 +382,7 @@ void FGeometryCollectionPhysicsProxy::Initialize()
 	//  2) Populate the buffer with the necessary data.
 	//  3) Deep copy the data to the other buffers. 
 	//
-	FGeometryDynamicCollection& DynamicCollection = *GameToPhysInterchange.AccessProducerBuffer();
+	FGeometryDynamicCollection& DynamicCollection = GameThreadCollection;
 
 
 	InitializeDynamicCollection(DynamicCollection, *Parameters.RestCollection, Parameters);
@@ -485,9 +485,7 @@ void FGeometryCollectionPhysicsProxy::Initialize()
 		//Results.Get(1).WorldBounds = FBoxSphereBounds(BoundingBox);
 	}
 	*/
-	GameToPhysInterchange.FlipProducer();
-	FGeometryDynamicCollection& ProducerCollection = *GameToPhysInterchange.AccessProducerBuffer();
-	ProducerCollection.CopyMatchingAttributesFrom(DynamicCollection);
+	PhysicsThreadCollection.CopyMatchingAttributesFrom(DynamicCollection);
 }
 
 void FGeometryCollectionPhysicsProxy::InitializeDynamicCollection(FGeometryDynamicCollection& DynamicCollection, const FGeometryCollection& RestCollection, const FSimulationParameters& Params)
@@ -625,7 +623,7 @@ void FGeometryCollectionPhysicsProxy::InitializeBodiesPT(
 	Chaos::FPBDRigidsSolver::FParticlesType& Particles)
 {
 	const FGeometryCollection* RestCollection = Parameters.RestCollection;
-	const FGeometryDynamicCollection& DynamicCollection = *GameToPhysInterchange.GetConsumerBuffer();
+	const FGeometryDynamicCollection& DynamicCollection = PhysicsThreadCollection;
 
 	if (Parameters.Simulating)
 	{
@@ -947,7 +945,7 @@ FGeometryCollectionPhysicsProxy::BuildClusters(
 	check(ChildHandles.Num() != 0);
 
 	check(Parameters.DynamicCollection);
-	FGeometryDynamicCollection& DynamicCollection = *GameToPhysInterchange.AccessProducerBuffer();
+	FGeometryDynamicCollection& DynamicCollection = PhysicsThreadCollection;
 	TManagedArray<int32>& DynamicState = DynamicCollection.DynamicState;
 	TManagedArray<int32>& ParentIndex = DynamicCollection.Parent;
 	TManagedArray<TSet<int32>>& Children = DynamicCollection.Children;
@@ -1177,7 +1175,7 @@ void FGeometryCollectionPhysicsProxy::InitializeRemoveOnFracture(FParticlesType&
 
 void FGeometryCollectionPhysicsProxy::OnRemoveFromSolver(Chaos::FPBDRigidsSolver *RBDSolver)
 {
-	const FGeometryDynamicCollection& DynamicCollection = *GameToPhysInterchange.GetConsumerBuffer();
+	const FGeometryDynamicCollection& DynamicCollection = PhysicsThreadCollection;
 
 	for (const FClusterHandle* Handle : SolverClusterHandles)
 	{
@@ -1428,7 +1426,7 @@ void FGeometryCollectionPhysicsProxy::BufferPhysicsResults()
 	SCOPE_CYCLE_COUNTER(STAT_CacheResultGeomCollection);
 
 	//FGeometryCollectionResults& TargetResults = Results.GetPhysicsDataForWrite();
-	const FGeometryDynamicCollection& PTDynamicCollection = *GameToPhysInterchange.GetConsumerBuffer();
+	const FGeometryDynamicCollection& PTDynamicCollection = PhysicsThreadCollection;
 	FGeometryCollectionResults& TargetResults = *PhysToGameInterchange.AccessProducerBuffer();
 
 	if (TargetResults.NumTransformGroup() != PTDynamicCollection.NumElements(FGeometryCollection::TransformGroup))
@@ -1688,7 +1686,7 @@ void FGeometryCollectionPhysicsProxy::PullFromPhysicsState()
 	 * Note: A read lock will have been acquired for this - so the physics thread won't force a buffer flip while this
 	 * sync is ongoing
 	 */
-	FGeometryDynamicCollection& DynamicCollection = *GameToPhysInterchange.AccessProducerBuffer();
+	FGeometryDynamicCollection& DynamicCollection = GameThreadCollection;
 
 	const FGeometryCollectionResults* TargetResultPtr = PhysToGameInterchange.GetConsumerBuffer();
 	if (!TargetResultPtr)
diff --git a/Engine/Source/Runtime/Experimental/ChaosSolvers/Public/PhysicsProxy/GeometryCollectionPhysicsProxy.h b/Engine/Source/Runtime/Experimental/ChaosSolvers/Public/PhysicsProxy/GeometryCollectionPhysicsProxy.h
--- a/Engine/Source/Runtime/Experimental/ChaosSolvers/Public/PhysicsProxy/GeometryCollectionPhysicsProxy.h
+++ b/Engine/Source/Runtime/Experimental/ChaosSolvers/Public/PhysicsProxy/GeometryCollectionPhysicsProxy.h
@@ -327,6 +327,13 @@ private:
 	// Per object collision fraction.
 	float CollisionParticlesPerObjectFraction;
 
+	// The Simulation data is copied between the game and physics thread. It is 
+	// expected that the two data sets will diverge, based on how the simulation
+	// uses the data, but at the start of the simulation the PhysicsThreadCollection
+	// is a deep copy from the GameThreadCollection. 
+	FGeometryDynamicCollection PhysicsThreadCollection;
+	FGeometryDynamicCollection GameThreadCollection;
+
 	// Currently this is using triple buffers for game-physics and 
 	// physics-game thread communication, but not for any reason other than this 
 	// is the only implementation we currently have of a guarded buffer - a buffer 
@@ -337,7 +344,7 @@ private:
 	// logic locks, and the triple buffer would enable a decoupled lock-free 
 	// paradigm, at least for this component of the handshake.
 	Chaos::FGuardedTripleBuffer<FGeometryCollectionResults> PhysToGameInterchange;
-	Chaos::FDoubleBuffer<FGeometryDynamicCollection> GameToPhysInterchange;
+
 };
 
 CHAOSSOLVERS_API void BuildSimulationData(Chaos::FErrorReporter& ErrorReporter, FGeometryCollection& GeometryCollection, const FSharedSimulationParameters& SharedParams);
