Fix CIS.

diff --git a/Engine/Source/Runtime/Experimental/Chaos/Public/Chaos/UniformGrid.h b/Engine/Source/Runtime/Experimental/Chaos/Public/Chaos/UniformGrid.h
--- a/Engine/Source/Runtime/Experimental/Chaos/Public/Chaos/UniformGrid.h
+++ b/Engine/Source/Runtime/Experimental/Chaos/Public/Chaos/UniformGrid.h
@@ -208,4 +208,10 @@ FArchive& operator<<(FArchive& Ar, TUniformGridBase<T, d>& Value)
 	return Ar;
 }
 
+#ifdef __clang__
+extern template class CHAOS_API Chaos::TUniformGridBase<float, 3>;
+extern template class CHAOS_API Chaos::TUniformGrid<float, 3>;
+extern template class CHAOS_API Chaos::TUniformGrid<float, 2>;
+#endif // __clang__
+
 }
