Physics acceleration structures were not updated in the case of single threaded / time sliced platforms
causing various collision issues like being able to push vehicles through terrain.

diff --git a/Engine/Source/Runtime/Experimental/Chaos/Private/Chaos/PBDRigidsEvolution.cpp b/Engine/Source/Runtime/Experimental/Chaos/Private/Chaos/PBDRigidsEvolution.cpp
--- a/Engine/Source/Runtime/Experimental/Chaos/Private/Chaos/PBDRigidsEvolution.cpp
+++ b/Engine/Source/Runtime/Experimental/Chaos/Private/Chaos/PBDRigidsEvolution.cpp
@@ -463,6 +463,10 @@ namespace Chaos
 				}
 				bExternalReady = true;
 			}
+			else
+			{
+				FlushInternalAccelerationQueue();
+			}
 			
 			// we run the task for both starting a new accel structure as well as for the time-slicing
 			AccelerationStructureTaskComplete = TGraphTask<FChaosAccelerationStructureTask>::CreateTask().ConstructAndDispatchWhenReady(*SpatialCollectionFactory, SpatialAccelerationCache, AsyncInternalAcceleration, AsyncExternalAcceleration, ForceFullBuild, bIsSingleThreaded);
