Fix particles not going to sleep when they should.

diff --git a/Engine/Source/Runtime/Experimental/Chaos/Private/Chaos/PBDConstraintGraph.cpp b/Engine/Source/Runtime/Experimental/Chaos/Private/Chaos/PBDConstraintGraph.cpp
--- a/Engine/Source/Runtime/Experimental/Chaos/Private/Chaos/PBDConstraintGraph.cpp
+++ b/Engine/Source/Runtime/Experimental/Chaos/Private/Chaos/PBDConstraintGraph.cpp
@@ -487,6 +487,23 @@ void FPBDConstraintGraph::ComputeIslands(const TParticleView<TPBDRigidParticles<
 			}
 			else
 			{
+				int32 NumRigidsInIsland = 0, NumRigidsInNewIsland = 0;
+				for (TGeometryParticleHandle<FReal, 3>* Particle : IslandToParticles[Island])
+				{
+					if (TPBDRigidParticleHandle<FReal, 3>* PBDRigid = Particle->CastToRigidParticle())
+					{
+						NumRigidsInIsland++;
+					}
+				}
+
+				for (TGeometryParticleHandle<FReal, 3>* Particle : NewIslandParticles[Island])
+				{
+					if (TPBDRigidParticleHandle<FReal, 3>* PBDRigid = Particle->CastToRigidParticle())
+					{
+						NumRigidsInNewIsland++;
+					}
+				}
+
 				for (TGeometryParticleHandle<FReal, 3>* Particle : IslandToParticles[Island])
 				{
 					if (CHAOS_ENSURE(Particle))
@@ -494,9 +511,14 @@ void FPBDConstraintGraph::ComputeIslands(const TParticleView<TPBDRigidParticles<
 						TPBDRigidParticleHandle<FReal, 3>* PBDRigid = Particle->CastToRigidParticle();
 						if (PBDRigid && PBDRigid->ObjectState() != EObjectStateType::Kinematic)
 						{
-							if (!PBDRigid->Disabled())	// todo: why is this needed? [we aren't handling enable/disable state changes properly so disabled particles end up in the graph.]
+							if (!PBDRigid->Disabled()) // todo: why is this needed? [we aren't handling enable/disable state changes properly so disabled particles end up in the graph.]
 							{
-								PBDRigid->SetSleeping(false);
+								//this condition is a hack to ensure we don't immediately wake up particles that were put to sleep due to them getting into a new island. 
+								//we should probably fix island generation code to not move particles to their own islands when sleeping.
+								if (!(NumRigidsInIsland == 1 && NumRigidsInNewIsland == 1 && PBDRigid->Sleeping()))
+								{
+									PBDRigid->SetSleeping(false);
+								}
 								Particles.ActivateParticle(Particle);
 							}
 						}
