Fix windows clang.

diff --git a/Engine/Source/Runtime/Experimental/Chaos/Public/Chaos/UniformGrid.h b/Engine/Source/Runtime/Experimental/Chaos/Public/Chaos/UniformGrid.h
--- a/Engine/Source/Runtime/Experimental/Chaos/Public/Chaos/UniformGrid.h
+++ b/Engine/Source/Runtime/Experimental/Chaos/Public/Chaos/UniformGrid.h
@@ -208,7 +208,7 @@ FArchive& operator<<(FArchive& Ar, TUniformGridBase<T, d>& Value)
 	return Ar;
 }
 
-#ifdef __clang__
+#if PLATFORM_MAC || PLATFORM_LINUX
 extern template class CHAOS_API Chaos::TUniformGridBase<float, 3>;
 extern template class CHAOS_API Chaos::TUniformGrid<float, 3>;
 extern template class CHAOS_API Chaos::TUniformGrid<float, 2>;
