Harden Modeling Mode against deleting objects currently being edited via ComponentTarget.
- don't allow delete when an Accept-style tool is active. This prevents most problems but is very hard-line from UX perspective.
- add FPrimitiveComponentTarget::IsValid(), use in various ComponentTarget functions to avoid trying to access a Component that has been deleted. Still returns null in many places which callers will need to handle.
#rnx

diff --git a/Engine/Plugins/Experimental/ModelingToolsEditorMode/Source/ModelingToolsEditorMode/Private/ModelingToolsEditorMode.cpp b/Engine/Plugins/Experimental/ModelingToolsEditorMode/Source/ModelingToolsEditorMode/Private/ModelingToolsEditorMode.cpp
--- a/Engine/Plugins/Experimental/ModelingToolsEditorMode/Source/ModelingToolsEditorMode/Private/ModelingToolsEditorMode.cpp
+++ b/Engine/Plugins/Experimental/ModelingToolsEditorMode/Source/ModelingToolsEditorMode/Private/ModelingToolsEditorMode.cpp
@@ -91,10 +91,37 @@ void FModelingToolsEditorMode::ActorSelectionChangeNotify()
 
 bool FModelingToolsEditorMode::ProcessEditDelete()
 {
-	return ToolsContext->ProcessEditDelete();
+	if (ToolsContext->ProcessEditDelete())
+	{
+		return true;
+	}
+
+	// for now we disable deleting in an Accept-style tool because it can result in crashes if we are deleting target object
+	if ( GetToolManager()->HasAnyActiveTool() && GetToolManager()->GetActiveTool(EToolSide::Mouse)->HasAccept() )
+	{
+		GetToolManager()->DisplayMessage(
+			LOCTEXT("CannotDeleteWarning", "Cannot delete objects while this Tool is active"), EToolMessageLevel::UserWarning);
+		return true;
+	}
+
+	return false;
 }
 
 
+bool FModelingToolsEditorMode::ProcessEditCut()
+{
+	// for now we disable deleting in an Accept-style tool because it can result in crashes if we are deleting target object
+	if (GetToolManager()->HasAnyActiveTool() && GetToolManager()->GetActiveTool(EToolSide::Mouse)->HasAccept())
+	{
+		GetToolManager()->DisplayMessage(
+			LOCTEXT("CannotCutWarning", "Cannot cut objects while this Tool is active"), EToolMessageLevel::UserWarning);
+		return true;
+	}
+	return false;
+}
+
+
+
 bool FModelingToolsEditorMode::CanAutoSave() const
 {
 	// prevent autosave if any tool is active
diff --git a/Engine/Source/Editor/UnrealEd/Private/Tools/EditorComponentSourceFactory.cpp b/Engine/Source/Editor/UnrealEd/Private/Tools/EditorComponentSourceFactory.cpp
--- a/Engine/Source/Editor/UnrealEd/Private/Tools/EditorComponentSourceFactory.cpp
+++ b/Engine/Source/Editor/UnrealEd/Private/Tools/EditorComponentSourceFactory.cpp
@@ -11,12 +11,17 @@
 
 FMeshDescription* FStaticMeshComponentTarget::GetMesh() 
 {
-	return Cast<UStaticMeshComponent>(Component)->GetStaticMesh()->GetMeshDescription(LODIndex);
+	return IsValid() ? Cast<UStaticMeshComponent>(Component)->GetStaticMesh()->GetMeshDescription(LODIndex) : nullptr;
 }
 
 
 void FStaticMeshComponentTarget::GetMaterialSet(FComponentMaterialSet& MaterialSetOut, bool bAssetMaterials) const
 {
+	if (IsValid() == false)
+	{
+		return;
+	}
+
 	if (bAssetMaterials)
 	{
 		UStaticMesh* StaticMesh = Cast<UStaticMeshComponent>(Component)->GetStaticMesh();
@@ -38,7 +43,8 @@ void FStaticMeshComponentTarget::CommitMaterialSetUpdate(const FComponentMateria
 {
 	// we only support this right now...
 	check(bApplyToAsset == true);
-
+	check(IsValid());
+	
 	UStaticMesh* StaticMesh = Cast<UStaticMeshComponent>(Component)->GetStaticMesh();
 
 	// make sure transactional flag is on
@@ -67,13 +73,22 @@ void FStaticMeshComponentTarget::CommitMaterialSetUpdate(const FComponentMateria
 
 bool FStaticMeshComponentTarget::HasSameSourceData(const FPrimitiveComponentTarget& OtherTarget) const
 {
-	const UStaticMesh* StaticMesh = Cast<UStaticMeshComponent>(Component)->GetStaticMesh();
-	const UStaticMesh* OtherStaticMesh = Cast<UStaticMeshComponent>(OtherTarget.Component)->GetStaticMesh();
-	return StaticMesh && StaticMesh == OtherStaticMesh;
+	if (IsValid())
+	{
+		const UStaticMesh* StaticMesh = Cast<UStaticMeshComponent>(Component)->GetStaticMesh();
+		const UStaticMesh* OtherStaticMesh = Cast<UStaticMeshComponent>(OtherTarget.Component)->GetStaticMesh();
+		return StaticMesh && StaticMesh == OtherStaticMesh;
+	}
+	else
+	{
+		return false;
+	}
 }
 
 void FStaticMeshComponentTarget::CommitMesh( const FCommitter& Committer )
 {
+	check(IsValid());
+
 	//bool bSaved = Component->Modify();
 	//check(bSaved);
 	UStaticMesh* StaticMesh = Cast<UStaticMeshComponent>(Component)->GetStaticMesh();
diff --git a/Engine/Source/Runtime/Experimental/InteractiveToolsFramework/Private/ComponentSourceInterfaces.cpp b/Engine/Source/Runtime/Experimental/InteractiveToolsFramework/Private/ComponentSourceInterfaces.cpp
--- a/Engine/Source/Runtime/Experimental/InteractiveToolsFramework/Private/ComponentSourceInterfaces.cpp
+++ b/Engine/Source/Runtime/Experimental/InteractiveToolsFramework/Private/ComponentSourceInterfaces.cpp
@@ -61,40 +61,51 @@ bool FComponentMaterialSet::operator!=(const FComponentMaterialSet& Other) const
 }
 
 
+bool FPrimitiveComponentTarget::IsValid() const
+{
+	return (Component->IsPendingKillOrUnreachable() == false) && Component->IsValidLowLevel();
+}
+
 AActor* FPrimitiveComponentTarget::GetOwnerActor() const
 {
-	return Component->GetOwner();
+	return IsValid() ? Component->GetOwner() : nullptr;
 }
 
 UPrimitiveComponent* FPrimitiveComponentTarget::GetOwnerComponent() const
 {
-	return Component;
+	return IsValid() ? Component : nullptr;
 }
 
 
 void FPrimitiveComponentTarget::SetOwnerVisibility(bool bVisible) const
 {
-	Component->SetVisibility(bVisible);
+	if (IsValid())
+	{
+		Component->SetVisibility(bVisible);
+	}
 }
 
 
 int32 FPrimitiveComponentTarget::GetNumMaterials() const
 {
-	return Component->GetNumMaterials();
+	return IsValid() ? Component->GetNumMaterials() : 0;
 }
 
 UMaterialInterface* FPrimitiveComponentTarget::GetMaterial(int32 MaterialIndex) const
 {
-	return Component->GetMaterial(MaterialIndex);
+	return IsValid() ? Component->GetMaterial(MaterialIndex) : nullptr;
 }
 
 void FPrimitiveComponentTarget::GetMaterialSet(FComponentMaterialSet& MaterialSetOut, bool bAssetMaterials) const
 {
-	int32 NumMaterials = Component->GetNumMaterials(); 
-	MaterialSetOut.Materials.SetNum(NumMaterials);
-	for (int32 k = 0; k < NumMaterials; ++k)
+	if (IsValid())
 	{
-		MaterialSetOut.Materials[k] = Component->GetMaterial(k);
+		int32 NumMaterials = Component->GetNumMaterials();
+		MaterialSetOut.Materials.SetNum(NumMaterials);
+		for (int32 k = 0; k < NumMaterials; ++k)
+		{
+			MaterialSetOut.Materials[k] = Component->GetMaterial(k);
+		}
 	}
 }
 
@@ -109,14 +120,13 @@ void FPrimitiveComponentTarget::CommitMaterialSetUpdate(const FComponentMaterial
 
 FTransform FPrimitiveComponentTarget::GetWorldTransform() const
 {
-	//return Component->GetOwner()->GetActorTransform();
-	return Component->GetComponentTransform();
+	return IsValid() ? Component->GetComponentTransform() : FTransform::Identity;
 }
 
 bool FPrimitiveComponentTarget::HitTest(const FRay& WorldRay, FHitResult& OutHit) const
 {
 	FVector End = WorldRay.PointAt(HALF_WORLD_MAX);
-	if (Component->LineTraceComponent(OutHit, WorldRay.Origin, End, FCollisionQueryParams(SCENE_QUERY_STAT(HitTest), true)))
+	if (IsValid() && Component->LineTraceComponent(OutHit, WorldRay.Origin, End, FCollisionQueryParams(SCENE_QUERY_STAT(HitTest), true)))
 	{
 		return true;
 	}
diff --git a/Engine/Source/Runtime/Experimental/InteractiveToolsFramework/Public/ComponentSourceInterfaces.h b/Engine/Source/Runtime/Experimental/InteractiveToolsFramework/Public/ComponentSourceInterfaces.h
--- a/Engine/Source/Runtime/Experimental/InteractiveToolsFramework/Public/ComponentSourceInterfaces.h
+++ b/Engine/Source/Runtime/Experimental/InteractiveToolsFramework/Public/ComponentSourceInterfaces.h
@@ -43,6 +43,9 @@ public:
 	 */
 	FPrimitiveComponentTarget( UPrimitiveComponent* Component ): Component( Component ){}
 
+	/** @return true if component target is still valid. May become invalid for various reasons (eg Component was deleted out from under us) */
+	bool IsValid() const;
+
 	/** @return the Actor that owns this Component */
 	AActor* GetOwnerActor() const;
 
