Fix crash.

diff --git a/Engine/Source/Runtime/Experimental/Chaos/Private/Chaos/PBDConstraintGraph.cpp b/Engine/Source/Runtime/Experimental/Chaos/Private/Chaos/PBDConstraintGraph.cpp
--- a/Engine/Source/Runtime/Experimental/Chaos/Private/Chaos/PBDConstraintGraph.cpp
+++ b/Engine/Source/Runtime/Experimental/Chaos/Private/Chaos/PBDConstraintGraph.cpp
@@ -496,7 +496,7 @@ void FPBDConstraintGraph::ComputeIslands(const TParticleView<TPBDRigidParticles<
 					}
 				}
 
-				for (TGeometryParticleHandle<FReal, 3>* Particle : NewIslandParticles[Island])
+				for (TGeometryParticleHandle<FReal, 3>* Particle : NewIslandParticles[OtherIsland])
 				{
 					if (TPBDRigidParticleHandle<FReal, 3>* PBDRigid = Particle->CastToRigidParticle())
 					{
