Guard against null dereferences/invalid array accesses when comparing old and new island rigid counts.

diff --git a/Engine/Source/Runtime/Experimental/Chaos/Private/Chaos/PBDConstraintGraph.cpp b/Engine/Source/Runtime/Experimental/Chaos/Private/Chaos/PBDConstraintGraph.cpp
--- a/Engine/Source/Runtime/Experimental/Chaos/Private/Chaos/PBDConstraintGraph.cpp
+++ b/Engine/Source/Runtime/Experimental/Chaos/Private/Chaos/PBDConstraintGraph.cpp
@@ -488,19 +488,28 @@ void FPBDConstraintGraph::ComputeIslands(const TParticleView<TPBDRigidParticles<
 			else
 			{
 				int32 NumRigidsInIsland = 0, NumRigidsInNewIsland = 0;
-				for (TGeometryParticleHandle<FReal, 3>* Particle : IslandToParticles[Island])
+				if (OtherIsland != INDEX_NONE)
 				{
-					if (TPBDRigidParticleHandle<FReal, 3>* PBDRigid = Particle->CastToRigidParticle())
+					for (TGeometryParticleHandle<FReal, 3>* Particle : IslandToParticles[Island])
 					{
-						NumRigidsInIsland++;
+						if (Particle)
+						{
+							if (TPBDRigidParticleHandle<FReal, 3>* PBDRigid = Particle->CastToRigidParticle())
+							{
+								NumRigidsInIsland++;
+							}
+						}
 					}
-				}
 
-				for (TGeometryParticleHandle<FReal, 3>* Particle : NewIslandParticles[OtherIsland])
-				{
-					if (TPBDRigidParticleHandle<FReal, 3>* PBDRigid = Particle->CastToRigidParticle())
+					for (TGeometryParticleHandle<FReal, 3>* Particle : NewIslandParticles[OtherIsland])
 					{
-						NumRigidsInNewIsland++;
+						if (Particle)
+						{
+							if (TPBDRigidParticleHandle<FReal, 3>* PBDRigid = Particle->CastToRigidParticle())
+							{
+								NumRigidsInNewIsland++;
+							}
+						}
 					}
 				}
 
