Remove the DynamicCollection from the geometry collection's FSimulationParameters.  Make the GameThreadCollection on the FGeometryCollectionPhysicsProxy a reference to a collection passed to the constructor.  So, the game thread dynamic collection is now owned outside of the proxy. Fixes the RigidBodies_SingleFallingUnderGravity test.

diff --git a/Engine/Source/Runtime/Experimental/ChaosSolvers/Private/PBDRigidsSolver.cpp b/Engine/Source/Runtime/Experimental/ChaosSolvers/Private/PBDRigidsSolver.cpp
--- a/Engine/Source/Runtime/Experimental/ChaosSolvers/Private/PBDRigidsSolver.cpp
+++ b/Engine/Source/Runtime/Experimental/ChaosSolvers/Private/PBDRigidsSolver.cpp
@@ -815,9 +815,11 @@ namespace Chaos
 		}
 	}
 
-	// This function is not called.  FPhysScene_ChaosInterface::EndFrame() calls 
+	// This function is not called during normal Engine execution.  
+	// FPhysScene_ChaosInterface::EndFrame() calls 
 	// FPhysScene_ChaosInterface::SyncBodies() instead, and then immediately afterwards 
-	// calls FPBDRigidsSovler::SyncEvents_GameThread().
+	// calls FPBDRigidsSovler::SyncEvents_GameThread().  This function is used by tests,
+	// however.
 	void FPBDRigidsSolver::UpdateGameThreadStructures()
 	{
 		//ensure(IsInGameThread());
diff --git a/Engine/Source/Runtime/Experimental/ChaosSolvers/Private/PhysicsProxy/GeometryCollectionPhysicsProxy.cpp b/Engine/Source/Runtime/Experimental/ChaosSolvers/Private/PhysicsProxy/GeometryCollectionPhysicsProxy.cpp
--- a/Engine/Source/Runtime/Experimental/ChaosSolvers/Private/PhysicsProxy/GeometryCollectionPhysicsProxy.cpp
+++ b/Engine/Source/Runtime/Experimental/ChaosSolvers/Private/PhysicsProxy/GeometryCollectionPhysicsProxy.cpp
@@ -338,10 +338,11 @@ void PopulateSimulatedParticle(
 //==============================================================================
 
 FGeometryCollectionPhysicsProxy::FGeometryCollectionPhysicsProxy(
-	UObject* InOwner, 
-	const FSimulationParameters & SimulationParameters,
-	FInitFunc InInitFunc, 
-	FCacheSyncFunc InCacheSyncFunc, 
+	UObject* InOwner,
+	FGeometryDynamicCollection& GameThreadCollectionIn,
+	const FSimulationParameters& SimulationParameters,
+	FInitFunc InInitFunc,
+	FCacheSyncFunc InCacheSyncFunc,
 	FFinalSyncFunc InFinalSyncFunc,
 	const Chaos::EMultiBufferMode BufferMode)
 	: Base(InOwner)
@@ -360,6 +361,8 @@ FGeometryCollectionPhysicsProxy::FGeometryCollectionPhysicsProxy(
 	, FinalSyncFunc(InFinalSyncFunc)
 
 	, CollisionParticlesPerObjectFraction(CollisionParticlesPerObjectFractionDefault)
+
+	, GameThreadCollection(GameThreadCollectionIn)
 {
 	// We rely on a guarded buffer.
 	check(BufferMode == Chaos::EMultiBufferMode::TripleGuarded);
@@ -396,13 +399,10 @@ void FGeometryCollectionPhysicsProxy::Initialize()
 	//	
 	//  Give clients the opportunity to update the parameters before the simualtion is setup. 
 	//
-	Parameters.DynamicCollection = &DynamicCollection;
 	if (InitFunc)
 	{
 		InitFunc(Parameters);
 	}
-	Parameters.DynamicCollection = nullptr; // safty setup to prevent usage of the Parameters.DynamicCollection. 
-
 
 	//
 	// Collision vertices down sampling validation.  
@@ -944,7 +944,6 @@ FGeometryCollectionPhysicsProxy::BuildClusters(
 	check(CollectionClusterIndex != INDEX_NONE);
 	check(ChildHandles.Num() != 0);
 
-	check(Parameters.DynamicCollection);
 	FGeometryDynamicCollection& DynamicCollection = PhysicsThreadCollection;
 	TManagedArray<int32>& DynamicState = DynamicCollection.DynamicState;
 	TManagedArray<int32>& ParentIndex = DynamicCollection.Parent;
@@ -1659,7 +1658,6 @@ void FGeometryCollectionPhysicsProxy::BufferPhysicsResults()
 
 	// We pulled state out of particles.  Push results into PTDynamicCollection, 
 	// so that it has an updated view.
-	UpdateGeometryCollection(TargetResults);
 }
 
 void FGeometryCollectionPhysicsProxy::FlipBuffer()
@@ -1686,7 +1684,6 @@ void FGeometryCollectionPhysicsProxy::PullFromPhysicsState()
 	 * Note: A read lock will have been acquired for this - so the physics thread won't force a buffer flip while this
 	 * sync is ongoing
 	 */
-	FGeometryDynamicCollection& DynamicCollection = GameThreadCollection;
 
 	const FGeometryCollectionResults* TargetResultPtr = PhysToGameInterchange.GetConsumerBuffer();
 	if (!TargetResultPtr)
@@ -1694,6 +1691,8 @@ void FGeometryCollectionPhysicsProxy::PullFromPhysicsState()
 	// Remove const-ness as we're going to do the ExchangeArrays() thing.
 	FGeometryCollectionResults& TR = *const_cast<FGeometryCollectionResults*>(TargetResultPtr);
 
+	FGeometryDynamicCollection& DynamicCollection = GameThreadCollection;
+
 	// We should never be changing the number of entries, this would break other 
 	// attributes in the transform group.
 	const int32 NumTransforms = DynamicCollection.Transform.Num();
@@ -1721,7 +1720,10 @@ void FGeometryCollectionPhysicsProxy::PullFromPhysicsState()
 void FGeometryCollectionPhysicsProxy::UpdateGeometryCollection( const FGeometryCollectionResults& GCResults, FGeometryDynamicCollection* Collection)
 {
 	check(IsInGameThread());
-	if (!Collection) return;
+	if (!Collection)
+	{
+		Collection = &GameThreadCollection;
+	}
 
 	TManagedArray<FTransform>& Transform = Collection->Transform;
 	TManagedArray<int32>& Parent = Collection->Parent;
@@ -2750,10 +2752,7 @@ void FGeometryCollectionPhysicsProxy::FieldForcesUpdateCallback(Chaos::FPhysicsS
 
 void FGeometryCollectionPhysicsProxy::ParameterUpdateCallback(FParticlesType& Particles, const float Time)
 {
-	FGeometryDynamicCollection* Collection = Parameters.DynamicCollection;
-	check(Collection);
-
-	if (Collection->Transform.Num())
+	if (GameThreadCollection.Transform.Num())
 	{
 		ProcessCommands(Particles, Time);
 	}
diff --git a/Engine/Source/Runtime/Experimental/ChaosSolvers/Public/PhysicsProxy/GeometryCollectionPhysicsProxy.h b/Engine/Source/Runtime/Experimental/ChaosSolvers/Public/PhysicsProxy/GeometryCollectionPhysicsProxy.h
--- a/Engine/Source/Runtime/Experimental/ChaosSolvers/Public/PhysicsProxy/GeometryCollectionPhysicsProxy.h
+++ b/Engine/Source/Runtime/Experimental/ChaosSolvers/Public/PhysicsProxy/GeometryCollectionPhysicsProxy.h
@@ -128,6 +128,7 @@ public:
 	 */
 	FGeometryCollectionPhysicsProxy(
 		UObject* InOwner, 
+		FGeometryDynamicCollection& GameThreadCollection,
 		const FSimulationParameters& SimulationParameters,
 		FInitFunc InInitFunc, 
 		FCacheSyncFunc InCacheSyncFunc, 
@@ -332,7 +333,7 @@ private:
 	// uses the data, but at the start of the simulation the PhysicsThreadCollection
 	// is a deep copy from the GameThreadCollection. 
 	FGeometryDynamicCollection PhysicsThreadCollection;
-	FGeometryDynamicCollection GameThreadCollection;
+	FGeometryDynamicCollection& GameThreadCollection;
 
 	// Currently this is using triple buffers for game-physics and 
 	// physics-game thread communication, but not for any reason other than this 
diff --git a/Engine/Source/Runtime/Experimental/GeometryCollectionEngine/Private/GeometryCollection/GeometryCollectionComponent.cpp b/Engine/Source/Runtime/Experimental/GeometryCollectionEngine/Private/GeometryCollection/GeometryCollectionComponent.cpp
--- a/Engine/Source/Runtime/Experimental/GeometryCollectionEngine/Private/GeometryCollection/GeometryCollectionComponent.cpp
+++ b/Engine/Source/Runtime/Experimental/GeometryCollectionEngine/Private/GeometryCollection/GeometryCollectionComponent.cpp
@@ -1345,8 +1345,7 @@ void UGeometryCollectionComponent::OnCreatePhysicsState()
 			}
 			// end temporary 
 
-			SimulationParameters.DynamicCollection = DynamicCollection.Get();
-			PhysicsProxy = new FGeometryCollectionPhysicsProxy(this, SimulationParameters, InitFunc, CacheSyncFunc, FinalSyncFunc);
+ 			PhysicsProxy = new FGeometryCollectionPhysicsProxy(this, *DynamicCollection, SimulationParameters, InitFunc, CacheSyncFunc, FinalSyncFunc);
 			FPhysScene_Chaos* Scene = GetPhysicsScene();
 			Scene->AddObject(this, PhysicsProxy);
 
diff --git a/Engine/Source/Runtime/Experimental/GeometryCollectionSimulationCore/Public/GeometryCollection/GeometryCollectionSimulationCoreTypes.h b/Engine/Source/Runtime/Experimental/GeometryCollectionSimulationCore/Public/GeometryCollection/GeometryCollectionSimulationCoreTypes.h
--- a/Engine/Source/Runtime/Experimental/GeometryCollectionSimulationCore/Public/GeometryCollection/GeometryCollectionSimulationCoreTypes.h
+++ b/Engine/Source/Runtime/Experimental/GeometryCollectionSimulationCore/Public/GeometryCollection/GeometryCollectionSimulationCoreTypes.h
@@ -201,7 +201,6 @@ struct FSimulationParameters
 	FSimulationParameters()
 		: Name("")
 		, RestCollection(nullptr)
-		, DynamicCollection(nullptr)
 		, RecordedTrack(nullptr)
 		, bOwnsTrack(false)
 		, Simulating(false)
@@ -228,7 +227,6 @@ struct FSimulationParameters
 
 	FSimulationParameters(const FSimulationParameters& Other)
 		: RestCollection(Other.RestCollection)
-		, DynamicCollection(Other.DynamicCollection)
 		, InitializationCommands(Other.InitializationCommands)
 		, RecordedTrack(Other.RecordedTrack)
 		, bOwnsTrack(false)
@@ -270,7 +268,6 @@ struct FSimulationParameters
 
 	FString Name;
 	const FGeometryCollection* RestCollection;
-	FGeometryDynamicCollection* DynamicCollection;
 	TArray<FFieldSystemCommand> InitializationCommands;
 	const FRecordedTransformTrack* RecordedTrack;
 	bool bOwnsTrack;
