Add inertia scale in CollisionContacts to allow for scaling down/skipping rotations for contacts.

diff --git a/Engine/Source/Runtime/Experimental/Chaos/Private/Chaos/PBDCollisionConstraintsContact.cpp b/Engine/Source/Runtime/Experimental/Chaos/Private/Chaos/PBDCollisionConstraintsContact.cpp
--- a/Engine/Source/Runtime/Experimental/Chaos/Private/Chaos/PBDCollisionConstraintsContact.cpp
+++ b/Engine/Source/Runtime/Experimental/Chaos/Private/Chaos/PBDCollisionConstraintsContact.cpp
@@ -82,8 +82,8 @@ namespace Chaos
 			{
 				*IterationParameters.NeedsAnotherIteration = true;
 
-				FMatrix33 WorldSpaceInvI1 = bIsRigidDynamic0 ? Utilities::ComputeWorldSpaceInertia(Q0, PBDRigid0->InvI()) : FMatrix33(0);
-				FMatrix33 WorldSpaceInvI2 = bIsRigidDynamic1 ? Utilities::ComputeWorldSpaceInertia(Q1, PBDRigid1->InvI()) : FMatrix33(0);
+				FMatrix33 WorldSpaceInvI1 = bIsRigidDynamic0 ? Utilities::ComputeWorldSpaceInertia(Q0, PBDRigid0->InvI()) * Contact.InvInertiaScale0 : FMatrix33(0);
+				FMatrix33 WorldSpaceInvI2 = bIsRigidDynamic1 ? Utilities::ComputeWorldSpaceInertia(Q1, PBDRigid1->InvI()) * Contact.InvInertiaScale1 : FMatrix33(0);
 				FMatrix33 Factor =
 					(bIsRigidDynamic0 ? ComputeFactorMatrix3(VectorToPoint1, WorldSpaceInvI1, PBDRigid0->InvM()) : FMatrix33(0)) +
 					(bIsRigidDynamic1 ? ComputeFactorMatrix3(VectorToPoint2, WorldSpaceInvI2, PBDRigid1->InvM()) : FMatrix33(0));
@@ -252,8 +252,8 @@ namespace Chaos
 			
 				FReal InvM0 = bIsRigidDynamic0 ? PBDRigid0->InvM() : 0.0f;
 				FReal InvM1 = bIsRigidDynamic1 ? PBDRigid1->InvM() : 0.0f;
-				FMatrix33 InvI0 = bIsRigidDynamic0 ? Utilities::ComputeWorldSpaceInertia(Q0, PBDRigid0->InvI()) : FMatrix33(0);
-				FMatrix33 InvI1 = bIsRigidDynamic1 ? Utilities::ComputeWorldSpaceInertia(Q1, PBDRigid1->InvI()) : FMatrix33(0);
+				FMatrix33 InvI0 = bIsRigidDynamic0 ? Utilities::ComputeWorldSpaceInertia(Q0, PBDRigid0->InvI()) * Contact.InvInertiaScale0 : FMatrix33(0);
+				FMatrix33 InvI1 = bIsRigidDynamic1 ? Utilities::ComputeWorldSpaceInertia(Q1, PBDRigid1->InvI()) * Contact.InvInertiaScale1 : FMatrix33(0);
 				FMatrix33 ContactInvI =
 					(bIsRigidDynamic0 ? ComputeFactorMatrix3(VectorToPoint0, InvI0, InvM0) : FMatrix33(0)) +
 					(bIsRigidDynamic1 ? ComputeFactorMatrix3(VectorToPoint1, InvI1, InvM1) : FMatrix33(0));
@@ -472,8 +472,8 @@ namespace Chaos
 				return AccumulatedImpulse;
 			}
 
-			FMatrix33 WorldSpaceInvI1 = bIsRigidDynamic0 ? Utilities::ComputeWorldSpaceInertia(Q0, PBDRigid0->InvI()) : FMatrix33(0);
-			FMatrix33 WorldSpaceInvI2 = bIsRigidDynamic1 ? Utilities::ComputeWorldSpaceInertia(Q1, PBDRigid1->InvI()) : FMatrix33(0);
+			FMatrix33 WorldSpaceInvI1 = bIsRigidDynamic0 ? Utilities::ComputeWorldSpaceInertia(Q0, PBDRigid0->InvI()) * Contact.InvInertiaScale0 : FMatrix33(0);
+			FMatrix33 WorldSpaceInvI2 = bIsRigidDynamic1 ? Utilities::ComputeWorldSpaceInertia(Q1, PBDRigid1->InvI()) * Contact.InvInertiaScale1 : FMatrix33(0);
 			FVec3 VectorToPoint1 = Contact.Location - P0;
 			FVec3 VectorToPoint2 = Contact.Location - P1;
 			FMatrix33 Factor =
diff --git a/Engine/Source/Runtime/Experimental/Chaos/Public/Chaos/CollisionResolutionTypes.h b/Engine/Source/Runtime/Experimental/Chaos/Public/Chaos/CollisionResolutionTypes.h
--- a/Engine/Source/Runtime/Experimental/Chaos/Public/Chaos/CollisionResolutionTypes.h
+++ b/Engine/Source/Runtime/Experimental/Chaos/Public/Chaos/CollisionResolutionTypes.h
@@ -74,7 +74,7 @@ namespace Chaos
 	{
 	public:
 		TCollisionContact(const FImplicitObject* InImplicit0 = nullptr, const FImplicitObject* InImplicit1 = nullptr)
-			: bDisabled(true), Normal(0), Location(0), Phi(FLT_MAX), Friction(0), AngularFriction(0), Restitution(0), ShapesType(EContactShapesType::Unknown)
+		    : bDisabled(true), Normal(0), Location(0), Phi(FLT_MAX), Friction(0), AngularFriction(0), Restitution(0), InvInertiaScale0(1.f), InvInertiaScale1(1.f), ShapesType(EContactShapesType::Unknown)
 		{
 			Implicit[0] = InImplicit0;
 			Implicit[1] = InImplicit1;
@@ -88,7 +88,8 @@ namespace Chaos
 		T Friction;
 		T AngularFriction;
 		T Restitution;
-
+		T InvInertiaScale0;
+		T InvInertiaScale1;
 		EContactShapesType ShapesType;
 
 
@@ -172,6 +173,12 @@ namespace Chaos
 		void SetLocation(const TVector<T, d> & InLocation) { Manifold.Location = InLocation; }
 		TVector<T, d> GetLocation() const { return Manifold.Location; }
 
+		void SetInvInertiaScale0(const T InInvInertiaScale) { Manifold.InvInertiaScale0 = InInvInertiaScale; }
+		T GetInvInertiaScale0() const { return Manifold.InvInertiaScale0; }
+
+		void SetInvInertiaScale1(const T InInvInertiaScale) { Manifold.InvInertiaScale1 = InInvInertiaScale; }
+		T GetInvInertiaScale1() const { return Manifold.InvInertiaScale1; }
+
 
 		FString ToString() const
 		{
