Fix CIS.

diff --git a/Engine/Source/Runtime/Experimental/Chaos/Private/Chaos/CollisionResolution.cpp b/Engine/Source/Runtime/Experimental/Chaos/Private/Chaos/CollisionResolution.cpp
--- a/Engine/Source/Runtime/Experimental/Chaos/Private/Chaos/CollisionResolution.cpp
+++ b/Engine/Source/Runtime/Experimental/Chaos/Private/Chaos/CollisionResolution.cpp
@@ -2551,8 +2551,8 @@ namespace Chaos
 			Dir /= Length;
 
 			// Attempt to get rid of invalid static analysis warnings
-			check(Implicit0 || Implicit0Type == ImplicitObjectType::Unknown);
-			check(Implicit1 || Implicit1Type == ImplicitObjectType::Unknown);
+			check(Implicit0 || (!Implicit0 && Implicit0Type == ImplicitObjectType::Unknown));
+			check(Implicit1 || (!Implicit1 && Implicit1Type == ImplicitObjectType::Unknown));
 
 			if (Implicit0Type == TBox<FReal, 3>::StaticType() && Implicit1Type == THeightField<FReal>::StaticType())
 			{
diff --git a/Engine/Source/Runtime/Experimental/Chaos/Private/Chaos/PBDRigidClustering.cpp b/Engine/Source/Runtime/Experimental/Chaos/Private/Chaos/PBDRigidClustering.cpp
--- a/Engine/Source/Runtime/Experimental/Chaos/Private/Chaos/PBDRigidClustering.cpp
+++ b/Engine/Source/Runtime/Experimental/Chaos/Private/Chaos/PBDRigidClustering.cpp
@@ -342,7 +342,7 @@ namespace Chaos
 				Child = OriginalChild;
 				ChildPosition = Child->X();
 			}
-			else if (ProxyData->KeyChild == OriginalChild)
+			else if (ProxyData && ProxyData->KeyChild == OriginalChild)
 			{
 				Child = MultiChildProxyId.Id;
 				const TRigidTransform<T, d> ProxyWorldTM = 
@@ -2011,7 +2011,7 @@ namespace Chaos
 		TPBDRigidClusteredParticleHandle<T, d>* ClusteredChild1,
 		TPBDRigidClusteredParticleHandle<T, d>* ClusteredChild2)
 	{
-		check(ClusteredChild1 && ClusteredChild1);
+		check(ClusteredChild1 && ClusteredChild2);
 		if (ClusteredChild1 == ClusteredChild2)
 			return;
 		const T AvgStrain = (ClusteredChild1->Strains() + ClusteredChild2->Strains()) * (T)0.5;
diff --git a/Engine/Source/Runtime/Experimental/GeometryCollectionEngine/Private/GeometryCollection/GeometryCollectionComponent.cpp b/Engine/Source/Runtime/Experimental/GeometryCollectionEngine/Private/GeometryCollection/GeometryCollectionComponent.cpp
--- a/Engine/Source/Runtime/Experimental/GeometryCollectionEngine/Private/GeometryCollection/GeometryCollectionComponent.cpp
+++ b/Engine/Source/Runtime/Experimental/GeometryCollectionEngine/Private/GeometryCollection/GeometryCollectionComponent.cpp
@@ -1133,8 +1133,11 @@ void UGeometryCollectionComponent::OnCreatePhysicsState()
 #if !(UE_BUILD_SHIPPING || UE_BUILD_TEST)
 				SimulationParameters.Name = GetPathName();
 #endif
-				RestCollection->GetSharedSimulationParams(SimulationParameters.Shared);
-				SimulationParameters.RestCollection = RestCollection->GetGeometryCollection().Get();
+				if (RestCollection)
+				{
+					RestCollection->GetSharedSimulationParams(SimulationParameters.Shared);
+					SimulationParameters.RestCollection = RestCollection->GetGeometryCollection().Get();
+				}
 				SimulationParameters.Simulating = Simulating;
 				SimulationParameters.EnableClustering = EnableClustering;
 				SimulationParameters.ClusterGroupIndex = EnableClustering ? ClusterGroupIndex : 0;
