Consolidate GeometryCollection storage of implicit shapes into a single ManagedArray of TSharedPtr<FImplicitObject> - non thread safe variant.  Removes SharedGeometryAttribute and SharedImplicitsAttribute.  This likely means that legacy geometry collections will need to be re-generated after this poin.

Copy implicit shapes so the physics thread has it's own copy of the data, for the time being for clarity.  Add a DeepCopy() function to TImplicitObject, where the base implementation simply forwards to Copy(), which is implemented for all implicit shapes except TLevelSet.  Implement TLevelSet::DeepCopy(), leaving TLevelSet::Copy() unimplemented, which asserts in the base implementation.

Add a cvar in GeometryCollectionPhysicsProxy that overrides collision filtering, and enables collisions for geometry collections against everything.  This cvar IS CURRENTLY ENABLED BY DEFAULT so that tests pass while proper support for collision filters is implemented for geometry collections.

These changes get all currently active geometry collection simulation tests passing, and geometry collections simulating in the editor again, although, not colliding.

(All of this was talked over with Brice.)

diff --git a/Engine/Source/Runtime/Experimental/Chaos/Private/Chaos/Levelset.cpp b/Engine/Source/Runtime/Experimental/Chaos/Private/Chaos/Levelset.cpp
--- a/Engine/Source/Runtime/Experimental/Chaos/Private/Chaos/Levelset.cpp
+++ b/Engine/Source/Runtime/Experimental/Chaos/Private/Chaos/Levelset.cpp
@@ -173,6 +173,19 @@ TLevelSet<T, d>::~TLevelSet()
 {
 }
 
+template<typename T, int d>
+TUniquePtr<FImplicitObject> TLevelSet<T, d>::DeepCopy() const
+{
+	TLevelSet<T, d>* Copy = new TLevelSet<T, d>();
+	Copy->MGrid = MGrid;
+	Copy->MPhi.Copy(MPhi);
+	Copy->MNormals.Copy(MNormals);
+	Copy->MLocalBoundingBox = MLocalBoundingBox;
+	Copy->MOriginalLocalBoundingBox = MOriginalLocalBoundingBox;
+	Copy->MBandWidth = MBandWidth;
+	return TUniquePtr<FImplicitObject>(Copy);
+}
+
 template<typename T, int d>
 bool TLevelSet<T, d>::ComputeMassProperties(T& OutVolume, TVector<T, d>& OutCOM, PMatrix<T, d, d>& OutInertia, TRotation<T, d>& OutRotationOfMass) const
 {
diff --git a/Engine/Source/Runtime/Experimental/Chaos/Private/Chaos/PBDRigidClustering.cpp b/Engine/Source/Runtime/Experimental/Chaos/Private/Chaos/PBDRigidClustering.cpp
--- a/Engine/Source/Runtime/Experimental/Chaos/Private/Chaos/PBDRigidClustering.cpp
+++ b/Engine/Source/Runtime/Experimental/Chaos/Private/Chaos/PBDRigidClustering.cpp
@@ -426,7 +426,7 @@ namespace Chaos
 		const int32 ClusterGroupIndex,
 		TArray<Chaos::TPBDRigidParticleHandle<T,d>*>&& Children,
 		const FClusterCreationParameters<T>& Parameters,
-		TSharedPtr<Chaos::FImplicitObject, ESPMode::ThreadSafe> ProxyGeometry,
+		TSharedPtr<Chaos::FImplicitObject> ProxyGeometry,
 		const TRigidTransform<T, d>* ForceMassOrientation)
 	{
 		SCOPE_CYCLE_COUNTER(STAT_CreateClusterParticle);
@@ -1348,7 +1348,7 @@ namespace Chaos
 	void TPBDRigidClustering<FPBDRigidsEvolution, FPBDCollisionConstraint, T, d>::UpdateGeometry(
 		Chaos::TPBDRigidClusteredParticleHandle<float, 3>* Parent, 
 		const TSet<TPBDRigidParticleHandle<T, d>*>& Children, 
-		TSharedPtr<Chaos::FImplicitObject, ESPMode::ThreadSafe> ProxyGeometry,
+		TSharedPtr<Chaos::FImplicitObject> ProxyGeometry,
 		const FClusterCreationParameters<T>& Parameters)
 	{
 		SCOPE_CYCLE_COUNTER(STAT_UpdateGeometry);
@@ -1462,7 +1462,7 @@ namespace Chaos
 		{
 			//ensureMsgf(false, TEXT("Checking usage with proxy"));
 			//@coverage {production}
-			Parent->SetSharedGeometry(ProxyGeometry);
+			Parent->SetSharedGeometry(TSharedPtr<FImplicitObject, ESPMode::ThreadSafe>(ProxyGeometry->DeepCopy().Release()));
 		}
 		else if (Objects.Num() == 0)
 		{
diff --git a/Engine/Source/Runtime/Experimental/Chaos/Public/Chaos/ImplicitObject.h b/Engine/Source/Runtime/Experimental/Chaos/Public/Chaos/ImplicitObject.h
--- a/Engine/Source/Runtime/Experimental/Chaos/Public/Chaos/ImplicitObject.h
+++ b/Engine/Source/Runtime/Experimental/Chaos/Public/Chaos/ImplicitObject.h
@@ -165,6 +165,7 @@ public:
 	virtual bool IsValidGeometry() const;
 
 	virtual TUniquePtr<FImplicitObject> Copy() const;
+	virtual TUniquePtr<FImplicitObject> DeepCopy() const { return Copy(); }
 
 	//This is strictly used for optimization purposes
 	bool IsUnderlyingUnion() const;
diff --git a/Engine/Source/Runtime/Experimental/Chaos/Public/Chaos/Levelset.h b/Engine/Source/Runtime/Experimental/Chaos/Public/Chaos/Levelset.h
--- a/Engine/Source/Runtime/Experimental/Chaos/Public/Chaos/Levelset.h
+++ b/Engine/Source/Runtime/Experimental/Chaos/Public/Chaos/Levelset.h
@@ -31,6 +31,8 @@ class CHAOS_API TLevelSet final : public FImplicitObject
 	TLevelSet(TLevelSet<T, d>&& Other);
 	virtual ~TLevelSet();
 
+	virtual TUniquePtr<FImplicitObject> DeepCopy() const;
+
 	void Write(std::ostream& Stream) const;
 	virtual T PhiWithNormal(const TVector<T, d>& x, TVector<T, d>& Normal) const override;
 	T SignedDistance(const TVector<T, d>& x) const;
diff --git a/Engine/Source/Runtime/Experimental/Chaos/Public/Chaos/PBDRigidClustering.h b/Engine/Source/Runtime/Experimental/Chaos/Public/Chaos/PBDRigidClustering.h
--- a/Engine/Source/Runtime/Experimental/Chaos/Public/Chaos/PBDRigidClustering.h
+++ b/Engine/Source/Runtime/Experimental/Chaos/Public/Chaos/PBDRigidClustering.h
@@ -62,7 +62,7 @@ public:
 		const int32 ClusterGroupIndex, 
 		TArray<Chaos::TPBDRigidParticleHandle<T,d>*>&& Children, 
 		const FClusterCreationParameters<T>& Parameters/* = FClusterCreationParameters<T>()*/,
-		TSharedPtr<Chaos::FImplicitObject, ESPMode::ThreadSafe> ProxyGeometry = nullptr,
+		TSharedPtr<Chaos::FImplicitObject> ProxyGeometry = nullptr,
 		const TRigidTransform<T, d>* ForceMassOrientation = nullptr);
 
 	/**
@@ -285,7 +285,7 @@ public:
 	void UpdateGeometry(
 		Chaos::TPBDRigidClusteredParticleHandle<float, 3>* Parent, 
 		const TSet<TPBDRigidParticleHandle<T, d>*>& Children, 
-		TSharedPtr<Chaos::FImplicitObject, ESPMode::ThreadSafe> ProxyGeometry,
+		TSharedPtr<Chaos::FImplicitObject> ProxyGeometry,
 		const FClusterCreationParameters<T>& Parameters);
 
 	void ComputeStrainFromCollision(const FPBDCollisionConstraint& CollisionRule);
diff --git a/Engine/Source/Runtime/Experimental/ChaosCore/Public/Chaos/ArrayND.h b/Engine/Source/Runtime/Experimental/ChaosCore/Public/Chaos/ArrayND.h
--- a/Engine/Source/Runtime/Experimental/ChaosCore/Public/Chaos/ArrayND.h
+++ b/Engine/Source/Runtime/Experimental/ChaosCore/Public/Chaos/ArrayND.h
@@ -79,6 +79,7 @@ class TArrayNDBase
 		return *this;
 	}
 	FORCEINLINE T_DERIVED Copy() const { return T_DERIVED(MCounts, MArray); }
+	FORCEINLINE void Copy(const TArrayNDBase<T_DERIVED, T, d>& Source) { MCounts = Source.MCounts; MArray = Source.MArray; }
 	FORCEINLINE void Fill(const T& Value)
 	{
 		for (auto& Elem : MArray)
@@ -90,6 +91,7 @@ class TArrayNDBase
 	FORCEINLINE T& operator[](const int32 i) { return MArray[i]; }
 
 	FORCEINLINE int32 Num() const { return MArray.Num(); }
+	FORCEINLINE TVector<int32, d> Counts() const { return MCounts; }
 
   protected:
 	TVector<int32, d> MCounts;
diff --git a/Engine/Source/Runtime/Experimental/ChaosSolvers/Private/GeometryCollectionProxyData.cpp b/Engine/Source/Runtime/Experimental/ChaosSolvers/Private/GeometryCollectionProxyData.cpp
--- a/Engine/Source/Runtime/Experimental/ChaosSolvers/Private/GeometryCollectionProxyData.cpp
+++ b/Engine/Source/Runtime/Experimental/ChaosSolvers/Private/GeometryCollectionProxyData.cpp
@@ -42,10 +42,8 @@ const FName FGeometryDynamicCollection::DynamicStateAttribute("DynamicState");
 const FName FGeometryDynamicCollection::ImplicitsAttribute("Implicits");
 const FName FGeometryDynamicCollection::ShapesQueryDataAttribute("ShapesQueryData");
 const FName FGeometryDynamicCollection::ShapesSimDataAttribute("ShapesSimData");
-const FName FGeometryDynamicCollection::SharedGeometryAttribute("SharedGeometry");
 const FName FGeometryDynamicCollection::SimplicialsAttribute("CollisionParticles");
 const FName FGeometryDynamicCollection::SimulatableParticlesAttribute("SimulatableParticlesAttribute");
-const FName FGeometryDynamicCollection::SharedImplicitsAttribute("SharedImplicits");
 
 FGeometryDynamicCollection::FGeometryDynamicCollection()
 	: FTransformDynamicCollection()
@@ -62,7 +60,6 @@ FGeometryDynamicCollection::FGeometryDynamicCollection()
 	AddExternalAttribute("MassToLocal", FTransformCollection::TransformGroup, MassToLocal);
 	//AddExternalAttribute(ShapesQueryDataAttribute, FTransformCollection::TransformGroup, ShapeQueryData);
 	//AddExternalAttribute(ShapesSimDataAttribute, FTransformCollection::TransformGroup, ShapeSimData);
-	AddExternalAttribute(SharedGeometryAttribute, FTransformCollection::TransformGroup, SharedGeometry);
 	//AddExternalAttribute(SimplicialsAttribute, FTransformCollection::TransformGroup, Simplicials);
 	AddExternalAttribute(SimulatableParticlesAttribute, FGeometryCollection::TransformGroup, SimulatableParticles);
 
diff --git a/Engine/Source/Runtime/Experimental/ChaosSolvers/Public/GeometryCollectionProxyData.h b/Engine/Source/Runtime/Experimental/ChaosSolvers/Public/GeometryCollectionProxyData.h
--- a/Engine/Source/Runtime/Experimental/ChaosSolvers/Public/GeometryCollectionProxyData.h
+++ b/Engine/Source/Runtime/Experimental/ChaosSolvers/Public/GeometryCollectionProxyData.h
@@ -59,7 +59,7 @@ public:
 	FGeometryDynamicCollection& operator=(FGeometryDynamicCollection&&) = delete;
 
 	typedef FTransformDynamicCollection Super;
-	typedef TSharedPtr<Chaos::FImplicitObject, ESPMode::ThreadSafe> FSharedImplicit;
+	typedef TSharedPtr<Chaos::FImplicitObject> FSharedImplicit;
 
 	static const FName ActiveAttribute;
 	static const FName CollisionGroupAttribute;
@@ -68,8 +68,6 @@ public:
 	static const FName ImplicitsAttribute;
 	static const FName ShapesQueryDataAttribute;
 	static const FName ShapesSimDataAttribute;
-	static const FName SharedGeometryAttribute;
-	static const FName SharedImplicitsAttribute;
 	static const FName SimplicialsAttribute;
 	static const FName SimulatableParticlesAttribute;
 
@@ -79,13 +77,12 @@ public:
 	TManagedArray<int32> CollisionMask;
 	TManagedArray<int32> CollisionStructureID;
 	TManagedArray<int32> DynamicState;
-	TManagedArray<TSharedPtr<Chaos::FImplicitObject, ESPMode::ThreadSafe>> Implicits;
+	TManagedArray<TSharedPtr<Chaos::FImplicitObject>> Implicits;
 	TManagedArray<FVector> InitialAngularVelocity;
 	TManagedArray<FVector> InitialLinearVelocity;
 	TManagedArray<FTransform> MassToLocal;
 	//TManagedArray<TArray<FCollisionFilterData>> ShapeQueryData;
 	//TManagedArray<TArray<FCollisionFilterData>> ShapeSimData;
-	TManagedArray<FSharedImplicit> SharedGeometry;
 	//TManagedArray<TUniquePtr<FCollisionStructureManager::FSimplicial>> Simplicials;
 	TManagedArray<bool> SimulatableParticles;
 
