A perforce snafu left this file edited, but not checked out, and it got left off of change 12047778.

diff --git a/Engine/Source/Runtime/Experimental/ChaosSolvers/Private/PhysicsProxy/GeometryCollectionPhysicsProxy.cpp b/Engine/Source/Runtime/Experimental/ChaosSolvers/Private/PhysicsProxy/GeometryCollectionPhysicsProxy.cpp
--- a/Engine/Source/Runtime/Experimental/ChaosSolvers/Private/PhysicsProxy/GeometryCollectionPhysicsProxy.cpp
+++ b/Engine/Source/Runtime/Experimental/ChaosSolvers/Private/PhysicsProxy/GeometryCollectionPhysicsProxy.cpp
@@ -49,6 +49,12 @@ FAutoConsoleVariableRef CVarGeometryCollectionDisableGravity(
 	DisableGeometryCollectionGravity,
 	TEXT("Disable gravity for geometry collections"));
 
+//bool GeometryCollectionCollideAll = false;
+bool GeometryCollectionCollideAll = true;
+FAutoConsoleVariableRef CVarGeometryCollectionCollideAll(
+	TEXT("p.GeometryCollectionCollideAll"),
+	GeometryCollectionCollideAll,
+	TEXT("Bypass the collision matrix and make geometry collections collide against everything"));
 
 DEFINE_LOG_CATEGORY_STATIC(UGCC_LOG, Error, All);
 
@@ -220,7 +226,7 @@ void PopulateSimulatedParticle(
 	Chaos::TPBDRigidParticleHandle<float,3>* Handle,
 	const FSharedSimulationParameters& SharedParams,
 	const FCollisionStructureManager::FSimplicial* Simplicial,
-	TSharedPtr<Chaos::FImplicitObject, ESPMode::ThreadSafe> Implicit,
+	TSharedPtr<Chaos::FImplicitObject> Implicit,
 	const float MassIn,
 	const FVector& InertiaTensorVec,
 	const FTransform& WorldTransform, 
@@ -275,7 +281,8 @@ void PopulateSimulatedParticle(
 
 	if (Implicit)	//todo(ocohen): this is only needed for cases where clusters have no proxy. Kind of gross though, should refactor
 	{
-		Handle->SetSharedGeometry(Implicit);
+		TSharedPtr<Chaos::FImplicitObject, ESPMode::ThreadSafe> SharedImplicitTS(Implicit->DeepCopy().Release());
+		Handle->SetSharedGeometry(SharedImplicitTS);
 	}
 
 	Handle->CollisionParticlesInitIfNeeded();
@@ -309,6 +316,23 @@ void PopulateSimulatedParticle(
 		CollisionParticles->UpdateAccelerationStructures();
 	}
 
+	if (GeometryCollectionCollideAll) // cvar
+	{
+		// Override collision filters and make this body collide with everything.
+		int32 CurrShape = 0;
+		FCollisionFilterData FilterData;
+		FilterData.Word1 = 0xFFFF; // this body channel
+		FilterData.Word3 = 0xFFFF; // collision candidate channels
+		for (const TUniquePtr<Chaos::TPerShapeData<Chaos::FReal, 3>>& Shape : Handle->ShapesArray())
+		{
+			// TODO Ryan - hardwired for now:
+			Shape->bDisable = false;
+			Shape->CollisionTraceType = Chaos::EChaosCollisionTraceFlag::Chaos_CTF_UseSimpleAndComplex; // Chaos_CTF_UseDefault?
+			Shape->SimData = FilterData;
+			Shape->QueryData = FCollisionFilterData();
+		}
+	}
+
 	//
 	//  Manage Object State
 	//
@@ -495,8 +519,8 @@ void FGeometryCollectionPhysicsProxy::InitializeDynamicCollection(FGeometryDynam
 	//
 
 	DynamicCollection.CopyMatchingAttributesFrom(RestCollection);
-
-	
+	check(DynamicCollection.HasAttribute(FGeometryDynamicCollection::ImplicitsAttribute, FTransformCollection::TransformGroup) &&
+		  RestCollection.HasAttribute(FGeometryDynamicCollection::ImplicitsAttribute, FTransformCollection::TransformGroup));
 
 	//
 	// User defined initial velocities need to be populated. 
@@ -509,53 +533,6 @@ void FGeometryCollectionPhysicsProxy::InitializeDynamicCollection(FGeometryDynam
 		}
 	}
 
-
-	//
-	// The RestCollection may have been authored with TUniquePtr<FImplicitObject> under the name 
-	// FGeometryDynamicCollection::ImplicitsAttribute, as they originally were.  As the particle system now uses TSharedPtr,
-	// we're now authoring TSharedPtr<FImplicitObject> under the name FGeometryDynamicCollection::SharedImplicitsAttribute.
-	{
-		FGeometryCollection& NCRestCollection = const_cast<FGeometryCollection&>(RestCollection);
-		TManagedArray<TUniquePtr<Chaos::FImplicitObject>>* UniqueRestImplicits =
-			NCRestCollection.FindAttribute<TUniquePtr<Chaos::FImplicitObject>>(
-				FGeometryDynamicCollection::ImplicitsAttribute, FTransformCollection::TransformGroup);
-		TManagedArray<TSharedPtr<Chaos::FImplicitObject, ESPMode::ThreadSafe>>* SharedRestImplicits =
-			NCRestCollection.FindAttribute<TSharedPtr<Chaos::FImplicitObject, ESPMode::ThreadSafe>>(
-				FGeometryDynamicCollection::FGeometryDynamicCollection::SharedImplicitsAttribute, FTransformCollection::TransformGroup);
-
-		// If we have unique pointers to geometry, we need to convert them to shared pointers.
-		TManagedArray<TSharedPtr<Chaos::FImplicitObject, ESPMode::ThreadSafe>>* RestCollectionSharedImplicits = nullptr;
-		if (!SharedRestImplicits)
-		{
-			RestCollectionSharedImplicits =
-				&NCRestCollection.AddAttribute<TSharedPtr<Chaos::FImplicitObject, ESPMode::ThreadSafe>>(
-					FGeometryDynamicCollection::SharedImplicitsAttribute, FTransformCollection::TransformGroup);
-		}
-
-		for (int32 Index = DynamicCollection.NumElements(FTransformCollection::TransformGroup) - 1; 0 <= Index; Index--)
-		{
-			//Implicits[Index] = Chaos::MakeSerializable(RestImplicits[Index]);
-			if (SharedRestImplicits)
-			{
-				DynamicCollection.Implicits[Index] = (*SharedRestImplicits)[Index];
-			}
-			else if (UniqueRestImplicits && RestCollectionSharedImplicits) // RestCollectionSharedImplicits is redundant, but static analysis build complains without it
-			{
-				// Steal the memory from the unique ptr.
-				DynamicCollection.Implicits[Index] = TSharedPtr<Chaos::FImplicitObject, ESPMode::ThreadSafe>((*UniqueRestImplicits)[Index].Release());
-				// Update the rest collection with a copy of the shared pointer.
-				(*RestCollectionSharedImplicits)[Index] = DynamicCollection.Implicits[Index];
-			}
-		}
-
-		// Clean up now invalid unique pointers.
-		if (RestCollectionSharedImplicits)
-		{
-			NCRestCollection.RemoveAttribute(FGeometryDynamicCollection::ImplicitsAttribute, FTransformCollection::TransformGroup);
-		}
-	}
-
-
 	// process simplicials
 	{
 		/*
@@ -643,7 +620,7 @@ void FGeometryCollectionPhysicsProxy::InitializeBodiesPT(
 		const TManagedArray<FTransform>& MassToLocal = DynamicCollection.MassToLocal;
 		const TManagedArray<FVector>& InitialAngularVelocity = DynamicCollection.InitialAngularVelocity;
 		const TManagedArray<FVector>& InitialLinearVelocity = DynamicCollection.InitialLinearVelocity;
-		const TManagedArray<FGeometryDynamicCollection::FSharedImplicit>& Implicits = DynamicCollection.SharedGeometry;
+		const TManagedArray<FGeometryDynamicCollection::FSharedImplicit>& Implicits = DynamicCollection.Implicits;
 
 		TArray<FTransform> Transform;
 		GeometryCollectionAlgo::GlobalMatrices(DynamicCollection.Transform, DynamicCollection.Parent, Transform);
@@ -951,7 +928,7 @@ FGeometryCollectionPhysicsProxy::BuildClusters(
 	TManagedArray<FTransform>& Transform = DynamicCollection.Transform;
 	TManagedArray<FTransform>& MassToLocal = DynamicCollection.MassToLocal;
 	//TManagedArray<TSharedPtr<FCollisionStructureManager::FSimplicial> >& Simplicials = DynamicCollection.Simplicials;
-	TManagedArray<TSharedPtr<Chaos::FImplicitObject, ESPMode::ThreadSafe>>& Implicits = DynamicCollection.Implicits;
+	TManagedArray<TSharedPtr<Chaos::FImplicitObject>>& Implicits = DynamicCollection.Implicits;
 
 	//If we are a root particle use the world transform, otherwise set the relative transform
 	const FTransform CollectionSpaceTransform = GeometryCollectionAlgo::GlobalMatrix(Transform, ParentIndex, CollectionClusterIndex);
@@ -1860,9 +1837,9 @@ void FGeometryCollectionPhysicsProxy::InitializeSharedCollisionStructures(
 	TManagedArray<TUniquePtr<FSimplicial>>& CollectionSimplicials = 
 		RestCollection.AddAttribute<TUniquePtr<FSimplicial>>(
 			FGeometryDynamicCollection::SimplicialsAttribute, FTransformCollection::TransformGroup);
-	TManagedArray<TSharedPtr<Chaos::FImplicitObject, ESPMode::ThreadSafe>>& CollectionImplicits = 
-		RestCollection.AddAttribute<TSharedPtr<Chaos::FImplicitObject, ESPMode::ThreadSafe>>(
-			FGeometryDynamicCollection::SharedImplicitsAttribute, FTransformCollection::TransformGroup);
+	TManagedArray<TSharedPtr<Chaos::FImplicitObject>>& CollectionImplicits = 
+		RestCollection.AddAttribute<TSharedPtr<Chaos::FImplicitObject>>(
+			FGeometryDynamicCollection::ImplicitsAttribute, FTransformCollection::TransformGroup);
 
 	FTransform IdentityXf(FQuat::Identity, FVector(0));
 	IdentityXf.NormalizeRotation();
@@ -2088,7 +2065,7 @@ void FGeometryCollectionPhysicsProxy::InitializeSharedCollisionStructures(
 			if (SizeSpecificData.ImplicitType == EImplicitTypeEnum::Chaos_Implicit_LevelSet)
 			{
 				ErrorReporter.SetPrefix(BaseErrorPrefix + " | Transform Index: " + FString::FromInt(TransformGroupIndex) + " of " + FString::FromInt(TransformIndex.Num()));
-				CollectionImplicits[TransformGroupIndex] = TSharedPtr<Chaos::FImplicitObject, ESPMode::ThreadSafe>(
+				CollectionImplicits[TransformGroupIndex] = TSharedPtr<Chaos::FImplicitObject>(
 					FCollisionStructureManager::NewImplicitLevelset(
 						ErrorReporter,
 						MassSpaceParticles,
@@ -2101,7 +2078,7 @@ void FGeometryCollectionPhysicsProxy::InitializeSharedCollisionStructures(
 			}
 			else if (SizeSpecificData.ImplicitType == EImplicitTypeEnum::Chaos_Implicit_Box)
 			{
-				CollectionImplicits[TransformGroupIndex] = TSharedPtr<Chaos::FImplicitObject, ESPMode::ThreadSafe>(
+				CollectionImplicits[TransformGroupIndex] = TSharedPtr<Chaos::FImplicitObject>(
 					FCollisionStructureManager::NewImplicitBox(
 						InstanceBoundingBox,
 						SizeSpecificData.CollisionObjectReductionPercentage,
@@ -2109,7 +2086,7 @@ void FGeometryCollectionPhysicsProxy::InitializeSharedCollisionStructures(
 			}
 			else if (SizeSpecificData.ImplicitType == EImplicitTypeEnum::Chaos_Implicit_Sphere)
 			{
-				CollectionImplicits[TransformGroupIndex] = TSharedPtr<Chaos::FImplicitObject, ESPMode::ThreadSafe>(
+				CollectionImplicits[TransformGroupIndex] = TSharedPtr<Chaos::FImplicitObject>(
 					FCollisionStructureManager::NewImplicitSphere(
 						InnerRadius[GeometryIndex],
 						SizeSpecificData.CollisionObjectReductionPercentage,
@@ -2289,7 +2266,7 @@ void FGeometryCollectionPhysicsProxy::InitializeSharedCollisionStructures(
 
 					//don't support non level-set serialization
 					ErrorReporter.SetPrefix(BaseErrorPrefix + " | Cluster Transform Index: " + FString::FromInt(ClusterTransformIdx));
-					CollectionImplicits[ClusterTransformIdx] = TSharedPtr<Chaos::FImplicitObject, ESPMode::ThreadSafe>(
+					CollectionImplicits[ClusterTransformIdx] = TSharedPtr<Chaos::FImplicitObject>(
 						FCollisionStructureManager::NewImplicitLevelset(
 							ErrorReporter,
 							MassSpaceParticles,
@@ -2306,7 +2283,7 @@ void FGeometryCollectionPhysicsProxy::InitializeSharedCollisionStructures(
 				}
 				else if (SizeSpecificData.ImplicitType == EImplicitTypeEnum::Chaos_Implicit_Box)
 				{
-					CollectionImplicits[ClusterTransformIdx] = TSharedPtr<Chaos::FImplicitObject, ESPMode::ThreadSafe>(
+					CollectionImplicits[ClusterTransformIdx] = TSharedPtr<Chaos::FImplicitObject>(
 						FCollisionStructureManager::NewImplicitBox(
 							InstanceBoundingBox,
 							SizeSpecificData.CollisionObjectReductionPercentage,
@@ -2318,7 +2295,7 @@ void FGeometryCollectionPhysicsProxy::InitializeSharedCollisionStructures(
 				}
 				else if (SizeSpecificData.ImplicitType == EImplicitTypeEnum::Chaos_Implicit_Sphere)
 				{
- 					CollectionImplicits[ClusterTransformIdx] = TSharedPtr<Chaos::FImplicitObject, ESPMode::ThreadSafe>(
+ 					CollectionImplicits[ClusterTransformIdx] = TSharedPtr<Chaos::FImplicitObject>(
 						FCollisionStructureManager::NewImplicitSphere(
 							InstanceBoundingBox.GetExtent().GetAbsMin() / 2,
 							SizeSpecificData.CollisionObjectReductionPercentage,
