Attempt to fix static analysis

diff --git a/Engine/Source/Runtime/Experimental/Chaos/Private/Chaos/CollisionResolution.cpp b/Engine/Source/Runtime/Experimental/Chaos/Private/Chaos/CollisionResolution.cpp
--- a/Engine/Source/Runtime/Experimental/Chaos/Private/Chaos/CollisionResolution.cpp
+++ b/Engine/Source/Runtime/Experimental/Chaos/Private/Chaos/CollisionResolution.cpp
@@ -2555,9 +2555,7 @@ namespace Chaos
 			}
 			Dir /= Length;
 
-			// Attempt to get rid of invalid static analysis warnings
-			check(Implicit0 || (!Implicit0 && Implicit0Type == ImplicitObjectType::Unknown));
-			check(Implicit1 || (!Implicit1 && Implicit1Type == ImplicitObjectType::Unknown));
+			CA_SUPPRESS(6011); // It's impossible to have an null Implicit0 or Implicit1 as we check types.
 
 			if (Implicit0Type == TBox<FReal, 3>::StaticType() && Implicit1Type == THeightField<FReal>::StaticType())
 			{
