Reinstate ESPMode::ThreadSafe for shared pointers to implicit objects.

diff --git a/Engine/Source/Runtime/Experimental/Chaos/Private/Chaos/PBDRigidClustering.cpp b/Engine/Source/Runtime/Experimental/Chaos/Private/Chaos/PBDRigidClustering.cpp
--- a/Engine/Source/Runtime/Experimental/Chaos/Private/Chaos/PBDRigidClustering.cpp
+++ b/Engine/Source/Runtime/Experimental/Chaos/Private/Chaos/PBDRigidClustering.cpp
@@ -426,7 +426,7 @@ namespace Chaos
 		const int32 ClusterGroupIndex,
 		TArray<Chaos::TPBDRigidParticleHandle<T,d>*>&& Children,
 		const FClusterCreationParameters<T>& Parameters,
-		TSharedPtr<Chaos::FImplicitObject> ProxyGeometry,
+		TSharedPtr<Chaos::FImplicitObject, ESPMode::ThreadSafe> ProxyGeometry,
 		const TRigidTransform<T, d>* ForceMassOrientation)
 	{
 		SCOPE_CYCLE_COUNTER(STAT_CreateClusterParticle);
@@ -1348,7 +1348,7 @@ namespace Chaos
 	void TPBDRigidClustering<FPBDRigidsEvolution, FPBDCollisionConstraint, T, d>::UpdateGeometry(
 		Chaos::TPBDRigidClusteredParticleHandle<float, 3>* Parent, 
 		const TSet<TPBDRigidParticleHandle<T, d>*>& Children, 
-		TSharedPtr<Chaos::FImplicitObject> ProxyGeometry,
+		TSharedPtr<Chaos::FImplicitObject, ESPMode::ThreadSafe> ProxyGeometry,
 		const FClusterCreationParameters<T>& Parameters)
 	{
 		SCOPE_CYCLE_COUNTER(STAT_UpdateGeometry);
diff --git a/Engine/Source/Runtime/Experimental/Chaos/Public/Chaos/PBDRigidClustering.h b/Engine/Source/Runtime/Experimental/Chaos/Public/Chaos/PBDRigidClustering.h
--- a/Engine/Source/Runtime/Experimental/Chaos/Public/Chaos/PBDRigidClustering.h
+++ b/Engine/Source/Runtime/Experimental/Chaos/Public/Chaos/PBDRigidClustering.h
@@ -62,7 +62,7 @@ public:
 		const int32 ClusterGroupIndex, 
 		TArray<Chaos::TPBDRigidParticleHandle<T,d>*>&& Children, 
 		const FClusterCreationParameters<T>& Parameters/* = FClusterCreationParameters<T>()*/,
-		TSharedPtr<Chaos::FImplicitObject> ProxyGeometry = nullptr,
+		TSharedPtr<Chaos::FImplicitObject, ESPMode::ThreadSafe> ProxyGeometry = nullptr,
 		const TRigidTransform<T, d>* ForceMassOrientation = nullptr);
 
 	/**
@@ -285,7 +285,7 @@ public:
 	void UpdateGeometry(
 		Chaos::TPBDRigidClusteredParticleHandle<float, 3>* Parent, 
 		const TSet<TPBDRigidParticleHandle<T, d>*>& Children, 
-		TSharedPtr<Chaos::FImplicitObject> ProxyGeometry,
+		TSharedPtr<Chaos::FImplicitObject, ESPMode::ThreadSafe> ProxyGeometry,
 		const FClusterCreationParameters<T>& Parameters);
 
 	void ComputeStrainFromCollision(const FPBDCollisionConstraint& CollisionRule);
diff --git a/Engine/Source/Runtime/Experimental/ChaosSolvers/Private/PhysicsProxy/GeometryCollectionPhysicsProxy.cpp b/Engine/Source/Runtime/Experimental/ChaosSolvers/Private/PhysicsProxy/GeometryCollectionPhysicsProxy.cpp
--- a/Engine/Source/Runtime/Experimental/ChaosSolvers/Private/PhysicsProxy/GeometryCollectionPhysicsProxy.cpp
+++ b/Engine/Source/Runtime/Experimental/ChaosSolvers/Private/PhysicsProxy/GeometryCollectionPhysicsProxy.cpp
@@ -226,7 +226,7 @@ void PopulateSimulatedParticle(
 	Chaos::TPBDRigidParticleHandle<float,3>* Handle,
 	const FSharedSimulationParameters& SharedParams,
 	const FCollisionStructureManager::FSimplicial* Simplicial,
-	TSharedPtr<Chaos::FImplicitObject> Implicit,
+	FGeometryDynamicCollection::FSharedImplicit Implicit,
 	const float MassIn,
 	const FVector& InertiaTensorVec,
 	const FTransform& WorldTransform, 
@@ -928,7 +928,7 @@ FGeometryCollectionPhysicsProxy::BuildClusters(
 	TManagedArray<FTransform>& Transform = DynamicCollection.Transform;
 	TManagedArray<FTransform>& MassToLocal = DynamicCollection.MassToLocal;
 	//TManagedArray<TSharedPtr<FCollisionStructureManager::FSimplicial> >& Simplicials = DynamicCollection.Simplicials;
-	TManagedArray<TSharedPtr<Chaos::FImplicitObject>>& Implicits = DynamicCollection.Implicits;
+	TManagedArray<FGeometryDynamicCollection::FSharedImplicit>& Implicits = DynamicCollection.Implicits;
 
 	//If we are a root particle use the world transform, otherwise set the relative transform
 	const FTransform CollectionSpaceTransform = GeometryCollectionAlgo::GlobalMatrix(Transform, ParentIndex, CollectionClusterIndex);
@@ -1837,8 +1837,8 @@ void FGeometryCollectionPhysicsProxy::InitializeSharedCollisionStructures(
 	TManagedArray<TUniquePtr<FSimplicial>>& CollectionSimplicials = 
 		RestCollection.AddAttribute<TUniquePtr<FSimplicial>>(
 			FGeometryDynamicCollection::SimplicialsAttribute, FTransformCollection::TransformGroup);
-	TManagedArray<TSharedPtr<Chaos::FImplicitObject>>& CollectionImplicits = 
-		RestCollection.AddAttribute<TSharedPtr<Chaos::FImplicitObject>>(
+	TManagedArray<FGeometryDynamicCollection::FSharedImplicit>& CollectionImplicits = 
+		RestCollection.AddAttribute<FGeometryDynamicCollection::FSharedImplicit>(
 			FGeometryDynamicCollection::ImplicitsAttribute, FTransformCollection::TransformGroup);
 
 	FTransform IdentityXf(FQuat::Identity, FVector(0));
@@ -2065,7 +2065,7 @@ void FGeometryCollectionPhysicsProxy::InitializeSharedCollisionStructures(
 			if (SizeSpecificData.ImplicitType == EImplicitTypeEnum::Chaos_Implicit_LevelSet)
 			{
 				ErrorReporter.SetPrefix(BaseErrorPrefix + " | Transform Index: " + FString::FromInt(TransformGroupIndex) + " of " + FString::FromInt(TransformIndex.Num()));
-				CollectionImplicits[TransformGroupIndex] = TSharedPtr<Chaos::FImplicitObject>(
+				CollectionImplicits[TransformGroupIndex] = FGeometryDynamicCollection::FSharedImplicit(
 					FCollisionStructureManager::NewImplicitLevelset(
 						ErrorReporter,
 						MassSpaceParticles,
@@ -2078,7 +2078,7 @@ void FGeometryCollectionPhysicsProxy::InitializeSharedCollisionStructures(
 			}
 			else if (SizeSpecificData.ImplicitType == EImplicitTypeEnum::Chaos_Implicit_Box)
 			{
-				CollectionImplicits[TransformGroupIndex] = TSharedPtr<Chaos::FImplicitObject>(
+				CollectionImplicits[TransformGroupIndex] = FGeometryDynamicCollection::FSharedImplicit(
 					FCollisionStructureManager::NewImplicitBox(
 						InstanceBoundingBox,
 						SizeSpecificData.CollisionObjectReductionPercentage,
@@ -2086,7 +2086,7 @@ void FGeometryCollectionPhysicsProxy::InitializeSharedCollisionStructures(
 			}
 			else if (SizeSpecificData.ImplicitType == EImplicitTypeEnum::Chaos_Implicit_Sphere)
 			{
-				CollectionImplicits[TransformGroupIndex] = TSharedPtr<Chaos::FImplicitObject>(
+				CollectionImplicits[TransformGroupIndex] = FGeometryDynamicCollection::FSharedImplicit(
 					FCollisionStructureManager::NewImplicitSphere(
 						InnerRadius[GeometryIndex],
 						SizeSpecificData.CollisionObjectReductionPercentage,
@@ -2266,7 +2266,7 @@ void FGeometryCollectionPhysicsProxy::InitializeSharedCollisionStructures(
 
 					//don't support non level-set serialization
 					ErrorReporter.SetPrefix(BaseErrorPrefix + " | Cluster Transform Index: " + FString::FromInt(ClusterTransformIdx));
-					CollectionImplicits[ClusterTransformIdx] = TSharedPtr<Chaos::FImplicitObject>(
+					CollectionImplicits[ClusterTransformIdx] = FGeometryDynamicCollection::FSharedImplicit(
 						FCollisionStructureManager::NewImplicitLevelset(
 							ErrorReporter,
 							MassSpaceParticles,
@@ -2283,7 +2283,7 @@ void FGeometryCollectionPhysicsProxy::InitializeSharedCollisionStructures(
 				}
 				else if (SizeSpecificData.ImplicitType == EImplicitTypeEnum::Chaos_Implicit_Box)
 				{
-					CollectionImplicits[ClusterTransformIdx] = TSharedPtr<Chaos::FImplicitObject>(
+					CollectionImplicits[ClusterTransformIdx] = FGeometryDynamicCollection::FSharedImplicit(
 						FCollisionStructureManager::NewImplicitBox(
 							InstanceBoundingBox,
 							SizeSpecificData.CollisionObjectReductionPercentage,
@@ -2295,7 +2295,7 @@ void FGeometryCollectionPhysicsProxy::InitializeSharedCollisionStructures(
 				}
 				else if (SizeSpecificData.ImplicitType == EImplicitTypeEnum::Chaos_Implicit_Sphere)
 				{
- 					CollectionImplicits[ClusterTransformIdx] = TSharedPtr<Chaos::FImplicitObject>(
+ 					CollectionImplicits[ClusterTransformIdx] = FGeometryDynamicCollection::FSharedImplicit(
 						FCollisionStructureManager::NewImplicitSphere(
 							InstanceBoundingBox.GetExtent().GetAbsMin() / 2,
 							SizeSpecificData.CollisionObjectReductionPercentage,
diff --git a/Engine/Source/Runtime/Experimental/ChaosSolvers/Public/GeometryCollectionProxyData.h b/Engine/Source/Runtime/Experimental/ChaosSolvers/Public/GeometryCollectionProxyData.h
--- a/Engine/Source/Runtime/Experimental/ChaosSolvers/Public/GeometryCollectionProxyData.h
+++ b/Engine/Source/Runtime/Experimental/ChaosSolvers/Public/GeometryCollectionProxyData.h
@@ -59,7 +59,7 @@ public:
 	FGeometryDynamicCollection& operator=(FGeometryDynamicCollection&&) = delete;
 
 	typedef FTransformDynamicCollection Super;
-	typedef TSharedPtr<Chaos::FImplicitObject> FSharedImplicit;
+	typedef TSharedPtr<Chaos::FImplicitObject, ESPMode::ThreadSafe> FSharedImplicit;
 
 	static const FName ActiveAttribute;
 	static const FName CollisionGroupAttribute;
@@ -77,7 +77,7 @@ public:
 	TManagedArray<int32> CollisionMask;
 	TManagedArray<int32> CollisionStructureID;
 	TManagedArray<int32> DynamicState;
-	TManagedArray<TSharedPtr<Chaos::FImplicitObject>> Implicits;
+	TManagedArray<FSharedImplicit> Implicits;
 	TManagedArray<FVector> InitialAngularVelocity;
 	TManagedArray<FVector> InitialLinearVelocity;
 	TManagedArray<FTransform> MassToLocal;
