Hooked up geometry collection to new SQ structures. GT dynamic collection now attached a TGeometryParticle array to store external geometry and transform for queries. Rotated Geometry Collection DDC key to fix previous serialization breakage.

diff --git a/Engine/Source/Runtime/Engine/Private/Collision/CollisionConversions.cpp b/Engine/Source/Runtime/Engine/Private/Collision/CollisionConversions.cpp
--- a/Engine/Source/Runtime/Engine/Private/Collision/CollisionConversions.cpp
+++ b/Engine/Source/Runtime/Engine/Private/Collision/CollisionConversions.cpp
@@ -171,7 +171,21 @@ static void SetHitResultFromShapeAndFaceIndex(const FPhysicsShape& Shape,  const
 #endif
 	else
 	{
-		ensureMsgf(false, TEXT("SetHitResultFromShapeAndFaceIndex hit shape with invalid userData"));
+		// Currently geom collections are registered with a primitive component user data, but maybe custom should be adapted
+		// to be more general so we can support leaf identification #BGTODO
+		void* UserData = Actor.UserData();
+		UPrimitiveComponent* PossibleOwner = FPhysxUserData::Get<UPrimitiveComponent>(UserData);
+
+		if(PossibleOwner)
+		{
+			OwningComponent = PossibleOwner;
+			OutResult.Item = INDEX_NONE;
+			OutResult.BoneName = NAME_None;
+		}
+		else
+		{
+			ensureMsgf(false, TEXT("SetHitResultFromShapeAndFaceIndex hit shape with invalid userData"));
+		}
 	}
 
 	OutResult.PhysMaterial = nullptr;
diff --git a/Engine/Source/Runtime/Experimental/Chaos/Public/Chaos/GeometryQueries.h b/Engine/Source/Runtime/Experimental/Chaos/Public/Chaos/GeometryQueries.h
--- a/Engine/Source/Runtime/Experimental/Chaos/Public/Chaos/GeometryQueries.h
+++ b/Engine/Source/Runtime/Experimental/Chaos/Public/Chaos/GeometryQueries.h
@@ -213,7 +213,7 @@ namespace Chaos
 				}
 				else
 				{
-					check(false);	//unsupported query type
+					ensureMsgf(false, TEXT("Unsupported query type: %u"), (uint8)AType);
 				}
 
 			}
diff --git a/Engine/Source/Runtime/Experimental/ChaosSolvers/Private/PhysicsProxy/GeometryCollectionPhysicsProxy.cpp b/Engine/Source/Runtime/Experimental/ChaosSolvers/Private/PhysicsProxy/GeometryCollectionPhysicsProxy.cpp
--- a/Engine/Source/Runtime/Experimental/ChaosSolvers/Private/PhysicsProxy/GeometryCollectionPhysicsProxy.cpp
+++ b/Engine/Source/Runtime/Experimental/ChaosSolvers/Private/PhysicsProxy/GeometryCollectionPhysicsProxy.cpp
@@ -365,6 +365,8 @@ FGeometryCollectionPhysicsProxy::FGeometryCollectionPhysicsProxy(
 	UObject* InOwner,
 	FGeometryDynamicCollection& GameThreadCollectionIn,
 	const FSimulationParameters& SimulationParameters,
+	FCollisionFilterData InSimFilter,
+	FCollisionFilterData InQueryFilter,
 	FInitFunc InInitFunc,
 	FCacheSyncFunc InCacheSyncFunc,
 	FFinalSyncFunc InFinalSyncFunc,
@@ -375,7 +377,8 @@ FGeometryCollectionPhysicsProxy::FGeometryCollectionPhysicsProxy(
 	, BaseParticleIndex(INDEX_NONE)
 	, IsObjectDynamic(false)
 	, IsObjectLoading(true)
-
+	, SimFilter(InSimFilter)
+	, QueryFilter(InQueryFilter)
 #if TODO_REIMPLEMENT_RIGID_CACHING
 	, ProxySimDuration(0.0f)
 	, LastSyncCountGT(MAX_uint32)
@@ -411,7 +414,6 @@ void FGeometryCollectionPhysicsProxy::Initialize()
 	//
 	FGeometryDynamicCollection& DynamicCollection = GameThreadCollection;
 
-
 	InitializeDynamicCollection(DynamicCollection, *Parameters.RestCollection, Parameters);
 
 	NumParticles = DynamicCollection.NumElements(FGeometryCollection::TransformGroup);
@@ -509,6 +511,38 @@ void FGeometryCollectionPhysicsProxy::Initialize()
 		//Results.Get(1).WorldBounds = FBoxSphereBounds(BoundingBox);
 	}
 	*/
+
+	// Initialise GT/External particles
+	const int32 NumTransforms = GameThreadCollection.Transform.Num();
+
+	// Attach the external particles to the gamethread collection
+	GameThreadCollection.AddExternalAttribute<TUniquePtr<Chaos::TGeometryParticle<Chaos::FReal, 3>>>(FGeometryCollection::ParticlesAttribute, FTransformCollection::TransformGroup, GTParticles);
+
+	if(ensure(NumTransforms == GameThreadCollection.Implicits.Num() && NumTransforms == GTParticles.Num())) // Implicits are in the transform group so this invariant should always hold
+	{
+		for(int32 Index = 0; Index < NumTransforms; ++Index)
+		{
+			GTParticles[Index] = Chaos::TGeometryParticle<Chaos::FReal, 3>::CreateParticle();
+			Chaos::TGeometryParticle<Chaos::FReal, 3>* P = GTParticles[Index].Get();
+
+			const FTransform& T = Parameters.WorldTransform * GameThreadCollection.Transform[Index];
+			P->SetX(T.GetTranslation(), false);
+			P->SetR(T.GetRotation(), false);
+			P->SetUserData(Parameters.UserData);
+			P->Proxy = this;
+			P->SetGeometry(GameThreadCollection.Implicits[Index]);
+
+			const Chaos::TShapesArray<Chaos::FReal, 3>& Shapes = P->ShapesArray();
+			const int32 NumShapes = Shapes.Num();
+			for(int32 ShapeIndex = 0; ShapeIndex < NumShapes; ++ShapeIndex)
+			{
+				Chaos::TPerShapeData<Chaos::FReal, 3>* Shape = Shapes[ShapeIndex].Get();
+				Shape->SimData = SimFilter;
+				Shape->QueryData = QueryFilter;
+			}
+		}
+	}
+
 	PhysicsThreadCollection.CopyMatchingAttributesFrom(DynamicCollection);
 }
 
@@ -519,8 +553,9 @@ void FGeometryCollectionPhysicsProxy::InitializeDynamicCollection(FGeometryDynam
 	//
 
 	DynamicCollection.CopyMatchingAttributesFrom(RestCollection);
-	check(DynamicCollection.HasAttribute(FGeometryDynamicCollection::ImplicitsAttribute, FTransformCollection::TransformGroup) &&
-		  RestCollection.HasAttribute(FGeometryDynamicCollection::ImplicitsAttribute, FTransformCollection::TransformGroup));
+
+	check(DynamicCollection.HasAttribute(FGeometryDynamicCollection::ImplicitsAttribute, FTransformCollection::TransformGroup));
+	check(RestCollection.HasAttribute(FGeometryDynamicCollection::ImplicitsAttribute, FTransformCollection::TransformGroup));
 
 	//
 	// User defined initial velocities need to be populated. 
@@ -658,6 +693,9 @@ void FGeometryCollectionPhysicsProxy::InitializeBodiesPT(
 				SolverParticleHandles[Idx] = Handle;
 				HandleToTransformGroupIndex.Add(Handle, Idx);
 
+				// We're on the physics thread here but we've already set up the GT particles and we're just linking here
+				Handle->GTGeometryParticle() = GTParticles[Idx].Get();
+
 				check(SolverParticleHandles[Idx]->GetParticleType() == Handle->GetParticleType());
 				RigidsSolver->GetEvolution()->CreateParticle(Handle);
 			}
@@ -1370,6 +1408,7 @@ void FGeometryCollectionPhysicsProxy::PushToPhysicsState(const Chaos::FParticleD
 			}
 		} // if rigid handle
 	} // for all SolverParticleHandles
+	*/
 
 	if (DisableGeometryCollectionGravity) // cvar
 	{
@@ -1388,7 +1427,6 @@ void FGeometryCollectionPhysicsProxy::PushToPhysicsState(const Chaos::FParticleD
 			}
 		}
 	}
-	*/
 }
 
 void FGeometryCollectionPhysicsProxy::BufferPhysicsResults()
@@ -1679,9 +1717,13 @@ void FGeometryCollectionPhysicsProxy::PullFromPhysicsState()
 		{
 			if(!TR.DisabledStates[TmIndex])
 			{
-				DynamicCollection.Transform[TmIndex] = TR.Transforms[TmIndex];
+				const FTransform& LocalTransform = TR.Transforms[TmIndex];
+				const FTransform& ParticleToWorld = TR.ParticleToWorldTransforms[TmIndex];
+
+				DynamicCollection.Transform[TmIndex] = LocalTransform;
+				GTParticles[TmIndex]->SetX(ParticleToWorld.GetTranslation());
+				GTParticles[TmIndex]->SetR(ParticleToWorld.GetRotation());
 			}
-			DynamicCollection.DynamicState[TmIndex] = TR.DynamicState[TmIndex];
 		}
 
 		//question: why do we need this? Sleeping objects will always have to update GPU
@@ -1798,6 +1840,10 @@ int32 FindSizeSpecificIdx(const TArray<FSharedSimulationSizeSpecificData>& SizeS
 	return UseIdx;
 }
 
+/** 
+	NOTE - Making any changes to data stored on the rest collection below MUST be accompanied
+	by a rotation of the DDC key in FDerivedDataGeometryCollectionCooker::GetVersionString
+*/
 void FGeometryCollectionPhysicsProxy::InitializeSharedCollisionStructures(
 	Chaos::FErrorReporter& ErrorReporter, 
 	FGeometryCollection& RestCollection, 
diff --git a/Engine/Source/Runtime/Experimental/ChaosSolvers/Public/PhysicsProxy/GeometryCollectionPhysicsProxy.h b/Engine/Source/Runtime/Experimental/ChaosSolvers/Public/PhysicsProxy/GeometryCollectionPhysicsProxy.h
--- a/Engine/Source/Runtime/Experimental/ChaosSolvers/Public/PhysicsProxy/GeometryCollectionPhysicsProxy.h
+++ b/Engine/Source/Runtime/Experimental/ChaosSolvers/Public/PhysicsProxy/GeometryCollectionPhysicsProxy.h
@@ -130,6 +130,8 @@ public:
 		UObject* InOwner, 
 		FGeometryDynamicCollection& GameThreadCollection,
 		const FSimulationParameters& SimulationParameters,
+		FCollisionFilterData InSimFilter,
+		FCollisionFilterData InQueryFilter,
 		FInitFunc InInitFunc, 
 		FCacheSyncFunc InCacheSyncFunc, 
 		FFinalSyncFunc InFinalSyncFunc  ,
@@ -229,6 +231,8 @@ public:
 	void DisableCollisionsCallback(TSet<TTuple<int32, int32>>& InPairs) {}
 	void AddForceCallback(FParticlesType& InParticles, const float InDt, const int32 InIndex) {}
 
+	bool IsGTCollectionDirty() const { return GameThreadCollection.IsDirty(); }
+
 protected:
 	/**
 	 * Build a physics thread cluster parent particle.
@@ -295,6 +299,11 @@ private:
 	//
 	bool IsObjectDynamic; // Records current dynamic state
 	bool IsObjectLoading; // Indicate when loaded
+
+	TManagedArray<TUniquePtr<Chaos::TGeometryParticle<Chaos::FReal, 3>>> GTParticles;
+	FCollisionFilterData SimFilter;
+	FCollisionFilterData QueryFilter;
+
 	TArray<int32> EndFrameUnparentingBuffer;
 	// This is a subset of the geometry group that are used in the transform hierarchy to represent geometry
 	TArray<FBox> ValidGeometryBoundingBoxes;
diff --git a/Engine/Source/Runtime/Experimental/GeometryCollectionCore/Private/GeometryCollection/TransformCollection.cpp b/Engine/Source/Runtime/Experimental/GeometryCollectionCore/Private/GeometryCollection/TransformCollection.cpp
--- a/Engine/Source/Runtime/Experimental/GeometryCollectionCore/Private/GeometryCollection/TransformCollection.cpp
+++ b/Engine/Source/Runtime/Experimental/GeometryCollectionCore/Private/GeometryCollection/TransformCollection.cpp
@@ -11,6 +11,7 @@ const FName FTransformCollection::TransformGroup = "Transform";
 const FName FTransformCollection::TransformAttribute = "Transform";
 const FName FTransformCollection::ParentAttribute = "Parent";
 const FName FTransformCollection::ChildrenAttribute = "Children";
+const FName FTransformCollection::ParticlesAttribute = "Particles";
 
 
 FTransformCollection::FTransformCollection()
diff --git a/Engine/Source/Runtime/Experimental/GeometryCollectionCore/Public/GeometryCollection/ManagedArrayCollection.h b/Engine/Source/Runtime/Experimental/GeometryCollectionCore/Public/GeometryCollection/ManagedArrayCollection.h
--- a/Engine/Source/Runtime/Experimental/GeometryCollectionCore/Public/GeometryCollection/ManagedArrayCollection.h
+++ b/Engine/Source/Runtime/Experimental/GeometryCollectionCore/Public/GeometryCollection/ManagedArrayCollection.h
@@ -312,11 +312,11 @@ public:
 	*/
 	void MakeDirty() { bDirty = true; }
 	void MakeClean() { bDirty = false; }
-	bool IsDirty() { return bDirty; }
+	bool IsDirty() const { return bDirty; }
 
 	/**
 	* Serialize
-	*/
+		*/
 	virtual void Serialize(Chaos::FChaosArchive& Ar);
 
 	/**
diff --git a/Engine/Source/Runtime/Experimental/GeometryCollectionCore/Public/GeometryCollection/TransformCollection.h b/Engine/Source/Runtime/Experimental/GeometryCollectionCore/Public/GeometryCollection/TransformCollection.h
--- a/Engine/Source/Runtime/Experimental/GeometryCollectionCore/Public/GeometryCollection/TransformCollection.h
+++ b/Engine/Source/Runtime/Experimental/GeometryCollectionCore/Public/GeometryCollection/TransformCollection.h
@@ -49,6 +49,7 @@ public:
 	static const FName TransformAttribute;
 	static const FName ParentAttribute;
 	static const FName ChildrenAttribute;
+	static const FName ParticlesAttribute;
 
 	/** Serialize */
 	virtual void Serialize(Chaos::FChaosArchive& Ar) override;
diff --git a/Engine/Source/Runtime/Experimental/GeometryCollectionEngine/Private/GeometryCollection/DerivedDataGeometryCollectionCooker.cpp b/Engine/Source/Runtime/Experimental/GeometryCollectionEngine/Private/GeometryCollection/DerivedDataGeometryCollectionCooker.cpp
--- a/Engine/Source/Runtime/Experimental/GeometryCollectionEngine/Private/GeometryCollection/DerivedDataGeometryCollectionCooker.cpp
+++ b/Engine/Source/Runtime/Experimental/GeometryCollectionEngine/Private/GeometryCollection/DerivedDataGeometryCollectionCooker.cpp
@@ -58,33 +58,8 @@ const TCHAR* FDerivedDataGeometryCollectionCooker::GetVersionString() const
 	{
 		return OverrideVersion;	//force load old ddc if found. Not recommended
 	}
-	//touch guid when new data needs to be serialized. Please leave comment saying what you did
-	//TEXT("4");	//original versions before using guids.
-	//return TEXT("8724C70A140146B5A2F4CF0C16083041");	//Switching to GUIDs so that people don't poison DDC when testing things locally
-	//return TEXT("C422BF0A539C455A90C6A8660BA0170F");	//Previous version was only saving to DDC. Changing GUID so we know which entries are potentially missing from assets
-	//return TEXT("951075CFF57C40D88E15756393EECE70");	//Simplicials are now stored in the DDC
-	//return TEXT("5AA67AB5FE4040629FD47946FE3CA3FD");	//Implicits are now stored in the DDC
-	//return TEXT("421E954DED254B95B20125FBAB507746");	//Mass and inertia are clamped together to ensure both are reasonable. At least one simplicial sample point always exists
-	//return TEXT("969470E420CD4F8E83356B8BC4007799");	//Properly report error when non-contiguous objects are stored in a geometry collection
-	//return TEXT("35629AF302984E83935233220788CDCC");	//Added size specific data for geometry collection
-	//return TEXT("C84C3A40E80B42FC865922587047255D");	//Fix non-coincident triangle verts being incorrectly sorted after coincident vertices
-	//return TEXT("9EC74FE6494D443DA12189BADF953A33");	//Fix cluster inertia tensor
-	//return TEXT("5A59EE91C27641DF8A083998EA32E560");	//Remove mesh normals when computing levelset
-	//return TEXT("08C81448E92D4B5EB9444C0E854A5BE4");	//Cull out slightly bigger triangles early in the pipeline. This is needed because of later transforms which add error
-	//return TEXT("E1615D65A2D645F3B9D119CC78D9240F");	//change default level set resolution from 5-10 to 10-10 to get better rasterization results
-	//return TEXT("C89DD39311BC49898736B7A49A1F51AF");	//Inertia and mass must have a minimum box of 10cm^3
-	//return TEXT("235FA4CD6A794305A4638D559DD1D359");	//Levelset normal updated to use surface normals
-	//return TEXT("43D27BBFDF96434BB116498A8375003E");	//Levelset error metrics
-	//return TEXT("142F76A9F1524A9F935B28C7FB74690E");	//Build simplicials
-	//return TEXT("132EC70C044646E48A725A67114DAE06");	//CleanCollisionParticles upgrade
-	//return TEXT("5EFA18C9C87E4350A34FD8F8100DDEA2");	//Bulk serialize changes
-	//return TEXT("1FD9874624304A0EB6C635DC2B4BE237");  //Add hole filling as pre-process before level set generation
-	//return TEXT("0182FCF14245402B8A902298417ECAA3");	//Increase resolution for cluster parents
-	//return TEXT("2F4F3A038A7B41F28CA86799DD8D4300");	//No tiny faces or non-manifold edges in the hole fills + fix materials
-	//return TEXT("B5F1831D60884F9DAC2C2CBCF2ECFA12");	//simplicials
-	//return TEXT("851938C6E7EA44C4B2DA0C8E600693A4");	//simplicials and inertia
-	return TEXT("E237FFBE23A54668A07C3D14AB68A7B4");	//box and sphere
 
+	return TEXT("ABC16C06507147A5A265CEDFAFDA96EE");
 }
 
 FString FDerivedDataGeometryCollectionCooker::GetPluginSpecificCacheKeySuffix() const
diff --git a/Engine/Source/Runtime/Experimental/GeometryCollectionEngine/Private/GeometryCollection/GeometryCollectionActor.cpp b/Engine/Source/Runtime/Experimental/GeometryCollectionEngine/Private/GeometryCollection/GeometryCollectionActor.cpp
--- a/Engine/Source/Runtime/Experimental/GeometryCollectionEngine/Private/GeometryCollection/GeometryCollectionActor.cpp
+++ b/Engine/Source/Runtime/Experimental/GeometryCollectionEngine/Private/GeometryCollection/GeometryCollectionActor.cpp
@@ -69,7 +69,7 @@ bool LowLevelRaycastImp(const Chaos::TVector<float, 3>& Start, const Chaos::TVec
 	//todo(ocohen): need to add thread safety / lock semantics
 	//const TManagedArray<int32>& RigidBodyIdArray = GeomCollectionActor.GetGeometryCollectionComponent()->GetRigidBodyIdArray();
 	const TManagedArray<FGuid>& RigidBodyIdArray = GeomCollectionActor.GetGeometryCollectionComponent()->GetRigidBodyGuidArray();
-	FPhysScene_Chaos* Scene = GeomCollectionActor.GetGeometryCollectionComponent()->GetPhysicsScene();
+	FPhysScene_Chaos* Scene = GeomCollectionActor.GetGeometryCollectionComponent()->GetInnerChaosScene();
 	ensure(Scene);
 
 	const Chaos::FPhysicsSolver* Solver = GetSolver(GeomCollectionActor);
diff --git a/Engine/Source/Runtime/Experimental/GeometryCollectionEngine/Private/GeometryCollection/GeometryCollectionComponent.cpp b/Engine/Source/Runtime/Experimental/GeometryCollectionEngine/Private/GeometryCollection/GeometryCollectionComponent.cpp
--- a/Engine/Source/Runtime/Experimental/GeometryCollectionEngine/Private/GeometryCollection/GeometryCollectionComponent.cpp
+++ b/Engine/Source/Runtime/Experimental/GeometryCollectionEngine/Private/GeometryCollection/GeometryCollectionComponent.cpp
@@ -20,6 +20,7 @@
 #include "ChaosStats.h"
 #include "PhysicsProxy/GeometryCollectionPhysicsProxy.h"
 #include "PhysicsSolver.h"
+#include "Physics/PhysicsFiltering.h"
 
 #if WITH_EDITOR
 #include "AssetToolsModule.h"
@@ -1010,8 +1011,7 @@ void UGeometryCollectionComponent::TickComponent(float DeltaTime, enum ELevelTic
 	//if (bRenderStateDirty && DynamicCollection)	//todo: always send for now
 	if(RestCollection)
 	{
-		//if (RestCollection->HasVisibleGeometry() || DynamicCollection->IsDirty())
-		if(RestCollection->HasVisibleGeometry())
+		if(RestCollection->HasVisibleGeometry() || DynamicCollection->IsDirty())
 		{
 			MarkRenderTransformDirty();
 			MarkRenderDynamicDataDirty();
@@ -1128,6 +1128,8 @@ void UGeometryCollectionComponent::OnCreatePhysicsState()
 				ChaosMaterial->SleepingAngularThreshold = PhysicalMaterial->SleepingAngularVelocityThreshold;
 			}
 
+			FPhysxUserData::Set<UPrimitiveComponent>(&PhysicsUserData, this);
+
 			FSimulationParameters SimulationParameters;
 			{
 #if !(UE_BUILD_SHIPPING || UE_BUILD_TEST)
@@ -1173,6 +1175,8 @@ void UGeometryCollectionComponent::OnCreatePhysicsState()
 				SimulationParameters.TrailingData.TrailingMinSpeedThreshold = CacheParameters.TrailingMinSpeedThreshold;
 				SimulationParameters.TrailingData.TrailingMinVolumeThreshold = CacheParameters.TrailingMinVolumeThreshold;
 				SimulationParameters.RemoveOnFractureEnabled = SimulationParameters.Shared.RemoveOnFractureIndices.Num() > 0;
+				SimulationParameters.WorldTransform = GetComponentToWorld();
+				SimulationParameters.UserData = static_cast<void*>(&PhysicsUserData);
 			}
 
 
@@ -1287,7 +1291,7 @@ void UGeometryCollectionComponent::OnCreatePhysicsState()
 					if (CacheParameters.TargetCache)
 					{
 						// Queue this up to be dirtied after PIE ends
-						FPhysScene_Chaos* Scene = GetPhysicsScene();
+						FPhysScene_Chaos* Scene = GetInnerChaosScene();
 
 						CacheParameters.TargetCache->PreEditChange(nullptr);
 						CacheParameters.TargetCache->Modify();
@@ -1348,8 +1352,17 @@ void UGeometryCollectionComponent::OnCreatePhysicsState()
 			}
 			// end temporary 
 
- 			PhysicsProxy = new FGeometryCollectionPhysicsProxy(this, *DynamicCollection, SimulationParameters, InitFunc, CacheSyncFunc, FinalSyncFunc);
-			FPhysScene_Chaos* Scene = GetPhysicsScene();
+			// Set up initial filter data for our particles
+			FCollisionResponseContainer Response;
+			AActor* ComponentOwner = GetOwner();
+			CreateShapeFilterData(ECC_WorldDynamic, FMaskFilter(0), ComponentOwner ? ComponentOwner->GetUniqueID() : 0, Response, GetUniqueID(), 0, InitialQueryFilter, InitialSimFilter, false, false, false, false);
+
+			// Enable for complex and simple (no dual representation currently like other meshes)
+			InitialQueryFilter.Word3 |= (EPDF_SimpleCollision | EPDF_ComplexCollision);
+			InitialSimFilter.Word3 |= (EPDF_SimpleCollision | EPDF_ComplexCollision);
+
+ 			PhysicsProxy = new FGeometryCollectionPhysicsProxy(this, *DynamicCollection, SimulationParameters, InitialQueryFilter, InitialSimFilter, InitFunc, CacheSyncFunc, FinalSyncFunc);
+			FPhysScene_Chaos* Scene = GetInnerChaosScene();
 			Scene->AddObject(this, PhysicsProxy);
 
 			RegisterForEvents();
@@ -1381,7 +1394,7 @@ void UGeometryCollectionComponent::OnDestroyPhysicsState()
 
 	if(PhysicsProxy)
 	{
-		FPhysScene_Chaos* Scene = GetPhysicsScene();
+		FPhysScene_Chaos* Scene = GetInnerChaosScene();
 		Scene->RemoveObject(PhysicsProxy);
 		InitializationState = ESimulationInitializationState::Unintialized;
 
@@ -1940,8 +1953,7 @@ void UGeometryCollectionComponent::GetInitializationCommands(TArray<FFieldSystem
 	}
 }
 
-
-FPhysScene_Chaos* UGeometryCollectionComponent::GetPhysicsScene() const
+FPhysScene_Chaos* UGeometryCollectionComponent::GetInnerChaosScene() const
 {
 	if (ChaosSolverActor)
 	{
@@ -1971,7 +1983,7 @@ AChaosSolverActor* UGeometryCollectionComponent::GetPhysicsSolverActor() const
 	}
 	else
 	{
-		FPhysScene_Chaos const* const Scene = GetPhysicsScene();
+		FPhysScene_Chaos const* const Scene = GetInnerChaosScene();
 		return Scene ? Cast<AChaosSolverActor>(Scene->GetSolverActor()) : nullptr;
 	}
 
diff --git a/Engine/Source/Runtime/Experimental/GeometryCollectionEngine/Public/GeometryCollection/GeometryCollectionComponent.h b/Engine/Source/Runtime/Experimental/GeometryCollectionEngine/Public/GeometryCollection/GeometryCollectionComponent.h
--- a/Engine/Source/Runtime/Experimental/GeometryCollectionEngine/Public/GeometryCollection/GeometryCollectionComponent.h
+++ b/Engine/Source/Runtime/Experimental/GeometryCollectionEngine/Public/GeometryCollection/GeometryCollectionComponent.h
@@ -483,7 +483,8 @@ public:
 	FORCEINLINE const TArray<int32>& GetHighlightedBones() const { return HighlightedBones; }
 #endif
 
-	FPhysScene_Chaos* GetPhysicsScene() const;
+	FPhysScene_ChaosInterface* GetPhysicsScene() const;
+	FPhysScene_Chaos* GetInnerChaosScene() const;
 	AChaosSolverActor* GetPhysicsSolverActor() const;
 	const FGeometryCollectionPhysicsProxy* GetPhysicsProxy() const { return PhysicsProxy; }
 
@@ -586,6 +587,10 @@ private:
 	bool IsObjectDynamic;
 	bool IsObjectLoading;
 
+	FCollisionFilterData InitialSimFilter;
+	FCollisionFilterData InitialQueryFilter;
+	FPhysxUserData PhysicsUserData;
+
 #if WITH_EDITORONLY_DATA
 	UPROPERTY(Transient)
 	TArray<int32> SelectedBones;
diff --git a/Engine/Source/Runtime/Experimental/GeometryCollectionSimulationCore/Private/GeometryCollection/GeometryCollectionCollisionStructureManager.cpp b/Engine/Source/Runtime/Experimental/GeometryCollectionSimulationCore/Private/GeometryCollection/GeometryCollectionCollisionStructureManager.cpp
--- a/Engine/Source/Runtime/Experimental/GeometryCollectionSimulationCore/Private/GeometryCollection/GeometryCollectionCollisionStructureManager.cpp
+++ b/Engine/Source/Runtime/Experimental/GeometryCollectionSimulationCore/Private/GeometryCollection/GeometryCollectionCollisionStructureManager.cpp
@@ -173,7 +173,7 @@ void FCollisionStructureManager::UpdateImplicitFlags(
 	FImplicit* Implicit, 
 	const ECollisionTypeEnum CollisionType)
 {
-	if (Implicit && (CollisionType == ECollisionTypeEnum::Chaos_Surface_Volumetric))
+	if (Implicit && (CollisionType == ECollisionTypeEnum::Chaos_Surface_Volumetric) && Implicit->GetType() == Chaos::ImplicitObjectType::LevelSet)
 	{
 		Implicit->SetDoCollide(false);
 		Implicit->SetConvex(false);
diff --git a/Engine/Source/Runtime/Experimental/GeometryCollectionSimulationCore/Public/GeometryCollection/GeometryCollectionSimulationCoreTypes.h b/Engine/Source/Runtime/Experimental/GeometryCollectionSimulationCore/Public/GeometryCollection/GeometryCollectionSimulationCoreTypes.h
--- a/Engine/Source/Runtime/Experimental/GeometryCollectionSimulationCore/Public/GeometryCollection/GeometryCollectionSimulationCoreTypes.h
+++ b/Engine/Source/Runtime/Experimental/GeometryCollectionSimulationCore/Public/GeometryCollection/GeometryCollectionSimulationCoreTypes.h
@@ -222,6 +222,9 @@ struct FSimulationParameters
 		, ReverseCacheBeginTime(0.0f)
 		, bClearCache(false)
 		, RemoveOnFractureEnabled(false)
+		, SimulationFilterData()
+		, QueryFilterData()
+		, UserData(nullptr)
 	{}
 
 
@@ -253,6 +256,9 @@ struct FSimulationParameters
 		, TrailingData(Other.TrailingData)
 		, Shared(Other.Shared)
 		, RemoveOnFractureEnabled(false)
+		, SimulationFilterData(Other.SimulationFilterData)
+		, QueryFilterData(Other.QueryFilterData)
+		, UserData(Other.UserData)
 	{}
 
 	~FSimulationParameters()
@@ -308,4 +314,8 @@ struct FSimulationParameters
 	FSharedSimulationParameters Shared;
 
 	bool RemoveOnFractureEnabled;
+
+	FCollisionFilterData SimulationFilterData;
+	FCollisionFilterData QueryFilterData;
+	void* UserData;
 };
