CIS fix

diff --git a/Engine/Source/Runtime/Experimental/GeometryCollectionEngine/Public/GeometryCollection/GeometryCollectionComponent.h b/Engine/Source/Runtime/Experimental/GeometryCollectionEngine/Public/GeometryCollection/GeometryCollectionComponent.h
--- a/Engine/Source/Runtime/Experimental/GeometryCollectionEngine/Public/GeometryCollection/GeometryCollectionComponent.h
+++ b/Engine/Source/Runtime/Experimental/GeometryCollectionEngine/Public/GeometryCollection/GeometryCollectionComponent.h
@@ -483,7 +483,6 @@ public:
 	FORCEINLINE const TArray<int32>& GetHighlightedBones() const { return HighlightedBones; }
 #endif
 
-	FPhysScene_ChaosInterface* GetPhysicsScene() const;
 	FPhysScene_Chaos* GetInnerChaosScene() const;
 	AChaosSolverActor* GetPhysicsSolverActor() const;
 	const FGeometryCollectionPhysicsProxy* GetPhysicsProxy() const { return PhysicsProxy; }
