Chaos: fix Code paths in GC fields to use Handles, replacing old RigidodyId

diff --git a/Engine/Source/Runtime/Experimental/ChaosSolvers/Private/PBDRigidsSolver.cpp b/Engine/Source/Runtime/Experimental/ChaosSolvers/Private/PBDRigidsSolver.cpp
--- a/Engine/Source/Runtime/Experimental/ChaosSolvers/Private/PBDRigidsSolver.cpp
+++ b/Engine/Source/Runtime/Experimental/ChaosSolvers/Private/PBDRigidsSolver.cpp
@@ -74,15 +74,21 @@ namespace Chaos
 				TMap<int32, int32> PositionTargetedParticles;
 				//TArray<FKinematicProxy> AnimatedPositions;
 				Chaos::TArrayCollectionArray<float> Strains;
-				for (FFieldSystemPhysicsProxy* Obj : MSolver->GetFieldSystemPhysicsProxies())
+				for (FFieldSystemPhysicsProxy* FieldObj : MSolver->GetFieldSystemPhysicsProxies())
 				{
 					auto& GeomCollectionParticles = MSolver->GetEvolution()->GetParticles().GetGeometryCollectionParticles();
-					Obj->FieldParameterUpdateCallback(MSolver, GeomCollectionParticles, Strains, 
+					FieldObj->FieldParameterUpdateCallback(MSolver, GeomCollectionParticles, Strains,
 						PositionTarget, PositionTargetedParticles, /*AnimatedPositions,*/ MSolver->GetSolverTime());
 					auto& ClusteredParticles = MSolver->GetEvolution()->GetParticles().GetClusteredParticles();
-					Obj->FieldParameterUpdateCallback(MSolver, ClusteredParticles, Strains,
+					FieldObj->FieldParameterUpdateCallback(MSolver, ClusteredParticles, Strains,
 						PositionTarget, PositionTargetedParticles, /*AnimatedPositions,*/ MSolver->GetSolverTime());
 				}
+
+				for (FGeometryCollectionPhysicsProxy* Obj : MSolver->GetGeometryCollectionPhysicsProxies())
+				{
+					Obj->ParameterUpdateCallback(MSolver->GetEvolution()->GetParticles().GetGeometryCollectionParticles(), MSolver->GetSolverTime());
+				}
+
 			}
 
 			{
@@ -119,6 +125,12 @@ namespace Chaos
 						auto& ClusteredParticles = MSolver->GetEvolution()->GetParticles().GetClusteredParticles();
 						Obj->FieldForcesUpdateCallback(MSolver, ClusteredParticles, Forces, Torques, MSolver->GetSolverTime());
 					}
+
+					for (FGeometryCollectionPhysicsProxy* Obj : MSolver->GetGeometryCollectionPhysicsProxies())
+					{
+						Obj->ParameterUpdateCallback(MSolver->GetEvolution()->GetParticles().GetGeometryCollectionParticles(), MSolver->GetSolverTime());
+					}
+
 					MSolver->GetEvolution()->AdvanceOneTimeStep(DeltaTime);
 				}
 
diff --git a/Engine/Source/Runtime/Experimental/ChaosSolvers/Private/PhysicsProxy/GeometryCollectionPhysicsProxy.cpp b/Engine/Source/Runtime/Experimental/ChaosSolvers/Private/PhysicsProxy/GeometryCollectionPhysicsProxy.cpp
--- a/Engine/Source/Runtime/Experimental/ChaosSolvers/Private/PhysicsProxy/GeometryCollectionPhysicsProxy.cpp
+++ b/Engine/Source/Runtime/Experimental/ChaosSolvers/Private/PhysicsProxy/GeometryCollectionPhysicsProxy.cpp
@@ -1096,11 +1096,54 @@ FGeometryCollectionPhysicsProxy::BuildClusters(
 	return Parent;
 }
 
-void FGeometryCollectionPhysicsProxy::ContiguousIndices(TArray<ContextIndex>& Array, const Chaos::FPhysicsSolver* RigidSolver, EFieldResolutionType ResolutionType, bool bForce = true)
+
+void FGeometryCollectionPhysicsProxy::GetRelevantHandles(
+	TArray<Chaos::TGeometryParticleHandle<float, 3>*>& Handles,
+	TArray<FVector>& Samples,
+	TArray<ContextIndex>& SampleIndices,
+	const Chaos::FPhysicsSolver* RigidSolver, 
+	EFieldResolutionType ResolutionType, 
+	bool bForce = true)
 {
-	check(false); // Not currently called
 	if (bForce)
 	{
+		// only the local handles
+		TArray<FClusterHandle*>& ParticleHandles = GetSolverParticleHandles();
+
+		if (ResolutionType == EFieldResolutionType::Field_Resolution_Maximum)
+		{
+			Handles.SetNumUninitialized(ParticleHandles.Num());
+
+			int32 NumUsedHandles = 0;
+			for (FClusterHandle* ClusterHandle : ParticleHandles)
+			{
+				if (!ClusterHandle->Disabled())
+				{
+					Handles.Add(ClusterHandle);
+					NumUsedHandles++;
+				}
+			}
+			Handles.SetNum(NumUsedHandles);
+
+		}
+		else if (ResolutionType == EFieldResolutionType::Field_Resolution_Minimal)
+		{
+			Handles.Reserve(ParticleHandles.Num());
+
+			for (FClusterHandle* ClusterHandle : ParticleHandles)
+			{
+				Handles.Add(ClusterHandle);
+			}
+		}
+
+		Samples.AddUninitialized(Handles.Num());
+		SampleIndices.AddUninitialized(Handles.Num());
+		for (int32 Idx = 0; Idx < Handles.Num(); ++Idx)
+		{
+			Samples[Idx] = Handles[Idx]->X();
+			SampleIndices[Idx] = ContextIndex(Idx, Idx);
+		}
+
 #if TODO_REIMPLEMENT_GET_RIGID_PARTICLES
 		const Chaos::FPhysicsSolver::FParticlesType & Particles = RigidSolver->GetRigidParticles();
 		if (ResolutionType == EFieldResolutionType::Field_Resolution_Minimal)
@@ -1151,6 +1194,59 @@ void FGeometryCollectionPhysicsProxy::ContiguousIndices(TArray<ContextIndex>& Ar
 	}
 }
 
+void FGeometryCollectionPhysicsProxy::PushKinematicStateToSolver(FParticlesType& Particles)
+{
+	FGeometryDynamicCollection& Collection = GameThreadCollection;
+	if (Collection.Transform.Num())
+	{
+		TManagedArray<int32>& DynamicState = Collection.DynamicState;
+
+		for (int32 TransformIndex = 0; TransformIndex < DynamicState.Num(); TransformIndex++)
+		{
+			Chaos::TPBDRigidClusteredParticleHandle<float, 3>* Handle = SolverParticleHandles[TransformIndex];
+			if (!Handle)
+			{
+				continue;
+			}
+	
+			if (DynamicState[TransformIndex] == (int32)EObjectStateTypeEnum::Chaos_Object_Dynamic
+				&& (Handle->ObjectState() == Chaos::EObjectStateType::Kinematic || Handle->ObjectState() == Chaos::EObjectStateType::Static)
+				&& FLT_EPSILON < Handle->M())
+			{
+				Handle->SetObjectState(Chaos::EObjectStateType::Dynamic);
+				if (Parameters.InitialVelocityType == EInitialVelocityTypeEnum::Chaos_Initial_Velocity_User_Defined)
+				{
+					Handle->SetV(Collection.InitialLinearVelocity[TransformIndex]);
+					Handle->SetW(Collection.InitialAngularVelocity[TransformIndex]);
+				}
+			}
+			else if ((DynamicState[TransformIndex] == (int32)EObjectStateTypeEnum::Chaos_Object_Kinematic)
+				&& (Handle->ObjectState() == Chaos::EObjectStateType::Dynamic)
+				&& FLT_EPSILON < Handle->M())
+			{
+				Handle->SetObjectState(Chaos::EObjectStateType::Kinematic);
+			}
+			else if ((DynamicState[TransformIndex] == (int32)EObjectStateTypeEnum::Chaos_Object_Static)
+				&& (Handle->ObjectState() == Chaos::EObjectStateType::Dynamic)
+				&& FLT_EPSILON < Handle->M())
+			{
+				Handle->SetObjectState(Chaos::EObjectStateType::Static);
+			}
+			else if ((DynamicState[TransformIndex] == (int32)EObjectStateTypeEnum::Chaos_Object_Sleeping)
+				&& (Handle->ObjectState() == Chaos::EObjectStateType::Dynamic))
+			{
+				Handle->SetObjectState(Chaos::EObjectStateType::Sleeping);
+			}
+			else if ((DynamicState[TransformIndex] == (int32)EObjectStateTypeEnum::Chaos_Object_Dynamic)
+				&& (Handle->ObjectState() == Chaos::EObjectStateType::Sleeping))
+			{
+				Handle->SetObjectState(Chaos::EObjectStateType::Dynamic);
+			}
+		
+		}
+	}
+}
+
 int32 FGeometryCollectionPhysicsProxy::CalculateHierarchyLevel(const FGeometryDynamicCollection& GeometryCollection, int32 TransformIndex) const
 {
 	int32 Level = 0;
@@ -1724,6 +1820,8 @@ void FGeometryCollectionPhysicsProxy::PullFromPhysicsState()
 				GTParticles[TmIndex]->SetX(ParticleToWorld.GetTranslation());
 				GTParticles[TmIndex]->SetR(ParticleToWorld.GetRotation());
 			}
+
+			DynamicCollection.DynamicState[TmIndex] = TR.DynamicState[TmIndex];
 		}
 
 		//question: why do we need this? Sleeping objects will always have to update GPU
@@ -2491,19 +2589,20 @@ void BuildSimulationData(Chaos::FErrorReporter& ErrorReporter, FGeometryCollecti
 
 void FGeometryCollectionPhysicsProxy::ProcessCommands(FParticlesType& Particles, const float Time)
 {
-	//const FGeometryDynamicCollection& Collection = *GameToPhysInterchange.AccessProducerBuffer();
-	//check(Collection);
+	FGeometryDynamicCollection& Collection = GameThreadCollection;
 
 	// Process Particle-Collection commands
 	if (Commands.Num())
 	{
-		/*
+		TArray<Chaos::TGeometryParticleHandle<float, 3>*> Handles;
+		TArray<FVector> Samples;
+		TArray<ContextIndex> SampleIndices;
+
 		TArray<ContextIndex> IndicesArray;
 		Chaos::FPhysicsSolver* CurrentSolver = GetSolver();
 
 		for (int32 CommandIndex = Commands.Num()-1; 0<=CommandIndex; CommandIndex--)
 		{
-
 			//
 			// Extract command and set metadata
 			//
@@ -2520,24 +2619,21 @@ void FGeometryCollectionPhysicsProxy::ProcessCommands(FParticlesType& Particles,
 				if (ensureMsgf(Command.RootNode->Type() == FFieldNode<int32>::StaticType(),
 					TEXT("Field based evaluation of the simulations 'DynamicState' parameter expects int32 field inputs.")))
 				{
-					FGeometryCollectionPhysicsProxy::ContiguousIndices(IndicesArray, CurrentSolver, ResolutionType, IndicesArray.Num() != Particles.Size());
-					if (IndicesArray.Num())
-					{
-						TArrayView<ContextIndex> IndexView(&(IndicesArray[0]), IndicesArray.Num());
-
-						FVector * xptr = &(Particles.X(0));
-						TArrayView<FVector> SamplesView(xptr, Particles.Size());
 
+					FGeometryCollectionPhysicsProxy::GetRelevantHandles(Handles, Samples, SampleIndices, CurrentSolver, ResolutionType, IndicesArray.Num() != Particles.Size());
+					if (Handles.Num())
+					{
+						TArrayView<FVector> SamplesView(&(Samples[0]), Samples.Num());
+						TArrayView<ContextIndex> SampleIndicesView(&(SampleIndices[0]), SampleIndices.Num());
 						FFieldContext Context{
-							IndexView,
+							SampleIndicesView,
 							SamplesView,
-							Command.MetaData
-						};
+							Command.MetaData };
 
-						TArrayView<int32> DynamicStateView(&(Collection->DynamicState[0]), Collection->DynamicState.Num());
+						TArrayView<int32> DynamicStateView(&(Collection.DynamicState[0]), Collection.DynamicState.Num());
 						static_cast<const FFieldNode<int32> *>(Command.RootNode.Get())->Evaluate(Context, DynamicStateView);
 
-						//PushKinematicStateToSolver(Particles);
+						PushKinematicStateToSolver(Particles);
 					}
 				}
 				Commands.RemoveAt(CommandIndex);
@@ -2550,21 +2646,17 @@ void FGeometryCollectionPhysicsProxy::ProcessCommands(FParticlesType& Particles,
 					if (ensureMsgf(Command.RootNode->Type() == FFieldNode<FVector>::StaticType(),
 						TEXT("Field based evaluation of the simulations 'InitialLinearVelocity' parameter expects FVector field inputs.")))
 					{
-						FGeometryCollectionPhysicsProxy::ContiguousIndices(IndicesArray, CurrentSolver, ResolutionType, IndicesArray.Num() != Particles.Size());
-						if (IndicesArray.Num())
+						FGeometryCollectionPhysicsProxy::GetRelevantHandles(Handles, Samples, SampleIndices, CurrentSolver, ResolutionType, IndicesArray.Num() != Particles.Size());
+						if (Handles.Num())
 						{
-							TArrayView<ContextIndex> IndexView(&(IndicesArray[0]), IndicesArray.Num());
-
-							FVector * xptr = &(Particles.X(0));
-							TArrayView<FVector> SamplesView(xptr, Particles.Size());
-
+							TArrayView<FVector> SamplesView(&(Samples[0]), Samples.Num());
+							TArrayView<ContextIndex> SampleIndicesView(&(SampleIndices[0]), SampleIndices.Num());
 							FFieldContext Context{
-								IndexView,
+								SampleIndicesView,
 								SamplesView,
-								Command.MetaData
-							};
+								Command.MetaData };
 
-							TArrayView<FVector> ResultsView(&(Collection->InitialLinearVelocity[0]), Collection->InitialLinearVelocity.Num());
+							TArrayView<FVector> ResultsView(&(Collection.InitialLinearVelocity[0]), Collection.InitialLinearVelocity.Num());
 							static_cast<const FFieldNode<FVector> *>(Command.RootNode.Get())->Evaluate(Context, ResultsView);
 						}
 					}
@@ -2579,29 +2671,26 @@ void FGeometryCollectionPhysicsProxy::ProcessCommands(FParticlesType& Particles,
 					if (ensureMsgf(Command.RootNode->Type() == FFieldNode<FVector>::StaticType(),
 						TEXT("Field based evaluation of the simulations 'InitialAngularVelocity' parameter expects FVector field inputs.")))
 					{
-						FGeometryCollectionPhysicsProxy::ContiguousIndices(IndicesArray, CurrentSolver, ResolutionType, IndicesArray.Num() != Particles.Size());
-						if (IndicesArray.Num())
+						FGeometryCollectionPhysicsProxy::GetRelevantHandles(Handles, Samples, SampleIndices, CurrentSolver, ResolutionType, IndicesArray.Num() != Particles.Size());
+						if (Handles.Num())
 						{
-							TArrayView<ContextIndex> IndexView(&(IndicesArray[0]), IndicesArray.Num());
-
-							FVector * xptr = &(Particles.X(0));
-							TArrayView<FVector> SamplesView(xptr, Particles.Size());
-
+							TArrayView<FVector> SamplesView(&(Samples[0]), Samples.Num());
+							TArrayView<ContextIndex> SampleIndicesView(&(SampleIndices[0]), SampleIndices.Num());
 							FFieldContext Context{
-								IndexView,
+								SampleIndicesView,
 								SamplesView,
-								Command.MetaData
-							};
+								Command.MetaData };
 
-							TArrayView<FVector> ResultsView(&(Collection->InitialAngularVelocity[0]), Collection->InitialAngularVelocity.Num());
+							TArrayView<FVector> ResultsView(&(Collection.InitialAngularVelocity[0]), Collection.InitialAngularVelocity.Num());
 							static_cast<const FFieldNode<FVector> *>(Command.RootNode.Get())->Evaluate(Context, ResultsView);
 						}
 					}
 				}
 				Commands.RemoveAt(CommandIndex);
+
 			}
+
 		}
-		*/
 	}
 
 	// Process Particle-Particle commands
@@ -2609,24 +2698,11 @@ void FGeometryCollectionPhysicsProxy::ProcessCommands(FParticlesType& Particles,
 	{
 		Chaos::FPhysicsSolver* CurrentSolver = GetSolver();
 
-		/*
-		@todo break everything
+		TArray<Chaos::TGeometryParticleHandle<float, 3>*> Handles;
+		TArray<FVector> Samples;
+		TArray<ContextIndex> SampleIndices;
 
-		//  Generate a Index mapping between the rigid body indices and 
-		//  the particle indices. This allows the geometry collection to
-		//  evaluate only its own particles. 
-		TArray<ContextIndex> IndicesArray;
-		int32 NumIndices = 0;
-		IndicesArray.SetNumUninitialized(RigidBodyID.Num());
-		for (int32 i = 0; i < RigidBodyID.Num(); i++)
-		{
-			if (RigidBodyID[i] != INDEX_NONE)
-			{
-				IndicesArray[NumIndices] = { RigidBodyID[i], RigidBodyID[i] };
-				NumIndices++;
-			}
-		}
-		IndicesArray.SetNum(NumIndices);
+		FGeometryCollectionPhysicsProxy::GetRelevantHandles(Handles, Samples, SampleIndices, CurrentSolver, EFieldResolutionType::Field_Resolution_Maximum, true);
 
 		for (int32 CommandIndex = Commands.Num() - 1; 0 <= CommandIndex; CommandIndex--)
 		{
@@ -2636,15 +2712,12 @@ void FGeometryCollectionPhysicsProxy::ProcessCommands(FParticlesType& Particles,
 				if (ensureMsgf(Command.RootNode->Type() == FFieldNode<FVector>::StaticType(),
 					TEXT("Field based evaluation of the simulations 'LinearVelocity' parameter expects FVector field inputs.")))
 				{
-					FVector * xptr = &(Particles.X(0));
-					TArrayView<FVector> SamplesView(xptr, Particles.Size());
-					TArrayView<ContextIndex> IndexView(&(IndicesArray[0]), IndicesArray.Num());
-
+					TArrayView<FVector> SamplesView(&(Samples[0]), Samples.Num());
+					TArrayView<ContextIndex> SampleIndicesView(&(SampleIndices[0]), SampleIndices.Num());
 					FFieldContext Context{
-						IndexView,
+						SampleIndicesView,
 						SamplesView,
-						Command.MetaData
-					};
+						Command.MetaData };
 
 					FVector * vptr = &(Particles.V(0));
 					TArrayView<FVector> ResultsView(vptr, Particles.Size());
@@ -2658,15 +2731,12 @@ void FGeometryCollectionPhysicsProxy::ProcessCommands(FParticlesType& Particles,
 					TEXT("Field based evaluation of the simulations 'AngularVelocity' parameter expects FVector field inputs.")))
 				{
 
-					FVector * xptr = &(Particles.X(0));
-					TArrayView<FVector> SamplesView(xptr, Particles.Size());
-					TArrayView<ContextIndex> IndexView(&(IndicesArray[0]), IndicesArray.Num());
-
+					TArrayView<FVector> SamplesView(&(Samples[0]), Samples.Num());
+					TArrayView<ContextIndex> SampleIndicesView(&(SampleIndices[0]), SampleIndices.Num());
 					FFieldContext Context{
-						IndexView,
+						SampleIndicesView,
 						SamplesView,
-						Command.MetaData
-					};
+						Command.MetaData };
 
 					FVector * vptr = &(Particles.W(0));
 					TArrayView<FVector> ResultsView(vptr, Particles.Size());
@@ -2679,15 +2749,12 @@ void FGeometryCollectionPhysicsProxy::ProcessCommands(FParticlesType& Particles,
 				if (ensureMsgf(Command.RootNode->Type() == FFieldNode<int32>::StaticType(),
 					TEXT("Field based evaluation of the simulations 'CollisionGroup' parameter expects int32 field inputs.")))
 				{
-					FVector * xptr = &(Particles.X(0));
-					TArrayView<FVector> SamplesView(xptr, Particles.Size());
-					TArrayView<ContextIndex> IndexView(&(IndicesArray[0]), IndicesArray.Num());
-
+					TArrayView<FVector> SamplesView(&(Samples[0]), Samples.Num());
+					TArrayView<ContextIndex> SampleIndicesView(&(SampleIndices[0]), SampleIndices.Num());
 					FFieldContext Context{
-						IndexView,
+						SampleIndicesView,
 						SamplesView,
-						Command.MetaData
-					};
+						Command.MetaData };
 
 					int32 * cptr = &(Particles.CollisionGroup(0));
 					TArrayView<int32> ResultsView(cptr, Particles.Size());
@@ -2696,33 +2763,20 @@ void FGeometryCollectionPhysicsProxy::ProcessCommands(FParticlesType& Particles,
 				Commands.RemoveAt(CommandIndex);
 			}
 		}
-		*/
+	
 	}
 }
 
+
 void FGeometryCollectionPhysicsProxy::FieldForcesUpdateCallback(Chaos::FPhysicsSolver* InSolver, FParticlesType& Particles, Chaos::TArrayCollectionArray<FVector> & Force, Chaos::TArrayCollectionArray<FVector> & Torque, const float Time)
 {
-	/*
-
-	@todo break everything : re-enable fields without rigid body ids. 
-
 	if (Commands.Num())
 		{
-			// @todo: This seems like a waste if we just want to get everything
-			int32 Counter = 0;
-			TArray<ContextIndex> IndicesArray;
-			IndicesArray.AddUninitialized(RigidBodyID.Num());
-			for (int32 i = 0; i < RigidBodyID.Num(); i++)
-			{
-				if (RigidBodyID[i] != INDEX_NONE)
-				{
-					IndicesArray[i] = { RigidBodyID[i],RigidBodyID[i] };
-					Counter++;
-				}
-			}
-			IndicesArray.SetNum(Counter, false);
-			TArrayView<ContextIndex> IndexView(&(IndicesArray[0]), IndicesArray.Num());
+		TArray<Chaos::TGeometryParticleHandle<float, 3>*> Handles;
+		TArray<FVector> Samples;
+		TArray<ContextIndex> SampleIndices;
 
+		FGeometryCollectionPhysicsProxy::GetRelevantHandles(Handles, Samples, SampleIndices, InSolver, EFieldResolutionType::Field_Resolution_Maximum, true);
 
 			for (int32 CommandIndex = 0; CommandIndex < Commands.Num(); CommandIndex++)
 			{
@@ -2733,14 +2787,13 @@ void FGeometryCollectionPhysicsProxy::FieldForcesUpdateCallback(Chaos::FPhysicsS
 					if (ensureMsgf(Command.RootNode->Type() == FFieldNode<FVector>::StaticType(),
 						TEXT("Field based evaluation of the simulations 'LinearForce' parameter expects FVector field inputs.")))
 					{
-						FVector * tptr = &(Particles.X(0));
-						TArrayView<FVector> SamplesView(tptr, int32(Particles.Size()));
-
+					TArrayView<FVector> SamplesView(&(Samples[0]), Samples.Num());
+					TArrayView<ContextIndex> SampleIndicesView(&(SampleIndices[0]), SampleIndices.Num());
 						FFieldContext Context{
-							IndexView,
+						SampleIndicesView,
 							SamplesView,
-							Command.MetaData
-						};
+						Command.MetaData };
+
 						TArrayView<FVector> ForceView(&(Force[0]), Force.Num());
 						static_cast<const FFieldNode<FVector> *>(Command.RootNode.Get())->Evaluate(Context, ForceView);
 					}
@@ -2751,14 +2804,13 @@ void FGeometryCollectionPhysicsProxy::FieldForcesUpdateCallback(Chaos::FPhysicsS
 					if (ensureMsgf(Command.RootNode->Type() == FFieldNode<FVector>::StaticType(),
 						TEXT("Field based evaluation of the simulations 'AngularTorque' parameter expects FVector field inputs.")))
 					{
-						FVector * tptr = &(Particles.X(0));
-						TArrayView<FVector> SamplesView(tptr, int32(Particles.Size()));
-
+					TArrayView<FVector> SamplesView(&(Samples[0]), Samples.Num());
+					TArrayView<ContextIndex> SampleIndicesView(&(SampleIndices[0]), SampleIndices.Num());
 						FFieldContext Context{
-							IndexView,
+						SampleIndicesView,
 							SamplesView,
-							Command.MetaData
-						};
+						Command.MetaData };
+
 						TArrayView<FVector> TorqueView(&(Torque[0]), Torque.Num());
 						static_cast<const FFieldNode<FVector> *>(Command.RootNode.Get())->Evaluate(Context, TorqueView);
 					}
@@ -2767,12 +2819,8 @@ void FGeometryCollectionPhysicsProxy::FieldForcesUpdateCallback(Chaos::FPhysicsS
 
 			}
 		}
-		*/
 }
 
-
-
-
 void FGeometryCollectionPhysicsProxy::ParameterUpdateCallback(FParticlesType& Particles, const float Time)
 {
 	if (GameThreadCollection.Transform.Num())
diff --git a/Engine/Source/Runtime/Experimental/ChaosSolvers/Public/PBDRigidsSolver.h b/Engine/Source/Runtime/Experimental/ChaosSolvers/Public/PBDRigidsSolver.h
--- a/Engine/Source/Runtime/Experimental/ChaosSolvers/Public/PBDRigidsSolver.h
+++ b/Engine/Source/Runtime/Experimental/ChaosSolvers/Public/PBDRigidsSolver.h
@@ -300,6 +300,11 @@ namespace Chaos
 			return FieldSystemPhysicsProxies;
 		}
 
+		TArray<FGeometryCollectionPhysicsProxy*>& GetGeometryCollectionPhysicsProxies()
+		{
+			return GeometryCollectionPhysicsProxies;
+		}
+
 		/** Events hooked up to the Chaos material manager */
 		void UpdateMaterial(Chaos::FMaterialHandle InHandle, const Chaos::FChaosPhysicsMaterial& InNewData);
 		void CreateMaterial(Chaos::FMaterialHandle InHandle, const Chaos::FChaosPhysicsMaterial& InNewData);
diff --git a/Engine/Source/Runtime/Experimental/ChaosSolvers/Public/PhysicsProxy/GeometryCollectionPhysicsProxy.h b/Engine/Source/Runtime/Experimental/ChaosSolvers/Public/PhysicsProxy/GeometryCollectionPhysicsProxy.h
--- a/Engine/Source/Runtime/Experimental/ChaosSolvers/Public/PhysicsProxy/GeometryCollectionPhysicsProxy.h
+++ b/Engine/Source/Runtime/Experimental/ChaosSolvers/Public/PhysicsProxy/GeometryCollectionPhysicsProxy.h
@@ -256,18 +256,17 @@ protected:
 		FGeometryDynamicCollection* TargetCollection=nullptr);
 
 	/**
-	 * Generates a mapping between the Position array and the results array. 
-	 * When EFieldResolutionType is set to Maximum the complete particle mapping 
-	 * is provided from the Particles.X to Particles.Attribute, when Minimum is 
-	 * set only the ActiveIndices and the direct children of the active clusters 
-	 * are set in the IndicesArray.
 	 */
-	void ContiguousIndices(
-		TArray<ContextIndex>& IndicesArray, 
+	void GetRelevantHandles(
+		TArray<Chaos::TGeometryParticleHandle<float, 3>*>& Handles,
+		TArray<FVector>& Samples,
+		TArray<ContextIndex>& SampleIndices,
 		const Chaos::FPhysicsSolver* RigidSolver, 
 		EFieldResolutionType ResolutionType, 
 		bool bForce);
 
+	void PushKinematicStateToSolver(FParticlesType& Particles);
+
 	/** 
 	 * Traverses the parents of \p TransformIndex in \p GeometryCollection, counting
 	 * the number of levels until the next parent is \c INDEX_NONE.
