Chaos : Desctuction

- Clearing out the per-frame update calls.

diff --git a/Engine/Source/Runtime/Experimental/ChaosSolvers/Private/PhysicsProxy/GeometryCollectionPhysicsProxy.cpp b/Engine/Source/Runtime/Experimental/ChaosSolvers/Private/PhysicsProxy/GeometryCollectionPhysicsProxy.cpp
--- a/Engine/Source/Runtime/Experimental/ChaosSolvers/Private/PhysicsProxy/GeometryCollectionPhysicsProxy.cpp
+++ b/Engine/Source/Runtime/Experimental/ChaosSolvers/Private/PhysicsProxy/GeometryCollectionPhysicsProxy.cpp
@@ -937,6 +937,26 @@ void FGeometryCollectionPhysicsProxy::InitializeBodiesPT(
 		}
 #endif // TODO_REIMPLEMENT_RIGID_CACHING
 
+
+
+		if (DisableGeometryCollectionGravity) // cvar
+		{
+			// Our assumption is that you'd only ever want to wholesale opt geometry 
+			// collections out of gravity for debugging, so we keep this conditional
+			// out of the loop above and on it's own.  This means we can't turn gravity
+			// back on once it's off, but even if we didn't enclose this in an if(),
+			// this function won't be called again unless something dirties the proxy.
+
+			Chaos::TPerParticleGravity<float, 3>& GravityForces = GetSolver()->GetEvolution()->GetGravityForces();
+			for (int32 HandleIdx = 0; HandleIdx < SolverParticleHandles.Num(); ++HandleIdx)
+			{
+				if (Chaos::TPBDRigidParticleHandle<float, 3>* Handle = SolverParticleHandles[HandleIdx])
+				{
+					GravityForces.SetEnabled(*Handle, false);
+				}
+			}
+		}
+
 	} // end if simulating...
 
 }
@@ -1336,193 +1356,21 @@ void FGeometryCollectionPhysicsProxy::SyncBeforeDestroy()
 void FGeometryCollectionPhysicsProxy::BufferGameState() 
 {
 	//
-	// NewData() redirects to this function.
-	/**
-	 * CONTEXT: GAMETHREAD
-	 * Returns a new unmanaged allocation of the data saved on the handle, otherwise nullptr
-	 */
-
-	// The other physics proxies allocate memory and returns a copy of game thread
-	// data as a TGeometryParticleData<T,d>, TKinematicGeometryParticleData<T,d>, 
-	// or a TPBDRigidParticleData<T,d>, which are all derived from FParticleData.
-	// That then gets passed to the physics thread via PushToPhysicsState().
-	//
-	// So all we're doing is pushing game thread state to the physics thread.
-	// I don't think we need to allocate memory to do this.  I think the other 
-	// proxies get away with that because they're dealing with a low number of 
-	// bodies; which is not the case here. Rather, let's just use a triple buffer.
-/*
-	if (!Parameters.Simulating || 
-		!ensure(NumParticles != INDEX_NONE)) // Make sure InitBodiesGT() has been called!
-	{
-		return;
-	}
-
-	FGeometryCollectionResults& Buffer = *GameToPhysInterchange.AccessProducerBuffer();
-	if (Buffer.NumParticlesAdded != NumParticles)
-	{
-		Buffer.InitArrays(*GTDynamicCollection);
-
-		// Initialize the buffer
-		Buffer.NumParticlesAdded = NumParticles;
-		Buffer.BaseIndex = BaseParticleIndex;
-
-		Buffer.SharedGeometry.SetNum(NumParticles);
-		Buffer.ShapeSimData.SetNum(NumParticles);
-		Buffer.ShapeQueryData.SetNum(NumParticles);
-
-		for (int32 Idx = 0; Idx < SimulatableParticles.Num(); Idx++)
-		{
-			Buffer.DisabledStates[Idx] = !SimulatableParticles[Idx];
-			Buffer.SharedGeometry[Idx] = Implicits[Idx];
-
-
-			// Hard coded to collide all
-			// todo(chaos) : Pull shape data from the interface. 
-			for (int32 Gdx = 0; Gdx<Buffer.SharedGeometry.Num();Gdx++)
-			{
-				int32 Idx1 = Buffer.ShapeQueryData[Idx].Add(FCollisionFilterData());
-				int32 Idx2 = Buffer.ShapeSimData[Idx].Add(FCollisionFilterData());
-				ensure(Idx1 == Idx2);
-
-				Buffer.ShapeSimData[Idx][Idx1].Word1 = 0xFFFF; // this body channel
-				Buffer.ShapeSimData[Idx][Idx1].Word3 = 0xFFFF; // collision candidate channels
-			}
-		}
-	}
-
-	const FGeometryDynamicCollection* DynamicCollection = Parameters.DynamicCollection;
-	const FGeometryCollection* RestCollection = Parameters.RestCollection;
-
-	//Buffer.TransformIndex.Init(RestCollection->TransformIndex);
-//	Buffer.BoneMap.Init(RestCollection->BoneMap);
-//	Buffer.Parent.Init(RestCollection->Parent);
-//	Buffer.Children.Init(RestCollection->Children);
-//	Buffer.SimulationType.Init(RestCollection->SimulationType);
-
-//	Buffer.DynamicState.Init(DynamicCollection->DynamicState);
-//	Buffer.Mass.Init(RestCollection->GetAttribute<float>("Mass", FTransformCollection::TransformGroup));
-//	Buffer.InertiaTensor.Init(RestCollection->GetAttribute<FVector>("InertiaTensor", FTransformCollection::TransformGroup));
-	Buffer.Transforms.Init(DynamicCollection->Transform);
-
-	GeometryCollectionAlgo::GlobalMatrices(
-		DynamicCollection->Transform, DynamicCollection->Parent, Buffer.GlobalTransforms);
-	
-	// Make the buffer available to the consumer.
-	GameToPhysInterchange.FlipProducer(); 
-	
-	// If we wanted to emulate the other proxies, we could augment the 
-	// FStubGeometryCollectionData class to carry along the consumer buffer:
-	//
-	// return new FStubGeometryCollectionData(GametoPhysInterchange.GetConsumerBuffer());
+	// There is currently no per advance updates to the GeometryCollection
 	//
-	// However, at this point, it seems pretty unnecessary to do the additional
-	// allocation and deallocation of the transient FStubGeometryCollectionData
-	// class, as I believe all these operations are happening with the same 
-	// instance of this class.  So skip all that, and just return null.
-	return;
-	*/
 }
 
 void FGeometryCollectionPhysicsProxy::PushToPhysicsState(const Chaos::FParticleData *InData)
 {
-	/**
-	 * CONTEXT: GAMETHREAD -> to -> PHYSICSTHREAD
-	 * Called on the game thread when the solver is about to advance forward. This
-	 * callback should Enqueue commands on the PhysicsThread to update the state of
-	 * the solver
-	 */
-
-	/*
-	@todo break everything : move to the init
-	// The other proxies do a deep copy into dynamic transient memory, passed into
-	// this function via InData.  Instead of that, we're using persistent memory
-	// in a buffer to facilitate communicating between threads.
-
-	const FGeometryCollectionResults* GState = GameToPhysInterchange.GetConsumerBuffer();
-	if (!GState)
-	{
-		// This will happen if GameToPhysInterchange hasn't been updated since
-		// the last time GetConsumerBuffer() was called. 
-		return;
-	}
-
-	//const FStubGeometryCollectionData* GCData = static_cast<const FStubGeometryCollectionData*>(InData);
-	//const FGeometryCollectionResults* GCState = GCData->GetStateData();
-
-	// Move the copied game thread data into the handles.
-	for (int32 HandleIdx = 0; HandleIdx < SolverParticleHandles.Num(); ++HandleIdx)
-	{
-		if (Chaos::TPBDRigidParticleHandle<float, 3>* Handle = SolverParticleHandles[HandleIdx])
-		{
-			const int32 Idx = GState->BaseIndex + HandleIdx;
-
-			FTransform ParticleTransform =
-				MassToLocal[Idx] * GState->Transforms[Idx] * Parameters.WorldTransform;
-
-			Handle->SetX(ParticleTransform.GetTranslation());
-			Handle->SetR(ParticleTransform.GetRotation());
-
-			// Particle disabled states are initialized on the physics thread, so 
-			// we don't want to pull this from the game thread
-			//Handle->SetDisabled(GState->DisabledStates[Idx]); ryan - fix this!
-
-			Handle->SetLinearEtherDrag(Parameters.LinearEtherDrag);
-			Handle->SetAngularEtherDrag(Parameters.AngularEtherDrag);
-
-			if (GState->SharedGeometry[Idx] && GState->SharedGeometry[Idx]->HasBoundingBox())
-			{
-				Handle->SetHasBounds(true);
-				Handle->SetLocalBounds(GState->SharedGeometry[Idx]->BoundingBox());
-
-				Chaos::TAABB<float, 3> WorldSpaceBox =
-					GState->SharedGeometry[Idx]->BoundingBox().TransformedAABB(
-						Chaos::TRigidTransform<float, 3>(Handle->X(), Handle->R()));
-				// TODO Ryan - WorldSpaceBox.ThickenSymmetrically(GState->V());
-				Handle->SetWorldSpaceInflatedBounds(WorldSpaceBox);
-			}
-
-			// TODO Ryan - check dirty flags like TGeometryParticleData::DirtyFlags.
-			// In the mean time, we only set once.
-			if (!Handle->SharedGeometry() && GState->SharedGeometry.IsValidIndex(Idx) && GState->SharedGeometry[Idx])
-			{
-				Handle->SetSharedGeometry(GState->SharedGeometry[Idx]);
-			}
-
-			int32 CurrShape = 0;
-			for (const TUniquePtr<Chaos::TPerShapeData<Chaos::FReal, 3>>& Shape : Handle->ShapesArray())
-			{
-				// TODO Ryan - hardwired for now:
-				Shape->bDisable = false;
-				Shape->CollisionTraceType = Chaos::EChaosCollisionTraceFlag::Chaos_CTF_UseSimpleAndComplex; // Chaos_CTF_UseDefault?
-				if (CurrShape < GState->ShapeSimData[Idx].Num())
-				{
-					Shape->SimData = GState->ShapeSimData[Idx][CurrShape];
-					Shape->QueryData = GState->ShapeQueryData[Idx][CurrShape];
-				}
-				++CurrShape;
-			}
-		} // if rigid handle
-	} // for all SolverParticleHandles
+	/*CONTEXT: GAMETHREAD->to->PHYSICSTHREAD
+	* Called on the game thread when the solver is about to advance forward.This
+	* callback should Enqueue commands on the PhysicsThread to update the state of
+	* the solver
 	*/
 
-	if (DisableGeometryCollectionGravity) // cvar
-	{
-		// Our assumption is that you'd only ever want to wholesale opt geometry 
-		// collections out of gravity for debugging, so we keep this conditional
-		// out of the loop above and on it's own.  This means we can't turn gravity
-		// back on once it's off, but even if we didn't enclose this in an if(),
-		// this function won't be called again unless something dirties the proxy.
-
-		Chaos::TPerParticleGravity<float, 3>& GravityForces = GetSolver()->GetEvolution()->GetGravityForces();
-		for (int32 HandleIdx = 0; HandleIdx < SolverParticleHandles.Num(); ++HandleIdx)
-		{
-			if (Chaos::TPBDRigidParticleHandle<float, 3>* Handle = SolverParticleHandles[HandleIdx])
-			{
-				GravityForces.SetEnabled(*Handle, false);
-			}
-		}
-	}
+	//
+	// There is currently no per advance updates to the GeometryCollection
+	//
 }
 
 void FGeometryCollectionPhysicsProxy::BufferPhysicsResults()
