Chaos :

- Moved Results to Data.h

diff --git a/Engine/Source/Runtime/Experimental/ChaosSolvers/Public/GeometryCollectionProxyData.h b/Engine/Source/Runtime/Experimental/ChaosSolvers/Public/GeometryCollectionProxyData.h
--- a/Engine/Source/Runtime/Experimental/ChaosSolvers/Public/GeometryCollectionProxyData.h
+++ b/Engine/Source/Runtime/Experimental/ChaosSolvers/Public/GeometryCollectionProxyData.h
@@ -87,4 +87,61 @@ public:
 	TManagedArray<bool> SimulatableParticles;
 
 
-};
\ No newline at end of file
+};
+
+
+/**
+ * Buffer structure for communicating simulation state between game and physics
+ * threads.
+ */
+class FGeometryCollectionResults
+{
+public:
+	FGeometryCollectionResults();
+
+	void Reset();
+
+	int32 NumTransformGroup() const { return Transforms.Num(); }
+
+	void InitArrays(const FGeometryDynamicCollection& Other)
+	{
+		// Managed arrays
+		Transforms.Init(Other.Transform);
+		DynamicState.Init(Other.DynamicState);
+		Parent.Init(Other.Parent);
+		Children.Init(Other.Children);
+		SimulationType.Init(Other.SimulationType);
+
+		// Arrays
+		const int32 NumTransforms = Other.NumElements(FGeometryCollection::TransformGroup);
+		DisabledStates.SetNumUninitialized(NumTransforms);
+		GlobalTransforms.SetNumUninitialized(NumTransforms);
+		ParticleToWorldTransforms.SetNumUninitialized(NumTransforms);
+	}
+
+	int32 BaseIndex;
+	int32 NumParticlesAdded;
+	TArray<bool> DisabledStates;
+	TArray<FMatrix> GlobalTransforms;
+	TArray<FTransform> ParticleToWorldTransforms;
+
+	TManagedArray<int32> TransformIndex;
+
+	TManagedArray<FTransform> Transforms;
+	TManagedArray<int32> BoneMap;
+	TManagedArray<int32> Parent;
+	TManagedArray<TSet<int32>> Children;
+	TManagedArray<int32> SimulationType;
+	TManagedArray<int32> DynamicState;
+	TManagedArray<float> Mass;
+	TManagedArray<FVector> InertiaTensor;
+
+	TManagedArray<int32> ClusterId;
+
+
+
+	bool IsObjectDynamic;
+	bool IsObjectLoading;
+
+	FBoxSphereBounds WorldBounds;
+};
diff --git a/Engine/Source/Runtime/Experimental/ChaosSolvers/Public/PhysicsProxy/GeometryCollectionPhysicsProxy.h b/Engine/Source/Runtime/Experimental/ChaosSolvers/Public/PhysicsProxy/GeometryCollectionPhysicsProxy.h
--- a/Engine/Source/Runtime/Experimental/ChaosSolvers/Public/PhysicsProxy/GeometryCollectionPhysicsProxy.h
+++ b/Engine/Source/Runtime/Experimental/ChaosSolvers/Public/PhysicsProxy/GeometryCollectionPhysicsProxy.h
@@ -23,61 +23,6 @@ namespace Chaos
 	template <typename T> struct FClusterCreationParameters;
 }
 
-/**
- * Buffer structure for communicating simulation state between game and physics 
- * threads.
- */
-class FGeometryCollectionResults
-{
-public:
-	FGeometryCollectionResults();
-
-	void Reset();
-
-	int32 NumTransformGroup() const { return Transforms.Num(); }
-
-	void InitArrays(const FGeometryDynamicCollection &Other)
-	{
-		// Managed arrays
-		Transforms.Init(Other.Transform);
-		DynamicState.Init(Other.DynamicState);
-		Parent.Init(Other.Parent);
-		Children.Init(Other.Children);
-		SimulationType.Init(Other.SimulationType);
-		
-		// Arrays
-		const int32 NumTransforms = Other.NumElements(FGeometryCollection::TransformGroup);
-		DisabledStates.SetNumUninitialized(NumTransforms);
-		GlobalTransforms.SetNumUninitialized(NumTransforms);
-		ParticleToWorldTransforms.SetNumUninitialized(NumTransforms);
-	}
-
-	int32 BaseIndex;
-	int32 NumParticlesAdded;
-	TArray<bool> DisabledStates;
-	TArray<FMatrix> GlobalTransforms;
-	TArray<FTransform> ParticleToWorldTransforms;
-
-	TManagedArray<int32> TransformIndex;
-
-	TManagedArray<FTransform> Transforms;
-	TManagedArray<int32> BoneMap;
-	TManagedArray<int32> Parent;
-	TManagedArray<TSet<int32>> Children;
-	TManagedArray<int32> SimulationType;
-	TManagedArray<int32> DynamicState;
-	TManagedArray<float> Mass;
-	TManagedArray<FVector> InertiaTensor;
-
-	TManagedArray<int32> ClusterId;
-
-
-
-	bool IsObjectDynamic;
-	bool IsObjectLoading;
-
-	FBoxSphereBounds WorldBounds;
-};
 
 
 class FStubGeometryCollectionData : public Chaos::FParticleData 
