Fixed attribute read for geometry collection reinterpretting an array of unique ptrs as shared ptrs

diff --git a/Engine/Source/Runtime/Experimental/Chaos/Private/Chaos/TriangleMesh.cpp b/Engine/Source/Runtime/Experimental/Chaos/Private/Chaos/TriangleMesh.cpp
--- a/Engine/Source/Runtime/Experimental/Chaos/Private/Chaos/TriangleMesh.cpp
+++ b/Engine/Source/Runtime/Experimental/Chaos/Private/Chaos/TriangleMesh.cpp
@@ -265,7 +265,7 @@ void TTriangleMesh<T>::GetFaceNormals(TArray<TVector<T, 3>>& Normals, const TArr
 			if (Size2 < SMALL_NUMBER)
 			{
 				//particles should not be coincident by the time they get here. Return empty to signal problem to caller
-				check(false);
+				ensure(false);
 				Normals.Empty();
 				return;
 			}
diff --git a/Engine/Source/Runtime/Experimental/ChaosSolvers/Private/PhysicsProxy/GeometryCollectionPhysicsProxy.cpp b/Engine/Source/Runtime/Experimental/ChaosSolvers/Private/PhysicsProxy/GeometryCollectionPhysicsProxy.cpp
--- a/Engine/Source/Runtime/Experimental/ChaosSolvers/Private/PhysicsProxy/GeometryCollectionPhysicsProxy.cpp
+++ b/Engine/Source/Runtime/Experimental/ChaosSolvers/Private/PhysicsProxy/GeometryCollectionPhysicsProxy.cpp
@@ -1829,6 +1829,8 @@ void FGeometryCollectionPhysicsProxy::InitializeSharedCollisionStructures(
 	TManagedArray<TUniquePtr<FSimplicial>>& CollectionSimplicials = 
 		RestCollection.AddAttribute<TUniquePtr<FSimplicial>>(
 			FGeometryDynamicCollection::SimplicialsAttribute, FTransformCollection::TransformGroup);
+
+	RestCollection.RemoveAttribute(FGeometryDynamicCollection::ImplicitsAttribute, FTransformCollection::TransformGroup);
 	TManagedArray<FGeometryDynamicCollection::FSharedImplicit>& CollectionImplicits = 
 		RestCollection.AddAttribute<FGeometryDynamicCollection::FSharedImplicit>(
 			FGeometryDynamicCollection::ImplicitsAttribute, FTransformCollection::TransformGroup);
