Ignore kinematic rigids for counting number of rigids in older island for waking up slept particles.

diff --git a/Engine/Source/Runtime/Experimental/Chaos/Private/Chaos/PBDConstraintGraph.cpp b/Engine/Source/Runtime/Experimental/Chaos/Private/Chaos/PBDConstraintGraph.cpp
--- a/Engine/Source/Runtime/Experimental/Chaos/Private/Chaos/PBDConstraintGraph.cpp
+++ b/Engine/Source/Runtime/Experimental/Chaos/Private/Chaos/PBDConstraintGraph.cpp
@@ -496,7 +496,10 @@ void FPBDConstraintGraph::ComputeIslands(const TParticleView<TPBDRigidParticles<
 						{
 							if (TPBDRigidParticleHandle<FReal, 3>* PBDRigid = Particle->CastToRigidParticle())
 							{
-								NumRigidsInIsland++;
+								if (PBDRigid->ObjectState() != EObjectStateType::Kinematic)
+								{
+									NumRigidsInIsland++;
+								}
 							}
 						}
 					}
