Chaos :

- Basic Transform initialization tests - solver side.
- Mass initialization udpates.

diff --git a/Engine/Source/Runtime/Experimental/ChaosSolvers/Private/PBDRigidsSolver.cpp b/Engine/Source/Runtime/Experimental/ChaosSolvers/Private/PBDRigidsSolver.cpp
--- a/Engine/Source/Runtime/Experimental/ChaosSolvers/Private/PBDRigidsSolver.cpp
+++ b/Engine/Source/Runtime/Experimental/ChaosSolvers/Private/PBDRigidsSolver.cpp
@@ -457,7 +457,6 @@ namespace Chaos
 				TEXT("FPBDRigidsSolver::RegisterObject(FGeometryCollectionPhysicsProxy*)"));
 			check(InParticles);
 			InProxy->InitializeBodiesPT(Solver, *InParticles);
-			InProxy->PushToPhysicsState(nullptr); // Uses buffer on proxy for transport, not ProxyData.
 		});
 	}
 
diff --git a/Engine/Source/Runtime/Experimental/ChaosSolvers/Private/PhysicsProxy/GeometryCollectionPhysicsProxy.cpp b/Engine/Source/Runtime/Experimental/ChaosSolvers/Private/PhysicsProxy/GeometryCollectionPhysicsProxy.cpp
--- a/Engine/Source/Runtime/Experimental/ChaosSolvers/Private/PhysicsProxy/GeometryCollectionPhysicsProxy.cpp
+++ b/Engine/Source/Runtime/Experimental/ChaosSolvers/Private/PhysicsProxy/GeometryCollectionPhysicsProxy.cpp
@@ -227,8 +227,8 @@ void PopulateSimulatedParticle(
 	const FSharedSimulationParameters& SharedParams,
 	const FCollisionStructureManager::FSimplicial* Simplicial,
 	FGeometryDynamicCollection::FSharedImplicit Implicit,
-	const float MassIn,
-	const FVector& InertiaTensorVec,
+	float MassIn,
+	FVector InertiaTensorVec,
 	const FTransform& WorldTransform, 
 	const uint8 DynamicState, 
 	const int16 CollisionGroup)
@@ -236,7 +236,6 @@ void PopulateSimulatedParticle(
 	SCOPE_CYCLE_COUNTER(STAT_PopulateSimulatedParticle);
 
 	Handle->SetDisabledLowLevel(false);
-
 	Handle->SetX(WorldTransform.GetTranslation());
 	Handle->SetV(Chaos::TVector<float, 3>(0.f));
 	Handle->SetR(WorldTransform.GetRotation().GetNormalized());
@@ -245,36 +244,37 @@ void PopulateSimulatedParticle(
 	Handle->SetQ(Handle->R());
 	Handle->SetIsland(INDEX_NONE);
 
-	//todo: If mass is too small use the right inertia
-	ensureMsgf(MassIn >= SharedParams.MinimumMassClamp, 
-		TEXT("Mass smaller than minimum mass clamp. Too late to change"));
-	ensureMsgf(Handle->InvM() > SMALL_NUMBER, 
-		TEXT("Object mass is too large. Too late to change"));
-	ensureMsgf(InertiaTensorVec[0] > SMALL_NUMBER && InertiaTensorVec[1] > SMALL_NUMBER && InertiaTensorVec[2] > SMALL_NUMBER, 
-		TEXT("Inertia tensor is too small. Too late to change"));
-
-	Handle->SetM(MassIn);
-	Handle->SetInvM(1.0f / MassIn);
-	if (FMath::IsNaN(InertiaTensorVec[0]) || FMath::IsNaN(InertiaTensorVec[1]) || FMath::IsNaN(InertiaTensorVec[2]) ||
-		InertiaTensorVec[0] < SMALL_NUMBER || InertiaTensorVec[1] < SMALL_NUMBER || InertiaTensorVec[2] < SMALL_NUMBER)
-	{
-		Handle->SetI(Chaos::PMatrix<float, 3, 3>(1.f, 1.f, 1.f));
-	}
-	else
+	//
+	// Setup Mass
+	//
 	{
-		Handle->SetI(Chaos::PMatrix<float, 3, 3>(InertiaTensorVec[0], InertiaTensorVec[1], InertiaTensorVec[2]));
-		Handle->SetInvI(Chaos::PMatrix<float, 3, 3>(1.0f / InertiaTensorVec[0], 1.0f / InertiaTensorVec[1], 1.0f / InertiaTensorVec[2]));
-	}
+		Handle->SetObjectState(Chaos::EObjectStateType::Uninitialized);
 
-	// for validation set the body to dynamic and check the inverse masses. 
-	Handle->SetObjectStateLowLevel(Chaos::EObjectStateType::Dynamic);
+		if (ensureMsgf(FMath::IsWithin(MassIn, SharedParams.MinimumMassClamp, SharedParams.MaximumMassClamp),
+			TEXT("Clamped mass[%3.5f] to range [%3.5f,%3.5f]"), MassIn, SharedParams.MinimumMassClamp, SharedParams.MaximumMassClamp))
+		{
+			MassIn = FMath::Clamp(MassIn, SharedParams.MinimumMassClamp, SharedParams.MaximumMassClamp);
+		}
 
-	const Chaos::PMatrix<float,3,3>& InvI = Handle->InvI();
-	if (!(InvI.M[0][0] > SMALL_NUMBER && 
-		  InvI.M[1][1] > SMALL_NUMBER && 
-		  InvI.M[2][2] > SMALL_NUMBER))
-	{
-		UE_LOG(LogChaos, Warning, TEXT("Inertia tensor is too large. Too late to change"));
+		if (ensureMsgf(!FMath::IsNaN(InertiaTensorVec[0]) && !FMath::IsNaN(InertiaTensorVec[1]) && !FMath::IsNaN(InertiaTensorVec[2]),
+			TEXT("Nan Tensor, reset to unit tesor")))
+		{
+			InertiaTensorVec = FVector(1);
+		}
+		else if (ensureMsgf(FMath::IsWithin(InertiaTensorVec[0], SharedParams.MinimumInertiaTensorDiagonalClamp, SharedParams.MaximumInertiaTensorDiagonalClamp)
+			&& FMath::IsWithin(InertiaTensorVec[1], SharedParams.MinimumInertiaTensorDiagonalClamp, SharedParams.MaximumInertiaTensorDiagonalClamp)
+			&& FMath::IsWithin(InertiaTensorVec[2], SharedParams.MinimumInertiaTensorDiagonalClamp, SharedParams.MaximumInertiaTensorDiagonalClamp),
+			TEXT("Clamped Inertia tensor[%3.5f,%3.5f,%3.5f]. Clamped each element to [%3.5f, %3.5f,]"), InertiaTensorVec[0], InertiaTensorVec[1], InertiaTensorVec[2],
+			SharedParams.MinimumInertiaTensorDiagonalClamp, SharedParams.MaximumInertiaTensorDiagonalClamp))
+		{
+			InertiaTensorVec[0] = FMath::Clamp(InertiaTensorVec[0], SharedParams.MinimumInertiaTensorDiagonalClamp, SharedParams.MaximumInertiaTensorDiagonalClamp);
+			InertiaTensorVec[1] = FMath::Clamp(InertiaTensorVec[1], SharedParams.MinimumInertiaTensorDiagonalClamp, SharedParams.MaximumInertiaTensorDiagonalClamp);
+			InertiaTensorVec[2] = FMath::Clamp(InertiaTensorVec[2], SharedParams.MinimumInertiaTensorDiagonalClamp, SharedParams.MaximumInertiaTensorDiagonalClamp);
+		}
+
+		Handle->SetM(MassIn);
+		Handle->SetI(Chaos::PMatrix<float, 3, 3>(InertiaTensorVec[0], InertiaTensorVec[1], InertiaTensorVec[2]));
+		Handle->SetObjectStateLowLevel(Chaos::EObjectStateType::Dynamic); // this step sets InvM, InvInertia, P, Q
 	}
 
 	Handle->SetCollisionGroup(CollisionGroup);
@@ -611,18 +611,6 @@ void FGeometryCollectionPhysicsProxy::InitializeDynamicCollection(FGeometryDynam
 		}
 	}
 
-	// Process Transforms
-	{
-		//TArray<FTransform> Transform;
-		//GeometryCollectionAlgo::GlobalMatrices(DynamicCollection->Transform, DynamicCollection->Parent, Transform);
-		//check(DynamicCollection->Transform.Num() == Transform.Num());
-
-		// Initialize global transforms
-		//TArray<FMatrix>& TmpGlobalTransforms = GameToPhysInterchange.AccessProducerBuffer()->GlobalTransforms;
-		//TArray<FMatrix> &TmpGlobalTransforms = Results.Get(0).GlobalTransforms;
-		//GeometryCollectionAlgo::GlobalMatrices(PTDynamicCollection.Transform, PTDynamicCollection.Parent, TmpGlobalTransforms);
-		//Results.Get(1).GlobalTransforms = TmpGlobalTransforms;
-	}
 }
 
 
@@ -675,11 +663,11 @@ void FGeometryCollectionPhysicsProxy::InitializeBodiesPT(
 
 		// Add entries into simulation array
 		RigidsSolver->GetEvolution()->ReserveParticles(NumSimulatedParticles);
-		TArray<Chaos::TPBDGeometryCollectionParticleHandle<float, 3>*> Handles = 
+		TArray<Chaos::TPBDGeometryCollectionParticleHandle<float, 3>*> Handles =
 			RigidsSolver->GetEvolution()->CreateGeometryCollectionParticles(NumLeafNodes);
 
 		int32 NextIdx = 0;
-		for(int32 Idx=0; Idx < SimulatableParticles.Num(); ++Idx)
+		for (int32 Idx = 0; Idx < SimulatableParticles.Num(); ++Idx)
 		{
 			if (SimulatableParticles[Idx] && Children[Idx].Num() == 0) // Leaf node
 			{
@@ -719,29 +707,29 @@ void FGeometryCollectionPhysicsProxy::InitializeBodiesPT(
 			{
 				// TODO: We probably want to take better advantage of this solver
 				// feature, rather than doing it all here ourselves?
-				Handle->SetCenterOfMass(Chaos::TVector<float,3>(0.0f));
-				Handle->SetRotationOfMass(Chaos::TRotation<float,3>::FromIdentity());
-			
+				Handle->SetCenterOfMass(MassToLocal[TransformGroupIndex].GetTranslation());
+				Handle->SetRotationOfMass(MassToLocal[TransformGroupIndex].GetRotation());
+
 				// Mass space -> Composed parent space -> world
-				const FTransform WorldTransform = 
+				const FTransform WorldTransform =
 					MassToLocal[TransformGroupIndex] * Transform[TransformGroupIndex] * Parameters.WorldTransform;
 
 				PopulateSimulatedParticle(
 					Handle,
 					//Particles, 
-					Parameters.Shared, 
+					Parameters.Shared,
 					nullptr,///Simplicials[TransformGroupIndex].Get(),
-					Implicits[TransformGroupIndex], 
+					Implicits[TransformGroupIndex],
 #if TODO_REIMPLEMENT_SOLVER_SETTINGS_ACCESSORS
-					RigidsSolver->GetMassScale() * Mass[TransformGroupIndex], 
+					RigidsSolver->GetMassScale() * Mass[TransformGroupIndex],
 					RigidsSolver->GetMassScale() * InertiaTensor[TransformGroupIndex],
 #else // TODO_REIMPLEMENT_SOLVER_SETTINGS_ACCESSORS
-					Mass[TransformGroupIndex], 
+					Mass[TransformGroupIndex],
 					InertiaTensor[TransformGroupIndex],
 #endif // TODO_REIMPLEMENT_SOLVER_SETTINGS_ACCESSORS
 					//RigidBodyIndex, 
-					WorldTransform, 
-					(uint8)DynamicState[TransformGroupIndex], 
+					WorldTransform,
+					(uint8)DynamicState[TransformGroupIndex],
 					CollisionGroup[TransformGroupIndex]);
 
 				if (Parameters.EnableClustering)
@@ -754,11 +742,11 @@ void FGeometryCollectionPhysicsProxy::InitializeBodiesPT(
 				// @todo break everything CollisionParticles.Reset(Simplicials[TransformGroupIndex].Release()); // Steals!
 				//Particles.CollisionParticles(RigidBodyIndex).Reset(Simplicials[TransformGroupIndex].Release());
 				if (CollisionParticles)
-				//if (Particles.CollisionParticles(RigidBodyIndex))
+					//if (Particles.CollisionParticles(RigidBodyIndex))
 				{
 					int32 NumCollisionParticles = CollisionParticles->Size();
 					//int32 NumCollisionParticles = Particles.CollisionParticles(RigidBodyIndex)->Size();
-					int32 CollisionParticlesSize = FMath::Max(0, FMath::Min(int(NumCollisionParticles*CollisionParticlesPerObjectFraction), NumCollisionParticles));
+					int32 CollisionParticlesSize = FMath::Max(0, FMath::Min(int(NumCollisionParticles * CollisionParticlesPerObjectFraction), NumCollisionParticles));
 					CollisionParticles->Resize(CollisionParticlesSize); // Truncates!
 					//Particles.CollisionParticles(RigidBodyIndex)->Resize(CollisionParticlesSize);
 				}
diff --git a/Engine/Source/Runtime/Experimental/GeometryCollectionSimulationCore/Private/GeometryCollection/GeometryCollectionSimulationCoreTypes.cpp b/Engine/Source/Runtime/Experimental/GeometryCollectionSimulationCore/Private/GeometryCollection/GeometryCollectionSimulationCoreTypes.cpp
--- a/Engine/Source/Runtime/Experimental/GeometryCollectionSimulationCore/Private/GeometryCollection/GeometryCollectionSimulationCoreTypes.cpp
+++ b/Engine/Source/Runtime/Experimental/GeometryCollectionSimulationCore/Private/GeometryCollection/GeometryCollectionSimulationCoreTypes.cpp
@@ -1,5 +1,3 @@
 // Copyright Epic Games, Inc. All Rights Reserved.
 
 #include "GeometryCollection/GeometryCollectionSimulationCoreTypes.h"
-
-constexpr float FSharedSimulationParameters::MaximumMassClamp;
\ No newline at end of file
diff --git a/Engine/Source/Runtime/Experimental/GeometryCollectionSimulationCore/Public/GeometryCollection/GeometryCollectionSimulationCoreTypes.h b/Engine/Source/Runtime/Experimental/GeometryCollectionSimulationCore/Public/GeometryCollection/GeometryCollectionSimulationCoreTypes.h
--- a/Engine/Source/Runtime/Experimental/GeometryCollectionSimulationCore/Public/GeometryCollection/GeometryCollectionSimulationCoreTypes.h
+++ b/Engine/Source/Runtime/Experimental/GeometryCollectionSimulationCore/Public/GeometryCollection/GeometryCollectionSimulationCoreTypes.h
@@ -53,11 +53,13 @@ enum ESimulationInitializationState { Unintialized = 0, Activated, Created, Init
 */
 struct FSharedSimulationParameters
 {
-	static constexpr float MaximumMassClamp = 1000;	//todo: this is way too low, need to handle this in a better way when combining with large inertia
 	FSharedSimulationParameters()
 	: bMassAsDensity(false)
 	, Mass(1.0)
-	, MinimumMassClamp(0.1)
+	, MinimumMassClamp(0.1f)
+	, MaximumMassClamp(1000.0f)
+	, MinimumInertiaTensorDiagonalClamp(SMALL_NUMBER) // todo : Expose to users with better initial values
+	, MaximumInertiaTensorDiagonalClamp(1.0e6f) // todo : Expose to users with better initial values
 	, MaximumCollisionParticleCount(60)
 	{
 		SizeSpecificData.AddDefaulted();
@@ -71,12 +73,18 @@ struct FSharedSimulationParameters
 		,int32 InMaxClusterLevelSetResolution
 		,bool InMassAsDensity
 		,float InMass
-		,float InMinimumMassClamp
+		, float InMinimumMassClamp
+		, float InMaximumMassClamp
+		, float InMinimumInertiaTensorDiagonalClamp
+		, float InMaximumInertiaTensorDiagonalClamp
 		,float InCollisionParticlesFraction
 		,int32 InMaximumCollisionParticleCount)
 	: bMassAsDensity(InMassAsDensity)
 	, Mass(InMass)
 	, MinimumMassClamp(InMinimumMassClamp)
+	, MaximumMassClamp(InMinimumMassClamp)
+	, MinimumInertiaTensorDiagonalClamp(InMinimumInertiaTensorDiagonalClamp)
+	, MaximumInertiaTensorDiagonalClamp(InMaximumInertiaTensorDiagonalClamp)
 	, MaximumCollisionParticleCount(InMaximumCollisionParticleCount)
 	{
 		SizeSpecificData.AddDefaulted();
@@ -93,6 +101,10 @@ struct FSharedSimulationParameters
 	bool bMassAsDensity;
 	float Mass;
 	float MinimumMassClamp;
+	float MaximumMassClamp;
+	float MinimumInertiaTensorDiagonalClamp;
+	float MaximumInertiaTensorDiagonalClamp;
+
 	TArray<FSharedSimulationSizeSpecificData> SizeSpecificData;
 	TArray<int32> RemoveOnFractureIndices;
 	int32 MaximumCollisionParticleCount;
@@ -235,6 +247,7 @@ struct FSimulationParameters
 		, bOwnsTrack(false)
 		, Simulating(Other.Simulating)
 		, WorldTransform(Other.WorldTransform)
+		, EnableClustering(Other.EnableClustering)
 		, ClusterGroupIndex(Other.ClusterGroupIndex)
 		, MaxClusterLevel(Other.MaxClusterLevel)
 		, DamageThreshold(Other.DamageThreshold)
