Chaos : Init

- Force active bodies out of the parent hierarchy in the FGeometryDynamicCollection when pulling simulation results out of the solver.

diff --git a/Engine/Source/Runtime/Experimental/ChaosSolvers/Private/PhysicsProxy/GeometryCollectionPhysicsProxy.cpp b/Engine/Source/Runtime/Experimental/ChaosSolvers/Private/PhysicsProxy/GeometryCollectionPhysicsProxy.cpp
--- a/Engine/Source/Runtime/Experimental/ChaosSolvers/Private/PhysicsProxy/GeometryCollectionPhysicsProxy.cpp
+++ b/Engine/Source/Runtime/Experimental/ChaosSolvers/Private/PhysicsProxy/GeometryCollectionPhysicsProxy.cpp
@@ -647,7 +647,6 @@ void FGeometryCollectionPhysicsProxy::InitializeBodiesPT(
 
 		TArray<FTransform> Transform;
 		GeometryCollectionAlgo::GlobalMatrices(DynamicCollection.Transform, DynamicCollection.Parent, Transform);
-		check(DynamicCollection.Transform.Num() == Transform.Num());
 
 		const int NumRigids = 0; // ryan - Since we're doing SOA, we start at zero?
 		BaseParticleIndex = NumRigids;
@@ -1655,6 +1654,7 @@ void FGeometryCollectionPhysicsProxy::PullFromPhysicsState()
 				DynamicCollection.Transform[TmIndex] = LocalTransform;
 				GTParticles[TmIndex]->SetX(ParticleToWorld.GetTranslation());
 				GTParticles[TmIndex]->SetR(ParticleToWorld.GetRotation());
+				DynamicCollection.Parent[TmIndex] = FGeometryCollection::Invalid;
 			}
 
 			DynamicCollection.DynamicState[TmIndex] = TR.DynamicState[TmIndex];
diff --git a/Engine/Source/Runtime/Experimental/GeometryCollectionCore/Private/GeometryCollection/ManagedArrayCollection.cpp b/Engine/Source/Runtime/Experimental/GeometryCollectionCore/Private/GeometryCollection/ManagedArrayCollection.cpp
--- a/Engine/Source/Runtime/Experimental/GeometryCollectionCore/Private/GeometryCollection/ManagedArrayCollection.cpp
+++ b/Engine/Source/Runtime/Experimental/GeometryCollectionCore/Private/GeometryCollection/ManagedArrayCollection.cpp
@@ -7,7 +7,7 @@
 
 DEFINE_LOG_CATEGORY_STATIC(FManagedArrayCollectionLogging, NoLogging, All);
 
-int8 FManagedArrayCollection::Invalid = -1;
+int8 FManagedArrayCollection::Invalid = INDEX_NONE;
 
 
 FManagedArrayCollection::FManagedArrayCollection()
