Push SimFilter and QueryFilter to physics thread particles from the geometry collection physics proxy.  Disable the GeometryCollectionCollideAll cvar by default.  Intialize the FGeometryCollectionPhysicsProxy used for tests with a FCollisionFilsterData SimFilter that enables all collisions.  Uncomment the body of RigidBodies_Initialize_ParticleImplicitCollisionGeometry(), and fix its compilation issues.

diff --git a/Engine/Source/Runtime/Experimental/ChaosSolvers/Private/PhysicsProxy/GeometryCollectionPhysicsProxy.cpp b/Engine/Source/Runtime/Experimental/ChaosSolvers/Private/PhysicsProxy/GeometryCollectionPhysicsProxy.cpp
--- a/Engine/Source/Runtime/Experimental/ChaosSolvers/Private/PhysicsProxy/GeometryCollectionPhysicsProxy.cpp
+++ b/Engine/Source/Runtime/Experimental/ChaosSolvers/Private/PhysicsProxy/GeometryCollectionPhysicsProxy.cpp
@@ -49,8 +49,7 @@ FAutoConsoleVariableRef CVarGeometryCollectionDisableGravity(
 	DisableGeometryCollectionGravity,
 	TEXT("Disable gravity for geometry collections"));
 
-//bool GeometryCollectionCollideAll = false;
-bool GeometryCollectionCollideAll = true;
+bool GeometryCollectionCollideAll = false;
 FAutoConsoleVariableRef CVarGeometryCollectionCollideAll(
 	TEXT("p.GeometryCollectionCollideAll"),
 	GeometryCollectionCollideAll,
@@ -227,6 +226,8 @@ void PopulateSimulatedParticle(
 	const FSharedSimulationParameters& SharedParams,
 	const FCollisionStructureManager::FSimplicial* Simplicial,
 	FGeometryDynamicCollection::FSharedImplicit Implicit,
+	const FCollisionFilterData SimFilterIn,
+	const FCollisionFilterData QueryFilterIn,
 	float MassIn,
 	FVector InertiaTensorVec,
 	const FTransform& WorldTransform, 
@@ -325,13 +326,21 @@ void PopulateSimulatedParticle(
 		FilterData.Word3 = 0xFFFF; // collision candidate channels
 		for (const TUniquePtr<Chaos::TPerShapeData<Chaos::FReal, 3>>& Shape : Handle->ShapesArray())
 		{
-			// TODO Ryan - hardwired for now:
 			Shape->bDisable = false;
-			Shape->CollisionTraceType = Chaos::EChaosCollisionTraceFlag::Chaos_CTF_UseSimpleAndComplex; // Chaos_CTF_UseDefault?
+			Shape->CollisionTraceType = Chaos::EChaosCollisionTraceFlag::Chaos_CTF_UseDefault;
+			//Shape->CollisionTraceType = Chaos::EChaosCollisionTraceFlag::Chaos_CTF_UseSimpleAndComplex;
 			Shape->SimData = FilterData;
 			Shape->QueryData = FCollisionFilterData();
 		}
 	}
+	else
+	{
+		for (const TUniquePtr<Chaos::TPerShapeData<Chaos::FReal, 3>>& Shape : Handle->ShapesArray())
+		{
+			Shape->SimData = SimFilterIn;
+			Shape->QueryData = QueryFilterIn;
+		}
+	}
 
 	//
 	//  Manage Object State
@@ -674,7 +683,6 @@ void FGeometryCollectionPhysicsProxy::InitializeBodiesPT(
 
 				Chaos::TPBDGeometryCollectionParticleHandle<float, 3>* Handle = Handles[NextIdx++];
 
-				Handle->GTGeometryParticle() = nullptr; // GeometryCollections do not support Game Thread TGeometryParticle
 				RigidsSolver->AddParticleToProxy(Handle, this);
 
 				SolverParticleHandles[Idx] = Handle;
@@ -719,6 +727,8 @@ void FGeometryCollectionPhysicsProxy::InitializeBodiesPT(
 					Parameters.Shared,
 					nullptr,///Simplicials[TransformGroupIndex].Get(),
 					Implicits[TransformGroupIndex],
+					SimFilter,
+					QueryFilter,
 #if TODO_REIMPLEMENT_SOLVER_SETTINGS_ACCESSORS
 					RigidsSolver->GetMassScale() * Mass[TransformGroupIndex],
 					RigidsSolver->GetMassScale() * InertiaTensor[TransformGroupIndex],
@@ -1042,6 +1052,8 @@ FGeometryCollectionPhysicsProxy::BuildClusters(
 		Parameters.Shared, 
 		nullptr, // CollisionParticles is optionally set from CreateClusterParticle()
 		nullptr, // Parent->Geometry() ? Parent->Geometry() : Implicits[CollectionClusterIndex], 
+		SimFilter,
+		QueryFilter,
 #if TODO_REIMPLEMENT_SOLVER_SETTINGS_ACCESSORS
 		// Note: if the children had this scaling baked in, then cluster parents
 		// will too, and we shouldn't apply the scaling to the parents twice!
@@ -2128,6 +2140,8 @@ void FGeometryCollectionPhysicsProxy::InitializeSharedCollisionStructures(
 					SharedParams, 
 					CollectionSimplicials[TransformGroupIndex].Get(),
 					CollectionImplicits[TransformGroupIndex], 
+					FCollisionFilterData(),		// SimFilter
+					FCollisionFilterData(),		// QueryFilter
 					CollectionMass[TransformGroupIndex],
 					CollectionInertiaTensor[TransformGroupIndex], 
 					MassToComponent,
