Chaos : bug fix handles gathering in fields

diff --git a/Engine/Source/Runtime/Experimental/ChaosSolvers/Private/PhysicsProxy/GeometryCollectionPhysicsProxy.cpp b/Engine/Source/Runtime/Experimental/ChaosSolvers/Private/PhysicsProxy/GeometryCollectionPhysicsProxy.cpp
--- a/Engine/Source/Runtime/Experimental/ChaosSolvers/Private/PhysicsProxy/GeometryCollectionPhysicsProxy.cpp
+++ b/Engine/Source/Runtime/Experimental/ChaosSolvers/Private/PhysicsProxy/GeometryCollectionPhysicsProxy.cpp
@@ -1128,33 +1128,33 @@ void FGeometryCollectionPhysicsProxy::GetRelevantHandles(
 	{
 		// only the local handles
 		TArray<FClusterHandle*>& ParticleHandles = GetSolverParticleHandles();
+		Handles.SetNumUninitialized(ParticleHandles.Num());
+		int32 NumUsedHandles = 0;
 
 		if (ResolutionType == EFieldResolutionType::Field_Resolution_Maximum)
 		{
-			Handles.SetNumUninitialized(ParticleHandles.Num());
 
-			int32 NumUsedHandles = 0;
 			for (FClusterHandle* ClusterHandle : ParticleHandles)
 			{
-				if (!ClusterHandle->Disabled())
+				if (ClusterHandle && !ClusterHandle->Disabled())
 				{
-					Handles.Add(ClusterHandle);
+					Handles[NumUsedHandles] = ClusterHandle;
 					NumUsedHandles++;
 				}
 			}
-			Handles.SetNum(NumUsedHandles);
-
 		}
 		else if (ResolutionType == EFieldResolutionType::Field_Resolution_Minimal)
 		{
-			Handles.Reserve(ParticleHandles.Num());
-
 			for (FClusterHandle* ClusterHandle : ParticleHandles)
 			{
-				Handles.Add(ClusterHandle);
+				if (ClusterHandle)
+				{
+					Handles[NumUsedHandles] = ClusterHandle;
+					NumUsedHandles++;
+				}
 			}
 		}
-
+		Handles.SetNum(NumUsedHandles);
 		Samples.AddUninitialized(Handles.Num());
 		SampleIndices.AddUninitialized(Handles.Num());
 		for (int32 Idx = 0; Idx < Handles.Num(); ++Idx)
