Chaos: Fix trailing events lambda

diff --git a/Engine/Source/Runtime/Experimental/ChaosSolvers/Private/EventDefaults.cpp b/Engine/Source/Runtime/Experimental/ChaosSolvers/Private/EventDefaults.cpp
--- a/Engine/Source/Runtime/Experimental/ChaosSolvers/Private/EventDefaults.cpp
+++ b/Engine/Source/Runtime/Experimental/ChaosSolvers/Private/EventDefaults.cpp
@@ -403,54 +403,49 @@ namespace Chaos
 
 			for (auto& ActiveParticle : Evolution->GetParticles().GetActiveParticlesView())
 			{
-				TPBDRigidParticle<float, 3>* GTParticle = ActiveParticle.Handle()->GTGeometryParticle()->CastToRigidParticle();
-				// Since Clustered GCs can be unioned the particleIndex representing the union 
-				// is not associated with a PhysicsProxy
-				if (GTParticle && GTParticle->Proxy != nullptr)
+
+				if (ensure(FMath::IsFinite(ActiveParticle.InvM())))
 				{
-					if (ensure(FMath::IsFinite(GTParticle->InvM())))
+					if (ActiveParticle.InvM() != 0.f &&
+						ActiveParticle.Geometry() &&
+						ActiveParticle.Geometry()->HasBoundingBox())
 					{
-						if (GTParticle->InvM() != 0.f &&
-							ActiveParticle.Geometry() &&
-							ActiveParticle.Geometry()->HasBoundingBox())
+						if (ensure(!ActiveParticle.X().ContainsNaN() &&
+							!ActiveParticle.V().ContainsNaN() &&
+							!ActiveParticle.W().ContainsNaN() &&
+							FMath::IsFinite(ActiveParticle.M())))
 						{
-							if (ensure(!GTParticle->X().ContainsNaN() &&
-								!GTParticle->V().ContainsNaN() &&
-								!GTParticle->W().ContainsNaN() &&
-								FMath::IsFinite(GTParticle->M())))
+							TTrailingData<float, 3> TrailingData;
+							TrailingData.Location = ActiveParticle.X();
+							TrailingData.Velocity = ActiveParticle.V();
+							TrailingData.AngularVelocity = ActiveParticle.W();
+							TrailingData.Mass = ActiveParticle.M();
+							TrailingData.Particle = nullptr; // #todo: provide a particle
+							if (ActiveParticle.Geometry()->HasBoundingBox())
 							{
-								TTrailingData<float, 3> TrailingData;
-								TrailingData.Location = GTParticle->X();
-								TrailingData.Velocity = GTParticle->V();
-								TrailingData.AngularVelocity = GTParticle->W();
-								TrailingData.Mass = GTParticle->M();
-								TrailingData.Particle = GTParticle;
-								if (ActiveParticle.Geometry()->HasBoundingBox())
-								{
-									TrailingData.BoundingBox = ActiveParticle.Geometry()->BoundingBox();
-								}
+								TrailingData.BoundingBox = ActiveParticle.Geometry()->BoundingBox();
+							}
 
-								const FSolverTrailingEventFilter* SolverTrailingEventFilter = Solver->GetEventFilters()->GetTrailingFilter();
-								if (!SolverTrailingEventFilter->Enabled() || SolverTrailingEventFilter->Pass(TrailingData))
-								{
-									int32 NewIdx = AllTrailingsDataArray.Add(TTrailingData<float, 3>());
-									TTrailingData<float, 3>& TrailingDataArrayItem = AllTrailingsDataArray[NewIdx];
-									TrailingDataArrayItem = TrailingData;
+							const FSolverTrailingEventFilter* SolverTrailingEventFilter = Solver->GetEventFilters()->GetTrailingFilter();
+							if (!SolverTrailingEventFilter->Enabled() || SolverTrailingEventFilter->Pass(TrailingData))
+							{
+								int32 NewIdx = AllTrailingsDataArray.Add(TTrailingData<float, 3>());
+								TTrailingData<float, 3>& TrailingDataArrayItem = AllTrailingsDataArray[NewIdx];
+								TrailingDataArrayItem = TrailingData;
 
-									// If IdxParticle is a cluster store an index for a mesh in this cluster
+								// If IdxParticle is a cluster store an index for a mesh in this cluster
 #if 0
-									if (ClusterIdsArray[IdxParticle].NumChildren > 0)
-									{
-										int32 ParticleIndexMesh = GetParticleIndexMesh(ParentToChildrenMap, IdxParticle);
-										ensure(ParticleIndexMesh != INDEX_NONE);
-										TrailingDataArrayItem.ParticleIndexMesh = ParticleIndexMesh;
-									}
+								if (ClusterIdsArray[IdxParticle].NumChildren > 0)
+								{
+									int32 ParticleIndexMesh = GetParticleIndexMesh(ParentToChildrenMap, IdxParticle);
+									ensure(ParticleIndexMesh != INDEX_NONE);
+									TrailingDataArrayItem.ParticleIndexMesh = ParticleIndexMesh;
+								}
 #endif
 
-									// Add to AllTrailingsIndicesByPhysicsProxy
-									IPhysicsProxyBase* PhysicsProxy = GTParticle->Proxy;
-									AllTrailingsIndicesByPhysicsProxy.FindOrAdd(PhysicsProxy).Add(NewIdx);
-								}
+								// Add to AllTrailingsIndicesByPhysicsProxy
+								IPhysicsProxyBase* PhysicsProxy = Solver->GetProxy(ActiveParticle.Handle());
+								AllTrailingsIndicesByPhysicsProxy.FindOrAdd(PhysicsProxy).Add(NewIdx);
 							}
 						}
 					}
