Unblocked robomerge by resolving PhysicsSolverBase.h. Original commit message:

Refactor data marshaling to be able to go wide during PushToPhysics

diff --git a/Engine/Source/Programs/HeadlessChaos/Private/HeadlessChaosTestDirtyParticle.cpp b/Engine/Source/Programs/HeadlessChaos/Private/HeadlessChaosTestDirtyParticle.cpp
--- a/Engine/Source/Programs/HeadlessChaos/Private/HeadlessChaosTestDirtyParticle.cpp
+++ b/Engine/Source/Programs/HeadlessChaos/Private/HeadlessChaosTestDirtyParticle.cpp
@@ -26,10 +26,10 @@ namespace ChaosTest
 		FDirtyPropertiesManager Manager;
 
 		{
-			FParticlePropertiesData RemoteData;
-			RemoteData.Manager = &Manager;
+			FParticlePropertiesData RemoteData(&Manager);
+			FShapeRemoteDataContainer ContainerData(&Manager);
 
-			Particle->SetRemoteData(&RemoteData);
+			Particle->SetRemoteData(RemoteData, ContainerData);
 
 			//then on PushToPhysics we'd do a pointer swap and use RemoteData internally
 
@@ -48,7 +48,7 @@ namespace ChaosTest
 			EXPECT_TRUE(RemoteData.HasGeometry());
 			EXPECT_EQ(RemoteData.GetGeometry().Get(),RawPtr);
 
-			Particle->SetRemoteData(nullptr);	//disconnect remote data so we can pretend we are freeing things on GT without affecting PT
+			Particle->DetachRemoteData();	//disconnect remote data so we can pretend we are freeing things on GT without affecting PT
 
 			Particle->SetGeometry(Ptr);	//free geometry on GT side
 
diff --git a/Engine/Source/Programs/HeadlessChaos/Private/HeadlessChaosTestUtility.cpp b/Engine/Source/Programs/HeadlessChaos/Private/HeadlessChaosTestUtility.cpp
--- a/Engine/Source/Programs/HeadlessChaos/Private/HeadlessChaosTestUtility.cpp
+++ b/Engine/Source/Programs/HeadlessChaos/Private/HeadlessChaosTestUtility.cpp
@@ -680,7 +680,7 @@ namespace ChaosTest {
 		InParticles.SetObjectStateLowLevel(EObjectStateType::Dynamic);
 
 		::ChaosTest::SetParticleSimDataToCollide({ &InParticles });
-		//for (const TUniquePtr<Chaos::TPerShapeData<Chaos::FReal, 3>>& Shape : InParticles.ShapesArray())
+		//for (const TUniquePtr<Chaos::FPerShapeData>& Shape : InParticles.ShapesArray())
 		//{
 		//	Shape->SimData.Word3 = 1;
 		//	Shape->SimData.Word1 = 1;
@@ -777,10 +777,13 @@ namespace ChaosTest {
 	{
 		for (Chaos::TGeometryParticle<float, 3>* Particle : ParticleArray)
 		{
-			for (const TUniquePtr<Chaos::TPerShapeData<Chaos::FReal, 3>>& Shape :Particle->ShapesArray())
+			for (const TUniquePtr<Chaos::FPerShapeData>& Shape :Particle->ShapesArray())
 			{
-				Shape->SimData.Word3 = 1;
-				Shape->SimData.Word1 = 1;
+				Shape->ModifySimData([](auto& SimData)
+				{
+					SimData.Word3 = 1;
+					SimData.Word1 = 1;
+				});
 			}
 		}
 	}
@@ -788,10 +791,13 @@ namespace ChaosTest {
 	{
 		for (Chaos::TGeometryParticleHandle<float, 3>* Particle : ParticleArray)
 		{
-			for (const TUniquePtr<Chaos::TPerShapeData<Chaos::FReal, 3>>& Shape : Particle->ShapesArray())
+			for (const TUniquePtr<Chaos::FPerShapeData>& Shape : Particle->ShapesArray())
 			{
-				Shape->SimData.Word3 = 1;
-				Shape->SimData.Word1 = 1;
+				Shape->ModifySimData([](auto& SimData)
+				{
+					SimData.Word3 = 1;
+					SimData.Word1 = 1;
+				});
 			}
 		}
 	}
diff --git a/Engine/Source/Runtime/Engine/Private/Collision/CollisionConversions.cpp b/Engine/Source/Runtime/Engine/Private/Collision/CollisionConversions.cpp
--- a/Engine/Source/Runtime/Engine/Private/Collision/CollisionConversions.cpp
+++ b/Engine/Source/Runtime/Engine/Private/Collision/CollisionConversions.cpp
@@ -108,7 +108,7 @@ static FVector FindGeomOpposingNormal(ECollisionShapeType QueryGeomType, const F
 			const FTransform ActorTM(Hit.Actor->R(), Hit.Actor->X());
 			const FVector LocalInNormal = ActorTM.InverseTransformVectorNoScale(InNormal);
 			const FVector LocalTraceDirectionDenorm = ActorTM.InverseTransformVectorNoScale(TraceDirectionDenorm);
-			const FVector LocalNormal = Shape->Geometry->FindGeometryOpposingNormal(LocalTraceDirectionDenorm, Hit.FaceIndex, LocalInNormal);
+			const FVector LocalNormal = Shape->GetGeometry()->FindGeometryOpposingNormal(LocalTraceDirectionDenorm, Hit.FaceIndex, LocalInNormal);
 			return ActorTM.TransformVectorNoScale(LocalNormal);
 #else
 			ECollisionShapeType GeomType = GetGeometryType(*Shape);
diff --git a/Engine/Source/Runtime/Engine/Private/PhysicsEngine/BodyInstance.cpp b/Engine/Source/Runtime/Engine/Private/PhysicsEngine/BodyInstance.cpp
--- a/Engine/Source/Runtime/Engine/Private/PhysicsEngine/BodyInstance.cpp
+++ b/Engine/Source/Runtime/Engine/Private/PhysicsEngine/BodyInstance.cpp
@@ -897,7 +897,7 @@ void FBodyInstance::UpdatePhysicsFilterData()
 
 #if WITH_CHAOS
 				// If this shape shouldn't collide in the sim we disable it here until we have more support.
-				Actor->SetShapeCollisionDisable(ShapeIndex, !CollisionEnabledHasPhysics(GetCollisionEnabled()));
+				Shape.Shape->SetDisable(!CollisionEnabledHasPhysics(GetCollisionEnabled()));
 #endif
 
 				// Apply new collision settings to this shape
@@ -1279,7 +1279,7 @@ void FInitBodiesHelperBase::InitBodies()
 					if (FPhysicsInterface::IsValid(ActorHandle))
 					{
 
-						PhysScene->AddToComponentMaps(BI->OwnerComponent.Get(), ActorHandle->Proxy);
+						PhysScene->AddToComponentMaps(BI->OwnerComponent.Get(), ActorHandle->GetProxy());
 					}
 					if (BI->bNotifyRigidBodyCollision)
 					{
@@ -1433,7 +1433,7 @@ void FBodyInstance::TermBody(bool bNeverDeferRelease)
 		{
 			if (FPhysicsInterface::IsValid(ActorHandle))
 			{
-				PhysScene->RemoveFromComponentMaps(ActorHandle->Proxy);
+				PhysScene->RemoveFromComponentMaps(ActorHandle->GetProxy());
 			}
 			if (bNotifyRigidBodyCollision)
 			{
diff --git a/Engine/Source/Runtime/Engine/Private/PhysicsEngine/CollisionQueryFilterCallback.cpp b/Engine/Source/Runtime/Engine/Private/PhysicsEngine/CollisionQueryFilterCallback.cpp
--- a/Engine/Source/Runtime/Engine/Private/PhysicsEngine/CollisionQueryFilterCallback.cpp
+++ b/Engine/Source/Runtime/Engine/Private/PhysicsEngine/CollisionQueryFilterCallback.cpp
@@ -108,7 +108,7 @@ ECollisionQueryHitType FCollisionQueryFilterCallback::PreFilterImp(const FCollis
 }
 #endif
 
-ECollisionQueryHitType FCollisionQueryFilterCallback::PreFilterImp(const FCollisionFilterData& FilterData, const Chaos::TPerShapeData<float,3>& Shape, const Chaos::TGeometryParticle<float,3>& Actor)
+ECollisionQueryHitType FCollisionQueryFilterCallback::PreFilterImp(const FCollisionFilterData& FilterData, const Chaos::FPerShapeData& Shape, const Chaos::TGeometryParticle<float,3>& Actor)
 {
 	//SCOPE_CYCLE_COUNTER(STAT_Collision_PreFilter);
 	FCollisionFilterData ShapeFilter = ChaosInterface::GetQueryFilterData(Shape);
diff --git a/Engine/Source/Runtime/Engine/Private/PhysicsEngine/CollisionQueryFilterCallback.h b/Engine/Source/Runtime/Engine/Private/PhysicsEngine/CollisionQueryFilterCallback.h
--- a/Engine/Source/Runtime/Engine/Private/PhysicsEngine/CollisionQueryFilterCallback.h
+++ b/Engine/Source/Runtime/Engine/Private/PhysicsEngine/CollisionQueryFilterCallback.h
@@ -69,14 +69,14 @@ public:
 	ECollisionQueryHitType PreFilterImp(const FCollisionFilterData& FilterData, const physx::PxShape& Shape, const physx::PxActor& Actor);
 #endif
 
-	ECollisionQueryHitType PreFilterImp(const FCollisionFilterData& FilterData, const Chaos::TPerShapeData<float,3>& Shape, const Chaos::TGeometryParticle<float,3>& Actor);
+	ECollisionQueryHitType PreFilterImp(const FCollisionFilterData& FilterData, const Chaos::FPerShapeData& Shape, const Chaos::TGeometryParticle<float,3>& Actor);
 	ECollisionQueryHitType PostFilterImp(const FCollisionFilterData& FilterData, const ChaosInterface::FQueryHit& Hit);
 
 	virtual ECollisionQueryHitType PostFilter(const FCollisionFilterData& FilterData, const ChaosInterface::FQueryHit& Hit) override
 	{
 		return PostFilterImp(FilterData, Hit);
 	}
-	virtual ECollisionQueryHitType PreFilter(const FCollisionFilterData& FilterData, const Chaos::TPerShapeData<float,3>& Shape, const Chaos::TGeometryParticle<float,3>& Actor) override
+	virtual ECollisionQueryHitType PreFilter(const FCollisionFilterData& FilterData, const Chaos::FPerShapeData& Shape, const Chaos::TGeometryParticle<float,3>& Actor) override
 	{
 		return PreFilterImp(FilterData, Shape, Actor);
 	}
diff --git a/Engine/Source/Runtime/Engine/Private/PhysicsEngine/Experimental/ChaosInterfaceUtils.cpp b/Engine/Source/Runtime/Engine/Private/PhysicsEngine/Experimental/ChaosInterfaceUtils.cpp
--- a/Engine/Source/Runtime/Engine/Private/PhysicsEngine/Experimental/ChaosInterfaceUtils.cpp
+++ b/Engine/Source/Runtime/Engine/Private/PhysicsEngine/Experimental/ChaosInterfaceUtils.cpp
@@ -150,12 +150,12 @@ namespace ChaosInterface
 		return Chaos::EChaosCollisionTraceFlag::Chaos_CTF_UseDefault;
 	}
 
-	void CreateGeometry(const FGeometryAddParams& InParams, TArray<TUniquePtr<Chaos::FImplicitObject>>& OutGeoms, Chaos::TShapesArray<float, 3>& OutShapes)
+	void CreateGeometry(const FGeometryAddParams& InParams, TArray<TUniquePtr<Chaos::FImplicitObject>>& OutGeoms, Chaos::FShapesArray& OutShapes)
 	{
 		LLM_SCOPE(ELLMTag::ChaosGeometry);
 		const FVector& Scale = InParams.Scale;
 		TArray<TUniquePtr<Chaos::FImplicitObject>>& Geoms = OutGeoms;
-		Chaos::TShapesArray<float, 3>& Shapes = OutShapes;
+		Chaos::FShapesArray& Shapes = OutShapes;
 
 		ECollisionTraceFlag CollisionTraceType = InParams.CollisionTraceType;
 		if (CollisionTraceType == CTF_UseDefault)
@@ -177,16 +177,16 @@ namespace ChaosInterface
 
 		ensure(bMakeComplexGeometry || bMakeSimpleGeometry);
 
-		auto NewShapeHelper = [&InParams, &CollisionTraceType](Chaos::TSerializablePtr<Chaos::FImplicitObject> InGeom, void* UserData, bool bComplexShape = false)
+		auto NewShapeHelper = [&InParams, &CollisionTraceType](Chaos::TSerializablePtr<Chaos::FImplicitObject> InGeom, int32 ShapeIdx, void* UserData, bool bComplexShape = false)
 		{
-			auto NewShape = Chaos::TPerShapeData<float, 3>::CreatePerShapeData();
-			NewShape->Geometry = InGeom;
-			NewShape->QueryData = bComplexShape ? InParams.CollisionData.CollisionFilterData.QueryComplexFilter : InParams.CollisionData.CollisionFilterData.QuerySimpleFilter;
-			NewShape->SimData = InParams.CollisionData.CollisionFilterData.SimFilter;
-			NewShape->CollisionTraceType = ConvertCollisionTraceFlag(CollisionTraceType);
+			auto NewShape = Chaos::FPerShapeData::CreatePerShapeData(ShapeIdx);
+			NewShape->SetGeometry(InGeom);
+			NewShape->SetQueryData(bComplexShape ? InParams.CollisionData.CollisionFilterData.QueryComplexFilter : InParams.CollisionData.CollisionFilterData.QuerySimpleFilter);
+			NewShape->SetSimData(InParams.CollisionData.CollisionFilterData.SimFilter);
+			NewShape->SetCollisionTraceType(ConvertCollisionTraceFlag(CollisionTraceType));
 			NewShape->UpdateShapeBounds(InParams.WorldTransform);
-			NewShape->UserData = UserData;
-			NewShape->bSimulate = bComplexShape ? InParams.CollisionData.CollisionFlags.bEnableSimCollisionComplex : InParams.CollisionData.CollisionFlags.bEnableSimCollisionSimple;
+			NewShape->SetUserData(UserData);
+			NewShape->SetSimulate(bComplexShape ? InParams.CollisionData.CollisionFlags.bEnableSimCollisionComplex : InParams.CollisionData.CollisionFlags.bEnableSimCollisionSimple);
 			return NewShape;
 		};
 
@@ -198,7 +198,7 @@ namespace ChaosInterface
 				const FKSphereElem ScaledSphereElem = SphereElem.GetFinalScaled(Scale, InParams.LocalTransform);
 				const float UseRadius = FMath::Max(ScaledSphereElem.Radius, KINDA_SMALL_NUMBER);
 				auto ImplicitSphere = MakeUnique<Chaos::TSphere<float, 3>>(ScaledSphereElem.Center, UseRadius);
-				auto NewShape = NewShapeHelper(MakeSerializable(ImplicitSphere), (void*)SphereElem.GetUserData());
+				auto NewShape = NewShapeHelper(MakeSerializable(ImplicitSphere), Shapes.Num(), (void*)SphereElem.GetUserData());
 				Shapes.Emplace(MoveTemp(NewShape));
 				Geoms.Add(MoveTemp(ImplicitSphere));
 			}
@@ -226,7 +226,7 @@ namespace ChaosInterface
 					Implicit = MakeUnique<Chaos::TBox<float, 3>>(BoxTransform.GetTranslation() - HalfExtents, BoxTransform.GetTranslation() + HalfExtents);
 				}
 
-				auto NewShape = NewShapeHelper(MakeSerializable(Implicit), (void*)BoxElem.GetUserData());
+				auto NewShape = NewShapeHelper(MakeSerializable(Implicit),Shapes.Num(), (void*)BoxElem.GetUserData());
 				Shapes.Emplace(MoveTemp(NewShape));
 				Geoms.Add(MoveTemp(Implicit));
 			}
@@ -241,7 +241,7 @@ namespace ChaosInterface
 				{
 					//not a capsule just use a sphere
 					auto ImplicitSphere = MakeUnique<Chaos::TSphere<float, 3>>(ScaledSphylElem.Center, Radius);
-					auto NewShape = NewShapeHelper(MakeSerializable(ImplicitSphere), (void*)UnscaledSphyl.GetUserData());
+					auto NewShape = NewShapeHelper(MakeSerializable(ImplicitSphere),Shapes.Num(), (void*)UnscaledSphyl.GetUserData());
 					Shapes.Emplace(MoveTemp(NewShape));
 					Geoms.Add(MoveTemp(ImplicitSphere));
 
@@ -251,7 +251,7 @@ namespace ChaosInterface
 					Chaos::TVector<float, 3> HalfExtents = ScaledSphylElem.Rotation.RotateVector(Chaos::TVector<float, 3>(0, 0, HalfHeight));
 
 					auto ImplicitCapsule = MakeUnique<Chaos::TCapsule<float>>(ScaledSphylElem.Center - HalfExtents, ScaledSphylElem.Center + HalfExtents, Radius);
-					auto NewShape = NewShapeHelper(MakeSerializable(ImplicitCapsule), (void*)UnscaledSphyl.GetUserData());
+					auto NewShape = NewShapeHelper(MakeSerializable(ImplicitCapsule),Shapes.Num(), (void*)UnscaledSphyl.GetUserData());
 					Shapes.Emplace(MoveTemp(NewShape));
 					Geoms.Add(MoveTemp(ImplicitCapsule));
 				}
@@ -311,7 +311,7 @@ namespace ChaosInterface
 							Implicit = TUniquePtr<Chaos::FImplicitObject>(new Chaos::TImplicitObjectScaled<Chaos::FConvex>(ConvexImplicit, Scale));
 						}
 
-						auto NewShape = NewShapeHelper(MakeSerializable(Implicit), (void*)CollisionBody.GetUserData());
+						auto NewShape = NewShapeHelper(MakeSerializable(Implicit),Shapes.Num(), (void*)CollisionBody.GetUserData());
 						Shapes.Emplace(MoveTemp(NewShape));
 						Geoms.Add(MoveTemp(Implicit));
 					}
@@ -335,7 +335,7 @@ namespace ChaosInterface
 
 				ChaosTriMesh->SetCullsBackFaceRaycast(!InParams.bDoubleSided);
 
-				auto NewShape = NewShapeHelper(MakeSerializable(Implicit), nullptr, true);
+				auto NewShape = NewShapeHelper(MakeSerializable(Implicit),Shapes.Num(), nullptr, true);
 				Shapes.Emplace(MoveTemp(NewShape));
 				Geoms.Add(MoveTemp(Implicit));
 			}
@@ -345,7 +345,7 @@ namespace ChaosInterface
 		for (const auto& PhysXMesh : InParams.TriMeshes)
 		{
 			auto Implicit = ConvertPhysXMeshToLevelset(PhysXMesh, Scale);
-			auto NewShape = NewShapeHelper(MakeSerializable(Implicit), nullptr, true);
+			auto NewShape = NewShapeHelper(MakeSerializable(Implicit),Shapes.Num(), nullptr, true);
 			Shapes.Emplace(MoveTemp(NewShape));
 			Geoms.Add(MoveTemp(Implicit));
 
diff --git a/Engine/Source/Runtime/Engine/Private/PhysicsEngine/Experimental/PhysInterface_Chaos.cpp b/Engine/Source/Runtime/Engine/Private/PhysicsEngine/Experimental/PhysInterface_Chaos.cpp
--- a/Engine/Source/Runtime/Engine/Private/PhysicsEngine/Experimental/PhysInterface_Chaos.cpp
+++ b/Engine/Source/Runtime/Engine/Private/PhysicsEngine/Experimental/PhysInterface_Chaos.cpp
@@ -84,31 +84,32 @@ DECLARE_DWORD_COUNTER_STAT(TEXT("(ASync) Mobile Bodies"), STAT_NumMobileBodiesAs
 DECLARE_DWORD_COUNTER_STAT(TEXT("(ASync) Static Bodies"), STAT_NumStaticBodiesAsync, STATGROUP_Physics);
 DECLARE_DWORD_COUNTER_STAT(TEXT("(ASync) Shapes"), STAT_NumShapesAsync, STATGROUP_Physics);
 
-ECollisionShapeType GetGeometryType(const Chaos::TPerShapeData<float, 3>& Shape)
+ECollisionShapeType GetGeometryType(const Chaos::FPerShapeData& Shape)
 {
-	return GetType(*Shape.Geometry);
+	return GetType(*Shape.GetGeometry());
 }
 
 Chaos::FChaosPhysicsMaterial* GetMaterialFromInternalFaceIndex(const FPhysicsShape& Shape, const FPhysicsActor& Actor, uint32 InternalFaceIndex)
 {
-	if(Shape.Materials.Num() > 0 && Actor.Proxy)
+	const auto& Materials = Shape.GetMaterials();
+	if(Materials.Num() > 0 && Actor.GetProxy())
 	{
-		Chaos::FPBDRigidsSolver* Solver = Actor.Proxy->GetSolver();
+		Chaos::FPBDRigidsSolver* Solver = Actor.GetProxy()->GetSolver();
 
 		if(ensure(Solver))
 		{
-			if(Shape.Materials.Num() == 1)
+			if(Materials.Num() == 1)
 			{
 				Chaos::TSolverQueryMaterialScope<Chaos::ELockType::Read> Scope(Solver);
-				return Solver->GetQueryMaterials().Get(Shape.Materials[0].InnerHandle);
+				return Solver->GetQueryMaterials().Get(Materials[0].InnerHandle);
 			}
 
-			uint8 Index = Shape.Geometry->GetMaterialIndex(InternalFaceIndex);
+			uint8 Index = Shape.GetGeometry()->GetMaterialIndex(InternalFaceIndex);
 
-			if(Shape.Materials.IsValidIndex(Index))
+			if(Materials.IsValidIndex(Index))
 			{
 				Chaos::TSolverQueryMaterialScope<Chaos::ELockType::Read> Scope(Solver);
-				return Solver->GetQueryMaterials().Get(Shape.Materials[Index].InnerHandle);
+				return Solver->GetQueryMaterials().Get(Materials[Index].InnerHandle);
 			}
 		}
 	}
@@ -121,13 +122,13 @@ Chaos::FChaosPhysicsMaterial* GetMaterialFromInternalFaceIndexAndHitLocation(con
 	{
 		SCOPE_CYCLE_COUNTER(STAT_QueryPhysicalMaterialMaskHit);
 
-		if (Shape.Materials.Num() > 0 && Actor.Proxy)
+		if (Shape.GetMaterials().Num() > 0 && Actor.GetProxy())
 		{
-			Chaos::FPBDRigidsSolver* Solver = Actor.Proxy->GetSolver();
+			Chaos::FPBDRigidsSolver* Solver = Actor.GetProxy()->GetSolver();
 
 			if (ensure(Solver))
 			{
-				if (Shape.MaterialMasks.Num() > 0)
+				if (Shape.GetMaterialMasks().Num() > 0)
 				{
 					UBodySetup* BodySetup = nullptr;
 
@@ -140,7 +141,7 @@ Chaos::FChaosPhysicsMaterial* GetMaterialFromInternalFaceIndexAndHitLocation(con
 						if (BodySetup && BodySetup->bSupportUVsAndFaceRemap && GetGeometryType(Shape) == ECollisionShapeType::Trimesh)
 						{
 							FVector Scale(1.0f, 1.0f, 1.0f);
-							const Chaos::FImplicitObject* Geometry = Shape.Geometry.Get();
+							const Chaos::FImplicitObject* Geometry = Shape.GetGeometry().Get();
 							if (const Chaos::TImplicitObjectScaled<Chaos::FTriangleMeshImplicitObject>* ScaledTrimesh = Chaos::TImplicitObjectScaled<Chaos::FTriangleMeshImplicitObject>::AsScaled(*Geometry))
 							{
 								Scale = ScaledTrimesh->GetScale();
@@ -150,13 +151,13 @@ Chaos::FChaosPhysicsMaterial* GetMaterialFromInternalFaceIndexAndHitLocation(con
 							Chaos::FRigidTransform3 ActorToWorld(Actor.X(), Actor.R(), Scale);
 							const FVector LocalHitPos = ActorToWorld.InverseTransformPosition(HitLocation);
 
-							uint8 Index = Shape.Geometry->GetMaterialIndex(InternalFaceIndex);
-							if (Shape.MaterialMasks.IsValidIndex(Index))
+							uint8 Index = Shape.GetGeometry()->GetMaterialIndex(InternalFaceIndex);
+							if (Shape.GetMaterialMasks().IsValidIndex(Index))
 							{
 								Chaos::FChaosPhysicsMaterialMask* Mask = nullptr;
 								{
 									Chaos::TSolverQueryMaterialScope<Chaos::ELockType::Read> Scope(Solver);
-									Mask = Solver->GetQueryMaterialMasks().Get(Shape.MaterialMasks[Index].InnerHandle);
+									Mask = Solver->GetQueryMaterialMasks().Get(Shape.GetMaterialMasks()[Index].InnerHandle);
 								}
 
 								if (Mask && InternalFaceIndex < (uint32)BodySetup->FaceRemap.Num())
@@ -169,13 +170,13 @@ Chaos::FChaosPhysicsMaterial* GetMaterialFromInternalFaceIndexAndHitLocation(con
 									{
 										uint32 MapIdx = UPhysicalMaterialMask::GetPhysMatIndex(Mask->MaskData, Mask->SizeX, Mask->SizeY, Mask->AddressX, Mask->AddressY, UV.X, UV.Y);
 										uint32 AdjustedMapIdx = Index * EPhysicalMaterialMaskColor::MAX + MapIdx;
-										if (Shape.MaterialMaskMaps.IsValidIndex(AdjustedMapIdx))
+										if (Shape.GetMaterialMaskMaps().IsValidIndex(AdjustedMapIdx))
 										{
-											uint32 MaterialIdx = Shape.MaterialMaskMaps[AdjustedMapIdx];
-											if (Shape.MaterialMaskMapMaterials.IsValidIndex(MaterialIdx))
+											uint32 MaterialIdx = Shape.GetMaterialMaskMaps()[AdjustedMapIdx];
+											if (Shape.GetMaterialMaskMapMaterials().IsValidIndex(MaterialIdx))
 											{
 												Chaos::TSolverQueryMaterialScope<Chaos::ELockType::Read> Scope(Solver);
-												return Solver->GetQueryMaterials().Get(Shape.MaterialMaskMapMaterials[MaterialIdx].InnerHandle);
+												return Solver->GetQueryMaterials().Get(Shape.GetMaterialMaskMapMaterials()[MaterialIdx].InnerHandle);
 											}
 										}
 									}
@@ -193,7 +194,7 @@ Chaos::FChaosPhysicsMaterial* GetMaterialFromInternalFaceIndexAndHitLocation(con
 
 const Chaos::FImplicitObject& FPhysicsShapeReference_Chaos::GetGeometry() const
 {
-	check(IsValid()); return *Shape->Geometry;
+	check(IsValid()); return *Shape->GetGeometry();
 }
 
 FPhysInterface_Chaos::FPhysInterface_Chaos(const AWorldSettings* Settings) 
@@ -279,7 +280,7 @@ void FPhysInterface_Chaos::ReleaseActor(FPhysicsActorHandle& Handle, FPhysScene*
 
 void FPhysInterface_Chaos::RemoveActorFromSolver(FPhysicsActorHandle& Handle, Chaos::FPhysicsSolver* Solver, Chaos::IDispatcher* Dispatcher)
 {
-	if (Solver && Handle->Proxy)
+	if (Solver && Handle->GetProxy())
 	{
 		Solver->UnregisterObject(Handle);
 	}
@@ -363,7 +364,7 @@ void FPhysInterface_Chaos::SetUserData(const FPhysicsShapeHandle& InShape, void*
 {
 	if (CHAOS_ENSURE(InShape.Shape))
 	{
-		InShape.Shape->UserData = InUserData;
+		InShape.Shape->SetUserData(InUserData);
 	}
 }
 
@@ -371,7 +372,7 @@ void* FPhysInterface_Chaos::GetUserData(const FPhysicsShapeHandle& InShape)
 {
 	if (ensure(InShape.Shape))
 	{
-		return InShape.Shape->UserData;
+		return InShape.Shape->GetUserData();
 	}
 	return nullptr;
 }
@@ -456,7 +457,7 @@ FPhysScene* FPhysInterface_Chaos::GetCurrentScene(const FPhysicsActorHandle& InH
 		return nullptr;
 	}
 
-	if (IPhysicsProxyBase* Proxy = InHandle->Proxy)
+	if (IPhysicsProxyBase* Proxy = InHandle->GetProxy())
 	{
 		Chaos::FPhysicsSolver* Solver = Proxy->GetSolver();
 		return static_cast<FPhysScene*>(Solver ? Solver->PhysSceneHack : nullptr);
@@ -1443,7 +1444,7 @@ FPhysicsShapeHandle FPhysInterface_Chaos::CreateShape(physx::PxGeometry* InGeom,
 	return { nullptr, bSimulation, bQuery, NewActor };
 }
 
-const FBodyInstance* FPhysInterface_Chaos::ShapeToOriginalBodyInstance(const FBodyInstance* InCurrentInstance, const Chaos::TPerShapeData<float, 3>* InShape)
+const FBodyInstance* FPhysInterface_Chaos::ShapeToOriginalBodyInstance(const FBodyInstance* InCurrentInstance, const Chaos::FPerShapeData* InShape)
 {
 	//question: this is identical to physx version, should it be in body instance?
 	check(InCurrentInstance);
@@ -1472,13 +1473,13 @@ void FPhysInterface_Chaos::AddGeometry(FPhysicsActorHandle& InActor, const FGeom
 {
 	LLM_SCOPE(ELLMTag::ChaosGeometry);
 	TArray<TUniquePtr<Chaos::FImplicitObject>> Geoms;
-	Chaos::TShapesArray<float, 3> Shapes;
+	Chaos::FShapesArray Shapes;
 	ChaosInterface::CreateGeometry(InParams, Geoms, Shapes);
 
 #if WITH_CHAOS
 	if (InActor && Geoms.Num())
 	{
-		for (TUniquePtr<Chaos::TPerShapeData<float, 3>>& Shape : Shapes)
+		for (TUniquePtr<Chaos::FPerShapeData>& Shape : Shapes)
 		{
 			FPhysicsShapeHandle NewHandle(Shape.Get(), true, true, InActor);
 			if (OutOptShapes)
@@ -1566,7 +1567,7 @@ FCollisionFilterData FPhysInterface_Chaos::GetSimulationFilter(const FPhysicsSha
 {
 	if (ensure(InShape.Shape))
 	{
-		return InShape.Shape->SimData;
+		return InShape.Shape->GetSimData();
 	}
 	else
 	{
@@ -1578,7 +1579,7 @@ FCollisionFilterData FPhysInterface_Chaos::GetQueryFilter(const FPhysicsShapeRef
 {
 	if (ensure(InShape.Shape))
 	{
-		return InShape.Shape->QueryData;
+		return InShape.Shape->GetQueryData();
 	}
 	else
 	{
@@ -1588,13 +1589,12 @@ FCollisionFilterData FPhysInterface_Chaos::GetQueryFilter(const FPhysicsShapeRef
 
 void FPhysInterface_Chaos::SetQueryFilter(const FPhysicsShapeReference_Chaos& InShapeRef, const FCollisionFilterData& InFilter)
 {
-	InShapeRef.Shape->QueryData = InFilter;
+	InShapeRef.Shape->SetQueryData(InFilter);
 }
 
 void FPhysInterface_Chaos::SetSimulationFilter(const FPhysicsShapeReference_Chaos& InShapeRef, const FCollisionFilterData& InFilter)
 {
-	InShapeRef.Shape->SimData = InFilter;
-	InShapeRef.ActorRef->MarkShapeSimDataDirty();
+	InShapeRef.Shape->SetSimData(InFilter);
 }
 
 bool FPhysInterface_Chaos::IsSimulationShape(const FPhysicsShapeHandle& InShape)
@@ -1609,15 +1609,15 @@ bool FPhysInterface_Chaos::IsQueryShape(const FPhysicsShapeHandle& InShape)
 
 ECollisionShapeType FPhysInterface_Chaos::GetShapeType(const FPhysicsShapeReference_Chaos& InShapeRef)
 {
-	return GetImplicitType(*InShapeRef.Shape->Geometry);
+	return GetImplicitType(*InShapeRef.Shape->GetGeometry());
 }
 
 FTransform FPhysInterface_Chaos::GetLocalTransform(const FPhysicsShapeReference_Chaos& InShapeRef)
 {
     // Transforms are baked into the object so there is never a local transform
-    if (InShapeRef.Shape->Geometry->GetType() == Chaos::ImplicitObjectType::Transformed && FPhysicsInterface::IsValid(InShapeRef.ActorRef))
+    if (InShapeRef.Shape->GetGeometry()->GetType() == Chaos::ImplicitObjectType::Transformed && FPhysicsInterface::IsValid(InShapeRef.ActorRef))
     {
-        return InShapeRef.Shape->Geometry->GetObject<Chaos::TImplicitObjectTransformed<float, 3>>()->GetTransform();
+        return InShapeRef.Shape->GetGeometry()->GetObject<Chaos::TImplicitObjectTransformed<float, 3>>()->GetTransform();
     }
     else
     {
@@ -1666,7 +1666,7 @@ void FPhysInterface_Chaos::SetMaterials(const FPhysicsShapeHandle& InShape, cons
 		NewMaterialHandles.Add(UnrealMaterial->GetPhysicsMaterial());
 	}
 
-	InShape.Shape->Materials = NewMaterialHandles;
+	InShape.Shape->SetMaterials(NewMaterialHandles);
 }
 
 void FPhysInterface_Chaos::SetMaterials(const FPhysicsShapeHandle& InShape, const TArrayView<UPhysicalMaterial*> InMaterials, const TArrayView<FPhysicalMaterialMaskParams>& InMaterialMasks)
@@ -1684,34 +1684,38 @@ void FPhysInterface_Chaos::SetMaterials(const FPhysicsShapeHandle& InShape, cons
 
 		int MaskMapMatIdx = 0;
 
-		for (FPhysicalMaterialMaskParams& MaterialMaskData : InMaterialMasks)
+		InShape.Shape->ModifyMaterialMaskMaps([&](auto& MaterialMaskMaps)
 		{
-			if (MaterialMaskData.PhysicalMaterialMask && ensure(MaterialMaskData.PhysicalMaterialMap))
+			for(FPhysicalMaterialMaskParams& MaterialMaskData : InMaterialMasks)
+		{
+				if(MaterialMaskData.PhysicalMaterialMask && ensure(MaterialMaskData.PhysicalMaterialMap))
 			{
 				NewMaterialMaskHandles.Add(MaterialMaskData.PhysicalMaterialMask->GetPhysicsMaterialMask());
-				for (int i = 0; i < EPhysicalMaterialMaskColor::MAX; i++)
+					for(int i = 0; i < EPhysicalMaterialMaskColor::MAX; i++)
 				{
-					if (UPhysicalMaterial* MapMat = MaterialMaskData.PhysicalMaterialMap->GetPhysicalMaterialFromMap(i))
+						if(UPhysicalMaterial* MapMat = MaterialMaskData.PhysicalMaterialMap->GetPhysicalMaterialFromMap(i))
 					{
-						InShape.Shape->MaterialMaskMaps.Emplace(MaskMapMatIdx);
+							MaterialMaskMaps.Emplace(MaskMapMatIdx);
 						MaskMapMatIdx++;
-					}
-					else
+						} else
 					{
-						InShape.Shape->MaterialMaskMaps.Emplace(INDEX_NONE);
-					}
+							MaterialMaskMaps.Emplace(INDEX_NONE);
 				}
 			}
-			else
+				} else
 			{
 				NewMaterialMaskHandles.Add(Chaos::FMaterialMaskHandle());
-				for (int i = 0; i < EPhysicalMaterialMaskColor::MAX; i++)
+					for(int i = 0; i < EPhysicalMaterialMaskColor::MAX; i++)
 				{
-					InShape.Shape->MaterialMaskMaps.Emplace(INDEX_NONE);
+						MaterialMaskMaps.Emplace(INDEX_NONE);
 				}
 			}
 		}
 
+		});
+		
+
+		
 		if (MaskMapMatIdx > 0)
 		{
 			NewMaterialMaskMaterialHandles.Reserve(MaskMapMatIdx);
@@ -1733,8 +1737,8 @@ void FPhysInterface_Chaos::SetMaterials(const FPhysicsShapeHandle& InShape, cons
 			}
 		}
 
-		InShape.Shape->MaterialMasks = NewMaterialMaskHandles;
-		InShape.Shape->MaterialMaskMapMaterials = NewMaterialMaskMaterialHandles;
+		InShape.Shape->SetMaterialMasks(NewMaterialMaskHandles);
+		InShape.Shape->SetMaterialMaskMapMaterials(NewMaterialMaskMaterialHandles);
 	}
 }
 
@@ -1799,9 +1803,9 @@ void FPhysInterface_Chaos::CalculateMassPropertiesFromShapeCollection(Chaos::TMa
 	TArray< Chaos::TMassProperties<float, 3> > MassPropertiesList;
 	for (const FPhysicsShapeHandle& ShapeHandle : InShapes)
 	{
-		if (const Chaos::TPerShapeData<float, 3>* Shape = ShapeHandle.Shape)
+		if (const Chaos::FPerShapeData* Shape = ShapeHandle.Shape)
 		{
-			if (const Chaos::FImplicitObject* ImplicitObject = Shape->Geometry.Get())
+			if (const Chaos::FImplicitObject* ImplicitObject = Shape->GetGeometry().Get())
 			{
 				FTransform WorldTransform(ShapeHandle.ActorRef->R(), ShapeHandle.ActorRef->X());
 				Chaos::TMassProperties<float, 3> MassProperties;
@@ -1875,7 +1879,7 @@ bool FPhysInterface_Chaos::LineTrace_Geom(FHitResult& OutHit, const FBodyInstanc
 					{
 						// #PHYS2 - SHAPES - Resolve this single cast case
 						FPhysicsShapeReference_Chaos& ShapeRef = Shapes[ShapeIdx];
-						Chaos::TPerShapeData<float, 3>* Shape = ShapeRef.Shape;
+						Chaos::FPerShapeData* Shape = ShapeRef.Shape;
 						check(Shape);
 
 						if (TargetInstance->IsShapeBoundToBody(ShapeRef) == false)
@@ -1884,7 +1888,7 @@ bool FPhysInterface_Chaos::LineTrace_Geom(FHitResult& OutHit, const FBodyInstanc
 						}
 
 						// Filter so we trace against the right kind of collision
-						FCollisionFilterData ShapeFilter = Shape->QueryData;
+						FCollisionFilterData ShapeFilter = Shape->GetQueryData();
 						const bool bShapeIsComplex = (ShapeFilter.Word3 & EPDF_ComplexCollision) != 0;
 						const bool bShapeIsSimple = (ShapeFilter.Word3 & EPDF_SimpleCollision) != 0;
 						if ((bTraceComplex && bShapeIsComplex) || (!bTraceComplex && bShapeIsSimple))
@@ -1895,7 +1899,7 @@ bool FPhysInterface_Chaos::LineTrace_Geom(FHitResult& OutHit, const FBodyInstanc
 							Chaos::TVector<float, 3> LocalNormal;
 
 							int32 FaceIndex;
-							if (Shape->Geometry->Raycast(LocalStart, LocalDelta / DeltaMag, DeltaMag, 0, Distance, LocalPosition, LocalNormal, FaceIndex))
+							if (Shape->GetGeometry()->Raycast(LocalStart, LocalDelta / DeltaMag, DeltaMag, 0, Distance, LocalPosition, LocalNormal, FaceIndex))
 							{
 								if (Distance < BestHit.Distance)
 								{
@@ -1978,7 +1982,7 @@ bool FPhysInterface_Chaos::Sweep_Geom(FHitResult& OutHit, const FBodyInstance* I
 					for (int32 ShapeIdx = 0; ShapeIdx < NumShapes; ShapeIdx++)
 					{
 						FPhysicsShapeReference_Chaos& ShapeRef = Shapes[ShapeIdx];
-						Chaos::TPerShapeData<float, 3>* Shape = ShapeRef.Shape;
+						Chaos::FPerShapeData* Shape = ShapeRef.Shape;
 						check(Shape);
 
 						// Skip shapes not bound to this instance
@@ -1988,7 +1992,7 @@ bool FPhysInterface_Chaos::Sweep_Geom(FHitResult& OutHit, const FBodyInstance* I
 						}
 
 						// Filter so we trace against the right kind of collision
-						FCollisionFilterData ShapeFilter = Shape->QueryData;;
+						FCollisionFilterData ShapeFilter = Shape->GetQueryData();
 						const bool bShapeIsComplex = (ShapeFilter.Word3 & EPDF_ComplexCollision) != 0;
 						const bool bShapeIsSimple = (ShapeFilter.Word3 & EPDF_SimpleCollision) != 0;
 						if ((bSweepComplex && bShapeIsComplex) || (!bSweepComplex && bShapeIsSimple))
@@ -1997,7 +2001,7 @@ bool FPhysInterface_Chaos::Sweep_Geom(FHitResult& OutHit, const FBodyInstance* I
 							Chaos::TVector<float, 3> WorldPosition;
 							Chaos::TVector<float, 3> WorldNormal;
 							int32 FaceIdx;
-							if (Chaos::Utilities::CastHelper(ShapeAdapter.GetGeometry(), ActorTM, [&](const auto& Downcast, const auto& FullActorTM) { return Chaos::SweepQuery(*Shape->Geometry, FullActorTM, Downcast, StartTM, Dir, DeltaMag, Hit.Distance, WorldPosition, WorldNormal, FaceIdx, 0.f, false); }))
+							if (Chaos::Utilities::CastHelper(ShapeAdapter.GetGeometry(), ActorTM, [&](const auto& Downcast, const auto& FullActorTM) { return Chaos::SweepQuery(*Shape->GetGeometry(), FullActorTM, Downcast, StartTM, Dir, DeltaMag, Hit.Distance, WorldPosition, WorldNormal, FaceIdx, 0.f, false); }))
 							{
 								// we just like to make sure if the hit is made
 								FCollisionFilterData QueryFilter;
@@ -2049,7 +2053,7 @@ bool Overlap_GeomInternal(const FBodyInstance* InInstance, const Chaos::FImplici
 	for (int32 ShapeIdx = 0; ShapeIdx < NumShapes; ++ShapeIdx)
 	{
 		FPhysicsShapeReference_Chaos& ShapeRef = Shapes[ShapeIdx];
-		const Chaos::TPerShapeData<float,3>* Shape = ShapeRef.Shape;
+		const Chaos::FPerShapeData* Shape = ShapeRef.Shape;
 		check(Shape);
 
 		if (TargetInstance->IsShapeBoundToBody(ShapeRef))
@@ -2057,7 +2061,7 @@ bool Overlap_GeomInternal(const FBodyInstance* InInstance, const Chaos::FImplici
 			if (OutOptResult)
 			{
 				Chaos::FMTDInfo MTDInfo;
-				if (Chaos::Utilities::CastHelper(InGeom, ActorTM, [&](const auto& Downcast, const auto& FullActorTM) { return Chaos::OverlapQuery(*Shape->Geometry, FullActorTM, Downcast, GeomTransform, /*Thickness=*/0, &MTDInfo); }))
+				if (Chaos::Utilities::CastHelper(InGeom, ActorTM, [&](const auto& Downcast, const auto& FullActorTM) { return Chaos::OverlapQuery(*Shape->GetGeometry(), FullActorTM, Downcast, GeomTransform, /*Thickness=*/0, &MTDInfo); }))
 				{
 					OutOptResult->Distance = MTDInfo.Penetration;
 					OutOptResult->Direction = MTDInfo.Normal;
@@ -2066,7 +2070,7 @@ bool Overlap_GeomInternal(const FBodyInstance* InInstance, const Chaos::FImplici
 			}
 			else	//question: why do we even allow user to not pass in MTD info?
 			{
-				if (Chaos::Utilities::CastHelper(InGeom, ActorTM, [&](const auto& Downcast, const auto& FullActorTM) { return Chaos::OverlapQuery(*Shape->Geometry, FullActorTM, Downcast, GeomTransform); }))
+				if (Chaos::Utilities::CastHelper(InGeom, ActorTM, [&](const auto& Downcast, const auto& FullActorTM) { return Chaos::OverlapQuery(*Shape->GetGeometry(), FullActorTM, Downcast, GeomTransform); }))
 				{
 					return true;
 				}
@@ -2131,7 +2135,7 @@ bool FPhysInterface_Chaos::GetSquaredDistanceToBody(const FBodyInstance* InInsta
 			bFoundValidBody = true;
 
 			Chaos::TVector<float, 3> Normal;
-			const float Phi = Shape.Shape->Geometry->PhiWithNormal(LocalPoint, Normal);
+			const float Phi = Shape.Shape->GetGeometry()->PhiWithNormal(LocalPoint, Normal);
 			if (Phi <= 0)
 			{
 				OutDistanceSquared = 0;
@@ -2165,7 +2169,7 @@ bool FPhysInterface_Chaos::GetSquaredDistanceToBody(const FBodyInstance* InInsta
 uint32 GetTriangleMeshExternalFaceIndex(const FPhysicsShape& Shape, uint32 InternalFaceIndex)
 {
 	using namespace Chaos;
-	uint8 OuterType = Shape.Geometry->GetType();
+	uint8 OuterType = Shape.GetGeometry()->GetType();
 	uint8 InnerType = GetInnerType(OuterType);
 	if (ensure(InnerType == ImplicitObjectType::TriangleMesh))
 	{
@@ -2173,16 +2177,16 @@ uint32 GetTriangleMeshExternalFaceIndex(const FPhysicsShape& Shape, uint32 Inter
 
 		if (IsScaled(OuterType))
 		{
-			const TImplicitObjectScaled<FTriangleMeshImplicitObject>& ScaledTriangleMesh = Shape.Geometry->GetObjectChecked<TImplicitObjectScaled<FTriangleMeshImplicitObject>>();
+			const TImplicitObjectScaled<FTriangleMeshImplicitObject>& ScaledTriangleMesh = Shape.GetGeometry()->GetObjectChecked<TImplicitObjectScaled<FTriangleMeshImplicitObject>>();
 			TriangleMesh = ScaledTriangleMesh.GetUnscaledObject();
 		}
 		else if(IsInstanced(OuterType))
 		{
-			TriangleMesh = Shape.Geometry->GetObjectChecked<TImplicitObjectInstanced<FTriangleMeshImplicitObject>>().GetInstancedObject();
+			TriangleMesh = Shape.GetGeometry()->GetObjectChecked<TImplicitObjectInstanced<FTriangleMeshImplicitObject>>().GetInstancedObject();
 		}
 		else
 		{
-			TriangleMesh = &Shape.Geometry->GetObjectChecked<FTriangleMeshImplicitObject>();
+			TriangleMesh = &Shape.GetGeometry()->GetObjectChecked<FTriangleMeshImplicitObject>();
 		}
 
 		return TriangleMesh->GetExternalFaceIndexFromInternal(InternalFaceIndex);
@@ -2195,11 +2199,11 @@ template<typename AllocatorType>
 int32 GetAllShapesInternal_AssumedLocked(const FPhysicsActorHandle& InActorHandle, TArray<FPhysicsShapeReference_Chaos, AllocatorType>& OutShapes)
 {
 	OutShapes.Reset();
-	const Chaos::TShapesArray<float, 3>& ShapesArray = InActorHandle->ShapesArray();
+	const Chaos::FShapesArray& ShapesArray = InActorHandle->ShapesArray();
 	//todo: can we avoid this construction?
-	for (const TUniquePtr<Chaos::TPerShapeData<float, 3>>& Shape : ShapesArray)
+	for (const TUniquePtr<Chaos::FPerShapeData>& Shape : ShapesArray)
 	{
-		OutShapes.Add(FPhysicsShapeReference_Chaos(Shape.Get(), Shape->bSimulate, true, InActorHandle));
+		OutShapes.Add(FPhysicsShapeReference_Chaos(Shape.Get(), Shape->GetSimulate(), true, InActorHandle));
 	}
 	return OutShapes.Num();
 }
diff --git a/Engine/Source/Runtime/Engine/Private/PhysicsEngine/Experimental/PhysScene_Chaos.cpp b/Engine/Source/Runtime/Engine/Private/PhysicsEngine/Experimental/PhysScene_Chaos.cpp
--- a/Engine/Source/Runtime/Engine/Private/PhysicsEngine/Experimental/PhysScene_Chaos.cpp
+++ b/Engine/Source/Runtime/Engine/Private/PhysicsEngine/Experimental/PhysScene_Chaos.cpp
@@ -1507,7 +1507,6 @@ void FPhysScene_ChaosInterface::ClearForces_AssumesLocked(FBodyInstance* BodyIns
 		if (ensure(Rigid))
 		{
 			Rigid->SetF(Chaos::TVector<float, 3>(0.f,0.f,0.f));
-			Rigid->MarkClean(Chaos::EParticleFlags::F);
 		}
 	}
 }
@@ -1549,7 +1548,6 @@ void FPhysScene_ChaosInterface::ClearTorques_AssumesLocked(FBodyInstance* BodyIn
 		if (ensure(Rigid))
 		{
 			Rigid->SetTorque(Chaos::TVector<float, 3>(0.f, 0.f, 0.f));
-			Rigid->MarkClean(Chaos::EParticleFlags::Torque);
 		}
 	}
 }
@@ -2136,7 +2134,7 @@ void FPhysScene_ChaosInterface::SyncBodies(Chaos::FPhysicsSolver* Solver)
 		const Chaos::FPBDRigidActiveParticlesBufferOut* ActiveParticleBuffer = Accessor.GetSolverOutData();
 		for (Chaos::TGeometryParticle<float, 3>* ActiveParticle : ActiveParticleBuffer->ActiveGameThreadParticles)
 		{
-			if (IPhysicsProxyBase * ProxyBase = ActiveParticle->Proxy)
+		if (IPhysicsProxyBase* ProxyBase = ActiveParticle->GetProxy())
 			{
 				if (ProxyBase->GetType() == EPhysicsProxyType::SingleRigidParticleType)
 				{
diff --git a/Engine/Source/Runtime/Engine/Private/PhysicsEngine/ImmediatePhysics/ImmediatePhysicsChaos/ImmediatePhysicsActorHandle_Chaos.cpp b/Engine/Source/Runtime/Engine/Private/PhysicsEngine/ImmediatePhysics/ImmediatePhysicsChaos/ImmediatePhysicsActorHandle_Chaos.cpp
--- a/Engine/Source/Runtime/Engine/Private/PhysicsEngine/ImmediatePhysics/ImmediatePhysicsChaos/ImmediatePhysicsActorHandle_Chaos.cpp
+++ b/Engine/Source/Runtime/Engine/Private/PhysicsEngine/ImmediatePhysics/ImmediatePhysicsChaos/ImmediatePhysicsActorHandle_Chaos.cpp
@@ -154,7 +154,7 @@ namespace ImmediatePhysics_Chaos
 		}
 	}
 
-	bool CreateDefaultGeometry(const FVector& Scale, float& OutMass, Chaos::TVector<float, 3>& OutInertia, Chaos::TRigidTransform<float, 3>& OutCoMTransform, TUniquePtr<Chaos::FImplicitObject>& OutGeom, TArray<TUniquePtr<Chaos::TPerShapeData<float, 3>>>& OutShapes)
+	bool CreateDefaultGeometry(const FVector& Scale, float& OutMass, Chaos::TVector<float, 3>& OutInertia, Chaos::TRigidTransform<float, 3>& OutCoMTransform, TUniquePtr<Chaos::FImplicitObject>& OutGeom, TArray<TUniquePtr<Chaos::FPerShapeData>>& OutShapes)
 	{
 		using namespace Chaos;
 
@@ -162,11 +162,11 @@ namespace ImmediatePhysics_Chaos
 		const FReal Radius = 1.0f * Scale.GetMax();
 
 		auto ImplicitSphere = MakeUnique<Chaos::TSphere<float, 3>>(FVec3(0), Radius);
-		auto NewShape = Chaos::TPerShapeData<float, 3>::CreatePerShapeData();
-		NewShape->Geometry = MakeSerializable(ImplicitSphere);
+		auto NewShape = Chaos::FPerShapeData::CreatePerShapeData(OutShapes.Num());
+		NewShape->SetGeometry(MakeSerializable(ImplicitSphere));
 		NewShape->UpdateShapeBounds(FTransform::Identity);
-		NewShape->UserData = 0;
-		NewShape->bSimulate = false;
+		NewShape->SetUserData(nullptr);
+		NewShape->SetSimulate(false);
 
 		OutMass = Mass;
 		OutInertia = TSphere<FReal, 3>::GetInertiaTensor(Mass, Radius).GetDiagonal();
@@ -177,7 +177,7 @@ namespace ImmediatePhysics_Chaos
 		return true;
 	}
 
-	bool CreateGeometry(FBodyInstance* BodyInstance, EActorType ActorType, const FVector& Scale, float& OutMass, Chaos::TVector<float, 3>& OutInertia, Chaos::TRigidTransform<float, 3>& OutCoMTransform, TUniquePtr<Chaos::FImplicitObject>& OutGeom, TArray<TUniquePtr<Chaos::TPerShapeData<float, 3>>>& OutShapes)
+	bool CreateGeometry(FBodyInstance* BodyInstance, EActorType ActorType, const FVector& Scale, float& OutMass, Chaos::TVector<float, 3>& OutInertia, Chaos::TRigidTransform<float, 3>& OutCoMTransform, TUniquePtr<Chaos::FImplicitObject>& OutGeom, TArray<TUniquePtr<Chaos::FPerShapeData>>& OutShapes)
 	{
 		using namespace Chaos;
 
@@ -237,7 +237,7 @@ namespace ImmediatePhysics_Chaos
 #endif
 
 		TArray<TUniquePtr<FImplicitObject>> Geoms;
-		TArray<TUniquePtr<TPerShapeData<float, 3>>, TInlineAllocator<1>> Shapes;
+		TArray<TUniquePtr<FPerShapeData>, TInlineAllocator<1>> Shapes;
 		ChaosInterface::CreateGeometry(AddParams, Geoms, Shapes);
 
 		if (Geoms.Num() == 0)
diff --git a/Engine/Source/Runtime/Engine/Public/Physics/Experimental/ChaosInterfaceUtils.h b/Engine/Source/Runtime/Engine/Public/Physics/Experimental/ChaosInterfaceUtils.h
--- a/Engine/Source/Runtime/Engine/Public/Physics/Experimental/ChaosInterfaceUtils.h
+++ b/Engine/Source/Runtime/Engine/Public/Physics/Experimental/ChaosInterfaceUtils.h
@@ -13,5 +13,5 @@ namespace ChaosInterface
 	/**
 	 * Create the Chaos Geometry based on the geometry parameters.
 	 */
-	void CreateGeometry(const FGeometryAddParams& InParams, TArray<TUniquePtr<Chaos::FImplicitObject>>& OutGeoms, TArray<TUniquePtr<Chaos::TPerShapeData<float, 3>>, TInlineAllocator<1>>& OutShapes);
+	void CreateGeometry(const FGeometryAddParams& InParams, TArray<TUniquePtr<Chaos::FImplicitObject>>& OutGeoms, TArray<TUniquePtr<Chaos::FPerShapeData>, TInlineAllocator<1>>& OutShapes);
 }
diff --git a/Engine/Source/Runtime/Engine/Public/Physics/Experimental/ChaosInterfaceWrapper.h b/Engine/Source/Runtime/Engine/Public/Physics/Experimental/ChaosInterfaceWrapper.h
--- a/Engine/Source/Runtime/Engine/Public/Physics/Experimental/ChaosInterfaceWrapper.h
+++ b/Engine/Source/Runtime/Engine/Public/Physics/Experimental/ChaosInterfaceWrapper.h
@@ -15,8 +15,7 @@ struct FCollisionQueryParams;
 
 namespace Chaos
 {
-	template <class T, int d>
-	class TPerShapeData;
+	class FPerShapeData;
 }
 
 namespace ChaosInterface
diff --git a/Engine/Source/Runtime/Engine/Public/Physics/Experimental/PhysInterface_Chaos.h b/Engine/Source/Runtime/Engine/Public/Physics/Experimental/PhysInterface_Chaos.h
--- a/Engine/Source/Runtime/Engine/Public/Physics/Experimental/PhysInterface_Chaos.h
+++ b/Engine/Source/Runtime/Engine/Public/Physics/Experimental/PhysInterface_Chaos.h
@@ -67,11 +67,9 @@ namespace Chaos
 	template <typename T, int>
 	class TSphere;
 
-	template <typename T, int>
-	class TPerShapeData;
+	class FPerShapeData;
 
-	template <typename T, int d>
-	using TShapesArray = TArray<TUniquePtr<TPerShapeData<T, d>>, TInlineAllocator<1>>;
+	using FShapesArray = TArray<TUniquePtr<FPerShapeData>, TInlineAllocator<1>>;
 }
 
 class FPhysicsConstraintReference_Chaos
@@ -92,7 +90,7 @@ public:
 	
 	FPhysicsShapeReference_Chaos()
 		: Shape(nullptr), bSimulation(false), bQuery(false), ActorRef() {}
-	FPhysicsShapeReference_Chaos(Chaos::TPerShapeData<float, 3>* ShapeIn, bool bSimulationIn, bool bQueryIn, FPhysicsActorHandle ActorRefIn)
+	FPhysicsShapeReference_Chaos(Chaos::FPerShapeData* ShapeIn, bool bSimulationIn, bool bQueryIn, FPhysicsActorHandle ActorRefIn)
 		: Shape(ShapeIn),bSimulation(bSimulationIn),bQuery(bQueryIn),ActorRef(ActorRefIn){}
 	FPhysicsShapeReference_Chaos(const FPhysicsShapeReference_Chaos& Other)
 		: Shape(Other.Shape)
@@ -108,7 +106,7 @@ public:
 	const Chaos::FImplicitObject& GetGeometry() const;
 
 
-	Chaos::TPerShapeData<float, 3>* Shape;
+	Chaos::FPerShapeData* Shape;
 	bool bSimulation;
 	bool bQuery;
     FPhysicsActorHandle ActorRef;
@@ -193,7 +191,7 @@ public:
 	static bool IsValid(const FPhysicsActorHandle& Handle) { return Handle != nullptr; }
 	static void AddActorToSolver(FPhysicsActorHandle& Handle, Chaos::FPhysicsSolver* Solver, Chaos::IDispatcher* Dispatcher);
 	static void RemoveActorFromSolver(FPhysicsActorHandle& Handle, Chaos::FPhysicsSolver* Solver, Chaos::IDispatcher* Dispatcher);
-	static const FBodyInstance* ShapeToOriginalBodyInstance(const FBodyInstance* InCurrentInstance, const Chaos::TPerShapeData<float, 3>* InShape);
+	static const FBodyInstance* ShapeToOriginalBodyInstance(const FBodyInstance* InCurrentInstance, const Chaos::FPerShapeData* InShape);
 
 	static FPhysicsAggregateReference_Chaos CreateAggregate(int32 MaxBodies);
 	static void ReleaseAggregate(FPhysicsAggregateReference_Chaos& InAggregate);
@@ -379,7 +377,7 @@ public:
 	// Shape interface functions
 	static FPhysicsShapeHandle CreateShape(physx::PxGeometry* InGeom, bool bSimulation = true, bool bQuery = true, UPhysicalMaterial* InSimpleMaterial = nullptr, TArray<UPhysicalMaterial*>* InComplexMaterials = nullptr);
 	
-	static void CreateGeometry(const FGeometryAddParams& InParams, TArray<TUniquePtr<Chaos::FImplicitObject>>& OutGeoms, Chaos::TShapesArray<float, 3>& OutShapes, TArray<FPhysicsShapeHandle>* OutOptShapes);
+	static void CreateGeometry(const FGeometryAddParams& InParams, TArray<TUniquePtr<Chaos::FImplicitObject>>& OutGeoms, Chaos::FShapesArray& OutShapes, TArray<FPhysicsShapeHandle>* OutOptShapes);
 	static void AddGeometry(FPhysicsActorHandle& InActor, const FGeometryAddParams& InParams, TArray<FPhysicsShapeHandle>* OutOptShapes = nullptr);
 	static FPhysicsShapeHandle CloneShape(const FPhysicsShapeHandle& InShape);
 	static FPhysicsGeometryCollection_Chaos GetGeometryCollection(const FPhysicsShapeHandle& InShape);
@@ -431,7 +429,7 @@ FORCEINLINE ECollisionShapeType GetType(const Chaos::FImplicitObject& Geom)
 	return ECollisionShapeType::None;
 }
 */
-ECollisionShapeType GetGeometryType(const Chaos::TPerShapeData<float, 3>& Shape);
+ECollisionShapeType GetGeometryType(const Chaos::FPerShapeData& Shape);
 
 /*
 FORCEINLINE float GetRadius(const Chaos::TCapsule<float>& Capsule)
diff --git a/Engine/Source/Runtime/Engine/Public/Physics/ImmediatePhysics/ImmediatePhysicsChaos/ImmediatePhysicsActorHandle_Chaos.h b/Engine/Source/Runtime/Engine/Public/Physics/ImmediatePhysics/ImmediatePhysicsChaos/ImmediatePhysicsActorHandle_Chaos.h
--- a/Engine/Source/Runtime/Engine/Public/Physics/ImmediatePhysics/ImmediatePhysicsChaos/ImmediatePhysicsActorHandle_Chaos.h
+++ b/Engine/Source/Runtime/Engine/Public/Physics/ImmediatePhysics/ImmediatePhysicsChaos/ImmediatePhysicsActorHandle_Chaos.h
@@ -134,7 +134,7 @@ namespace ImmediatePhysics_Chaos
 		Chaos::TPBDRigidsSOAs<FReal, 3>& Particles;
 		Chaos::TGeometryParticleHandle<FReal, Dimensions>* ParticleHandle;
 		TUniquePtr<Chaos::FImplicitObject> Geometry;
-		TArray<TUniquePtr<Chaos::TPerShapeData<float, 3>>> Shapes;
+		TArray<TUniquePtr<Chaos::FPerShapeData>> Shapes;
 		int32 Level;
 	};
 
diff --git a/Engine/Source/Runtime/Engine/Public/Physics/ImmediatePhysics/ImmediatePhysicsChaos/ImmediatePhysicsCore_Chaos.h b/Engine/Source/Runtime/Engine/Public/Physics/ImmediatePhysics/ImmediatePhysicsChaos/ImmediatePhysicsCore_Chaos.h
--- a/Engine/Source/Runtime/Engine/Public/Physics/ImmediatePhysics/ImmediatePhysicsChaos/ImmediatePhysicsCore_Chaos.h
+++ b/Engine/Source/Runtime/Engine/Public/Physics/ImmediatePhysics/ImmediatePhysicsChaos/ImmediatePhysicsCore_Chaos.h
@@ -19,7 +19,7 @@ namespace Chaos
 	template<typename P> class TPBDConstraintIslandRule;
 	template<typename T, int D> struct TPBDRigidParticleParameters;
 	template<typename T, int D> class TPBDRigidsSOAs;
-	template<typename T, int D> class TPerShapeData;
+	class FPerShapeData;
 }
 
 namespace ImmediatePhysics_Chaos
diff --git a/Engine/Source/Runtime/Experimental/Chaos/Private/Chaos/ChaosDebugDraw.cpp b/Engine/Source/Runtime/Experimental/Chaos/Private/Chaos/ChaosDebugDraw.cpp
--- a/Engine/Source/Runtime/Experimental/Chaos/Private/Chaos/ChaosDebugDraw.cpp
+++ b/Engine/Source/Runtime/Experimental/Chaos/Private/Chaos/ChaosDebugDraw.cpp
@@ -127,8 +127,8 @@ namespace Chaos
 
 		void DrawParticleShapesImpl(const FRigidTransform3& SpaceTransform, const TGeometryParticle<FReal, 3>* Particle, const FColor& Color)
 		{
-			FVec3 P = SpaceTransform.TransformPosition(Particle->ObjectState() == EObjectStateType::Dynamic ? Particle->CastToRigidParticle()->P() : Particle->X());
-			FRotation3 Q = SpaceTransform.GetRotation() * (Particle->ObjectState() == EObjectStateType::Dynamic ? Particle->CastToRigidParticle()->Q() : Particle->R());
+			FVec3 P = SpaceTransform.TransformPosition(Particle->X());
+			FRotation3 Q = SpaceTransform.GetRotation() * (Particle->R());
 
 			DrawShapesImpl(FRigidTransform3(P, Q), Particle->Geometry().Get(), Color);
 		}
diff --git a/Engine/Source/Runtime/Experimental/Chaos/Private/Chaos/CollisionResolution.cpp b/Engine/Source/Runtime/Experimental/Chaos/Private/Chaos/CollisionResolution.cpp
--- a/Engine/Source/Runtime/Experimental/Chaos/Private/Chaos/CollisionResolution.cpp
+++ b/Engine/Source/Runtime/Experimental/Chaos/Private/Chaos/CollisionResolution.cpp
@@ -2297,25 +2297,24 @@ namespace Chaos
 			return (uint32)ChannelMask;
 		}
 
-		template<class T, int d>
-		inline bool DoCollide(EImplicitObjectType Implicit0Type, const TPerShapeData<T, d>* Shape0, EImplicitObjectType Implicit1Type, const TPerShapeData<T, d>* Shape1)
+		inline bool DoCollide(EImplicitObjectType Implicit0Type, const FPerShapeData* Shape0, EImplicitObjectType Implicit1Type, const FPerShapeData* Shape1)
 		{
 			//
 			// Disabled shapes do not collide
 			//
-			if (Shape0 && (Shape0->bDisable || !IsValid(Shape0->SimData) ) ) return false;
-			if (Shape1 && (Shape1->bDisable || !IsValid(Shape1->SimData) ) ) return false;
+			if (Shape0 && (Shape0->GetDisable() || !IsValid(Shape0->GetSimData()) ) ) return false;
+			if (Shape1 && (Shape1->GetDisable ()|| !IsValid(Shape1->GetSimData()) ) ) return false;
 
 			//
 			// Triangle Mesh geometry is only used if the shape specifies UseComplexAsSimple
 			//
 			if (Shape0)
 			{
-				if (Implicit0Type == ImplicitObjectType::TriangleMesh && Shape0->CollisionTraceType != Chaos_CTF_UseComplexAsSimple)
+				if (Implicit0Type == ImplicitObjectType::TriangleMesh && Shape0->GetCollisionTraceType() != Chaos_CTF_UseComplexAsSimple)
 				{
 					return false;
 				}
-				else if (Shape0->CollisionTraceType == Chaos_CTF_UseComplexAsSimple && Implicit0Type != ImplicitObjectType::TriangleMesh)
+				else if (Shape0->GetCollisionTraceType() == Chaos_CTF_UseComplexAsSimple && Implicit0Type != ImplicitObjectType::TriangleMesh)
 				{
 					return false;
 				}
@@ -2327,11 +2326,11 @@ namespace Chaos
 
 			if (Shape1)
 			{
-				if (Implicit1Type == ImplicitObjectType::TriangleMesh && Shape1->CollisionTraceType != Chaos_CTF_UseComplexAsSimple)
+				if (Implicit1Type == ImplicitObjectType::TriangleMesh && Shape1->GetCollisionTraceType() != Chaos_CTF_UseComplexAsSimple)
 				{
 					return false;
 				}
-				else if (Shape1->CollisionTraceType == Chaos_CTF_UseComplexAsSimple && Implicit1Type != ImplicitObjectType::TriangleMesh)
+				else if (Shape1->GetCollisionTraceType() == Chaos_CTF_UseComplexAsSimple && Implicit1Type != ImplicitObjectType::TriangleMesh)
 				{
 					return false;
 				}
@@ -2347,11 +2346,11 @@ namespace Chaos
 			if (Shape0 && Shape1)
 			{
 
-				if (IsValid(Shape0->SimData) && IsValid(Shape1->SimData))
+				if (IsValid(Shape0->GetSimData()) && IsValid(Shape1->GetSimData()))
 				{
 					FMaskFilter Filter0Mask, Filter1Mask;
-					const uint32 Filter0Channel = GetChaosCollisionChannelAndExtraFilter(Shape0->SimData.Word3, Filter0Mask);
-					const uint32 Filter1Channel = GetChaosCollisionChannelAndExtraFilter(Shape1->SimData.Word3, Filter1Mask);
+					const uint32 Filter0Channel = GetChaosCollisionChannelAndExtraFilter(Shape0->GetSimData().Word3, Filter0Mask);
+					const uint32 Filter1Channel = GetChaosCollisionChannelAndExtraFilter(Shape1->GetSimData().Word3, Filter1Mask);
 
 					if ((Filter0Mask & Filter1Mask) != 0)
 					{
@@ -2360,7 +2359,7 @@ namespace Chaos
 
 					const uint32 Filter1Bit = 1 << (Filter1Channel); // SIMDATA_TO_BITFIELD
 					uint32 const Filter0Bit = 1 << (Filter0Channel); // SIMDATA_TO_BITFIELD
-					return (Filter0Bit & Shape1->SimData.Word1) && (Filter1Bit & Shape0->SimData.Word1);
+					return (Filter0Bit & Shape1->GetSimData().Word1) && (Filter1Bit & Shape0->GetSimData().Word1);
 				}
 			}
 
diff --git a/Engine/Source/Runtime/Experimental/Chaos/Private/Chaos/GeometryParticles.cpp b/Engine/Source/Runtime/Experimental/Chaos/Private/Chaos/GeometryParticles.cpp
--- a/Engine/Source/Runtime/Experimental/Chaos/Private/Chaos/GeometryParticles.cpp
+++ b/Engine/Source/Runtime/Experimental/Chaos/Private/Chaos/GeometryParticles.cpp
@@ -6,11 +6,11 @@
 #include "Chaos/ImplicitObject.h"
 #include "Chaos/ImplicitObjectUnion.h"
 #include "Chaos/ParticleHandle.h"
+#include "Chaos/Framework/PhysicsSolverBase.h"
 
 namespace Chaos
 {
-	template <typename T, int d>
-	void UpdateShapesArrayFromGeometry(TShapesArray<T, d>& ShapesArray, TSerializablePtr<FImplicitObject> Geometry, const FRigidTransform3& ActorTM)
+	void UpdateShapesArrayFromGeometry(FShapesArray& ShapesArray, TSerializablePtr<FImplicitObject> Geometry, const FRigidTransform3& ActorTM, IPhysicsProxyBase* Proxy)
 	{
 		if(Geometry)
 		{
@@ -25,17 +25,25 @@ namespace Chaos
 					if (ShapeIndex >= OldShapeNum)
 					{
 						// If newly allocated shape, initialize it.
-						ShapesArray[ShapeIndex] = TPerShapeData<T, d>::CreatePerShapeData();
+						ShapesArray[ShapeIndex] = FPerShapeData::CreatePerShapeData(ShapeIndex);
 					}
 
-					ShapesArray[ShapeIndex]->Geometry = MakeSerializable(Union->GetObjects()[ShapeIndex]);
+					ShapesArray[ShapeIndex]->SetGeometry(MakeSerializable(Union->GetObjects()[ShapeIndex]));
 				}
 			}
 			else
 			{
 				ShapesArray.SetNum(1);
-				ShapesArray[0] = TPerShapeData<T, d>::CreatePerShapeData();
-				ShapesArray[0]->Geometry = Geometry;
+				ShapesArray[0] = FPerShapeData::CreatePerShapeData(0);
+				ShapesArray[0]->SetGeometry(Geometry);
+
+				if(Proxy)
+				{
+					if(FPhysicsSolverBase* PhysicsSolverBase = Proxy->GetSolver<FPhysicsSolverBase>())
+					{
+						PhysicsSolverBase->SetNumDirtyShapes(Proxy, 1);
+					}
+				}
 			}
 
 			if (Geometry->HasBoundingBox())
@@ -49,97 +57,64 @@ namespace Chaos
 		else
 		{
 			ShapesArray.Reset();
+			if(Proxy)
+			{
+				if(FPhysicsSolverBase* PhysicsSolverBase = Proxy->GetSolver<FPhysicsSolverBase>())
+				{
+					PhysicsSolverBase->SetNumDirtyShapes(Proxy,0);
+				}
+			}
 		}
 	}
 
-	template <typename T, int d>
-	TPerShapeData<T, d>::TPerShapeData()
-		: QueryData()
-		, SimData()
-		, UserData(nullptr)
+	FPerShapeData::FPerShapeData(int32 InShapeIdx)
+		: Proxy(nullptr)
+		, ShapeIdx(InShapeIdx)
 		, Geometry()
 		, WorldSpaceInflatedShapeBounds(TAABB<FReal, 3>(FVec3(0), FVec3(0)))
-		, Materials()
-		, MaterialMasks()
-		, MaterialMaskMaps()
-		, MaterialMaskMapMaterials()
-		, bDisable(false)
-		, bSimulate(true)
-		, CollisionTraceType(EChaosCollisionTraceFlag::Chaos_CTF_UseDefault)
 	{
 	}
 
-	template <typename T, int d>
-	TPerShapeData<T, d>::~TPerShapeData()
+	FPerShapeData::~FPerShapeData()
 	{
 	}
 
-	template <typename T, int d>
-	TUniquePtr<TPerShapeData<T, d>> TPerShapeData<T, d>::CreatePerShapeData()
+	TUniquePtr<FPerShapeData> FPerShapeData::CreatePerShapeData(int32 ShapeIdx)
 	{
-		return TUniquePtr<TPerShapeData<T, d>>(new TPerShapeData<T, d>());
+		return TUniquePtr<FPerShapeData>(new FPerShapeData(ShapeIdx));
 	}
 
-	template<typename T, int d>
-	void TPerShapeData<T, d>::UpdateShapeBounds(const TRigidTransform<T, d>& WorldTM)
+	void FPerShapeData::UpdateShapeBounds(const FRigidTransform3& WorldTM)
 	{
 		if (Geometry && Geometry->HasBoundingBox())
 		{
-			WorldSpaceInflatedShapeBounds = Geometry->BoundingBox().TransformedAABB(WorldTM);
+			SetWorldSpaceInflatedShapeBounds(Geometry->BoundingBox().TransformedAABB(WorldTM));
 		}
 	}
 
-	template <typename T, int d>
-	TPerShapeData<T, d>* TPerShapeData<T, d>::SerializationFactory(FChaosArchive& Ar, TPerShapeData<T, d>*)
+	FPerShapeData* FPerShapeData::SerializationFactory(FChaosArchive& Ar, FPerShapeData*)
 	{
-		return Ar.IsLoading() ? new TPerShapeData<T, d>() : nullptr;
+		//todo: need to rework serialization for shapes, for now just give them all shape idx 0
+		return Ar.IsLoading() ? new FPerShapeData(0) : nullptr;
 	}
 
-	template <typename T, int d>
-	void TPerShapeData<T, d>::Serialize(FChaosArchive& Ar)
+	void FPerShapeData::Serialize(FChaosArchive& Ar)
 	{
 		Ar.UsingCustomVersion(FExternalPhysicsCustomObjectVersion::GUID);
 		Ar.UsingCustomVersion(FExternalPhysicsMaterialCustomObjectVersion::GUID);
 
 		Ar << Geometry;
-		Ar << QueryData;
-		Ar << SimData;
+		Ar << CollisionData;
+		Ar << Materials;
 
 		if (Ar.CustomVer(FExternalPhysicsCustomObjectVersion::GUID) >= FExternalPhysicsCustomObjectVersion::SerializeShapeWorldSpaceBounds)
 		{
-			TBox<T, d>::SerializeAsAABB(Ar, WorldSpaceInflatedShapeBounds);
+			TBox<FReal,3>::SerializeAsAABB(Ar,WorldSpaceInflatedShapeBounds);
 		}
 		else
 		{
-			// This should be set by particle serializing this TPerShapeData.
-			WorldSpaceInflatedShapeBounds = TAABB<FReal, 3>(FVec3(0.0f, 0.0f, 0.0f), FVec3(0.0f, 0.0f, 0.0f));
-		}
-
-		if(Ar.CustomVer(FExternalPhysicsCustomObjectVersion::GUID) >= FExternalPhysicsCustomObjectVersion::AddedMaterialManager)
-		{
-			Ar << Materials;
-		}
-
-		if(Ar.CustomVer(FExternalPhysicsCustomObjectVersion::GUID) >= FExternalPhysicsCustomObjectVersion::AddShapeCollisionDisable)
-		{
-			Ar << bDisable;
-		}
-
-		if (Ar.CustomVer(FExternalPhysicsCustomObjectVersion::GUID) >= FExternalPhysicsCustomObjectVersion::SerializePerShapeDataSimulateFlag)
-		{
-			Ar << bSimulate;
-		}
-
-		if (Ar.CustomVer(FExternalPhysicsCustomObjectVersion::GUID) >= FExternalPhysicsCustomObjectVersion::SerializeCollisionTraceType)
-		{
-			int32 Data = (int32)CollisionTraceType;
-			Ar << Data;
-			CollisionTraceType = (EChaosCollisionTraceFlag)Data;
-		}
-
-		if (Ar.CustomVer(FExternalPhysicsMaterialCustomObjectVersion::GUID) >= FExternalPhysicsMaterialCustomObjectVersion::AddedMaterialMasks)
-		{
-			Ar << MaterialMasks << MaterialMaskMaps << MaterialMaskMapMaterials;
+			// This should be set by particle serializing this FPerShapeData.
+			SetWorldSpaceInflatedShapeBounds(TAABB<FReal, 3>(FVec3(0.0f, 0.0f, 0.0f), FVec3(0.0f, 0.0f, 0.0f)));
 		}
 
 	}
@@ -205,12 +180,12 @@ namespace Chaos
 		checkSlow(Index >= 0 && Index < ImplicitShapeMap.Num());
 
 		TMap<const FImplicitObject*, int32>& Mapping = ImplicitShapeMap[Index];
-		TShapesArray<T, d>& ShapeArray = MShapesArray[Index];
+		FShapesArray& ShapeArray = MShapesArray[Index];
 		Mapping.Reset();
 
 		for (int32 ShapeIndex = 0; ShapeIndex < ShapeArray.Num(); ++ ShapeIndex)
 		{
-			const FImplicitObject* ImplicitObject = ShapeArray[ShapeIndex]->Geometry.Get();
+			const FImplicitObject* ImplicitObject = ShapeArray[ShapeIndex]->GetGeometry().Get();
 			Mapping.Add(ImplicitObject, ShapeIndex);
 
 			const FImplicitObject* ImplicitChildObject = Utilities::ImplicitChildHelper(ImplicitObject);
@@ -269,6 +244,4 @@ namespace Chaos
 	
 	template class TGeometryParticlesImp<float, 3, EGeometryParticlesSimType::RigidBodySim>;
 	template class TGeometryParticlesImp<float, 3, EGeometryParticlesSimType::Other>;
-	template class TPerShapeData<float, 3>;
-	template void UpdateShapesArrayFromGeometry(TShapesArray<float, 3>& ShapesArray, TSerializablePtr<FImplicitObject> Geometry, const FRigidTransform3& ActorTM);
 }
diff --git a/Engine/Source/Runtime/Experimental/Chaos/Private/Chaos/PBDConstraintGraph.cpp b/Engine/Source/Runtime/Experimental/Chaos/Private/Chaos/PBDConstraintGraph.cpp
--- a/Engine/Source/Runtime/Experimental/Chaos/Private/Chaos/PBDConstraintGraph.cpp
+++ b/Engine/Source/Runtime/Experimental/Chaos/Private/Chaos/PBDConstraintGraph.cpp
@@ -668,11 +668,11 @@ bool FPBDConstraintGraph::SleepInactive(const int32 Island, const TArrayCollecti
 				}
 				else if (PBDRigid->ShapesArray().Num())
 				{
-					if (TPerShapeData<FReal, 3>* PerShapeData = PBDRigid->ShapesArray()[0].Get())
+					if (FPerShapeData* PerShapeData = PBDRigid->ShapesArray()[0].Get())
 					{
-						if (PerShapeData->Materials.Num())
+						if (PerShapeData->GetMaterials().Num())
 						{
-							if (FChaosPhysicsMaterial* Material = SolverPhysicsMaterials.Get(PBDRigid->ShapesArray()[0].Get()->Materials[0].InnerHandle))
+							if (FChaosPhysicsMaterial* Material = SolverPhysicsMaterials.Get(PBDRigid->ShapesArray()[0].Get()->GetMaterials()[0].InnerHandle))
 							{
 								LinearSleepingThreshold = FMath::Min(LinearSleepingThreshold, Material->SleepingLinearThreshold);
 								AngularSleepingThreshold = FMath::Min(AngularSleepingThreshold, Material->SleepingAngularThreshold);
diff --git a/Engine/Source/Runtime/Experimental/Chaos/Private/Chaos/ParticleHandle.cpp b/Engine/Source/Runtime/Experimental/Chaos/Private/Chaos/ParticleHandle.cpp
--- a/Engine/Source/Runtime/Experimental/Chaos/Private/Chaos/ParticleHandle.cpp
+++ b/Engine/Source/Runtime/Experimental/Chaos/Private/Chaos/ParticleHandle.cpp
@@ -16,7 +16,7 @@ namespace Chaos
 
 		for (int32 ShapeIndex = 0; ShapeIndex < MShapesArray.Num(); ++ ShapeIndex)
 		{
-			const FImplicitObject* ImplicitObject = MShapesArray[ShapeIndex]->Geometry.Get();
+			const FImplicitObject* ImplicitObject = MShapesArray[ShapeIndex]->GetGeometry().Get();
 			ImplicitShapeMap.Add(ImplicitObject, ShapeIndex);
 
 			const FImplicitObject* ImplicitChildObject = Utilities::ImplicitChildHelper(ImplicitObject);
@@ -26,7 +26,7 @@ namespace Chaos
 			}
 		}
 
-		auto& Geometry = MGeometry.Read();
+		auto& Geometry = MNonFrequentData.Read().Geometry;
 		if (Geometry)
 		{
 			int32 CurrentShapeIndex = INDEX_NONE;
diff --git a/Engine/Source/Runtime/Experimental/Chaos/Private/Framework/PhysicsProxyBase.cpp b/Engine/Source/Runtime/Experimental/Chaos/Private/Framework/PhysicsProxyBase.cpp
new file mode 100644
--- /dev/null
+++ b/Engine/Source/Runtime/Experimental/Chaos/Private/Framework/PhysicsProxyBase.cpp
@@ -0,0 +1,12 @@
+// Copyright Epic Games, Inc. All Rights Reserved.
+
+#include "Chaos/Framework/PhysicsProxyBase.h"
+#include "Chaos/Framework/PhysicsSolverBase.h"
+
+IPhysicsProxyBase::~IPhysicsProxyBase()
+{
+	if(GetSolver<Chaos::FPhysicsSolverBase>())
+	{
+		GetSolver<Chaos::FPhysicsSolverBase>()->RemoveDirtyProxy(this);
+	}
+}
\ No newline at end of file
diff --git a/Engine/Source/Runtime/Experimental/Chaos/Public/Chaos/Declares.h b/Engine/Source/Runtime/Experimental/Chaos/Public/Chaos/Declares.h
--- a/Engine/Source/Runtime/Experimental/Chaos/Public/Chaos/Declares.h
+++ b/Engine/Source/Runtime/Experimental/Chaos/Public/Chaos/Declares.h
@@ -11,7 +11,7 @@ namespace Chaos
 	typedef FPBDRigidsSolver FPhysicsSolver;
 
 	class FImplicitObject;
-	template <typename T, int d> class TPerShapeData;
+	class FPerShapeData;
 	template <typename TSOA> class TParticleView;
 	template<class T, int d> class TPBDRigidParticles;
 }
diff --git a/Engine/Source/Runtime/Experimental/Chaos/Public/Chaos/Framework/MultiBufferResource.h b/Engine/Source/Runtime/Experimental/Chaos/Public/Chaos/Framework/MultiBufferResource.h
--- a/Engine/Source/Runtime/Experimental/Chaos/Public/Chaos/Framework/MultiBufferResource.h
+++ b/Engine/Source/Runtime/Experimental/Chaos/Public/Chaos/Framework/MultiBufferResource.h
@@ -37,7 +37,7 @@ namespace Chaos
 	 * Single Buffer Implementation
 	 */
 	template<typename ResourceType>
-	class FSingleBuffer : public IBufferResource<ResourceType>
+	class FSingleBuffer final : public IBufferResource<ResourceType>
 	{
 	public:
 
@@ -61,7 +61,7 @@ namespace Chaos
 	* Double Buffer Implementation - Not thread-safe requires external locks
 	*/
 	template<typename ResourceType>
-	class FDoubleBuffer : public IBufferResource<ResourceType>
+	class FDoubleBuffer final : public IBufferResource<ResourceType>
 	{
 	public:
 
@@ -103,7 +103,7 @@ namespace Chaos
 	* Triple Buffer Implementation - Not thread-safe requires external locks
 	*/
 	template<typename ResourceType>
-	class FTripleBuffer : public IBufferResource<ResourceType>
+	class FTripleBuffer final : public IBufferResource<ResourceType>
 	{
 	public:
 
@@ -149,7 +149,7 @@ namespace Chaos
 	 * Not thread-safe, requires external locks.
 	 */
 	template<typename ResourceType>
-	class FGuardedTripleBuffer : public IBufferResource<ResourceType>
+	class FGuardedTripleBuffer final : public IBufferResource<ResourceType>
 	{
 	public:
 		/** This class implements a circular buffer access pattern, such that during
diff --git a/Engine/Source/Runtime/Experimental/Chaos/Public/Chaos/Framework/PhysicsProxy.h b/Engine/Source/Runtime/Experimental/Chaos/Public/Chaos/Framework/PhysicsProxy.h
--- a/Engine/Source/Runtime/Experimental/Chaos/Public/Chaos/Framework/PhysicsProxy.h
+++ b/Engine/Source/Runtime/Experimental/Chaos/Public/Chaos/Framework/PhysicsProxy.h
@@ -38,7 +38,6 @@ class TPhysicsProxy : public IPhysicsProxyBase
 
 public:
 	using FParticleType = Concrete;
-	using FParticleData = ConcreteData;
 
 	using FParticlesType = Chaos::TPBDRigidParticles<float, 3>;
 	using FCollisionConstraintsType = Chaos::TPBDCollisionConstraints<float, 3>;
@@ -88,7 +87,7 @@ public:
 	 * CONTEXT: GAMETHREAD
 	* Returns a new unmanaged allocation of the data saved on the handle, otherwise nullptr
 	*/
-	Chaos::FParticleData* NewData() { return static_cast<Concrete*>(this)->NewData(); }
+	//ConcreteData* NewData() { return static_cast<Concrete*>(this)->NewData(); }
 
 	/**
 	* CONTEXT: GAMETHREAD -> to -> PHYSICSTHREAD
@@ -96,7 +95,7 @@ public:
 	* callback should Enqueue commands on the PhysicsThread to update the state of
 	* the solver
 	*/
-	void PushToPhysicsState(const Chaos::FParticleData* InData) { static_cast<Concrete*>(this)->PushToPhysicsState(InData); }
+	//void PushToPhysicsState(const ConcreteData* InData) { static_cast<Concrete*>(this)->PushToPhysicsState(InData); }
 
 	/**
 	* CONTEXT: GAMETHREAD
diff --git a/Engine/Source/Runtime/Experimental/Chaos/Public/Chaos/Framework/PhysicsProxyBase.h b/Engine/Source/Runtime/Experimental/Chaos/Public/Chaos/Framework/PhysicsProxyBase.h
--- a/Engine/Source/Runtime/Experimental/Chaos/Public/Chaos/Framework/PhysicsProxyBase.h
+++ b/Engine/Source/Runtime/Experimental/Chaos/Public/Chaos/Framework/PhysicsProxyBase.h
@@ -3,7 +3,6 @@
 #pragma once
 
 #include "Chaos/Declares.h"
-#include "Chaos/Framework/PhysicsSolverBase.h"
 #include "UObject/GCObject.h"
 
 enum class EPhysicsProxyType
@@ -18,11 +17,17 @@ enum class EPhysicsProxyType
 	SingleRigidParticleType = 7
 };
 
-class IPhysicsProxyBase
+namespace Chaos
+{
+	class FPhysicsSolverBase;
+}
+
+class CHAOS_API IPhysicsProxyBase
 {
 public:
 	IPhysicsProxyBase(EPhysicsProxyType InType)
 		: Solver(nullptr)
+		, DirtyIdx(INDEX_NONE)
 		, Type(InType)
 	{}
 
@@ -40,21 +45,21 @@ public:
 	//todo: remove this
 	virtual void* GetHandleUnsafe() const { check(false); return nullptr; }
 
+	int32 GetDirtyIdx() const { return DirtyIdx; }
+	void SetDirtyIdx(const int32 Idx) { DirtyIdx = Idx; }
+	void ResetDirtyIdx() { DirtyIdx = INDEX_NONE; }
+
 protected:
 	// Ensures that derived classes can successfully call this destructor
 	// but no one can delete using a IPhysicsProxyBase*
-	virtual ~IPhysicsProxyBase()
-	{
-		if (GetSolver<Chaos::FPhysicsSolverBase>())
-		{
-			GetSolver<Chaos::FPhysicsSolverBase>()->RemoveDirtyProxy(this);
-		}
-	}
-
+	virtual ~IPhysicsProxyBase();
 	
 	/** The solver that owns the solver object */
 	Chaos::FPhysicsSolverBase* Solver;
 
+private:
+	int32 DirtyIdx;
+protected:
 	/** Proxy type */
 	EPhysicsProxyType Type;
 };
diff --git a/Engine/Source/Runtime/Experimental/Chaos/Public/Chaos/Framework/PhysicsSolverBase.h b/Engine/Source/Runtime/Experimental/Chaos/Public/Chaos/Framework/PhysicsSolverBase.h
--- a/Engine/Source/Runtime/Experimental/Chaos/Public/Chaos/Framework/PhysicsSolverBase.h
+++ b/Engine/Source/Runtime/Experimental/Chaos/Public/Chaos/Framework/PhysicsSolverBase.h
@@ -4,13 +4,43 @@
 #include "Chaos/Framework/MultiBufferResource.h"
 #include "Chaos/Matrix.h"
 #include "Misc/ScopeLock.h"
-
-class UObject;
-class IPhysicsProxyBase; // WTF - not in Chaos?
+#include "ChaosLog.h"
+#include "Chaos/Framework/PhysicsProxyBase.h"
+#include "Chaos/ParticleDirtyFlags.h"
+#include "Async/ParallelFor.h"
 
 namespace Chaos
 {
+	struct FDirtyProxy
+	{
+		IPhysicsProxyBase* Proxy;
+		FParticleDirtyData ParticleData;
+		TArray<int32> ShapeDataIndices;
 
+		FDirtyProxy(IPhysicsProxyBase* InProxy)
+			: Proxy(InProxy)
+		{
+		}
+
+		void SetDirtyIdx(int32 Idx)
+		{
+			Proxy->SetDirtyIdx(Idx);
+		}
+
+		void AddShape(int32 ShapeDataIdx)
+		{
+			ShapeDataIndices.Add(ShapeDataIdx);
+		}
+
+		void Clear(FDirtyPropertiesManager& Manager,int32 DataIdx,FShapeDirtyData* ShapesData)
+		{
+			ParticleData.Clear(Manager,DataIdx);
+			for(int32 ShapeDataIdx : ShapeDataIndices)
+			{
+				ShapesData[ShapeDataIdx].Clear(Manager,ShapeDataIdx);
+			}
+		}
+	};
 	class CHAOS_API FPhysicsSolverBase
 	{
 	public:
@@ -19,11 +49,21 @@ namespace Chaos
 
 		void AddDirtyProxy(IPhysicsProxyBase * ProxyBaseIn)
 		{
-			DirtyProxiesSet.Add(ProxyBaseIn);
+			DirtyProxiesDataBuffer.AccessProducerBuffer()->Add(ProxyBaseIn);
 		}
 		void RemoveDirtyProxy(IPhysicsProxyBase * ProxyBaseIn)
 		{
-			DirtyProxiesSet.Remove(ProxyBaseIn);
+			DirtyProxiesDataBuffer.AccessProducerBuffer()->Remove(ProxyBaseIn);
+		}
+
+		void AddDirtyProxyShape(IPhysicsProxyBase* ProxyBaseIn, int32 ShapeIdx)
+		{
+			DirtyProxiesDataBuffer.AccessProducerBuffer()->AddShape(ProxyBaseIn,ShapeIdx);
+		}
+
+		void SetNumDirtyShapes(IPhysicsProxyBase* Proxy, int32 NumShapes)
+		{
+			DirtyProxiesDataBuffer.AccessProducerBuffer()->SetNumDirtyShapes(Proxy,NumShapes);
 		}
 
 		const UObject* GetOwner() const
@@ -50,6 +90,107 @@ namespace Chaos
 #endif
 
 	protected:
+		/** Mode that the results buffers should be set to (single, double, triple) */
+		EMultiBufferMode BufferMode;
+
+		class FDirtySet
+		{
+		public:
+			void Add(IPhysicsProxyBase* Base)
+			{
+				if(Base->GetDirtyIdx() == INDEX_NONE)
+				{
+					const int32 Idx = ProxiesData.Num();
+					Base->SetDirtyIdx(Idx);
+					ProxiesData.Add(Base);
+				}
+			}
+
+			void Remove(IPhysicsProxyBase* Base)
+			{
+				const int32 Idx = Base->GetDirtyIdx();
+				if(Idx != INDEX_NONE)
+				{
+					if(Idx == ProxiesData.Num() - 1)
+					{
+						//last element so just pop
+						ProxiesData.Pop(/*bAllowShrinking=*/false);
+					}
+					else
+					{
+						//update other proxy's idx
+						ProxiesData.RemoveAtSwap(Idx);
+						ProxiesData[Idx].SetDirtyIdx(Idx);
+					}
+
+					Base->ResetDirtyIdx();
+				}
+			}
+
+			void Reset()
+			{
+				ProxiesData.Reset();
+				ShapesData.Reset();
+			}
+
+			int32 NumDirtyProxies() const { return ProxiesData.Num(); }
+			int32 NumDirtyShapes() const { return ShapesData.Num(); }
+
+			FShapeDirtyData* GetShapesDirtyData(){ return ShapesData.GetData(); }
+
+			template <typename Lambda>
+			void ParallelForEachProxy(const Lambda& Func)
+			{
+				::ParallelFor(ProxiesData.Num(),[this,&Func](int32 Idx)
+				{
+					Func(Idx,ProxiesData[Idx]);
+				});
+			}
+
+			template <typename Lambda>
+			void ForEachProxy(const Lambda& Func)
+			{
+				int32 Idx = 0;
+				for(FDirtyProxy& Dirty : ProxiesData)
+				{
+					Func(Idx++, Dirty);
+				}
+			}
+
+			void AddShape(IPhysicsProxyBase* Proxy, int32 ShapeIdx)
+			{
+				Add(Proxy);
+				FDirtyProxy& Dirty = ProxiesData[Proxy->GetDirtyIdx()];
+				for(int32 NewShapeIdx = Dirty.ShapeDataIndices.Num(); NewShapeIdx <= ShapeIdx; ++NewShapeIdx)
+				{
+					const int32 ShapeDataIdx = ShapesData.Add(FShapeDirtyData(NewShapeIdx));
+					Dirty.AddShape(ShapeDataIdx);
+				}
+			}
+
+			void SetNumDirtyShapes(IPhysicsProxyBase* Proxy,int32 NumShapes)
+			{
+				Add(Proxy);
+				FDirtyProxy& Dirty = ProxiesData[Proxy->GetDirtyIdx()];
+				
+				if(NumShapes < Dirty.ShapeDataIndices.Num())
+				{
+					Dirty.ShapeDataIndices.SetNum(NumShapes);
+				}
+				else
+				{
+					for(int32 NewShapeIdx = Dirty.ShapeDataIndices.Num(); NewShapeIdx < NumShapes; ++NewShapeIdx)
+					{
+						const int32 ShapeDataIdx = ShapesData.Add(FShapeDirtyData(NewShapeIdx));
+						Dirty.AddShape(ShapeDataIdx);
+					}
+				}
+			}
+
+		private:
+			TArray<FDirtyProxy> ProxiesData;
+			TArray<FShapeDirtyData> ShapesData;
+		};
 
 		/** Protected construction so callers still have to go through the module to create new instances */
 		FPhysicsSolverBase(const EMultiBufferMode BufferingModeIn, UObject* InOwner)
@@ -64,11 +205,10 @@ namespace Chaos
 		FPhysicsSolverBase& operator =(const FPhysicsSolverBase& InCopy) = delete;
 		FPhysicsSolverBase& operator =(FPhysicsSolverBase&& InSteal) = delete;
 
-		/** Mode that the results buffers should be set to (single, double, triple) */
-		EMultiBufferMode BufferMode;
 
-		// Input Proxy Map
-		TSet< IPhysicsProxyBase *> DirtyProxiesSet;
+		//NOTE: if you want to make this virtual, you need to make sure AddProxy is still inlinable since it gets called every time we do a write
+		//The easiest way is probably to have an FDirtySet* that we always write to, and then swap it into this generic buffer that is virtual
+		FDoubleBuffer<FDirtySet> DirtyProxiesDataBuffer;
 
 #if CHAOS_CHECKED
 		FName DebugName;
diff --git a/Engine/Source/Runtime/Experimental/Chaos/Public/Chaos/GeometryParticles.h b/Engine/Source/Runtime/Experimental/Chaos/Public/Chaos/GeometryParticles.h
--- a/Engine/Source/Runtime/Experimental/Chaos/Public/Chaos/GeometryParticles.h
+++ b/Engine/Source/Runtime/Experimental/Chaos/Public/Chaos/GeometryParticles.h
@@ -13,6 +13,7 @@
 #include "UObject/PhysicsObjectVersion.h"
 #include "UObject/ExternalPhysicsCustomObjectVersion.h"
 #include "UObject/ExternalPhysicsMaterialCustomObjectVersion.h"
+#include "Chaos/Properties.h"
 
 #ifndef CHAOS_DETERMINISTIC
 #define CHAOS_DETERMINISTIC 1
@@ -21,71 +22,210 @@
 namespace Chaos
 {
 
-	enum EChaosCollisionTraceFlag
-	{
-		/** Use project physics settings (DefaultShapeComplexity) */
-		Chaos_CTF_UseDefault,
-		/** Create both simple and complex shapes. Simple shapes are used for regular scene queries and collision tests. Complex shape (per poly) is used for complex scene queries.*/
-		Chaos_CTF_UseSimpleAndComplex,
-		/** Create only simple shapes. Use simple shapes for all scene queries and collision tests.*/
-		Chaos_CTF_UseSimpleAsComplex,
-		/** Create only complex shapes (per poly). Use complex shapes for all scene queries and collision tests. Can be used in simulation for static shapes only (i.e can be collided against but not moved through forces or velocity.) */
-		Chaos_CTF_UseComplexAsSimple,
-		/** */
-		Chaos_CTF_MAX,
-	};
-
-
 	/** Data that is associated with geometry. If a union is used an entry is created per internal geometry */
-	template <typename T, int d>
-	class CHAOS_API TPerShapeData
+	class CHAOS_API FPerShapeData
 	{
 	public:
 
 		static constexpr bool AlwaysSerializable = true;
-		static TUniquePtr<TPerShapeData<T, d>> CreatePerShapeData();
+		static TUniquePtr<FPerShapeData> CreatePerShapeData(int32 InShapeIdx);
 
-		~TPerShapeData();
-		TPerShapeData(const TPerShapeData<T, d>& Other) = delete;
-		FCollisionFilterData QueryData;
-		FCollisionFilterData SimData;
-		void* UserData;
-		TSerializablePtr<FImplicitObject> Geometry;
-		TAABB<FReal, 3> WorldSpaceInflatedShapeBounds;
-		TArray<FMaterialHandle> Materials;
-		TArray<FMaterialMaskHandle> MaterialMasks;
-		TArray<uint32> MaterialMaskMaps;
-		TArray<FMaterialHandle> MaterialMaskMapMaterials;
+		~FPerShapeData();
+		FPerShapeData(const FPerShapeData& Other) = delete;
+		
+		void UpdateShapeBounds(const FRigidTransform3& WorldTM);
 
-		// TODO: Bitfields?
-		bool bDisable;
-		bool bSimulate;
-		EChaosCollisionTraceFlag CollisionTraceType;
+		static FPerShapeData* SerializationFactory(FChaosArchive& Ar, FPerShapeData*);
+		void Serialize(FChaosArchive& Ar);
 
-		void UpdateShapeBounds(const TRigidTransform<T, d>& WorldTM);
+		void* GetUserData() const { return CollisionData.Read().UserData; }
+		void SetUserData(void* InUserData)
+		{
+			CollisionData.Modify(true, DirtyFlags, Proxy, ShapeIdx, [InUserData](FCollisionData& Data){ Data.UserData = InUserData; });
+		}
 
-		static TPerShapeData<T, d>* SerializationFactory(FChaosArchive& Ar, TPerShapeData<T, d>*);
-		void Serialize(FChaosArchive& Ar);
+		const FCollisionFilterData& GetQueryData() const { return CollisionData.Read().QueryData; }
+		void SetQueryData(const FCollisionFilterData& InQueryData)
+		{
+			CollisionData.Modify(true, DirtyFlags, Proxy, ShapeIdx, [InQueryData](FCollisionData& Data){ Data.QueryData = InQueryData; });
+		}
+
+		const FCollisionFilterData& GetSimData() const { return CollisionData.Read().SimData; }
+		void SetSimData(const FCollisionFilterData& InSimData)
+		{
+			CollisionData.Modify(true, DirtyFlags, Proxy, ShapeIdx, [InSimData](FCollisionData& Data){ Data.SimData = InSimData; });
+		}
+
+		template <typename Lambda>
+		void ModifySimData(const Lambda& LambdaFunc)
+		{
+			CollisionData.Modify(true,DirtyFlags,Proxy, ShapeIdx,[&LambdaFunc](FCollisionData& Data){ LambdaFunc(Data.SimData);});
+		}
+
+		TSerializablePtr<FImplicitObject> GetGeometry() const { return Geometry; }
+		void SetGeometry(TSerializablePtr<FImplicitObject> InGeometry)
+		{
+			Geometry = InGeometry;
+		}
+
+		const TAABB<FReal,3>& GetWorldSpaceInflatedShapeBounds() const { return WorldSpaceInflatedShapeBounds; }
+		void SetWorldSpaceInflatedShapeBounds(const TAABB<FReal,3>& InWorldSpaceInflatedShapeBounds)
+		{
+			WorldSpaceInflatedShapeBounds = InWorldSpaceInflatedShapeBounds;
+		}
+
+		const TArray<FMaterialHandle>& GetMaterials() const { return Materials.Read().Materials; }
+		const TArray<FMaterialMaskHandle>& GetMaterialMasks() const { return Materials.Read().MaterialMasks; }
+		const TArray<uint32>& GetMaterialMaskMaps() const { return Materials.Read().MaterialMaskMaps; }
+		const TArray<FMaterialHandle>& GetMaterialMaskMapMaterials() const { return Materials.Read().MaterialMaskMapMaterials; }
+
+		const FShapeDirtyFlags GetDirtyFlags() const { return DirtyFlags; }
+
+		void SetMaterials(const TArray<FMaterialHandle>& InMaterials)
+		{
+			Materials.Modify(true,DirtyFlags,Proxy, ShapeIdx,[&InMaterials](FMaterialData& Data)
+			{
+				Data.Materials = InMaterials;
+			});
+		}
+
+		void SetMaterialMasks(const TArray<FMaterialMaskHandle>& InMaterialMasks)
+		{
+			Materials.Modify(true,DirtyFlags,Proxy, ShapeIdx,[&InMaterialMasks](FMaterialData& Data)
+			{
+				Data.MaterialMasks = InMaterialMasks;
+			});
+		}
+
+		void SetMaterialMaskMaps(const TArray<uint32>& InMaterialMaskMaps)
+		{
+			Materials.Modify(true,DirtyFlags,Proxy, ShapeIdx,[&InMaterialMaskMaps](FMaterialData& Data)
+			{
+				Data.MaterialMaskMaps = InMaterialMaskMaps;
+			});
+		}
+
+		void SetMaterialMaskMapMaterials(const TArray<FMaterialHandle>& InMaterialMaskMapMaterials)
+		{
+			Materials.Modify(true,DirtyFlags,Proxy, ShapeIdx,[&InMaterialMaskMapMaterials](FMaterialData& Data)
+			{
+				Data.MaterialMaskMapMaterials = InMaterialMaskMapMaterials;
+			});
+		}
+
+		template <typename Lambda>
+		void ModifyMaterials(const Lambda& LambdaFunc)
+		{
+			Materials.Modify(true,DirtyFlags,Proxy, ShapeIdx,[&LambdaFunc](FMaterialData& Data)
+			{
+				LambdaFunc(Data.Materials);
+			});
+		}
+
+		template <typename Lambda>
+		void ModifyMaterialMasks(const Lambda& LambdaFunc)
+		{
+			Materials.Modify(true,DirtyFlags,Proxy, ShapeIdx,[&LambdaFunc](FMaterialData& Data)
+			{
+				LambdaFunc(Data.MaterialMasks);
+			});
+		}
+
+		template <typename Lambda>
+		void ModifyMaterialMaskMaps(const Lambda& LambdaFunc)
+		{
+			Materials.Modify(true,DirtyFlags,Proxy, ShapeIdx,[&LambdaFunc](FMaterialData& Data)
+			{
+				LambdaFunc(Data.MaterialMaskMaps);
+			});
+		}
+
+		template <typename Lambda>
+		void ModifyMaterialMaskMapMaterials(const Lambda& LambdaFunc)
+		{
+			Materials.Modify(true,DirtyFlags,Proxy, ShapeIdx,[&LambdaFunc](FMaterialData& Data)
+			{
+				LambdaFunc(Data.MaterialMaskMapMaterials);
+			});
+		}
+
+		bool GetDisable() const { return CollisionData.Read().bDisable; }
+		void SetDisable(const bool bDisable)
+		{
+			CollisionData.Modify(true, DirtyFlags, Proxy, ShapeIdx, [bDisable](FCollisionData& Data){ Data.bDisable = bDisable; });
+		}
+
+		bool GetSimulate() const { return CollisionData.Read().bSimulate; }
+		void SetSimulate(const bool bSimulate)
+		{
+			CollisionData.Modify(true,DirtyFlags,Proxy, ShapeIdx,[bSimulate](FCollisionData& Data){ Data.bSimulate = bSimulate; });
+		}
+
+		EChaosCollisionTraceFlag GetCollisionTraceType() const { return CollisionData.Read().CollisionTraceType; }
+		void SetCollisionTraceType(const EChaosCollisionTraceFlag InTraceFlag)
+		{
+			CollisionData.Modify(true,DirtyFlags,Proxy, ShapeIdx,[InTraceFlag](FCollisionData& Data){ Data.CollisionTraceType = InTraceFlag; });
+		}
+
+		void SetCollisionData(const FCollisionData& Data)
+		{
+			CollisionData.Write(Data,true,DirtyFlags,Proxy, ShapeIdx);
+		}
+
+		void SetMaterialData(const FMaterialData& Data)
+		{
+			Materials.Write(Data,true,DirtyFlags,Proxy,ShapeIdx);
+		}
+
+		void SyncRemoteData(FDirtyPropertiesManager& Manager, int32 ShapeDataIdx, FShapeDirtyData& RemoteData)
+		{
+			RemoteData.SetFlags(DirtyFlags);
+			CollisionData.SyncRemote(Manager, ShapeDataIdx, RemoteData);
+			Materials.SyncRemote(Manager, ShapeDataIdx, RemoteData);
+			DirtyFlags.Clear();
+		}
+
+		void SetProxy(class IPhysicsProxyBase* InProxy)
+		{
+			Proxy = InProxy;
+			if(Proxy)
+			{
+				if(DirtyFlags.IsDirty())
+				{
+					if(FPhysicsSolverBase* PhysicsSolverBase = Proxy->GetSolver<FPhysicsSolverBase>())
+					{
+						PhysicsSolverBase->AddDirtyProxyShape(Proxy,ShapeIdx);
+					}
+				}
+			}
+		}
 
 	private:
-		//use factory
-		TPerShapeData();
+
+		class IPhysicsProxyBase* Proxy;
+		FShapeDirtyFlags DirtyFlags;
+		int32 ShapeIdx;
+
+		TShapeProperty<FCollisionData,EShapeProperty::CollisionData> CollisionData;
+		TShapeProperty<FMaterialData,EShapeProperty::Materials> Materials;
+
+		TSerializablePtr<FImplicitObject> Geometry;
+		TAABB<FReal,3> WorldSpaceInflatedShapeBounds;
+		
+		// use CreatePerShapeData
+		FPerShapeData(int32 InShapeIdx);
 	};
 
-	template <typename T, int d>
-	FChaosArchive& operator<<(FChaosArchive& Ar, TPerShapeData<T, d>& Shape)
+	inline FChaosArchive& operator<<(FChaosArchive& Ar, FPerShapeData& Shape)
 	{
 		Shape.Serialize(Ar);
 		return Ar;
 	}
 
-	template <typename T, int d>
-	using TShapesArray = TArray<TUniquePtr<TPerShapeData<T, d>>, TInlineAllocator<1>>;
+	using FShapesArray = TArray<TUniquePtr<FPerShapeData>, TInlineAllocator<1>>;
 
-	template <typename T, int d>
-	void UpdateShapesArrayFromGeometry(TShapesArray<T, d>& ShapesArray, TSerializablePtr<FImplicitObject> Geometry, const FRigidTransform3& ActorTM);
+	void CHAOS_API UpdateShapesArrayFromGeometry(FShapesArray& ShapesArray, TSerializablePtr<FImplicitObject> Geometry, const FRigidTransform3& ActorTM, IPhysicsProxyBase* Proxy);
 
-	extern template void CHAOS_API UpdateShapesArrayFromGeometry(TShapesArray<float, 3>& ShapesArray, TSerializablePtr<FImplicitObject> Geometry,  const FRigidTransform3& ActorTM);
 
 #if CHAOS_DETERMINISTIC
 	using FParticleID = int32;	//Used to break ties when determinism is needed. Should not be used for anything else
@@ -239,7 +379,7 @@ namespace Chaos
 
 		CHAOS_API const TSharedPtr<FImplicitObject, ESPMode::ThreadSafe>& SharedGeometry(const int32 Index) const { return MSharedGeometry[Index]; }
 
-		CHAOS_API const TShapesArray<T, d>& ShapesArray(const int32 Index) const { return MShapesArray[Index]; }
+		CHAOS_API const FShapesArray& ShapesArray(const int32 Index) const { return MShapesArray[Index]; }
 
 #if CHAOS_DETERMINISTIC
 		CHAOS_API FParticleID ParticleID(const int32 Idx) const { return MParticleIDs[Idx]; }
@@ -325,10 +465,10 @@ namespace Chaos
 		{
 			MWorldSpaceInflatedBounds[Index] = Bounds;
 
-			const TShapesArray<FReal, 3>& Shapes = ShapesArray(Index);
+			const FShapesArray& Shapes = ShapesArray(Index);
 			for (const auto& Shape : Shapes)
 			{
-				if (Shape->Geometry->HasBoundingBox())
+				if (Shape->GetGeometry()->HasBoundingBox())
 				{
 					const TRigidTransform<FReal, 3> ActorTM(X(Index), R(Index));
 					Shape->UpdateShapeBounds(ActorTM);
@@ -420,11 +560,11 @@ namespace Chaos
 
 		CHAOS_API EParticleType ParticleType() const { return MParticleType; }
 
-		const TPerShapeData<T, d>* GetImplicitShape(int32 Index, const FImplicitObject* InObject)
+		const FPerShapeData* GetImplicitShape(int32 Index, const FImplicitObject* InObject)
 		{
 			checkSlow(Index >= 0 && Index < ImplicitShapeMap.Num());
 			TMap<const FImplicitObject*, int32>& Mapping = ImplicitShapeMap[Index];
-			TShapesArray<T, d>& ShapeArray = MShapesArray[Index];
+			FShapesArray& ShapeArray = MShapesArray[Index];
 			int32* ShapeIndex = Mapping.Find(InObject);
 
 			if(ShapeIndex && ShapeArray.IsValidIndex(*ShapeIndex))
@@ -452,7 +592,7 @@ namespace Chaos
 		TArrayCollectionArray<TUniquePtr<FImplicitObject>> MDynamicGeometry;
 		TArrayCollectionArray<TSerializablePtr<TGeometryParticleHandle<T, d>>> MGeometryParticleHandle;
 		TArrayCollectionArray<TGeometryParticle<T, d>*> MGeometryParticle;
-		TArrayCollectionArray<TShapesArray<T,d>> MShapesArray;
+		TArrayCollectionArray<FShapesArray> MShapesArray;
 		TArrayCollectionArray<TMap<const FImplicitObject*, int32>> ImplicitShapeMap;
 		TArrayCollectionArray<TAABB<T,d>> MLocalBounds;
 		TArrayCollectionArray<TAABB<T, d>> MWorldSpaceInflatedBounds;
@@ -462,7 +602,7 @@ namespace Chaos
 
 		void UpdateShapesArray(const int32 Index)
 		{
-			UpdateShapesArrayFromGeometry(MShapesArray[Index], MGeometry[Index], FRigidTransform3(X(Index), R(Index)));
+			UpdateShapesArrayFromGeometry(MShapesArray[Index], MGeometry[Index], FRigidTransform3(X(Index), R(Index)), nullptr);
 			MapImplicitShapes(Index);
 		}
 
diff --git a/Engine/Source/Runtime/Experimental/Chaos/Public/Chaos/ParticleDirtyFlags.h b/Engine/Source/Runtime/Experimental/Chaos/Public/Chaos/ParticleDirtyFlags.h
--- a/Engine/Source/Runtime/Experimental/Chaos/Public/Chaos/ParticleDirtyFlags.h
+++ b/Engine/Source/Runtime/Experimental/Chaos/Public/Chaos/ParticleDirtyFlags.h
@@ -3,18 +3,250 @@
 #pragma once
 
 #include "Chaos/Core.h"
-#include "Framework/PhysicsProxyBase.h"
 #include "Chaos/ChaosArchive.h"
+#include "Chaos/Box.h"
+#include "Chaos/Particles.h"
+#include "Chaos/PhysicalMaterials.h"
+#include "Chaos/GeometryParticlesFwd.h"
+#include "Chaos/CollisionFilterData.h"
+#include "UObject/ExternalPhysicsCustomObjectVersion.h"
+#include "UObject/ExternalPhysicsMaterialCustomObjectVersion.h"
 
 class FName;
 
 namespace Chaos
 {
 
-using ChaosGeometrySp = TSharedPtr<FImplicitObject,ESPMode::ThreadSafe>;
+struct FParticlePositionRotation
+{
+	FVec3 X;
+	FRotation3 R;
+
+	void Serialize(FChaosArchive& Ar)
+	{
+		Ar << X << R;
+	}
+};
+
+inline FChaosArchive& operator<<(FChaosArchive& Ar,FParticlePositionRotation& Data)
+{
+	Data.Serialize(Ar);
+	return Ar;
+}
+
+struct FParticleVelocities
+{
+	void Serialize(FChaosArchive& Ar)
+	{
+		Ar << V << W;
+	}
+
+	FVec3 V;
+	FVec3 W;
+};
+
+inline FChaosArchive& operator<<(FChaosArchive& Ar,FParticleVelocities& Data)
+{
+	Data.Serialize(Ar);
+	return Ar;
+}
+
+struct FParticleDynamics
+{
+	FVec3 F;
+	FVec3 Torque;
+	FVec3 LinearImpulse;
+	FVec3 AngularImpulse;
+	FReal LinearEtherDrag;
+	FReal AngularEtherDrag;
+
+	void Serialize(FChaosArchive& Ar)
+	{
+		Ar << F;
+		Ar << Torque;
+		Ar << LinearImpulse;
+		Ar << AngularImpulse;
+		Ar << LinearEtherDrag;
+		Ar << AngularEtherDrag;		
+	}
+};
+
+inline FChaosArchive& operator<<(FChaosArchive& Ar, FParticleDynamics& Data)
+{
+	Data.Serialize(Ar);
+	return Ar;
+}
+
+struct FParticleMassProps
+{
+	FVec3 CenterOfMass;
+	FRotation3 RotationOfMass;
+	FMatrix33 I;
+	FMatrix33 InvI;
+	FReal M;
+	FReal InvM;
+
+	void Serialize(FChaosArchive& Ar)
+	{
+		Ar << CenterOfMass;
+		Ar << RotationOfMass;
+		Ar << I;
+		Ar << InvI;
+		Ar << M;
+		Ar << InvM;
+	}
+};
+
+inline FChaosArchive& operator<<(FChaosArchive& Ar,FParticleMassProps& Data)
+{
+	Data.Serialize(Ar);
+	return Ar;
+}
+
+struct FParticleMisc
+{
+	int32 CollisionGroup;
+	EObjectStateType ObjectState;
+	FSpatialAccelerationIdx SpatialIdx;
+
+	uint8 bDisabled : 1;
+	uint8 bGravityEnabled : 1;
+
+	void Serialize(FChaosArchive& Ar)
+	{
+		Ar << CollisionGroup;
+		Ar << ObjectState;
+		bool Disabled = bDisabled;
+		Ar << Disabled;
+		bDisabled = Disabled;
+
+		bool GravityEnabled = bGravityEnabled;
+		Ar << GravityEnabled;
+		bGravityEnabled = GravityEnabled;
+	}
+};
+
+inline FChaosArchive& operator<<(FChaosArchive& Ar,FParticleMisc& Data)
+{
+	Data.Serialize(Ar);
+	return Ar;
+}
+
+struct FParticleNonFrequentData
+{
+	FParticleNonFrequentData()
+		: UserData(nullptr)
+	{
+
+	}
+
+	TSharedPtr<FImplicitObject,ESPMode::ThreadSafe> Geometry;
+	void* UserData;
+	FUniqueIdx UniqueIdx;
+#if CHAOS_CHECKED
+	FName DebugName;
+#endif
+
+	void Serialize(FChaosArchive& Ar)
+	{
+		Ar << Geometry;
+	}
+};
+
+inline FChaosArchive& operator<<(FChaosArchive& Ar,FParticleNonFrequentData& Data)
+{
+	Data.Serialize(Ar);
+	return Ar;
+}
+
+struct FCollisionData
+{
+	FCollisionFilterData QueryData;
+	FCollisionFilterData SimData;
+	void* UserData;
+	EChaosCollisionTraceFlag CollisionTraceType;
+	uint8 bDisable : 1;
+	uint8 bSimulate : 1;
+
+	FCollisionData()
+	: UserData(nullptr)
+	, CollisionTraceType(EChaosCollisionTraceFlag::Chaos_CTF_UseDefault)
+	, bDisable(false)
+	, bSimulate(true)
+	{
+	}
+
+	void Serialize(FChaosArchive& Ar)
+	{
+		Ar.UsingCustomVersion(FExternalPhysicsCustomObjectVersion::GUID);
+		Ar.UsingCustomVersion(FExternalPhysicsMaterialCustomObjectVersion::GUID);
+
+		Ar << QueryData;
+		Ar << SimData;
+
+		if(Ar.CustomVer(FExternalPhysicsCustomObjectVersion::GUID) >= FExternalPhysicsCustomObjectVersion::AddShapeCollisionDisable)
+		{
+			bool Disable = bDisable;
+			Ar << Disable;
+			bDisable = Disable;
+		}
+
+		if(Ar.CustomVer(FExternalPhysicsCustomObjectVersion::GUID) >= FExternalPhysicsCustomObjectVersion::SerializePerShapeDataSimulateFlag)
+		{
+			bool Simulate = bSimulate;
+			Ar << Simulate;
+			bSimulate= Simulate;
+		}
+
+		if(Ar.CustomVer(FExternalPhysicsCustomObjectVersion::GUID) >= FExternalPhysicsCustomObjectVersion::SerializeCollisionTraceType)
+		{
+			int32 Data = (int32)CollisionTraceType;
+			Ar << Data;
+			CollisionTraceType = (EChaosCollisionTraceFlag)Data;
+		}
+	}
+};
+
+inline FChaosArchive& operator<<(FChaosArchive& Ar,FCollisionData& Data)
+{
+	//TODO: should this only work with dirty flag? Not sure if this path really matters at this point
+	Data.Serialize(Ar);
+	return Ar;
+}
+
+struct FMaterialData
+{
+	TArray<FMaterialHandle> Materials;
+	TArray<FMaterialMaskHandle> MaterialMasks;
+	TArray<uint32> MaterialMaskMaps;
+	TArray<FMaterialHandle> MaterialMaskMapMaterials;
+
+	void Serialize(FChaosArchive& Ar)
+	{
+		Ar.UsingCustomVersion(FExternalPhysicsCustomObjectVersion::GUID);
+		Ar.UsingCustomVersion(FExternalPhysicsMaterialCustomObjectVersion::GUID);
+
+		if(Ar.CustomVer(FExternalPhysicsCustomObjectVersion::GUID) >= FExternalPhysicsCustomObjectVersion::AddedMaterialManager)
+		{
+			Ar << Materials;
+		}
+
+		if(Ar.CustomVer(FExternalPhysicsMaterialCustomObjectVersion::GUID) >= FExternalPhysicsMaterialCustomObjectVersion::AddedMaterialMasks)
+		{
+			Ar << MaterialMasks << MaterialMaskMaps << MaterialMaskMapMaterials;
+		}
+	}
+};
+
+inline FChaosArchive& operator<<(FChaosArchive& Ar,FMaterialData& Data)
+{
+	//TODO: should this only work with dirty flag? Not sure if this path really matters at this point
+	Data.Serialize(Ar);
+	return Ar;
+}
 
 #define PARTICLE_PROPERTY(PropName, Type) PropName,
-	enum EParticleProperty : uint32
+	enum class EParticleProperty : uint32
 	{
 #include "ParticleProperties.inl"
 		NumProperties
@@ -22,33 +254,30 @@ using ChaosGeometrySp = TSharedPtr<FImplicitObject,ESPMode::ThreadSafe>;
 
 #undef PARTICLE_PROPERTY
 
-#define PARTICLE_PROPERTY_TYPE(TypeName, Type) TypeName,
-	enum EParticlePropertyType: uint32
+#define PROPERTY_TYPE(TypeName, Type) TypeName,
+	enum class EPropertyType: uint32
 	{
-#include "ParticlePropertiesTypes.inl"
+#include "PropertiesTypes.inl"
 		NumTypes
 	};
 
-#undef PARTICLE_PROPERTY_TYPE
+#undef PROPERTY_TYPE
 
 template <typename T>
-struct TParticlePropertyTrait
+struct TPropertyTypeTrait
 {
 };
 
-#define PARTICLE_PROPERTY_TYPE(TypeName, Type) \
+#define PROPERTY_TYPE(TypeName, Type) \
 template <>\
-struct TParticlePropertyTrait<Type>\
+struct TPropertyTypeTrait<Type>\
 {\
-	static constexpr uint32 PoolIdx = (uint32)EParticlePropertyType::TypeName;\
+	static constexpr EPropertyType PoolIdx = EPropertyType::TypeName;\
 };
 
-#include "ParticlePropertiesTypes.inl"
-#undef PARTICLE_PROPERTY_TYPE
-	
-	// There is a dirty flag for every user-settable particle property.
-	// Dirty property values will get copied from game to physics thread buffers,
-	// but clean property values will get overridden with physics thread results.
+#include "PropertiesTypes.inl"
+#undef PROPERTY_TYPE
+
 #define PARTICLE_PROPERTY(PropName, Type) PropName = (uint32)1 << (uint32)EParticleProperty::PropName,
 
 	enum class EParticleFlags : uint32
@@ -62,25 +291,55 @@ struct TParticlePropertyTrait<Type>\
 	{
 		switch(Prop)
 		{
-			#define PARTICLE_PROPERTY(PropName, Type) case PropName: return EParticleFlags::PropName;
+			#define PARTICLE_PROPERTY(PropName, Type) case EParticleProperty::PropName: return EParticleFlags::PropName;
 #include "ParticleProperties.inl"
 #undef PARTICLE_PROPERTY
 		default: return (EParticleFlags)0;
 		}
 	}
 
-	class FParticleDirtyFlags
+
+#define SHAPE_PROPERTY(PropName, Type) PropName,
+	enum class EShapeProperty: uint32
+	{
+#include "ShapeProperties.inl"
+		NumShapeProperties
+	};
+
+#undef SHAPE_PROPERTY
+
+#define SHAPE_PROPERTY(PropName, Type) PropName = (uint32)1 << (uint32)EShapeProperty::PropName,
+
+	enum class EShapeFlags: uint32
+	{
+#include "ShapeProperties.inl"
+		DummyFlag
+	};
+#undef SHAPE_PROPERTY
+
+	constexpr EShapeFlags ShapePropToFlag(EShapeProperty Prop)
+	{
+		switch(Prop)
+		{
+#define SHAPE_PROPERTY(PropName, Type) case EShapeProperty::PropName: return EShapeFlags::PropName;
+#include "ShapeProperties.inl"
+#undef SHAPE_PROPERTY
+		default: return (EShapeFlags)0;
+		}
+	}
+
+	template <typename FlagsType>
+	class TDirtyFlags
 	{
 	public:
-		FParticleDirtyFlags() : Bits(0) { }
-		~FParticleDirtyFlags() { }
+		TDirtyFlags() : Bits(0) { }
 
 		bool IsDirty() const
 		{
 			return Bits != 0;
 		}
 
-		bool IsDirty(const EParticleFlags CheckBits) const
+		bool IsDirty(const FlagsType CheckBits) const
 		{
 			return (Bits & (int32)CheckBits) != 0;
 		}
@@ -90,12 +349,12 @@ struct TParticlePropertyTrait<Type>\
 			return (Bits & CheckBits) != 0;
 		}
 
-		void MarkDirty(const EParticleFlags DirtyBits)
+		void MarkDirty(const FlagsType DirtyBits)
 		{
 			Bits |= (int32)DirtyBits;
 		}
 
-		void MarkClean(const EParticleFlags CleanBits)
+		void MarkClean(const FlagsType CleanBits)
 		{
 			Bits &= ~(int32)CleanBits;
 		}
@@ -114,6 +373,9 @@ struct TParticlePropertyTrait<Type>\
 		int32 Bits;
 	};
 
+	using FParticleDirtyFlags = TDirtyFlags<EParticleFlags>;
+	using FShapeDirtyFlags = TDirtyFlags<EShapeFlags>;
+
 	struct FDirtyIdx
 	{
 		uint32 bHasEntry : 1;
@@ -123,7 +385,32 @@ struct TParticlePropertyTrait<Type>\
 	template <typename T>
 	class TDirtyElementPool
 	{
-		static_assert(sizeof(TParticlePropertyTrait<T>::PoolIdx),"Property type must be registed. Is it in ParticlePropertiesTypes?");
+		static_assert(sizeof(TPropertyTypeTrait<T>::PoolIdx),"Property type must be registered. Is it in PropertiesTypes.inl?");
+	public:
+		const T& GetElement(int32 Idx) const { return Elements[Idx]; }
+		T& GetElement(int32 Idx){ return Elements[Idx]; }
+		
+		void Reset(int32 Idx)
+		{
+			Elements[Idx] = T();
+		}
+
+		void SetNum(int32 Num)
+		{
+			Elements.SetNum(Num);
+		}
+
+	private:
+
+		TArray<T> Elements;
+	};
+
+	//want this for sparse representation
+#if 0
+	template <typename T>
+	class TDirtyElementPool
+	{
+		static_assert(sizeof(TPropertyTypeTrait<T>::PoolIdx),"Property type must be registered. Is it in PropertiesTypes.inl?");
 	public:
 		const T& Read(int32 Idx) const
 		{
@@ -164,7 +451,7 @@ struct TParticlePropertyTrait<Type>\
 			//todo: can we avoid default constructors? maybe if std::is_trivially_copyable
 			if(FreeIndices.Num())
 			{
-				int32 NewIdx = FreeIndices.Pop();
+				int32 NewIdx = FreeIndices.Pop(/*bAllowShrinking=*/false);
 				Elements[NewIdx] = T();
 				return NewIdx;
 			}
@@ -178,243 +465,553 @@ struct TParticlePropertyTrait<Type>\
 		TArray<int32> FreeIndices;
 	};
 
-	class FDirtyPropertiesManager
-	{
-	public:
-		template <typename T>
-		TDirtyElementPool<T>& GetPool()
-		{
-			switch(TParticlePropertyTrait<T>::PoolIdx)
-			{
-#define PARTICLE_PROPERTY_TYPE(TypeName, Type) case EParticlePropertyType::TypeName: return (TDirtyElementPool<T>&)TypeName##Pool;
-#include "ParticlePropertiesTypes.inl"
-#undef PARTICLE_PROPERTY_TYPE
-			default: check(false);
-			}
-			
-			static TDirtyElementPool<T> ErrorPool;
-			return ErrorPool;
-		}
-
-		template <typename T>
-		const TDirtyElementPool<T>& GetPool() const
-		{
-			switch(TParticlePropertyTrait<T>::PoolIdx)
-			{
-#define PARTICLE_PROPERTY_TYPE(TypeName, Type) case EParticlePropertyType::TypeName: return (const TDirtyElementPool<T>&)TypeName##Pool;
-#include "ParticlePropertiesTypes.inl"
-#undef PARTICLE_PROPERTY_TYPE
-			default: check(false);
-			}
-
-			static TDirtyElementPool<T> ErrorPool;
-			return ErrorPool;
-		}
-
-	private:
-
-#define PARTICLE_PROPERTY_TYPE(TypeName, Type) TDirtyElementPool<Type> TypeName##Pool;
-#include "ParticlePropertiesTypes.inl"
-#undef PARTICLE_PROPERTY_TYPE
-	};
-
+class FDirtyPropertiesManager;
 
 template <typename T>
-class TRemoteParticleProperty
+class TRemoteProperty
 {
 public:
-	TRemoteParticleProperty()
+	TRemoteProperty()
 	{
 		Idx.bHasEntry = false;
 	}
 
-	TRemoteParticleProperty(const TRemoteParticleProperty<T>& Rhs) = delete;
-	TRemoteParticleProperty(TRemoteParticleProperty<T>&& Rhs)
+	TRemoteProperty(const TRemoteProperty<T>& Rhs) = delete;
+	TRemoteProperty(TRemoteProperty<T>&& Rhs)
 	: Idx(Rhs.Idx)
 	{
 		Rhs.bHasEntry = false;
 	}
 
-	~TRemoteParticleProperty()
+	~TRemoteProperty()
 	{
 		ensure(!Idx.bHasEntry);	//leaking, make sure to call Pop
 	}
 
-	const T& Read(const FDirtyPropertiesManager& Manager) const
+	const T& Read(const FDirtyPropertiesManager& Manager) const;
+	void Clear(FDirtyPropertiesManager& Manager);
+	void Write(FDirtyPropertiesManager& Manager,const T& Val);
+	
+	bool IsSet() const
 	{
-		ensure(Idx.bHasEntry);
-		return Manager.GetPool<T>().Read(Idx.Entry);
+		return Idx.bHasEntry;
 	}
+private:
+	FDirtyIdx Idx;
+
+	TRemoteProperty<T>& operator=(const TRemoteProperty<T>& Rhs){}
+};
 
-	T Pop(FDirtyPropertiesManager& Manager)
+struct FParticlePropertiesData
+{
+	FParticlePropertiesData(FDirtyPropertiesManager* InManager = nullptr)
+		: Manager(InManager)
 	{
-		ensure(Idx.bHasEntry);
-		Idx.bHasEntry = false;
-		return Manager.GetPool<T>().Pop(Idx.Entry);
 	}
 
-	void Clear(FDirtyPropertiesManager& Manager)
+	template <typename T, EParticleProperty PropertyIdx>
+	TRemoteProperty<T>& GetProperty()
 	{
-		if(Idx.bHasEntry)
+		switch(PropertyIdx)
 		{
-			Idx.bHasEntry = false;
-			Manager.GetPool<T>().Pop(Idx.Entry);
+#define PARTICLE_PROPERTY(PropName, Type) case EParticleProperty::PropName: return (TRemoteProperty<T>&) PropName;
+#include "ParticleProperties.inl"
+#undef PARTICLE_PROPERTY
+		default: check(false);
 		}
+
+		static TRemoteProperty<T> Error;
+		return Error;
 	}
 
-	void Write(FDirtyPropertiesManager& Manager, const T& Val)
+	void Clear()
 	{
-		if(Idx.bHasEntry)
-		{
-			Manager.GetPool<T>().Update(Idx.Entry,Val);
-		}
-		else
+		if(Manager)
 		{
-			Idx.Entry = Manager.GetPool<T>().Write(Val);
-			Idx.bHasEntry = true;
+#define PARTICLE_PROPERTY(PropName, Type) PropName.Clear(*Manager);
+#include "ParticleProperties.inl"
+#undef PARTICLE_PROPERTY
 		}
 	}
 
-	bool IsSet() const
+	inline void FreeToManager();
+	
+	~FParticlePropertiesData()
 	{
-		return Idx.bHasEntry;
+		Clear();
 	}
+
+	FDirtyPropertiesManager* GetManager(){ return Manager; }
+	const FDirtyPropertiesManager* GetManager() const { return Manager; }
+
+#define PARTICLE_PROPERTY(PropName, Type)\
+Type const & Get##PropName() const { return PropName.Read(*Manager); }\
+bool Has##PropName() const { return PropName.IsSet(); }\
+Type const * Find##PropName() const { return Has##PropName() ? &Get##PropName() : nullptr; }
+
+#include "ParticleProperties.inl"
+#undef PARTICLE_PROPERTY
+
 private:
-	FDirtyIdx Idx;
+#define PARTICLE_PROPERTY(PropName, Type) TRemoteProperty<Type> PropName;
+#include "ParticleProperties.inl"
+#undef PARTICLE_PROPERTY
+
+	FDirtyPropertiesManager* Manager;
 
-	TRemoteParticleProperty<T>& operator=(const TRemoteParticleProperty<T>& Rhs){}
 };
 
-struct FParticlePropertiesData
+struct FShapePropertiesData
 {
-	template <typename T, EParticleProperty PropertyIdx>
-	TRemoteParticleProperty<T>& GetProperty()
+	FShapePropertiesData(FDirtyPropertiesManager* InManager)
+		: Manager(InManager)
+	{
+
+	}
+	template <typename T,EShapeProperty PropertyIdx>
+	TRemoteProperty<T>& GetProperty()
 	{
 		switch(PropertyIdx)
 		{
-#define PARTICLE_PROPERTY(PropName, Type) case EParticleProperty::PropName: return (TRemoteParticleProperty<T>&) PropName;
-#include "ParticleProperties.inl"
-#undef PARTICLE_PROPERTY
+#define SHAPE_PROPERTY(PropName, Type) case EShapeProperty::PropName: return (TRemoteProperty<T>&) PropName;
+#include "ShapeProperties.inl"
+#undef SHAPE_PROPERTY
 		default: check(false);
 		}
 
-		static TRemoteParticleProperty<T> Error;
+		static TRemoteProperty<T> Error;
 		return Error;
 	}
 
+	void Clear()
+	{
+		if(Manager)
+		{
+#define SHAPE_PROPERTY(PropName, Type) PropName.Clear(*Manager);
+#include "ShapeProperties.inl"
+#undef SHAPE_PROPERTY
+		}
+	}
+
+	~FShapePropertiesData()
+	{
+		Clear();
+	}
+
+	inline void FreeToManager();
+
+	FDirtyPropertiesManager* GetManager(){ return Manager; }
+	const FDirtyPropertiesManager* GetManager() const { return Manager; }
+
+#define SHAPE_PROPERTY(PropName, Type)\
+Type const & Get##PropName() const { return PropName.Read(*Manager); }\
+bool Has##PropName() const { return PropName.IsSet(); }\
+Type const * Find##PropName() const { return Has##PropName() ? &Get##PropName() : nullptr; }
+
+#include "ShapeProperties.inl"
+#undef SHAPE_PROPERTY
+
+private:
+#define SHAPE_PROPERTY(PropName, Type) TRemoteProperty<Type> PropName;
+#include "ShapeProperties.inl"
+#undef SHAPE_PROPERTY
+
 	FDirtyPropertiesManager* Manager;
+};
 
-	~FParticlePropertiesData()
+class FPerShapeData;
+
+class CHAOS_API FShapeRemoteDataContainer
+{
+public:
+	FShapeRemoteDataContainer(FDirtyPropertiesManager* InManager)
+	: Manager(InManager)
 	{
-		if(Manager)
+	}
+
+	~FShapeRemoteDataContainer()
+	{
+		Clear();
+	}
+
+	void SyncShapes(TArray<TUniquePtr<FPerShapeData>, TInlineAllocator<1>>& Shapes);
+	
+	void DetachRemoteData(TArray<TUniquePtr<FPerShapeData>,TInlineAllocator<1>>& Shapes);
+
+	inline void FreeToManager();
+
+	inline FShapePropertiesData* NewRemoteShapeProperties();
+	
+	void Clear()
+	{
+		//todo: avoid iterating all remote data regardless of if dirty or not
+		for(FShapePropertiesData* RemoteData : RemoteDatas)
 		{
-#define PARTICLE_PROPERTY(PropName, Type) PropName.Clear(*Manager);
+			if(RemoteData)
+			{
+				RemoteData->FreeToManager();
+			}
+		}
+
+		RemoteDatas.Reset();
+	}
+
+	const auto& GetRemoteDatas() const
+	{
+		return RemoteDatas;
+	}
+
+	auto& GetRemoteDatas()
+	{
+		return RemoteDatas;
+	}
+	
+private:
+
+	FDirtyPropertiesManager* Manager;
+	TArray<FShapePropertiesData*,TInlineAllocator<4>> RemoteDatas;
+};
+#endif
+
+class FDirtyPropertiesManager
+{
+public:
+
+#if 0
+	FParticlePropertiesData* NewRemoteParticleProperties()
+	{
+		return RemoteParticlePool.NewEntry(this);
+	}
+
+	void FreeRemoteParticleProperties(FParticlePropertiesData* Entry)
+	{
+		RemoteParticlePool.FreeEntry(Entry);
+	}
+
+	FShapePropertiesData* NewRemoteShapeProperties()
+	{
+		return RemoteShapePool.NewEntry(this);
+	}
+
+	void FreeRemoteShapeProperties(FShapePropertiesData* Entry)
+	{
+		RemoteShapePool.FreeEntry(Entry);
+	}
+
+	FShapeRemoteDataContainer* NewRemoteShapeContainer()
+	{
+		return RemoteShapeContainerPool.NewEntry(this);
+	}
+
+	void FreeRemoteShapeContainer(FShapeRemoteDataContainer* Entry)
+	{
+		RemoteShapeContainerPool.FreeEntry(Entry);
+	}
+#endif
+
+	void SetNumParticles(int32 NumParticles)
+	{
+#define PARTICLE_PROPERTY(PropName, Type) PropName##Pool.SetNum(NumParticles);
 #include "ParticleProperties.inl"
 #undef PARTICLE_PROPERTY
-		}
 	}
 
-#define PARTICLE_PROPERTY(PropName, Type)\
-Type const & Get##PropName() { return PropName.Read(*Manager); };\
-bool Has##PropName() const { return PropName.IsSet(); };
+	void SetNumShapes(int32 NumShapes)
+	{
+#define SHAPE_PROPERTY(PropName, Type) PropName##ShapePool.SetNum(NumShapes);
+#include "ShapeProperties.inl"
+#undef SHAPE_PROPERTY
+	}
 
+	template <typename T, EParticleProperty PropName>
+	TDirtyElementPool<T>& GetParticlePool()
+	{
+		switch(PropName)
+		{
+#define PARTICLE_PROPERTY(PropName, Type) case EParticleProperty::PropName: return (TDirtyElementPool<T>&)PropName##Pool;
 #include "ParticleProperties.inl"
 #undef PARTICLE_PROPERTY
+		default: check(false);
+		}
+
+		static TDirtyElementPool<T> ErrorPool;
+		return ErrorPool;
+	}
+
+	template <typename T,EParticleProperty PropName>
+	const TDirtyElementPool<T>& GetParticlePool() const
+	{
+		switch(PropName)
+		{
+#define PARTICLE_PROPERTY(PropName, Type) case EParticleProperty::PropName: return (TDirtyElementPool<T>&)PropName##Pool;
+#include "ParticleProperties.inl"
+#undef PARTICLE_PROPERTY
+		default: check(false);
+		}
+
+		static TDirtyElementPool<T> ErrorPool;
+		return ErrorPool;
+	}
+
+	template <typename T,EShapeProperty PropName>
+	TDirtyElementPool<T>& GetShapePool()
+	{
+		switch(PropName)
+		{
+#define SHAPE_PROPERTY(PropName, Type) case EShapeProperty::PropName: return (TDirtyElementPool<T>&)PropName##ShapePool;
+#include "ShapeProperties.inl"
+#undef SHAPE_PROPERTY
+		default: check(false);
+		}
+
+		static TDirtyElementPool<T> ErrorPool;
+		return ErrorPool;
+	}
+
+	template <typename T,EShapeProperty PropName>
+	const TDirtyElementPool<T>& GetShapePool() const
+	{
+		switch(PropName)
+		{
+#define SHAPE_PROPERTY(PropName, Type) case EShapeProperty::PropName: return (TDirtyElementPool<T>&)PropName##ShapePool;
+#include "ShapeProperties.inl"
+#undef SHAPE_PROPERTY
+		default: check(false);
+		}
+
+		static TDirtyElementPool<T> ErrorPool;
+		return ErrorPool;
+	}
 
 private:
-#define PARTICLE_PROPERTY(PropName, Type) TRemoteParticleProperty<Type> PropName;
+
+#define PARTICLE_PROPERTY(PropName, Type) TDirtyElementPool<Type> PropName##Pool;
 #include "ParticleProperties.inl"
 #undef PARTICLE_PROPERTY
 
+#define SHAPE_PROPERTY(PropName, Type) TDirtyElementPool<Type> PropName##ShapePool;
+#include "ShapeProperties.inl"
+#undef SHAPE_PROPERTY
+
+#if 0
+	template <typename T>
+	class TRemotePropertiesPool
+	{
+	public:
+		T* NewEntry(FDirtyPropertiesManager* Manager)
+		{
+			if(Pool.Num())
+			{
+				return Pool.Pop(/*bAllowShrinking=*/false);
+			} else
+			{
+				return new T(Manager);
+			}
+		}
+
+		void FreeEntry(T* Entry)
+		{
+			Entry->Clear();
+			Pool.Add(Entry);
+		}
+
+		~TRemotePropertiesPool()
+		{
+			for(T* Entry : Pool)
+			{
+				delete Entry;
+			}
+		}
+
+		int32 NumInPool() const { return Pool.Num(); }
+
+	private:
+		TArray<T*> Pool;
+	};
+
+	TRemotePropertiesPool<FParticlePropertiesData> RemoteParticlePool;
+	TRemotePropertiesPool<FShapePropertiesData> RemoteShapePool;
+	TRemotePropertiesPool<FShapeRemoteDataContainer> RemoteShapeContainerPool;
+#endif
+
 };
 
-template <typename T, EParticleProperty PropName>
-class TParticleProperty
+class FParticleDirtyData
 {
 public:
-	TParticleProperty() = default;
-	TParticleProperty(const T& Val)
-		: Property(Val)
+	
+	void SetFlags(FParticleDirtyFlags InFlags)
 	{
+		Flags = InFlags;
 	}
 
-	//we don't support this because it could lead to bugs with values not being properly written to remote
-	TParticleProperty(const TParticleProperty<T,PropName>& Rhs) = delete;
+	template <typename T, EParticleProperty PropName>
+	void SyncRemote(FDirtyPropertiesManager& Manager, int32 Idx, const T& Val) const
+	{
+		if(Flags.IsDirty(ParticlePropToFlag(PropName)))
+		{
+			Manager.GetParticlePool<T,PropName>().GetElement(Idx) = Val;
+		}
+	}
 
-	const T& Read() const { return Property; }
-	void Write(const T& Val, bool bInvalidate, FParticlePropertiesData* Remote, FParticleDirtyFlags& Dirty, IPhysicsProxyBase* Proxy)
+	void Clear(FDirtyPropertiesManager& Manager, int32 Idx)
 	{
-		Property = Val;
-		WriteImp(bInvalidate, Remote, Dirty, Proxy);
+#define PARTICLE_PROPERTY(PropName, Type) ClearHelper<Type, EParticleProperty::PropName>(Manager, Idx);
+#include "ParticleProperties.inl"
+#undef PARTICLE_PROPERTY
+		Flags.Clear();
 	}
 
-	template <typename Lambda>
-	void Modify(bool bInvalidate,FParticlePropertiesData* Remote,FParticleDirtyFlags& Dirty,IPhysicsProxyBase* Proxy,const Lambda& LambdaFunc)
+#define PARTICLE_PROPERTY(PropName, Type)\
+Type const & Get##PropName(const FDirtyPropertiesManager& Manager, int32 Idx) const { return ReadImp<Type, EParticleProperty::PropName>(Manager, Idx); }\
+bool Has##PropName() const { return Flags.IsDirty(ParticlePropToFlag(EParticleProperty::PropName)); }\
+Type const * Find##PropName(const FDirtyPropertiesManager& Manager, int32 Idx) const { return Has##PropName() ? &Get##PropName(Manager, Idx) : nullptr; }
+
+#include "ParticleProperties.inl"
+#undef PARTICLE_PROPERTY
+
+private:
+	FParticleDirtyFlags Flags;
+
+	template <typename T,EParticleProperty PropName>
+	const T& ReadImp(const FDirtyPropertiesManager& Manager, int32 Idx) const
 	{
-		LambdaFunc(Property);
-		WriteImp(bInvalidate, Remote, Dirty, Proxy);
+		ensure(Flags.IsDirty(ParticlePropToFlag(PropName)));
+		return Manager.GetParticlePool<T,PropName>().GetElement(Idx);
 	}
 
-	void InitialFlush(FParticlePropertiesData& Remote,FParticleDirtyFlags& Dirty, IPhysicsProxyBase* Proxy)
+	template <typename T, EParticleProperty PropName>
+	void ClearHelper(FDirtyPropertiesManager& Manager, int32 Idx)
 	{
-		if(Dirty.IsDirty(PropertyFlag))
+		if(Flags.IsDirty(ParticlePropToFlag(PropName)))
 		{
-			WriteToRemote(Remote,Proxy);
+			Manager.GetParticlePool<T, PropName>().Reset(Idx);
 		}
 	}
+};
 
-	void Serialize(FChaosArchive& Ar)
+class FShapeDirtyData
+{
+public:
+
+	FShapeDirtyData(int32 InShapeIdx)
+	: ShapeIdx(InShapeIdx)
 	{
-		Ar << Property;
+
 	}
 
-private:
-	T Property;
-	static constexpr EParticleFlags PropertyFlag = ParticlePropToFlag(PropName);
+	int32 GetShapeIdx() const { return ShapeIdx; }
+
+	void SetFlags(FShapeDirtyFlags InFlags)
+	{
+		Flags = InFlags;
+	}
 
-	void WriteToRemote(FParticlePropertiesData& Remote, IPhysicsProxyBase* Proxy)
+	template <typename T,EShapeProperty PropName>
+	void SyncRemote(FDirtyPropertiesManager& Manager,int32 Idx, const T& Val) const
 	{
-		Remote.GetProperty<T,PropName>().Write(*Remote.Manager, Property);
-		if(Proxy)
+		if(Flags.IsDirty(ShapePropToFlag(PropName)))
 		{
-			if(FPhysicsSolverBase* PhysicsSolverBase = Proxy->GetSolver<FPhysicsSolverBase>())
-			{
-				PhysicsSolverBase->AddDirtyProxy(Proxy);
-			}
+			Manager.GetShapePool<T,PropName>().GetElement(Idx) = Val;
 		}
 	}
 
-	void WriteImp(bool bInvalidate, FParticlePropertiesData* Remote, FParticleDirtyFlags& Dirty, IPhysicsProxyBase* Proxy)
+	void Clear(FDirtyPropertiesManager& Manager, int32 Idx)
 	{
-		if(bInvalidate)
-		{
-			Dirty.MarkDirty(PropertyFlag);
-			if(Remote)
-			{
-				WriteToRemote(*Remote,Proxy);
+#define SHAPE_PROPERTY(PropName, Type) ClearHelper<Type, EShapeProperty::PropName>(Manager, Idx);
+#include "ShapeProperties.inl"
+#undef SHAPE_PROPERTY
+		Flags.Clear();
+	}
 
-			}
+#define SHAPE_PROPERTY(PropName, Type)\
+Type const & Get##PropName(const FDirtyPropertiesManager& Manager, int32 Idx) const { return ReadImp<Type, EShapeProperty::PropName>(Manager, Idx); }\
+bool Has##PropName() const { return Flags.IsDirty(ShapePropToFlag(EShapeProperty::PropName)); }\
+Type const * Find##PropName(const FDirtyPropertiesManager& Manager, int32 Idx) const { return Has##PropName() ? &Get##PropName(Manager, Idx) : nullptr; }
 
-			//for now we have to dirty proxy in here because rest of pipeline is still old
-			//todo: remove this once remote data is set, writing to remote should take care of proxy dirtying
-			if(Proxy)
-			{
-				if(FPhysicsSolverBase* PhysicsSolverBase = Proxy->GetSolver<FPhysicsSolverBase>())
-				{
-					PhysicsSolverBase->AddDirtyProxy(Proxy);
-				}
-			}
-		}
+#include "ShapeProperties.inl"
+#undef PARTICLE_PROPERTY
+
+private:
+	int32 ShapeIdx;
+	FShapeDirtyFlags Flags;
+
+	template <typename T,EShapeProperty PropName>
+	const T& ReadImp(const FDirtyPropertiesManager& Manager, int32 Idx) const
+	{
+		ensure(Flags.IsDirty(ShapePropToFlag(PropName)));
+		return Manager.GetShapePool<T,PropName>().GetElement(Idx);
 	}
 
-	//we don't support this because it could lead to bugs with values not being properly written to remote
-	TParticleProperty<T,PropName>& operator=(const TParticleProperty<T,PropName>& Rhs) {}
+	template <typename T,EShapeProperty PropName>
+	void ClearHelper(FDirtyPropertiesManager& Manager, int32 Idx)
+	{
+		if(Flags.IsDirty(ShapePropToFlag(PropName)))
+		{
+			Manager.GetShapePool<T,PropName>().Reset(Idx);
+		}
+	}
 };
 
-template <typename T, EParticleProperty PropName>
-FChaosArchive& operator<<(FChaosArchive& Ar, TParticleProperty<T, PropName>& Prop)
+#if 0
+void FParticlePropertiesData::FreeToManager()
 {
-	Prop.Serialize(Ar);
-	return Ar;
+	if(Manager)
+	{
+		Manager->FreeRemoteParticleProperties(this);
+	}
 }
+
+void FShapePropertiesData::FreeToManager()
+{
+	if(Manager)
+	{
+		Manager->FreeRemoteShapeProperties(this);
+	}
 }
+
+void FShapeRemoteDataContainer::FreeToManager()
+{
+	if(Manager)
+	{
+		Manager->FreeRemoteShapeContainer(this);
+	}
+}
+
+FShapePropertiesData* FShapeRemoteDataContainer::NewRemoteShapeProperties()
+{
+	return Manager ? Manager->NewRemoteShapeProperties() : nullptr;
+}
+
+template <typename T>
+const T& TRemoteProperty<T>::Read(const FDirtyPropertiesManager& Manager) const
+{
+	ensure(Idx.bHasEntry);
+	return Manager.GetPool<T>().Read(Idx.Entry);
+}
+
+template <typename T>
+void TRemoteProperty<T>::Clear(FDirtyPropertiesManager& Manager)
+{
+	if(Idx.bHasEntry)
+	{
+		Idx.bHasEntry = false;
+		Manager.GetPool<T>().Pop(Idx.Entry);
+	}
+}
+
+template <typename T>
+void TRemoteProperty<T>::Write(FDirtyPropertiesManager& Manager,const T& Val)
+{
+	if(Idx.bHasEntry)
+	{
+		Manager.GetPool<T>().Update(Idx.Entry,Val);
+	} else
+	{
+		Idx.Entry = Manager.GetPool<T>().Write(Val);
+		Idx.bHasEntry = true;
+	}
+}
+#endif
+
+}
\ No newline at end of file
diff --git a/Engine/Source/Runtime/Experimental/Chaos/Public/Chaos/ParticleHandle.h b/Engine/Source/Runtime/Experimental/Chaos/Public/Chaos/ParticleHandle.h
--- a/Engine/Source/Runtime/Experimental/Chaos/Public/Chaos/ParticleHandle.h
+++ b/Engine/Source/Runtime/Experimental/Chaos/Public/Chaos/ParticleHandle.h
@@ -6,7 +6,7 @@
 #include "Chaos/PBDGeometryCollectionParticles.h"
 #include "Chaos/ParticleHandleFwd.h"
 #include "Chaos/ParticleIterator.h"
-#include "Chaos/ParticleDirtyFlags.h"
+#include "Chaos/Properties.h"
 #include "ChaosCheck.h"
 #include "Chaos/ChaosDebugDrawDeclares.h"
 #if CHAOS_DEBUG_DRAW
@@ -96,10 +96,6 @@ void PBDRigidParticleDefaultConstruct(TPBDRigidParticle<T,d>& Concrete, const TP
 	//don't bother calling parent since the call gets made by the corresponding hierarchy in FConcrete
 	Concrete.SetCollisionGroup(0);
 	Concrete.SetDisabled(Params.bDisabled);
-	Concrete.SetPreV(Concrete.V());
-	Concrete.SetPreW(Concrete.W());
-	Concrete.SetP(Concrete.X());
-	Concrete.SetQ(Concrete.R());
 	Concrete.SetF(TVector<T, d>(0));
 	Concrete.SetTorque(TVector<T, d>(0));
 	Concrete.SetLinearImpulse(TVector<T, d>(0));
@@ -435,7 +431,7 @@ public:
 	const TUniquePtr<FImplicitObject>& DynamicGeometry() const { return GeometryParticles->DynamicGeometry(ParticleIdx); }
 	void SetDynamicGeometry(TUniquePtr<FImplicitObject>&& Unique) { GeometryParticles->SetDynamicGeometry(ParticleIdx, MoveTemp(Unique)); }
 
-	const TShapesArray<T,d>& ShapesArray() const { return GeometryParticles->ShapesArray(ParticleIdx); }
+	const FShapesArray& ShapesArray() const { return GeometryParticles->ShapesArray(ParticleIdx); }
 
 	const TAABB<T, d>& LocalBounds() const { return GeometryParticles->LocalBounds(ParticleIdx); }
 	void SetLocalBounds(const TAABB<T, d>& NewBounds) { GeometryParticles->LocalBounds(ParticleIdx) = NewBounds; }
@@ -521,7 +517,7 @@ public:
 		GeometryParticles->SetHandle(ParticleIdx, this);
 	}
 
-	const TPerShapeData<T, d>* GetImplicitShape(const FImplicitObject* InObject)
+	const FPerShapeData* GetImplicitShape(const FImplicitObject* InObject)
 	{
 		return GeometryParticles->GetImplicitShape(ParticleIdx, InObject);
 	}
@@ -1355,8 +1351,6 @@ public:
 protected:
 
 	TGeometryParticle(const TGeometryParticleParameters<T, d>& StaticParams = TGeometryParticleParameters<T, d>())
-		: MUserData(nullptr)
-		, RemoteData(nullptr)
 	{
 		Type = EParticleType::Static;
 		Proxy = nullptr;
@@ -1377,9 +1371,8 @@ public:
 
 	virtual void Serialize(FChaosArchive& Ar)
 	{
-		Ar << MX;
-		Ar << MR;
-		Ar << MGeometry;
+		Ar << MXR;
+		Ar << MNonFrequentData;
 		Ar << MShapesArray;
 		Ar << Type;
 		//Ar << MDirtyFlags;
@@ -1398,28 +1391,28 @@ public:
 
 	virtual bool IsParticleValid() const
 	{
-		auto& Geometry = MGeometry.Read();
+		auto& Geometry = MNonFrequentData.Read().Geometry;
 		return Geometry && Geometry->IsValidGeometry();	//todo: if we want support for sample particles without geometry we need to adjust this
 	}
 
 	static TGeometryParticle<T, d>* SerializationFactory(FChaosArchive& Ar, TGeometryParticle<T, d>* Serializable);
 
-	const TVector<T, d>& X() const { return MX.Read(); }
+	const TVector<T, d>& X() const { return MXR.Read().X; }
 	void SetX(const TVector<T, d>& InX, bool bInvalidate = true)
 	{
-		MX.Write(InX, bInvalidate, RemoteData, MDirtyFlags, Proxy);
+		MXR.Modify(bInvalidate,MDirtyFlags,Proxy,[&InX](auto& Data){ Data.X = InX;});
 	}
 
-	FUniqueIdx UniqueIdx() const { return MUniqueIdx.Read(); }
+	FUniqueIdx UniqueIdx() const { return MNonFrequentData.Read().UniqueIdx; }
 	void SetUniqueIdx(const FUniqueIdx UniqueIdx, bool bInvalidate = true)
 	{
-		MUniqueIdx.Write(UniqueIdx,bInvalidate,RemoteData,MDirtyFlags,Proxy);
+		MNonFrequentData.Modify(bInvalidate,MDirtyFlags,Proxy,[UniqueIdx](auto& Data){ Data.UniqueIdx = UniqueIdx;});
 	}
 
-	const TRotation<T, d>& R() const { return MR.Read(); }
+	const TRotation<T, d>& R() const { return MXR.Read().R; }
 	void SetR(const TRotation<T, d>& InR, bool bInvalidate = true)
 	{
-		MR.Write(InR,bInvalidate,RemoteData,MDirtyFlags,Proxy);
+		MXR.Modify(bInvalidate,MDirtyFlags,Proxy,[&InR](auto& Data){ Data.R = InR;});
 	}
 
 	//todo: geometry should not be owned by particle
@@ -1437,7 +1430,7 @@ public:
 	//       We should replace this with a method for supporting SetGeometry(RawGeometry).
 	void SetGeometry(TSharedPtr<FImplicitObject, ESPMode::ThreadSafe> SharedGeometry)
 	{
-		MGeometry.Write(SharedGeometry,/*bInvalidate=*/true,RemoteData,MDirtyFlags,Proxy);
+		MNonFrequentData.Modify(true,MDirtyFlags,Proxy,[&SharedGeometry](auto& Data){ Data.Geometry = SharedGeometry;});
 		UpdateShapesArray();
 	}
 
@@ -1448,42 +1441,69 @@ public:
 		check(false);
 	}
 
-	void* UserData() const { return MUserData.Read(); }
+	void* UserData() const { return MNonFrequentData.Read().UserData; }
 	void SetUserData(void* InUserData)
 	{
-		MUserData.Write(InUserData,/*bInvalidate=*/true,RemoteData,MDirtyFlags,Proxy);
+		MNonFrequentData.Modify(true,MDirtyFlags,Proxy,[InUserData](auto& Data){ Data.UserData = InUserData;});
 	}
 
 	void UpdateShapeBounds()
 	{
-		if (MGeometry.Read()->HasBoundingBox())
+		if (MNonFrequentData.Read().Geometry->HasBoundingBox())
 		{
 			for (auto& Shape : MShapesArray)
 			{
-				Shape->UpdateShapeBounds(FRigidTransform3(MX.Read(), MR.Read()));
+				Shape->UpdateShapeBounds(FRigidTransform3(X(), R()));
 			}
 		}
 	}
 
+	void SetShapeCollisionDisable(int32 InShapeIndex,bool bInDisable)
+	{
+		const bool bCurrent = MShapesArray[InShapeIndex]->GetDisable();
+		if(bCurrent != bInDisable)
+		{
+			MShapesArray[InShapeIndex]->SetDisable(bInDisable);
+		}
+	}
+
+	void SetShapeCollisionTraceType(int32 InShapeIndex,EChaosCollisionTraceFlag TraceType)
+	{
+		const EChaosCollisionTraceFlag Current = MShapesArray[InShapeIndex]->GetCollisionTraceType();
+		if(Current != TraceType)
+		{
+			MShapesArray[InShapeIndex]->SetCollisionTraceType(TraceType);
+		}
+	}
+
+	void SetShapeSimData(int32 InShapeIndex,const FCollisionFilterData& SimData)
+	{
+		const FCollisionFilterData& Current = MShapesArray[InShapeIndex]->GetSimData();
+		if(Current != SimData)
+		{
+			MShapesArray[InShapeIndex]->SetSimData(SimData);
+		}
+	}
+
 #if CHAOS_CHECKED
-	const FName& DebugName() const { return MDebugName.Read(); }
+	const FName& DebugName() const { return MNonFrequentData.Read().DebugName; }
 	void SetDebugName(const FName& InDebugName)
 	{
-		MDebugName.Write(InDebugName,/*bInvalidate=*/true,RemoteData,MDirtyFlags,Proxy);
+		MNonFrequentData.Modify(bInvalidate,MDirtyFlags,Proxy,[&InDebugName](auto& Data){ Data.DebugName = InDebugName;});
 	}
 #endif
 
 	//Note: this must be called after setting geometry. This API seems bad. Should probably be part of setting geometry
-	void SetShapesArray(TShapesArray<T, d>&& InShapesArray)
+	void SetShapesArray(FShapesArray&& InShapesArray)
 	{
 		ensure(InShapesArray.Num() == MShapesArray.Num());
 		MShapesArray = MoveTemp(InShapesArray);
 		MapImplicitShapes();
 	}
 
-	TSerializablePtr<FImplicitObject> Geometry() const { return MakeSerializable(MGeometry.Read()); }
+	TSerializablePtr<FImplicitObject> Geometry() const { return MakeSerializable(MNonFrequentData.Read().Geometry); }
 
-	const TShapesArray<T,d>& ShapesArray() const { return MShapesArray; }
+	const FShapesArray& ShapesArray() const { return MShapesArray; }
 
 	EObjectStateType ObjectState() const;
 	void SetObjectState(const EObjectStateType InState, bool bAllowEvents = false);
@@ -1500,47 +1520,20 @@ public:
 	const TPBDRigidParticle<T, d>* CastToRigidParticle() const;
 	TPBDRigidParticle<T, d>* CastToRigidParticle();
 
-	FSpatialAccelerationIdx SpatialIdx() const { return MSpatialIdx.Read(); }
+	FSpatialAccelerationIdx SpatialIdx() const { return MMiscData.Read().SpatialIdx; }
 	void SetSpatialIdx(FSpatialAccelerationIdx Idx)
 	{
-		MSpatialIdx.Write(Idx,/*bInvalidate=*/true,RemoteData,MDirtyFlags,Proxy);
-	}
-
-	void SetShapeCollisionDisable(int32 InShapeIndex, bool bInDisable)
-	{
-		const bool bCurrent = MShapesArray[InShapeIndex]->bDisable;
-		if(bCurrent != bInDisable)
-		{
-			MShapesArray[InShapeIndex]->bDisable = bInDisable;
-			MarkDirty(EParticleFlags::ShapeDisableCollision);
-		}
-	}
-
-	void SetShapeCollisionTraceType(int32 InShapeIndex, EChaosCollisionTraceFlag TraceType)
-	{
-		const EChaosCollisionTraceFlag Current = MShapesArray[InShapeIndex]->CollisionTraceType;
-		if (Current != TraceType)
-		{
-			MShapesArray[InShapeIndex]->CollisionTraceType = TraceType;
-			MarkDirty(EParticleFlags::CollisionTraceType);
-		}
+		MMiscData.Modify(true,MDirtyFlags,Proxy,[Idx](auto& Data){ Data.SpatialIdx = Idx;});
 	}
 
-	void SetShapeSimData(int32 InShapeIndex, const FCollisionFilterData& SimData)
+	bool IsDirty() const
 	{
-		const FCollisionFilterData& Current = MShapesArray[InShapeIndex]->SimData;
-		if (Current != SimData)
-		{
-			MShapesArray[InShapeIndex]->SimData = SimData;
-			MarkDirty(EParticleFlags::ShapeSimData);
-		}
+		return MDirtyFlags.IsDirty();
 	}
 
-	FParticleData* NewData() const { return new TGeometryParticleData<T, d>( *this ); }
-
-	bool IsDirty() const
+	bool IsClean() const
 	{
-		return MDirtyFlags.IsDirty();
+		return MDirtyFlags.IsClean();
 	}
 
 	bool IsDirty(const EParticleFlags CheckBits) const
@@ -1553,19 +1546,11 @@ public:
 		return MDirtyFlags;
 	}
 
-	//todo: should this exist?
 	void ClearDirtyFlags()
 	{
 		MDirtyFlags.Clear();
 	}
-
-	//todo: should this exist?
-	void MarkClean(const EParticleFlags CleanBits)
-	{
-		MDirtyFlags.MarkClean(CleanBits);
-	}
-
-
+	
 	TGeometryParticleHandle<T, d>* Handle() const
 	{
 		if (Proxy)
@@ -1576,7 +1561,7 @@ public:
 		return nullptr;
 	}
 
-	const TPerShapeData<T, d>* GetImplicitShape(const FImplicitObject* InImplicit) const
+	const FPerShapeData* GetImplicitShape(const FImplicitObject* InImplicit) const
 	{
 		const int32* ShapeIndex = ImplicitShapeMap.Find(InImplicit);
 		if(ShapeIndex)
@@ -1587,32 +1572,63 @@ public:
 		return nullptr;
 	}
 
+	void SyncRemoteData(FDirtyPropertiesManager& Manager, int32 DataIdx, FParticleDirtyData& RemoteData, const TArray<int32>& ShapeDataIndices, FShapeDirtyData* ShapesRemoteData) const
+	{
+		RemoteData.SetFlags(MDirtyFlags);
+		SyncRemoteDataImp(Manager, DataIdx, RemoteData);
+
+		for(const int32 ShapeDataIdx : ShapeDataIndices)
+		{
+			FShapeDirtyData& ShapeRemoteData = ShapesRemoteData[ShapeDataIdx];
+			const int32 ShapeIdx = ShapeRemoteData.GetShapeIdx();
+			MShapesArray[ShapeIdx]->SyncRemoteData(Manager, ShapeDataIdx, ShapeRemoteData);
+		}
+	}
+
+
+	class IPhysicsProxyBase* GetProxy() const
+	{
+		return Proxy;
+	}
+
+	void SetProxy(IPhysicsProxyBase* InProxy)
+	{
+		Proxy = InProxy;
+		if(Proxy)
+		{
+			if(MDirtyFlags.IsDirty())
+			{
+				if(FPhysicsSolverBase* PhysicsSolverBase = Proxy->GetSolver<FPhysicsSolverBase>())
+				{
+					PhysicsSolverBase->AddDirtyProxy(Proxy);
+				}
+			}
+		}
+
+		for(auto& Shape : MShapesArray)
+		{
+			Shape->SetProxy(Proxy);
+		}
+	}
+
+protected:
+
 	// Pointer to any data that the solver wants to associate with this particle
 	// TODO: It's important to eventually hide this!
 	// Right now it's exposed to lubricate the creation of the whole proxy system.
 	class IPhysicsProxyBase* Proxy;
+private:
 
-	// TODO: This is an awful side effect of housing the dirty flag for shape data
-	//       inside the particle, but not setting the shape data through it.
-	void MarkShapeSimDataDirty() { MarkDirty(EParticleFlags::ShapeSimData); }
-
-	void SetRemoteData(FParticlePropertiesData* InRemoteData)
-	{
-		RemoteData = InRemoteData;
-		FlushToRemoteData();
-	}
+	TParticleProperty<FParticlePositionRotation, EParticleProperty::XR> MXR;
+	TParticleProperty<FParticleNonFrequentData, EParticleProperty::NonFrequentData> MNonFrequentData;
 
+protected:
+	TParticleProperty<FParticleMisc,EParticleProperty::Misc> MMiscData;
 private:
-	TParticleProperty<TVector<T, d>, EParticleProperty::X> MX;
-	TParticleProperty<FUniqueIdx, EParticleProperty::UniqueIdx> MUniqueIdx;
-	TParticleProperty<TRotation<T, d>, EParticleProperty::R> MR;
-	TParticleProperty<TSharedPtr<FImplicitObject, ESPMode::ThreadSafe>, EParticleProperty::Geometry> MGeometry;	//TODO: geometry should live in bodysetup
-	TShapesArray<T,d> MShapesArray;
+
+	FShapesArray MShapesArray;
 	TMap<const FImplicitObject*, int32> ImplicitShapeMap;
-	TParticleProperty<FSpatialAccelerationIdx, EParticleProperty::SpatialIdx> MSpatialIdx;
 
-	// Pointer to some arbitrary data associated with the particle, but not used by Chaos. External systems may use this for whatever.
-	TParticleProperty<void*, EParticleProperty::UserData> MUserData;
 
 public:
 	// Ryan: FGeometryCollectionPhysicsProxy needs access to GeometrySharedLowLevel(), 
@@ -1624,16 +1640,11 @@ public:
 	// when the geometry is being copied from GT to PT.
 	TSharedPtr<FImplicitObject, ESPMode::ThreadSafe> GeometrySharedLowLevel() const
 	{
-		return MGeometry.Read();
+		return MNonFrequentData.Read().Geometry;
 	}
 private:
 
-#if CHAOS_CHECKED
-	TParticleProperty<FName, EParticleProperty::DebugName> MDebugName;
-#endif
-
 protected:
-	FParticlePropertiesData* RemoteData;
 
 	EParticleType Type;
 	FParticleDirtyFlags MDirtyFlags;
@@ -1642,16 +1653,15 @@ protected:
 
 	void UpdateShapesArray()
 	{
-		UpdateShapesArrayFromGeometry<T, d>(MShapesArray, MakeSerializable(MGeometry.Read()), FRigidTransform3(X(), R()));
+		UpdateShapesArrayFromGeometry(MShapesArray, MakeSerializable(MNonFrequentData.Read().Geometry), FRigidTransform3(X(), R()), Proxy);
 		MapImplicitShapes();
 	}
 
-	void FlushToRemoteData()
+	virtual void SyncRemoteDataImp(FDirtyPropertiesManager& Manager, int32 DataIdx, const FParticleDirtyData& RemoteData) const
 	{
-		if(RemoteData && MDirtyFlags.IsDirty())
-		{
-			MX.InitialFlush(*RemoteData,MDirtyFlags,Proxy);
-		}
+		MXR.SyncRemote(Manager, DataIdx, RemoteData);
+		MNonFrequentData.SyncRemote(Manager, DataIdx, RemoteData);
+		MMiscData.SyncRemote(Manager, DataIdx, RemoteData);
 	}
 
 	void MapImplicitShapes();
@@ -1695,24 +1705,26 @@ public:
 		, DebugName(InParticle.DebugName())
 #endif
 	{
-		const TShapesArray<T, d>& Shapes = InParticle.ShapesArray();
+		const FShapesArray& Shapes = InParticle.ShapesArray();
 		ShapeCollisionDisableFlags.Empty(Shapes.Num());
 		CollisionTraceType.Empty(Shapes.Num());
 		ShapeSimData.Empty(Shapes.Num());
 		ShapeQueryData.Empty(Shapes.Num());
 		ShapeMaterials.Empty(Shapes.Num());
-		for (const TUniquePtr<TPerShapeData<T, d>>& ShapePtr : Shapes)
+		for (const TUniquePtr<FPerShapeData>& ShapePtr : Shapes)
 		{
-			ShapeCollisionDisableFlags.Add(ShapePtr->bDisable);
-			CollisionTraceType.Add(ShapePtr->CollisionTraceType);
-			ShapeSimData.Add(ShapePtr->SimData);
-			ShapeQueryData.Add(ShapePtr->QueryData);
-			ShapeMaterials.Add(ShapePtr->Materials);
+			ShapeCollisionDisableFlags.Add(ShapePtr->GetDisable());
+			CollisionTraceType.Add(ShapePtr->GetCollisionTraceType());
+			ShapeSimData.Add(ShapePtr->GetSimData());
+			ShapeQueryData.Add(ShapePtr->GetQueryData());
+			ShapeMaterials.Add(ShapePtr->GetMaterials());
 		}
 	}
 
+	//todo: remove
 	void Reset() 
 	{ 
+#if 0
 		FParticleData::Reset();  
 		X = TVector<T, d>(0); 
 		R = TRotation<T, d>(); 
@@ -1728,11 +1740,14 @@ public:
 		ShapeMaterials.Reset();
 #if CHAOS_CHECKED
 		DebugName = NAME_None;
+#endif
 #endif
 	}
 	
+	//todo: remove
 	void Init(const TGeometryParticle<T, d>& InParticle)
 		{
+#if 0
 			Type = EParticleType::Static;
 			X = InParticle.X();
 			R = InParticle.R();
@@ -1744,13 +1759,12 @@ public:
 #if CHAOS_CHECKED
 			DebugName = InParticle.DebugName();
 #endif
-			const TShapesArray<T, d>& Shapes = InParticle.ShapesArray();
+			const FShapesArray& Shapes = InParticle.ShapesArray();
 			ShapeCollisionDisableFlags.Empty(Shapes.Num());
 			CollisionTraceType.Empty(Shapes.Num());
 			ShapeSimData.Empty(Shapes.Num());
 			ShapeQueryData.Empty(Shapes.Num());
-			ShapeMaterials.Empty(Shapes.Num());
-			for (const TUniquePtr<TPerShapeData<T, d>>& ShapePtr : Shapes)
+			for (const TUniquePtr<FPerShapeData>& ShapePtr : Shapes)
 			{
 				ShapeCollisionDisableFlags.Add(ShapePtr->bDisable);
 				CollisionTraceType.Add(ShapePtr->CollisionTraceType);
@@ -1758,6 +1772,7 @@ public:
 				ShapeQueryData.Add(ShapePtr->QueryData);
 				ShapeMaterials.Add(ShapePtr->Materials);
 			}
+#endif
 	}
 	TVector<T, d> X;
 	TRotation<T, d> R;
@@ -1787,8 +1802,13 @@ public:
 
 	using TGeometryParticle<T, d>::Type;
 	using TGeometryParticle<T, d>::CastToRigidParticle;
+	using Base = TGeometryParticle<T,d>;
+	using Base::MDirtyFlags;
+	using Base::MMiscData;
 
 protected:
+	using Base::Proxy;
+
 	friend TGeometryParticle<T,d>* TGeometryParticle<T, d>::SerializationFactory(FChaosArchive& Ar, TGeometryParticle<T, d>* Serializable);
 	TKinematicGeometryParticle(const TKinematicGeometryParticleParameters<T, d>& KinematicParams = TKinematicGeometryParticleParameters<T,d>())
 		: TGeometryParticle<T, d>(KinematicParams)
@@ -1805,46 +1825,49 @@ public:
 	virtual void Serialize(FChaosArchive& Ar) override
 	{
 		TGeometryParticle<T, d>::Serialize(Ar);
-		Ar << MV;
-		Ar << MW;
+		Ar << MVelocities;
 	}
 
-	const TVector<T, d>& V() const { return MV.Read(); }
+	const TVector<T, d>& V() const { return MVelocities.Read().V; }
 	void SetV(const TVector<T, d>& InV, bool bInvalidate = true)
 	{
-		MV.Write(InV,bInvalidate,this->RemoteData,this->MDirtyFlags,this->Proxy);
+		MVelocities.Modify(bInvalidate,MDirtyFlags,Proxy,[&InV](auto& Data){ Data.V = InV;});
 	}
 
-	const TVector<T, d>& W() const { return MW.Read(); }
+	const TVector<T, d>& W() const { return MVelocities.Read().W; }
 	void SetW(const TVector<T, d>& InW, bool bInvalidate = true)
 	{
-		MW.Write(InW,bInvalidate,this->RemoteData,this->MDirtyFlags,this->Proxy);
+		MVelocities.Modify(bInvalidate,MDirtyFlags,Proxy,[&InW](auto& Data){ Data.W = InW;});
 	}
 
-	const TVector<T, d>& CenterOfMass() const { return MCenterOfMass.Read(); }
+	const TVector<T, d>& CenterOfMass() const { return MMassProps.Read().CenterOfMass; }
 	void SetCenterOfMass(const TVector<T, d>& InCenterOfMass, bool bInvalidate = true)
 	{
-		MCenterOfMass.Write(InCenterOfMass,bInvalidate,this->RemoteData,this->MDirtyFlags,this->Proxy);
+		MMassProps.Modify(bInvalidate,MDirtyFlags,Proxy,[&InCenterOfMass](auto& Data){ Data.CenterOfMass = InCenterOfMass;});
 	}
 	
-	const TRotation<T, d>& RotationOfMass() const { return MRotationOfMass.Read(); }
+	const TRotation<T, d>& RotationOfMass() const { return MMassProps.Read().RotationOfMass; }
 	void SetRotationOfMass(const TRotation<T, d>& InRotationOfMass, bool bInvalidate = true)
 	{
-		MRotationOfMass.Write(InRotationOfMass,bInvalidate,this->RemoteData,this->MDirtyFlags,this->Proxy);
+		MMassProps.Modify(bInvalidate,MDirtyFlags,Proxy,[&InRotationOfMass](auto& Data){ Data.RotationOfMass = InRotationOfMass;});
 	}
 
 	EObjectStateType ObjectState() const;
 
-	FParticleData* NewData() const
+private:
+	TParticleProperty<FParticleVelocities, EParticleProperty::Velocities> MVelocities;
+
+protected:
+	//todo: move this into pbdrigid
+	TParticleProperty<FParticleMassProps,EParticleProperty::MassProps> MMassProps;
+
+protected:
+	virtual void SyncRemoteDataImp(FDirtyPropertiesManager& Manager, int32 DataIdx, const FParticleDirtyData& RemoteData) const
 	{
-		return new TKinematicGeometryParticleData<T, d>(*this);
+		Base::SyncRemoteDataImp(Manager, DataIdx, RemoteData);
+		MVelocities.SyncRemote(Manager, DataIdx, RemoteData);
+		MMassProps.SyncRemote(Manager, DataIdx, RemoteData);
 	}
-
-private:
-	TParticleProperty<TVector<T, d>, EParticleProperty::V> MV;
-	TParticleProperty<TVector<T, d>, EParticleProperty::W> MW;
-	TParticleProperty<TVector<T, d>, EParticleProperty::CenterOfMass> MCenterOfMass;
-	TParticleProperty<TRotation<T, d>, EParticleProperty::RotationOfMass> MRotationOfMass;
 };
 
 template <typename T, int d>
@@ -1905,8 +1928,15 @@ public:
 	typedef TPBDRigidParticleHandle<T, d> FHandle;
 
 	using TGeometryParticle<T, d>::Type;
+	using Base = TKinematicGeometryParticle<T,d>;
+
+	using Base::MDirtyFlags;
+	using Base::MMiscData;
+	using Base::MMassProps;
+
 
 protected:
+	using Base::Proxy;
 	friend TGeometryParticle<T, d>* TGeometryParticle<T, d>::SerializationFactory(FChaosArchive& Ar, TGeometryParticle<T, d>* Serializable);
 	TPBDRigidParticle<T, d>(const TPBDRigidParticleParameters<T, d>& DynamicParams = TPBDRigidParticleParameters<T, d>())
 		: TKinematicGeometryParticle<T, d>(DynamicParams), MAwakeEvent(false)
@@ -1926,53 +1956,30 @@ public:
 	void Serialize(FChaosArchive& Ar) override
 	{
 		TKinematicGeometryParticle<T, d>::Serialize(Ar);
-		Ar << MQ;
-		Ar << MPreV;
-		Ar << MPreW;
-		Ar << MP;
-		Ar << MF;
-		Ar << MTorque;
-		Ar << MLinearImpulse;
-		Ar << MAngularImpulse;
-		Ar << MI;
-		Ar << MInvI;
-		Ar << MCollisionParticles;
-		Ar << MM;
-		Ar << MInvM;
-
-		Ar.UsingCustomVersion(FExternalPhysicsCustomObjectVersion::GUID);
-		if (Ar.CustomVer(FExternalPhysicsCustomObjectVersion::GUID) >= FExternalPhysicsCustomObjectVersion::AddDampingToRigids)
-		{
-			Ar << MLinearEtherDrag;
-			Ar << MAngularEtherDrag;
-		}
+		Ar << MDynamics;
 
 		Ar << MIsland;
-		Ar << MCollisionGroup;
-		Ar << MObjectState;
-		Ar << MDisabled;
 		Ar << MToBeRemovedOnFracture;
-		Ar << MGravityEnabled;
 	}
 
 	const TUniquePtr<TBVHParticles<T, d>>& CollisionParticles() const { return MCollisionParticles; }
 
-	int32 CollisionGroup() const { return MCollisionGroup.Read(); }
+	int32 CollisionGroup() const { return MMiscData.Read().CollisionGroup; }
 	void SetCollisionGroup(const int32 InCollisionGroup)
 	{
-		MCollisionGroup.Write(InCollisionGroup, /*bInvalidate=*/true,this->RemoteData,this->MDirtyFlags,this->Proxy);
+		MMiscData.Modify(true,MDirtyFlags,Proxy,[InCollisionGroup](auto& Data){ Data.CollisionGroup = InCollisionGroup;});
 	}
 
-	bool Disabled() const { return MDisabled.Read(); }
+	bool Disabled() const { return MMiscData.Read().bDisabled; }
 	void SetDisabled(const bool InDisabled)
 	{
-		MDisabled.Write(InDisabled, /*bInvalidate=*/true,this->RemoteData,this->MDirtyFlags,this->Proxy);
+		MMiscData.Modify(true,MDirtyFlags,Proxy,[InDisabled](auto& Data){ Data.bDisabled = InDisabled;});
 	}
 
-	bool IsGravityEnabled() const { return MGravityEnabled.Read(); }
+	bool IsGravityEnabled() const { return MMiscData.Read().bGravityEnabled; }
 	void SetGravityEnabled(const bool InGravityEnabled)
 	{
-		MGravityEnabled.Write(InGravityEnabled, /*bInvalidate=*/true,this->RemoteData,this->MDirtyFlags,this->Proxy);
+		MMiscData.Modify(true,MDirtyFlags,Proxy,[InGravityEnabled](auto& Data){ Data.bGravityEnabled = InGravityEnabled;});
 	}
 
 	//todo: remove this
@@ -1982,97 +1989,64 @@ public:
 		this->MInitialized = InInitialized;
 	}
 
-	const TVector<T, d>& PreV() const { return MPreV.Read(); }
-	void SetPreV(const TVector<T, d>& InPreV)
-	{
-		MPreV.Write(InPreV, /*bInvalidate=*/true,this->RemoteData,this->MDirtyFlags,this->Proxy);
-	}
-
-	const TVector<T, d>& PreW() const { return MPreW.Read(); }
-	void SetPreW(const TVector<T, d>& InPreW)
-	{
-		MPreW.Write(InPreW, /*bInvalidate=*/true,this->RemoteData,this->MDirtyFlags,this->Proxy);
-	}
-
-	const TVector<T, d>& P() const { return MP.Read(); }
-	void SetP(const TVector<T, d>& InP)
-	{
-		MP.Write(InP, /*bInvalidate=*/true,this->RemoteData,this->MDirtyFlags,this->Proxy);
-	}
-
-	const TRotation<T, d>& Q() const { return MQ.Read(); }
-	void SetQ(const TRotation<T, d>& InQ)
-	{
-		MQ.Write(InQ, /*bInvalidate=*/true,this->RemoteData,this->MDirtyFlags,this->Proxy);
-	}
-
-	const TVector<T, d>& F() const { return MF.Read(); }
-	void SetF(const TVector<T, d>& InF)
+	const TVector<T, d>& F() const { return MDynamics.Read().F; }
+	void SetF(const TVector<T, d>& InF, bool bInvalidate = true)
 	{
-		//question: should we do this check? only adding because we clear forces after removing from dirty list, but this marks dirty
-		if(InF != MF.Read())
-		{
-			MF.Write(InF, /*bInvalidate=*/true,this->RemoteData,this->MDirtyFlags,this->Proxy);
-		}
-		
+		MDynamics.Modify(bInvalidate,MDirtyFlags,Proxy,[&InF](auto& Data){ Data.F = InF;});
 	}
 
-	const TVector<T, d>& Torque() const { return MTorque.Read(); }
-	void SetTorque(const TVector<T, d>& InTorque)
+	const TVector<T, d>& Torque() const { return MDynamics.Read().Torque; }
+	void SetTorque(const TVector<T, d>& InTorque, bool bInvalidate=true)
 	{
-		//question: should we do this check? only adding because we clear forces after removing from dirty list, but this marks dirty
-		if(InTorque != MTorque.Read())
-		{
-			MTorque.Write(InTorque, /*bInvalidate=*/true,this->RemoteData,this->MDirtyFlags,this->Proxy);
-		}
+		MDynamics.Modify(bInvalidate,MDirtyFlags,Proxy,[&InTorque](auto& Data){ Data.Torque = InTorque;});
 	}
 
-	const TVector<T, d>& LinearImpulse() const { return MLinearImpulse.Read(); }
+	const TVector<T, d>& LinearImpulse() const { return MDynamics.Read().LinearImpulse; }
 	void SetLinearImpulse(const TVector<T, d>& InLinearImpulse, bool bInvalidate = true)
 	{
-		MLinearImpulse.Write(InLinearImpulse, bInvalidate,this->RemoteData,this->MDirtyFlags,this->Proxy);
+		MDynamics.Modify(bInvalidate,MDirtyFlags,Proxy,[&InLinearImpulse](auto& Data){ Data.LinearImpulse = InLinearImpulse;});
 	}
 
-	const TVector<T, d>& AngularImpulse() const { return MAngularImpulse.Read(); }
+	const TVector<T, d>& AngularImpulse() const { return MDynamics.Read().AngularImpulse; }
 	void SetAngularImpulse(const TVector<T, d>& InAngularImpulse, bool bInvalidate = true)
 	{
-		MAngularImpulse.Write(InAngularImpulse, bInvalidate,this->RemoteData,this->MDirtyFlags,this->Proxy);
+		MDynamics.Modify(bInvalidate,MDirtyFlags,Proxy,[&InAngularImpulse](auto& Data){ Data.AngularImpulse = InAngularImpulse;});
 	}
 
-	const PMatrix<T, d, d>& I() const { return MI.Read(); }
+	const PMatrix<T, d, d>& I() const { return MMassProps.Read().I; }
 	void SetI(const PMatrix<T, d, d>& InI)
 	{
-		MI.Write(InI, /*bInvalidate=*/true,this->RemoteData,this->MDirtyFlags,this->Proxy);
+		MMassProps.Modify(true,MDirtyFlags,Proxy,[&InI](auto& Data){ Data.I = InI;});
 	}
 
-	const PMatrix<T, d, d>& InvI() const { return MInvI.Read(); }
+	const PMatrix<T, d, d>& InvI() const { return MMassProps.Read().InvI; }
 	void SetInvI(const PMatrix<T, d, d>& InInvI)
 	{
-		MInvI.Write(InInvI, /*bInvalidate=*/true,this->RemoteData,this->MDirtyFlags,this->Proxy);
+		MMassProps.Modify(true,MDirtyFlags,Proxy,[&InInvI](auto& Data){ Data.InvI = InInvI;});
 	}
 
-	T M() const { return MM.Read(); }
+	T M() const { return MMassProps.Read().M; }
 	void SetM(const T& InM)
 	{
-		MM.Write(InM, /*bInvalidate=*/true,this->RemoteData,this->MDirtyFlags,this->Proxy);
+		MMassProps.Modify(true,MDirtyFlags,Proxy,[InM](auto& Data){ Data.M = InM;});
 	}
 
-	T InvM() const { return MInvM.Read(); }
+	T InvM() const { return MMassProps.Read().InvM; }
 	void SetInvM(const T& InInvM)
 	{
-		MInvM.Write(InInvM, /*bInvalidate=*/true,this->RemoteData,this->MDirtyFlags,this->Proxy);
+		MMassProps.Modify(true,MDirtyFlags,Proxy,[InInvM](auto& Data){ Data.InvM = InInvM;});
 	}
 
-	T LinearEtherDrag() const { return MLinearEtherDrag.Read(); }
+	T LinearEtherDrag() const { return MDynamics.Read().LinearEtherDrag; }
 	void SetLinearEtherDrag(const T& InLinearEtherDrag)
 	{
-		MLinearEtherDrag.Write(InLinearEtherDrag, /*bInvalidate=*/true,this->RemoteData,this->MDirtyFlags,this->Proxy);
+		MDynamics.Modify(true,MDirtyFlags,Proxy,[&InLinearEtherDrag](auto& Data){ Data.LinearEtherDrag = InLinearEtherDrag;});
 	}
 
-	T AngularEtherDrag() const { return MAngularEtherDrag.Read(); }
+	T AngularEtherDrag() const { return MDynamics.Read().AngularEtherDrag; }
 	void SetAngularEtherDrag(const T& InAngularEtherDrag)
 	{
-		MAngularEtherDrag.Write(InAngularEtherDrag, /*bInvalidate=*/true,this->RemoteData,this->MDirtyFlags,this->Proxy);
+		MDynamics.Modify(true,MDirtyFlags,Proxy,[&InAngularEtherDrag](auto& Data){ Data.AngularEtherDrag = InAngularEtherDrag;});
 	}
 
 	int32 Island() const { return MIsland; }
@@ -2089,49 +2063,33 @@ public:
 		this->MToBeRemovedOnFracture = bToBeRemovedOnFracture;
 	}
 
-	EObjectStateType ObjectState() const { return MObjectState.Read(); }
+	EObjectStateType ObjectState() const { return MMiscData.Read().ObjectState; }
 	void SetObjectState(const EObjectStateType InState, bool bAllowEvents = false)
 	{
-		if (bAllowEvents && MObjectState.Read() != EObjectStateType::Dynamic && InState == EObjectStateType::Dynamic)
+		if (bAllowEvents && ObjectState() != EObjectStateType::Dynamic && InState == EObjectStateType::Dynamic)
 		{
 			MAwakeEvent |= true;
 		}
-		MObjectState.Write(InState, /*bInvalidate=*/true,this->RemoteData,this->MDirtyFlags,this->Proxy);
+		MMiscData.Modify(true,MDirtyFlags,Proxy,[&InState](auto& Data){ Data.ObjectState = InState;});
 	}
 
 	void ClearEvents() { MAwakeEvent = false; }
 	bool HasAwakeEvent() { return MAwakeEvent; }
 
-	FParticleData* NewData() const
-	{
-		return new TPBDRigidParticleData<T, d>(*this);
-	}
-
-
 private:
-	TParticleProperty<TRotation<T, d>, EParticleProperty::Q> MQ;
-	TParticleProperty<TVector<T, d>, EParticleProperty::PreV> MPreV;
-	TParticleProperty<TVector<T, d>, EParticleProperty::PreW> MPreW;
-	TParticleProperty<TVector<T, d>, EParticleProperty::P> MP;
-	TParticleProperty<TVector<T, d>, EParticleProperty::F> MF;
-	TParticleProperty<TVector<T, d>, EParticleProperty::Torque> MTorque;
-	TParticleProperty<TVector<T, d>, EParticleProperty::LinearImpulse> MLinearImpulse;
-	TParticleProperty<TVector<T, d>, EParticleProperty::AngularImpulse> MAngularImpulse;
-	TParticleProperty<PMatrix<T, d, d>, EParticleProperty::I> MI;
-	TParticleProperty<PMatrix<T, d, d>, EParticleProperty::InvI> MInvI;
+	TParticleProperty<FParticleDynamics, EParticleProperty::Dynamics> MDynamics;
 	TUniquePtr<TBVHParticles<T, d>> MCollisionParticles;
-	TParticleProperty<T, EParticleProperty::M> MM;
-	TParticleProperty<T, EParticleProperty::InvM> MInvM;
-	TParticleProperty<T, EParticleProperty::LinearEtherDrag> MLinearEtherDrag;
-	TParticleProperty<T, EParticleProperty::AngularEtherDrag> MAngularEtherDrag;
 	int32 MIsland;
-	TParticleProperty<int32, EParticleProperty::CollisionGroup> MCollisionGroup;
-	TParticleProperty<EObjectStateType, EParticleProperty::ObjectState> MObjectState;
-	TParticleProperty<bool, EParticleProperty::Disabled> MDisabled;
 	bool MToBeRemovedOnFracture;
-	TParticleProperty<bool, EParticleProperty::GravityEnabled> MGravityEnabled;
 	bool MInitialized;
 	bool MAwakeEvent;
+
+protected:
+	virtual void SyncRemoteDataImp(FDirtyPropertiesManager& Manager, int32 DataIdx, const FParticleDirtyData& RemoteData) const
+	{
+		Base::SyncRemoteDataImp(Manager,DataIdx,RemoteData);
+		MDynamics.SyncRemote(Manager,DataIdx,RemoteData);
+	}
 };
 
 template <typename T, int d>
@@ -2165,10 +2123,6 @@ public:
 
 	TPBDRigidParticleData(EParticleType InType = EParticleType::Rigid)
 		: Base(InType)
-		, MQ(TRotation<T, d>())
-		, MPreV(TVector<T, d>(0))
-		, MPreW(TVector<T, d>(0))
-		, MP(TVector<T, d>(0))
 		, MF(TVector<T, d>(0))
 		, MTorque(TVector<T, d>(0))
 		, MLinearImpulse(TVector<T, d>(0))
@@ -2191,10 +2145,6 @@ public:
 
 	TPBDRigidParticleData(const TPBDRigidParticle<T, d>& InParticle)
 		: Base(InParticle)
-		, MQ(InParticle.Q())
-		, MPreV(InParticle.PreV())
-		, MPreW(InParticle.PreW())
-		, MP(InParticle.P())
 		, MF(InParticle.F())
 		, MTorque(InParticle.Torque())
 		, MLinearImpulse(InParticle.LinearImpulse())
@@ -2218,10 +2168,6 @@ public:
 	}
 
 
-	TRotation<T, d> MQ;
-	TVector<T, d> MPreV;
-	TVector<T, d> MPreW;
-	TVector<T, d> MP;
 	TVector<T, d> MF;
 	TVector<T, d> MTorque;
 	TVector<T, d> MLinearImpulse;
@@ -2244,10 +2190,6 @@ public:
 	void Reset() {
 		TKinematicGeometryParticleData<T, d>::Reset();
 		Type = EParticleType::Rigid;
-		MQ = TRotation<T, d>();
-		MPreV = TVector<T, d>(0);
-		MPreW = TVector<T, d>(0);
-		MP = TVector<T, d>(0);
 		MF = TVector<T, d>(0);
 		MTorque = TVector<T, d>(0);
 		MLinearImpulse = TVector<T, d>(0);
@@ -2269,10 +2211,6 @@ public:
 	}
 	void Init(const TPBDRigidParticle<T, d>& InParticle) {
 			Base::Init(InParticle);
-			MQ = InParticle.Q();
-			MPreV = InParticle.PreV();
-			MPreW = InParticle.PreW();
-			MP = InParticle.P();
 			MF = InParticle.F();
 			MTorque = InParticle.Torque();
 			MLinearImpulse = InParticle.LinearImpulse();
@@ -2433,10 +2371,10 @@ void TAccelerationStructureHandle<T, d>::UpdatePrePreFilter(const TParticle& Par
 	const auto& Shapes = Particle.ShapesArray();
 	for (const auto& Shape : Shapes)
 	{
-		UnionFilterData.Word0 |= Shape->QueryData.Word0;
-		UnionFilterData.Word1 |= Shape->QueryData.Word1;
-		UnionFilterData.Word2 |= Shape->QueryData.Word2;
-		UnionFilterData.Word3 |= Shape->QueryData.Word3;
+		UnionFilterData.Word0 |= Shape->GetQueryData().Word0;
+		UnionFilterData.Word1 |= Shape->GetQueryData().Word1;
+		UnionFilterData.Word2 |= Shape->GetQueryData().Word2;
+		UnionFilterData.Word3 |= Shape->GetQueryData().Word3;
 	}
 	
 	bCanPrePreFilter = true;
diff --git a/Engine/Source/Runtime/Experimental/Chaos/Public/Chaos/ParticleProperties.inl b/Engine/Source/Runtime/Experimental/Chaos/Public/Chaos/ParticleProperties.inl
--- a/Engine/Source/Runtime/Experimental/Chaos/Public/Chaos/ParticleProperties.inl
+++ b/Engine/Source/Runtime/Experimental/Chaos/Public/Chaos/ParticleProperties.inl
@@ -7,35 +7,9 @@
 #define PARTICLE_PROPERTY_CHECKED(x, Type)
 #endif
 
-PARTICLE_PROPERTY(X,FVec3)
-PARTICLE_PROPERTY(R,FRotation3)
-PARTICLE_PROPERTY(V,FVec3)
-PARTICLE_PROPERTY(W,FVec3)
-PARTICLE_PROPERTY(CenterOfMass,FVec3)
-PARTICLE_PROPERTY(RotationOfMass,FRotation3)
-PARTICLE_PROPERTY(CollisionGroup,int32)
-PARTICLE_PROPERTY(Disabled,bool)
-PARTICLE_PROPERTY(PreV,FVec3)
-PARTICLE_PROPERTY(PreW,FVec3)
-PARTICLE_PROPERTY(P,FVec3)
-PARTICLE_PROPERTY(Q,FRotation3)
-PARTICLE_PROPERTY(F,FVec3)
-PARTICLE_PROPERTY(Torque,FVec3)
-PARTICLE_PROPERTY(I,FMatrix33)
-PARTICLE_PROPERTY(InvI,FMatrix33)
-PARTICLE_PROPERTY(M,FReal)
-PARTICLE_PROPERTY(InvM,FReal)
-PARTICLE_PROPERTY(LinearEtherDrag,FReal)
-PARTICLE_PROPERTY(AngularEtherDrag,FReal)
-PARTICLE_PROPERTY(ObjectState,EObjectStateType)
-PARTICLE_PROPERTY(Geometry, ChaosGeometrySp)
-PARTICLE_PROPERTY(LinearImpulse, FVec3)
-PARTICLE_PROPERTY(AngularImpulse, FVec3)
-PARTICLE_PROPERTY(GravityEnabled, bool)
-PARTICLE_PROPERTY(SpatialIdx, FSpatialAccelerationIdx)
-PARTICLE_PROPERTY(UniqueIdx, FUniqueIdx)
-PARTICLE_PROPERTY(ShapeDisableCollision, int32)
-PARTICLE_PROPERTY(CollisionTraceType, EChaosCollisionTraceFlag)
-PARTICLE_PROPERTY(ShapeSimData, TArray<FCollisionFilterData>)
-PARTICLE_PROPERTY(UserData, void*)
-PARTICLE_PROPERTY_CHECKED(DebugName, FName)
\ No newline at end of file
+PARTICLE_PROPERTY(XR, FParticlePositionRotation)
+PARTICLE_PROPERTY(Velocities, FParticleVelocities)
+PARTICLE_PROPERTY(Dynamics, FParticleDynamics)
+PARTICLE_PROPERTY(Misc, FParticleMisc)
+PARTICLE_PROPERTY(NonFrequentData, FParticleNonFrequentData)
+PARTICLE_PROPERTY(MassProps, FParticleMassProps)
\ No newline at end of file
diff --git a/Engine/Source/Runtime/Experimental/Chaos/Public/Chaos/ParticlePropertiesTypes.inl b/Engine/Source/Runtime/Experimental/Chaos/Public/Chaos/ParticlePropertiesTypes.inl
deleted file mode 100644
diff --git a/Engine/Source/Runtime/Experimental/Chaos/Public/Chaos/Particles.h b/Engine/Source/Runtime/Experimental/Chaos/Public/Chaos/Particles.h
--- a/Engine/Source/Runtime/Experimental/Chaos/Public/Chaos/Particles.h
+++ b/Engine/Source/Runtime/Experimental/Chaos/Public/Chaos/Particles.h
@@ -201,4 +201,28 @@ namespace Chaos
 		return InParticles.GetTypeHash();
 	}
 	
+	enum class EObjectStateType: int8
+	{
+		Uninitialized = 0,
+		Sleeping = 1,
+		Kinematic = 2,
+		Static = 3,
+		Dynamic = 4,
+
+		Count
+	};
+
+	enum EChaosCollisionTraceFlag
+	{
+		/** Use project physics settings (DefaultShapeComplexity) */
+		Chaos_CTF_UseDefault,
+		/** Create both simple and complex shapes. Simple shapes are used for regular scene queries and collision tests. Complex shape (per poly) is used for complex scene queries.*/
+		Chaos_CTF_UseSimpleAndComplex,
+		/** Create only simple shapes. Use simple shapes for all scene queries and collision tests.*/
+		Chaos_CTF_UseSimpleAsComplex,
+		/** Create only complex shapes (per poly). Use complex shapes for all scene queries and collision tests. Can be used in simulation for static shapes only (i.e can be collided against but not moved through forces or velocity.) */
+		Chaos_CTF_UseComplexAsSimple,
+		/** */
+		Chaos_CTF_MAX,
+	};
 } // namespace Chaos
diff --git a/Engine/Source/Runtime/Experimental/Chaos/Public/Chaos/Properties.h b/Engine/Source/Runtime/Experimental/Chaos/Public/Chaos/Properties.h
new file mode 100644
--- /dev/null
+++ b/Engine/Source/Runtime/Experimental/Chaos/Public/Chaos/Properties.h
@@ -0,0 +1,150 @@
+// Copyright Epic Games, Inc. All Rights Reserved.
+
+#pragma once
+
+#include "Chaos/ParticleDirtyFlags.h"
+#include "Framework/PhysicsSolverBase.h"
+
+namespace Chaos
+{
+
+template <typename T, EParticleProperty PropName>
+class TParticleProperty
+{
+public:
+	TParticleProperty() = default;
+	TParticleProperty(const T& Val)
+		: Property(Val)
+	{
+	}
+
+	//we don't support this because it could lead to bugs with values not being properly written to remote
+	TParticleProperty(const TParticleProperty<T,PropName>& Rhs) = delete;
+
+	const T& Read() const { return Property; }
+	void Write(const T& Val, bool bInvalidate, FParticleDirtyFlags& Dirty, IPhysicsProxyBase* Proxy)
+	{
+		Property = Val;
+		MarkDirty(bInvalidate, Dirty, Proxy);
+	}
+
+	template <typename Lambda>
+	void Modify(bool bInvalidate, FParticleDirtyFlags& Dirty,IPhysicsProxyBase* Proxy,const Lambda& LambdaFunc)
+	{
+		LambdaFunc(Property);
+		MarkDirty(bInvalidate, Dirty, Proxy);
+	}
+
+	void SyncRemote(FDirtyPropertiesManager& Manager, int32 DataIdx, const FParticleDirtyData& Remote) const
+	{
+		Remote.SyncRemote<T,PropName>(Manager, DataIdx, Property);
+	}
+
+	void Serialize(FChaosArchive& Ar)
+	{
+		Ar << Property;
+	}
+
+private:
+	T Property;
+	static constexpr EParticleFlags PropertyFlag = ParticlePropToFlag(PropName);
+
+	void MarkDirty(bool bInvalidate, FParticleDirtyFlags& Dirty, IPhysicsProxyBase* Proxy)
+	{
+		if(bInvalidate)
+		{
+			Dirty.MarkDirty(PropertyFlag);
+			
+			if(Proxy)
+			{
+				if(FPhysicsSolverBase* PhysicsSolverBase = Proxy->GetSolver<FPhysicsSolverBase>())
+				{
+					PhysicsSolverBase->AddDirtyProxy(Proxy);
+				}
+			}
+		}
+	}
+
+	//we don't support this because it could lead to bugs with values not being properly written to remote
+	TParticleProperty<T,PropName>& operator=(const TParticleProperty<T,PropName>& Rhs) {}
+};
+
+template <typename T, EParticleProperty PropName>
+FChaosArchive& operator<<(FChaosArchive& Ar, TParticleProperty<T, PropName>& Prop)
+{
+	//TODO: should this only work with dirty flag? Not sure if this path really matters at this point
+	Prop.Serialize(Ar);
+	return Ar;
+}
+
+
+template <typename T,EShapeProperty PropName>
+class TShapeProperty
+{
+public:
+	TShapeProperty() = default;
+	TShapeProperty(const T& Val)
+		: Property(Val)
+	{
+	}
+
+	//we don't support this because it could lead to bugs with values not being properly written to remote
+	TShapeProperty(const TShapeProperty<T,PropName>& Rhs) = delete;
+
+	const T& Read() const { return Property; }
+	void Write(const T& Val,bool bInvalidate,FShapeDirtyFlags& Dirty,IPhysicsProxyBase* Proxy, int32 ShapeIdx)
+	{
+		Property = Val;
+		MarkDirty(bInvalidate,Dirty,Proxy, ShapeIdx);
+	}
+
+	template <typename Lambda>
+	void Modify(bool bInvalidate,FShapeDirtyFlags& Dirty,IPhysicsProxyBase* Proxy, int32 ShapeIdx, const Lambda& LambdaFunc)
+	{
+		LambdaFunc(Property);
+		MarkDirty(bInvalidate,Dirty,Proxy, ShapeIdx);
+	}
+
+	void SyncRemote(FDirtyPropertiesManager& Manager, int32 DataIdx, FShapeDirtyData& Remote) const
+	{
+		Remote.SyncRemote<T,PropName>(Manager, DataIdx, Property);
+	}
+
+	void Serialize(FChaosArchive& Ar)
+	{
+		Ar << Property;
+	}
+
+private:
+	T Property;
+	static constexpr EShapeFlags PropertyFlag = ShapePropToFlag(PropName);
+
+	void MarkDirty(bool bInvalidate, FShapeDirtyFlags& Dirty,IPhysicsProxyBase* Proxy, int32 ShapeIdx)
+	{
+		if(bInvalidate)
+		{
+			const bool bFirstDirty = Dirty.IsClean();
+			Dirty.MarkDirty(PropertyFlag);
+
+			if(bFirstDirty && Proxy)
+			{
+				if(FPhysicsSolverBase* PhysicsSolverBase = Proxy->GetSolver<FPhysicsSolverBase>())
+				{
+					PhysicsSolverBase->AddDirtyProxyShape(Proxy, ShapeIdx);
+				}
+			}
+		}
+	}
+
+	//we don't support this because it could lead to bugs with values not being properly written to remote
+	TShapeProperty<T,PropName>& operator=(const TShapeProperty<T,PropName>& Rhs) {}
+};
+
+template <typename T,EShapeProperty PropName>
+FChaosArchive& operator<<(FChaosArchive& Ar,TShapeProperty<T,PropName>& Prop)
+{
+	//TODO: should this use dirty flags? not sure if this path even matters
+	Prop.Serialize(Ar);
+	return Ar;
+}
+}
\ No newline at end of file
diff --git a/Engine/Source/Runtime/Experimental/Chaos/Public/Chaos/PropertiesTypes.inl b/Engine/Source/Runtime/Experimental/Chaos/Public/Chaos/PropertiesTypes.inl
new file mode 100644
--- /dev/null
+++ b/Engine/Source/Runtime/Experimental/Chaos/Public/Chaos/PropertiesTypes.inl
@@ -0,0 +1,10 @@
+// Copyright Epic Games, Inc. All Rights Reserved.
+
+PROPERTY_TYPE(FParticlePositionRotation_type, FParticlePositionRotation)
+PROPERTY_TYPE(FParticleVelocities_type, FParticleVelocities)
+PROPERTY_TYPE(FParticleDynamics_type, FParticleDynamics)
+PROPERTY_TYPE(FParticleMisc_type, FParticleMisc)
+PROPERTY_TYPE(FParticleNonFrequentData_type, FParticleNonFrequentData)
+PROPERTY_TYPE(FCollisionData_type, FCollisionData)
+PROPERTY_TYPE(FMaterialData_type, FMaterialData)
+PROPERTY_TYPE(FParticleMassProps_type, FParticleMassProps)
\ No newline at end of file
diff --git a/Engine/Source/Runtime/Experimental/Chaos/Public/Chaos/RigidParticles.h b/Engine/Source/Runtime/Experimental/Chaos/Public/Chaos/RigidParticles.h
--- a/Engine/Source/Runtime/Experimental/Chaos/Public/Chaos/RigidParticles.h
+++ b/Engine/Source/Runtime/Experimental/Chaos/Public/Chaos/RigidParticles.h
@@ -14,16 +14,6 @@
 
 namespace Chaos
 {
-enum class EObjectStateType : int8
-{
-	Uninitialized = 0,
-	Sleeping = 1,
-	Kinematic = 2,
-	Static = 3,
-	Dynamic = 4,
-
-	Count
-};
 
 template<class T, int d>
 struct TSleepData
diff --git a/Engine/Source/Runtime/Experimental/Chaos/Public/Chaos/ShapeProperties.inl b/Engine/Source/Runtime/Experimental/Chaos/Public/Chaos/ShapeProperties.inl
new file mode 100644
--- /dev/null
+++ b/Engine/Source/Runtime/Experimental/Chaos/Public/Chaos/ShapeProperties.inl
@@ -0,0 +1,4 @@
+// Copyright Epic Games, Inc. All Rights Reserved.
+
+SHAPE_PROPERTY(CollisionData,FCollisionData)
+SHAPE_PROPERTY(Materials,FMaterialData)
\ No newline at end of file
diff --git a/Engine/Source/Runtime/Experimental/ChaosSolverEngine/Private/Chaos/ChaosGameplayEventDispatcher.cpp b/Engine/Source/Runtime/Experimental/ChaosSolverEngine/Private/Chaos/ChaosGameplayEventDispatcher.cpp
--- a/Engine/Source/Runtime/Experimental/ChaosSolverEngine/Private/Chaos/ChaosGameplayEventDispatcher.cpp
+++ b/Engine/Source/Runtime/Experimental/ChaosSolverEngine/Private/Chaos/ChaosGameplayEventDispatcher.cpp
@@ -388,9 +388,9 @@ void UChaosGameplayEventDispatcher::HandleBreakingEvents(const Chaos::FBreakingE
 		{
 			for (Chaos::TBreakingData<float, 3> const& BreakingDataItem : BreakingData)
 			{	
-				if (BreakingDataItem.Particle && (BreakingDataItem.Particle->Proxy))
+				if (BreakingDataItem.Particle && (BreakingDataItem.Particle->GetProxy()))
 				{
-					UPrimitiveComponent* const PrimComp = Cast<UPrimitiveComponent>(BreakingDataItem.Particle->Proxy->GetOwner());
+					UPrimitiveComponent* const PrimComp = Cast<UPrimitiveComponent>(BreakingDataItem.Particle->GetProxy()->GetOwner());
 					if (PrimComp && BreakEventRegistrations.Contains(PrimComp))
 					{
 						// queue them up so we can release the physics data before trigging BP events
@@ -417,7 +417,7 @@ void UChaosGameplayEventDispatcher::HandleSleepingEvents(const Chaos::FSleepingE
 
 	for (const Chaos::TSleepingData<float, 3>& SleepData : SleepingArray)
 	{
-		if (SleepData.Particle->Proxy != nullptr)
+		if (SleepData.Particle->GetProxy()!= nullptr)
 		{
 			if (FBodyInstance* BodyInstance = FPhysicsUserData::Get<FBodyInstance>(SleepData.Particle->UserData()))
 			{
diff --git a/Engine/Source/Runtime/Experimental/ChaosSolvers/Private/EventDefaults.cpp b/Engine/Source/Runtime/Experimental/ChaosSolvers/Private/EventDefaults.cpp
--- a/Engine/Source/Runtime/Experimental/ChaosSolvers/Private/EventDefaults.cpp
+++ b/Engine/Source/Runtime/Experimental/ChaosSolvers/Private/EventDefaults.cpp
@@ -330,7 +330,7 @@ namespace Chaos
 				{
 					// Since Clustered GCs can be unioned the particleIndex representing the union 
 					// is not associated with a PhysicsProxy
-					if(AllBreakingsArray[Idx].Particle->Proxy != nullptr)
+					if(AllBreakingsArray[Idx].Particle->GetProxy() != nullptr)
 					{
 						if(ensure(!AllBreakingsArray[Idx].Location.ContainsNaN() &&
 							!Particles.V(AllBreakingsArray[Idx].ParticleIndex).ContainsNaN() &&
@@ -365,7 +365,7 @@ namespace Chaos
 								}
 
 								// Add to AllBreakingsIndicesByPhysicsProxy
-								IPhysicsProxyBase* PhysicsProxy = AllBreakingsArray[Idx].Particle->Proxy;
+								IPhysicsProxyBase* PhysicsProxy = AllBreakingsArray[Idx].Particle->GetProxy();
 								AllBreakingIndicesByPhysicsProxy.FindOrAdd(PhysicsProxy).Add(NewIdx);
 							}
 						}
diff --git a/Engine/Source/Runtime/Experimental/ChaosSolvers/Private/PBDRigidsSolver.cpp b/Engine/Source/Runtime/Experimental/ChaosSolvers/Private/PBDRigidsSolver.cpp
--- a/Engine/Source/Runtime/Experimental/ChaosSolvers/Private/PBDRigidsSolver.cpp
+++ b/Engine/Source/Runtime/Experimental/ChaosSolvers/Private/PBDRigidsSolver.cpp
@@ -252,10 +252,11 @@ namespace Chaos
 		// NOTE: Do we really need these lists of proxies if we can just
 		// access them through the GTParticles list?
 		
-		Chaos::FParticleData* ProxyData;
 		IPhysicsProxyBase* ProxyBase;
 
 		GTParticle->SetUniqueIdx(GetEvolution()->GenerateUniqueIdx());
+		//Chaos::FParticlePropertiesData& RemoteParticleData = *DirtyPropertiesManager->AccessProducerBuffer()->NewRemoteParticleProperties();
+		//Chaos::FShapeRemoteDataContainer& RemoteShapeContainer = *DirtyPropertiesManager->AccessProducerBuffer()->NewRemoteShapeContainer();
 
 		// Make a physics proxy, giving it our particle and particle handle
 		const EParticleType InParticleType = GTParticle->ObjectType();
@@ -263,74 +264,30 @@ namespace Chaos
 		{
 			auto Proxy = new FRigidParticlePhysicsProxy(GTParticle->CastToRigidParticle(), nullptr);
 			RigidParticlePhysicsProxies.Add((FRigidParticlePhysicsProxy*)Proxy);
-			ProxyData = Proxy->NewData();
 			ProxyBase = Proxy;
+			//Proxy->SetRemoteDataForGT(RemoteParticleData, RemoteShapeContainer);
 		}
 		else if (InParticleType == EParticleType::Kinematic)
 		{
 			auto Proxy = new FKinematicGeometryParticlePhysicsProxy(GTParticle->CastToKinematicParticle(), nullptr);
 			KinematicGeometryParticlePhysicsProxies.Add((FKinematicGeometryParticlePhysicsProxy*)Proxy);
-			ProxyData = Proxy->NewData();
 			ProxyBase = Proxy;
+			//Proxy->SetRemoteDataForGT(RemoteParticleData, RemoteShapeContainer);
 		}
 		else // Assume it's a static (geometry) if it's not dynamic or kinematic
 		{
 			auto Proxy = new FGeometryParticlePhysicsProxy(GTParticle, nullptr);
 			GeometryParticlePhysicsProxies.Add((FGeometryParticlePhysicsProxy*)Proxy);
-			ProxyData = Proxy->NewData();
 			ProxyBase = Proxy;
+			//Proxy->SetRemoteDataForGT(RemoteParticleData, RemoteShapeContainer);
 		}
 
 		ProxyBase->SetSolver(this);
 
 		// Associate the proxy with the particle
-		GTParticle->Proxy = ProxyBase;
+		GTParticle->SetProxy(ProxyBase);
 
-		//enqueue onto physics thread for finalizing registration
-		FChaosSolversModule::GetModule()->GetDispatcher()->EnqueueCommandImmediate(this, [GTParticle, InParticleType, ProxyBase, ProxyData](FPBDRigidsSolver* Solver)
-		{
-			UE_LOG(LogPBDRigidsSolver, Verbose, TEXT("FPBDRigidsSolver::RegisterObject() ~ Dequeue"));
-
-			// Create a handle for the new particle
-			TGeometryParticleHandle<float, 3>* Handle;
-
-			// Insert the particle ptr into the auxiliary array
-			//
-			// TODO: Is this array even necessary?
-			//       Proxy should be able to map back to Particle.
-			//
-
-			const auto& GeomData = static_cast<const Chaos::TGeometryParticleData<FReal,3>*>(ProxyData);
-
-			if (InParticleType == EParticleType::Rigid)
-			{
-				Handle = Solver->Particles.CreateDynamicParticles(1, &GeomData->UniqueIdx)[0];
-				auto Proxy = static_cast<FRigidParticlePhysicsProxy*>(ProxyBase);
-				Proxy->SetHandle(Handle->CastToRigidParticle());
-				Proxy->PushToPhysicsState(ProxyData);
-			}
-			else if (InParticleType == EParticleType::Kinematic)
-			{
-				Handle = Solver->Particles.CreateKinematicParticles(1, &GeomData->UniqueIdx)[0];
-				auto Proxy = static_cast<FKinematicGeometryParticlePhysicsProxy*>(ProxyBase);
-				Proxy->SetHandle(Handle->CastToKinematicParticle());
-				Proxy->PushToPhysicsState(ProxyData);
-			}
-			else // Assume it's a static (geometry) if it's not dynamic or kinematic
-			{
-				Handle = Solver->Particles.CreateStaticParticles(1, &GeomData->UniqueIdx)[0];
-
-				auto Proxy = static_cast<FGeometryParticlePhysicsProxy*>(ProxyBase);
-				Proxy->SetHandle(Handle);
-				Proxy->PushToPhysicsState(ProxyData);
-			}
-
-			Handle->GTGeometryParticle() = GTParticle;
-			Solver->MParticleToProxy.Add(Handle, ProxyBase);
-			Solver->GetEvolution()->CreateParticle(Handle);
-
-			delete ProxyData;
-		});
+		AddDirtyProxy(ProxyBase);
 	}
 
 	void FPBDRigidsSolver::UnregisterObject(TGeometryParticle<float, 3>* GTParticle)
@@ -338,17 +295,17 @@ namespace Chaos
 		UE_LOG(LogPBDRigidsSolver, Verbose, TEXT("FPBDRigidsSolver::UnregisterObject()"));
 
 		// Get the proxy associated with this particle
-		IPhysicsProxyBase* InProxy = GTParticle->Proxy;
+		IPhysicsProxyBase* InProxy = GTParticle->GetProxy();
 		check(InProxy);
 
 		// Grab the particle's type
 		const EParticleType InParticleType = GTParticle->ObjectType();
 
 		// remove the proxy from the invalidation list
-		RemoveDirtyProxy(GTParticle->Proxy);
+		RemoveDirtyProxy(GTParticle->GetProxy());
 
 		// Null out the particle's proxy pointer
-		GTParticle->Proxy = nullptr;
+		GTParticle->SetProxy(nullptr);
 
 		// Remove the proxy from the GT proxy map
 		if (InParticleType == EParticleType::Rigid)
@@ -427,14 +384,19 @@ namespace Chaos
 				delete Proxy;
 			}
 
+			//If particle was created and destroyed before commands were enqueued just skip. I suspect we can skip entire lambda, but too much code to verify right now
+
+			if(Handle)
+			{
 			// Remove game thread particle from ActiveGameThreadParticles so we won't crash when pulling physics state
 			// if this particle was deleted after buffering results. 
 			Solver->GetActiveParticlesBuffer()->RemoveActiveParticleFromConsumerBuffer(Handle->GTGeometryParticle());
-
+  
 			Solver->MParticleToProxy.Remove(Handle);
-
+  
 			// Use the handle to destroy the particle data
 			Solver->GetEvolution()->DestroyParticle(Handle);
+			}
 
 		});
 
@@ -531,6 +493,8 @@ namespace Chaos
 		MMaxSubSteps = 1;
 		MEvolution = TUniquePtr<FPBDRigidsEvolution>(new FPBDRigidsEvolution(Particles, SimMaterials, ChaosSolverCollisionDefaultIterationsCVar, ChaosSolverCollisionDefaultPushoutIterationsCVar, BufferMode == EMultiBufferMode::Single)); 
 
+		DirtyPropertiesManager = MakeUnique<FDoubleBuffer<FDirtyPropertiesManager>>();
+
 		FEventDefaults::RegisterSystemEvents(*GetEventManager());
 	}
 
@@ -565,9 +529,10 @@ namespace Chaos
 	template<typename ProxyType>
 	void PushPhysicsStateExec(FPBDRigidsSolver * Solver, ProxyType* Proxy, Chaos::IDispatcher* Dispatcher)
 	{
+		//todo: remove 99% of this
 		LLM_SCOPE(ELLMTag::Chaos);
-
-		if (Chaos::FParticleData* ProxyData = Proxy->NewData())
+#if 0
+		if (Chaos::FParticlePropertiesData* ProxyData = Proxy->NewData())
 		{
 			if (auto* RigidHandle = static_cast<Chaos::TGeometryParticleHandle<float, 3>*>(Proxy->GetHandle()))
 			{
@@ -581,6 +546,7 @@ namespace Chaos
 							Solver->GetEvolution()->DirtyParticle(*Handle);
 							Proxy->PushToPhysicsState(ProxyData);
 						}
+						//todo: ProxyData->Clear()
 						delete ProxyData;
 					});
 				}
@@ -600,11 +566,13 @@ namespace Chaos
 				delete ProxyData;
 			}
 		}
+#endif
 	}
 
 	template<typename ParticleEntry, typename ProxyEntry, SIZE_T PreAllocCount>
 	void FlushExec(FPBDRigidsSolver::TFramePool<ParticleEntry, ProxyEntry, PreAllocCount>& PoolParticles, Chaos::IDispatcher* Dispatcher, FPBDRigidsSolver * Solver)
 	{
+#if 0
 		Dispatcher->EnqueueCommandImmediate([Solver, &PoolParticles](Chaos::FPersistentPhysicsTask* PhysThread)
 		{
 			for (int32 i = 0; i < PoolParticles.GetEntryCount(); i++)
@@ -619,6 +587,7 @@ namespace Chaos
 				Entry.Particle.Reset();
 			}
 		});
+#endif
 	}
 
 	template<typename ProxyType, typename ParticleEntry, typename ProxyEntry, SIZE_T PreAllocCount>
@@ -641,7 +610,7 @@ namespace Chaos
 
 	void FPBDRigidsSolver::PushPhysicsStatePooled(IDispatcher* Dispatcher)
 	{
-
+#if 0
 		ensure(Dispatcher != nullptr);
 
 		// reset the per frame pool
@@ -649,7 +618,7 @@ namespace Chaos
 		KinematicGeometryParticlePool.ResetPool();
 		GeometryParticlePool.ResetPool();
 
-		TArray< IPhysicsProxyBase *> DirtyProxiesArray = DirtyProxiesSet.Array();
+		const TArray< IPhysicsProxyBase *>& DirtyProxiesArray = DirtyProxiesSet.GetProxies();
 		for (auto & Proxy : DirtyProxiesArray)
 		{
 			switch (Proxy->GetType())
@@ -671,6 +640,7 @@ namespace Chaos
 		FlushExec(RigidParticlePool, Dispatcher, this);
 		FlushExec(KinematicGeometryParticlePool, Dispatcher, this);
 		FlushExec(GeometryParticlePool, Dispatcher, this);
+#endif
 	}
 	
 	void PushPhysicsStateExec(FPBDRigidsSolver* Solver, FGeometryCollectionPhysicsProxy* Proxy, Chaos::IDispatcher* Dispatcher)
@@ -706,42 +676,115 @@ namespace Chaos
 
 	void FPBDRigidsSolver::PushPhysicsState(IDispatcher* Dispatcher)
 	{
-		if (ChaosSolverUseParticlePool)
+		QUICK_SCOPE_CYCLE_COUNTER(STAT_PushPhysicsState);
+		FDirtySet* DirtyProxiesData = DirtyProxiesDataBuffer.AccessProducerBuffer();
+
+		if(DirtyProxiesData->NumDirtyProxies() == 0)
 		{
-			PushPhysicsStatePooled(Dispatcher);
+			//is this early out really necessary?
+			return;
 		}
-		else
+
+		FDirtyPropertiesManager* Manager = DirtyPropertiesManager->AccessProducerBuffer();
+		Manager->SetNumParticles(DirtyProxiesData->NumDirtyProxies());
+		Manager->SetNumShapes(DirtyProxiesData->NumDirtyShapes());
+		FShapeDirtyData* ShapeDirtyData = DirtyProxiesData->GetShapesDirtyData();
+		auto ProcessProxyGT =[ShapeDirtyData, Manager, DirtyProxiesData](auto Proxy, int32 ParticleDataIdx, FDirtyProxy& DirtyProxy)
+		{
+			auto Particle = Proxy->GetParticle();
+			Particle->SyncRemoteData(*Manager,ParticleDataIdx,DirtyProxy.ParticleData,DirtyProxy.ShapeDataIndices,ShapeDirtyData);
+			Proxy->ClearAccumulatedData();
+			Proxy->ResetDirtyIdx();
+		};
+
+
+		//todo: if we allocate remote data ahead of time we could go wide
+		DirtyProxiesData->ParallelForEachProxy([&ProcessProxyGT](int32 DataIdx, FDirtyProxy& Dirty)
 		{
-			if (DirtyProxiesSet.Num() == 0)
-				return;
-			TArray< IPhysicsProxyBase *> DirtyProxiesArray = DirtyProxiesSet.Array();
-			for (auto & Proxy : DirtyProxiesArray)
+			switch(Dirty.Proxy->GetType())
 			{
-				switch (Proxy->GetType())
+			case EPhysicsProxyType::SingleRigidParticleType:
 				{
-					//case EPhysicsProxyType::NoneType: // 0
-					//case EPhysicsProxyType::StaticMeshType: // 1
-				case EPhysicsProxyType::GeometryCollectionType: // 2
-					PushPhysicsStateExec(this, static_cast<FGeometryCollectionPhysicsProxy*>(Proxy), Dispatcher);
+				auto Proxy = static_cast<FRigidParticlePhysicsProxy*>(Dirty.Proxy);
+				ProcessProxyGT(Proxy,DataIdx,Dirty);
 					break;
-				case EPhysicsProxyType::FieldType: // 3
-					PushPhysicsStateExec(this, static_cast<FFieldSystemPhysicsProxy*>(Proxy), Dispatcher);
+			}
+			case EPhysicsProxyType::SingleKinematicParticleType:
+			{
+				auto Proxy = static_cast<FKinematicGeometryParticlePhysicsProxy*>(Dirty.Proxy);
+				ProcessProxyGT(Proxy,DataIdx,Dirty);
 					break;
-					//case EPhysicsProxyType::SkeletalMeshType: // 4
-				case EPhysicsProxyType::SingleRigidParticleType: // 7
-					PushPhysicsStateExec(this, static_cast<FSingleParticlePhysicsProxy< Chaos::TPBDRigidParticle<float, 3> >*>(Proxy), Dispatcher);
+			}
+			case EPhysicsProxyType::SingleGeometryParticleType:
+			{
+				auto Proxy = static_cast<FGeometryParticlePhysicsProxy*>(Dirty.Proxy);
+				ProcessProxyGT(Proxy,DataIdx,Dirty);
 					break;
-				case EPhysicsProxyType::SingleKinematicParticleType: // 6
-					PushPhysicsStateExec(this, static_cast<FSingleParticlePhysicsProxy< Chaos::TKinematicGeometryParticle<float, 3> >*>(Proxy), Dispatcher);
+			}
+			default:
+			ensure("Unknown proxy type in physics solver.");
+			}
+		});
+
+		DirtyPropertiesManager->FlipProducer();
+		DirtyProxiesDataBuffer.FlipProducer();
+
+		FChaosSolversModule::GetModule()->GetDispatcher()->EnqueueCommandImmediate(this,[Manager, DirtyProxiesData, ShapeDirtyData](FPBDRigidsSolver* Solver)
+		{
+			auto ProcessProxyPT = [Solver, Manager, DirtyProxiesData, ShapeDirtyData](auto& Proxy, int32 DataIdx, FDirtyProxy& Dirty, const auto& CreateHandleFunc)
+			{
+				const bool bIsNew = !Proxy->IsInitialized();
+				if(bIsNew)
+				{
+					const auto* NonFrequentData = Dirty.ParticleData.FindNonFrequentData(*Manager, DataIdx);
+					const FUniqueIdx* UniqueIdx = NonFrequentData ? &NonFrequentData->UniqueIdx : nullptr;
+					Proxy->SetHandle(CreateHandleFunc(UniqueIdx));
+				}
+
+				Proxy->PushToPhysicsState(*Manager, DataIdx, Dirty, ShapeDirtyData);
+
+				if(bIsNew)
+				{
+					auto Handle = Proxy->GetHandle();
+					Handle->GTGeometryParticle() = Proxy->GetParticle();
+					Solver->MParticleToProxy.Add(Handle,Proxy);
+					Solver->GetEvolution()->CreateParticle(Handle);
+					Proxy->SetInitialized(true);
+				}
+
+				Dirty.Clear(*Manager, DataIdx, ShapeDirtyData);
+			};
+
+			//need to create new particle handles
+			DirtyProxiesData->ForEachProxy([Solver, &ProcessProxyPT](int32 DataIdx,FDirtyProxy& Dirty)
+			{
+				switch(Dirty.Proxy->GetType())
+				{
+				case EPhysicsProxyType::SingleRigidParticleType:
+				{
+					auto Proxy = static_cast<FRigidParticlePhysicsProxy*>(Dirty.Proxy);
+					ProcessProxyPT(Proxy,DataIdx,Dirty,[Solver](const FUniqueIdx* UniqueIdx){ return Solver->Particles.CreateDynamicParticles(1,UniqueIdx)[0]; });
+					break;
+				}
+				case EPhysicsProxyType::SingleKinematicParticleType:
+				{
+					auto Proxy = static_cast<FKinematicGeometryParticlePhysicsProxy*>(Dirty.Proxy);
+					ProcessProxyPT(Proxy,DataIdx,Dirty,[Solver](const FUniqueIdx* UniqueIdx){ return Solver->Particles.CreateKinematicParticles(1,UniqueIdx)[0]; });
 					break;
-				case EPhysicsProxyType::SingleGeometryParticleType: // 5
-					PushPhysicsStateExec(this, static_cast<FSingleParticlePhysicsProxy< Chaos::TGeometryParticle<float, 3> >*>(Proxy), Dispatcher);
+				}
+				case EPhysicsProxyType::SingleGeometryParticleType:
+				{
+					auto Proxy = static_cast<FGeometryParticlePhysicsProxy*>(Dirty.Proxy);
+					ProcessProxyPT(Proxy,DataIdx,Dirty,[Solver](const FUniqueIdx* UniqueIdx){ return Solver->Particles.CreateStaticParticles(1,UniqueIdx)[0]; });
 					break;
+				}
 				default:
 					ensure("Unknown proxy type in physics solver.");
 				}
-			}
-		}
+			});
+
+			DirtyProxiesData->Reset();
+		});
 	}
 
 	void FPBDRigidsSolver::BufferPhysicsResults()
diff --git a/Engine/Source/Runtime/Experimental/ChaosSolvers/Private/PhysicsProxy/GeometryCollectionPhysicsProxy.cpp b/Engine/Source/Runtime/Experimental/ChaosSolvers/Private/PhysicsProxy/GeometryCollectionPhysicsProxy.cpp
--- a/Engine/Source/Runtime/Experimental/ChaosSolvers/Private/PhysicsProxy/GeometryCollectionPhysicsProxy.cpp
+++ b/Engine/Source/Runtime/Experimental/ChaosSolvers/Private/PhysicsProxy/GeometryCollectionPhysicsProxy.cpp
@@ -324,21 +324,21 @@ void PopulateSimulatedParticle(
 		FCollisionFilterData FilterData;
 		FilterData.Word1 = 0xFFFF; // this body channel
 		FilterData.Word3 = 0xFFFF; // collision candidate channels
-		for (const TUniquePtr<Chaos::TPerShapeData<Chaos::FReal, 3>>& Shape : Handle->ShapesArray())
+		for (const TUniquePtr<Chaos::FPerShapeData>& Shape : Handle->ShapesArray())
 		{
-			Shape->bDisable = false;
-			Shape->CollisionTraceType = Chaos::EChaosCollisionTraceFlag::Chaos_CTF_UseDefault;
+			Shape->SetDisable(false);
+			Shape->SetCollisionTraceType(Chaos::EChaosCollisionTraceFlag::Chaos_CTF_UseDefault);
 			//Shape->CollisionTraceType = Chaos::EChaosCollisionTraceFlag::Chaos_CTF_UseSimpleAndComplex;
-			Shape->SimData = FilterData;
-			Shape->QueryData = FCollisionFilterData();
+			Shape->SetSimData(FilterData);
+			Shape->SetQueryData(FCollisionFilterData());
 		}
 	}
 	else
 	{
-		for (const TUniquePtr<Chaos::TPerShapeData<Chaos::FReal, 3>>& Shape : Handle->ShapesArray())
+		for (const TUniquePtr<Chaos::FPerShapeData>& Shape : Handle->ShapesArray())
 		{
-			Shape->SimData = SimFilterIn;
-			Shape->QueryData = QueryFilterIn;
+			Shape->SetSimData(SimFilterIn);
+			Shape->SetQueryData(QueryFilterIn);
 		}
 	}
 
@@ -538,16 +538,16 @@ void FGeometryCollectionPhysicsProxy::Initialize()
 			P->SetX(T.GetTranslation(), false);
 			P->SetR(T.GetRotation(), false);
 			P->SetUserData(Parameters.UserData);
-			P->Proxy = this;
+			P->SetProxy(this);
 			P->SetGeometry(GameThreadCollection.Implicits[Index]);
 
-			const Chaos::TShapesArray<Chaos::FReal, 3>& Shapes = P->ShapesArray();
+			const Chaos::FShapesArray& Shapes = P->ShapesArray();
 			const int32 NumShapes = Shapes.Num();
 			for(int32 ShapeIndex = 0; ShapeIndex < NumShapes; ++ShapeIndex)
 			{
-				Chaos::TPerShapeData<Chaos::FReal, 3>* Shape = Shapes[ShapeIndex].Get();
-				Shape->SimData = SimFilter;
-				Shape->QueryData = QueryFilter;
+				Chaos::FPerShapeData* Shape = Shapes[ShapeIndex].Get();
+				Shape->SetSimData(SimFilter);
+				Shape->SetQueryData(QueryFilter);
 			}
 		}
 	}
diff --git a/Engine/Source/Runtime/Experimental/ChaosSolvers/Private/PhysicsProxy/SingleParticlePhysicsProxy.cpp b/Engine/Source/Runtime/Experimental/ChaosSolvers/Private/PhysicsProxy/SingleParticlePhysicsProxy.cpp
--- a/Engine/Source/Runtime/Experimental/ChaosSolvers/Private/PhysicsProxy/SingleParticlePhysicsProxy.cpp
+++ b/Engine/Source/Runtime/Experimental/ChaosSolvers/Private/PhysicsProxy/SingleParticlePhysicsProxy.cpp
@@ -36,74 +36,168 @@ FSingleParticlePhysicsProxy<PARTICLE_TYPE>::GetInitialState() const
 	return InitialState; 
 }
 
-
-//
-// TGeometryParticle<float, 3> template specialization 
-//
-
-template< >
-void FSingleParticlePhysicsProxy<Chaos::TGeometryParticle<float, 3>>::PushToPhysicsState(const Chaos::FParticleData* InData)
+template <Chaos::EParticleType ParticleType>
+void PushToPhysicsStateImp(const Chaos::FDirtyPropertiesManager& Manager, Chaos::TGeometryParticleHandle<float,3>* Handle, int32 DataIdx, const Chaos::FDirtyProxy& Dirty, Chaos::FShapeDirtyData* ShapesData, Chaos::FPhysicsSolver* Solver, const bool bInitialized)
 {
+	using namespace Chaos;
+	auto KinematicHandle = ParticleType == EParticleType::Static ? nullptr : static_cast<Chaos::TKinematicGeometryParticleHandle<float,3>*>(Handle);
+	auto RigidHandle = ParticleType != EParticleType::Rigid ? nullptr : static_cast<Chaos::TPBDRigidParticleHandle<float,3>*>(Handle);
+	const FParticleDirtyData& ParticleData = Dirty.ParticleData;
 	// move the copied game thread data into the handle
-	if (auto* RigidHandle = static_cast<Chaos::TGeometryParticleHandle<float, 3>*>(Handle))
 	{
-		typedef Chaos::TKinematicGeometryParticle<float, 3> PARTICLE_TYPE;
-		const auto* Data = static_cast<const typename PARTICLE_TYPE::FData*>(InData);
+		//  DynamicProperties are anything that would wake up a sleeping
+		//  particle. For example, if we set position on a sleeping particle
+		//  we will need to wake it up
+		bool bDynamicPropertyUpdated = false;
 
-		RigidHandle->SetX(Data->X);
-		RigidHandle->SetR(Data->R);
-		RigidHandle->SetSharedGeometry(Data->Geometry);
-		if (Data->Geometry && Data->Geometry->HasBoundingBox())
+		auto NewXR = ParticleData.FindXR(Manager, DataIdx);
+		auto NewNonFrequentData = ParticleData.FindNonFrequentData(Manager, DataIdx);
+
+		if(NewXR)
 		{
-			RigidHandle->SetHasBounds(true);
-			RigidHandle->SetLocalBounds(Data->Geometry->BoundingBox());
-			RigidHandle->SetWorldSpaceInflatedBounds(Data->Geometry->BoundingBox().TransformedAABB(Chaos::TRigidTransform<float, 3>(Data->X, Data->R)));
+			Handle->SetX(NewXR->X);
+			Handle->SetR(NewXR->R);
+			if(RigidHandle)
+			{
+				RigidHandle->SetP(NewXR->X);
+				RigidHandle->SetQ(NewXR->R);
+				bDynamicPropertyUpdated = true;
+			}
 		}
-		RigidHandle->SetSpatialIdx(Data->SpatialIdx);	//todo: this needs to only happen once during initialization
-		RigidHandle->SetUniqueIdx(Data->UniqueIdx);
-		RigidHandle->SetUserData(Data->UserData);
+
+		if(NewNonFrequentData)
+		{
+			Handle->SetSharedGeometry(NewNonFrequentData->Geometry);
+			Handle->SetUniqueIdx(NewNonFrequentData->UniqueIdx);
+			Handle->SetUserData(NewNonFrequentData->UserData);
 #if CHAOS_CHECKED
-		RigidHandle->SetDebugName(Data->DebugName);
+			Handle->SetDebugName(NewNonFrequentData->DebugName);
 #endif
+		}
 
-		if (Data->DirtyFlags.IsDirty(Chaos::EParticleFlags::ShapeDisableCollision))
+		auto NewVelocities = KinematicHandle != nullptr ? ParticleData.FindVelocities(Manager, DataIdx) : nullptr;
+		if(NewVelocities)
 		{
-			int32 CurrShape = 0;
-			for (const TUniquePtr<Chaos::TPerShapeData<Chaos::FReal, 3>>& Shape : RigidHandle->ShapesArray())
+			KinematicHandle->SetV(NewVelocities->V);
+			KinematicHandle->SetW(NewVelocities->W);
+			if(RigidHandle)
 			{
-				Shape->bDisable = Data->ShapeCollisionDisableFlags[CurrShape++];
+				bDynamicPropertyUpdated = true;
 			}
 		}
 
-		// @todo(chaos) : Fix - The initial update is not flagging the dirty flag, so this state is not updated. 
-		//if (Data->DirtyFlags.IsDirty(Chaos::EParticleFlags::CollisionTraceType))
+		if(NewXR || NewNonFrequentData || NewVelocities)
 		{
-			int32 CurrShape = 0;
-			for (const TUniquePtr<Chaos::TPerShapeData<Chaos::FReal, 3>>& Shape : RigidHandle->ShapesArray())
+			const auto& Geometry = Handle->Geometry();
+			if(Geometry && Geometry->HasBoundingBox())
 			{
-				Shape->CollisionTraceType = Data->CollisionTraceType[CurrShape++];
+				Handle->SetHasBounds(true);
+				Handle->SetLocalBounds(Geometry->BoundingBox());
+				TAABB<float,3> WorldSpaceBounds = Geometry->BoundingBox().TransformedAABB(TRigidTransform<float,3>(Handle->X(),Handle->R()));
+				if(KinematicHandle)
+				{
+					WorldSpaceBounds.ThickenSymmetrically(KinematicHandle->V());
+				}
+				Handle->SetWorldSpaceInflatedBounds(WorldSpaceBounds);
 			}
 		}
 
-		// @todo(chaos) : Fix - The initial update is not flagging the dirty flag, so this state is not updated. 
-		//if (Data->DirtyFlags.IsDirty(Chaos::EParticleFlags::ShapeSimData))
+		if(auto NewData = ParticleData.FindMisc(Manager, DataIdx))
 		{
-			int32 CurrShape = 0;
-			for (const TUniquePtr<Chaos::TPerShapeData<Chaos::FReal, 3>>& Shape : RigidHandle->ShapesArray())
+			Handle->SetSpatialIdx(NewData->SpatialIdx);
+
+			if(RigidHandle)
 			{
-				Shape->SimData = Data->ShapeSimData[CurrShape];
-				Shape->QueryData = Data->ShapeQueryData[CurrShape++];
+				Solver->GetEvolution()->SetParticleObjectState(RigidHandle,NewData->ObjectState);
+				Solver->GetEvolution()->GetGravityForces().SetEnabled(*RigidHandle,NewData->bGravityEnabled);
 			}
 		}
 
+		if(auto NewData = ParticleData.FindMassProps(Manager, DataIdx))
 		{
-			int32 CurrShape = 0;
-			for (const TUniquePtr<Chaos::TPerShapeData<Chaos::FReal, 3>>& Shape : RigidHandle->ShapesArray())
+			if(KinematicHandle)
 			{
-				Shape->Materials = Data->ShapeMaterials[CurrShape++];
+				KinematicHandle->SetCenterOfMass(NewData->CenterOfMass);
+				KinematicHandle->SetRotationOfMass(NewData->RotationOfMass);
+			}
+
+			if(RigidHandle)
+			{
+				RigidHandle->SetM(NewData->M);
+				RigidHandle->SetInvM(NewData->InvM);
+				RigidHandle->SetI(NewData->I);
+				RigidHandle->SetInvI(NewData->InvI);
 			}
 		}
+
+		if(RigidHandle)
+		{
+			if(auto NewData = ParticleData.FindDynamics(Manager, DataIdx))
+			{
+				RigidHandle->SetF(NewData->F);
+				RigidHandle->SetTorque(NewData->Torque);
+				RigidHandle->SetLinearImpulse(NewData->LinearImpulse);
+				RigidHandle->SetAngularImpulse(NewData->AngularImpulse);
+
+				RigidHandle->SetLinearEtherDrag(NewData->LinearEtherDrag);
+				RigidHandle->SetAngularEtherDrag(NewData->AngularEtherDrag);
+			}
+		}
+
+		//if dynamic property changed trigger a wakeup
+		//question: should this be here?
+		if(bDynamicPropertyUpdated)
+		{
+			if(bInitialized && RigidHandle->ObjectState() == Chaos::EObjectStateType::Sleeping)
+			{
+				Solver->GetEvolution()->SetParticleObjectState(RigidHandle,Chaos::EObjectStateType::Dynamic);
+			}
+		}
+
+		//shape properties
+		for(int32 ShapeDataIdx : Dirty.ShapeDataIndices)
+		{
+			const FShapeDirtyData& ShapeData = ShapesData[ShapeDataIdx];
+			const int32 ShapeIdx = ShapeData.GetShapeIdx();
+			if(auto NewData = ShapeData.FindCollisionData(Manager, ShapeDataIdx)){ Handle->ShapesArray()[ShapeIdx]->SetCollisionData(*NewData); }
+			if(auto NewData = ShapeData.FindMaterials(Manager, ShapeDataIdx)){ Handle->ShapesArray()[ShapeIdx]->SetMaterialData(*NewData); }
+		}
+	}
+}
+
+//
+// TGeometryParticle<float, 3> template specialization 
+//
+
+template< >
+void FSingleParticlePhysicsProxy<Chaos::TGeometryParticle<float, 3>>::PushToPhysicsState(const Chaos::FDirtyPropertiesManager& Manager, int32 DataIdx, const Chaos::FDirtyProxy& Dirty, Chaos::FShapeDirtyData* ShapesData)
+{
+	PushToPhysicsStateImp<Chaos::EParticleType::Static>(Manager,Handle,DataIdx,Dirty,ShapesData,GetSolver(),bInitialized);
+#if 0
+	// move the copied game thread data into the handle
+	if (auto* RigidHandle = static_cast<Chaos::TGeometryParticleHandle<float, 3>*>(Handle))
+	{
+		auto NewX = Data->FindX();
+		auto NewR = Data->FindR();
+		auto NewGeometry = Data->FindGeometry();
+
+		if(NewX){ RigidHandle->SetX(*NewX); }
+		if(NewR){ RigidHandle->SetR(*NewR); }
+		if(NewGeometry){ RigidHandle->SetSharedGeometry(*NewGeometry); }
+
+		if(NewX || NewR || NewGeometry)
+		{
+			const auto& Geometry = RigidHandle->Geometry();
+			if(Geometry && Geometry->HasBoundingBox())
+			{
+				RigidHandle->SetHasBounds(true);
+				RigidHandle->SetLocalBounds(Geometry->BoundingBox());
+				RigidHandle->SetWorldSpaceInflatedBounds(Geometry->BoundingBox().TransformedAABB(Chaos::TRigidTransform<float,3>(RigidHandle->X(),RigidHandle->R())));
+			}
+		}
+
 	}
+#endif
 }
 
 
@@ -156,8 +250,62 @@ void FSingleParticlePhysicsProxy<Chaos::TGeometryParticle<float, 3>>::ClearEvent
 //
 
 template< >
-void FSingleParticlePhysicsProxy<Chaos::TKinematicGeometryParticle<float, 3>>::PushToPhysicsState(const Chaos::FParticleData* InData)
+void FSingleParticlePhysicsProxy<Chaos::TKinematicGeometryParticle<float, 3>>::PushToPhysicsState(const Chaos::FDirtyPropertiesManager& Manager, int32 DataIdx, const Chaos::FDirtyProxy& Dirty, Chaos::FShapeDirtyData* ShapesData)
 {
+	PushToPhysicsStateImp<Chaos::EParticleType::Kinematic>(Manager,Handle,DataIdx,Dirty,ShapesData,GetSolver(),bInitialized);
+#if 0
+	// move the copied game thread data into the handle
+	if(auto* RigidHandle = static_cast<Chaos::TKinematicGeometryParticleHandle<float,3>*>(Handle))
+	{
+		auto NewX = Data->FindX();
+		auto NewR = Data->FindR();
+		auto NewGeometry = Data->FindGeometry();
+
+		if(NewX){ RigidHandle->SetX(*NewX); }
+		if(NewR){ RigidHandle->SetR(*NewR); }
+		if(NewGeometry){ RigidHandle->SetSharedGeometry(*NewGeometry); }
+
+		if(NewX || NewR || NewGeometry)
+		{
+			const auto& Geometry = RigidHandle->Geometry();
+			if(Geometry && Geometry->HasBoundingBox())
+			{
+				RigidHandle->SetHasBounds(true);
+				RigidHandle->SetLocalBounds(Geometry->BoundingBox());
+				RigidHandle->SetWorldSpaceInflatedBounds(Geometry->BoundingBox().TransformedAABB(Chaos::TRigidTransform<float,3>(RigidHandle->X(),RigidHandle->R())));
+			}
+		}
+
+		if(auto NewSpatialIdx = Data->FindSpatialIdx()){ RigidHandle->SetSpatialIdx(*NewSpatialIdx); }
+		if(auto NewUniqueIdx = Data->FindUniqueIdx()){ RigidHandle->SetUniqueIdx(*NewUniqueIdx); }
+		if(auto NewUserData = Data->FindUserData()){ RigidHandle->SetUserData(*NewUserData); }
+#if CHAOS_CHECKED
+		if(auto NewDebugName = Data->FindDebugName()){ RigidHandle->SetDebugName(*NewDebugName); }
+#endif
+
+		int32 ShapeIdx = 0;
+		const auto& Shapes = RigidHandle->ShapesArray();
+		for(Chaos::FShapePropertiesData* ShapeData : ShapesData->GetRemoteDatas())
+		{
+			//todo: avoid iterating over every shape regardless of if it's dirty or not
+			if(auto NewData = ShapeData->FindQueryData()){ Shapes[ShapeIdx]->SetQueryData(*NewData); }
+			if(auto NewData = ShapeData->FindSimData()){ Shapes[ShapeIdx]->SetSimData(*NewData); }
+			if(auto NewData = ShapeData->FindUserData()){ Shapes[ShapeIdx]->SetUserData(*NewData); }
+			if(auto NewData = ShapeData->FindGeometry()){ Shapes[ShapeIdx]->SetGeometry(*NewData); }
+			//todo: fully marshal materials
+#if 0
+			if(auto NewData = ShapeData->FindMaterials()){ Shapes[ShapeIdx]->SetMaterials(*NewData); }
+			if(auto NewData = ShapeData->FindMaterialMasks()){ Shapes[ShapeIdx]->SetMaterialMasks(*NewData); }
+			if(auto NewData = ShapeData->FindMaterialMaskMaps()){ Shapes[ShapeIdx]->SetMaterialMaskMaps(*NewData); }
+			if(auto NewData = ShapeData->FindMaterialMaskMapMaterials()){ Shapes[ShapeIdx]->SetMaterialMaskMapMaterials(*NewData); }
+#endif
+			if(auto NewData = ShapeData->FindDisable()){ Shapes[ShapeIdx]->SetDisable(*NewData); }
+			if(auto NewData = ShapeData->FindSimulate()){ Shapes[ShapeIdx]->SetSimulate(*NewData); }
+			if(auto NewData = ShapeData->FindCollisionTraceType()){ Shapes[ShapeIdx]->SetCollisionTraceType(*NewData); }
+		}
+	}
+
+#if 0
 	// move the copied game thread data into the handle
 	if (auto* RigidHandle = static_cast<Chaos::TKinematicGeometryParticleHandle<float, 3>*>(Handle))
 	{
@@ -223,6 +371,8 @@ void FSingleParticlePhysicsProxy<Chaos::TKinematicGeometryParticle<float, 3>>::P
 
 
 	}
+#endif
+#endif
 }
 
 template< >
@@ -272,8 +422,10 @@ void FSingleParticlePhysicsProxy<Chaos::TKinematicGeometryParticle<float, 3>>::C
 //
 
 template< >
-void FSingleParticlePhysicsProxy<Chaos::TPBDRigidParticle<float, 3>>::PushToPhysicsState(const Chaos::FParticleData* InData)
+void FSingleParticlePhysicsProxy<Chaos::TPBDRigidParticle<float, 3>>::PushToPhysicsState(const Chaos::FDirtyPropertiesManager& Manager, int32 DataIdx, const Chaos::FDirtyProxy& Dirty, Chaos::FShapeDirtyData* ShapesData)
 {
+	PushToPhysicsStateImp<Chaos::EParticleType::Rigid>(Manager,Handle,DataIdx,Dirty,ShapesData,GetSolver(),bInitialized);
+#if 0
 	// move the copied game thread data into the handle
 	if (auto* RigidHandle = static_cast<Chaos::TPBDRigidParticleHandle<float, 3>*>(Handle))
 	{
@@ -434,16 +586,16 @@ void FSingleParticlePhysicsProxy<Chaos::TPBDRigidParticle<float, 3>>::PushToPhys
 			bInitialized = true;
 		}
 	}
-
+#endif
 }
 
 template< >
 void FSingleParticlePhysicsProxy<Chaos::TPBDRigidParticle<float, 3>>::ClearAccumulatedData()
 {
-	Particle->SetF(Chaos::TVector<float, 3>(0));
-	Particle->SetTorque(Chaos::TVector<float, 3>(0));
-	Particle->SetLinearImpulse(Chaos::TVector<float, 3>(0));
-	Particle->SetAngularImpulse(Chaos::TVector<float, 3>(0));
+	Particle->SetF(Chaos::TVector<float, 3>(0), false);
+	Particle->SetTorque(Chaos::TVector<float, 3>(0), false);
+	Particle->SetLinearImpulse(Chaos::TVector<float, 3>(0), false);
+	Particle->SetAngularImpulse(Chaos::TVector<float, 3>(0), false);
 	Particle->ClearDirtyFlags();
 }
 
@@ -474,7 +626,9 @@ void FSingleParticlePhysicsProxy<Chaos::TPBDRigidParticle<float, 3>>::PullFromPh
 		Particle->SetV(Buffer->MV, false);
 		Particle->SetW(Buffer->MW, false);
 		Particle->UpdateShapeBounds();
-		if (!Particle->IsDirty(Chaos::EParticleFlags::ObjectState))
+		//if (!Particle->IsDirty(Chaos::EParticleFlags::ObjectState))
+		//question: is it ok to call this when it was one of the other properties that changed?
+		if (!Particle->IsDirty(Chaos::EParticleFlags::Misc))
 		{
 			Particle->SetObjectState(Buffer->MObjectState, true);
 		}
diff --git a/Engine/Source/Runtime/Experimental/ChaosSolvers/Public/PBDRigidsSolver.h b/Engine/Source/Runtime/Experimental/ChaosSolvers/Public/PBDRigidsSolver.h
--- a/Engine/Source/Runtime/Experimental/ChaosSolvers/Public/PBDRigidsSolver.h
+++ b/Engine/Source/Runtime/Experimental/ChaosSolvers/Public/PBDRigidsSolver.h
@@ -402,6 +402,7 @@ namespace Chaos
 		template<ELockType>
 		friend struct TSolverQueryMaterialScope;
 
+		TUniquePtr<IBufferResource<FDirtyPropertiesManager>> DirtyPropertiesManager;
 	public:
 
 		template<typename ParticleEntry, typename ProxyEntry, SIZE_T PreAllocCount>
diff --git a/Engine/Source/Runtime/Experimental/ChaosSolvers/Public/PhysicsProxy/FieldSystemPhysicsProxy.h b/Engine/Source/Runtime/Experimental/ChaosSolvers/Public/PhysicsProxy/FieldSystemPhysicsProxy.h
--- a/Engine/Source/Runtime/Experimental/ChaosSolvers/Public/PhysicsProxy/FieldSystemPhysicsProxy.h
+++ b/Engine/Source/Runtime/Experimental/ChaosSolvers/Public/PhysicsProxy/FieldSystemPhysicsProxy.h
@@ -81,7 +81,7 @@ public:
 	void EndFrameCallback(const float InDt) { check(false); } // never called
 
 	// Called by FPBDRigidsSolver::PushPhysicsState() on game thread.
-	Chaos::FParticleData* NewData() { return nullptr; }
+	FStubFieldData* NewData() { return nullptr; }
 	// Called by FPBDRigidsSolver::PushPhysicsState() on physics thread.
 	void PushToPhysicsState(const Chaos::FParticleData*) {};
 	// Called by FPBDRigidsSolver::PushPhysicsState() on game thread.
diff --git a/Engine/Source/Runtime/Experimental/ChaosSolvers/Public/PhysicsProxy/SingleParticlePhysicsProxy.h b/Engine/Source/Runtime/Experimental/ChaosSolvers/Public/PhysicsProxy/SingleParticlePhysicsProxy.h
--- a/Engine/Source/Runtime/Experimental/ChaosSolvers/Public/PhysicsProxy/SingleParticlePhysicsProxy.h
+++ b/Engine/Source/Runtime/Experimental/ChaosSolvers/Public/PhysicsProxy/SingleParticlePhysicsProxy.h
@@ -47,9 +47,9 @@ private:
  *		\c Chaos::TPBDRigidParticle<float,3>
  */
 template<class PARTICLE_TYPE>
-class FSingleParticlePhysicsProxy : public TPhysicsProxy<FSingleParticlePhysicsProxy<PARTICLE_TYPE>, typename PARTICLE_TYPE::FData>
+class FSingleParticlePhysicsProxy : public TPhysicsProxy<FSingleParticlePhysicsProxy<PARTICLE_TYPE>,void>
 {
-	typedef TPhysicsProxy<FSingleParticlePhysicsProxy<PARTICLE_TYPE>, typename PARTICLE_TYPE::FData> Base;
+	typedef TPhysicsProxy<FSingleParticlePhysicsProxy<PARTICLE_TYPE>, void> Base;
 
 public:
 	using FParticleHandle = typename PARTICLE_TYPE::FHandle;
@@ -105,14 +105,8 @@ public:
 		return Chaos::TRigidTransform<float, 3>(Handle->X(), Handle->R());
 	}
 
-	/** 
-	 * Creates a copy of sim state and returns it as \c TGeometryParticleData, 
-	 * \c TKinemticGeometr yParticleData, or \c TPBDRigidParticleData. 
-	 */
-	Chaos::FParticleData* NewData()
+	void* NewData()
 	{
-		if (Particle)
-			return Particle->NewData();
 		return nullptr;
 	}
 
@@ -124,8 +118,7 @@ public:
 	/**/
 	void FlipBuffer() { BufferedData->FlipProducer(); }
 
-	/**/
-	void PushToPhysicsState(const Chaos::FParticleData*);
+	void PushToPhysicsState(const Chaos::FDirtyPropertiesManager& Manager,int32 DataIdx,const Chaos::FDirtyProxy& Dirty,Chaos::FShapeDirtyData* ShapesData);
 
 	/**/
 	void ClearAccumulatedData();
@@ -139,6 +132,9 @@ public:
 	/**/
 	bool IsDirty();
 
+	bool IsInitialized() const { return bInitialized; }
+	void SetInitialized(bool InInitialized) { bInitialized = InInitialized; }
+
 	/**/
 	bool HasAwakeEvent() const;
 
@@ -160,6 +156,8 @@ private:
 	PARTICLE_TYPE* Particle;
 	FParticleHandle* Handle;
 	TUniquePtr<Chaos::IBufferResource<FStorageData>> BufferedData;
+	//TUniquePtr<Chaos::IBufferResource<FPropertiesDataHolder>> PropertiesData;
+	//TUniquePtr<Chaos::FDoubleBuffer<FPropertiesDataHolder>> PropertiesData;
 };
 
 
@@ -177,7 +175,7 @@ template< >
 bool FSingleParticlePhysicsProxy<Chaos::TGeometryParticle<float, 3>>::IsDirty();
 
 template< >
-void FSingleParticlePhysicsProxy<Chaos::TGeometryParticle<float, 3>>::PushToPhysicsState(const Chaos::FParticleData* InData);
+void FSingleParticlePhysicsProxy<Chaos::TGeometryParticle<float,3>>::PushToPhysicsState(const Chaos::FDirtyPropertiesManager& Manager,int32 DataIdx,const Chaos::FDirtyProxy& Dirty,Chaos::FShapeDirtyData* ShapesData);
 
 template< >
 EPhysicsProxyType FSingleParticlePhysicsProxy<Chaos::TGeometryParticle<float, 3>>::ConcreteType();
@@ -191,7 +189,7 @@ CHAOSSOLVERS_API void FSingleParticlePhysicsProxy<Chaos::TGeometryParticle<float
 // TKinematicGeometryParticle specialization prototypes
 
 template< >
-void FSingleParticlePhysicsProxy<Chaos::TKinematicGeometryParticle<float, 3>>::PushToPhysicsState(const Chaos::FParticleData* InData);
+void FSingleParticlePhysicsProxy<Chaos::TKinematicGeometryParticle<float,3>>::PushToPhysicsState(const Chaos::FDirtyPropertiesManager& Manager,int32 DataIdx,const Chaos::FDirtyProxy& Dirty,Chaos::FShapeDirtyData* ShapesData);
 
 template< >
 void FSingleParticlePhysicsProxy<Chaos::TKinematicGeometryParticle<float, 3>>::ClearAccumulatedData();
@@ -229,7 +227,7 @@ template< >
 bool FSingleParticlePhysicsProxy<Chaos::TPBDRigidParticle<float, 3>>::IsDirty();
 
 template< >
-void FSingleParticlePhysicsProxy<Chaos::TPBDRigidParticle<float, 3>>::PushToPhysicsState(const Chaos::FParticleData* InData);
+void FSingleParticlePhysicsProxy<Chaos::TPBDRigidParticle<float,3>>::PushToPhysicsState(const Chaos::FDirtyPropertiesManager& Manager,int32 DataIdx,const Chaos::FDirtyProxy& Dirty,Chaos::FShapeDirtyData* ShapesData);
 
 template< >
 EPhysicsProxyType FSingleParticlePhysicsProxy<Chaos::TPBDRigidParticle<float, 3>>::ConcreteType();
diff --git a/Engine/Source/Runtime/Experimental/ChaosSolvers/Public/PhysicsProxy/SkeletalMeshPhysicsProxy.h b/Engine/Source/Runtime/Experimental/ChaosSolvers/Public/PhysicsProxy/SkeletalMeshPhysicsProxy.h
--- a/Engine/Source/Runtime/Experimental/ChaosSolvers/Public/PhysicsProxy/SkeletalMeshPhysicsProxy.h
+++ b/Engine/Source/Runtime/Experimental/ChaosSolvers/Public/PhysicsProxy/SkeletalMeshPhysicsProxy.h
@@ -138,7 +138,7 @@ public:
 
 	void BufferCommand(Chaos::FPhysicsSolver* InSolver, const FFieldSystemCommand& InCommmand) {}
 
-	Chaos::FParticleData* NewData() { return nullptr; }
+	FSkeletalMeshPhysicsProxyOutputs* NewData() { return nullptr; }
 	void SyncBeforeDestroy();
 	void OnRemoveFromScene();
 	void PushToPhysicsState(const Chaos::FParticleData*) {};
diff --git a/Engine/Source/Runtime/Experimental/ChaosSolvers/Public/PhysicsProxy/StaticMeshPhysicsProxy.h b/Engine/Source/Runtime/Experimental/ChaosSolvers/Public/PhysicsProxy/StaticMeshPhysicsProxy.h
--- a/Engine/Source/Runtime/Experimental/ChaosSolvers/Public/PhysicsProxy/StaticMeshPhysicsProxy.h
+++ b/Engine/Source/Runtime/Experimental/ChaosSolvers/Public/PhysicsProxy/StaticMeshPhysicsProxy.h
@@ -100,7 +100,7 @@ public:
 	void FlipBuffer();
 	void PullFromPhysicsState();
 	bool IsDirty() { return false; }
-	Chaos::FParticleData* NewData() { return nullptr; }
+	FStubSkeletalMeshData* NewData() { return nullptr; }
 	EPhysicsProxyType ConcreteType() { return EPhysicsProxyType::StaticMeshType; }
 	/** ----------------------- */
 
diff --git a/Engine/Source/Runtime/Landscape/Private/LandscapeCollision.cpp b/Engine/Source/Runtime/Landscape/Private/LandscapeCollision.cpp
--- a/Engine/Source/Runtime/Landscape/Private/LandscapeCollision.cpp
+++ b/Engine/Source/Runtime/Landscape/Private/LandscapeCollision.cpp
@@ -403,11 +403,11 @@ void ULandscapeHeightfieldCollisionComponent::OnCreatePhysicsState()
 				FPhysicsActorHandle PhysHandle;
 				FPhysicsInterface::CreateActor(Params, PhysHandle);
 
-				Chaos::TShapesArray<float, 3> ShapeArray;
+				Chaos::FShapesArray ShapeArray;
 				TArray<TUniquePtr<Chaos::FImplicitObject>> Geoms;
 
 				// First add complex geometry
-				TUniquePtr<Chaos::TPerShapeData<float, 3>> NewShape = Chaos::TPerShapeData<float, 3>::CreatePerShapeData();
+				TUniquePtr<Chaos::FPerShapeData> NewShape = Chaos::FPerShapeData::CreatePerShapeData(ShapeArray.Num());
 
 				HeightfieldRef->Heightfield->SetScale(FinalScale);
 				TUniquePtr<Chaos::TImplicitObjectTransformed<float, 3>> ChaosHeightFieldFromCooked = MakeUnique<Chaos::TImplicitObjectTransformed<float, 3>>(MakeSerializable(HeightfieldRef->Heightfield), Chaos::TRigidTransform<float, 3>(FTransform::Identity));
@@ -420,10 +420,10 @@ void ULandscapeHeightfieldCollisionComponent::OnCreatePhysicsState()
 				QueryFilterData.Word3 |= bCreateSimpleCollision ? EPDF_ComplexCollision : (EPDF_SimpleCollision | EPDF_ComplexCollision);
 				SimFilterData.Word3 |= bCreateSimpleCollision ? EPDF_ComplexCollision : (EPDF_SimpleCollision | EPDF_ComplexCollision);
 
-				NewShape->Geometry = MakeSerializable(ChaosHeightFieldFromCooked);
-				NewShape->QueryData = QueryFilterData;
-				NewShape->SimData = SimFilterData;
-				NewShape->Materials = HeightfieldRef->UsedChaosMaterials;
+				NewShape->SetGeometry(MakeSerializable(ChaosHeightFieldFromCooked));
+				NewShape->SetQueryData(QueryFilterData);
+				NewShape->SetSimData(SimFilterData);
+				NewShape->SetMaterials(HeightfieldRef->UsedChaosMaterials);
 
 				Geoms.Emplace(MoveTemp(ChaosHeightFieldFromCooked));
 				ShapeArray.Emplace(MoveTemp(NewShape));
@@ -431,7 +431,7 @@ void ULandscapeHeightfieldCollisionComponent::OnCreatePhysicsState()
 				// Add simple geometry if necessary
 				if(bCreateSimpleCollision)
 				{
-					TUniquePtr<Chaos::TPerShapeData<float, 3>> NewSimpleShape = Chaos::TPerShapeData<float, 3>::CreatePerShapeData();
+					TUniquePtr<Chaos::FPerShapeData> NewSimpleShape = Chaos::FPerShapeData::CreatePerShapeData(ShapeArray.Num());
 
 					HeightfieldRef->HeightfieldSimple->SetScale(FinalScale);
 					TUniquePtr<Chaos::TImplicitObjectTransformed<float, 3>> ChaosSimpleHeightFieldFromCooked = MakeUnique<Chaos::TImplicitObjectTransformed<float, 3>>(MakeSerializable(HeightfieldRef->HeightfieldSimple), Chaos::TRigidTransform<float, 3>(FTransform::Identity));
@@ -441,9 +441,9 @@ void ULandscapeHeightfieldCollisionComponent::OnCreatePhysicsState()
 					QueryFilterDataSimple.Word3 = (QueryFilterDataSimple.Word3 & ~EPDF_ComplexCollision) | EPDF_SimpleCollision;
 					SimFilterDataSimple.Word3 = (SimFilterDataSimple.Word3 & ~EPDF_ComplexCollision) | EPDF_SimpleCollision;
 
-					NewShape->Geometry = MakeSerializable(ChaosSimpleHeightFieldFromCooked);
-					NewShape->QueryData = QueryFilterDataSimple;
-					NewShape->SimData = SimFilterDataSimple;
+					NewShape->SetGeometry(MakeSerializable(ChaosSimpleHeightFieldFromCooked));
+					NewShape->SetQueryData(QueryFilterDataSimple);
+					NewShape->SetSimData(SimFilterDataSimple);
 
 					Geoms.Emplace(MoveTemp(ChaosHeightFieldFromCooked));
 					ShapeArray.Emplace(MoveTemp(NewSimpleShape));
@@ -453,7 +453,7 @@ void ULandscapeHeightfieldCollisionComponent::OnCreatePhysicsState()
 				// Create a shape for a heightfield which is used only by the landscape editor
 				if(!GetWorld()->IsGameWorld())
 				{
-					TUniquePtr<Chaos::TPerShapeData<float, 3>> NewEditorShape = Chaos::TPerShapeData<float, 3>::CreatePerShapeData();
+					TUniquePtr<Chaos::FPerShapeData> NewEditorShape = Chaos::FPerShapeData::CreatePerShapeData(ShapeArray.Num());
 
 					HeightfieldRef->EditorHeightfield->SetScale(FinalScale);
 					TUniquePtr<Chaos::TImplicitObjectTransformed<float, 3>> ChaosEditorHeightFieldFromCooked = MakeUnique<Chaos::TImplicitObjectTransformed<float, 3>>(MakeSerializable(HeightfieldRef->EditorHeightfield), Chaos::TRigidTransform<float, 3>(FTransform::Identity));
@@ -466,9 +466,9 @@ void ULandscapeHeightfieldCollisionComponent::OnCreatePhysicsState()
 
 					QueryFilterDataEd.Word3 |= (EPDF_SimpleCollision | EPDF_ComplexCollision);
 
-					NewEditorShape->Geometry = MakeSerializable(ChaosEditorHeightFieldFromCooked);
-					NewEditorShape->QueryData = QueryFilterDataEd;
-					NewEditorShape->SimData = SimFilterDataEd;
+					NewEditorShape->SetGeometry(MakeSerializable(ChaosEditorHeightFieldFromCooked));
+					NewEditorShape->SetQueryData(QueryFilterDataEd);
+					NewEditorShape->SetSimData(SimFilterDataEd);
 
 					Geoms.Emplace(MoveTemp(ChaosEditorHeightFieldFromCooked));
 					ShapeArray.Emplace(MoveTemp(NewEditorShape));
@@ -511,7 +511,7 @@ void ULandscapeHeightfieldCollisionComponent::OnCreatePhysicsState()
 				bool bImmediateAccelStructureInsertion = true;
 				PhysScene->AddActorsToScene_AssumesLocked(Actors, bImmediateAccelStructureInsertion);
 
-				PhysScene->AddToComponentMaps(this, PhysHandle->Proxy);
+				PhysScene->AddToComponentMaps(this, PhysHandle->GetProxy());
 				if (BodyInstance.bNotifyRigidBodyCollision)
 				{
 					FPhysScene_Chaos& Scene = PhysScene->GetScene();
@@ -545,7 +545,7 @@ void ULandscapeHeightfieldCollisionComponent::OnDestroyPhysicsState()
 		FPhysicsActorHandle& ActorHandle = BodyInstance.GetPhysicsActorHandle();
 		if (FPhysicsInterface::IsValid(ActorHandle))
 		{
-			PhysScene->RemoveFromComponentMaps(ActorHandle->Proxy);
+			PhysScene->RemoveFromComponentMaps(ActorHandle->GetProxy());
 		}
 		if (BodyInstance.bNotifyRigidBodyCollision)
 		{
diff --git a/Engine/Source/Runtime/PhysicsCore/Private/ChaosInterfaceWrapperCore.cpp b/Engine/Source/Runtime/PhysicsCore/Private/ChaosInterfaceWrapperCore.cpp
--- a/Engine/Source/Runtime/PhysicsCore/Private/ChaosInterfaceWrapperCore.cpp
+++ b/Engine/Source/Runtime/PhysicsCore/Private/ChaosInterfaceWrapperCore.cpp
@@ -50,14 +50,14 @@ namespace ChaosInterface
 		return InCapsule.GetHeight() / 2.;
 	}
 
-	FCollisionFilterData GetQueryFilterData(const Chaos::TPerShapeData<float, 3>& Shape)
+	FCollisionFilterData GetQueryFilterData(const Chaos::FPerShapeData& Shape)
 	{
-		return Shape.QueryData;
+		return Shape.GetQueryData();
 	}
 
-	FCollisionFilterData GetSimulationFilterData(const Chaos::TPerShapeData<float, 3>& Shape)
+	FCollisionFilterData GetSimulationFilterData(const Chaos::FPerShapeData& Shape)
 	{
-		return Shape.SimData;
+		return Shape.GetSimData();
 	}
 
 
diff --git a/Engine/Source/Runtime/PhysicsCore/Private/ChaosSQTypes.cpp b/Engine/Source/Runtime/PhysicsCore/Private/ChaosSQTypes.cpp
--- a/Engine/Source/Runtime/PhysicsCore/Private/ChaosSQTypes.cpp
+++ b/Engine/Source/Runtime/PhysicsCore/Private/ChaosSQTypes.cpp
@@ -10,7 +10,7 @@ namespace ChaosInterface
 	void FActorShape::Serialize(Chaos::FChaosArchive& Ar)
 	{
 		Ar << Chaos::AsAlwaysSerializable(Actor);
-		auto NonConstShape = const_cast<Chaos::TPerShapeData<float, 3>*>(Shape);
+		auto NonConstShape = const_cast<Chaos::FPerShapeData*>(Shape);
 		Ar << Chaos::AsAlwaysSerializable(NonConstShape);
 		Shape = NonConstShape;
 	}
diff --git a/Engine/Source/Runtime/PhysicsCore/Private/PhysTestSerializer.cpp b/Engine/Source/Runtime/PhysicsCore/Private/PhysTestSerializer.cpp
--- a/Engine/Source/Runtime/PhysicsCore/Private/PhysTestSerializer.cpp
+++ b/Engine/Source/Runtime/PhysicsCore/Private/PhysTestSerializer.cpp
@@ -283,9 +283,6 @@ void FPhysTestSerializer::CreateChaosData()
 			{
 				PBDRigid->P() = Particle->X();
 				PBDRigid->Q() = Particle->R();
-
-				PBDRigid->GTGeometryParticle()->CastToRigidParticle()->SetP(PBDRigid->P());
-				PBDRigid->GTGeometryParticle()->CastToRigidParticle()->SetQ(PBDRigid->R());
 			}
 
 			PxActorToChaosHandle.Add(Act, Particle.Get());
@@ -325,7 +322,7 @@ void FPhysTestSerializer::CreateChaosData()
 					auto& ShapeArray = GTParticle->ShapesArray();
 					for (auto& Shape : ShapeArray)
 					{
-						Shape->WorldSpaceInflatedShapeBounds = Geom->BoundingBox().TransformedAABB(TRigidTransform<FReal, 3>(Particle->X(), Particle->R()));
+						Shape->SetWorldSpaceInflatedShapeBounds(Geom->BoundingBox().TransformedAABB(TRigidTransform<FReal, 3>(Particle->X(), Particle->R())));
 					}
 				}
 			}
diff --git a/Engine/Source/Runtime/PhysicsCore/Private/PhysicsInterfaceUtilsCore.cpp b/Engine/Source/Runtime/PhysicsCore/Private/PhysicsInterfaceUtilsCore.cpp
--- a/Engine/Source/Runtime/PhysicsCore/Private/PhysicsInterfaceUtilsCore.cpp
+++ b/Engine/Source/Runtime/PhysicsCore/Private/PhysicsInterfaceUtilsCore.cpp
@@ -287,6 +287,6 @@ uint32 FindFaceIndex(const FHitLocation& PHit, const FVector& UnitDir)
 	const FTransform WorldTM(PHit.Actor->R(), PHit.Actor->X());
 	const FVector LocalPosition = WorldTM.InverseTransformPositionNoScale(PHit.WorldPosition);
 	const FVector LocalNormal = WorldTM.InverseTransformVectorNoScale(UnitDir);
-	return PHit.Shape->Geometry->FindMostOpposingFace(LocalPosition, LocalNormal, PHit.FaceIndex, 1);	//todo:this number matches the one above, but is it right?
+	return PHit.Shape->GetGeometry()->FindMostOpposingFace(LocalPosition, LocalNormal, PHit.FaceIndex, 1);	//todo:this number matches the one above, but is it right?
 #endif
 }
diff --git a/Engine/Source/Runtime/PhysicsCore/Private/SQCapture.cpp b/Engine/Source/Runtime/PhysicsCore/Private/SQCapture.cpp
--- a/Engine/Source/Runtime/PhysicsCore/Private/SQCapture.cpp
+++ b/Engine/Source/Runtime/PhysicsCore/Private/SQCapture.cpp
@@ -20,7 +20,7 @@ public:
 	FSQCaptureFilterCallback(const FSQCapture& InCapture) : Capture(InCapture) {}
 	virtual ~FSQCaptureFilterCallback() {}
 	virtual ECollisionQueryHitType PostFilter(const FCollisionFilterData& FilterData, const ChaosInterface::FQueryHit& Hit) override { /*check(false);*/  return ECollisionQueryHitType::Touch; }
-	virtual ECollisionQueryHitType PreFilter(const FCollisionFilterData& FilterData, const Chaos::TPerShapeData<float, 3>& Shape, const Chaos::TGeometryParticle<float, 3>& Actor) override { return Capture.GetFilterResult(&Shape, &Actor); }
+	virtual ECollisionQueryHitType PreFilter(const FCollisionFilterData& FilterData, const Chaos::FPerShapeData& Shape, const Chaos::TGeometryParticle<float, 3>& Actor) override { return Capture.GetFilterResult(&Shape, &Actor); }
 
 #if PHYSICS_INTERFACE_PHYSX
 	virtual ECollisionQueryHitType PostFilter(const FCollisionFilterData& FilterData, const physx::PxQueryHit& Hit) override { return ECollisionQueryHitType::Touch; }
@@ -206,14 +206,14 @@ void FSQCapture::SerializeChaosActorToShapeHitsArray(Chaos::FChaosArchive& Ar)
 			int32 NumShapes;
 			Ar << NumShapes;
 
-			TArray<TPair<Chaos::TPerShapeData<float,3>*, ECollisionQueryHitType>> Pairs;
+			TArray<TPair<Chaos::FPerShapeData*, ECollisionQueryHitType>> Pairs;
 			for (int32 ShapeIdx = 0; ShapeIdx < NumShapes; ++ShapeIdx)
 			{
-				Chaos::TSerializablePtr<Chaos::TPerShapeData<float, 3>> Shape;
+				Chaos::TSerializablePtr<Chaos::FPerShapeData> Shape;
 				Ar << Shape;
 				ECollisionQueryHitType HitType;
 				Ar << HitType;
-				Pairs.Emplace(const_cast<Chaos::TPerShapeData<float,3>*>(Shape.Get()), HitType);
+				Pairs.Emplace(const_cast<Chaos::FPerShapeData*>(Shape.Get()), HitType);
 			}
 			ChaosActorToShapeHitsArray.Add(const_cast<Chaos::TGeometryParticle<float,3>*>(Actor.Get()), Pairs);
 		}
@@ -456,9 +456,9 @@ void FSQCapture::CaptureChaosFilterResults(const Chaos::FPBDRigidsEvolutionGBF&
 	for (int32 Idx = 0; Idx < NumTransientActors; ++Idx)
 	{
 		TGeometryParticle<float, 3>* TransientActor = Particles.GetParticleHandles().Handle(Idx)->GTGeometryParticle();
-		const TShapesArray<float, 3>& TransientShapes = TransientActor->ShapesArray();
+		const FShapesArray& TransientShapes = TransientActor->ShapesArray();
 		const int32 NumTransientShapes = TransientShapes.Num();
-		TArray<TPair<TPerShapeData<float, 3>*, ECollisionQueryHitType>> ShapeHitsArray; ShapeHitsArray.Reserve(NumTransientShapes);
+		TArray<TPair<FPerShapeData*, ECollisionQueryHitType>> ShapeHitsArray; ShapeHitsArray.Reserve(NumTransientShapes);
 		for (const auto& TransientShape : TransientShapes)
 		{
 			const ECollisionQueryHitType Result = Callback.PreFilter(FilterData, *TransientShape, *TransientActor);
@@ -488,7 +488,7 @@ ECollisionQueryHitType GetFilterResultHelper(const TShape* Shape, const TActor*
 	return ECollisionQueryHitType::None;
 }
 
-ECollisionQueryHitType FSQCapture::GetFilterResult(const Chaos::TPerShapeData<float, 3>* Shape, const Chaos::TGeometryParticle<float, 3>* Actor) const
+ECollisionQueryHitType FSQCapture::GetFilterResult(const Chaos::FPerShapeData* Shape, const Chaos::TGeometryParticle<float, 3>* Actor) const
 {
 	return GetFilterResultHelper(Shape, Actor, ChaosActorToShapeHitsArray);
 }
@@ -805,13 +805,13 @@ void FSQCapture::CreateChaosFilterResults()
 	{
 		Chaos::TGeometryParticle<float, 3>* Actor = PhysSerializer.PhysXActorToChaosHandle(Itr.Key);
 		const auto& ShapesArray = Actor->ShapesArray();
-		TArray<TPair<Chaos::TPerShapeData<float,3>*, ECollisionQueryHitType>> FilterResults;
+		TArray<TPair<Chaos::FPerShapeData*, ECollisionQueryHitType>> FilterResults;
 
 		const auto& Pairs = Itr.Value;
 		int32 Idx = 0;
 		for (const auto& Pair : Pairs)
 		{
-			FilterResults.Add(TPair<Chaos::TPerShapeData<float, 3>*, ECollisionQueryHitType>(ShapesArray[Idx++].Get(), Pair.Value));
+			FilterResults.Add(TPair<Chaos::FPerShapeData*, ECollisionQueryHitType>(ShapesArray[Idx++].Get(), Pair.Value));
 		}
 		ChaosActorToShapeHitsArray.Add(Actor, FilterResults);
 	}
diff --git a/Engine/Source/Runtime/PhysicsCore/Public/ChaosInterfaceWrapperCore.h b/Engine/Source/Runtime/PhysicsCore/Public/ChaosInterfaceWrapperCore.h
--- a/Engine/Source/Runtime/PhysicsCore/Public/ChaosInterfaceWrapperCore.h
+++ b/Engine/Source/Runtime/PhysicsCore/Public/ChaosInterfaceWrapperCore.h
@@ -96,8 +96,8 @@ struct FQueryDebugParams
 };
 #endif
 
-extern PHYSICSCORE_API FCollisionFilterData GetQueryFilterData(const Chaos::TPerShapeData<float, 3>& Shape);
-extern PHYSICSCORE_API FCollisionFilterData GetSimulationFilterData(const Chaos::TPerShapeData<float, 3>& Shape);
+extern PHYSICSCORE_API FCollisionFilterData GetQueryFilterData(const Chaos::FPerShapeData& Shape);
+extern PHYSICSCORE_API FCollisionFilterData GetSimulationFilterData(const Chaos::FPerShapeData& Shape);
 
 
 PHYSICSCORE_API ECollisionShapeType GetImplicitType(const Chaos::FImplicitObject& InGeometry);
@@ -117,7 +117,7 @@ inline bool HadInitialOverlap(const FLocationHit& Hit)
 	return Hit.Distance <= 0.f;
 }
 
-inline const Chaos::TPerShapeData<float, 3>* GetShape(const FActorShape& Hit)
+inline const Chaos::FPerShapeData* GetShape(const FActorShape& Hit)
 {
 	return Hit.Shape;
 }
diff --git a/Engine/Source/Runtime/PhysicsCore/Public/ChaosSQTypes.h b/Engine/Source/Runtime/PhysicsCore/Public/ChaosSQTypes.h
--- a/Engine/Source/Runtime/PhysicsCore/Public/ChaosSQTypes.h
+++ b/Engine/Source/Runtime/PhysicsCore/Public/ChaosSQTypes.h
@@ -16,7 +16,7 @@ namespace ChaosInterface
 	struct FActorShape
 	{
 		Chaos::TGeometryParticle<float, 3>* Actor;
-		const Chaos::TPerShapeData<float, 3>* Shape;
+		const Chaos::FPerShapeData* Shape;
 
 		void Serialize(Chaos::FChaosArchive& Ar);
 	};
diff --git a/Engine/Source/Runtime/PhysicsCore/Public/CollisionQueryFilterCallbackCore.h b/Engine/Source/Runtime/PhysicsCore/Public/CollisionQueryFilterCallbackCore.h
--- a/Engine/Source/Runtime/PhysicsCore/Public/CollisionQueryFilterCallbackCore.h
+++ b/Engine/Source/Runtime/PhysicsCore/Public/CollisionQueryFilterCallbackCore.h
@@ -39,7 +39,7 @@ public:
 	
 	virtual ~ICollisionQueryFilterCallbackBase() {}
 	virtual ECollisionQueryHitType PostFilter(const FCollisionFilterData& FilterData, const ChaosInterface::FQueryHit& Hit) = 0;
-	virtual ECollisionQueryHitType PreFilter(const FCollisionFilterData& FilterData, const Chaos::TPerShapeData<float, 3>& Shape, const Chaos::TGeometryParticle<float, 3>& Actor) = 0;
+	virtual ECollisionQueryHitType PreFilter(const FCollisionFilterData& FilterData, const Chaos::FPerShapeData& Shape, const Chaos::TGeometryParticle<float, 3>& Actor) = 0;
 
 #if  PHYSICS_INTERFACE_PHYSX
 	virtual ECollisionQueryHitType PostFilter(const FCollisionFilterData& FilterData, const physx::PxQueryHit& Hit) = 0;
@@ -52,7 +52,7 @@ class PHYSICSCORE_API FBlockAllQueryCallback : public ICollisionQueryFilterCallb
 public:
 	virtual ~FBlockAllQueryCallback() {}
 	virtual ECollisionQueryHitType PostFilter(const FCollisionFilterData& FilterData, const ChaosInterface::FQueryHit& Hit) override { return ECollisionQueryHitType::Block; }
-	virtual ECollisionQueryHitType PreFilter(const FCollisionFilterData& FilterData, const Chaos::TPerShapeData<float, 3>& Shape, const Chaos::TGeometryParticle<float, 3>& Actor) override { return ECollisionQueryHitType::Block; }
+	virtual ECollisionQueryHitType PreFilter(const FCollisionFilterData& FilterData, const Chaos::FPerShapeData& Shape, const Chaos::TGeometryParticle<float, 3>& Actor) override { return ECollisionQueryHitType::Block; }
 
 #if PHYSICS_INTERFACE_PHYSX
 	virtual ECollisionQueryHitType PostFilter(const FCollisionFilterData& FilterData, const physx::PxQueryHit& Hit) override { return ECollisionQueryHitType::Block; }
@@ -67,7 +67,7 @@ class PHYSICSCORE_API FOverlapAllQueryCallback : public ICollisionQueryFilterCal
 public:
 	virtual ~FOverlapAllQueryCallback() {}
 	virtual ECollisionQueryHitType PostFilter(const FCollisionFilterData& FilterData, const ChaosInterface::FQueryHit& Hit) override { return ECollisionQueryHitType::Touch; }
-	virtual ECollisionQueryHitType PreFilter(const FCollisionFilterData& FilterData, const Chaos::TPerShapeData<float, 3>& Shape, const Chaos::TGeometryParticle<float, 3>& Actor) override { return ECollisionQueryHitType::Touch; }
+	virtual ECollisionQueryHitType PreFilter(const FCollisionFilterData& FilterData, const Chaos::FPerShapeData& Shape, const Chaos::TGeometryParticle<float, 3>& Actor) override { return ECollisionQueryHitType::Touch; }
 
 #if PHYSICS_INTERFACE_PHYSX
 	virtual ECollisionQueryHitType PostFilter(const FCollisionFilterData& FilterData, const physx::PxQueryHit& Hit) override { return ECollisionQueryHitType::Touch; }
diff --git a/Engine/Source/Runtime/PhysicsCore/Public/PhysTestSerializer.h b/Engine/Source/Runtime/PhysicsCore/Public/PhysTestSerializer.h
--- a/Engine/Source/Runtime/PhysicsCore/Public/PhysTestSerializer.h
+++ b/Engine/Source/Runtime/PhysicsCore/Public/PhysTestSerializer.h
@@ -120,7 +120,7 @@ public:
 	physx::PxBase* FindObject(uint64 Id);
 
 	Chaos::TGeometryParticle<float,3>* PhysXActorToChaosHandle(physx::PxActor* Actor) const { return PxActorToChaosHandle.FindChecked(Actor)->GTGeometryParticle(); }
-	Chaos::TPerShapeData<float,3>* PhysXShapeToChaosImplicit(physx::PxShape* Shape) const { return PxShapeToChaosShapes.FindRef(Shape); }
+	Chaos::FPerShapeData* PhysXShapeToChaosImplicit(physx::PxShape* Shape) const { return PxShapeToChaosShapes.FindRef(Shape); }
 #endif
 
 private:
@@ -167,7 +167,7 @@ private:
 	TUniquePtr<FPhysXSerializerData> AlignedDataHelper;
 
 	TMap<physx::PxActor*, Chaos::TGeometryParticleHandle<float, 3>*> PxActorToChaosHandle;
-	TMap<physx::PxShape*, Chaos::TPerShapeData<float, 3>*> PxShapeToChaosShapes;
+	TMap<physx::PxShape*, Chaos::FPerShapeData*> PxShapeToChaosShapes;
 #endif
 };
 
diff --git a/Engine/Source/Runtime/PhysicsCore/Public/PhysicsInterfaceDeclaresCore.h b/Engine/Source/Runtime/PhysicsCore/Public/PhysicsInterfaceDeclaresCore.h
--- a/Engine/Source/Runtime/PhysicsCore/Public/PhysicsInterfaceDeclaresCore.h
+++ b/Engine/Source/Runtime/PhysicsCore/Public/PhysicsInterfaceDeclaresCore.h
@@ -28,8 +28,7 @@ namespace Chaos
 	template <typename T, int d>
 	class TGeometryParticle;
 
-	template <typename T, int d>
-	class TPerShapeData;
+	class FPerShapeData;
 
 	class FPhysicalMaterial;
 	class FPhysicalMaterialMask;
@@ -67,7 +66,7 @@ using FPhysicsQueryHit = ChaosInterface::FQueryHit;
 
 using FPhysicsTransform = FTransform;
 
-using FPhysicsShape = Chaos::TPerShapeData<float, 3>;
+using FPhysicsShape = Chaos::FPerShapeData;
 using FPhysicsGeometry = Chaos::FImplicitObject;
 using FPhysicsCapsuleGeometry = Chaos::TCapsule<float>;
 using FPhysicsMaterial = Chaos::FChaosPhysicsMaterial;
diff --git a/Engine/Source/Runtime/PhysicsCore/Public/SQCapture.h b/Engine/Source/Runtime/PhysicsCore/Public/SQCapture.h
--- a/Engine/Source/Runtime/PhysicsCore/Public/SQCapture.h
+++ b/Engine/Source/Runtime/PhysicsCore/Public/SQCapture.h
@@ -63,7 +63,7 @@ struct PHYSICSCORE_API FSQCapture
 	void StartCaptureChaosOverlap(const Chaos::FPBDRigidsEvolutionGBF& Evolution, const Chaos::FImplicitObject& InQueryGeom, const FTransform& InStartTM, const FQueryFilterData& QueryFilter, const FCollisionFilterData& FilterData, ICollisionQueryFilterCallbackBase& Callback);
 	void EndCaptureChaosOverlap(const ChaosInterface::FSQHitBuffer<ChaosInterface::FOverlapHit>& Results);
 
-	ECollisionQueryHitType GetFilterResult(const Chaos::TPerShapeData<float,3>* Shape, const Chaos::TGeometryParticle<float,3>* Actor) const;
+	ECollisionQueryHitType GetFilterResult(const Chaos::FPerShapeData* Shape, const Chaos::TGeometryParticle<float,3>* Actor) const;
 
 #if PHYSICS_INTERFACE_PHYSX
 	ECollisionQueryHitType GetFilterResult(const physx::PxShape* Shape, const physx::PxActor* Actor) const;
@@ -161,7 +161,7 @@ private:
 
 	void CaptureChaosFilterResults(const Chaos::FPBDRigidsEvolutionGBF& Evolution, const FCollisionFilterData& FilterData, ICollisionQueryFilterCallbackBase& Callback);
 
-	TMap<Chaos::TGeometryParticle<float,3>*, TArray<TPair<Chaos::TPerShapeData<float,3>*, ECollisionQueryHitType>>> ChaosActorToShapeHitsArray;
+	TMap<Chaos::TGeometryParticle<float,3>*, TArray<TPair<Chaos::FPerShapeData*, ECollisionQueryHitType>>> ChaosActorToShapeHitsArray;
 
 	template <typename THit>
 	void SerializeChaosBuffers(Chaos::FChaosArchive& Ar, int32 Version, ChaosInterface::FSQHitBuffer<THit>& ChaosBuffer);
diff --git a/Engine/Source/Runtime/PhysicsSQ/Private/SQAccelerator.cpp b/Engine/Source/Runtime/PhysicsSQ/Private/SQAccelerator.cpp
--- a/Engine/Source/Runtime/PhysicsSQ/Private/SQAccelerator.cpp
+++ b/Engine/Source/Runtime/PhysicsSQ/Private/SQAccelerator.cpp
@@ -199,7 +199,7 @@ private:
 		//todo: add a check to ensure hitbuffer matches SQ type
 		using namespace Chaos;
 		TGeometryParticle<float, 3>* GeometryParticle = Payload.GetExternalGeometryParticle_ExternalThread();
-		const TShapesArray<float,3>& Shapes = GeometryParticle->ShapesArray();
+		const FShapesArray& Shapes = GeometryParticle->ShapesArray();
 
 		const bool bTestShapeBounds =  Shapes.Num() > 1;
 		bool bContinue = true;
@@ -213,21 +213,21 @@ private:
 
 		for (const auto& Shape : Shapes)
 		{
-			const FImplicitObject* Geom = Shape->Geometry.Get();
+			const FImplicitObject* Geom = Shape->GetGeometry().Get();
 
 			if (bTestShapeBounds)
 			{
 				TAABB<FReal, 3> InflatedWorldBounds;
 				if (SQ == ESQType::Raycast)
 				{
-					InflatedWorldBounds = Shape->WorldSpaceInflatedShapeBounds;
+					InflatedWorldBounds = Shape->GetWorldSpaceInflatedShapeBounds();
 				}
 				else
 				{
 					// Transform to world bounds and get the proper half extent.
 					const FVec3 WorldHalfExtent = QueryGeom ? QueryGeom->BoundingBox().TransformedAABB(StartTM).Extents() * 0.5f : HalfExtents;
 
-					InflatedWorldBounds = TAABB<FReal, 3>(Shape->WorldSpaceInflatedShapeBounds.Min() - WorldHalfExtent, Shape->WorldSpaceInflatedShapeBounds.Max() + WorldHalfExtent);
+					InflatedWorldBounds = TAABB<FReal, 3>(Shape->GetWorldSpaceInflatedShapeBounds().Min() - WorldHalfExtent, Shape->GetWorldSpaceInflatedShapeBounds().Max() + WorldHalfExtent);
 				}
 				if (SQ != ESQType::Overlap)
 				{
