Don't fail level set rasterization if the mesh contains 1 small triangle (creating face normals fails).  Fall back on implicit sphere if level set rasterization fails during geometry collection asset creation.  Make sure geometry collection InnerRadius and OuterRadius are valid, and calculate if necessary.  Generate implicit geometry for geometry collections in parallel. Downgrade spewage about empty-ish level sets to a warning.

diff --git a/Engine/Source/Runtime/Experimental/Chaos/Private/Chaos/Levelset.cpp b/Engine/Source/Runtime/Experimental/Chaos/Private/Chaos/Levelset.cpp
--- a/Engine/Source/Runtime/Experimental/Chaos/Private/Chaos/Levelset.cpp
+++ b/Engine/Source/Runtime/Experimental/Chaos/Private/Chaos/Levelset.cpp
@@ -45,11 +45,13 @@ TLevelSet<T, d>::TLevelSet(FErrorReporter& ErrorReporter, const TUniformGrid<T,
 	check(MGrid.Counts()[0] > 1 && MGrid.Counts()[1] > 1 && MGrid.Counts()[2] > 1);
 	check(Mesh.GetSurfaceElements().Num());
 
-	const TArray<TVector<T, 3>> Normals = Mesh.GetFaceNormals(InParticles);
-
+	const TArray<TVector<T, 3>> Normals = 
+		Mesh.GetFaceNormals(
+			InParticles, 
+			false);			// Don't fail if the mesh has small faces
 	if (Normals.Num() == 0)
 	{
-		ErrorReporter.ReportError(TEXT("Normals came back empty. Does mesh contain coincident points?"));
+		ErrorReporter.ReportError(TEXT("Normals came back empty."));
 		return;
 	}
 
diff --git a/Engine/Source/Runtime/Experimental/Chaos/Public/Chaos/ErrorReporter.h b/Engine/Source/Runtime/Experimental/Chaos/Public/Chaos/ErrorReporter.h
--- a/Engine/Source/Runtime/Experimental/Chaos/Public/Chaos/ErrorReporter.h
+++ b/Engine/Source/Runtime/Experimental/Chaos/Public/Chaos/ErrorReporter.h
@@ -25,7 +25,7 @@ public:
 		return bUnhandledErrors;
 	}
 
-	void ReportError(const TCHAR* ErrorMsg)
+	void ReportWarning(const TCHAR* ErrorMsg)
 	{
 		if (Prefix != "")
 		{
@@ -35,6 +35,11 @@ public:
 		{
 			UE_LOG(LogChaos, Warning, TEXT("ErrorReporter: %s"), *Prefix, ErrorMsg);
 		}
+	}
+
+	void ReportError(const TCHAR* ErrorMsg)
+	{
+		ReportWarning(ErrorMsg);
 		bEncountedErrors = true;
 		bUnhandledErrors = true;
 	}
diff --git a/Engine/Source/Runtime/Experimental/ChaosSolvers/Private/PhysicsProxy/GeometryCollectionPhysicsProxy.cpp b/Engine/Source/Runtime/Experimental/ChaosSolvers/Private/PhysicsProxy/GeometryCollectionPhysicsProxy.cpp
--- a/Engine/Source/Runtime/Experimental/ChaosSolvers/Private/PhysicsProxy/GeometryCollectionPhysicsProxy.cpp
+++ b/Engine/Source/Runtime/Experimental/ChaosSolvers/Private/PhysicsProxy/GeometryCollectionPhysicsProxy.cpp
@@ -1844,7 +1844,8 @@ void FGeometryCollectionPhysicsProxy::InitializeSharedCollisionStructures(
 
 	// GeometryGroup
 	const TManagedArray<FBox>& BoundingBox = RestCollection.BoundingBox;
-	const TManagedArray<float>& InnerRadius = RestCollection.InnerRadius;
+	TManagedArray<float>& InnerRadius = RestCollection.InnerRadius;
+	TManagedArray<float>& OuterRadius = RestCollection.OuterRadius;
 	const TManagedArray<int32>& VertexCount = RestCollection.VertexCount;
 	const TManagedArray<int32>& VertexStart = RestCollection.VertexStart;
 	const TManagedArray<int32>& FaceCount = RestCollection.FaceCount;
@@ -1962,6 +1963,32 @@ void FGeometryCollectionPhysicsProxy::InitializeSharedCollisionStructures(
 			}
 #endif
 
+			if (InnerRadius[GeometryIndex] == 0.0f || OuterRadius[GeometryIndex] == 0.0f)
+			{
+				//const FVector Center = MassProperties.CenterOfMass;
+				const FVector Center = BoundingBox[GeometryIndex].GetCenter();
+				const int32 VStart = VertexStart[GeometryIndex];
+				const int32 VCount = VertexCount[GeometryIndex];
+				InnerRadius[GeometryIndex] = VCount ? TNumericLimits<float>::Max() : 0.0f;
+				OuterRadius[GeometryIndex] = 0.0f;
+				for (int32 VIdx = 0; VIdx < VCount; ++VIdx)
+				{
+					const int32 PtIdx = VStart + VIdx;
+					const FVector& Pt = Vertex[PtIdx];
+					const float DistSq = FVector::DistSquared(Pt, Center);
+					if (InnerRadius[GeometryIndex] > DistSq)
+					{
+						InnerRadius[GeometryIndex] = DistSq;
+					}
+					if (OuterRadius[GeometryIndex] < DistSq)
+					{
+						OuterRadius[GeometryIndex] = DistSq;
+					}
+				}
+				InnerRadius[GeometryIndex] = FMath::Sqrt(InnerRadius[GeometryIndex]);
+				OuterRadius[GeometryIndex] = FMath::Sqrt(OuterRadius[GeometryIndex]);
+			}
+
 			TotalVolume += MassProperties.Volume;
 			TriangleMeshesArray[TransformGroupIndex] = MoveTemp(TriMesh);
 		}
@@ -1977,7 +2004,8 @@ void FGeometryCollectionPhysicsProxy::InitializeSharedCollisionStructures(
 	const float DesiredDensity = ClampedTotalMass / TotalVolume;
 	TVector<float, 3> MaxChildBounds(1);
 
-	for (int32 GeometryIndex = 0; GeometryIndex < NumGeometries; GeometryIndex++)
+	ParallelFor(NumGeometries, [&](int32 GeometryIndex)
+	//for (int32 GeometryIndex = 0; GeometryIndex < NumGeometries; GeometryIndex++)
 	{
 		const int32 TransformGroupIndex = TransformIndex[GeometryIndex];
 		if (CollectionSimulatableParticles[TransformGroupIndex])
@@ -2069,6 +2097,15 @@ void FGeometryCollectionPhysicsProxy::InitializeSharedCollisionStructures(
 						SizeSpecificData.MaxLevelSetResolution,
 						SizeSpecificData.CollisionObjectReductionPercentage,
 						SizeSpecificData.CollisionType));
+				// Fall back on sphere if level set rasterization failed.
+				if (!CollectionImplicits[TransformGroupIndex])
+				{
+					CollectionImplicits[TransformGroupIndex] = FGeometryDynamicCollection::FSharedImplicit(
+						FCollisionStructureManager::NewImplicitSphere(
+							InnerRadius[GeometryIndex],
+							SizeSpecificData.CollisionObjectReductionPercentage,
+							SizeSpecificData.CollisionType));
+				}
 			}
 			else if (SizeSpecificData.ImplicitType == EImplicitTypeEnum::Chaos_Implicit_Box)
 			{
@@ -2095,14 +2132,14 @@ void FGeometryCollectionPhysicsProxy::InitializeSharedCollisionStructures(
 				ensure(false); // unsupported implicit type!
 			}
 
-			if (CollectionImplicits[TransformGroupIndex])
+			if (CollectionImplicits[TransformGroupIndex] && CollectionImplicits[TransformGroupIndex]->HasBoundingBox())
 			{
-				const TVector<float, 3> Extents = 
-					CollectionImplicits[TransformGroupIndex]->BoundingBox().Extents(); // Chaos::TAABB::Extents() is Max - Min
+				const auto BBox = CollectionImplicits[TransformGroupIndex]->BoundingBox();
+				const TVector<float, 3> Extents = BBox.Extents(); // Chaos::TAABB::Extents() is Max - Min
 				MaxChildBounds = MaxChildBounds.ComponentwiseMax(Extents);
 			}
 		}
-	}
+	});
 
 	// question: at the moment we always build cluster data in the asset. This 
 	// allows for per instance toggling. Is this needed? It increases memory 
diff --git a/Engine/Source/Runtime/Experimental/GeometryCollectionSimulationCore/Private/GeometryCollection/GeometryCollectionCollisionStructureManager.cpp b/Engine/Source/Runtime/Experimental/GeometryCollectionSimulationCore/Private/GeometryCollection/GeometryCollectionCollisionStructureManager.cpp
--- a/Engine/Source/Runtime/Experimental/GeometryCollectionSimulationCore/Private/GeometryCollection/GeometryCollectionCollisionStructureManager.cpp
+++ b/Engine/Source/Runtime/Experimental/GeometryCollectionSimulationCore/Private/GeometryCollection/GeometryCollectionCollisionStructureManager.cpp
@@ -256,13 +256,14 @@ FCollisionStructureManager::NewImplicitLevelset(
 		const float FilledVolume = LevelSet->ApproximateNegativeMaterial();
 		if (FilledVolume < DomainVolume * 0.05)
 		{
-			ErrorReporter.ReportError(
+			const Chaos::TVector<float,3> Extent = LevelSet->BoundingBox().Extents();
+			ErrorReporter.ReportWarning(
 				*FString::Printf(TEXT(
 					"Level set is small or empty:\n"
-					"    domain volume: %g\n"
+					"    domain extent: (%g %g %g) volume: %g\n"
 					"    estimated level set volume: %g\n"
 					"    percentage filled: %g%%"),
-					DomainVolume, FilledVolume, FilledVolume / DomainVolume * 100.0));
+					Extent[0], Extent[1], Extent[2], DomainVolume, FilledVolume, FilledVolume / DomainVolume * 100.0));
 		}
 		HalfExtents *= CollisionObjectReduction / 100.f;
 		const float MinExtent = FMath::Min(HalfExtents[0], FMath::Min(HalfExtents[1], HalfExtents[2]));
