GeometryCollectionTest

- Addd MassPropertiesTests
- Remoted OLD.cpp

diff --git a/Engine/Source/Runtime/Experimental/ChaosSolvers/Private/PhysicsProxy/GeometryCollectionPhysicsProxyOLD.cpp b/Engine/Source/Runtime/Experimental/ChaosSolvers/Private/PhysicsProxy/GeometryCollectionPhysicsProxyOLD.cpp
deleted file mode 100644
diff --git a/Engine/Source/Runtime/Experimental/ChaosSolvers/Public/PhysicsProxy/GeometryCollectionPhysicsProxy.h b/Engine/Source/Runtime/Experimental/ChaosSolvers/Public/PhysicsProxy/GeometryCollectionPhysicsProxy.h
--- a/Engine/Source/Runtime/Experimental/ChaosSolvers/Public/PhysicsProxy/GeometryCollectionPhysicsProxy.h
+++ b/Engine/Source/Runtime/Experimental/ChaosSolvers/Public/PhysicsProxy/GeometryCollectionPhysicsProxy.h
@@ -301,4 +301,5 @@ private:
 
 };
 
+CHAOSSOLVERS_API Chaos::TTriangleMesh<float>* CreateTriangleMesh(const int32 FaceCount,const int32 VertexOffset,const int32 StartIndex,const TManagedArray<FVector>& Vertex,const TManagedArray<bool>& Visible,const TManagedArray<FIntVector>& Indices,TSet<int32>& VertsAdded);
 CHAOSSOLVERS_API void BuildSimulationData(Chaos::FErrorReporter& ErrorReporter, FGeometryCollection& GeometryCollection, const FSharedSimulationParameters& SharedParams);
