Fixed Debug build.

#rnx

diff --git a/Engine/Source/Runtime/Experimental/ChaosSolvers/Private/PBDRigidsSolver.cpp b/Engine/Source/Runtime/Experimental/ChaosSolvers/Private/PBDRigidsSolver.cpp
--- a/Engine/Source/Runtime/Experimental/ChaosSolvers/Private/PBDRigidsSolver.cpp
+++ b/Engine/Source/Runtime/Experimental/ChaosSolvers/Private/PBDRigidsSolver.cpp
@@ -238,7 +238,7 @@ namespace Chaos
 		UE_LOG(LogPBDRigidsSolver, Verbose, TEXT("FPBDRigidsSolver::RegisterObject()"));
 
 		// Make sure this particle doesn't already have a proxy
-		checkSlow(GTParticle->Proxy == nullptr);
+		checkSlow(GTParticle->GetProxy() == nullptr);
 
 		if (GTParticle->Geometry() && GTParticle->Geometry()->HasBoundingBox() && GTParticle->Geometry()->BoundingBox().Extents().Max() >= MaxBoundsForTree)
 		{
