Fix crash in single threaded mode

diff --git a/Engine/Source/Runtime/Experimental/ChaosSolvers/Private/PBDRigidsSolver.cpp b/Engine/Source/Runtime/Experimental/ChaosSolvers/Private/PBDRigidsSolver.cpp
--- a/Engine/Source/Runtime/Experimental/ChaosSolvers/Private/PBDRigidsSolver.cpp
+++ b/Engine/Source/Runtime/Experimental/ChaosSolvers/Private/PBDRigidsSolver.cpp
@@ -258,6 +258,8 @@ namespace Chaos
 		//Chaos::FParticlePropertiesData& RemoteParticleData = *DirtyPropertiesManager->AccessProducerBuffer()->NewRemoteParticleProperties();
 		//Chaos::FShapeRemoteDataContainer& RemoteShapeContainer = *DirtyPropertiesManager->AccessProducerBuffer()->NewRemoteShapeContainer();
 
+		const bool bIsSingleThreaded = FChaosSolversModule::GetModule()->GetDispatcher()->GetMode() == EThreadingMode::SingleThread;
+
 		// Make a physics proxy, giving it our particle and particle handle
 		const EParticleType InParticleType = GTParticle->ObjectType();
 		if (InParticleType == EParticleType::Rigid)
@@ -265,21 +267,42 @@ namespace Chaos
 			auto Proxy = new FRigidParticlePhysicsProxy(GTParticle->CastToRigidParticle(), nullptr);
 			RigidParticlePhysicsProxies.Add((FRigidParticlePhysicsProxy*)Proxy);
 			ProxyBase = Proxy;
-			//Proxy->SetRemoteDataForGT(RemoteParticleData, RemoteShapeContainer);
+			
+			//todo: remove this and have just one handle
+			if(bIsSingleThreaded)
+			{
+				FUniqueIdx UniqueIdx = GTParticle->UniqueIdx();
+				auto* Handle = Particles.CreateDynamicParticles(1,&UniqueIdx)[0];
+				Proxy->SetHandle(Handle);
+			}
 		}
 		else if (InParticleType == EParticleType::Kinematic)
 		{
 			auto Proxy = new FKinematicGeometryParticlePhysicsProxy(GTParticle->CastToKinematicParticle(), nullptr);
 			KinematicGeometryParticlePhysicsProxies.Add((FKinematicGeometryParticlePhysicsProxy*)Proxy);
 			ProxyBase = Proxy;
-			//Proxy->SetRemoteDataForGT(RemoteParticleData, RemoteShapeContainer);
+			
+			//todo: remove this and have just one handle
+			if(bIsSingleThreaded)
+			{
+				FUniqueIdx UniqueIdx = GTParticle->UniqueIdx();
+				auto* Handle = Particles.CreateKinematicParticles(1,&UniqueIdx)[0];
+				Proxy->SetHandle(Handle);
+			}
 		}
 		else // Assume it's a static (geometry) if it's not dynamic or kinematic
 		{
 			auto Proxy = new FGeometryParticlePhysicsProxy(GTParticle, nullptr);
 			GeometryParticlePhysicsProxies.Add((FGeometryParticlePhysicsProxy*)Proxy);
 			ProxyBase = Proxy;
-			//Proxy->SetRemoteDataForGT(RemoteParticleData, RemoteShapeContainer);
+			
+			//todo: remove this and have just one handle
+			if(bIsSingleThreaded)
+			{
+				FUniqueIdx UniqueIdx = GTParticle->UniqueIdx();
+				auto* Handle = Particles.CreateStaticParticles(1,&UniqueIdx)[0];
+				Proxy->SetHandle(Handle);
+			}
 		}
 
 		ProxyBase->SetSolver(this);
@@ -728,15 +751,17 @@ namespace Chaos
 
 		DirtyPropertiesManager->FlipProducer();
 		DirtyProxiesDataBuffer.FlipProducer();
+		const bool bIsSingleThreaded = FChaosSolversModule::GetModule()->GetDispatcher()->GetMode() == EThreadingMode::SingleThread;
 
-		FChaosSolversModule::GetModule()->GetDispatcher()->EnqueueCommandImmediate(this,[Manager, DirtyProxiesData, ShapeDirtyData](FPBDRigidsSolver* Solver)
+		FChaosSolversModule::GetModule()->GetDispatcher()->EnqueueCommandImmediate(this,[bIsSingleThreaded, Manager,DirtyProxiesData,ShapeDirtyData](FPBDRigidsSolver* Solver)
 		{
-			auto ProcessProxyPT = [Solver, Manager, DirtyProxiesData, ShapeDirtyData](auto& Proxy, int32 DataIdx, FDirtyProxy& Dirty, const auto& CreateHandleFunc)
+			auto ProcessProxyPT = [bIsSingleThreaded, Solver,Manager,DirtyProxiesData,ShapeDirtyData](auto& Proxy,int32 DataIdx,FDirtyProxy& Dirty,const auto& CreateHandleFunc)
 			{
 				const bool bIsNew = !Proxy->IsInitialized();
-				if(bIsNew)
+				//single threaded version already created particle, but didn't initialize it
+				if(!bIsSingleThreaded && bIsNew)
 				{
-					const auto* NonFrequentData = Dirty.ParticleData.FindNonFrequentData(*Manager, DataIdx);
+					const auto* NonFrequentData = Dirty.ParticleData.FindNonFrequentData(*Manager,DataIdx);
 					const FUniqueIdx* UniqueIdx = NonFrequentData ? &NonFrequentData->UniqueIdx : nullptr;
 					Proxy->SetHandle(CreateHandleFunc(UniqueIdx));
 				}
