Fix Build w/ bChaosChecked.

diff --git a/Engine/Source/Runtime/Experimental/Chaos/Public/Chaos/ParticleHandle.h b/Engine/Source/Runtime/Experimental/Chaos/Public/Chaos/ParticleHandle.h
--- a/Engine/Source/Runtime/Experimental/Chaos/Public/Chaos/ParticleHandle.h
+++ b/Engine/Source/Runtime/Experimental/Chaos/Public/Chaos/ParticleHandle.h
@@ -1489,7 +1489,7 @@ public:
 	const FName& DebugName() const { return MNonFrequentData.Read().DebugName; }
 	void SetDebugName(const FName& InDebugName)
 	{
-		MNonFrequentData.Modify(bInvalidate,MDirtyFlags,Proxy,[&InDebugName](auto& Data){ Data.DebugName = InDebugName;});
+		MNonFrequentData.Modify(true,MDirtyFlags,Proxy,[&InDebugName](auto& Data){ Data.DebugName = InDebugName;});
 	}
 #endif
 
