AABBTree: Update Bounding boxes of elements contained in leaves when updated. Even if the bounding box is still within the leave bounding box.

diff --git a/Engine/Source/Runtime/Experimental/Chaos/Public/Chaos/AABBTree.h b/Engine/Source/Runtime/Experimental/Chaos/Public/Chaos/AABBTree.h
--- a/Engine/Source/Runtime/Experimental/Chaos/Public/Chaos/AABBTree.h
+++ b/Engine/Source/Runtime/Experimental/Chaos/Public/Chaos/AABBTree.h
@@ -199,6 +199,21 @@ struct TAABBTreeLeafArray : public TBoundsWrapperHelper<TPayloadType, T, bComput
 		}
 	}
 
+	void UpdateElement(const TPayloadType& Payload, const TAABB<T, 3>& NewBounds, bool bHasBounds)
+	{
+		if (!bHasBounds)
+			return;
+
+		for (int32 Idx = 0; Idx < Elems.Num(); ++Idx)
+		{
+			if (Elems[Idx].Payload == Payload)
+			{
+				Elems[Idx].Bounds = NewBounds;
+				break;
+			}
+		}
+	}
+
 	const TAABB<T, 3>& GetBounds() const
 	{
 		return Bounds;
@@ -494,6 +509,8 @@ public:
 						const TAABB<T,3>& LeafBounds = Leaves[PayloadInfo->LeafIdx].GetBounds();
 						if (LeafBounds.Contains(NewBounds.Min()) && LeafBounds.Contains(NewBounds.Max()))
 						{
+							// We still need to update the constituent bounds
+							Leaves[PayloadInfo->LeafIdx].UpdateElement(Payload, NewBounds, bHasBounds);
 							return;
 						}
 					}
