Fix Particles pushing to physics state not dirtying themselves in accel structure update queues.

diff --git a/Engine/Source/Runtime/Experimental/ChaosSolvers/Private/PhysicsProxy/SingleParticlePhysicsProxy.cpp b/Engine/Source/Runtime/Experimental/ChaosSolvers/Private/PhysicsProxy/SingleParticlePhysicsProxy.cpp
--- a/Engine/Source/Runtime/Experimental/ChaosSolvers/Private/PhysicsProxy/SingleParticlePhysicsProxy.cpp
+++ b/Engine/Source/Runtime/Experimental/ChaosSolvers/Private/PhysicsProxy/SingleParticlePhysicsProxy.cpp
@@ -100,6 +100,8 @@ void PushToPhysicsStateImp(const Chaos::FDirtyPropertiesManager& Manager, Chaos:
 				}
 				Handle->SetWorldSpaceInflatedBounds(WorldSpaceBounds);
 			}
+
+			Solver->GetEvolution()->DirtyParticle(*Handle);
 		}
 
 		if(auto NewData = ParticleData.FindMisc(Manager, DataIdx))
