Fix static analysis

diff --git a/Engine/Source/Runtime/Experimental/ChaosSolvers/Private/PhysicsProxy/SingleParticlePhysicsProxy.cpp b/Engine/Source/Runtime/Experimental/ChaosSolvers/Private/PhysicsProxy/SingleParticlePhysicsProxy.cpp
--- a/Engine/Source/Runtime/Experimental/ChaosSolvers/Private/PhysicsProxy/SingleParticlePhysicsProxy.cpp
+++ b/Engine/Source/Runtime/Experimental/ChaosSolvers/Private/PhysicsProxy/SingleParticlePhysicsProxy.cpp
@@ -40,8 +40,10 @@ template <Chaos::EParticleType ParticleType>
 void PushToPhysicsStateImp(const Chaos::FDirtyPropertiesManager& Manager, Chaos::TGeometryParticleHandle<float,3>* Handle, int32 DataIdx, const Chaos::FDirtyProxy& Dirty, Chaos::FShapeDirtyData* ShapesData, Chaos::FPhysicsSolver* Solver, const bool bInitialized)
 {
 	using namespace Chaos;
-	auto KinematicHandle = ParticleType == EParticleType::Static ? nullptr : static_cast<Chaos::TKinematicGeometryParticleHandle<float,3>*>(Handle);
-	auto RigidHandle = ParticleType != EParticleType::Rigid ? nullptr : static_cast<Chaos::TPBDRigidParticleHandle<float,3>*>(Handle);
+	constexpr bool bHasKinematicData = ParticleType != EParticleType::Static;
+	constexpr bool bHasDynamicData = ParticleType == EParticleType::Rigid;
+	auto KinematicHandle = bHasKinematicData ? static_cast<Chaos::TKinematicGeometryParticleHandle<float,3>*>(Handle) : nullptr;
+	auto RigidHandle = bHasDynamicData ? static_cast<Chaos::TPBDRigidParticleHandle<float,3>*>(Handle) : nullptr;
 	const FParticleDirtyData& ParticleData = Dirty.ParticleData;
 	// move the copied game thread data into the handle
 	{
@@ -57,7 +59,7 @@ void PushToPhysicsStateImp(const Chaos::FDirtyPropertiesManager& Manager, Chaos:
 		{
 			Handle->SetX(NewXR->X);
 			Handle->SetR(NewXR->R);
-			if(RigidHandle)
+			if(bHasDynamicData)
 			{
 				RigidHandle->SetP(NewXR->X);
 				RigidHandle->SetQ(NewXR->R);
@@ -75,12 +77,12 @@ void PushToPhysicsStateImp(const Chaos::FDirtyPropertiesManager& Manager, Chaos:
 #endif
 		}
 
-		auto NewVelocities = KinematicHandle != nullptr ? ParticleData.FindVelocities(Manager, DataIdx) : nullptr;
+		auto NewVelocities = bHasKinematicData ? ParticleData.FindVelocities(Manager, DataIdx) : nullptr;
 		if(NewVelocities)
 		{
 			KinematicHandle->SetV(NewVelocities->V);
 			KinematicHandle->SetW(NewVelocities->W);
-			if(RigidHandle)
+			if(bHasDynamicData)
 			{
 				bDynamicPropertyUpdated = true;
 			}
@@ -94,7 +96,7 @@ void PushToPhysicsStateImp(const Chaos::FDirtyPropertiesManager& Manager, Chaos:
 				Handle->SetHasBounds(true);
 				Handle->SetLocalBounds(Geometry->BoundingBox());
 				TAABB<float,3> WorldSpaceBounds = Geometry->BoundingBox().TransformedAABB(TRigidTransform<float,3>(Handle->X(),Handle->R()));
-				if(KinematicHandle)
+				if(bHasKinematicData)
 				{
 					WorldSpaceBounds.ThickenSymmetrically(KinematicHandle->V());
 				}
@@ -108,7 +110,7 @@ void PushToPhysicsStateImp(const Chaos::FDirtyPropertiesManager& Manager, Chaos:
 		{
 			Handle->SetSpatialIdx(NewData->SpatialIdx);
 
-			if(RigidHandle)
+			if(bHasDynamicData)
 			{
 				Solver->GetEvolution()->SetParticleObjectState(RigidHandle,NewData->ObjectState);
 				Solver->GetEvolution()->GetGravityForces().SetEnabled(*RigidHandle,NewData->bGravityEnabled);
@@ -117,13 +119,13 @@ void PushToPhysicsStateImp(const Chaos::FDirtyPropertiesManager& Manager, Chaos:
 
 		if(auto NewData = ParticleData.FindMassProps(Manager, DataIdx))
 		{
-			if(KinematicHandle)
+			if(bHasKinematicData)
 			{
 				KinematicHandle->SetCenterOfMass(NewData->CenterOfMass);
 				KinematicHandle->SetRotationOfMass(NewData->RotationOfMass);
 			}
 
-			if(RigidHandle)
+			if(bHasDynamicData)
 			{
 				RigidHandle->SetM(NewData->M);
 				RigidHandle->SetInvM(NewData->InvM);
@@ -132,7 +134,7 @@ void PushToPhysicsStateImp(const Chaos::FDirtyPropertiesManager& Manager, Chaos:
 			}
 		}
 
-		if(RigidHandle)
+		if(bHasDynamicData)
 		{
 			if(auto NewData = ParticleData.FindDynamics(Manager, DataIdx))
 			{
