Fixed for implicit type mismatch and bad serialization.
Revived SharedImplicitsAttribute to correctly handle the type change.
Undid breaking of old assets on load.
Removed const casting from dynamic collection initialization and safely handle attribute upgrades during geometry collection serialization.
Fixed attempted copies of managed attributes of different types.

diff --git a/Engine/Source/Runtime/Experimental/ChaosSolvers/Private/GeometryCollectionProxyData.cpp b/Engine/Source/Runtime/Experimental/ChaosSolvers/Private/GeometryCollectionProxyData.cpp
--- a/Engine/Source/Runtime/Experimental/ChaosSolvers/Private/GeometryCollectionProxyData.cpp
+++ b/Engine/Source/Runtime/Experimental/ChaosSolvers/Private/GeometryCollectionProxyData.cpp
@@ -44,6 +44,7 @@ const FName FGeometryDynamicCollection::ShapesQueryDataAttribute("ShapesQueryDat
 const FName FGeometryDynamicCollection::ShapesSimDataAttribute("ShapesSimData");
 const FName FGeometryDynamicCollection::SimplicialsAttribute("CollisionParticles");
 const FName FGeometryDynamicCollection::SimulatableParticlesAttribute("SimulatableParticlesAttribute");
+const FName FGeometryDynamicCollection::SharedImplicitsAttribute("SharedImplicits");
 
 FGeometryDynamicCollection::FGeometryDynamicCollection()
 	: FTransformDynamicCollection()
diff --git a/Engine/Source/Runtime/Experimental/ChaosSolvers/Private/PhysicsProxy/GeometryCollectionPhysicsProxy.cpp b/Engine/Source/Runtime/Experimental/ChaosSolvers/Private/PhysicsProxy/GeometryCollectionPhysicsProxy.cpp
--- a/Engine/Source/Runtime/Experimental/ChaosSolvers/Private/PhysicsProxy/GeometryCollectionPhysicsProxy.cpp
+++ b/Engine/Source/Runtime/Experimental/ChaosSolvers/Private/PhysicsProxy/GeometryCollectionPhysicsProxy.cpp
@@ -564,7 +564,7 @@ void FGeometryCollectionPhysicsProxy::InitializeDynamicCollection(FGeometryDynam
 	DynamicCollection.CopyMatchingAttributesFrom(RestCollection);
 
 	check(DynamicCollection.HasAttribute(FGeometryDynamicCollection::ImplicitsAttribute, FTransformCollection::TransformGroup));
-	check(RestCollection.HasAttribute(FGeometryDynamicCollection::ImplicitsAttribute, FTransformCollection::TransformGroup));
+	check(RestCollection.HasAttribute(FGeometryDynamicCollection::SharedImplicitsAttribute, FTransformCollection::TransformGroup));
 
 	//
 	// User defined initial velocities need to be populated. 
@@ -577,6 +577,29 @@ void FGeometryCollectionPhysicsProxy::InitializeDynamicCollection(FGeometryDynam
 		}
 	}
 
+	/** 
+		Copy over our implicits from the rest collection to the dynamic collection. This doesn't automatically happen
+		in the copy above due to the rest collection storing implicits in a different attribute that the dynamic collection
+		doesn't contain
+	*/
+	{
+		const TManagedArray<TSharedPtr<Chaos::FImplicitObject, ESPMode::ThreadSafe>>* SharedRestImplicits = RestCollection.FindAttribute<TSharedPtr<Chaos::FImplicitObject, ESPMode::ThreadSafe>>(FGeometryDynamicCollection::FGeometryDynamicCollection::SharedImplicitsAttribute, FTransformCollection::TransformGroup);
+
+		const int32 NumRestElems = RestCollection.NumElements(FTransformCollection::TransformGroup);
+		const int32 NumDynamicElems = DynamicCollection.NumElements(FTransformCollection::TransformGroup);
+
+		if(ensure(SharedRestImplicits && (NumRestElems == NumDynamicElems)))
+		{
+			for(int32 Index = 0; Index < NumDynamicElems; ++Index)
+			{
+				if(SharedRestImplicits)
+				{
+					DynamicCollection.Implicits[Index] = (*SharedRestImplicits)[Index];
+				}
+			}
+		}
+	}
+
 	// process simplicials
 	{
 		/*
@@ -1829,11 +1852,9 @@ void FGeometryCollectionPhysicsProxy::InitializeSharedCollisionStructures(
 	TManagedArray<TUniquePtr<FSimplicial>>& CollectionSimplicials = 
 		RestCollection.AddAttribute<TUniquePtr<FSimplicial>>(
 			FGeometryDynamicCollection::SimplicialsAttribute, FTransformCollection::TransformGroup);
-
-	RestCollection.RemoveAttribute(FGeometryDynamicCollection::ImplicitsAttribute, FTransformCollection::TransformGroup);
 	TManagedArray<FGeometryDynamicCollection::FSharedImplicit>& CollectionImplicits = 
 		RestCollection.AddAttribute<FGeometryDynamicCollection::FSharedImplicit>(
-			FGeometryDynamicCollection::ImplicitsAttribute, FTransformCollection::TransformGroup);
+			FGeometryDynamicCollection::SharedImplicitsAttribute, FTransformCollection::TransformGroup);
 
 	FTransform IdentityXf(FQuat::Identity, FVector(0));
 	IdentityXf.NormalizeRotation();
diff --git a/Engine/Source/Runtime/Experimental/ChaosSolvers/Public/GeometryCollectionProxyData.h b/Engine/Source/Runtime/Experimental/ChaosSolvers/Public/GeometryCollectionProxyData.h
--- a/Engine/Source/Runtime/Experimental/ChaosSolvers/Public/GeometryCollectionProxyData.h
+++ b/Engine/Source/Runtime/Experimental/ChaosSolvers/Public/GeometryCollectionProxyData.h
@@ -68,6 +68,7 @@ public:
 	static const FName ImplicitsAttribute;
 	static const FName ShapesQueryDataAttribute;
 	static const FName ShapesSimDataAttribute;
+	static const FName SharedImplicitsAttribute;
 	static const FName SimplicialsAttribute;
 	static const FName SimulatableParticlesAttribute;
 
diff --git a/Engine/Source/Runtime/Experimental/GeometryCollectionCore/Private/GeometryCollection/ManagedArrayCollection.cpp b/Engine/Source/Runtime/Experimental/GeometryCollectionCore/Private/GeometryCollection/ManagedArrayCollection.cpp
--- a/Engine/Source/Runtime/Experimental/GeometryCollectionCore/Private/GeometryCollection/ManagedArrayCollection.cpp
+++ b/Engine/Source/Runtime/Experimental/GeometryCollectionCore/Private/GeometryCollection/ManagedArrayCollection.cpp
@@ -288,7 +288,14 @@ void FManagedArrayCollection::CopyMatchingAttributesFrom(const FManagedArrayColl
 	{
 		if (InCollection.HasAttribute(Entry.Key.Get<0>(), Entry.Key.Get<1>()))
 		{
-			CopyAttribute(InCollection, Entry.Key.Get<0>(), Entry.Key.Get<1>());
+			const FValueType& OriginalValue = InCollection.Map[Entry.Key];
+			const FValueType& DestValue = Map[Entry.Key];
+
+			// If we don't have a type match don't attempt the copy.
+			if(OriginalValue.ArrayType == DestValue.ArrayType)
+			{
+				CopyAttribute(InCollection, Entry.Key.Get<0>(), Entry.Key.Get<1>());
+			}
 		}
 	}
 
diff --git a/Engine/Source/Runtime/Experimental/GeometryCollectionCore/Public/GeometryCollection/ManagedArrayCollection.h b/Engine/Source/Runtime/Experimental/GeometryCollectionCore/Public/GeometryCollection/ManagedArrayCollection.h
--- a/Engine/Source/Runtime/Experimental/GeometryCollectionCore/Public/GeometryCollection/ManagedArrayCollection.h
+++ b/Engine/Source/Runtime/Experimental/GeometryCollectionCore/Public/GeometryCollection/ManagedArrayCollection.h
@@ -189,6 +189,44 @@ public:
 		return nullptr;
 	};
 
+	/**
+	* Returns attribute(Name) of Type(T) from the group if and only if the types of T and the array match
+	* @param Name - The name of the attribute
+	* @param Group - The group that manages the attribute
+	* @return ManagedArray<T> &
+	*/
+	template<typename T>
+	TManagedArray<T>* FindAttributeTyped(FName Name, FName Group)
+	{
+		if(HasAttribute(Name, Group))
+		{
+			FKeyType Key = FManagedArrayCollection::MakeMapKey(Name, Group);
+			FValueType& FoundValue = Map[Key];
+
+			if(FoundValue.ArrayType == ManagedArrayType<T>())
+			{
+				return static_cast<TManagedArray<T>*>(Map[Key].Value);
+			}
+		}
+		return nullptr;
+	};
+
+	template<typename T>
+	const TManagedArray<T>* FindAttributeTyped(FName Name, FName Group) const
+	{
+		if(HasAttribute(Name, Group))
+		{
+			FKeyType Key = FManagedArrayCollection::MakeMapKey(Name, Group);
+			const FValueType& FoundValue = Map[Key];
+
+			if(FoundValue.ArrayType == ManagedArrayType<T>())
+			{
+				return static_cast<TManagedArray<T>*>(Map[Key].Value);
+			}
+		}
+		return nullptr;
+	};
+
 	/**
 	* Returns attribute access of Type(T) from the group
 	* @param Name - The name of the attribute
diff --git a/Engine/Source/Runtime/Experimental/GeometryCollectionEngine/Private/GeometryCollection/DerivedDataGeometryCollectionCooker.cpp b/Engine/Source/Runtime/Experimental/GeometryCollectionEngine/Private/GeometryCollection/DerivedDataGeometryCollectionCooker.cpp
--- a/Engine/Source/Runtime/Experimental/GeometryCollectionEngine/Private/GeometryCollection/DerivedDataGeometryCollectionCooker.cpp
+++ b/Engine/Source/Runtime/Experimental/GeometryCollectionEngine/Private/GeometryCollection/DerivedDataGeometryCollectionCooker.cpp
@@ -59,7 +59,7 @@ const TCHAR* FDerivedDataGeometryCollectionCooker::GetVersionString() const
 		return OverrideVersion;	//force load old ddc if found. Not recommended
 	}
 
-	return TEXT("ABC16C06507147A5A265CEDFAFDA96EE");
+	return TEXT("AE5BE323BD8A445E939DF316EB074D37");
 }
 
 FString FDerivedDataGeometryCollectionCooker::GetPluginSpecificCacheKeySuffix() const
diff --git a/Engine/Source/Runtime/Experimental/GeometryCollectionEngine/Private/GeometryCollection/GeometryCollectionObject.cpp b/Engine/Source/Runtime/Experimental/GeometryCollectionEngine/Private/GeometryCollection/GeometryCollectionObject.cpp
--- a/Engine/Source/Runtime/Experimental/GeometryCollectionEngine/Private/GeometryCollection/GeometryCollectionObject.cpp
+++ b/Engine/Source/Runtime/Experimental/GeometryCollectionEngine/Private/GeometryCollection/GeometryCollectionObject.cpp
@@ -21,6 +21,7 @@
 
 #include "GeometryCollection/GeometryCollectionSimulationCoreTypes.h"
 #include "Chaos/ChaosArchive.h"
+#include "GeometryCollectionProxyData.h"
 
 DEFINE_LOG_CATEGORY_STATIC(UGeometryCollectionLogging, NoLogging, All);
 
@@ -309,6 +310,25 @@ void UGeometryCollection::Serialize(FArchive& Ar)
 		}
 #endif
 		GeometryCollection->Serialize(ChaosAr);
+
+		// Fix up the type change for implicits here, previously they were unique ptrs, now they're shared
+		TManagedArray<TUniquePtr<Chaos::FImplicitObject>>* OldAttr = GeometryCollection->FindAttributeTyped<TUniquePtr<Chaos::FImplicitObject>>(FGeometryDynamicCollection::ImplicitsAttribute, FTransformCollection::TransformGroup);
+		TManagedArray<TSharedPtr<Chaos::FImplicitObject, ESPMode::ThreadSafe>>* NewAttr = GeometryCollection->FindAttributeTyped<TSharedPtr<Chaos::FImplicitObject, ESPMode::ThreadSafe>>(FGeometryDynamicCollection::SharedImplicitsAttribute, FTransformCollection::TransformGroup);
+		if(OldAttr)
+		{
+			if(!NewAttr)
+			{
+				NewAttr = &GeometryCollection->AddAttribute<TSharedPtr<Chaos::FImplicitObject, ESPMode::ThreadSafe>>(FGeometryDynamicCollection::SharedImplicitsAttribute, FTransformCollection::TransformGroup);
+
+				const int32 NumElems = GeometryCollection->NumElements(FTransformCollection::TransformGroup);
+				for(int32 Index = 0; Index < NumElems; ++Index)
+				{
+					(*NewAttr)[Index] = TSharedPtr<Chaos::FImplicitObject, ESPMode::ThreadSafe>((*OldAttr)[Index].Release());
+				}
+			}
+
+			GeometryCollection->RemoveAttribute(FGeometryDynamicCollection::ImplicitsAttribute, FTransformCollection::TransformGroup);
+		}
 	}
 
 	if (Ar.CustomVer(FDestructionObjectVersion::GUID) < FDestructionObjectVersion::GroupAndAttributeNameRemapping)
@@ -332,6 +352,9 @@ void UGeometryCollection::Serialize(FArchive& Ar)
 #if WITH_EDITOR
 void UGeometryCollection::CreateSimulationDataImp(bool bCopyFromDDC, const TCHAR* OverrideVersion)
 {
+	// Skips the DDC fetch entirely for testing the builder without adding to the DDC
+	const static bool bSkipDDC = false;
+
 	//Use the DDC to build simulation data. If we are loading in the editor we then serialize this data into the geometry collection
 	TArray<uint8> DDCData;
 	FDerivedDataGeometryCollectionCooker* GeometryCollectionCooker = new FDerivedDataGeometryCollectionCooker(*this);
@@ -339,8 +362,16 @@ void UGeometryCollection::CreateSimulationDataImp(bool bCopyFromDDC, const TCHAR
 
 	if (GeometryCollectionCooker->CanBuild())
 	{
-		GetDerivedDataCacheRef().GetSynchronous(GeometryCollectionCooker, DDCData);
-		if (bCopyFromDDC)
+		if(bSkipDDC)
+		{
+			GeometryCollectionCooker->Build(DDCData);
+		}
+		else
+		{
+			GetDerivedDataCacheRef().GetSynchronous(GeometryCollectionCooker, DDCData);
+		}
+
+		if(bCopyFromDDC)
 		{
 			FMemoryReader Ar(DDCData);
 			Chaos::FChaosArchive ChaosAr(Ar);
