Removed redundant ptr check in dynamic implicit initialization

diff --git a/Engine/Source/Runtime/Experimental/ChaosSolvers/Private/PhysicsProxy/GeometryCollectionPhysicsProxy.cpp b/Engine/Source/Runtime/Experimental/ChaosSolvers/Private/PhysicsProxy/GeometryCollectionPhysicsProxy.cpp
--- a/Engine/Source/Runtime/Experimental/ChaosSolvers/Private/PhysicsProxy/GeometryCollectionPhysicsProxy.cpp
+++ b/Engine/Source/Runtime/Experimental/ChaosSolvers/Private/PhysicsProxy/GeometryCollectionPhysicsProxy.cpp
@@ -592,10 +592,7 @@ void FGeometryCollectionPhysicsProxy::InitializeDynamicCollection(FGeometryDynam
 		{
 			for(int32 Index = 0; Index < NumDynamicElems; ++Index)
 			{
-				if(SharedRestImplicits)
-				{
-					DynamicCollection.Implicits[Index] = (*SharedRestImplicits)[Index];
-				}
+				DynamicCollection.Implicits[Index] = (*SharedRestImplicits)[Index];
 			}
 		}
 	}
