GeometryCollectionTests : resurrecting clustering test

diff --git a/Engine/Source/Runtime/Experimental/ChaosSolvers/Private/PhysicsProxy/GeometryCollectionPhysicsProxy.cpp b/Engine/Source/Runtime/Experimental/ChaosSolvers/Private/PhysicsProxy/GeometryCollectionPhysicsProxy.cpp
--- a/Engine/Source/Runtime/Experimental/ChaosSolvers/Private/PhysicsProxy/GeometryCollectionPhysicsProxy.cpp
+++ b/Engine/Source/Runtime/Experimental/ChaosSolvers/Private/PhysicsProxy/GeometryCollectionPhysicsProxy.cpp
@@ -1690,6 +1690,8 @@ void FGeometryCollectionPhysicsProxy::PullFromPhysicsState()
 			}
 
 			DynamicCollection.DynamicState[TmIndex] = TR.DynamicState[TmIndex];
+
+			DynamicCollection.Active[TmIndex] = !TR.DisabledStates[TmIndex];
 		}
 
 		//question: why do we need this? Sleeping objects will always have to update GPU
