Check for valid index when calculating number of rigids in new island in ComputeIslands.

diff --git a/Engine/Source/Runtime/Experimental/Chaos/Private/Chaos/PBDConstraintGraph.cpp b/Engine/Source/Runtime/Experimental/Chaos/Private/Chaos/PBDConstraintGraph.cpp
--- a/Engine/Source/Runtime/Experimental/Chaos/Private/Chaos/PBDConstraintGraph.cpp
+++ b/Engine/Source/Runtime/Experimental/Chaos/Private/Chaos/PBDConstraintGraph.cpp
@@ -488,7 +488,7 @@ void FPBDConstraintGraph::ComputeIslands(const TParticleView<TPBDRigidParticles<
 			else
 			{
 				int32 NumRigidsInIsland = 0, NumRigidsInNewIsland = 0;
-				if (OtherIsland != INDEX_NONE)
+				if (NewIslandParticles.IsValidIndex(OtherIsland))
 				{
 					for (TGeometryParticleHandle<FReal, 3>* Particle : IslandToParticles[Island])
 					{
