Make sure particles are created and cluster is disabled

diff --git a/Engine/Source/Runtime/Experimental/ChaosSolvers/Public/PhysicsProxy/GeometryCollectionPhysicsProxy.h b/Engine/Source/Runtime/Experimental/ChaosSolvers/Public/PhysicsProxy/GeometryCollectionPhysicsProxy.h
--- a/Engine/Source/Runtime/Experimental/ChaosSolvers/Public/PhysicsProxy/GeometryCollectionPhysicsProxy.h
+++ b/Engine/Source/Runtime/Experimental/ChaosSolvers/Public/PhysicsProxy/GeometryCollectionPhysicsProxy.h
@@ -151,6 +151,8 @@ public:
 	void SetCollisionParticlesPerObjectFraction(float CollisionParticlesPerObjectFractionIn) 
 	{CollisionParticlesPerObjectFraction = CollisionParticlesPerObjectFractionIn;}
 
+	TArray<FClusterHandle*>& GetSolverClusterHandles() { return SolverClusterHandles; }
+
 	TArray<FClusterHandle*>& GetSolverParticleHandles() { return SolverParticleHandles; }
 
 	const FGeometryCollectionResults* GetConsumerResultsGT() const 
