Chaos: Range check inclusive of max, matches clamp

diff --git a/Engine/Source/Runtime/Experimental/ChaosSolvers/Private/PhysicsProxy/GeometryCollectionPhysicsProxy.cpp b/Engine/Source/Runtime/Experimental/ChaosSolvers/Private/PhysicsProxy/GeometryCollectionPhysicsProxy.cpp
--- a/Engine/Source/Runtime/Experimental/ChaosSolvers/Private/PhysicsProxy/GeometryCollectionPhysicsProxy.cpp
+++ b/Engine/Source/Runtime/Experimental/ChaosSolvers/Private/PhysicsProxy/GeometryCollectionPhysicsProxy.cpp
@@ -251,7 +251,7 @@ void PopulateSimulatedParticle(
 	{
 		Handle->SetObjectState(Chaos::EObjectStateType::Uninitialized);
 
-		if (ensureMsgf(FMath::IsWithin(MassIn, SharedParams.MinimumMassClamp, SharedParams.MaximumMassClamp),
+		if (ensureMsgf(FMath::IsWithinInclusive(MassIn, SharedParams.MinimumMassClamp, SharedParams.MaximumMassClamp),
 			TEXT("Clamped mass[%3.5f] to range [%3.5f,%3.5f]"), MassIn, SharedParams.MinimumMassClamp, SharedParams.MaximumMassClamp))
 		{
 			MassIn = FMath::Clamp(MassIn, SharedParams.MinimumMassClamp, SharedParams.MaximumMassClamp);
@@ -262,9 +262,9 @@ void PopulateSimulatedParticle(
 		{
 			InertiaTensorVec = FVector(1);
 		}
-		else if (ensureMsgf(FMath::IsWithin(InertiaTensorVec[0], SharedParams.MinimumInertiaTensorDiagonalClamp, SharedParams.MaximumInertiaTensorDiagonalClamp)
-			&& FMath::IsWithin(InertiaTensorVec[1], SharedParams.MinimumInertiaTensorDiagonalClamp, SharedParams.MaximumInertiaTensorDiagonalClamp)
-			&& FMath::IsWithin(InertiaTensorVec[2], SharedParams.MinimumInertiaTensorDiagonalClamp, SharedParams.MaximumInertiaTensorDiagonalClamp),
+		else if (ensureMsgf(FMath::IsWithinInclusive(InertiaTensorVec[0], SharedParams.MinimumInertiaTensorDiagonalClamp, SharedParams.MaximumInertiaTensorDiagonalClamp)
+			&& FMath::IsWithinInclusive(InertiaTensorVec[1], SharedParams.MinimumInertiaTensorDiagonalClamp, SharedParams.MaximumInertiaTensorDiagonalClamp)
+			&& FMath::IsWithinInclusive(InertiaTensorVec[2], SharedParams.MinimumInertiaTensorDiagonalClamp, SharedParams.MaximumInertiaTensorDiagonalClamp),
 			TEXT("Clamped Inertia tensor[%3.5f,%3.5f,%3.5f]. Clamped each element to [%3.5f, %3.5f,]"), InertiaTensorVec[0], InertiaTensorVec[1], InertiaTensorVec[2],
 			SharedParams.MinimumInertiaTensorDiagonalClamp, SharedParams.MaximumInertiaTensorDiagonalClamp))
 		{
