Levelset sweep implementation

diff --git a/Engine/Source/Runtime/Experimental/Chaos/Private/Chaos/Levelset.cpp b/Engine/Source/Runtime/Experimental/Chaos/Private/Chaos/Levelset.cpp
--- a/Engine/Source/Runtime/Experimental/Chaos/Private/Chaos/Levelset.cpp
+++ b/Engine/Source/Runtime/Experimental/Chaos/Private/Chaos/Levelset.cpp
@@ -10,6 +10,10 @@
 #include "Chaos/ErrorReporter.h"
 #include "Chaos/MassProperties.h"
 #include "Chaos/TriangleMesh.h"
+#include "Chaos/Sphere.h"
+#include "Chaos/Box.h"
+#include "Chaos/Capsule.h"
+#include "Chaos/Convex.h"
 
 int32 OutputFailedLevelSetDebugData = 0;
 FAutoConsoleVariableRef CVarOutputFailedLevelSetDebugData(TEXT("p.LevelSetOutputFailedDebugData"), OutputFailedLevelSetDebugData, TEXT("Output debug obj files for level set and mesh when error tolerances are too high"));
@@ -1364,4 +1368,175 @@ T TLevelSet<T, d>::PhiWithNormal(const TVector<T, d>& x, TVector<T, d>& Normal)
 	return SizeSquared ? (sqrt(SizeSquared) + Phi) : Phi;
 }
 
+void GetGeomSurfaceSamples(const TSphere<FReal, 3>& InGeom, TArray<FVec3>& OutSamples)
+{
+	OutSamples.Reset();
+	OutSamples.AddUninitialized(6);
+
+	const FReal Radius = InGeom.GetRadius();
+
+	OutSamples[0] = FVec3(Radius, 0, 0);
+	OutSamples[1] = FVec3(-Radius, 0, 0);
+	OutSamples[2] = FVec3(0, Radius, Radius);
+	OutSamples[3] = FVec3(0, -Radius, Radius);
+	OutSamples[4] = FVec3(0, -Radius, -Radius);
+	OutSamples[5] = FVec3(0, Radius, -Radius);
+}
+
+void GetGeomSurfaceSamples(const TBox<FReal, 3>& InGeom, TArray<FVec3>& OutSamples)
+{
+	OutSamples.Reset();
+	OutSamples.AddUninitialized(8);
+
+	const FVec3& Min = InGeom.Min();
+	const FVec3& Max = InGeom.Max();
+
+	OutSamples[0] = FVec3(Min.X, Min.Y, Min.Z);
+	OutSamples[1] = FVec3(Min.X, Min.Y, Max.Z);
+	OutSamples[2] = FVec3(Min.X, Max.Y, Min.Z);
+	OutSamples[3] = FVec3(Max.X, Min.Y, Min.Z);
+	OutSamples[4] = FVec3(Max.X, Max.Y, Max.Z);
+	OutSamples[5] = FVec3(Max.X, Max.Y, Min.Z);
+	OutSamples[6] = FVec3(Max.X, Min.Y, Max.Z);
+	OutSamples[7] = FVec3(Min.X, Max.Y, Max.Z);
+}
+
+void GetGeomSurfaceSamples(const TCapsule<FReal>& InGeom, TArray<FVec3>& OutSamples)
+{
+	OutSamples.Reset();
+	OutSamples.AddUninitialized(14);
+
+	FReal Radius = InGeom.GetRadius();
+	FReal HalfHeight = InGeom.GetHeight() * 0.5f;
+
+	OutSamples[0]	= FVec3(HalfHeight + Radius, 0, 0);
+	OutSamples[1]	= FVec3(-HalfHeight - Radius, 0, 0);
+	OutSamples[2]	= FVec3(HalfHeight, Radius, Radius);
+	OutSamples[3]	= FVec3(HalfHeight, -Radius, Radius);
+	OutSamples[4]	= FVec3(HalfHeight, -Radius, -Radius);
+	OutSamples[5]	= FVec3(HalfHeight, Radius, -Radius);
+	OutSamples[6]	= FVec3(0, Radius, Radius);
+	OutSamples[7]	= FVec3(0, -Radius, Radius);
+	OutSamples[8]	= FVec3(0, -Radius, -Radius);
+	OutSamples[9]	= FVec3(0, Radius, -Radius);
+	OutSamples[10]	= FVec3(-HalfHeight, Radius, Radius);
+	OutSamples[11]	= FVec3(-HalfHeight, -Radius, Radius);
+	OutSamples[12]	= FVec3(-HalfHeight, -Radius, -Radius);
+	OutSamples[13]	= FVec3(-HalfHeight, Radius, -Radius);
+}
+
+void GetGeomSurfaceSamples(const FConvex& InGeom, TArray<FVec3>& OutSamples)
+{
+	const TParticles<FReal, 3>& Particles = InGeom.GetSurfaceParticles();
+
+	const int32 NumParticles = Particles.Size();
+
+	OutSamples.Reset();
+	OutSamples.AddUninitialized(NumParticles);
+
+	for(int32 Index = 0; Index < NumParticles; ++Index)
+	{
+		OutSamples[Index] = Particles.X(Index);
+	}
+}
+
+template<typename InnerT>
+void GetGeomSurfaceSamples(const TImplicitObjectScaled<InnerT>& InScaledGeom, TArray<FVec3>& OutSamples)
+{
+	const InnerT* InnerObject = InScaledGeom.Object().Get();
+
+	if(InnerObject)
+	{
+		GetGeomSurfaceSamples(*InnerObject, OutSamples);
+
+		const FVec3 Scale = InScaledGeom.GetScale();
+
+		for(FVec3& Sample : OutSamples)
+		{
+			Sample *= Scale;
+		}
+	}
+}
+
+template<class T, int d>
+template <typename QueryGeomType>
+bool TLevelSet<T, d>::SweepGeomImp(const QueryGeomType& QueryGeom, const FRigidTransform3& StartTM, const FVec3& Dir, const FReal Length,
+	FReal& OutTime, FVec3& OutPosition, FVec3& OutNormal, int32& OutFaceIndex, const FReal Thickness, const bool bComputeMTD, FVec3 TriMeshScale) const
+{
+	TArray<FVec3> Samples;
+	GetGeomSurfaceSamples(QueryGeom, Samples);
+
+	OutTime = TNumericLimits<FReal>::Max();
+	FReal TempTime = TNumericLimits<FReal>::Max();
+	FVec3 TempNormal(0);
+	FVec3 TempPosition(0);
+	int32 FaceIndex;
+
+	bool bHit = false;
+
+	for(const FVec3& Sample : Samples)
+	{
+		FVec3 Transformed = StartTM.TransformPosition(Sample);
+		Raycast(Transformed, Dir, Length, 0, TempTime, TempPosition, TempNormal, FaceIndex);
+
+		if(TempTime < OutTime)
+		{
+			OutTime = TempTime;
+			OutPosition = TempPosition;
+			OutNormal = TempNormal;
+			bHit = true;
+		}
+	}
+
+	return bHit;
+}
+
+template<class T, int d>
+bool TLevelSet<T, d>::SweepGeom(const TSphere<FReal, 3>& QueryGeom, const FRigidTransform3& StartTM, const FVec3& Dir, const FReal Length, FReal& OutTime, FVec3& OutPosition, FVec3& OutNormal, int32& OutFaceIndex, const FReal Thickness, const bool bComputeMTD) const
+{
+	return SweepGeomImp(QueryGeom, StartTM, Dir, Length, OutTime, OutPosition, OutNormal, OutFaceIndex, Thickness, bComputeMTD);
+}
+
+template<class T, int d>
+bool TLevelSet<T, d>::SweepGeom(const TBox<FReal, 3>& QueryGeom, const FRigidTransform3& StartTM, const FVec3& Dir, const FReal Length, FReal& OutTime, FVec3& OutPosition, FVec3& OutNormal, int32& OutFaceIndex, const FReal Thickness, const bool bComputeMTD) const
+{
+	return SweepGeomImp(QueryGeom, StartTM, Dir, Length, OutTime, OutPosition, OutNormal, OutFaceIndex, Thickness, bComputeMTD);
+}
+
+template<class T, int d>
+bool TLevelSet<T, d>::SweepGeom(const TCapsule<FReal>& QueryGeom, const FRigidTransform3& StartTM, const FVec3& Dir, const FReal Length, FReal& OutTime, FVec3& OutPosition, FVec3& OutNormal, int32& OutFaceIndex, const FReal Thickness, const bool bComputeMTD) const
+{
+	return SweepGeomImp(QueryGeom, StartTM, Dir, Length, OutTime, OutPosition, OutNormal, OutFaceIndex, Thickness, bComputeMTD);
+}
+
+template<class T, int d>
+bool TLevelSet<T, d>::SweepGeom(const FConvex& QueryGeom, const FRigidTransform3& StartTM, const FVec3& Dir, const FReal Length, FReal& OutTime, FVec3& OutPosition, FVec3& OutNormal, int32& OutFaceIndex, const FReal Thickness, const bool bComputeMTD) const
+{
+	return SweepGeomImp(QueryGeom, StartTM, Dir, Length, OutTime, OutPosition, OutNormal, OutFaceIndex, Thickness, bComputeMTD);
+}
+
+template<class T, int d>
+bool TLevelSet<T, d>::SweepGeom(const TImplicitObjectScaled<TSphere<FReal, 3>>& QueryGeom, const FRigidTransform3& StartTM, const FVec3& Dir, const FReal Length, FReal& OutTime, FVec3& OutPosition, FVec3& OutNormal, int32& OutFaceIndex, const FReal Thickness, const bool bComputeMTD, FVec3 TriMeshScale) const
+{
+	return SweepGeomImp(QueryGeom, StartTM, Dir, Length, OutTime, OutPosition, OutNormal, OutFaceIndex, Thickness, bComputeMTD, TriMeshScale);
+}
+
+template<class T, int d>
+bool TLevelSet<T, d>::SweepGeom(const TImplicitObjectScaled<TBox<FReal, 3>>& QueryGeom, const FRigidTransform3& StartTM, const FVec3& Dir, const FReal Length, FReal& OutTime, FVec3& OutPosition, FVec3& OutNormal, int32& OutFaceIndex, const FReal Thickness, const bool bComputeMTD, FVec3 TriMeshScale) const
+{
+	return SweepGeomImp(QueryGeom, StartTM, Dir, Length, OutTime, OutPosition, OutNormal, OutFaceIndex, Thickness, bComputeMTD, TriMeshScale);
+}
+
+template<class T, int d>
+bool TLevelSet<T, d>::SweepGeom(const TImplicitObjectScaled<TCapsule<FReal>>& QueryGeom, const FRigidTransform3& StartTM, const FVec3& Dir, const FReal Length, FReal& OutTime, FVec3& OutPosition, FVec3& OutNormal, int32& OutFaceIndex, const FReal Thickness, const bool bComputeMTD, FVec3 TriMeshScale) const
+{
+	return SweepGeomImp(QueryGeom, StartTM, Dir, Length, OutTime, OutPosition, OutNormal, OutFaceIndex, Thickness, bComputeMTD, TriMeshScale);
+}
+
+template<class T, int d>
+bool TLevelSet<T, d>::SweepGeom(const TImplicitObjectScaled<FConvex>& QueryGeom, const FRigidTransform3& StartTM, const FVec3& Dir, const FReal Length, FReal& OutTime, FVec3& OutPosition, FVec3& OutNormal, int32& OutFaceIndex, const FReal Thickness, const bool bComputeMTD, FVec3 TriMeshScale) const
+{
+	return SweepGeomImp(QueryGeom, StartTM, Dir, Length, OutTime, OutPosition, OutNormal, OutFaceIndex, Thickness, bComputeMTD, TriMeshScale);
+}
+
 template class Chaos::TLevelSet<float, 3>;
diff --git a/Engine/Source/Runtime/Experimental/Chaos/Public/Chaos/GeometryQueries.h b/Engine/Source/Runtime/Experimental/Chaos/Public/Chaos/GeometryQueries.h
--- a/Engine/Source/Runtime/Experimental/Chaos/Public/Chaos/GeometryQueries.h
+++ b/Engine/Source/Runtime/Experimental/Chaos/Public/Chaos/GeometryQueries.h
@@ -12,6 +12,7 @@
 #include "ImplicitObjectScaled.h"
 #include "Chaos/Box.h"
 #include "Chaos/Sphere.h"
+#include "Chaos/Levelset.h"
 
 #include "ChaosArchive.h"
 #include <algorithm>
@@ -186,19 +187,25 @@ namespace Chaos
 			//todo: pass bComputeMTD into these functions
 			switch (AType)
 			{
-			case ImplicitObjectType::HeightField:
-			{
-				const auto& AHeightField = static_cast<const THeightField<FReal>&>(A);
-				bResult = AHeightField.SweepGeom(B, BToATM, LocalDir, Length, OutTime, LocalPosition, LocalNormal, OutFaceIndex, Thickness, bComputeMTD);
-				break;
-			}
-			case ImplicitObjectType::TriangleMesh:
-			{
-				const auto& ATriangleMesh = static_cast<const FTriangleMeshImplicitObject&>(A);
-				bResult = ATriangleMesh.SweepGeom(B, BToATM, LocalDir, Length, OutTime, LocalPosition, LocalNormal, OutFaceIndex, Thickness, bComputeMTD);
-				break;
-			}
-			default:
+				case ImplicitObjectType::HeightField:
+				{
+					const THeightField<FReal>& AHeightField = static_cast<const THeightField<FReal>&>(A);
+					bResult = AHeightField.SweepGeom(B, BToATM, LocalDir, Length, OutTime, LocalPosition, LocalNormal, OutFaceIndex, Thickness, bComputeMTD);
+					break;
+				}
+				case ImplicitObjectType::TriangleMesh:
+				{
+					const FTriangleMeshImplicitObject& ATriangleMesh = static_cast<const FTriangleMeshImplicitObject&>(A);
+					bResult = ATriangleMesh.SweepGeom(B, BToATM, LocalDir, Length, OutTime, LocalPosition, LocalNormal, OutFaceIndex, Thickness, bComputeMTD);
+					break;
+				}
+				case ImplicitObjectType::LevelSet:
+				{
+					const TLevelSet<FReal, 3>& ALevelSet = static_cast<const TLevelSet<FReal, 3>&>(A);
+					bResult = ALevelSet.SweepGeom(B, BToATM, LocalDir, Length, OutTime, LocalPosition, LocalNormal, OutFaceIndex, Thickness, bComputeMTD);
+					break;
+				}
+				default:
 				if (IsScaled(AType))
 				{
 					const auto& AScaled = TImplicitObjectScaled<FTriangleMeshImplicitObject>::AsScaledChecked(A);
@@ -215,7 +222,6 @@ namespace Chaos
 				{
 					ensureMsgf(false, TEXT("Unsupported query type: %u"), (uint8)AType);
 				}
-
 			}
 		}
 
diff --git a/Engine/Source/Runtime/Experimental/Chaos/Public/Chaos/Levelset.h b/Engine/Source/Runtime/Experimental/Chaos/Public/Chaos/Levelset.h
--- a/Engine/Source/Runtime/Experimental/Chaos/Public/Chaos/Levelset.h
+++ b/Engine/Source/Runtime/Experimental/Chaos/Public/Chaos/Levelset.h
@@ -4,6 +4,7 @@
 #include "Chaos/ArrayND.h"
 #include "Chaos/Box.h"
 #include "Chaos/ImplicitObject.h"
+#include "Chaos/ImplicitObjectScaled.h"
 #include "Chaos/Particles.h"
 #include "Chaos/UniformGrid.h"
 
@@ -18,6 +19,11 @@ class TTriangleMesh;
 template<class T, int D>
 class TPlane;
 
+template<typename T>
+class TCapsule;
+
+class FConvex;
+
 template<class T, int d>
 class CHAOS_API TLevelSet final : public FImplicitObject
 {
@@ -123,6 +129,20 @@ class CHAOS_API TLevelSet final : public FImplicitObject
 		return Result;
 	}
 
+
+	bool SweepGeom(const TSphere<FReal, 3>& QueryGeom, const FRigidTransform3& StartTM, const FVec3& Dir, const FReal Length, FReal& OutTime, FVec3& OutPosition, FVec3& OutNormal, int32& OutFaceIndex, const FReal Thickness = 0, const bool bComputeMTD = false) const;
+	bool SweepGeom(const TBox<FReal, 3>& QueryGeom, const FRigidTransform3& StartTM, const FVec3& Dir, const FReal Length, FReal& OutTime, FVec3& OutPosition, FVec3& OutNormal, int32& OutFaceIndex, const FReal Thickness = 0, const bool bComputeMTD = false) const;
+	bool SweepGeom(const TCapsule<FReal>& QueryGeom, const FRigidTransform3& StartTM, const FVec3& Dir, const FReal Length, FReal& OutTime, FVec3& OutPosition, FVec3& OutNormal, int32& OutFaceIndex, const FReal Thickness = 0, const bool bComputeMTD = false) const;
+	bool SweepGeom(const FConvex& QueryGeom, const FRigidTransform3& StartTM, const FVec3& Dir, const FReal Length, FReal& OutTime, FVec3& OutPosition, FVec3& OutNormal, int32& OutFaceIndex, const FReal Thickness = 0, const bool bComputeMTD = false) const;
+
+	bool SweepGeom(const TImplicitObjectScaled<TSphere<FReal, 3>>& QueryGeom, const FRigidTransform3& StartTM, const FVec3& Dir, const FReal Length, FReal& OutTime, FVec3& OutPosition, FVec3& OutNormal, int32& OutFaceIndex, const FReal Thickness = 0, const bool bComputeMTD = false, FVec3 TriMeshScale = FVec3(1.0f)) const;
+	bool SweepGeom(const TImplicitObjectScaled<TBox<FReal, 3>>& QueryGeom, const FRigidTransform3& StartTM, const FVec3& Dir, const FReal Length, FReal& OutTime, FVec3& OutPosition, FVec3& OutNormal, int32& OutFaceIndex, const FReal Thickness = 0, const bool bComputeMTD = false, FVec3 TriMeshScale = FVec3(1.0f)) const;
+	bool SweepGeom(const TImplicitObjectScaled<TCapsule<FReal>>& QueryGeom, const FRigidTransform3& StartTM, const FVec3& Dir, const FReal Length, FReal& OutTime, FVec3& OutPosition, FVec3& OutNormal, int32& OutFaceIndex, const FReal Thickness = 0, const bool bComputeMTD = false, FVec3 TriMeshScale = FVec3(1.0f)) const;
+	bool SweepGeom(const TImplicitObjectScaled<FConvex>& QueryGeom, const FRigidTransform3& StartTM, const FVec3& Dir, const FReal Length, FReal& OutTime, FVec3& OutPosition, FVec3& OutNormal, int32& OutFaceIndex, const FReal Thickness = 0, const bool bComputeMTD = false, FVec3 TriMeshScale = FVec3(1.0f)) const;
+
+	template<typename QueryGeomType>
+	bool SweepGeomImp(const QueryGeomType& QueryGeom, const FRigidTransform3& StartTM, const FVec3& Dir, const FReal Length, FReal& OutTime, FVec3& OutPosition, FVec3& OutNormal, int32& OutFaceIndex, const FReal Thickness, const bool bComputeMTD, FVec3 TriMeshScale = FVec3(1.0f)) const;
+
   private:
 	bool ComputeDistancesNearZeroIsocontour(FErrorReporter& ErrorReporter, const TParticles<T, d>& InParticles, const TArray<TVector<T, 3>> &Normals, const TTriangleMesh<T>& Mesh, TArrayND<bool, d>& BlockedFaceX, TArrayND<bool, d>& BlockedFaceY, TArrayND<bool, d>& BlockedFaceZ, TArray<TVector<int32, d>>& InterfaceIndices);
 	void ComputeDistancesNearZeroIsocontour(const FImplicitObject& Object, const TArrayND<T, d>& ObjectPhi, TArray<TVector<int32, d>>& InterfaceIndices);
