Make SetObjectState(Sleep) work like PhysX's PutToSleep. That is, it buffers zero velocity and cleans dynamic properties.

diff --git a/Engine/Source/Runtime/Experimental/Chaos/Public/Chaos/ParticleHandle.h b/Engine/Source/Runtime/Experimental/Chaos/Public/Chaos/ParticleHandle.h
--- a/Engine/Source/Runtime/Experimental/Chaos/Public/Chaos/ParticleHandle.h
+++ b/Engine/Source/Runtime/Experimental/Chaos/Public/Chaos/ParticleHandle.h
@@ -2070,7 +2070,23 @@ public:
 		{
 			MAwakeEvent |= true;
 		}
+
+		if (InState == EObjectStateType::Sleeping)
+		{
+			// When an object is forced into a sleep state, the velocities must be zeroed and buffered,
+			// in case the velocity is queried during sleep, or in case the object is woken up again.
+			SetV(FVec3(0.f), true);
+			SetW(FVec3(0.f), true);
+
+			// Dynamic particle properties must be marked clean in order not to actually apply forces which
+			// have been buffered. If another force is added after the object is put to sleep, the old forces
+			// will remain and the new ones will accumulate and re-dirty the dynamic properties which will
+			// wake the body.
+			MDirtyFlags.MarkClean(ParticlePropToFlag(EParticleProperty::Dynamics));
+		}
+
 		MMiscData.Modify(true,MDirtyFlags,Proxy,[&InState](auto& Data){ Data.ObjectState = InState;});
+
 	}
 
 	void ClearEvents() { MAwakeEvent = false; }
