Help out linux with crtp

diff --git a/Engine/Source/Runtime/Experimental/Chaos/Public/Chaos/ParticleHandle.h b/Engine/Source/Runtime/Experimental/Chaos/Public/Chaos/ParticleHandle.h
--- a/Engine/Source/Runtime/Experimental/Chaos/Public/Chaos/ParticleHandle.h
+++ b/Engine/Source/Runtime/Experimental/Chaos/Public/Chaos/ParticleHandle.h
@@ -2075,8 +2075,8 @@ public:
 		{
 			// When an object is forced into a sleep state, the velocities must be zeroed and buffered,
 			// in case the velocity is queried during sleep, or in case the object is woken up again.
-			SetV(FVec3(0.f), true);
-			SetW(FVec3(0.f), true);
+			this->SetV(FVec3(0.f), true);
+			this->SetW(FVec3(0.f), true);
 
 			// Dynamic particle properties must be marked clean in order not to actually apply forces which
 			// have been buffered. If another force is added after the object is put to sleep, the old forces
