Fixing bad merge.

diff --git a/Engine/Source/Runtime/Experimental/Chaos/Public/Chaos/PBDRigidParticles.h b/Engine/Source/Runtime/Experimental/Chaos/Public/Chaos/PBDRigidParticles.h
--- a/Engine/Source/Runtime/Experimental/Chaos/Public/Chaos/PBDRigidParticles.h
+++ b/Engine/Source/Runtime/Experimental/Chaos/Public/Chaos/PBDRigidParticles.h
@@ -66,10 +66,10 @@ class TPBDRigidParticles : public TRigidParticles<T, d>
 
     // Must be reinterpret cast instead of static_cast as it's a forward declare
 	typedef TPBDRigidParticleHandle<T, d> THandleType;
-	CHAOS_API const THandleType* Handle(int32 Index) const { return static_cast<const THandleType*>(TGeometryParticles<T,d>::Handle(Index)); }
+	CHAOS_API const THandleType* Handle(int32 Index) const { return reinterpret_cast<const THandleType*>(TGeometryParticles<T,d>::Handle(Index)); }
 
 	//cannot be reference because double pointer would allow for badness, but still useful to have non const access to handle
-	CHAOS_API THandleType* Handle(int32 Index) { return static_cast<THandleType*>(TGeometryParticles<T, d>::Handle(Index)); }
+	CHAOS_API THandleType* Handle(int32 Index) { return reinterpret_cast<THandleType*>(TGeometryParticles<T, d>::Handle(Index)); }
 
 	CHAOS_API void SetSleeping(int32 Index, bool bSleeping)
 	{
@@ -82,7 +82,7 @@ class TPBDRigidParticles : public TRigidParticles<T, d>
 		bool CurrentlySleeping = this->ObjectState(Index) == EObjectStateType::Sleeping;
 		if (CurrentlySleeping != bSleeping)
 		{
-			TGeometryParticleHandle<T, d>* Particle = static_cast<TGeometryParticleHandle<T, d>*>(this->Handle(Index));
+			TGeometryParticleHandle<T, d>* Particle = reinterpret_cast<TGeometryParticleHandle<T, d>*>(this->Handle(Index));
 			this->AddSleepData(Particle, bSleeping);
 		}
 
@@ -147,7 +147,7 @@ class TPBDRigidParticles : public TRigidParticles<T, d>
 		const bool bNewSleeping = InObjectState == EObjectStateType::Sleeping;
 		if(bCurrentSleeping != bNewSleeping)
 		{
-			TGeometryParticleHandle<T, d>* Particle = static_cast<TGeometryParticleHandle<T, d>*>(this->Handle(Index));
+			TGeometryParticleHandle<T, d>* Particle = reinterpret_cast<TGeometryParticleHandle<T, d>*>(this->Handle(Index));
  			this->AddSleepData(Particle, bNewSleeping);
 		}
 
