Chaos: Fix field system registration/unregistration crash

diff --git a/Engine/Source/Runtime/Engine/Private/PhysicsEngine/Experimental/PhysScene_Chaos.cpp b/Engine/Source/Runtime/Engine/Private/PhysicsEngine/Experimental/PhysScene_Chaos.cpp
--- a/Engine/Source/Runtime/Engine/Private/PhysicsEngine/Experimental/PhysScene_Chaos.cpp
+++ b/Engine/Source/Runtime/Engine/Private/PhysicsEngine/Experimental/PhysScene_Chaos.cpp
@@ -599,7 +599,7 @@ void FPhysScene_Chaos::AddObject(UPrimitiveComponent* Component, FFieldSystemPhy
 
 	if (Chaos::IDispatcher* Dispatcher = GetDispatcher())
 	{
-		TArray<Chaos::FPhysicsSolver*> WorldSolverList = ChaosModule->GetSolversMutable(Component->GetWorld());
+		TArray<Chaos::FPhysicsSolver*> WorldSolverList = ChaosModule->GetAllSolvers();
 
 		for(Chaos::FPhysicsSolver* Solver : WorldSolverList)
 		{
@@ -807,13 +807,13 @@ void FPhysScene_Chaos::RemoveObject(FFieldSystemPhysicsProxy* InObject)
 
 		if(Chaos::IDispatcher* Dispatcher = GetDispatcher())
 		{
-			TArray<Chaos::FPhysicsSolver*> SolverList = ChaosModule->GetSolversMutable(CurrSceneSolver->GetOwner());
+			TArray<Chaos::FPhysicsSolver*> SolverList = ChaosModule->GetAllSolvers();
 
 			for(Chaos::FPhysicsSolver* Solver : SolverList)
 			{
 				if(true || Solver->HasActiveParticles())
 				{
-					Solver->RegisterObject(InObject);
+					Solver->UnregisterObject(InObject);
 
 					if(/*bDedicatedThread && */Dispatcher)
 					{
diff --git a/Engine/Source/Runtime/Experimental/ChaosSolvers/Private/PhysicsProxy/FieldSystemPhysicsProxy.cpp b/Engine/Source/Runtime/Experimental/ChaosSolvers/Private/PhysicsProxy/FieldSystemPhysicsProxy.cpp
--- a/Engine/Source/Runtime/Experimental/ChaosSolvers/Private/PhysicsProxy/FieldSystemPhysicsProxy.cpp
+++ b/Engine/Source/Runtime/Experimental/ChaosSolvers/Private/PhysicsProxy/FieldSystemPhysicsProxy.cpp
@@ -415,33 +415,6 @@ void FFieldSystemPhysicsProxy::FieldParameterUpdateCallback(
 							}
 							++i;
 						}
-
-						// #todo: no special casing for the floor
-						/*
-						bool bHasFloor = false;
-						int32 FloorIndex = CurrentSolver->GetFloorIndex();
-						if (FloorIndex != INDEX_NONE)
-						{
-							bHasFloor = !Particles.Disabled(FloorIndex);
-						}
-
-						TSet<uint32> RemovedParticles;
-						for (const ContextIndex& Index : IndicesArray)
-						{
-							const int32 i = Index.Result;
-							if (!Particles.Disabled(i) && Results[i] > 0.0)
-							{
-								RemovedParticles.Add((uint32)i);
-								CurrentSolver->GetEvolution()->DisableParticle(i);
-							}
-						}
-
-						if (RemovedParticles.Num() && bHasFloor)
-						{
-							CurrentSolver->GetEvolution()->DisableParticle(FloorIndex);
-							Particles.SetObjectState(FloorIndex, Chaos::EObjectStateType::Static);
-						}
-						*/
 					}
 				}
 				CommandsToRemove.Add(CommandIndex);
@@ -946,8 +919,12 @@ void FFieldSystemPhysicsProxy::FieldForcesUpdateCallback(
 						for (Chaos::TGeometryParticleHandle<float, 3>* Handle : Handles)
 						{
 							Chaos::TPBDRigidParticleHandle<float, 3>* RigidHandle = Handle->CastToRigidParticle();
-							if(RigidHandle && RigidHandle->ObjectState() == Chaos::EObjectStateType::Dynamic)
-							{								
+							if(RigidHandle && (RigidHandle->ObjectState() == Chaos::EObjectStateType::Dynamic || RigidHandle->ObjectState() == Chaos::EObjectStateType::Sleeping))
+							{				
+								if (RigidHandle->Sleeping())
+								{
+									RigidHandle->SetObjectState(Chaos::EObjectStateType::Dynamic);
+								}
 								RigidHandle->F() += ForceView[i];
 							}
 							++i;
@@ -1006,8 +983,12 @@ void FFieldSystemPhysicsProxy::FieldForcesUpdateCallback(
 						for (Chaos::TGeometryParticleHandle<float, 3>* Handle : Handles)
 						{
 							Chaos::TPBDRigidParticleHandle<float, 3>* RigidHandle = Handle->CastToRigidParticle();
-							if(RigidHandle && RigidHandle->ObjectState() == Chaos::EObjectStateType::Dynamic)
+							if(RigidHandle && (RigidHandle->ObjectState() == Chaos::EObjectStateType::Dynamic || RigidHandle->ObjectState() == Chaos::EObjectStateType::Sleeping))
 							{
+								if (RigidHandle->Sleeping())
+								{
+									RigidHandle->SetObjectState(Chaos::EObjectStateType::Dynamic);
+								}
 								RigidHandle->Torque() += TorqueView[i];
 							}
 							++i;
diff --git a/Engine/Source/Runtime/Experimental/FieldSystem/Source/FieldSystemEngine/Private/Field/FieldSystemComponent.cpp b/Engine/Source/Runtime/Experimental/FieldSystem/Source/FieldSystemEngine/Private/Field/FieldSystemComponent.cpp
--- a/Engine/Source/Runtime/Experimental/FieldSystem/Source/FieldSystemEngine/Private/Field/FieldSystemComponent.cpp
+++ b/Engine/Source/Runtime/Experimental/FieldSystem/Source/FieldSystemEngine/Private/Field/FieldSystemComponent.cpp
@@ -152,8 +152,7 @@ void UFieldSystemComponent::DispatchCommand(const FFieldSystemCommand& InCommand
 			}
 		}
 
-		TArray<Chaos::FPhysicsSolver*> WorldSolverList;
-		WorldSolverList = ChaosModule->GetSolversMutable(GetWorld());
+		TArray<Chaos::FPhysicsSolver*> WorldSolverList = ChaosModule->GetAllSolvers();
 
 		// #BGTODO Currently all commands will end up actually executing a frame late. That's because this command has to be logged as a global command
 		// so we don't end up with multiple solver threads writing to the proxy. We need a better way to buffer up multi-solver commands so they can be
