Levelset overlap implementation
Fixed TSphere sample point generation for tiny radii

diff --git a/Engine/Source/Runtime/Engine/Private/Collision/CollisionConversions.cpp b/Engine/Source/Runtime/Engine/Private/Collision/CollisionConversions.cpp
--- a/Engine/Source/Runtime/Engine/Private/Collision/CollisionConversions.cpp
+++ b/Engine/Source/Runtime/Engine/Private/Collision/CollisionConversions.cpp
@@ -559,7 +559,23 @@ void ConvertQueryOverlap(const FPhysicsShape& Shape, const FPhysicsActor& Actor,
 #endif
 	else
 	{
-		ensureMsgf(false, TEXT("ConvertQueryOverlap called with bad payload type"));
+#if WITH_CHAOS
+		// Currently geom collections are registered with a primitive component user data, but maybe custom should be adapted
+		// to be more general so we can support leaf identification #BGTODO
+		void* UserData = Actor.UserData();
+		UPrimitiveComponent* PossibleOwner = FPhysxUserData::Get<UPrimitiveComponent>(UserData);
+
+		if(PossibleOwner)
+		{
+			OutOverlap.Component = PossibleOwner;
+			OutOverlap.Actor = OutOverlap.Component->GetOwner();
+			OutOverlap.ItemIndex = INDEX_NONE;
+		}
+		else
+#endif
+		{
+			ensureMsgf(false, TEXT("ConvertQueryOverlap called with bad payload type"));
+		}
 	}
 
 	// Other info
diff --git a/Engine/Source/Runtime/Experimental/Chaos/Private/Chaos/Levelset.cpp b/Engine/Source/Runtime/Experimental/Chaos/Private/Chaos/Levelset.cpp
--- a/Engine/Source/Runtime/Experimental/Chaos/Private/Chaos/Levelset.cpp
+++ b/Engine/Source/Runtime/Experimental/Chaos/Private/Chaos/Levelset.cpp
@@ -14,6 +14,7 @@
 #include "Chaos/Box.h"
 #include "Chaos/Capsule.h"
 #include "Chaos/Convex.h"
+#include "Chaos/GeometryQueries.h"
 
 int32 OutputFailedLevelSetDebugData = 0;
 FAutoConsoleVariableRef CVarOutputFailedLevelSetDebugData(TEXT("p.LevelSetOutputFailedDebugData"), OutputFailedLevelSetDebugData, TEXT("Output debug obj files for level set and mesh when error tolerances are too high"));
@@ -30,6 +31,12 @@ FAutoConsoleVariableRef CVarMaxDistErrorTolerance(TEXT("p.LevelSetMaxDistErrorTo
 float AvgAngleErrorTolerance = 1.;
 FAutoConsoleVariableRef CVarAvgAngleErrorTolerance(TEXT("p.LevelSetAvgAngleErrorTolerance"), AvgAngleErrorTolerance, TEXT("Average error in of the mesh normal and computed normal on the level set."));
 
+int32 NumOverlapSphereSamples = 16;
+FAutoConsoleVariableRef CVarNumOverlapSphereSamples(TEXT("p.LevelsetOverlapSphereSamples"), NumOverlapSphereSamples, TEXT("Number of spiral points to generate for levelset-sphere overlaps"));
+
+int32 NumOverlapCapsuleSamples = 24;
+FAutoConsoleVariableRef CVarNumOverlapCapsuleSamples(TEXT("p.LevelsetOverlapCapsuleSamples"), NumOverlapCapsuleSamples, TEXT("Number of spiral points to generate for levelset-capsule overlaps"));
+
 using namespace Chaos;
 
 #define MAX_CLAMP(a, comp, b) (a >= comp ? b : a)
@@ -1461,7 +1468,7 @@ void GetGeomSurfaceSamples(const TImplicitObjectScaled<InnerT>& InScaledGeom, TA
 template<class T, int d>
 template <typename QueryGeomType>
 bool TLevelSet<T, d>::SweepGeomImp(const QueryGeomType& QueryGeom, const FRigidTransform3& StartTM, const FVec3& Dir, const FReal Length,
-	FReal& OutTime, FVec3& OutPosition, FVec3& OutNormal, int32& OutFaceIndex, const FReal Thickness, const bool bComputeMTD, FVec3 TriMeshScale) const
+	FReal& OutTime, FVec3& OutPosition, FVec3& OutNormal, int32& OutFaceIndex, const FReal Thickness, const bool bComputeMTD) const
 {
 	TArray<FVec3> Samples;
 	GetGeomSurfaceSamples(QueryGeom, Samples);
@@ -1516,27 +1523,160 @@ bool TLevelSet<T, d>::SweepGeom(const FConvex& QueryGeom, const FRigidTransform3
 }
 
 template<class T, int d>
-bool TLevelSet<T, d>::SweepGeom(const TImplicitObjectScaled<TSphere<FReal, 3>>& QueryGeom, const FRigidTransform3& StartTM, const FVec3& Dir, const FReal Length, FReal& OutTime, FVec3& OutPosition, FVec3& OutNormal, int32& OutFaceIndex, const FReal Thickness, const bool bComputeMTD, FVec3 TriMeshScale) const
+bool TLevelSet<T, d>::SweepGeom(const TImplicitObjectScaled<TSphere<FReal, 3>>& QueryGeom, const FRigidTransform3& StartTM, const FVec3& Dir, const FReal Length, FReal& OutTime, FVec3& OutPosition, FVec3& OutNormal, int32& OutFaceIndex, const FReal Thickness, const bool bComputeMTD) const
+{
+	return SweepGeomImp(QueryGeom, StartTM, Dir, Length, OutTime, OutPosition, OutNormal, OutFaceIndex, Thickness, bComputeMTD);
+}
+
+template<class T, int d>
+bool TLevelSet<T, d>::SweepGeom(const TImplicitObjectScaled<TBox<FReal, 3>>& QueryGeom, const FRigidTransform3& StartTM, const FVec3& Dir, const FReal Length, FReal& OutTime, FVec3& OutPosition, FVec3& OutNormal, int32& OutFaceIndex, const FReal Thickness, const bool bComputeMTD) const
+{
+	return SweepGeomImp(QueryGeom, StartTM, Dir, Length, OutTime, OutPosition, OutNormal, OutFaceIndex, Thickness, bComputeMTD);
+}
+
+template<class T, int d>
+bool TLevelSet<T, d>::SweepGeom(const TImplicitObjectScaled<TCapsule<FReal>>& QueryGeom, const FRigidTransform3& StartTM, const FVec3& Dir, const FReal Length, FReal& OutTime, FVec3& OutPosition, FVec3& OutNormal, int32& OutFaceIndex, const FReal Thickness, const bool bComputeMTD) const
+{
+	return SweepGeomImp(QueryGeom, StartTM, Dir, Length, OutTime, OutPosition, OutNormal, OutFaceIndex, Thickness, bComputeMTD);
+}
+
+template<class T, int d>
+bool TLevelSet<T, d>::SweepGeom(const TImplicitObjectScaled<FConvex>& QueryGeom, const FRigidTransform3& StartTM, const FVec3& Dir, const FReal Length, FReal& OutTime, FVec3& OutPosition, FVec3& OutNormal, int32& OutFaceIndex, const FReal Thickness, const bool bComputeMTD) const
+{
+	return SweepGeomImp(QueryGeom, StartTM, Dir, Length, OutTime, OutPosition, OutNormal, OutFaceIndex, Thickness, bComputeMTD);
+}
+
+void GetGeomSurfaceSamplesExtended(const TSphere<FReal, 3>& InGeom, TArray<FVec3>& OutSamples)
+{
+	OutSamples = InGeom.ComputeLocalSamplePoints(NumOverlapSphereSamples);
+}
+
+void GetGeomSurfaceSamplesExtended(const TBox<FReal, 3>& InGeom, TArray<FVec3>& OutSamples)
+{
+	OutSamples = InGeom.ComputeLocalSamplePoints();
+}
+
+void GetGeomSurfaceSamplesExtended(const TCapsule<FReal>& InGeom, TArray<FVec3>& OutSamples)
+{
+	OutSamples = InGeom.ComputeLocalSamplePoints(NumOverlapCapsuleSamples);
+}
+
+void GetGeomSurfaceSamplesExtended(const FConvex& InGeom, TArray<FVec3>& OutSamples)
+{
+	// Convex doesn't have extended samples
+	GetGeomSurfaceSamples(InGeom, OutSamples);
+}
+
+template<typename InnerT>
+void GetGeomSurfaceSamplesExtended(const TImplicitObjectScaled<InnerT>& InScaledGeom, TArray<FVec3>& OutSamples)
+{
+	const InnerT* InnerObject = InScaledGeom.Object().Get();
+
+	if(InnerObject)
+	{
+		GetGeomSurfaceSamplesExtended(*InnerObject, OutSamples);
+
+		const FVec3 Scale = InScaledGeom.GetScale();
+
+		for(FVec3& Sample : OutSamples)
+		{
+			Sample *= Scale;
+		}
+	}
+}
+
+template<class T, int d>
+template <typename QueryGeomType>
+bool TLevelSet<T, d>::OverlapGeomImp(const QueryGeomType& QueryGeom, const FRigidTransform3& QueryTM, const FReal Thickness, FMTDInfo* OutMTD) const
+{
+	// NOTE: This isn't a perfect overlap implementation. It takes particle samples from the query
+	// geoemetry and looks for intersections, often this means that we're only detecting on the surface
+	// of the query geometry.
+	// #BGTODO better sampling of levelset, then invert the check for levelset point inside querygeom
+	bool bResult = false;
+
+	if(OutMTD)
+	{
+		OutMTD->Normal = FVec3(0);
+		OutMTD->Penetration = 0;
+	}
+
+	TArray<FVec3> SamplePoints;
+	FVec3 TempNormal;
+	FReal TempPhi;
+
+	// Use an extended set of points here to attempt to get a better overlap
+	GetGeomSurfaceSamplesExtended(QueryGeom, SamplePoints);
+
+	for(const FVec3& Sample : SamplePoints)
+	{
+		const FVec3 Transformed = QueryTM.TransformPosition(Sample);
+		TempPhi = PhiWithNormal(Transformed, TempNormal);
+
+		if(OutMTD && (-TempPhi) > OutMTD->Penetration)
+		{
+			OutMTD->Penetration = -TempPhi;
+			OutMTD->Normal = TempNormal;
+			bResult = true;
+		}
+		else
+		{
+			if(TempPhi <= 0)
+			{
+				return true;
+			}
+		}
+	}
+
+	return bResult;
+}
+
+template<class T, int d>
+bool TLevelSet<T, d>::OverlapGeom(const TSphere<FReal, 3>& QueryGeom, const FRigidTransform3& QueryTM, const FReal Thickness, FMTDInfo* OutMTD) const
+{
+	return OverlapGeomImp(QueryGeom, QueryTM, Thickness, OutMTD);
+}
+
+template<class T, int d>
+bool TLevelSet<T, d>::OverlapGeom(const TBox<FReal, 3>& QueryGeom, const FRigidTransform3& QueryTM, const FReal Thickness, FMTDInfo* OutMTD) const
+{
+	return OverlapGeomImp(QueryGeom, QueryTM, Thickness, OutMTD);
+}
+
+template<class T, int d>
+bool TLevelSet<T, d>::OverlapGeom(const TCapsule<FReal>& QueryGeom, const FRigidTransform3& QueryTM, const FReal Thickness, FMTDInfo* OutMTD) const
+{
+	return OverlapGeomImp(QueryGeom, QueryTM, Thickness, OutMTD);
+}
+
+template<class T, int d>
+bool TLevelSet<T, d>::OverlapGeom(const FConvex& QueryGeom, const FRigidTransform3& QueryTM, const FReal Thickness, FMTDInfo* OutMTD) const
+{
+	return OverlapGeomImp(QueryGeom, QueryTM, Thickness, OutMTD);
+}
+
+template<class T, int d>
+bool TLevelSet<T, d>::OverlapGeom(const TImplicitObjectScaled<TSphere<FReal, 3>>& QueryGeom, const FRigidTransform3& QueryTM, const FReal Thickness, FMTDInfo* OutMTD) const
 {
-	return SweepGeomImp(QueryGeom, StartTM, Dir, Length, OutTime, OutPosition, OutNormal, OutFaceIndex, Thickness, bComputeMTD, TriMeshScale);
+	return OverlapGeomImp(QueryGeom, QueryTM, Thickness, OutMTD);
 }
 
 template<class T, int d>
-bool TLevelSet<T, d>::SweepGeom(const TImplicitObjectScaled<TBox<FReal, 3>>& QueryGeom, const FRigidTransform3& StartTM, const FVec3& Dir, const FReal Length, FReal& OutTime, FVec3& OutPosition, FVec3& OutNormal, int32& OutFaceIndex, const FReal Thickness, const bool bComputeMTD, FVec3 TriMeshScale) const
+bool TLevelSet<T, d>::OverlapGeom(const TImplicitObjectScaled<TBox<FReal, 3>>& QueryGeom, const FRigidTransform3& QueryTM, const FReal Thickness, FMTDInfo* OutMTD) const
 {
-	return SweepGeomImp(QueryGeom, StartTM, Dir, Length, OutTime, OutPosition, OutNormal, OutFaceIndex, Thickness, bComputeMTD, TriMeshScale);
+	return OverlapGeomImp(QueryGeom, QueryTM, Thickness, OutMTD);
 }
 
 template<class T, int d>
-bool TLevelSet<T, d>::SweepGeom(const TImplicitObjectScaled<TCapsule<FReal>>& QueryGeom, const FRigidTransform3& StartTM, const FVec3& Dir, const FReal Length, FReal& OutTime, FVec3& OutPosition, FVec3& OutNormal, int32& OutFaceIndex, const FReal Thickness, const bool bComputeMTD, FVec3 TriMeshScale) const
+bool TLevelSet<T, d>::OverlapGeom(const TImplicitObjectScaled<TCapsule<FReal>>& QueryGeom, const FRigidTransform3& QueryTM, const FReal Thickness, FMTDInfo* OutMTD) const
 {
-	return SweepGeomImp(QueryGeom, StartTM, Dir, Length, OutTime, OutPosition, OutNormal, OutFaceIndex, Thickness, bComputeMTD, TriMeshScale);
+	return OverlapGeomImp(QueryGeom, QueryTM, Thickness, OutMTD);
 }
 
 template<class T, int d>
-bool TLevelSet<T, d>::SweepGeom(const TImplicitObjectScaled<FConvex>& QueryGeom, const FRigidTransform3& StartTM, const FVec3& Dir, const FReal Length, FReal& OutTime, FVec3& OutPosition, FVec3& OutNormal, int32& OutFaceIndex, const FReal Thickness, const bool bComputeMTD, FVec3 TriMeshScale) const
+bool TLevelSet<T, d>::OverlapGeom(const TImplicitObjectScaled<FConvex>& QueryGeom, const FRigidTransform3& QueryTM, const FReal Thickness, FMTDInfo* OutMTD) const
 {
-	return SweepGeomImp(QueryGeom, StartTM, Dir, Length, OutTime, OutPosition, OutNormal, OutFaceIndex, Thickness, bComputeMTD, TriMeshScale);
+	return OverlapGeomImp(QueryGeom, QueryTM, Thickness, OutMTD);
 }
 
 template class Chaos::TLevelSet<float, 3>;
diff --git a/Engine/Source/Runtime/Experimental/Chaos/Public/Chaos/GeometryQueries.h b/Engine/Source/Runtime/Experimental/Chaos/Public/Chaos/GeometryQueries.h
--- a/Engine/Source/Runtime/Experimental/Chaos/Public/Chaos/GeometryQueries.h
+++ b/Engine/Source/Runtime/Experimental/Chaos/Public/Chaos/GeometryQueries.h
@@ -80,30 +80,37 @@ namespace Chaos
 		{
 			switch (AType)
 			{
-			case ImplicitObjectType::HeightField:
-			{
-				const THeightField<FReal>& AHeightField = static_cast<const THeightField<FReal>&>(A);
-				return AHeightField.OverlapGeom(B, BToATM, Thickness, OutMTD);
-			}
-			case ImplicitObjectType::TriangleMesh:
-			{
-				const FTriangleMeshImplicitObject& ATriangleMesh = static_cast<const FTriangleMeshImplicitObject&>(A);
-				return ATriangleMesh.OverlapGeom(B, BToATM, Thickness, OutMTD);
-			}
-			default:
-				if (IsScaled(AType))
+				case ImplicitObjectType::HeightField:
 				{
-					const auto& AScaled = TImplicitObjectScaled<FTriangleMeshImplicitObject>::AsScaledChecked(A);
-					return AScaled.LowLevelOverlapGeom(B, BToATM, Thickness, OutMTD);
+					const THeightField<FReal>& AHeightField = static_cast<const THeightField<FReal>&>(A);
+					return AHeightField.OverlapGeom(B, BToATM, Thickness, OutMTD);
 				}
-				else if(IsInstanced(AType))
+				case ImplicitObjectType::TriangleMesh:
 				{
-					const auto& AInstanced = TImplicitObjectInstanced<FTriangleMeshImplicitObject>::AsInstancedChecked(A);
-					return AInstanced.LowLevelOverlapGeom(B,BToATM,Thickness, OutMTD);
+					const FTriangleMeshImplicitObject& ATriangleMesh = static_cast<const FTriangleMeshImplicitObject&>(A);
+					return ATriangleMesh.OverlapGeom(B, BToATM, Thickness, OutMTD);
 				}
-				else
+				case ImplicitObjectType::LevelSet:
+				{
+					const TLevelSet<FReal, 3>& ALevelSet = static_cast<const TLevelSet<FReal, 3>&>(A);
+					return ALevelSet.OverlapGeom(B, BToATM, Thickness, OutMTD);
+				}
+				default:
 				{
-					check(false);	//unsupported query type
+					if(IsScaled(AType))
+					{
+						const auto& AScaled = TImplicitObjectScaled<FTriangleMeshImplicitObject>::AsScaledChecked(A);
+						return AScaled.LowLevelOverlapGeom(B, BToATM, Thickness, OutMTD);
+					}
+					else if(IsInstanced(AType))
+					{
+						const auto& AInstanced = TImplicitObjectInstanced<FTriangleMeshImplicitObject>::AsInstancedChecked(A);
+						return AInstanced.LowLevelOverlapGeom(B, BToATM, Thickness, OutMTD);
+					}
+					else
+					{
+						check(false);	//unsupported query type
+					}
 				}
 			}
 		}
diff --git a/Engine/Source/Runtime/Experimental/Chaos/Public/Chaos/Levelset.h b/Engine/Source/Runtime/Experimental/Chaos/Public/Chaos/Levelset.h
--- a/Engine/Source/Runtime/Experimental/Chaos/Public/Chaos/Levelset.h
+++ b/Engine/Source/Runtime/Experimental/Chaos/Public/Chaos/Levelset.h
@@ -135,13 +135,26 @@ class CHAOS_API TLevelSet final : public FImplicitObject
 	bool SweepGeom(const TCapsule<FReal>& QueryGeom, const FRigidTransform3& StartTM, const FVec3& Dir, const FReal Length, FReal& OutTime, FVec3& OutPosition, FVec3& OutNormal, int32& OutFaceIndex, const FReal Thickness = 0, const bool bComputeMTD = false) const;
 	bool SweepGeom(const FConvex& QueryGeom, const FRigidTransform3& StartTM, const FVec3& Dir, const FReal Length, FReal& OutTime, FVec3& OutPosition, FVec3& OutNormal, int32& OutFaceIndex, const FReal Thickness = 0, const bool bComputeMTD = false) const;
 
-	bool SweepGeom(const TImplicitObjectScaled<TSphere<FReal, 3>>& QueryGeom, const FRigidTransform3& StartTM, const FVec3& Dir, const FReal Length, FReal& OutTime, FVec3& OutPosition, FVec3& OutNormal, int32& OutFaceIndex, const FReal Thickness = 0, const bool bComputeMTD = false, FVec3 TriMeshScale = FVec3(1.0f)) const;
-	bool SweepGeom(const TImplicitObjectScaled<TBox<FReal, 3>>& QueryGeom, const FRigidTransform3& StartTM, const FVec3& Dir, const FReal Length, FReal& OutTime, FVec3& OutPosition, FVec3& OutNormal, int32& OutFaceIndex, const FReal Thickness = 0, const bool bComputeMTD = false, FVec3 TriMeshScale = FVec3(1.0f)) const;
-	bool SweepGeom(const TImplicitObjectScaled<TCapsule<FReal>>& QueryGeom, const FRigidTransform3& StartTM, const FVec3& Dir, const FReal Length, FReal& OutTime, FVec3& OutPosition, FVec3& OutNormal, int32& OutFaceIndex, const FReal Thickness = 0, const bool bComputeMTD = false, FVec3 TriMeshScale = FVec3(1.0f)) const;
-	bool SweepGeom(const TImplicitObjectScaled<FConvex>& QueryGeom, const FRigidTransform3& StartTM, const FVec3& Dir, const FReal Length, FReal& OutTime, FVec3& OutPosition, FVec3& OutNormal, int32& OutFaceIndex, const FReal Thickness = 0, const bool bComputeMTD = false, FVec3 TriMeshScale = FVec3(1.0f)) const;
+	bool SweepGeom(const TImplicitObjectScaled<TSphere<FReal, 3>>& QueryGeom, const FRigidTransform3& StartTM, const FVec3& Dir, const FReal Length, FReal& OutTime, FVec3& OutPosition, FVec3& OutNormal, int32& OutFaceIndex, const FReal Thickness = 0, const bool bComputeMTD = false) const;
+	bool SweepGeom(const TImplicitObjectScaled<TBox<FReal, 3>>& QueryGeom, const FRigidTransform3& StartTM, const FVec3& Dir, const FReal Length, FReal& OutTime, FVec3& OutPosition, FVec3& OutNormal, int32& OutFaceIndex, const FReal Thickness = 0, const bool bComputeMTD = false) const;
+	bool SweepGeom(const TImplicitObjectScaled<TCapsule<FReal>>& QueryGeom, const FRigidTransform3& StartTM, const FVec3& Dir, const FReal Length, FReal& OutTime, FVec3& OutPosition, FVec3& OutNormal, int32& OutFaceIndex, const FReal Thickness = 0, const bool bComputeMTD = false) const;
+	bool SweepGeom(const TImplicitObjectScaled<FConvex>& QueryGeom, const FRigidTransform3& StartTM, const FVec3& Dir, const FReal Length, FReal& OutTime, FVec3& OutPosition, FVec3& OutNormal, int32& OutFaceIndex, const FReal Thickness = 0, const bool bComputeMTD = false) const;
 
 	template<typename QueryGeomType>
-	bool SweepGeomImp(const QueryGeomType& QueryGeom, const FRigidTransform3& StartTM, const FVec3& Dir, const FReal Length, FReal& OutTime, FVec3& OutPosition, FVec3& OutNormal, int32& OutFaceIndex, const FReal Thickness, const bool bComputeMTD, FVec3 TriMeshScale = FVec3(1.0f)) const;
+	bool SweepGeomImp(const QueryGeomType& QueryGeom, const FRigidTransform3& StartTM, const FVec3& Dir, const FReal Length, FReal& OutTime, FVec3& OutPosition, FVec3& OutNormal, int32& OutFaceIndex, const FReal Thickness, const bool bComputeMTD) const;
+
+	bool OverlapGeom(const TSphere<FReal, 3>& QueryGeom, const FRigidTransform3& QueryTM, const FReal Thickness, FMTDInfo* OutMTD = nullptr) const;
+	bool OverlapGeom(const TBox<FReal, 3>& QueryGeom, const FRigidTransform3& QueryTM, const FReal Thickness, FMTDInfo* OutMTD = nullptr) const;
+	bool OverlapGeom(const TCapsule<FReal>& QueryGeom, const FRigidTransform3& QueryTM, const FReal Thickness, FMTDInfo* OutMTD = nullptr) const;
+	bool OverlapGeom(const FConvex& QueryGeom, const FRigidTransform3& QueryTM, const FReal Thickness, FMTDInfo* OutMTD = nullptr) const;
+
+	bool OverlapGeom(const TImplicitObjectScaled<TSphere<FReal, 3>>& QueryGeom, const FRigidTransform3& QueryTM, const FReal Thickness, FMTDInfo* OutMTD = nullptr) const;
+	bool OverlapGeom(const TImplicitObjectScaled<TBox<FReal, 3>>& QueryGeom, const FRigidTransform3& QueryTM, const FReal Thickness, FMTDInfo* OutMTD = nullptr) const;
+	bool OverlapGeom(const TImplicitObjectScaled<TCapsule<FReal>>& QueryGeom, const FRigidTransform3& QueryTM, const FReal Thickness, FMTDInfo* OutMTD = nullptr) const;
+	bool OverlapGeom(const TImplicitObjectScaled<FConvex>& QueryGeom, const FRigidTransform3& QueryTM, const FReal Thickness, FMTDInfo* OutMTD = nullptr) const;
+
+	template<typename QueryGeomType>
+	bool OverlapGeomImp(const QueryGeomType& QueryGeom, const FRigidTransform3& QueryTM, const FReal Thickness, FMTDInfo* OutMTD = nullptr) const;
 
   private:
 	bool ComputeDistancesNearZeroIsocontour(FErrorReporter& ErrorReporter, const TParticles<T, d>& InParticles, const TArray<TVector<T, 3>> &Normals, const TTriangleMesh<T>& Mesh, TArrayND<bool, d>& BlockedFaceX, TArrayND<bool, d>& BlockedFaceY, TArrayND<bool, d>& BlockedFaceZ, TArray<TVector<int32, d>>& InterfaceIndices);
diff --git a/Engine/Source/Runtime/Experimental/Chaos/Public/Chaos/Sphere.h b/Engine/Source/Runtime/Experimental/Chaos/Public/Chaos/Sphere.h
--- a/Engine/Source/Runtime/Experimental/Chaos/Public/Chaos/Sphere.h
+++ b/Engine/Source/Runtime/Experimental/Chaos/Public/Chaos/Sphere.h
@@ -266,8 +266,19 @@ namespace Chaos
 		TArray<TVector<T, d>> ComputeLocalSamplePoints(const int NumPoints) const
 		{
 			TArray<TVector<T, d>> Points;
-			TSphere<T, d> LocalSphere(TVector<T, d>(0.0), Radius);
-			TSphereSpecializeSamplingHelper<T, d>::ComputeSamplePoints(Points, LocalSphere, NumPoints);
+			
+			if(Radius <= KINDA_SMALL_NUMBER)
+			{
+				// If we're too small (and will create NaNs) then just take the centre
+				Points.Add(TVector<T, d>(0.0));
+			}
+			else
+			{
+				Points.Reserve(NumPoints);
+				TSphere<T, d> LocalSphere(TVector<T, d>(0.0), Radius);
+				TSphereSpecializeSamplingHelper<T, d>::ComputeSamplePoints(Points, LocalSphere, NumPoints);
+			}
+
 			return Points;
 		}
 
