First simple test and implementation for recording physics state

diff --git a/Engine/Source/Programs/HeadlessChaos/Private/HeadlessChaosTestRewind.cpp b/Engine/Source/Programs/HeadlessChaos/Private/HeadlessChaosTestRewind.cpp
new file mode 100644
--- /dev/null
+++ b/Engine/Source/Programs/HeadlessChaos/Private/HeadlessChaosTestRewind.cpp
@@ -0,0 +1,78 @@
+// Copyright Epic Games, Inc. All Rights Reserved.
+
+#include "HeadlessChaosTestUtility.h"
+#include "Chaos/ParticleHandle.h"
+#include "Chaos/ErrorReporter.h"
+#include "PhysicsProxy/SingleParticlePhysicsProxy.h"
+#include "PhysicsProxy/GeometryCollectionPhysicsProxy.h"
+#include "Chaos/Utilities.h"
+#include "PBDRigidsSolver.h"
+#include "ChaosSolversModule.h"
+
+#include "Modules/ModuleManager.h"
+#include "Framework/PhysicsTickTask.h"
+#include "Chaos/RewindData.h"
+
+
+namespace ChaosTest {
+
+    using namespace Chaos;
+
+	GTEST_TEST(RewindTest,DataCapture)
+	{
+		auto Sphere = TSharedPtr<FImplicitObject,ESPMode::ThreadSafe>(new TSphere<float,3>(TVector<float,3>(0),10));
+
+		FChaosSolversModule* Module = FChaosSolversModule::GetModule();
+		Module->ChangeThreadingMode(EChaosThreadingMode::SingleThread);
+
+		// Make a solver
+		Chaos::FPhysicsSolver* Solver = Module->CreateSolver(nullptr,ESolverFlags::Standalone);
+		Solver->SetEnabled(true);
+
+		// Make a particle
+
+		auto Particle = Chaos::TKinematicGeometryParticle<float,3>::CreateParticle();
+
+		Particle->SetGeometry(Sphere);
+		Solver->RegisterObject(Particle.Get());
+		Solver->EnableRewindCapture(20);
+
+		for(int Step = 0; Step < 11; ++Step)
+		{
+			Particle->SetX(FVec3(0,0,100 - Step));
+
+			Solver->PushPhysicsState(Module->GetDispatcher());
+			FPhysicsSolverAdvanceTask AdvanceTask(Solver,1.0f);
+			AdvanceTask.DoTask(ENamedThreads::GameThread,FGraphEventRef());
+			Solver->BufferPhysicsResults();
+			Solver->FlipBuffers();
+			Solver->UpdateGameThreadStructures();
+		}
+
+		//ended up at z = 90
+		EXPECT_EQ(Particle->X()[2],90);
+
+		const FRewindData* RewindData = Solver->GetRewindData();
+
+		//check state at every step except latest
+		for(int Step = 0; Step < 10; ++Step)
+		{
+			const auto ParticleState = RewindData->GetStateAtFrame(*Particle,Step);
+
+			EXPECT_TRUE(ParticleState != nullptr);
+			if(ParticleState)
+			{
+				EXPECT_EQ(ParticleState->X()[2],100 - Step);
+			}
+		}
+		
+		//no data at head because we always save the previous frame
+		EXPECT_TRUE(RewindData->GetStateAtFrame(*Particle,10) == nullptr);
+		
+		// Throw out the proxy
+		Solver->UnregisterObject(Particle.Get());
+
+		Module->DestroySolver(Solver);
+	}
+
+}
diff --git a/Engine/Source/Runtime/Experimental/Chaos/Public/Chaos/Framework/PhysicsSolverBase.h b/Engine/Source/Runtime/Experimental/Chaos/Public/Chaos/Framework/PhysicsSolverBase.h
--- a/Engine/Source/Runtime/Experimental/Chaos/Public/Chaos/Framework/PhysicsSolverBase.h
+++ b/Engine/Source/Runtime/Experimental/Chaos/Public/Chaos/Framework/PhysicsSolverBase.h
@@ -41,6 +41,113 @@ namespace Chaos
 			}
 		}
 	};
+
+	class FDirtySet
+	{
+	public:
+		void Add(IPhysicsProxyBase* Base)
+		{
+			if(Base->GetDirtyIdx() == INDEX_NONE)
+			{
+				const int32 Idx = ProxiesData.Num();
+				Base->SetDirtyIdx(Idx);
+				ProxiesData.Add(Base);
+			}
+		}
+
+		void Remove(IPhysicsProxyBase* Base)
+		{
+			const int32 Idx = Base->GetDirtyIdx();
+			if(Idx != INDEX_NONE)
+			{
+				if(Idx == ProxiesData.Num() - 1)
+				{
+					//last element so just pop
+					ProxiesData.Pop(/*bAllowShrinking=*/false);
+				} else
+				{
+					//update other proxy's idx
+					ProxiesData.RemoveAtSwap(Idx);
+					ProxiesData[Idx].SetDirtyIdx(Idx);
+				}
+
+				Base->ResetDirtyIdx();
+			}
+		}
+
+		void Reset()
+		{
+			ProxiesData.Reset();
+			ShapesData.Reset();
+		}
+
+		int32 NumDirtyProxies() const { return ProxiesData.Num(); }
+		int32 NumDirtyShapes() const { return ShapesData.Num(); }
+
+		FShapeDirtyData* GetShapesDirtyData(){ return ShapesData.GetData(); }
+
+		template <typename Lambda>
+		void ParallelForEachProxy(const Lambda& Func)
+		{
+			::ParallelFor(ProxiesData.Num(),[this,&Func](int32 Idx)
+			{
+				Func(Idx,ProxiesData[Idx]);
+			});
+		}
+
+		template <typename Lambda>
+		void ParallelForEachProxy(const Lambda& Func) const
+		{
+			::ParallelFor(ProxiesData.Num(),[this,&Func](int32 Idx)
+			{
+				Func(Idx,ProxiesData[Idx]);
+			});
+		}
+
+		template <typename Lambda>
+		void ForEachProxy(const Lambda& Func)
+		{
+			int32 Idx = 0;
+			for(FDirtyProxy& Dirty : ProxiesData)
+			{
+				Func(Idx++,Dirty);
+			}
+		}
+
+		void AddShape(IPhysicsProxyBase* Proxy,int32 ShapeIdx)
+		{
+			Add(Proxy);
+			FDirtyProxy& Dirty = ProxiesData[Proxy->GetDirtyIdx()];
+			for(int32 NewShapeIdx = Dirty.ShapeDataIndices.Num(); NewShapeIdx <= ShapeIdx; ++NewShapeIdx)
+			{
+				const int32 ShapeDataIdx = ShapesData.Add(FShapeDirtyData(NewShapeIdx));
+				Dirty.AddShape(ShapeDataIdx);
+			}
+		}
+
+		void SetNumDirtyShapes(IPhysicsProxyBase* Proxy,int32 NumShapes)
+		{
+			Add(Proxy);
+			FDirtyProxy& Dirty = ProxiesData[Proxy->GetDirtyIdx()];
+
+			if(NumShapes < Dirty.ShapeDataIndices.Num())
+			{
+				Dirty.ShapeDataIndices.SetNum(NumShapes);
+			} else
+			{
+				for(int32 NewShapeIdx = Dirty.ShapeDataIndices.Num(); NewShapeIdx < NumShapes; ++NewShapeIdx)
+				{
+					const int32 ShapeDataIdx = ShapesData.Add(FShapeDirtyData(NewShapeIdx));
+					Dirty.AddShape(ShapeDataIdx);
+				}
+			}
+		}
+
+	private:
+		TArray<FDirtyProxy> ProxiesData;
+		TArray<FShapeDirtyData> ShapesData;
+	};
+
 	class CHAOS_API FPhysicsSolverBase
 	{
 	public:
@@ -93,105 +200,6 @@ namespace Chaos
 		/** Mode that the results buffers should be set to (single, double, triple) */
 		EMultiBufferMode BufferMode;
 
-		class FDirtySet
-		{
-		public:
-			void Add(IPhysicsProxyBase* Base)
-			{
-				if(Base->GetDirtyIdx() == INDEX_NONE)
-				{
-					const int32 Idx = ProxiesData.Num();
-					Base->SetDirtyIdx(Idx);
-					ProxiesData.Add(Base);
-				}
-			}
-
-			void Remove(IPhysicsProxyBase* Base)
-			{
-				const int32 Idx = Base->GetDirtyIdx();
-				if(Idx != INDEX_NONE)
-				{
-					if(Idx == ProxiesData.Num() - 1)
-					{
-						//last element so just pop
-						ProxiesData.Pop(/*bAllowShrinking=*/false);
-					}
-					else
-					{
-						//update other proxy's idx
-						ProxiesData.RemoveAtSwap(Idx);
-						ProxiesData[Idx].SetDirtyIdx(Idx);
-					}
-
-					Base->ResetDirtyIdx();
-				}
-			}
-
-			void Reset()
-			{
-				ProxiesData.Reset();
-				ShapesData.Reset();
-			}
-
-			int32 NumDirtyProxies() const { return ProxiesData.Num(); }
-			int32 NumDirtyShapes() const { return ShapesData.Num(); }
-
-			FShapeDirtyData* GetShapesDirtyData(){ return ShapesData.GetData(); }
-
-			template <typename Lambda>
-			void ParallelForEachProxy(const Lambda& Func)
-			{
-				::ParallelFor(ProxiesData.Num(),[this,&Func](int32 Idx)
-				{
-					Func(Idx,ProxiesData[Idx]);
-				});
-			}
-
-			template <typename Lambda>
-			void ForEachProxy(const Lambda& Func)
-			{
-				int32 Idx = 0;
-				for(FDirtyProxy& Dirty : ProxiesData)
-				{
-					Func(Idx++, Dirty);
-				}
-			}
-
-			void AddShape(IPhysicsProxyBase* Proxy, int32 ShapeIdx)
-			{
-				Add(Proxy);
-				FDirtyProxy& Dirty = ProxiesData[Proxy->GetDirtyIdx()];
-				for(int32 NewShapeIdx = Dirty.ShapeDataIndices.Num(); NewShapeIdx <= ShapeIdx; ++NewShapeIdx)
-				{
-					const int32 ShapeDataIdx = ShapesData.Add(FShapeDirtyData(NewShapeIdx));
-					Dirty.AddShape(ShapeDataIdx);
-				}
-			}
-
-			void SetNumDirtyShapes(IPhysicsProxyBase* Proxy,int32 NumShapes)
-			{
-				Add(Proxy);
-				FDirtyProxy& Dirty = ProxiesData[Proxy->GetDirtyIdx()];
-				
-				if(NumShapes < Dirty.ShapeDataIndices.Num())
-				{
-					Dirty.ShapeDataIndices.SetNum(NumShapes);
-				}
-				else
-				{
-					for(int32 NewShapeIdx = Dirty.ShapeDataIndices.Num(); NewShapeIdx < NumShapes; ++NewShapeIdx)
-					{
-						const int32 ShapeDataIdx = ShapesData.Add(FShapeDirtyData(NewShapeIdx));
-						Dirty.AddShape(ShapeDataIdx);
-					}
-				}
-			}
-
-		private:
-			TArray<FDirtyProxy> ProxiesData;
-			TArray<FShapeDirtyData> ShapesData;
-		};
-
 		/** Protected construction so callers still have to go through the module to create new instances */
 		FPhysicsSolverBase(const EMultiBufferMode BufferingModeIn, UObject* InOwner)
 			: BufferMode(BufferingModeIn)
diff --git a/Engine/Source/Runtime/Experimental/Chaos/Public/Chaos/ISpatialAcceleration.h b/Engine/Source/Runtime/Experimental/Chaos/Public/Chaos/ISpatialAcceleration.h
--- a/Engine/Source/Runtime/Experimental/Chaos/Public/Chaos/ISpatialAcceleration.h
+++ b/Engine/Source/Runtime/Experimental/Chaos/Public/Chaos/ISpatialAcceleration.h
@@ -396,6 +396,16 @@ public:
 		return nullptr;
 	}
 
+	const TValue* Find(const TKey& Key) const
+	{
+		const int32 Idx = GetUniqueIdx(Key).Idx;
+		if(Idx < Entries.Num() && Entries[Idx].bSet)
+		{
+			return &Entries[Idx].Value;
+		}
+		return nullptr;
+	}
+
 	TValue& FindChecked(const TKey& Key)
 	{
 		return Entries[GetUniqueIdx(Key).Idx].Value;
diff --git a/Engine/Source/Runtime/Experimental/Chaos/Public/Chaos/ParticleHandle.h b/Engine/Source/Runtime/Experimental/Chaos/Public/Chaos/ParticleHandle.h
--- a/Engine/Source/Runtime/Experimental/Chaos/Public/Chaos/ParticleHandle.h
+++ b/Engine/Source/Runtime/Experimental/Chaos/Public/Chaos/ParticleHandle.h
@@ -1611,6 +1611,8 @@ public:
 		}
 	}
 
+	const FParticlePositionRotation& XR() const { return MXR.Read(); }
+
 protected:
 
 	// Pointer to any data that the solver wants to associate with this particle
diff --git a/Engine/Source/Runtime/Experimental/ChaosSolvers/Private/PBDRigidsSolver.cpp b/Engine/Source/Runtime/Experimental/ChaosSolvers/Private/PBDRigidsSolver.cpp
--- a/Engine/Source/Runtime/Experimental/ChaosSolvers/Private/PBDRigidsSolver.cpp
+++ b/Engine/Source/Runtime/Experimental/ChaosSolvers/Private/PBDRigidsSolver.cpp
@@ -18,6 +18,7 @@
 #include "PhysicsProxy/FieldSystemPhysicsProxy.h"
 #include "EventDefaults.h"
 #include "EventsData.h"
+#include "Chaos/RewindData.h"
 
 
 DEFINE_LOG_CATEGORY_STATIC(LogPBDRigidsSolver, Log, All);
@@ -755,7 +756,8 @@ namespace Chaos
 
 		FChaosSolversModule::GetModule()->GetDispatcher()->EnqueueCommandImmediate(this,[bIsSingleThreaded, Manager,DirtyProxiesData,ShapeDirtyData](FPBDRigidsSolver* Solver)
 		{
-			auto ProcessProxyPT = [bIsSingleThreaded, Solver,Manager,DirtyProxiesData,ShapeDirtyData](auto& Proxy,int32 DataIdx,FDirtyProxy& Dirty,const auto& CreateHandleFunc)
+			FRewindData* RewindData = Solver->GetRewindData();
+			auto ProcessProxyPT = [bIsSingleThreaded, Solver,Manager,DirtyProxiesData,ShapeDirtyData, RewindData](auto& Proxy,int32 DataIdx,FDirtyProxy& Dirty,const auto& CreateHandleFunc)
 			{
 				const bool bIsNew = !Proxy->IsInitialized();
 				//single threaded version already created particle, but didn't initialize it
@@ -780,6 +782,11 @@ namespace Chaos
 				Dirty.Clear(*Manager, DataIdx, ShapeDirtyData);
 			};
 
+			if(RewindData)
+			{
+				RewindData->SavePrevFrameState(*DirtyProxiesData);
+			}
+
 			//need to create new particle handles
 			DirtyProxiesData->ForEachProxy([Solver, &ProcessProxyPT](int32 DataIdx,FDirtyProxy& Dirty)
 			{
@@ -988,4 +995,9 @@ namespace Chaos
 		QueryMaterialMasks = SimMaterialMasks;
 	}
 
+	void FPBDRigidsSolver::EnableRewindCapture(int32 NumFrames)
+	{
+		MRewindData = MakeUnique<FRewindData>(NumFrames);
+	}
+
 }; // namespace Chaos
diff --git a/Engine/Source/Runtime/Experimental/ChaosSolvers/Public/PBDRigidsSolver.h b/Engine/Source/Runtime/Experimental/ChaosSolvers/Public/PBDRigidsSolver.h
--- a/Engine/Source/Runtime/Experimental/ChaosSolvers/Public/PBDRigidsSolver.h
+++ b/Engine/Source/Runtime/Experimental/ChaosSolvers/Public/PBDRigidsSolver.h
@@ -47,6 +47,7 @@ namespace Chaos
 	class FPersistentPhysicsTask;
 	class FChaosArchive;
 	class FPBDRigidsSolver;
+	class FRewindData;
 
 	enum class ELockType : uint8
 	{
@@ -128,6 +129,12 @@ namespace Chaos
 
 		bool IsSimulating() const;
 
+		void EnableRewindCapture(int32 NumFrames);
+		FRewindData* GetRewindData()
+		{
+			return MRewindData.Get();
+		}
+
 		template<typename Lambda>
 		void ForEachPhysicsProxy(Lambda InCallable)
 		{
@@ -369,6 +376,7 @@ namespace Chaos
 		TUniquePtr<FSolverEventFilters> MSolverEventFilters;
 		TUniquePtr<FActiveParticlesBuffer> MActiveParticlesBuffer;
 		TMap<const Chaos::TGeometryParticleHandle<float, 3>*, IPhysicsProxyBase*> MParticleToProxy;
+		TUniquePtr<FRewindData> MRewindData;
 
 		//
 		// Commands
diff --git a/Engine/Source/Runtime/Experimental/ChaosSolvers/Public/PhysicsProxy/SingleParticlePhysicsProxy.h b/Engine/Source/Runtime/Experimental/ChaosSolvers/Public/PhysicsProxy/SingleParticlePhysicsProxy.h
--- a/Engine/Source/Runtime/Experimental/ChaosSolvers/Public/PhysicsProxy/SingleParticlePhysicsProxy.h
+++ b/Engine/Source/Runtime/Experimental/ChaosSolvers/Public/PhysicsProxy/SingleParticlePhysicsProxy.h
@@ -84,6 +84,11 @@ public:
 		return Handle;
 	}
 
+	const FParticleHandle* GetHandle() const
+	{
+		return Handle;
+	}
+
 	virtual void* GetHandleUnsafe() const override
 	{
 		return Handle;
@@ -145,6 +150,11 @@ public:
 	{
 		return Particle;
 	}
+
+	const PARTICLE_TYPE* GetParticle() const
+	{
+		return Particle;
+	}
 	
 private:
 	bool bInitialized;
diff --git a/Engine/Source/Runtime/Experimental/ChaosSolvers/Public/RewindData.h b/Engine/Source/Runtime/Experimental/ChaosSolvers/Public/RewindData.h
new file mode 100644
--- /dev/null
+++ b/Engine/Source/Runtime/Experimental/ChaosSolvers/Public/RewindData.h
@@ -0,0 +1,138 @@
+// Copyright Epic Games, Inc. All Rights Reserved.
+#pragma once
+#include "Chaos/Core.h"
+#include "Chaos/ParticleHandle.h"
+
+namespace Chaos
+{
+
+template <typename T,EParticleProperty PropName>
+class TParticleStateProperty
+{
+public:
+
+	TParticleStateProperty()
+	: Manager(nullptr)
+	{
+	}
+
+	TParticleStateProperty(const FDirtyPropertiesManager* InManager, int32 InIdx)
+	: Manager(InManager)
+	, Idx(InIdx)
+	{
+	}
+
+	template <typename LambdaRead>
+	const T& Read(const LambdaRead& ReadFunc) const
+	{
+		if(Manager)
+		{
+			const TDirtyElementPool<T>& Pool = Manager->GetParticlePool<T,PropName>();
+			return Pool.GetElement(Idx);
+		}
+		else
+		{
+			return ReadFunc();
+		}
+	}
+
+	template <typename LambdaWrite>
+	void SyncToParticle(LambdaWrite& WriteFunc) const
+	{
+		if(Manager)
+		{
+			const TDirtyElementPool<T>& Pool = Manager->GetParticlePool<T,PropName>();
+			const T& Value = Pool.GetElement(Idx);
+			WriteFunc(Value);
+		}
+	}
+
+private:
+	FDirtyPropertiesManager* Manager;
+	int32 Idx;
+};
+
+	class FGeometryParticleState
+	{
+	public:
+
+		FGeometryParticleState(TGeometryParticle<FReal,3>* InParticle)
+		: Particle(InParticle)
+		{
+		}
+
+		const FVec3& X() const
+		{
+			return ParticlePositionRotation.Read([this]() -> const auto& {return Particle->XR();}).X;
+		}
+
+	protected:
+		TGeometryParticle<FReal,3>* Particle;
+	private:
+		TParticleStateProperty<FParticlePositionRotation,EParticleProperty::XR> ParticlePositionRotation;
+	};
+
+	class FRewindData
+	{
+	public:
+		FRewindData(int32 NumFrames)
+			: CurFrame(0)
+		{
+		}
+
+		const FGeometryParticleState* GetStateAtFrame(const TGeometryParticle<FReal,3>& Particle, int32 Frame) const
+		{
+			if(const TArray<FFrameInfo>* Frames = ParticleToFrameInfo.Find(Particle))
+			{
+				//is it worth doing binary search?
+				for(const FFrameInfo& Info : (*Frames))
+				{
+					if(Info.Frame == Frame)
+					{
+						return &Info.State;
+					}
+				}
+			}
+			
+			return nullptr;
+		}
+
+		void SetSecondLastState(const FDirtyProxy& DirtyProxy)
+		{
+			//NOTE: this is called from PT so we cannot use the GTParticle to read
+			//Instead we  must use the proxy system
+			//However, the rewind happens on GT so that's why we still want to associate this data with the GTParticle
+
+			auto ProcessProxy = [this](const auto Proxy)
+			{
+				const auto GTParticleUnsafe = Proxy->GetParticle();
+				TArray<FGeometryParticleState>& Frames = ParticleToFrameInfo.FindOrAdd(*GTParticleUnsafe);
+				Frames.Add({FGeometryParticleState(GTParticleUnsafe),CurFrame});
+			};
+
+			switch(DirtyProxy.Proxy->GetType())
+			{
+			case EPhysicsProxyType::SingleGeometryParticleType: ProcessProxy(static_cast<const FGeometryParticlePhysicsProxy*>(DirtyProxy.Proxy)); break;
+			case EPhysicsProxyType::SingleKinematicParticleType: ProcessProxy(static_cast<const FKinematicGeometryParticlePhysicsProxy*>(DirtyProxy.Proxy)); break;
+			case EPhysicsProxyType::SingleRigidParticleType: ProcessProxy(static_cast<const FRigidParticlePhysicsProxy*>(DirtyProxy.Proxy)); break;
+			default: check(false);
+			}
+		}
+
+		void AdvanceFrame()
+		{
+			++CurFrame
+		}
+		
+	private:
+
+		struct FFrameInfo
+		{
+			FGeometryParticleState State;
+			int32 Frame;
+		};
+
+		TArrayAsMap<TGeometryParticle<FReal,3>, TArray<FFrameInfo>> ParticleToFrameInfo;
+		int32 CurFrame;
+	};
+}
