Fix CIS. Submitted stale file

diff --git a/Engine/Source/Programs/HeadlessChaos/Private/HeadlessChaosTestRewind.cpp b/Engine/Source/Programs/HeadlessChaos/Private/HeadlessChaosTestRewind.cpp
--- a/Engine/Source/Programs/HeadlessChaos/Private/HeadlessChaosTestRewind.cpp
+++ b/Engine/Source/Programs/HeadlessChaos/Private/HeadlessChaosTestRewind.cpp
@@ -11,7 +11,7 @@
 
 #include "Modules/ModuleManager.h"
 #include "Framework/PhysicsTickTask.h"
-#include "Chaos/RewindData.h"
+#include "RewindData.h"
 
 
 namespace ChaosTest {
diff --git a/Engine/Source/Runtime/Experimental/ChaosSolvers/Private/PBDRigidsSolver.cpp b/Engine/Source/Runtime/Experimental/ChaosSolvers/Private/PBDRigidsSolver.cpp
--- a/Engine/Source/Runtime/Experimental/ChaosSolvers/Private/PBDRigidsSolver.cpp
+++ b/Engine/Source/Runtime/Experimental/ChaosSolvers/Private/PBDRigidsSolver.cpp
@@ -18,7 +18,7 @@
 #include "PhysicsProxy/FieldSystemPhysicsProxy.h"
 #include "EventDefaults.h"
 #include "EventsData.h"
-#include "Chaos/RewindData.h"
+#include "RewindData.h"
 
 
 DEFINE_LOG_CATEGORY_STATIC(LogPBDRigidsSolver, Log, All);
diff --git a/Engine/Source/Runtime/Experimental/ChaosSolvers/Public/RewindData.h b/Engine/Source/Runtime/Experimental/ChaosSolvers/Public/RewindData.h
--- a/Engine/Source/Runtime/Experimental/ChaosSolvers/Public/RewindData.h
+++ b/Engine/Source/Runtime/Experimental/ChaosSolvers/Public/RewindData.h
@@ -2,6 +2,8 @@
 #pragma once
 #include "Chaos/Core.h"
 #include "Chaos/ParticleHandle.h"
+#include "PhysicsProxy/SingleParticlePhysicsProxyFwd.h"
+#include "Chaos/Framework/PhysicsSolverBase.h"
 
 namespace Chaos
 {
@@ -12,13 +14,13 @@ class TParticleStateProperty
 public:
 
 	TParticleStateProperty()
-	: Manager(nullptr)
+		: Manager(nullptr)
 	{
 	}
 
-	TParticleStateProperty(const FDirtyPropertiesManager* InManager, int32 InIdx)
-	: Manager(InManager)
-	, Idx(InIdx)
+	TParticleStateProperty(const FDirtyPropertiesManager* InManager,int32 InIdx)
+		: Manager(InManager)
+		,Idx(InIdx)
 	{
 	}
 
@@ -29,8 +31,7 @@ public:
 		{
 			const TDirtyElementPool<T>& Pool = Manager->GetParticlePool<T,PropName>();
 			return Pool.GetElement(Idx);
-		}
-		else
+		} else
 		{
 			return ReadFunc();
 		}
@@ -47,92 +48,140 @@ public:
 		}
 	}
 
+	template <typename LambdaSet>
+	void SyncRemoteData(FDirtyPropertiesManager* InManager,int32 InIdx,const LambdaSet& SetFunc)
+	{
+		ensure(Manager == nullptr);
+		Manager = InManager;
+		Idx = InIdx;
+		SetFunc(Manager->GetParticlePool<T,PropName>().GetElement(Idx));
+	}
+
 private:
 	FDirtyPropertiesManager* Manager;
 	int32 Idx;
 };
 
-	class FGeometryParticleState
-	{
-	public:
+class FGeometryParticleState
+{
+public:
 
-		FGeometryParticleState(TGeometryParticle<FReal,3>* InParticle)
+	FGeometryParticleState(TGeometryParticle<FReal,3>* InParticle)
 		: Particle(InParticle)
-		{
-		}
+	{
+	}
+
+	const FVec3& X() const
+	{
+		return ParticlePositionRotation.Read([this]() -> const auto& {return Particle->XR();}).X;
+	}
 
-		const FVec3& X() const
+	void SyncRemoteData(FDirtyPropertiesManager* Manager,int32 Idx,const FDirtyProxy& Dirty)
+	{
+		const auto Proxy = static_cast<const FGeometryParticlePhysicsProxy*>(Dirty.Proxy);
+		const auto Handle = Proxy->GetHandle();
+		ParticlePositionRotation.SyncRemoteData(Manager,Idx,[Handle](FParticlePositionRotation& Data)
 		{
-			return ParticlePositionRotation.Read([this]() -> const auto& {return Particle->XR();}).X;
-		}
+			Data.X = Handle->X();
+			Data.R = Handle->R();
+		});
+	}
 
-	protected:
-		TGeometryParticle<FReal,3>* Particle;
-	private:
-		TParticleStateProperty<FParticlePositionRotation,EParticleProperty::XR> ParticlePositionRotation;
-	};
+protected:
+	TGeometryParticle<FReal,3>* Particle;
+private:
+	TParticleStateProperty<FParticlePositionRotation,EParticleProperty::XR> ParticlePositionRotation;
+};
 
-	class FRewindData
+class FRewindData
+{
+public:
+	FRewindData(int32 NumFrames)
+		: CurFrame(0)
 	{
-	public:
-		FRewindData(int32 NumFrames)
-			: CurFrame(0)
-		{
-		}
+	}
 
-		const FGeometryParticleState* GetStateAtFrame(const TGeometryParticle<FReal,3>& Particle, int32 Frame) const
+	const FGeometryParticleState* GetStateAtFrame(const TGeometryParticle<FReal,3>& Particle,int32 Frame) const
+	{
+		if(const TArray<FFrameInfo>* Frames = ParticleToFrameInfo.Find(Particle.UniqueIdx()))
 		{
-			if(const TArray<FFrameInfo>* Frames = ParticleToFrameInfo.Find(Particle))
+			//is it worth doing binary search?
+			for(const FFrameInfo& Info : (*Frames))
 			{
-				//is it worth doing binary search?
-				for(const FFrameInfo& Info : (*Frames))
+				if(Info.Frame == Frame)
 				{
-					if(Info.Frame == Frame)
-					{
-						return &Info.State;
-					}
+					return &Info.State;
 				}
 			}
-			
-			return nullptr;
 		}
 
-		void SetSecondLastState(const FDirtyProxy& DirtyProxy)
-		{
-			//NOTE: this is called from PT so we cannot use the GTParticle to read
-			//Instead we  must use the proxy system
-			//However, the rewind happens on GT so that's why we still want to associate this data with the GTParticle
+		return nullptr;
+	}
+
+	void SavePrevFrameState(const FDirtySet& DirtyData)
+	{
+		Managers.Emplace(MakeUnique<FDirtyPropertiesManager>());
 
-			auto ProcessProxy = [this](const auto Proxy)
+		FDirtyPropertiesManager* Manager = Managers.Last().Get();
+		Manager->SetNumParticles(DirtyData.NumDirtyProxies());
+
+		//NOTE: this is called from PT so we cannot use the GTParticle to read
+		//Instead we  must use the proxy system
+		//However, the rewind happens on GT so that's why we still want to associate this data with the GTParticle
+
+		auto ProcessProxy = [this,Manager](const auto Proxy,const int32 DataIdx,const FDirtyProxy& Dirty)
+		{
+			//Since we're saving last frame's data the proxy must already be initialized
+			if(Proxy->IsInitialized())
 			{
 				const auto GTParticleUnsafe = Proxy->GetParticle();
-				TArray<FGeometryParticleState>& Frames = ParticleToFrameInfo.FindOrAdd(*GTParticleUnsafe);
-				Frames.Add({FGeometryParticleState(GTParticleUnsafe),CurFrame});
-			};
+				const auto PTParticle = Proxy->GetHandle();
+				TArray<FFrameInfo>& Frames = ParticleToFrameInfo.FindOrAdd(PTParticle->UniqueIdx());
+				Frames.Add({FGeometryParticleState(GTParticleUnsafe),CurFrame-1});
+				Frames.Last().State.SyncRemoteData(Manager,DataIdx,Dirty);
+			}
+		};
 
-			switch(DirtyProxy.Proxy->GetType())
+		DirtyData.ParallelForEachProxy([&ProcessProxy](int32 DataIdx,const FDirtyProxy& Dirty)
+		{
+			switch(Dirty.Proxy->GetType())
 			{
-			case EPhysicsProxyType::SingleGeometryParticleType: ProcessProxy(static_cast<const FGeometryParticlePhysicsProxy*>(DirtyProxy.Proxy)); break;
-			case EPhysicsProxyType::SingleKinematicParticleType: ProcessProxy(static_cast<const FKinematicGeometryParticlePhysicsProxy*>(DirtyProxy.Proxy)); break;
-			case EPhysicsProxyType::SingleRigidParticleType: ProcessProxy(static_cast<const FRigidParticlePhysicsProxy*>(DirtyProxy.Proxy)); break;
-			default: check(false);
+			case EPhysicsProxyType::SingleRigidParticleType:
+			{
+				auto Proxy = static_cast<FRigidParticlePhysicsProxy*>(Dirty.Proxy);
+				ProcessProxy(Proxy,DataIdx,Dirty);
+				break;
 			}
-		}
+			case EPhysicsProxyType::SingleKinematicParticleType:
+			{
+				auto Proxy = static_cast<FKinematicGeometryParticlePhysicsProxy*>(Dirty.Proxy);
+				ProcessProxy(Proxy,DataIdx,Dirty);
+				break;
+			}
+			case EPhysicsProxyType::SingleGeometryParticleType:
+			{
+				auto Proxy = static_cast<FGeometryParticlePhysicsProxy*>(Dirty.Proxy);
+				ProcessProxy(Proxy,DataIdx,Dirty);
+				break;
+			}
+			default:
+			ensure("Unknown proxy type in physics solver.");
+			}
+		});
 
-		void AdvanceFrame()
-		{
-			++CurFrame
-		}
-		
-	private:
+		++CurFrame;
+	}
 
-		struct FFrameInfo
-		{
-			FGeometryParticleState State;
-			int32 Frame;
-		};
+private:
 
-		TArrayAsMap<TGeometryParticle<FReal,3>, TArray<FFrameInfo>> ParticleToFrameInfo;
-		int32 CurFrame;
+	struct FFrameInfo
+	{
+		FGeometryParticleState State;
+		int32 Frame;
 	};
+
+	TArrayAsMap<FUniqueIdx,TArray<FFrameInfo>> ParticleToFrameInfo;
+	TArray<TUniquePtr<FDirtyPropertiesManager>> Managers;
+	int32 CurFrame;
+};
 }
