Fix CIS take 2

diff --git a/Engine/Source/Runtime/Experimental/Chaos/Public/Chaos/Framework/PhysicsSolverBase.h b/Engine/Source/Runtime/Experimental/Chaos/Public/Chaos/Framework/PhysicsSolverBase.h
--- a/Engine/Source/Runtime/Experimental/Chaos/Public/Chaos/Framework/PhysicsSolverBase.h
+++ b/Engine/Source/Runtime/Experimental/Chaos/Public/Chaos/Framework/PhysicsSolverBase.h
@@ -114,6 +114,16 @@ namespace Chaos
 			}
 		}
 
+		template <typename Lambda>
+		void ForEachProxy(const Lambda& Func) const
+		{
+			int32 Idx = 0;
+			for(const FDirtyProxy& Dirty : ProxiesData)
+			{
+				Func(Idx++,Dirty);
+			}
+		}
+
 		void AddShape(IPhysicsProxyBase* Proxy,int32 ShapeIdx)
 		{
 			Add(Proxy);
diff --git a/Engine/Source/Runtime/Experimental/ChaosSolvers/Public/RewindData.h b/Engine/Source/Runtime/Experimental/ChaosSolvers/Public/RewindData.h
--- a/Engine/Source/Runtime/Experimental/ChaosSolvers/Public/RewindData.h
+++ b/Engine/Source/Runtime/Experimental/ChaosSolvers/Public/RewindData.h
@@ -18,7 +18,7 @@ public:
 	{
 	}
 
-	TParticleStateProperty(const FDirtyPropertiesManager* InManager,int32 InIdx)
+	TParticleStateProperty(FDirtyPropertiesManager* InManager,int32 InIdx)
 		: Manager(InManager)
 		,Idx(InIdx)
 	{
@@ -137,12 +137,13 @@ public:
 				const auto GTParticleUnsafe = Proxy->GetParticle();
 				const auto PTParticle = Proxy->GetHandle();
 				TArray<FFrameInfo>& Frames = ParticleToFrameInfo.FindOrAdd(PTParticle->UniqueIdx());
-				Frames.Add({FGeometryParticleState(GTParticleUnsafe),CurFrame-1});
+				Frames.Add(FFrameInfo{FGeometryParticleState(GTParticleUnsafe),CurFrame-1});
 				Frames.Last().State.SyncRemoteData(Manager,DataIdx,Dirty);
 			}
 		};
 
-		DirtyData.ParallelForEachProxy([&ProcessProxy](int32 DataIdx,const FDirtyProxy& Dirty)
+		//can't be parallel because ParticletoFrameInfo is modified
+		DirtyData.ForEachProxy([&ProcessProxy](int32 DataIdx,const FDirtyProxy& Dirty)
 		{
 			switch(Dirty.Proxy->GetType())
 			{
