Reinstate Simplicials to the geometry collection physics proxy.

Add functions to TTriangleMesh for finding open boundaries, merging coincident vertices, and removing degernate and duplicate faces.  Have CleanCollisionParticles() first attempt to repair holes in geometry due to such discontinuities prior to calling GetVertexImportanceOrdering(), as some topologies were failing to get surface curvature.

Added the ability for ManagedArrayCollection::CopyMatchingAttributesFrom() to skip specified attributes.  This is so we can skip uncopyable arrays (TUniquePtr) rather than fail.

diff --git a/Engine/Source/Runtime/Experimental/Chaos/Private/Chaos/TriangleMesh.cpp b/Engine/Source/Runtime/Experimental/Chaos/Private/Chaos/TriangleMesh.cpp
--- a/Engine/Source/Runtime/Experimental/Chaos/Private/Chaos/TriangleMesh.cpp
+++ b/Engine/Source/Runtime/Experimental/Chaos/Private/Chaos/TriangleMesh.cpp
@@ -93,6 +93,17 @@ void TTriangleMesh<T>::InitHelper(const int32 StartIdx, const int32 EndIdx, cons
 	ExpandVertexRange(StartIdx, EndIdx);
 }
 
+template <class T>
+void TTriangleMesh<T>::ResetAuxiliaryStructures()
+{
+	MPointToTriangleMap.Reset();
+	MPointToNeighborsMap.Reset();
+	TArray<TVector<int32, 2>> EmptyEdges;
+	MSegmentMesh.Init(EmptyEdges);
+	MFaceToEdges.Reset();
+	MEdgeToFaces.Reset();
+}
+
 template<class T>
 TPair<int32, int32> TTriangleMesh<T>::GetVertexRange() const
 {
@@ -581,6 +592,25 @@ FORCEINLINE TVector<int32, 2> GetOrdered(const TVector<int32, 2>& elem)
 	    ordered ? elem[1] : elem[0]);
 }
 
+void Order(int32& A, int32& B)
+{
+	if (B < A)
+	{
+		int32 Tmp = A;
+		A = B;
+		B = Tmp;
+	}
+}
+
+TVector<int32, 3> GetOrdered(const TVector<int32, 3>& Elem)
+{
+	TVector<int32, 3> OrderedElem = Elem;  // 3 2 1		1 2 3		1 2 1	2 1 1
+	Order(OrderedElem[0], OrderedElem[1]); // 2 3 1		1 2 3		1 2 1	1 2 1
+	Order(OrderedElem[1], OrderedElem[2]); // 2 1 3		1 2 3		1 1 2	1 1 2
+	Order(OrderedElem[0], OrderedElem[1]); // 1 2 3		1 2 3		1 1 2	1 1 2
+	return OrderedElem;
+}
+
 /**
  * Comparator for TSet<TVector<int32,2>> that compares the components of vectors in ascending
  * order.
@@ -684,6 +714,146 @@ const TArray<TVector<int32, 2>>& TTriangleMesh<T>::GetEdgeToFaces()
 	return MEdgeToFaces;
 }
 
+
+template <class T>
+TSet<int32> TTriangleMesh<T>::GetBoundaryPoints()
+{
+	TSegmentMesh<T>& SegmentMesh = GetSegmentMesh();
+	const TArray<TVector<int32, 2>>& Edges = SegmentMesh.GetElements();
+	const TArray<TVector<int32, 2>>& EdgeToFaces = GetEdgeToFaces();
+	TSet<int32> OpenBoundaryPoints;
+	for (int32 EdgeIdx = 0; EdgeIdx < EdgeToFaces.Num(); ++EdgeIdx)
+	{
+		const TVector<int32, 2>& CoincidentFaces = EdgeToFaces[EdgeIdx];
+		if (CoincidentFaces[0] == INDEX_NONE || CoincidentFaces[1] == INDEX_NONE)
+		{
+			const TVector<int32, 2>& Edge = Edges[EdgeIdx];
+			OpenBoundaryPoints.Add(Edge[0]);
+			OpenBoundaryPoints.Add(Edge[1]);
+		}
+	}
+	return OpenBoundaryPoints;
+}
+
+template <class T>
+TMap<int32, int32> TTriangleMesh<T>::FindCoincidentVertexRemappings(
+	const TArray<int32>& TestIndices,
+	const TArrayView<const TVector<T, 3>>& Points)
+{
+	// From index -> To index
+	TMap<int32, int32> Remappings;
+
+	const int32 NumPoints = TestIndices.Num();
+	if (NumPoints <= 1)
+	{
+		return Remappings;
+	}
+
+	// Move the points to the origin to avoid floating point aliasing far away
+	// from the origin.
+	TAABB<T, 3> Bbox(Points[0], Points[0]);
+	for (int i = 1; i < NumPoints; i++)
+	{
+		Bbox.GrowToInclude(Points[TestIndices[i]]);
+	}
+	const TVector<T, 3> Center = Bbox.Center();
+
+	TArray<TVector<T, 3>> LocalPoints;
+	LocalPoints.AddUninitialized(NumPoints);
+	LocalPoints[0] = Points[TestIndices[0]] - Center;
+	TAABB<T, 3> LocalBBox(LocalPoints[0], LocalPoints[0]);
+	for (int i = 1; i < NumPoints; i++)
+	{
+		LocalPoints[i] = Points[TestIndices[i]] - Center;
+		LocalBBox.GrowToInclude(LocalPoints[i]);
+	}
+
+	// Return early if all points are coincident
+	if (LocalBBox.Extents().Max() < KINDA_SMALL_NUMBER)
+	{
+		int32 First = INDEX_NONE;
+		for (const int32 Pt : TestIndices)
+		{
+			if (First == INDEX_NONE)
+			{
+				First = Pt;
+			}
+			else
+			{
+				// Remap Pt to First
+				Remappings.Add(Pt, First);
+			}
+		}
+		return Remappings;
+	}
+
+	LocalBBox.Thicken(1.0e-3);
+	const TVector<T, 3> LocalCenter = LocalBBox.Center();
+	const TVector<T, 3>& LocalMin = LocalBBox.Min();
+
+	const T MaxBBoxDim = LocalBBox.Extents().Max();
+
+	// Find coincident vertices.
+	// We hash to a grid of fine enough resolution such that if 2 particles 
+	// hash to the same cell, then we're going to consider them coincident.
+	TMap<int64, TSet<int32>> OccupiedCells;
+	OccupiedCells.Reserve(NumPoints);
+
+	const int64 Resolution = static_cast<int64>(floor(MaxBBoxDim / 0.01));
+	const T CellSize = MaxBBoxDim / Resolution;
+	for (int i = 0; i < 2; i++)
+	{
+		OccupiedCells.Reset();
+
+		// Shift the grid by 1/2 a grid cell the second iteration so that
+		// we don't miss slightly adjacent coincident points across cell
+		// boundaries.
+		const TVector<T, 3> GridCenter = LocalCenter - TVector<T, 3>(i * CellSize / 2);
+		for (int32 LocalIdx = 0; LocalIdx < NumPoints; LocalIdx++)
+		{
+			const int32 Idx = TestIndices[LocalIdx];
+			if (i != 0 && Remappings.Contains(Idx))
+			{
+				// Already remapped
+				continue;
+			}
+
+			const TVector<T, 3>& Pos = LocalPoints[LocalIdx];
+			const TVector<int64, 3> Coord(
+				static_cast<int64>(floor((Pos[0] - GridCenter[0]) / CellSize + Resolution / 2)),
+				static_cast<int64>(floor((Pos[1] - GridCenter[1]) / CellSize + Resolution / 2)),
+				static_cast<int64>(floor((Pos[2] - GridCenter[2]) / CellSize + Resolution / 2)));
+			const int64 FlatIdx =
+				((Coord[0] * Resolution + Coord[1]) * Resolution) + Coord[2];
+
+			TSet<int32>& Bucket = OccupiedCells.FindOrAdd(FlatIdx);
+			Bucket.Add(Idx);
+		}
+
+		// Iterate over all occupied cells and remap redundant vertices to the first index.
+		for (auto& KV : OccupiedCells)
+		{
+			const TSet<int32>& CoincidentVertices = KV.Value;
+			if (CoincidentVertices.Num() <= 1)
+				continue;
+			int32 First = INDEX_NONE;
+			for (const int32 Idx : CoincidentVertices)
+			{
+				if (First == INDEX_NONE)
+				{
+					First = Idx;
+				}
+				else
+				{
+					Remappings.Add(Idx, First);
+				}
+			}
+		}
+	}
+
+	return Remappings;
+}
+
 template<class T>
 TArray<T> TTriangleMesh<T>::GetCurvatureOnEdges(const TArray<TVector<T, 3>>& FaceNormals)
 {
@@ -859,7 +1029,7 @@ TArray<int32> TTriangleMesh<T>::GetVertexImportanceOrdering(
 		const bool IsFree = Neighbors == nullptr || Neighbors->Num() == 0;
 		Rank[Idx - Offset] = IsFree ? 1 : 0;
 		FoundFreeParticle |= IsFree;
-	}
+			}
 	if (FoundFreeParticle)
 	{
 		StableSort(&PointOrder[0], NumPoints, AscendingRankPred);
@@ -896,7 +1066,7 @@ TArray<int32> TTriangleMesh<T>::GetVertexImportanceOrdering(
 	TAABB<T, 3> Bbox(Points[0], Points[0]);
 	for (int i = 1; i < NumPoints; i++)
 	{
-		Bbox.GrowToInclude(Points[i]);
+		Bbox.GrowToInclude(Points[i+Offset]);
 	}
 	const TVector<T, 3> Center = Bbox.Center();
 
@@ -1059,12 +1229,97 @@ template<class T>
 void TTriangleMesh<T>::RemapVertices(const TArray<int32>& Order)
 {
 	// Remap element indices
+	int32 MinIdx = TNumericLimits<int32>::Max();
+	int32 MaxIdx = -TNumericLimits<int32>::Max();
 	for (int32 i = 0; i < MElements.Num(); i++)
 	{
 		TVector<int32, 3>& elem = MElements[i];
-		elem[0] = Order[elem[0]];
-		elem[1] = Order[elem[1]];
-		elem[2] = Order[elem[2]];
+		for (int32 j = 0; j < 3; ++j)
+		{
+			if (elem[j] != Order[elem[j]])
+			{
+				elem[j] = Order[elem[j]];
+				MinIdx = elem[j] < MinIdx ? elem[j] : MinIdx;
+				MaxIdx = elem[j] > MaxIdx ? elem[j] : MaxIdx;
+			}
+		}
+	}
+	if (MinIdx != TNumericLimits<int32>::Max())
+	{
+		ExpandVertexRange(MinIdx, MaxIdx);
+		RemoveDuplicateElements();
+		RemoveDegenerateElements();
+		ResetAuxiliaryStructures();
+	}
+}
+
+template <class T>
+void TTriangleMesh<T>::RemapVertices(const TMap<int32, int32>& Remapping)
+{
+	if (!Remapping.Num())
+	{
+		return;
+	}
+	int32 MinIdx = TNumericLimits<int32>::Max();
+	int32 MaxIdx = -TNumericLimits<int32>::Max();
+	for (TVector<int32, 3>& Tri : MElements)
+	{
+		for (int32 Idx = 0; Idx < 3; ++Idx)
+		{
+			if (const int32* ToIdx = Remapping.Find(Tri[Idx]))
+			{
+				Tri[Idx] = *ToIdx;
+				MinIdx = *ToIdx < MinIdx ? *ToIdx : MinIdx;
+				MaxIdx = *ToIdx > MaxIdx ? *ToIdx : MaxIdx;
+			}
+		}
+	}
+	if (MinIdx != TNumericLimits<int32>::Max())
+	{
+		ExpandVertexRange(MinIdx, MaxIdx);
+		RemoveDuplicateElements();
+		RemoveDegenerateElements();
+		ResetAuxiliaryStructures();
+	}
+}
+
+template <class T>
+void TTriangleMesh<T>::RemoveDuplicateElements()
+{
+	TArray<int32> ToRemove;
+	TSet<TVector<int32, 3>> Existing;
+	for (int32 Idx = 0; Idx < MElements.Num(); ++Idx)
+	{
+		const TVector<int32, 3>& Tri = MElements[Idx];
+		const TVector<int32, 3> OrderedTri = GetOrdered(Tri);
+		if (!Existing.Contains(OrderedTri))
+		{
+			Existing.Add(OrderedTri);
+			continue;
+		}
+		ToRemove.Add(Idx);
+	}
+	for (int32 Idx = ToRemove.Num() - 1; Idx >= 0; --Idx)
+	{
+		MElements.RemoveAtSwap(ToRemove[Idx]);
+	}
+}
+
+template <class T>
+void TTriangleMesh<T>::RemoveDegenerateElements()
+{
+	for (int i = MElements.Num() - 1; i >= 0; --i)
+	{
+		if (MElements[i][0] == MElements[i][1] ||
+			MElements[i][0] == MElements[i][2] ||
+			MElements[i][1] == MElements[i][2])
+		{
+			// It's possible that the order of the triangles might be important.
+			// RemoveAtSwap() changes the order of the array.  I figure that if
+			// you're up for CullDegenerateElements, then triangle reordering is
+			// fair game.
+			MElements.RemoveAtSwap(i);
+		}
 	}
 }
 
diff --git a/Engine/Source/Runtime/Experimental/Chaos/Public/Chaos/PBDRigidClustering.h b/Engine/Source/Runtime/Experimental/Chaos/Public/Chaos/PBDRigidClustering.h
--- a/Engine/Source/Runtime/Experimental/Chaos/Public/Chaos/PBDRigidClustering.h
+++ b/Engine/Source/Runtime/Experimental/Chaos/Public/Chaos/PBDRigidClustering.h
@@ -442,6 +442,18 @@ CleanCollisionParticles(
 	if (Fraction <= 0.0)
 		return CollisionVertices;
 
+	// If the tri mesh has any open boundaries, see if we can merge any coincident
+	// vertices on the boundary.  This makes the importance ordering work much better
+	// as we need the curvature at each edge of the tri mesh, and we can't calculate
+	// curvature on discontiguous triangles.
+	TSet<int32> BoundaryPoints = TriMesh.GetBoundaryPoints();
+	if (BoundaryPoints.Num())
+	{
+		TMap<int32, int32> Remapping =
+			TriMesh.FindCoincidentVertexRemappings(BoundaryPoints.Array(), Vertices);
+		TriMesh.RemapVertices(Remapping);
+	}
+
 	// Get the importance vertex ordering, from most to least.  Reorder the 
 	// particles accordingly.
 	TArray<int32> CoincidentVertices;
diff --git a/Engine/Source/Runtime/Experimental/Chaos/Public/Chaos/TriangleMesh.h b/Engine/Source/Runtime/Experimental/Chaos/Public/Chaos/TriangleMesh.h
--- a/Engine/Source/Runtime/Experimental/Chaos/Public/Chaos/TriangleMesh.h
+++ b/Engine/Source/Runtime/Experimental/Chaos/Public/Chaos/TriangleMesh.h
@@ -29,6 +29,8 @@ namespace Chaos
 		void Init(TArray<TVector<int32, 3>>&& Elements, const int32 StartIdx = 0, const int32 EndIdx = -1, const bool CullDegenerateElements=true);
 		void Init(const TArray<TVector<int32, 3>>& Elements, const int32 StartIdx = 0, const int32 EndIdx = -1, const bool CullDegenerateElements=true);
 
+		void ResetAuxiliaryStructures();
+
 		/**
 		 * Returns the closed interval of the smallest vertex index used by 
 		 * this class, to the largest.
@@ -186,6 +188,21 @@ namespace Chaos
 		/** @brief Helper that generates edge curvatures on the fly. */
 		TArray<T> GetCurvatureOnPoints(const TArrayView<const TVector<T, 3>>& points);
 
+		/**
+		 * Get the set of point indices that live on the boundary (an edge with only 1 
+		 * coincident face).
+		 */
+		TSet<int32> GetBoundaryPoints();
+
+		/**
+		 * Find vertices that are coincident within the subset @param TestIndices 
+		 * of given coordinates @param Points, and return a correspondence mapping
+		 * from redundant vertex index to consolidated vertex index.
+		 */
+		TMap<int32, int32> FindCoincidentVertexRemappings(
+			const TArray<int32>& TestIndices,
+			const TArrayView<const TVector<T, 3>>& Points);
+
 		/**
 		 * @ret An array of vertex indices ordered from most important to least.
 		 * @param Points - point positions.
@@ -206,6 +223,10 @@ namespace Chaos
 
 		/** @brief Reorder vertices according to @param Order. */
 		void RemapVertices(const TArray<int32>& Order);
+		void RemapVertices(const TMap<int32, int32>& Remapping);
+
+		void RemoveDuplicateElements();
+		void RemoveDegenerateElements();
 
 		static void InitEquilateralTriangleXY(TTriangleMesh<T>& TriMesh, TParticles<T, 3>& Particles)
 		{
diff --git a/Engine/Source/Runtime/Experimental/ChaosSolvers/Private/GeometryCollectionProxyData.cpp b/Engine/Source/Runtime/Experimental/ChaosSolvers/Private/GeometryCollectionProxyData.cpp
--- a/Engine/Source/Runtime/Experimental/ChaosSolvers/Private/GeometryCollectionProxyData.cpp
+++ b/Engine/Source/Runtime/Experimental/ChaosSolvers/Private/GeometryCollectionProxyData.cpp
@@ -61,7 +61,7 @@ FGeometryDynamicCollection::FGeometryDynamicCollection()
 	AddExternalAttribute("MassToLocal", FTransformCollection::TransformGroup, MassToLocal);
 	//AddExternalAttribute(ShapesQueryDataAttribute, FTransformCollection::TransformGroup, ShapeQueryData);
 	//AddExternalAttribute(ShapesSimDataAttribute, FTransformCollection::TransformGroup, ShapeSimData);
-	//AddExternalAttribute(SimplicialsAttribute, FTransformCollection::TransformGroup, Simplicials);
+	AddExternalAttribute(SimplicialsAttribute, FTransformCollection::TransformGroup, Simplicials);
 	AddExternalAttribute(SimulatableParticlesAttribute, FGeometryCollection::TransformGroup, SimulatableParticles);
 
 }
diff --git a/Engine/Source/Runtime/Experimental/ChaosSolvers/Private/PhysicsProxy/GeometryCollectionPhysicsProxy.cpp b/Engine/Source/Runtime/Experimental/ChaosSolvers/Private/PhysicsProxy/GeometryCollectionPhysicsProxy.cpp
--- a/Engine/Source/Runtime/Experimental/ChaosSolvers/Private/PhysicsProxy/GeometryCollectionPhysicsProxy.cpp
+++ b/Engine/Source/Runtime/Experimental/ChaosSolvers/Private/PhysicsProxy/GeometryCollectionPhysicsProxy.cpp
@@ -552,7 +552,35 @@ void FGeometryCollectionPhysicsProxy::Initialize()
 		}
 	}
 
-	PhysicsThreadCollection.CopyMatchingAttributesFrom(DynamicCollection);
+	// Skip simplicials, as they're owned by unique pointers.
+	TMap<FName, TSet<FName>> SkipList;
+	TSet<FName>& TransformGroupSkipList = SkipList.Emplace(FTransformCollection::TransformGroup);
+	TransformGroupSkipList.Add(DynamicCollection.SimplicialsAttribute);
+
+	PhysicsThreadCollection.CopyMatchingAttributesFrom(DynamicCollection, &SkipList);
+
+	// Copy simplicials.
+	// TODO: Ryan - Should we just transfer ownership of the SimplicialsAttribute from the DynamicCollection to
+	// the PhysicsThreadCollection?
+	{
+		if (DynamicCollection.HasAttribute(DynamicCollection.SimplicialsAttribute, FTransformCollection::TransformGroup))
+		{
+			const auto& SourceSimplicials = DynamicCollection.GetAttribute<TUniquePtr<FSimplicial>>(
+				DynamicCollection.SimplicialsAttribute, FTransformCollection::TransformGroup);
+			for (int32 Index = PhysicsThreadCollection.NumElements(FTransformCollection::TransformGroup) - 1; 0 <= Index; Index--)
+			{
+				PhysicsThreadCollection.Simplicials[Index].Reset(
+					SourceSimplicials[Index] ? SourceSimplicials[Index]->NewCopy() : nullptr);
+			}
+		}
+		else
+		{
+			for (int32 Index = PhysicsThreadCollection.NumElements(FTransformCollection::TransformGroup) - 1; 0 <= Index; Index--)
+			{
+				PhysicsThreadCollection.Simplicials[Index].Reset();
+			}
+		}
+	}
 }
 
 void FGeometryCollectionPhysicsProxy::InitializeDynamicCollection(FGeometryDynamicCollection& DynamicCollection, const FGeometryCollection& RestCollection, const FSimulationParameters& Params)
@@ -561,10 +589,13 @@ void FGeometryCollectionPhysicsProxy::InitializeDynamicCollection(FGeometryDynam
 	// This function will use the rest collection to populate the dynamic collection. 
 	//
 
-	DynamicCollection.CopyMatchingAttributesFrom(RestCollection);
+	TMap<FName, TSet<FName>> SkipList;
+	TSet<FName>& TransformGroupSkipList = SkipList.Emplace(FTransformCollection::TransformGroup);
+	TransformGroupSkipList.Add(DynamicCollection.SimplicialsAttribute);
+	DynamicCollection.CopyMatchingAttributesFrom(RestCollection, &SkipList);
 
 	check(DynamicCollection.HasAttribute(FGeometryDynamicCollection::ImplicitsAttribute, FTransformCollection::TransformGroup));
-	check(RestCollection.HasAttribute(FGeometryDynamicCollection::SharedImplicitsAttribute, FTransformCollection::TransformGroup));
+	check(RestCollection.HasAttribute(FGeometryDynamicCollection::ImplicitsAttribute, FTransformCollection::TransformGroup));
 
 	//
 	// User defined initial velocities need to be populated. 
@@ -577,49 +608,27 @@ void FGeometryCollectionPhysicsProxy::InitializeDynamicCollection(FGeometryDynam
 		}
 	}
 
-	/** 
-		Copy over our implicits from the rest collection to the dynamic collection. This doesn't automatically happen
-		in the copy above due to the rest collection storing implicits in a different attribute that the dynamic collection
-		doesn't contain
-	*/
-	{
-		const TManagedArray<TSharedPtr<Chaos::FImplicitObject, ESPMode::ThreadSafe>>* SharedRestImplicits = RestCollection.FindAttribute<TSharedPtr<Chaos::FImplicitObject, ESPMode::ThreadSafe>>(FGeometryDynamicCollection::FGeometryDynamicCollection::SharedImplicitsAttribute, FTransformCollection::TransformGroup);
-
-		const int32 NumRestElems = RestCollection.NumElements(FTransformCollection::TransformGroup);
-		const int32 NumDynamicElems = DynamicCollection.NumElements(FTransformCollection::TransformGroup);
-
-		if(ensure(SharedRestImplicits && (NumRestElems == NumDynamicElems)))
-		{
-			for(int32 Index = 0; Index < NumDynamicElems; ++Index)
-			{
-				DynamicCollection.Implicits[Index] = (*SharedRestImplicits)[Index];
-			}
-		}
-	}
-
 	// process simplicials
 	{
-		/*
 		if (RestCollection.HasAttribute(DynamicCollection.SimplicialsAttribute, FTransformCollection::TransformGroup))
 		{
 			const auto& RestSimplicials = RestCollection.GetAttribute<TUniquePtr<FSimplicial>>(
 				DynamicCollection.SimplicialsAttribute, FTransformCollection::TransformGroup);
 			for (int32 Index = DynamicCollection.NumElements(FTransformCollection::TransformGroup) - 1; 0 <= Index; Index--)
 			{
-				DynamicCollection.Simplicials[Index] = TSharedPtr<FSimplicial>(RestSimplicials[Index] ? RestSimplicials[Index]->NewCopy() : nullptr);
+				DynamicCollection.Simplicials[Index].Reset(
+					RestSimplicials[Index] ? RestSimplicials[Index]->NewCopy() : nullptr);
 			}
 		}
 		else
 		{
 			for (int32 Index = DynamicCollection.NumElements(FTransformCollection::TransformGroup) - 1; 0 <= Index; Index--)
 			{
-				DynamicCollection.Simplicials[Index] = TSharedPtr<FSimplicial>(nullptr);
+				DynamicCollection.Simplicials[Index].Reset();
 			}
 		}
-		*/
 	}
 
-
 	// Process Activity
 	{
 		const int32 NumTransforms = DynamicCollection.SimulatableParticles.Num();
@@ -673,6 +682,7 @@ void FGeometryCollectionPhysicsProxy::InitializeBodiesPT(
 		const TManagedArray<FVector>& InitialAngularVelocity = DynamicCollection.InitialAngularVelocity;
 		const TManagedArray<FVector>& InitialLinearVelocity = DynamicCollection.InitialLinearVelocity;
 		const TManagedArray<FGeometryDynamicCollection::FSharedImplicit>& Implicits = DynamicCollection.Implicits;
+		const TManagedArray<TUniquePtr<FCollisionStructureManager::FSimplicial>>& Simplicials = DynamicCollection.Simplicials;
 
 		TArray<FTransform> Transform;
 		GeometryCollectionAlgo::GlobalMatrices(DynamicCollection.Transform, DynamicCollection.Parent, Transform);
@@ -745,7 +755,7 @@ void FGeometryCollectionPhysicsProxy::InitializeBodiesPT(
 					Handle,
 					//Particles, 
 					Parameters.Shared,
-					nullptr,///Simplicials[TransformGroupIndex].Get(),
+					Simplicials[TransformGroupIndex].Get(),
 					Implicits[TransformGroupIndex],
 					SimFilter,
 					QueryFilter,
@@ -1851,9 +1861,11 @@ void FGeometryCollectionPhysicsProxy::InitializeSharedCollisionStructures(
 	TManagedArray<TUniquePtr<FSimplicial>>& CollectionSimplicials = 
 		RestCollection.AddAttribute<TUniquePtr<FSimplicial>>(
 			FGeometryDynamicCollection::SimplicialsAttribute, FTransformCollection::TransformGroup);
+
+	RestCollection.RemoveAttribute(FGeometryDynamicCollection::ImplicitsAttribute, FTransformCollection::TransformGroup);
 	TManagedArray<FGeometryDynamicCollection::FSharedImplicit>& CollectionImplicits = 
 		RestCollection.AddAttribute<FGeometryDynamicCollection::FSharedImplicit>(
-			FGeometryDynamicCollection::SharedImplicitsAttribute, FTransformCollection::TransformGroup);
+			FGeometryDynamicCollection::ImplicitsAttribute, FTransformCollection::TransformGroup);
 
 	FTransform IdentityXf(FQuat::Identity, FVector(0));
 	IdentityXf.NormalizeRotation();
@@ -2154,7 +2166,8 @@ void FGeometryCollectionPhysicsProxy::InitializeSharedCollisionStructures(
 
 			if (CollectionImplicits[TransformGroupIndex] && CollectionImplicits[TransformGroupIndex]->HasBoundingBox())
 			{
-				const auto BBox = CollectionImplicits[TransformGroupIndex]->BoundingBox();
+				const auto Implicit = CollectionImplicits[TransformGroupIndex];
+				const auto BBox = Implicit->BoundingBox();
 				const TVector<float, 3> Extents = BBox.Extents(); // Chaos::TAABB::Extents() is Max - Min
 				MaxChildBounds = MaxChildBounds.ComponentwiseMax(Extents);
 			}
diff --git a/Engine/Source/Runtime/Experimental/ChaosSolvers/Public/GeometryCollectionProxyData.h b/Engine/Source/Runtime/Experimental/ChaosSolvers/Public/GeometryCollectionProxyData.h
--- a/Engine/Source/Runtime/Experimental/ChaosSolvers/Public/GeometryCollectionProxyData.h
+++ b/Engine/Source/Runtime/Experimental/ChaosSolvers/Public/GeometryCollectionProxyData.h
@@ -84,10 +84,8 @@ public:
 	TManagedArray<FTransform> MassToLocal;
 	//TManagedArray<TArray<FCollisionFilterData>> ShapeQueryData;
 	//TManagedArray<TArray<FCollisionFilterData>> ShapeSimData;
-	//TManagedArray<TUniquePtr<FCollisionStructureManager::FSimplicial>> Simplicials;
+	TManagedArray<TUniquePtr<FCollisionStructureManager::FSimplicial>> Simplicials;
 	TManagedArray<bool> SimulatableParticles;
-
-
 };
 
 
diff --git a/Engine/Source/Runtime/Experimental/GeometryCollectionCore/Private/GeometryCollection/ManagedArrayCollection.cpp b/Engine/Source/Runtime/Experimental/GeometryCollectionCore/Private/GeometryCollection/ManagedArrayCollection.cpp
--- a/Engine/Source/Runtime/Experimental/GeometryCollectionCore/Private/GeometryCollection/ManagedArrayCollection.cpp
+++ b/Engine/Source/Runtime/Experimental/GeometryCollectionCore/Private/GeometryCollection/ManagedArrayCollection.cpp
@@ -278,7 +278,9 @@ void FManagedArrayCollection::SyncGroupSizeFrom(const FManagedArrayCollection& I
 	Resize(InCollection.GroupInfo[Group].Size, Group);
 }
 
-void FManagedArrayCollection::CopyMatchingAttributesFrom(const FManagedArrayCollection& InCollection)
+void FManagedArrayCollection::CopyMatchingAttributesFrom(
+	const FManagedArrayCollection& InCollection,
+	const TMap<FName, TSet<FName>>* SkipList)
 {
 	for (const auto& Pair : InCollection.GroupInfo)
 	{
@@ -286,6 +288,16 @@ void FManagedArrayCollection::CopyMatchingAttributesFrom(const FManagedArrayColl
 	}
 	for (TTuple<FKeyType, FValueType>& Entry : Map)
 	{
+		if (SkipList)
+		{
+			if (const TSet<FName>* Attrs = SkipList->Find(Entry.Key.Get<1>()))
+			{
+				if (Attrs->Contains(Entry.Key.Get<0>()))
+				{
+					continue;
+				}
+			}
+		}
 		if (InCollection.HasAttribute(Entry.Key.Get<0>(), Entry.Key.Get<1>()))
 		{
 			const FValueType& OriginalValue = InCollection.Map[Entry.Key];
diff --git a/Engine/Source/Runtime/Experimental/GeometryCollectionCore/Public/GeometryCollection/ManagedArrayCollection.h b/Engine/Source/Runtime/Experimental/GeometryCollectionCore/Public/GeometryCollection/ManagedArrayCollection.h
--- a/Engine/Source/Runtime/Experimental/GeometryCollectionCore/Public/GeometryCollection/ManagedArrayCollection.h
+++ b/Engine/Source/Runtime/Experimental/GeometryCollectionCore/Public/GeometryCollection/ManagedArrayCollection.h
@@ -308,8 +308,9 @@ public:
 	/**
 	* Copy attributes that match the input collection. This is a utility to easily sync collections
 	* @param InCollection - All groups from this collection found in the input will be sized accordingly
+	* @param SkipList - Group/Attrs to skip. Keys are group names, values are attributes in those groups.
 	*/
-	void CopyMatchingAttributesFrom(const FManagedArrayCollection& InCollection);
+	void CopyMatchingAttributesFrom(const FManagedArrayCollection& InCollection, const TMap<FName, TSet<FName>>* SkipList=nullptr);
 
 	/**
 	* Number of elements in a group
