Mark geometry collection particles as HasBounds = true if they have an implicit shape.

diff --git a/Engine/Source/Runtime/Experimental/ChaosSolvers/Private/PhysicsProxy/GeometryCollectionPhysicsProxy.cpp b/Engine/Source/Runtime/Experimental/ChaosSolvers/Private/PhysicsProxy/GeometryCollectionPhysicsProxy.cpp
--- a/Engine/Source/Runtime/Experimental/ChaosSolvers/Private/PhysicsProxy/GeometryCollectionPhysicsProxy.cpp
+++ b/Engine/Source/Runtime/Experimental/ChaosSolvers/Private/PhysicsProxy/GeometryCollectionPhysicsProxy.cpp
@@ -284,6 +284,7 @@ void PopulateSimulatedParticle(
 	{
 		TSharedPtr<Chaos::FImplicitObject, ESPMode::ThreadSafe> SharedImplicitTS(Implicit->DeepCopy().Release());
 		Handle->SetSharedGeometry(SharedImplicitTS);
+		Handle->SetHasBounds(true);
 	}
 
 	Handle->CollisionParticlesInitIfNeeded();
