Fix hoagie not going to sleep when flipped.

diff --git a/Engine/Source/Runtime/Experimental/ChaosSolvers/Private/PhysicsProxy/SingleParticlePhysicsProxy.cpp b/Engine/Source/Runtime/Experimental/ChaosSolvers/Private/PhysicsProxy/SingleParticlePhysicsProxy.cpp
--- a/Engine/Source/Runtime/Experimental/ChaosSolvers/Private/PhysicsProxy/SingleParticlePhysicsProxy.cpp
+++ b/Engine/Source/Runtime/Experimental/ChaosSolvers/Private/PhysicsProxy/SingleParticlePhysicsProxy.cpp
@@ -477,12 +477,20 @@ void FSingleParticlePhysicsProxy<Chaos::TPBDRigidParticle<float, 3>>::PushToPhys
 		if (Data->DirtyFlags.IsDirty(Chaos::EParticleFlags::V))
 		{
 			RigidHandle->SetV(Data->MV);
-			bDynamicPropertyUpdated = true;
+			//TODO: This is gross, we should clean it up
+			if (!(Data->MV.IsNearlyZero() && Data->MObjectState == Chaos::EObjectStateType::Sleeping))
+			{
+				bDynamicPropertyUpdated = true;
+			}
 		}
 		if (Data->DirtyFlags.IsDirty(Chaos::EParticleFlags::W))
 		{
 			RigidHandle->SetW(Data->MW);
-			bDynamicPropertyUpdated = true;
+			//TODO: Same as above
+			if (!(Data->MW.IsNearlyZero() && Data->MObjectState == Chaos::EObjectStateType::Sleeping))
+			{
+				bDynamicPropertyUpdated = true;
+			}
 		}
 
 		RigidHandle->SetCenterOfMass(Data->MCenterOfMass);
