Chaos: fix crash in FPBDConstraintGraph::ComputeIslands originating from fields disabling GC particle

diff --git a/Engine/Source/Runtime/Experimental/Chaos/Private/Chaos/PBDRigidsEvolutionGBF.cpp b/Engine/Source/Runtime/Experimental/Chaos/Private/Chaos/PBDRigidsEvolutionGBF.cpp
--- a/Engine/Source/Runtime/Experimental/Chaos/Private/Chaos/PBDRigidsEvolutionGBF.cpp
+++ b/Engine/Source/Runtime/Experimental/Chaos/Private/Chaos/PBDRigidsEvolutionGBF.cpp
@@ -411,7 +411,7 @@ void FPBDRigidsEvolutionGBF::AdvanceOneTimeStep(const FReal Dt, const FReal Step
 			}
 			for (const auto Particle : DisabledParticles[Island])
 			{
-				Particles.DisableParticle(Particle);
+				DisableParticle(Particle);
 			}
 		}
 	}
