Fix rewind data writing even when properties were clean.

diff --git a/Engine/Source/Programs/HeadlessChaos/Private/HeadlessChaosTestRewind.cpp b/Engine/Source/Programs/HeadlessChaos/Private/HeadlessChaosTestRewind.cpp
--- a/Engine/Source/Programs/HeadlessChaos/Private/HeadlessChaosTestRewind.cpp
+++ b/Engine/Source/Programs/HeadlessChaos/Private/HeadlessChaosTestRewind.cpp
@@ -18,20 +18,32 @@ namespace ChaosTest {
 
     using namespace Chaos;
 
+	void TickSolverHelper(FChaosSolversModule* Module, FPhysicsSolver* Solver)
+	{
+		Solver->PushPhysicsState(Module->GetDispatcher());
+		FPhysicsSolverAdvanceTask AdvanceTask(Solver,1.0f);
+		AdvanceTask.DoTask(ENamedThreads::GameThread,FGraphEventRef());
+		Solver->BufferPhysicsResults();
+		Solver->FlipBuffers();
+		Solver->UpdateGameThreadStructures();
+	}
+
 	GTEST_TEST(RewindTest,DataCapture)
 	{
 		auto Sphere = TSharedPtr<FImplicitObject,ESPMode::ThreadSafe>(new TSphere<float,3>(TVector<float,3>(0),10));
+		auto Box = TSharedPtr<FImplicitObject,ESPMode::ThreadSafe>(new TBox<float,3>(FVec3(0),FVec3(1)));
+		auto Box2 = TSharedPtr<FImplicitObject,ESPMode::ThreadSafe>(new TBox<float,3>(FVec3(2),FVec3(3)));
 
 		FChaosSolversModule* Module = FChaosSolversModule::GetModule();
 		Module->ChangeThreadingMode(EChaosThreadingMode::SingleThread);
 
 		// Make a solver
-		Chaos::FPhysicsSolver* Solver = Module->CreateSolver(nullptr,ESolverFlags::Standalone);
+		FPhysicsSolver* Solver = Module->CreateSolver(nullptr,ESolverFlags::Standalone);
 		Solver->SetEnabled(true);
 
 		// Make a particle
 
-		auto Particle = Chaos::TKinematicGeometryParticle<float,3>::CreateParticle();
+		auto Particle = TKinematicGeometryParticle<float,3>::CreateParticle();
 
 		Particle->SetGeometry(Sphere);
 		Solver->RegisterObject(Particle.Get());
@@ -39,19 +51,34 @@ namespace ChaosTest {
 
 		for(int Step = 0; Step < 11; ++Step)
 		{
+			//property that changes every step
 			Particle->SetX(FVec3(0,0,100 - Step));
 
-			Solver->PushPhysicsState(Module->GetDispatcher());
-			FPhysicsSolverAdvanceTask AdvanceTask(Solver,1.0f);
-			AdvanceTask.DoTask(ENamedThreads::GameThread,FGraphEventRef());
-			Solver->BufferPhysicsResults();
-			Solver->FlipBuffers();
-			Solver->UpdateGameThreadStructures();
+			//property that changes once half way through
+			if(Step == 3)
+			{
+				Particle->SetGeometry(Box);
+			}
+
+			if(Step == 5)
+			{
+				Particle->SetGeometry(Box2);
+			}
+
+			if(Step == 7)
+			{
+				Particle->SetGeometry(Box);
+			}
+
+			TickSolverHelper(Module, Solver);
 		}
 
 		//ended up at z = 90
 		EXPECT_EQ(Particle->X()[2],90);
 
+		//ended up with box geometry
+		EXPECT_EQ(Box.Get(),Particle->Geometry().Get());
+		
 		const FRewindData* RewindData = Solver->GetRewindData();
 
 		//check state at every step except latest
@@ -63,12 +90,28 @@ namespace ChaosTest {
 			if(ParticleState)
 			{
 				EXPECT_EQ(ParticleState->X()[2],100 - Step);
+
+				if(Step < 3)
+				{
+					//was sphere
+					EXPECT_EQ(ParticleState->Geometry().Get(),Sphere.Get());
+				}
+				else if(Step < 5 || Step >= 7)
+				{
+					//then became box
+					EXPECT_EQ(ParticleState->Geometry().Get(),Box.Get());
+				}
+				else
+				{
+					//second box
+					EXPECT_EQ(ParticleState->Geometry().Get(),Box2.Get());
+				}
 			}
 		}
 		
 		//no data at head because we always save the previous frame
 		EXPECT_TRUE(RewindData->GetStateAtFrame(*Particle,10) == nullptr);
-		
+
 		// Throw out the proxy
 		Solver->UnregisterObject(Particle.Get());
 
diff --git a/Engine/Source/Runtime/Experimental/Chaos/Public/Chaos/ParticleDirtyFlags.h b/Engine/Source/Runtime/Experimental/Chaos/Public/Chaos/ParticleDirtyFlags.h
--- a/Engine/Source/Runtime/Experimental/Chaos/Public/Chaos/ParticleDirtyFlags.h
+++ b/Engine/Source/Runtime/Experimental/Chaos/Public/Chaos/ParticleDirtyFlags.h
@@ -862,6 +862,11 @@ public:
 		Flags.Clear();
 	}
 
+	bool IsDirty(EParticleFlags InBits) const
+	{
+		return Flags.IsDirty(InBits);
+	}
+
 #define PARTICLE_PROPERTY(PropName, Type)\
 Type const & Get##PropName(const FDirtyPropertiesManager& Manager, int32 Idx) const { return ReadImp<Type, EParticleProperty::PropName>(Manager, Idx); }\
 bool Has##PropName() const { return Flags.IsDirty(ParticlePropToFlag(EParticleProperty::PropName)); }\
diff --git a/Engine/Source/Runtime/Experimental/Chaos/Public/Chaos/ParticleHandle.h b/Engine/Source/Runtime/Experimental/Chaos/Public/Chaos/ParticleHandle.h
--- a/Engine/Source/Runtime/Experimental/Chaos/Public/Chaos/ParticleHandle.h
+++ b/Engine/Source/Runtime/Experimental/Chaos/Public/Chaos/ParticleHandle.h
@@ -428,6 +428,8 @@ public:
 	TSharedPtr<FImplicitObject, ESPMode::ThreadSafe> SharedGeometry() const { return GeometryParticles->SharedGeometry(ParticleIdx); }
 	void SetSharedGeometry(TSharedPtr<FImplicitObject, ESPMode::ThreadSafe> InGeometry) { GeometryParticles->SetSharedGeometry(ParticleIdx, InGeometry); }
 
+	TSharedPtr<FImplicitObject, ESPMode::ThreadSafe> SharedGeometryLowLevel() const { return GeometryParticles->SharedGeometry(ParticleIdx); }
+
 	const TUniquePtr<FImplicitObject>& DynamicGeometry() const { return GeometryParticles->DynamicGeometry(ParticleIdx); }
 	void SetDynamicGeometry(TUniquePtr<FImplicitObject>&& Unique) { GeometryParticles->SetDynamicGeometry(ParticleIdx, MoveTemp(Unique)); }
 
@@ -1611,8 +1613,6 @@ public:
 		}
 	}
 
-	const FParticlePositionRotation& XR() const { return MXR.Read(); }
-
 protected:
 
 	// Pointer to any data that the solver wants to associate with this particle
diff --git a/Engine/Source/Runtime/Experimental/ChaosSolvers/Public/RewindData.h b/Engine/Source/Runtime/Experimental/ChaosSolvers/Public/RewindData.h
--- a/Engine/Source/Runtime/Experimental/ChaosSolvers/Public/RewindData.h
+++ b/Engine/Source/Runtime/Experimental/ChaosSolvers/Public/RewindData.h
@@ -24,17 +24,10 @@ public:
 	{
 	}
 
-	template <typename LambdaRead>
-	const T& Read(const LambdaRead& ReadFunc) const
+	const T& Read() const
 	{
-		if(Manager)
-		{
-			const TDirtyElementPool<T>& Pool = Manager->GetParticlePool<T,PropName>();
-			return Pool.GetElement(Idx);
-		} else
-		{
-			return ReadFunc();
-		}
+		const TDirtyElementPool<T>& Pool = Manager->GetParticlePool<T,PropName>();
+		return Pool.GetElement(Idx);
 	}
 
 	template <typename LambdaWrite>
@@ -49,12 +42,19 @@ public:
 	}
 
 	template <typename LambdaSet>
-	void SyncRemoteData(FDirtyPropertiesManager* InManager,int32 InIdx,const LambdaSet& SetFunc)
+	void SyncRemoteData(FDirtyPropertiesManager* InManager,int32 InIdx, const FParticleDirtyData& DirtyData, const LambdaSet& SetFunc)
 	{
-		ensure(Manager == nullptr);
-		Manager = InManager;
-		Idx = InIdx;
-		SetFunc(Manager->GetParticlePool<T,PropName>().GetElement(Idx));
+		if(DirtyData.IsDirty(ParticlePropToFlag(PropName)))
+		{
+			Manager = InManager;
+			Idx = InIdx;
+			SetFunc(Manager->GetParticlePool<T,PropName>().GetElement(Idx));
+		}
+	}
+
+	bool IsSet() const
+	{
+		return Manager != nullptr;
 	}
 
 private:
@@ -73,24 +73,61 @@ public:
 
 	const FVec3& X() const
 	{
-		return ParticlePositionRotation.Read([this]() -> const auto& {return Particle->XR();}).X;
+		return ParticlePositionRotation.IsSet() ? ParticlePositionRotation.Read().X : Particle->X();
+	}
+
+	TSerializablePtr<FImplicitObject> Geometry() const
+	{
+		return NonFrequentData.IsSet() ? MakeSerializable(NonFrequentData.Read().Geometry) : Particle->Geometry();
 	}
 
 	void SyncRemoteData(FDirtyPropertiesManager* Manager,int32 Idx,const FDirtyProxy& Dirty)
 	{
 		const auto Proxy = static_cast<const FGeometryParticlePhysicsProxy*>(Dirty.Proxy);
 		const auto Handle = Proxy->GetHandle();
-		ParticlePositionRotation.SyncRemoteData(Manager,Idx,[Handle](FParticlePositionRotation& Data)
+		ParticlePositionRotation.SyncRemoteData(Manager,Idx,Dirty.ParticleData, [Handle](FParticlePositionRotation& Data)
 		{
 			Data.X = Handle->X();
 			Data.R = Handle->R();
 		});
+
+		NonFrequentData.SyncRemoteData(Manager,Idx,Dirty.ParticleData, [Handle](FParticleNonFrequentData& Data)
+		{
+			Data.Geometry = Handle->SharedGeometryLowLevel();
+			Data.UserData = Handle->UserData();
+
+			//note: this data is keyed based on unique idx so it's not really possible to change this
+			//but we save it anyway since it's part of a big struct
+			Data.UniqueIdx = Handle->UniqueIdx();
+#if CHAOS_CHECKED
+			Data.DebugName = Handle->DebugName();
+#endif
+		});
+	}
+
+	bool CoalesceState(const FGeometryParticleState& LatestState)
+	{
+		bool bCoalesced = false;
+		if(!ParticlePositionRotation.IsSet() && LatestState.ParticlePositionRotation.IsSet())
+		{
+			ParticlePositionRotation = LatestState.ParticlePositionRotation;
+			bCoalesced = true;
+		}
+
+		if(!NonFrequentData.IsSet() && LatestState.NonFrequentData.IsSet())
+		{
+			NonFrequentData = LatestState.NonFrequentData;
+			bCoalesced = true;
+		}
+
+		return bCoalesced;
 	}
 
 protected:
 	TGeometryParticle<FReal,3>* Particle;
 private:
 	TParticleStateProperty<FParticlePositionRotation,EParticleProperty::XR> ParticlePositionRotation;
+	TParticleStateProperty<FParticleNonFrequentData, EParticleProperty::NonFrequentData> NonFrequentData;
 };
 
 class FRewindData
@@ -138,7 +175,20 @@ public:
 				const auto PTParticle = Proxy->GetHandle();
 				TArray<FFrameInfo>& Frames = ParticleToFrameInfo.FindOrAdd(PTParticle->UniqueIdx());
 				Frames.Add(FFrameInfo{FGeometryParticleState(GTParticleUnsafe),CurFrame-1});
-				Frames.Last().State.SyncRemoteData(Manager,DataIdx,Dirty);
+
+				FGeometryParticleState& LatestState = Frames.Last().State;
+				LatestState.SyncRemoteData(Manager,DataIdx,Dirty);
+
+				//coalesce any previous data
+				for(int32 FrameIdx = Frames.Num() - 2; FrameIdx >= 0; --FrameIdx)
+				{
+					FFrameInfo& Frame = Frames[FrameIdx];
+					if(Frame.State.CoalesceState(LatestState) == false)
+					{
+						//nothing to coalesce so no need to check earlier frames
+						break;
+					}
+				}
 			}
 		};
 
