Chaos : GeometryCollection Center of Mass
- Removed usage of MassToLocal for simplicial mass alignment.
- MassToLocal is only used for initialization of the Particle CenterOfMass and RotationOfMass.

diff --git a/Engine/Source/Runtime/Experimental/ChaosSolvers/Private/PhysicsProxy/GeometryCollectionPhysicsProxy.cpp b/Engine/Source/Runtime/Experimental/ChaosSolvers/Private/PhysicsProxy/GeometryCollectionPhysicsProxy.cpp
--- a/Engine/Source/Runtime/Experimental/ChaosSolvers/Private/PhysicsProxy/GeometryCollectionPhysicsProxy.cpp
+++ b/Engine/Source/Runtime/Experimental/ChaosSolvers/Private/PhysicsProxy/GeometryCollectionPhysicsProxy.cpp
@@ -230,6 +230,7 @@ void PopulateSimulatedParticle(
 	const FCollisionFilterData QueryFilterIn,
 	float MassIn,
 	FVector InertiaTensorVec,
+	const FTransform& MassOffset,
 	const FTransform& WorldTransform, 
 	const uint8 DynamicState, 
 	const int16 CollisionGroup)
@@ -244,6 +245,8 @@ void PopulateSimulatedParticle(
 	Handle->SetP(Handle->X());
 	Handle->SetQ(Handle->R());
 	Handle->SetIsland(INDEX_NONE);
+	Handle->SetCenterOfMass(MassOffset.GetTranslation());
+	Handle->SetRotationOfMass(MassOffset.GetRotation());
 
 	//
 	// Setup Mass
@@ -741,31 +744,19 @@ void FGeometryCollectionPhysicsProxy::InitializeBodiesPT(
 			const int32 TransformGroupIndex = TransformIndex[GeometryIndex];
 			if (FClusterHandle* Handle = SolverParticleHandles[TransformGroupIndex])
 			{
-				// TODO: We probably want to take better advantage of this solver
-				// feature, rather than doing it all here ourselves?
-				Handle->SetCenterOfMass(MassToLocal[TransformGroupIndex].GetTranslation());
-				Handle->SetRotationOfMass(MassToLocal[TransformGroupIndex].GetRotation());
-
 				// Mass space -> Composed parent space -> world
-				const FTransform WorldTransform =
-					MassToLocal[TransformGroupIndex] * Transform[TransformGroupIndex] * Parameters.WorldTransform;
+				const FTransform WorldTransform = Transform[TransformGroupIndex] * Parameters.WorldTransform;
 
 				PopulateSimulatedParticle(
 					Handle,
-					//Particles, 
 					Parameters.Shared,
 					Simplicials[TransformGroupIndex].Get(),
 					Implicits[TransformGroupIndex],
 					SimFilter,
 					QueryFilter,
-#if TODO_REIMPLEMENT_SOLVER_SETTINGS_ACCESSORS
-					RigidsSolver->GetMassScale() * Mass[TransformGroupIndex],
-					RigidsSolver->GetMassScale() * InertiaTensor[TransformGroupIndex],
-#else // TODO_REIMPLEMENT_SOLVER_SETTINGS_ACCESSORS
 					Mass[TransformGroupIndex],
 					InertiaTensor[TransformGroupIndex],
-#endif // TODO_REIMPLEMENT_SOLVER_SETTINGS_ACCESSORS
-					//RigidBodyIndex, 
+					MassToLocal[TransformGroupIndex],
 					WorldTransform,
 					(uint8)DynamicState[TransformGroupIndex],
 					CollisionGroup[TransformGroupIndex]);
@@ -1016,7 +1007,7 @@ FGeometryCollectionPhysicsProxy::BuildClusters(
 
 	//If we are a root particle use the world transform, otherwise set the relative transform
 	const FTransform CollectionSpaceTransform = GeometryCollectionAlgo::GlobalMatrix(Transform, ParentIndex, CollectionClusterIndex);
-	const Chaos::TRigidTransform<float, 3> ParticleTM = MassToLocal[CollectionClusterIndex] * CollectionSpaceTransform * Parameters.WorldTransform;
+	const Chaos::TRigidTransform<float, 3> ParticleTM = CollectionSpaceTransform * Parameters.WorldTransform;
 
 	//create new cluster particle
 	//The reason we need to pass in a mass orientation override is as follows:
@@ -1083,16 +1074,9 @@ FGeometryCollectionPhysicsProxy::BuildClusters(
 		nullptr, // Parent->Geometry() ? Parent->Geometry() : Implicits[CollectionClusterIndex], 
 		SimFilter,
 		QueryFilter,
-#if TODO_REIMPLEMENT_SOLVER_SETTINGS_ACCESSORS
-		// Note: if the children had this scaling baked in, then cluster parents
-		// will too, and we shouldn't apply the scaling to the parents twice!
-		ThisSolver->GetMassScale() * Mass[CollectionClusterIndex], 
-		ThisSolver->GetMassScale() * InertiaTensor[CollectionClusterIndex],
-#else
-		// We computed M and I from children in CreateClusterParticle().
 		Parent->M() > 0.0 ? Parent->M() : Mass[CollectionClusterIndex], 
 		Parent->I().GetDiagonal() != Chaos::TVector<float,3>(0.0) ? Parent->I().GetDiagonal() : InertiaTensor[CollectionClusterIndex],
-#endif
+		MassToLocal[CollectionClusterIndex],
 		ParticleTM, 
 		(uint8)DynamicState[CollectionClusterIndex], 
 		0); // CollisionGroup
@@ -1519,23 +1503,15 @@ void FGeometryCollectionPhysicsProxy::BufferPhysicsResults()
 			//   bodies out of the collections hierarchy during the simulation.  
 			if (!Handle->Disabled())
 			{
-				const TManagedArray<FTransform>& MassToLocal = PTDynamicCollection.MassToLocal;
-
 				// Update the transform of the active body. The active body can be either a single rigid
 				// or a collection of rigidly attached geometries (Clustering). The cluster is represented as a
 				// single transform in the GeometryCollection, and all children are stored in the local space
 				// of the parent cluster.
-				// ... When setting cluster transforms it is expected that the MassToLocal is identity.
-				//     Cluster initialization will set the vertices in the MassSpace of the rigid body.
-				// ... When setting individual rigid bodies that are not clustered, the MassToLocal will be 
-				//     non-Identity, and will reflect the difference between the geometric center of the geometry
-				//     and that corresponding rigid bodies center of mass. 
-
+	
 				FTransform& ParticleToWorld = TargetResults.ParticleToWorldTransforms[TransformIndex];
 				ParticleToWorld = FTransform(Handle->R(), Handle->X());
 
-				TargetResults.Transforms[TransformIndex] =
-					MassToLocal[TransformIndex].GetRelativeTransformReverse(ParticleToWorld).GetRelativeTransform(ActorToWorld);
+				TargetResults.Transforms[TransformIndex] = ParticleToWorld.GetRelativeTransform(ActorToWorld);
 				TargetResults.Transforms[TransformIndex].NormalizeRotation();
 
 				// If the parent of this NON DISABLED body is set to anything other than INDEX_NONE,
@@ -1585,7 +1561,6 @@ void FGeometryCollectionPhysicsProxy::BufferPhysicsResults()
 					{
 						if (ClusteredParentHandle->InternalCluster())
 						{
-							const TManagedArray<FTransform>& MassToLocal = PTDynamicCollection.MassToLocal;
 
 							//const FTransform ClusterChildToWorld = ClusterChildToParentMap[RigidBodyIndex] * FTransform(Particles.R(ClusterParentIndex), Particles.X(ClusterParentIndex));
 							const TRigidTransform<float, 3>& ChildToParent = Handle->ChildToParent();
@@ -1595,9 +1570,7 @@ void FGeometryCollectionPhysicsProxy::BufferPhysicsResults()
 								Handle->SetX(ClusterChildToWorld.GetTranslation());
 								Handle->SetR(ClusterChildToWorld.GetRotation());
 							}
-							// GeomToActor = ActorToWorld.Inv() * ClusterChildToWorld * MassToLocal.Inv();
-							TargetResults.Transforms[TransformIndex] = 
-								MassToLocal[TransformIndex].GetRelativeTransformReverse(ClusterChildToWorld).GetRelativeTransform(ActorToWorld);
+							TargetResults.Transforms[TransformIndex] = ClusterChildToWorld.GetRelativeTransform(ActorToWorld);
 							TargetResults.Transforms[TransformIndex].NormalizeRotation();
 						}
 					}
@@ -1727,14 +1700,6 @@ void FGeometryCollectionPhysicsProxy::UpdateGeometryCollection( const FGeometryC
 	TManagedArray<int32>& SimulationType = Collection->SimulationType;
 	TManagedArray<int32>& DynamicState = Collection->GetAttribute<int32>("DynamicState", FGeometryCollection::TransformGroup);
 
-	// Particles X and R are aligned with center of mass and inertia principal axes.
-	// Renderer doesn't know about this and simply does ActorToWorld * GeomToActor * LocalSpaceVerts
-	// In proper math multiplication order:
-	//		ParticleToWorld = ActorToWorld * GeomToActor * MassToLocal
-	//		GeomToWorld = ActorToWorld * GeomToActor
-	//		=> GeomToWorld = ParticleToWorld * MassToLocal.Inv()
-	//		=> GeomToActor = ActorToWorld.Inv() * ParticleToWorld * MassToLocal.Inv()
-
 	const FTransform& ActorToWorld = Parameters.WorldTransform;
 	const int32 TransformSize = Collection->NumElements(FGeometryCollection::TransformGroup);
 	for (int32 TransformGroupIndex = 0; TransformGroupIndex < TransformSize; ++TransformGroupIndex)
@@ -1765,17 +1730,7 @@ void FGeometryCollectionPhysicsProxy::UpdateGeometryCollection( const FGeometryC
 		if(!GCResults.DisabledStates[TransformGroupIndex])
 		{
 #endif
-			// Update the transform of the active body. The active body can be either a single rigid
-			// or a collection of rigidly attached geometries (Clustering). The cluster is represented as a
-			// single transform in the GeometryCollection, and all children are stored in the local space
-			// of the parent cluster.
-			// ... When setting cluster transforms it is expected that the MassToLocal is identity.
-			//     Cluster initialization will set the vertices in the MassSpace of the rigid body.
-			// ... When setting individual rigid bodies that are not clustered, the MassToLocal will be 
-			//     non-Identity, and will reflect the difference between the geometric center of the geometry
-			//     and that corresponding rigid bodies center of mass. 
-
-			// GeomToActor = ActorToWorld.Inv() * ParticleToWorld * MassToLocal.Inv();
+
 			Transform[TransformGroupIndex] = GCResults.Transforms[TransformGroupIndex];
 
 			// When a leaf node rigid body is removed from a cluster the rigid
@@ -1822,8 +1777,8 @@ int32 FindSizeSpecificIdx(const TArray<FSharedSimulationSizeSpecificData>& SizeS
 	by a rotation of the DDC key in FDerivedDataGeometryCollectionCooker::GetVersionString
 */
 void FGeometryCollectionPhysicsProxy::InitializeSharedCollisionStructures(
-	Chaos::FErrorReporter& ErrorReporter, 
-	FGeometryCollection& RestCollection, 
+	Chaos::FErrorReporter& ErrorReporter,
+	FGeometryCollection& RestCollection,
 	const FSharedSimulationParameters& SharedParams)
 {
 	FString BaseErrorPrefix = ErrorReporter.GetPrefix();
@@ -1832,40 +1787,48 @@ void FGeometryCollectionPhysicsProxy::InitializeSharedCollisionStructures(
 	if (RestCollection.NumElements(FGeometryCollection::GeometryGroup) == 0)
 		return;
 
+	// clamps
+	const float MinBoundsExtents = SharedParams.MinimumBoundingExtentClamp;
+	const float MaxBoundsExtents = SharedParams.MaximumBoundingExtentClamp;
+	const float MinVolume = SharedParams.MinimumVolumeClamp();
+	const float MaxVolume = SharedParams.MaximumVolumeClamp();
+	const float MinMass = FMath::Max(SMALL_NUMBER, SharedParams.MaximumMassClamp);
+	const float MaxMass = SharedParams.MinimumMassClamp;
+
+
 	//TArray<TArray<TArray<int32>>> BoundaryVertexIndices;
 	//GeometryCollectionAlgo::FindOpenBoundaries(&RestCollection, 1e-2, BoundaryVertexIndices);
 	//GeometryCollectionAlgo::TriangulateBoundaries(&RestCollection, BoundaryVertexIndices);
 	//RestCollection.ReindexMaterials();
 
 	using namespace Chaos;
-	const TManagedArray<bool>& Visible = RestCollection.Visible;
 
 	// TransformGroup
 	const TManagedArray<int32>& BoneMap = RestCollection.BoneMap;
 	const TManagedArray<int32>& Parent = RestCollection.Parent;
 	const TManagedArray<TSet<int32>>& Children = RestCollection.Children;
-	const TManagedArray<bool>& CollectionSimulatableParticles = 
+	TManagedArray<bool>& CollectionSimulatableParticles =
 		RestCollection.GetAttribute<bool>(
 			FGeometryCollection::SimulatableParticlesAttribute, FTransformCollection::TransformGroup);
-
-	TManagedArray<FVector>& CollectionInertiaTensor = 
+	TManagedArray<FVector>& CollectionInertiaTensor =
 		RestCollection.AddAttribute<FVector>(
 			TEXT("InertiaTensor"), FTransformCollection::TransformGroup);
-	TManagedArray<FTransform>& CollectionMassToLocal = 
-		RestCollection.AddAttribute<FTransform>(
-			TEXT("MassToLocal"), FTransformCollection::TransformGroup);
-	TManagedArray<float>& CollectionMass = 
+	TManagedArray<float>& CollectionMass =
 		RestCollection.AddAttribute<float>(
 			TEXT("Mass"), FTransformCollection::TransformGroup);
-	TManagedArray<TUniquePtr<FSimplicial>>& CollectionSimplicials = 
+	TManagedArray<TUniquePtr<FSimplicial>>& CollectionSimplicials =
 		RestCollection.AddAttribute<TUniquePtr<FSimplicial>>(
 			FGeometryDynamicCollection::SimplicialsAttribute, FTransformCollection::TransformGroup);
 
 	RestCollection.RemoveAttribute(FGeometryDynamicCollection::ImplicitsAttribute, FTransformCollection::TransformGroup);
-	TManagedArray<FGeometryDynamicCollection::FSharedImplicit>& CollectionImplicits = 
+	TManagedArray<FGeometryDynamicCollection::FSharedImplicit>& CollectionImplicits =
 		RestCollection.AddAttribute<FGeometryDynamicCollection::FSharedImplicit>(
 			FGeometryDynamicCollection::ImplicitsAttribute, FTransformCollection::TransformGroup);
 
+	// @todo(chaos_transforms) : do we still use this?
+	TManagedArray<FTransform>& CollectionMassToLocal =
+		RestCollection.AddAttribute<FTransform>(
+			TEXT("MassToLocal"), FTransformCollection::TransformGroup);
 	FTransform IdentityXf(FQuat::Identity, FVector(0));
 	IdentityXf.NormalizeRotation();
 	CollectionMassToLocal.Fill(IdentityXf);
@@ -1873,6 +1836,9 @@ void FGeometryCollectionPhysicsProxy::InitializeSharedCollisionStructures(
 	// VerticesGroup
 	const TManagedArray<FVector>& Vertex = RestCollection.Vertex;
 
+	// Faces
+	const TManagedArray<bool>& Visible = RestCollection.Visible;
+
 	// GeometryGroup
 	const TManagedArray<FBox>& BoundingBox = RestCollection.BoundingBox;
 	TManagedArray<float>& InnerRadius = RestCollection.InnerRadius;
@@ -1888,7 +1854,6 @@ void FGeometryCollectionPhysicsProxy::InitializeSharedCollisionStructures(
 	{ // tmp scope
 		const TManagedArray<FTransform>& HierarchyTransform = RestCollection.Transform;
 		GeometryCollectionAlgo::GlobalMatrices(HierarchyTransform, Parent, CollectionSpaceTransforms);
-		check(HierarchyTransform.Num() == CollectionSpaceTransforms.Num());
 	} // tmp scope
 
 	const int32 NumTransforms = CollectionSpaceTransforms.Num();
@@ -1907,17 +1872,14 @@ void FGeometryCollectionPhysicsProxy::InitializeSharedCollisionStructures(
 		MassSpaceParticles.X(Idx) = Vertex[Idx];	//mass space computation done later down
 	}
 
-	TArray<TMassProperties<float,3>> MassPropertiesArray;
+	TArray<TMassProperties<float, 3>> MassPropertiesArray;
 	MassPropertiesArray.AddUninitialized(NumGeometries);
 
 	TArray<bool> InertiaComputationNeeded;
-	InertiaComputationNeeded.AddUninitialized(NumGeometries);
+	InertiaComputationNeeded.Init(false, NumGeometries);
 
-	const float MinBoundsExtents = 10.f;
 
 	float TotalVolume = 0.f;
-	int32 NumSimulatableParticles = 0;
-
 	// The geometry group has a set of transform indices that maps a geometry index
 	// to a transform index, but only in the case where there is a 1-to-1 mapping 
 	// between the two.  In the event where a geometry is instanced for multiple
@@ -1930,73 +1892,66 @@ void FGeometryCollectionPhysicsProxy::InitializeSharedCollisionStructures(
 		const int32 TransformGroupIndex = TransformIndex[GeometryIndex];
 		if (CollectionSimulatableParticles[TransformGroupIndex])
 		{
-			NumSimulatableParticles++;
-
 			TUniquePtr<TTriangleMesh<float>> TriMesh(
 				CreateTriangleMesh(
-					FaceCount[GeometryIndex], 
-					VertexStart[GeometryIndex], 
-					FaceStart[GeometryIndex], 
-					Vertex, 
-					Visible, 
-					Indices, 
+					FaceCount[GeometryIndex],
+					VertexStart[GeometryIndex],
+					FaceStart[GeometryIndex],
+					Vertex,
+					Visible,
+					Indices,
 					VertsAddedArray[GeometryIndex]));
 
 			TMassProperties<float, 3>& MassProperties = MassPropertiesArray[GeometryIndex];
 
-			//CalculateVolumeAndCenterOfMass(MassSpaceParticles, *TriMesh, MassProperties.Volume, MassProperties.CenterOfMass);
-			//InertiaComputationNeeded[GeometryIndex] = true;
-			//if (MassProperties.Volume < MinVolume)
-			//just assume everything is a box, need better computation later but requires mesh cleanup
 			{
-				//since volume is too small we just use a fallback one
-				InertiaComputationNeeded[GeometryIndex] = false;	
-
-				CollectionMassToLocal[TransformGroupIndex] = 
-					FTransform(TRotation<float, 3>(FQuat(0, 0, 0, 1)), BoundingBox[GeometryIndex].GetCenter());
-
+				MassProperties.CenterOfMass = FVector::ZeroVector;
+				MassProperties.RotationOfMass = FRotation3(FQuat::Identity).GetNormalized();
+				MassProperties.Volume = 0.f;
+				MassProperties.InertiaTensor = FMatrix33(1,1,1);
+				MassProperties.Mass = 1.0f; // start with unit mass, scaled later by density
 
-				const TVector<float, 3> Size = BoundingBox[GeometryIndex].GetSize().ComponentMax(FVector(MinBoundsExtents));
-				MassProperties.Volume = Size.Product();	//just fake a box if tiny
-				MassProperties.InertiaTensor = 
-					Chaos::TBox<float, 3>::GetInertiaTensor(
-						1.0f, // Mass
-						Size);
-				/*
-				if (MassProperties.Volume < MinVolume)
+				if (!ensureMsgf(BoundingBox[GeometryIndex].GetExtent().GetAbsMin() > MinVolume, TEXT("Geometry too small to simulate. Idx (%d)"), GeometryIndex))
+				{
+					CollectionSimulatableParticles[TransformGroupIndex] = false;	//do not simulate tiny particles
+					MassProperties.Mass = 0.f;
+					MassProperties.InertiaTensor = FMatrix33(0,0,0);
+				}
+				else
 				{
-					MassProperties.Volume = MinVolume;	//for thin shells volume is 0. Just use min volume to avoid divie by 0. Probably needs more thought
+					CalculateVolumeAndCenterOfMass(MassSpaceParticles, TriMesh->GetElements(), MassProperties.Volume, MassProperties.CenterOfMass);
 
-					//if all dimensions are tiny we should probably fix content
-					if(!ensureMsgf(BoundingBox[GeometryIndex].GetExtent().GetAbsMin() > MinVolume, TEXT("Geometry too small to simulate. Idx (%d)"), GeometryIndex))
+					if (MassProperties.Volume < MinVolume)
 					{
-						ErrorReporter.ReportError(*FString::Printf(TEXT("Geometry too small to simulate. Idx(%d)"), GeometryIndex));
-						CollectionSimulatableParticles[TransformGroupIndex] = false;	//do not simulate tiny particles
-						ErrorReporter.HandleLatestError();
+						// For rigid bodies outside of range just defaut to a clamped bounding box, and warn the user.
+						MassProperties.Volume = MinVolume;
+						CollectionMassToLocal[TransformGroupIndex] = FTransform(FQuat::Identity, BoundingBox[GeometryIndex].GetCenter());
 					}
-				}*/
-			}
-#if 0
-			else if (MassProperties.Volume > BoundsVolume)
-			{
-				ensure(false);
-				ErrorReporter.ReportError(*FString::Printf(TEXT("Geometry has invalid volume")));
-				ErrorReporter.HandleLatestError();
+					else if (MaxVolume < MassProperties.Volume)
+					{
+						// For rigid bodies outside of range just defaut to a clamped bounding box, and warn the user
+						MassProperties.Volume = MaxVolume;
+						CollectionMassToLocal[TransformGroupIndex] = FTransform(FQuat::Identity, BoundingBox[GeometryIndex].GetCenter());
+					}
+					else
+					{
+						InertiaComputationNeeded[GeometryIndex] = true;
 
-				//somehow ended up with huge volume, just use bounds volume
-				InertiaComputationNeeded[GeometryIndex] = false;
-				CollectionMassToLocal[TransformGroupIndex] = FTransform(TRotation<float, 3>(FQuat(0, 0, 0, 1)), BoundingBox[GeometryIndex].GetCenter());
+						if (!FMath::IsNearlyZero(MassProperties.CenterOfMass.SizeSquared()))
+						{
+							for (int32 Idx = 0; Idx < Vertex.Num(); ++Idx)
+							{
+								MassSpaceParticles.X(Idx) -= MassProperties.CenterOfMass;
+							}
+						}
 
-				FVector Size = BoundingBox[GeometryIndex].GetSize();
-				FVector SideSquared(Size.X * Size.X, Size.Y * Size.Y, Size.Z * Size.Z);
-				MassProperties.Volume = BoundsVolume;
-				MassProperties.InertiaTensor = PMatrix<float, 3, 3>((SideSquared.Y + SideSquared.Z) / 12.f, (SideSquared.X + SideSquared.Z) / 12.f, (SideSquared.X + SideSquared.Y) / 12.f);
+						CollectionMassToLocal[TransformGroupIndex] = FTransform(FQuat::Identity, MassProperties.CenterOfMass);
+					}
+				}
 			}
-#endif
 
 			if (InnerRadius[GeometryIndex] == 0.0f || OuterRadius[GeometryIndex] == 0.0f)
 			{
-				//const FVector Center = MassProperties.CenterOfMass;
 				const FVector Center = BoundingBox[GeometryIndex].GetCenter();
 				const int32 VStart = VertexStart[GeometryIndex];
 				const int32 VCount = VertexCount[GeometryIndex];
@@ -2025,74 +1980,62 @@ void FGeometryCollectionPhysicsProxy::InitializeSharedCollisionStructures(
 		}
 	}
 
-	ensureMsgf(TotalVolume > KINDA_SMALL_NUMBER || !NumSimulatableParticles, TEXT("Geometry collection too small"));
-	TotalVolume = FMath::Max(TotalVolume, MinBoundsExtents*MinBoundsExtents*MinBoundsExtents);
 	//User provides us with total mass or density.
 	//Density must be the same for individual parts and the total. Density_i = Density = Mass_i / Volume_i
 	//Total mass must equal sum of individual parts. Mass_i = TotalMass * Volume_i / TotalVolume => Density_i = TotalMass / TotalVolume
+	TotalVolume = FMath::Max(TotalVolume, MinBoundsExtents * MinBoundsExtents * MinBoundsExtents);
 	const float DesiredTotalMass = SharedParams.bMassAsDensity ? SharedParams.Mass * TotalVolume : SharedParams.Mass;
 	const float ClampedTotalMass = FMath::Clamp(DesiredTotalMass, SharedParams.MinimumMassClamp, SharedParams.MaximumMassClamp);
 	const float DesiredDensity = ClampedTotalMass / TotalVolume;
-	TVector<float, 3> MaxChildBounds(1);
 
+	TVector<float, 3> MaxChildBounds(1);
 	ParallelFor(NumGeometries, [&](int32 GeometryIndex)
 	//for (int32 GeometryIndex = 0; GeometryIndex < NumGeometries; GeometryIndex++)
 	{
 		const int32 TransformGroupIndex = TransformIndex[GeometryIndex];
+
+		const float Volume_i = MassPropertiesArray[GeometryIndex].Volume;
 		if (CollectionSimulatableParticles[TransformGroupIndex])
 		{
-			TUniquePtr<TTriangleMesh<float>>& TriMesh = TriangleMeshesArray[TransformGroupIndex];
-			TMassProperties<float, 3>& MassProperties = MassPropertiesArray[GeometryIndex];
-
 			//Must clamp each individual mass regardless of desired density
-			const float Volume_i = MassPropertiesArray[GeometryIndex].Volume;
 			if (DesiredDensity * Volume_i > SharedParams.MaximumMassClamp)
 			{
-				ensure(false);
+				// For rigid bodies outside of range just defaut to a clamped bounding box, and warn the user.
+				ErrorReporter.ReportError(*FString::Printf(TEXT("Geometry has invalid mass (to large)")));
+				ErrorReporter.HandleLatestError();
+
+				CollectionSimulatableParticles[TransformGroupIndex] = false;
 			}
+		}
+
+		if (CollectionSimulatableParticles[TransformGroupIndex])
+		{
+			TUniquePtr<TTriangleMesh<float>>& TriMesh = TriangleMeshesArray[TransformGroupIndex];
+			TMassProperties<float, 3>& MassProperties = MassPropertiesArray[GeometryIndex];
+
 			const float Mass_i = FMath::Max(DesiredDensity * Volume_i, SharedParams.MinimumMassClamp);
 			const float Density_i = Mass_i / Volume_i;
 			CollectionMass[TransformGroupIndex] = Mass_i;
 
 			if (InertiaComputationNeeded[GeometryIndex])
 			{
-				check(false); // Ryan - Currently everything is hard wired to be a box in the loop above, so this should never run.
-				CalculateInertiaAndRotationOfMass(
-					MassSpaceParticles, 
-					TriMesh->GetSurfaceElements(), 
-					Density_i, 
-					MassProperties.CenterOfMass, 
-					MassProperties.InertiaTensor, 
-					MassProperties.RotationOfMass);
-
-				CollectionMassToLocal[TransformGroupIndex] = 
-					FTransform(MassProperties.RotationOfMass, MassProperties.CenterOfMass);
-
-				const TVector<float, 3> DiagonalInertia(
-					MassProperties.InertiaTensor.M[0][0], MassProperties.InertiaTensor.M[1][1], MassProperties.InertiaTensor.M[2][2]);
-				//computation includes mass already
-				CollectionInertiaTensor[TransformGroupIndex] = DiagonalInertia;
+				CalculateInertiaAndRotationOfMass(MassSpaceParticles, TriMesh->GetSurfaceElements(), Density_i, MassProperties.CenterOfMass, MassProperties.InertiaTensor, MassProperties.RotationOfMass);
+				CollectionMassToLocal[TransformGroupIndex] = FTransform(MassProperties.RotationOfMass, MassProperties.CenterOfMass);
+				CollectionInertiaTensor[TransformGroupIndex] = TVector<float, 3>(MassProperties.InertiaTensor.M[0][0], MassProperties.InertiaTensor.M[1][1], MassProperties.InertiaTensor.M[2][2]);
 			}
 			else
 			{
-				const TVector<float, 3> DiagonalInertia(
-					MassProperties.InertiaTensor.M[0][0], 
-					MassProperties.InertiaTensor.M[1][1], 
-					MassProperties.InertiaTensor.M[2][2]);
-				//using fallback computation, but adjust for mass
+				const TVector<float, 3> DiagonalInertia(MassProperties.InertiaTensor.M[0][0], MassProperties.InertiaTensor.M[1][1], MassProperties.InertiaTensor.M[2][2]);
 				CollectionInertiaTensor[TransformGroupIndex] = DiagonalInertia * Mass_i;
 			}
 
-			// Update vertex buffer to be in mass space so that at runtime geometry aligns properly.
 			FBox InstanceBoundingBox(EForceInit::ForceInitToZero);
 			const TSet<int32>& VertsAdded = VertsAddedArray[GeometryIndex];
 			for (int32 VertIdx = VertexStart[GeometryIndex]; VertIdx < VertexStart[GeometryIndex] + VertexCount[GeometryIndex]; ++VertIdx)
 			{
+				MassSpaceParticles.X(VertIdx) += MassProperties.CenterOfMass;
 				if (VertsAdded.Contains(VertIdx))	//only consider verts from the trimesh
 				{
-					MassSpaceParticles.X(VertIdx) = 
-						CollectionMassToLocal[TransformGroupIndex].InverseTransformPosition(
-							MassSpaceParticles.X(VertIdx));
 					InstanceBoundingBox += MassSpaceParticles.X(VertIdx);	//build bounding box for visible verts in mass space
 				}
 			}
@@ -2104,32 +2047,50 @@ void FGeometryCollectionPhysicsProxy::InitializeSharedCollisionStructures(
 			//  Build the simplicial for the rest collection. This will be used later in the DynamicCollection to 
 			//  populate the collision structures of the simulation. 
 			//
-			check(TriMesh);
-			Chaos::TBVHParticles<float, 3> * Simplicial = 
-				FCollisionStructureManager::NewSimplicial(
-					MassSpaceParticles, 
-					BoneMap, 
-					SizeSpecificData.CollisionType, 
-					*TriMesh, 
-					SizeSpecificData.CollisionParticlesFraction);
-			CollectionSimplicials[TransformGroupIndex] = TUniquePtr<FSimplicial>(Simplicial);
-			ensureMsgf(CollectionSimplicials[TransformGroupIndex], TEXT("No simplicial representation."));
-
-			if (SizeSpecificData.ImplicitType == EImplicitTypeEnum::Chaos_Implicit_LevelSet)
+			if (ensureMsgf(TriMesh, TEXT("No Triangle representation")))
 			{
-				ErrorReporter.SetPrefix(BaseErrorPrefix + " | Transform Index: " + FString::FromInt(TransformGroupIndex) + " of " + FString::FromInt(TransformIndex.Num()));
-				CollectionImplicits[TransformGroupIndex] = FGeometryDynamicCollection::FSharedImplicit(
-					FCollisionStructureManager::NewImplicitLevelset(
-						ErrorReporter,
+				Chaos::TBVHParticles<float, 3>* Simplicial =
+					FCollisionStructureManager::NewSimplicial(
 						MassSpaceParticles,
+						BoneMap,
+						SizeSpecificData.CollisionType,
 						*TriMesh,
-						InstanceBoundingBox,
-						SizeSpecificData.MinLevelSetResolution,
-						SizeSpecificData.MaxLevelSetResolution,
-						SizeSpecificData.CollisionObjectReductionPercentage,
-						SizeSpecificData.CollisionType));
-				// Fall back on sphere if level set rasterization failed.
-				if (!CollectionImplicits[TransformGroupIndex])
+						SizeSpecificData.CollisionParticlesFraction);
+				CollectionSimplicials[TransformGroupIndex] = TUniquePtr<FSimplicial>(Simplicial);
+				ensureMsgf(CollectionSimplicials[TransformGroupIndex], TEXT("No simplicial representation."));
+
+				if (SizeSpecificData.ImplicitType == EImplicitTypeEnum::Chaos_Implicit_LevelSet)
+				{
+					ErrorReporter.SetPrefix(BaseErrorPrefix + " | Transform Index: " + FString::FromInt(TransformGroupIndex) + " of " + FString::FromInt(TransformIndex.Num()));
+					CollectionImplicits[TransformGroupIndex] = FGeometryDynamicCollection::FSharedImplicit(
+						FCollisionStructureManager::NewImplicitLevelset(
+							ErrorReporter,
+							MassSpaceParticles,
+							*TriMesh,
+							InstanceBoundingBox,
+							SizeSpecificData.MinLevelSetResolution,
+							SizeSpecificData.MaxLevelSetResolution,
+							SizeSpecificData.CollisionObjectReductionPercentage,
+							SizeSpecificData.CollisionType));
+					// Fall back on sphere if level set rasterization failed.
+					if (!CollectionImplicits[TransformGroupIndex])
+					{
+						CollectionImplicits[TransformGroupIndex] = FGeometryDynamicCollection::FSharedImplicit(
+							FCollisionStructureManager::NewImplicitSphere(
+								InnerRadius[GeometryIndex],
+								SizeSpecificData.CollisionObjectReductionPercentage,
+								SizeSpecificData.CollisionType));
+					}
+				}
+				else if (SizeSpecificData.ImplicitType == EImplicitTypeEnum::Chaos_Implicit_Box)
+				{
+					CollectionImplicits[TransformGroupIndex] = FGeometryDynamicCollection::FSharedImplicit(
+						FCollisionStructureManager::NewImplicitBox(
+							InstanceBoundingBox,
+							SizeSpecificData.CollisionObjectReductionPercentage,
+							SizeSpecificData.CollisionType));
+				}
+				else if (SizeSpecificData.ImplicitType == EImplicitTypeEnum::Chaos_Implicit_Sphere)
 				{
 					CollectionImplicits[TransformGroupIndex] = FGeometryDynamicCollection::FSharedImplicit(
 						FCollisionStructureManager::NewImplicitSphere(
@@ -2137,38 +2098,22 @@ void FGeometryCollectionPhysicsProxy::InitializeSharedCollisionStructures(
 							SizeSpecificData.CollisionObjectReductionPercentage,
 							SizeSpecificData.CollisionType));
 				}
-			}
-			else if (SizeSpecificData.ImplicitType == EImplicitTypeEnum::Chaos_Implicit_Box)
-			{
-				CollectionImplicits[TransformGroupIndex] = FGeometryDynamicCollection::FSharedImplicit(
-					FCollisionStructureManager::NewImplicitBox(
-						InstanceBoundingBox,
-						SizeSpecificData.CollisionObjectReductionPercentage,
-						SizeSpecificData.CollisionType));
-			}
-			else if (SizeSpecificData.ImplicitType == EImplicitTypeEnum::Chaos_Implicit_Sphere)
-			{
-				CollectionImplicits[TransformGroupIndex] = FGeometryDynamicCollection::FSharedImplicit(
-					FCollisionStructureManager::NewImplicitSphere(
-						InnerRadius[GeometryIndex],
-						SizeSpecificData.CollisionObjectReductionPercentage,
-						SizeSpecificData.CollisionType));
-			}
-			else if (SizeSpecificData.ImplicitType == EImplicitTypeEnum::Chaos_Implicit_None)
-			{
-				CollectionImplicits[TransformGroupIndex] = nullptr;
-			}
-			else
-			{
-				ensure(false); // unsupported implicit type!
-			}
+				else if (SizeSpecificData.ImplicitType == EImplicitTypeEnum::Chaos_Implicit_None)
+				{
+					CollectionImplicits[TransformGroupIndex] = nullptr;
+				}
+				else
+				{
+					ensure(false); // unsupported implicit type!
+				}
 
-			if (CollectionImplicits[TransformGroupIndex] && CollectionImplicits[TransformGroupIndex]->HasBoundingBox())
-			{
-				const auto Implicit = CollectionImplicits[TransformGroupIndex];
-				const auto BBox = Implicit->BoundingBox();
-				const TVector<float, 3> Extents = BBox.Extents(); // Chaos::TAABB::Extents() is Max - Min
-				MaxChildBounds = MaxChildBounds.ComponentwiseMax(Extents);
+				if (CollectionImplicits[TransformGroupIndex] && CollectionImplicits[TransformGroupIndex]->HasBoundingBox())
+				{
+					const auto Implicit = CollectionImplicits[TransformGroupIndex];
+					const auto BBox = Implicit->BoundingBox();
+					const TVector<float, 3> Extents = BBox.Extents(); // Chaos::TAABB::Extents() is Max - Min
+					MaxChildBounds = MaxChildBounds.ComponentwiseMax(Extents);
+				}
 			}
 		}
 	});
@@ -2203,7 +2148,6 @@ void FGeometryCollectionPhysicsProxy::InitializeSharedCollisionStructures(
 			const int32 TransformGroupIndex = TransformIndex[GeometryIdx];
 			if (CollectionSimulatableParticles[TransformGroupIndex])
 			{
-				FTransform MassToComponent = CollectionMassToLocal[TransformGroupIndex] * CollectionSpaceTransforms[TransformGroupIndex];
 				PopulateSimulatedParticle(
 					Handles[TransformGroupIndex].Get(),
 					SharedParams, 
@@ -2213,7 +2157,8 @@ void FGeometryCollectionPhysicsProxy::InitializeSharedCollisionStructures(
 					FCollisionFilterData(),		// QueryFilter
 					CollectionMass[TransformGroupIndex],
 					CollectionInertiaTensor[TransformGroupIndex], 
-					MassToComponent,
+					CollectionMassToLocal[TransformGroupIndex],
+					CollectionSpaceTransforms[TransformGroupIndex],
 					(uint8)EObjectStateTypeEnum::Chaos_Object_Dynamic, 
 					INDEX_NONE); // CollisionGroup
 			}
@@ -2293,8 +2238,7 @@ void FGeometryCollectionPhysicsProxy::InitializeSharedCollisionStructures(
 							UnionMeshIndices.Append(ChildIndices);
 
 							const FTransform ChildMassToClusterMass =
-								(CollectionMassToLocal[ChildTransformIdx] * CollectionSpaceTransforms[ChildTransformIdx]).GetRelativeTransform(
-									ClusterMassToCollection);
+								CollectionSpaceTransforms[ChildTransformIdx].GetRelativeTransform(ClusterMassToCollection);
 
 							ChildMesh->GetVertexSet(VertsAdded);
 							for (const int32 VertIdx : VertsAdded)
diff --git a/Engine/Source/Runtime/Experimental/GeometryCollectionSimulationCore/Public/GeometryCollection/GeometryCollectionSimulationCoreTypes.h b/Engine/Source/Runtime/Experimental/GeometryCollectionSimulationCore/Public/GeometryCollection/GeometryCollectionSimulationCoreTypes.h
--- a/Engine/Source/Runtime/Experimental/GeometryCollectionSimulationCore/Public/GeometryCollection/GeometryCollectionSimulationCoreTypes.h
+++ b/Engine/Source/Runtime/Experimental/GeometryCollectionSimulationCore/Public/GeometryCollection/GeometryCollectionSimulationCoreTypes.h
@@ -56,10 +56,12 @@ struct FSharedSimulationParameters
 	FSharedSimulationParameters()
 	: bMassAsDensity(false)
 	, Mass(1.0)
-	, MinimumMassClamp(0.1f)
-	, MaximumMassClamp(1000.0f)
-	, MinimumInertiaTensorDiagonalClamp(SMALL_NUMBER) // todo : Expose to users with better initial values
-	, MaximumInertiaTensorDiagonalClamp(1.0e6f) // todo : Expose to users with better initial values
+	, MinimumMassClamp(0.1)								// todo : Expose to users with better initial values
+	, MaximumMassClamp(1e5f)							// todo : Expose to users with better initial values
+	, MinimumBoundingExtentClamp(0.1)					// todo : Expose to users with better initial values
+	, MaximumBoundingExtentClamp(1e6f)					// todo : Expose to users with better initial values
+	, MinimumInertiaTensorDiagonalClamp(SMALL_NUMBER)	// todo : Expose to users with better initial values
+	, MaximumInertiaTensorDiagonalClamp(1e20f)			// todo : Expose to users with better initial values
 	, MaximumCollisionParticleCount(60)
 	{
 		SizeSpecificData.AddDefaulted();
@@ -75,6 +77,8 @@ struct FSharedSimulationParameters
 		,float InMass
 		, float InMinimumMassClamp
 		, float InMaximumMassClamp
+		, float InMinimumBoundingExtentClamp
+		, float InMaximumBoundingExtentClamp
 		, float InMinimumInertiaTensorDiagonalClamp
 		, float InMaximumInertiaTensorDiagonalClamp
 		,float InCollisionParticlesFraction
@@ -83,6 +87,8 @@ struct FSharedSimulationParameters
 	, Mass(InMass)
 	, MinimumMassClamp(InMinimumMassClamp)
 	, MaximumMassClamp(InMinimumMassClamp)
+	, MinimumBoundingExtentClamp(InMinimumBoundingExtentClamp)
+	, MaximumBoundingExtentClamp(InMinimumBoundingExtentClamp)
 	, MinimumInertiaTensorDiagonalClamp(InMinimumInertiaTensorDiagonalClamp)
 	, MaximumInertiaTensorDiagonalClamp(InMaximumInertiaTensorDiagonalClamp)
 	, MaximumCollisionParticleCount(InMaximumCollisionParticleCount)
@@ -102,9 +108,14 @@ struct FSharedSimulationParameters
 	float Mass;
 	float MinimumMassClamp;
 	float MaximumMassClamp;
+	float MinimumBoundingExtentClamp;
+	float MaximumBoundingExtentClamp;
 	float MinimumInertiaTensorDiagonalClamp;
 	float MaximumInertiaTensorDiagonalClamp;
 
+	float MinimumVolumeClamp() const { return MinimumBoundingExtentClamp * MinimumBoundingExtentClamp * MinimumBoundingExtentClamp; }
+	float MaximumVolumeClamp() const { return MaximumBoundingExtentClamp * MaximumBoundingExtentClamp * MaximumBoundingExtentClamp; }
+
 	TArray<FSharedSimulationSizeSpecificData> SizeSpecificData;
 	TArray<int32> RemoveOnFractureIndices;
 	int32 MaximumCollisionParticleCount;
