Call FGeometryCollectionPhysicsProxy::PushToPhysicsState() from the solver.  Update geometry collection particle world space inflated bounds.

diff --git a/Engine/Source/Runtime/Experimental/ChaosSolvers/Private/PBDRigidsSolver.cpp b/Engine/Source/Runtime/Experimental/ChaosSolvers/Private/PBDRigidsSolver.cpp
--- a/Engine/Source/Runtime/Experimental/ChaosSolvers/Private/PBDRigidsSolver.cpp
+++ b/Engine/Source/Runtime/Experimental/ChaosSolvers/Private/PBDRigidsSolver.cpp
@@ -632,41 +632,6 @@ namespace Chaos
 		Solver->RemoveDirtyProxy(Proxy);
 	}
 
-	void FPBDRigidsSolver::PushPhysicsStatePooled(IDispatcher* Dispatcher)
-	{
-#if 0
-		ensure(Dispatcher != nullptr);
-
-		// reset the per frame pool
-		RigidParticlePool.ResetPool();
-		KinematicGeometryParticlePool.ResetPool();
-		GeometryParticlePool.ResetPool();
-
-		const TArray< IPhysicsProxyBase *>& DirtyProxiesArray = DirtyProxiesSet.GetProxies();
-		for (auto & Proxy : DirtyProxiesArray)
-		{
-			switch (Proxy->GetType())
-			{
-			case EPhysicsProxyType::SingleRigidParticleType:
-				PushPhysicsStateExec(this, static_cast<FSingleParticlePhysicsProxy< Chaos::TPBDRigidParticle<float, 3> >*>(Proxy), RigidParticlePool, Dispatcher);
-				break;
-			case EPhysicsProxyType::SingleKinematicParticleType:
-				PushPhysicsStateExec(this, static_cast<FSingleParticlePhysicsProxy< Chaos::TKinematicGeometryParticle<float, 3> >*>(Proxy), KinematicGeometryParticlePool, Dispatcher);
-				break;
-			case EPhysicsProxyType::SingleGeometryParticleType:
-				PushPhysicsStateExec(this, static_cast<FSingleParticlePhysicsProxy< Chaos::TGeometryParticle<float, 3> >*>(Proxy), GeometryParticlePool, Dispatcher);
-				break;
-			default:
-				ensure("Unknown proxy type in physics solver.");
-			}
-		}
-
-		FlushExec(RigidParticlePool, Dispatcher, this);
-		FlushExec(KinematicGeometryParticlePool, Dispatcher, this);
-		FlushExec(GeometryParticlePool, Dispatcher, this);
-#endif
-	}
-	
 	void PushPhysicsStateExec(FPBDRigidsSolver* Solver, FGeometryCollectionPhysicsProxy* Proxy, Chaos::IDispatcher* Dispatcher)
 	{
 		Proxy->NewData();
@@ -698,6 +663,48 @@ namespace Chaos
 		Solver->RemoveDirtyProxy(Proxy);
 	}
 
+	void FPBDRigidsSolver::PushPhysicsStatePooled(IDispatcher* Dispatcher)
+	{
+#if 0
+		ensure(Dispatcher != nullptr);
+
+		// reset the per frame pool
+		RigidParticlePool.ResetPool();
+		KinematicGeometryParticlePool.ResetPool();
+		GeometryParticlePool.ResetPool();
+
+		TArray< IPhysicsProxyBase*> DirtyProxiesArray = DirtyProxiesSet.Array();
+		for (auto& Proxy : DirtyProxiesArray)
+		{
+			switch (Proxy->GetType())
+			{
+				//case EPhysicsProxyType::NoneType: // 0
+				//case EPhysicsProxyType::StaticMeshType: // 1
+			case EPhysicsProxyType::GeometryCollectionType: // 2
+				PushPhysicsStateExec(this, static_cast<FGeometryCollectionPhysicsProxy*>(Proxy), Dispatcher); // non pool api
+				break;
+			case EPhysicsProxyType::SingleRigidParticleType: // 7
+				PushPhysicsStateExec(this, static_cast<FSingleParticlePhysicsProxy< Chaos::TPBDRigidParticle<float, 3> >*>(Proxy), RigidParticlePool, Dispatcher);
+				break;
+				// case EPhysicsProxyType::FieldType: // 3
+				// case EPhysicsProxyType::SkeletalMeshType: // 4
+			case EPhysicsProxyType::SingleKinematicParticleType: // 6
+				PushPhysicsStateExec(this, static_cast<FSingleParticlePhysicsProxy< Chaos::TKinematicGeometryParticle<float, 3> >*>(Proxy), KinematicGeometryParticlePool, Dispatcher);
+				break;
+			case EPhysicsProxyType::SingleGeometryParticleType: // 5
+				PushPhysicsStateExec(this, static_cast<FSingleParticlePhysicsProxy< Chaos::TGeometryParticle<float, 3> >*>(Proxy), GeometryParticlePool, Dispatcher);
+				break;
+			default:
+				ensure("Unknown proxy type in physics solver.");
+			}
+		}
+
+		FlushExec(RigidParticlePool, Dispatcher, this);
+		FlushExec(KinematicGeometryParticlePool, Dispatcher, this);
+		FlushExec(GeometryParticlePool, Dispatcher, this);
+#endif
+	}
+
 	void FPBDRigidsSolver::PushPhysicsState(IDispatcher* Dispatcher)
 	{
 		QUICK_SCOPE_CYCLE_COUNTER(STAT_PushPhysicsState);
diff --git a/Engine/Source/Runtime/Experimental/ChaosSolvers/Private/PhysicsProxy/GeometryCollectionPhysicsProxy.cpp b/Engine/Source/Runtime/Experimental/ChaosSolvers/Private/PhysicsProxy/GeometryCollectionPhysicsProxy.cpp
--- a/Engine/Source/Runtime/Experimental/ChaosSolvers/Private/PhysicsProxy/GeometryCollectionPhysicsProxy.cpp
+++ b/Engine/Source/Runtime/Experimental/ChaosSolvers/Private/PhysicsProxy/GeometryCollectionPhysicsProxy.cpp
@@ -288,6 +288,11 @@ void PopulateSimulatedParticle(
 		TSharedPtr<Chaos::FImplicitObject, ESPMode::ThreadSafe> SharedImplicitTS(Implicit->DeepCopy().Release());
 		Handle->SetSharedGeometry(SharedImplicitTS);
 		Handle->SetHasBounds(true);
+		Handle->SetLocalBounds(SharedImplicitTS->BoundingBox());
+		const Chaos::TAABB<float, 3>& LocalBounds = Handle->LocalBounds();
+		const Chaos::TRigidTransform<float, 3> Xf(Handle->X(), Handle->R());
+		const Chaos::TAABB<float, 3> TransformedBBox = LocalBounds.TransformedAABB(Xf);
+		Handle->SetWorldSpaceInflatedBounds(TransformedBBox);
 	}
 
 	Handle->CollisionParticlesInitIfNeeded();
@@ -1380,9 +1385,18 @@ void FGeometryCollectionPhysicsProxy::PushToPhysicsState(const Chaos::FParticleD
 	* the solver
 	*/
 
-	//
-	// There is currently no per advance updates to the GeometryCollection
-	//
+	// We need to update physics thread world space bounding boxes.
+	for (auto* Handle : SolverParticleHandles)
+	{
+		if (!Handle || !Handle->HasBounds())
+		{
+			continue;
+		}
+		const Chaos::TAABB<float, 3>& LocalBounds = Handle->LocalBounds();
+		const Chaos::TRigidTransform<float, 3> Xf(Handle->X(), Handle->R());
+		const Chaos::TAABB<float, 3> TransformedBBox = LocalBounds.TransformedAABB(Xf);
+		Handle->SetWorldSpaceInflatedBounds(TransformedBBox);
+	}
 }
 
 void FGeometryCollectionPhysicsProxy::BufferPhysicsResults()
