Hook up the third copy of the GC particles to the GT base particle for SQ to run without crashing against clustered GCs

diff --git a/Engine/Source/Runtime/Experimental/ChaosSolvers/Private/PhysicsProxy/GeometryCollectionPhysicsProxy.cpp b/Engine/Source/Runtime/Experimental/ChaosSolvers/Private/PhysicsProxy/GeometryCollectionPhysicsProxy.cpp
--- a/Engine/Source/Runtime/Experimental/ChaosSolvers/Private/PhysicsProxy/GeometryCollectionPhysicsProxy.cpp
+++ b/Engine/Source/Runtime/Experimental/ChaosSolvers/Private/PhysicsProxy/GeometryCollectionPhysicsProxy.cpp
@@ -919,6 +919,9 @@ void FGeometryCollectionPhysicsProxy::InitializeBodiesPT(
 					}
 					SolverClusterID[TransformGroupIndex] = Handle->ClusterIds().Id;
 
+					// Hook the handle up with the GT particle
+					Handle->GTGeometryParticle() = GTParticles[TransformGroupIndex].Get();
+
 					SolverClusterHandles[TransformGroupIndex] = Handle;
 					SolverParticleHandles[TransformGroupIndex] = Handle;
 					HandleToTransformGroupIndex.Add(Handle, TransformGroupIndex);
