Remove unecessary auto

diff --git a/Engine/Source/Runtime/Experimental/ChaosSolvers/Private/PhysicsProxy/GeometryCollectionPhysicsProxy.cpp b/Engine/Source/Runtime/Experimental/ChaosSolvers/Private/PhysicsProxy/GeometryCollectionPhysicsProxy.cpp
--- a/Engine/Source/Runtime/Experimental/ChaosSolvers/Private/PhysicsProxy/GeometryCollectionPhysicsProxy.cpp
+++ b/Engine/Source/Runtime/Experimental/ChaosSolvers/Private/PhysicsProxy/GeometryCollectionPhysicsProxy.cpp
@@ -1389,7 +1389,7 @@ void FGeometryCollectionPhysicsProxy::PushToPhysicsState(const Chaos::FParticleD
 	*/
 
 	// We need to update physics thread world space bounding boxes.
-	for (auto* Handle : SolverParticleHandles)
+	for (FClusterHandle* Handle : SolverParticleHandles)
 	{
 		if (!Handle || !Handle->HasBounds())
 		{
