Fixed incorrect iteration over child partricles during cluster build. Children may not have particles currently if they also have further children.

diff --git a/Engine/Source/Runtime/Experimental/ChaosSolvers/Private/PhysicsProxy/GeometryCollectionPhysicsProxy.cpp b/Engine/Source/Runtime/Experimental/ChaosSolvers/Private/PhysicsProxy/GeometryCollectionPhysicsProxy.cpp
--- a/Engine/Source/Runtime/Experimental/ChaosSolvers/Private/PhysicsProxy/GeometryCollectionPhysicsProxy.cpp
+++ b/Engine/Source/Runtime/Experimental/ChaosSolvers/Private/PhysicsProxy/GeometryCollectionPhysicsProxy.cpp
@@ -913,9 +913,9 @@ void FGeometryCollectionPhysicsProxy::InitializeBodiesPT(
 							CreationParameters);
 
 					int32 RigidChildrenIdx = 0;
-					for (const int32 ChildIndex : Children[TransformGroupIndex])
+					for(const int32 ChildTransformIndex : RigidChildrenTransformGroupIndex)
 					{
-						SolverClusterID[ChildIndex] = RigidChildren[RigidChildrenIdx++]->CastToClustered()->ClusterIds().Id;
+						SolverClusterID[ChildTransformIndex] = RigidChildren[RigidChildrenIdx++]->CastToClustered()->ClusterIds().Id;;
 					}
 					SolverClusterID[TransformGroupIndex] = Handle->ClusterIds().Id;
 
