ToolsFramework: UClickDragInputBehavior now tracks whether it is in drag or not, and during drag will not update modifier key states unless new bUpdateModifiersDuringDrag is enabled (default false). MeshSurfacePointTool now updates modifier states in ClickBehavior, and not just hover behavior. This resolves issue where shift-to-smooth/etc wouldn't kick in until the viewport was focused (ie had to click once first)
#rnx

diff --git a/Engine/Source/Runtime/Experimental/InteractiveToolsFramework/Private/BaseBehaviors/ClickDragBehavior.cpp b/Engine/Source/Runtime/Experimental/InteractiveToolsFramework/Private/BaseBehaviors/ClickDragBehavior.cpp
--- a/Engine/Source/Runtime/Experimental/InteractiveToolsFramework/Private/BaseBehaviors/ClickDragBehavior.cpp
+++ b/Engine/Source/Runtime/Experimental/InteractiveToolsFramework/Private/BaseBehaviors/ClickDragBehavior.cpp
@@ -18,6 +18,8 @@ void UClickDragInputBehavior::Initialize(IClickDragBehaviorTarget* TargetIn)
 
 FInputCaptureRequest UClickDragInputBehavior::WantsCapture(const FInputDeviceState& Input)
 {
+	bInClickDrag = false;	// should never be true here, but weird things can happen w/ focus
+
 	if (IsPressed(Input) && (ModifierCheckFunc == nullptr || ModifierCheckFunc(Input)) )
 	{
 		FInputRayHit HitResult = Target->CanBeginClickDragSequence(GetDeviceRay(Input));
@@ -34,17 +36,22 @@ FInputCaptureUpdate UClickDragInputBehavior::BeginCapture(const FInputDeviceStat
 {
 	Modifiers.UpdateModifiers(Input, Target);
 	OnClickPressInternal(Input, Side);
+	bInClickDrag = true;
 	return FInputCaptureUpdate::Begin(this, EInputCaptureSide::Any);
 }
 
 
 FInputCaptureUpdate UClickDragInputBehavior::UpdateCapture(const FInputDeviceState& Input, const FInputCaptureData& Data)
 {
-	Modifiers.UpdateModifiers(Input, Target);
+	if (bUpdateModifiersDuringDrag)
+	{
+		Modifiers.UpdateModifiers(Input, Target);
+	}
 
 	if (IsReleased(Input)) 
 	{
 		OnClickReleaseInternal(Input, Data);
+		bInClickDrag = false;
 		return FInputCaptureUpdate::End();
 	}
 	else
@@ -58,6 +65,7 @@ FInputCaptureUpdate UClickDragInputBehavior::UpdateCapture(const FInputDeviceSta
 void UClickDragInputBehavior::ForceEndCapture(const FInputCaptureData& Data)
 {
 	Target->OnTerminateDragSequence();
+	bInClickDrag = false;
 }
 
 
diff --git a/Engine/Source/Runtime/Experimental/InteractiveToolsFramework/Private/BaseTools/MeshSurfacePointTool.cpp b/Engine/Source/Runtime/Experimental/InteractiveToolsFramework/Private/BaseTools/MeshSurfacePointTool.cpp
--- a/Engine/Source/Runtime/Experimental/InteractiveToolsFramework/Private/BaseTools/MeshSurfacePointTool.cpp
+++ b/Engine/Source/Runtime/Experimental/InteractiveToolsFramework/Private/BaseTools/MeshSurfacePointTool.cpp
@@ -57,6 +57,8 @@ void UMeshSurfacePointTool::Setup()
 
 	// add input behaviors
 	UClickDragInputBehavior* DragBehavior = NewObject<UClickDragInputBehavior>();
+	DragBehavior->Modifiers.RegisterModifier(UMeshSurfacePointTool_ShiftModifier, FInputDeviceState::IsShiftKeyDown);
+	DragBehavior->Modifiers.RegisterModifier(UMeshSurfacePointTool_CtrlModifier, FInputDeviceState::IsCtrlKeyDown);
 	DragBehavior->Initialize(this);
 	AddInputBehavior(DragBehavior);
 
diff --git a/Engine/Source/Runtime/Experimental/InteractiveToolsFramework/Public/BaseBehaviors/ClickDragBehavior.h b/Engine/Source/Runtime/Experimental/InteractiveToolsFramework/Public/BaseBehaviors/ClickDragBehavior.h
--- a/Engine/Source/Runtime/Experimental/InteractiveToolsFramework/Public/BaseBehaviors/ClickDragBehavior.h
+++ b/Engine/Source/Runtime/Experimental/InteractiveToolsFramework/Public/BaseBehaviors/ClickDragBehavior.h
@@ -53,11 +53,21 @@ public:
 	virtual FInputCaptureUpdate UpdateCapture(const FInputDeviceState& Input, const FInputCaptureData& Data) override;
 	virtual void ForceEndCapture(const FInputCaptureData& Data) override;
 
+	/**
+	 * If true, then we will update Modifier states in UpdateCapture(). This defaults to false because
+	 * in most cases during a capture you don't want this, eg in a brush interaction, if you have a shift-to-smooth modifier,
+	 * you don't want to toggle this on/off during a sculpting stroke.
+	 */
+	UPROPERTY()
+	bool bUpdateModifiersDuringDrag = false;
 
 protected:
 	/** Click Target object */
 	IClickDragBehaviorTarget* Target;
 
+	/** set to true if we are in a capture */
+	bool bInClickDrag = false;
+
 	/**
 	 * Internal function that forwards click evens to Target::OnClickPress, you can customize behavior here
 	 */
