Added rewind support for sim-writable properties

diff --git a/Engine/Source/Programs/HeadlessChaos/Private/HeadlessChaosTestRewind.cpp b/Engine/Source/Programs/HeadlessChaos/Private/HeadlessChaosTestRewind.cpp
--- a/Engine/Source/Programs/HeadlessChaos/Private/HeadlessChaosTestRewind.cpp
+++ b/Engine/Source/Programs/HeadlessChaos/Private/HeadlessChaosTestRewind.cpp
@@ -41,18 +41,25 @@ namespace ChaosTest {
 		FPhysicsSolver* Solver = Module->CreateSolver(nullptr,ESolverFlags::Standalone);
 		Solver->SetEnabled(true);
 
-		// Make a particle
+		Solver->EnableRewindCapture(20);
+
 
+		// Make particles
 		auto Particle = TKinematicGeometryParticle<float,3>::CreateParticle();
 
 		Particle->SetGeometry(Sphere);
 		Solver->RegisterObject(Particle.Get());
-		Solver->EnableRewindCapture(20);
+
+		auto Dynamic = TPBDRigidParticle<float,3>::CreateParticle();
+		Dynamic->SetGeometry(Sphere);
+		Solver->RegisterObject(Dynamic.Get());
+		
 
 		for(int Step = 0; Step < 11; ++Step)
 		{
 			//property that changes every step
 			Particle->SetX(FVec3(0,0,100 - Step));
+			Dynamic->SetF(FVec3(0,0,Step + 1));
 
 			//property that changes once half way through
 			if(Step == 3)
@@ -107,6 +114,14 @@ namespace ChaosTest {
 					EXPECT_EQ(ParticleState->Geometry().Get(),Box2.Get());
 				}
 			}
+
+			const auto DynamicState = RewindData->GetStateAtFrame(*Dynamic,Step);
+
+			EXPECT_TRUE(DynamicState != nullptr);
+			if(DynamicState)
+			{
+				EXPECT_EQ(DynamicState->F()[2],Step+1);
+			}
 		}
 		
 		//no data at head because we always save the previous frame
diff --git a/Engine/Source/Runtime/Experimental/ChaosSolvers/Private/PBDRigidsSolver.cpp b/Engine/Source/Runtime/Experimental/ChaosSolvers/Private/PBDRigidsSolver.cpp
--- a/Engine/Source/Runtime/Experimental/ChaosSolvers/Private/PBDRigidsSolver.cpp
+++ b/Engine/Source/Runtime/Experimental/ChaosSolvers/Private/PBDRigidsSolver.cpp
@@ -775,6 +775,11 @@ namespace Chaos
 					Proxy->SetHandle(CreateHandleFunc(UniqueIdx));
 				}
 
+				if(RewindData)
+				{
+					RewindData->PushGTDirtyData(*Manager,DataIdx,Dirty);
+				}
+
 				Proxy->PushToPhysicsState(*Manager, DataIdx, Dirty, ShapeDirtyData);
 
 				if(bIsNew)
@@ -791,7 +796,7 @@ namespace Chaos
 
 			if(RewindData)
 			{
-				RewindData->SavePrevFrameState(*DirtyProxiesData);
+				RewindData->PrepareFrame(DirtyProxiesData->NumDirtyProxies());
 			}
 
 			//need to create new particle handles
@@ -822,6 +827,11 @@ namespace Chaos
 				}
 			});
 
+			if(RewindData)
+			{
+				RewindData->AdvanceFrame();
+			}
+
 			DirtyProxiesData->Reset();
 		});
 	}
diff --git a/Engine/Source/Runtime/Experimental/ChaosSolvers/Public/RewindData.h b/Engine/Source/Runtime/Experimental/ChaosSolvers/Public/RewindData.h
--- a/Engine/Source/Runtime/Experimental/ChaosSolvers/Public/RewindData.h
+++ b/Engine/Source/Runtime/Experimental/ChaosSolvers/Public/RewindData.h
@@ -42,11 +42,11 @@ public:
 	}
 
 	template <typename LambdaSet>
-	void SyncRemoteData(FDirtyPropertiesManager* InManager,int32 InIdx, const FParticleDirtyData& DirtyData, const LambdaSet& SetFunc)
+	void SyncRemoteData(FDirtyPropertiesManager& InManager,int32 InIdx, const FParticleDirtyData& DirtyData, const LambdaSet& SetFunc)
 	{
 		if(DirtyData.IsDirty(ParticlePropToFlag(PropName)))
 		{
-			Manager = InManager;
+			Manager = &InManager;
 			Idx = InIdx;
 			SetFunc(Manager->GetParticlePool<T,PropName>().GetElement(Idx));
 		}
@@ -62,33 +62,92 @@ private:
 	int32 Idx;
 };
 
+//Set of properties that can be written to from the simulation
+class FSimWritableData
+{
+public:
+
+	FSimWritableData()
+	{
+	}
+
+	const FParticlePositionRotation& GetXR() const { return XR; }
+	const FParticleVelocities& GetVelocities() const { return Velocities; }
+	const FParticleDynamics& GetDynamics() const { return Dynamics; }
+
+	void SyncData(const FDirtyPropertiesManager& Manager, const int32 DataIdx, const FParticleDirtyData& Dirty)
+	{
+		if(const auto Data = Dirty.FindXR(Manager, DataIdx))
+		{
+			XR = *Data;
+		}
+
+		if(const auto Data = Dirty.FindVelocities(Manager, DataIdx))
+		{
+			Velocities = *Data;
+		}
+
+		if(const auto Data = Dirty.FindDynamics(Manager, DataIdx))
+		{
+			Dynamics = *Data;
+		}
+	}
+
+private:
+	FParticlePositionRotation XR;
+	FParticleVelocities Velocities;
+	FParticleDynamics Dynamics;
+};
+
+inline bool SimWritablePropsMayChange(const TGeometryParticleHandle<FReal,3>& Handle)
+{
+	if(const auto Rigid = Handle.CastToRigidParticle())
+	{
+		return Rigid->ObjectState() == EObjectStateType::Dynamic;
+	}
+
+	return false;
+}
+
 class FGeometryParticleState
 {
 public:
 
-	FGeometryParticleState(TGeometryParticle<FReal,3>* InParticle)
+	FGeometryParticleState(const TGeometryParticle<FReal,3>& InParticle)
 		: Particle(InParticle)
 	{
 	}
 
 	const FVec3& X() const
 	{
-		return ParticlePositionRotation.IsSet() ? ParticlePositionRotation.Read().X : Particle->X();
+		return ParticlePositionRotation.IsSet() ? ParticlePositionRotation.Read().X : Particle.X();
 	}
 
 	TSerializablePtr<FImplicitObject> Geometry() const
 	{
-		return NonFrequentData.IsSet() ? MakeSerializable(NonFrequentData.Read().Geometry) : Particle->Geometry();
+		return NonFrequentData.IsSet() ? MakeSerializable(NonFrequentData.Read().Geometry) : Particle.Geometry();
+	}
+
+	const FVec3& F() const
+	{
+		return Dynamics.IsSet() ? Dynamics.Read().F : Particle.CastToRigidParticle()->F();
 	}
 
-	void SyncRemoteData(FDirtyPropertiesManager* Manager,int32 Idx,const FDirtyProxy& Dirty)
+	void SyncRemoteData(FDirtyPropertiesManager& Manager,int32 Idx,const FDirtyProxy& Dirty, const FSimWritableData* SimWritableData)
 	{
 		const auto Proxy = static_cast<const FGeometryParticlePhysicsProxy*>(Dirty.Proxy);
 		const auto Handle = Proxy->GetHandle();
-		ParticlePositionRotation.SyncRemoteData(Manager,Idx,Dirty.ParticleData, [Handle](FParticlePositionRotation& Data)
+		ParticlePositionRotation.SyncRemoteData(Manager,Idx,Dirty.ParticleData, [Handle, SimWritableData](FParticlePositionRotation& Data)
 		{
-			Data.X = Handle->X();
-			Data.R = Handle->R();
+			if(SimWritableData)
+			{
+				Data = SimWritableData->GetXR();
+			}
+			else
+			{
+				Data.X = Handle->X();
+				Data.R = Handle->R();
+			}
 		});
 
 		NonFrequentData.SyncRemoteData(Manager,Idx,Dirty.ParticleData, [Handle](FParticleNonFrequentData& Data)
@@ -103,6 +162,26 @@ public:
 			Data.DebugName = Handle->DebugName();
 #endif
 		});
+
+		if(auto PBDRigid = Handle->CastToRigidParticle())
+		{
+			Dynamics.SyncRemoteData(Manager,Idx,Dirty.ParticleData,[PBDRigid, SimWritableData](FParticleDynamics& Data)
+			{
+				if(SimWritableData)
+				{
+					Data = SimWritableData->GetDynamics();
+				}
+				else
+				{
+					Data.F = PBDRigid->F();
+					Data.Torque = PBDRigid->Torque();
+					Data.LinearImpulse = PBDRigid->LinearImpulse();
+					Data.AngularImpulse = PBDRigid->AngularImpulse();
+					Data.LinearEtherDrag = PBDRigid->LinearEtherDrag();
+					Data.AngularEtherDrag = PBDRigid->AngularEtherDrag();
+				}
+			});
+		}
 	}
 
 	bool CoalesceState(const FGeometryParticleState& LatestState)
@@ -120,14 +199,28 @@ public:
 			bCoalesced = true;
 		}
 
+		if(!Dynamics.IsSet() && LatestState.Dynamics.IsSet())
+		{
+			Dynamics = LatestState.Dynamics;
+			bCoalesced = true;
+		}
+
 		return bCoalesced;
 	}
 
 protected:
-	TGeometryParticle<FReal,3>* Particle;
+	const TGeometryParticle<FReal,3>& Particle;
 private:
 	TParticleStateProperty<FParticlePositionRotation,EParticleProperty::XR> ParticlePositionRotation;
 	TParticleStateProperty<FParticleNonFrequentData, EParticleProperty::NonFrequentData> NonFrequentData;
+	TParticleStateProperty<FParticleDynamics, EParticleProperty::Dynamics> Dynamics;
+	/*
+	PARTICLE_PROPERTY(XR,FParticlePositionRotation)
+		PARTICLE_PROPERTY(Velocities,FParticleVelocities)
+		PARTICLE_PROPERTY(Dynamics,FParticleDynamics)
+		PARTICLE_PROPERTY(Misc,FParticleMisc)
+		PARTICLE_PROPERTY(NonFrequentData,FParticleNonFrequentData)
+		PARTICLE_PROPERTY(MassProps,FParticleMassProps)*/
 };
 
 class FRewindData
@@ -140,14 +233,14 @@ public:
 
 	const FGeometryParticleState* GetStateAtFrame(const TGeometryParticle<FReal,3>& Particle,int32 Frame) const
 	{
-		if(const TArray<FFrameInfo>* Frames = ParticleToFrameInfo.Find(Particle.UniqueIdx()))
+		if(const FParticleRewindInfo* Info = ParticleToRewindInfo.Find(Particle.UniqueIdx()))
 		{
 			//is it worth doing binary search?
-			for(const FFrameInfo& Info : (*Frames))
+			for(const FFrameInfo& FrameInfo : Info->Frames)
 			{
-				if(Info.Frame == Frame)
+				if(FrameInfo.Frame == Frame)
 				{
-					return &Info.State;
+					return &FrameInfo.State;
 				}
 			}
 		}
@@ -155,34 +248,37 @@ public:
 		return nullptr;
 	}
 
-	void SavePrevFrameState(const FDirtySet& DirtyData)
+	void AdvanceFrame()
+	{
+		++CurFrame;
+	}
+
+	void PrepareFrame(int32 NumDirtyParticles)
 	{
 		Managers.Emplace(MakeUnique<FDirtyPropertiesManager>());
+		Managers.Last()->SetNumParticles(NumDirtyParticles);
+	}
 
-		FDirtyPropertiesManager* Manager = Managers.Last().Get();
-		Manager->SetNumParticles(DirtyData.NumDirtyProxies());
+	void PushGTDirtyData(const FDirtyPropertiesManager& Manager, int32 DataIdx, const FDirtyProxy& Dirty)
+	{
+		auto ProcessProxy = [this,Manager, DataIdx, Dirty](const auto Proxy)
+		{
+			const auto PTParticle = Proxy->GetHandle();
+			FParticleRewindInfo& Info = ParticleToRewindInfo.FindOrAdd(PTParticle->UniqueIdx());
 
-		//NOTE: this is called from PT so we cannot use the GTParticle to read
-		//Instead we  must use the proxy system
-		//However, the rewind happens on GT so that's why we still want to associate this data with the GTParticle
+			//Any previous frames that are pointing at head of a dirty property must be written out (since head is changing)
 
-		auto ProcessProxy = [this,Manager](const auto Proxy,const int32 DataIdx,const FDirtyProxy& Dirty)
-		{
-			//Since we're saving last frame's data the proxy must already be initialized
+			//Since we are updating previous frames, the proxy must have been initialized already
 			if(Proxy->IsInitialized())
 			{
-				const auto GTParticleUnsafe = Proxy->GetParticle();
-				const auto PTParticle = Proxy->GetHandle();
-				TArray<FFrameInfo>& Frames = ParticleToFrameInfo.FindOrAdd(PTParticle->UniqueIdx());
-				Frames.Add(FFrameInfo{FGeometryParticleState(GTParticleUnsafe),CurFrame-1});
+				FGeometryParticleState& LatestState = Info.AddFrame(*Proxy->GetParticle(),CurFrame-1);
 
-				FGeometryParticleState& LatestState = Frames.Last().State;
-				LatestState.SyncRemoteData(Manager,DataIdx,Dirty);
+				LatestState.SyncRemoteData(*Managers.Last(),DataIdx,Dirty, Info.SimWritableData.Get());
 
-				//coalesce any previous data
-				for(int32 FrameIdx = Frames.Num() - 2; FrameIdx >= 0; --FrameIdx)
+				//for frames further back a simply copy is enough
+				for(int32 FrameIdx = Info.Frames.Num() - 2; FrameIdx >= 0; --FrameIdx)
 				{
-					FFrameInfo& Frame = Frames[FrameIdx];
+					FFrameInfo& Frame = Info.Frames[FrameIdx];
 					if(Frame.State.CoalesceState(LatestState) == false)
 					{
 						//nothing to coalesce so no need to check earlier frames
@@ -190,37 +286,45 @@ public:
 					}
 				}
 			}
-		};
 
-		//can't be parallel because ParticletoFrameInfo is modified
-		DirtyData.ForEachProxy([&ProcessProxy](int32 DataIdx,const FDirtyProxy& Dirty)
-		{
-			switch(Dirty.Proxy->GetType())
-			{
-			case EPhysicsProxyType::SingleRigidParticleType:
+			//if particle is sim-writable make sure to record appropriate dirty properties
+			if(SimWritablePropsMayChange(*PTParticle))
 			{
-				auto Proxy = static_cast<FRigidParticlePhysicsProxy*>(Dirty.Proxy);
-				ProcessProxy(Proxy,DataIdx,Dirty);
-				break;
-			}
-			case EPhysicsProxyType::SingleKinematicParticleType:
-			{
-				auto Proxy = static_cast<FKinematicGeometryParticlePhysicsProxy*>(Dirty.Proxy);
-				ProcessProxy(Proxy,DataIdx,Dirty);
-				break;
+				if(!Info.SimWritableData)
+				{
+					Info.SimWritableData = MakeUnique<FSimWritableData>();
+				}
+				Info.SimWritableData->SyncData(Manager,DataIdx,Dirty.ParticleData);
 			}
-			case EPhysicsProxyType::SingleGeometryParticleType:
+			else
 			{
-				auto Proxy = static_cast<FGeometryParticlePhysicsProxy*>(Dirty.Proxy);
-				ProcessProxy(Proxy,DataIdx,Dirty);
-				break;
-			}
-			default:
-			ensure("Unknown proxy type in physics solver.");
+				Info.SimWritableData = nullptr;
 			}
-		});
+		};
 
-		++CurFrame;
+		switch(Dirty.Proxy->GetType())
+		{
+		case EPhysicsProxyType::SingleRigidParticleType:
+		{
+			auto Proxy = static_cast<FRigidParticlePhysicsProxy*>(Dirty.Proxy);
+			ProcessProxy(Proxy);
+			break;
+		}
+		case EPhysicsProxyType::SingleKinematicParticleType:
+		{
+			auto Proxy = static_cast<FKinematicGeometryParticlePhysicsProxy*>(Dirty.Proxy);
+			ProcessProxy(Proxy);
+			break;
+		}
+		case EPhysicsProxyType::SingleGeometryParticleType:
+		{
+			auto Proxy = static_cast<FGeometryParticlePhysicsProxy*>(Dirty.Proxy);
+			ProcessProxy(Proxy);
+			break;
+		}
+		default:
+		ensure("Unknown proxy type in physics solver.");
+		}
 	}
 
 private:
@@ -231,7 +335,19 @@ private:
 		int32 Frame;
 	};
 
-	TArrayAsMap<FUniqueIdx,TArray<FFrameInfo>> ParticleToFrameInfo;
+	struct FParticleRewindInfo
+	{
+		TArray<FFrameInfo> Frames;
+		TUniquePtr<FSimWritableData> SimWritableData;
+
+		FGeometryParticleState& AddFrame(TGeometryParticle<FReal,3>& GTParticleUnsafe, int32 FrameIdx)
+		{
+			Frames.Add(FFrameInfo{FGeometryParticleState(GTParticleUnsafe),FrameIdx});
+			return Frames.Last().State;
+		}
+	};
+
+	TArrayAsMap<FUniqueIdx,FParticleRewindInfo> ParticleToRewindInfo;
 	TArray<TUniquePtr<FDirtyPropertiesManager>> Managers;
 	int32 CurFrame;
 };
