Chaos: Bug fix clustering, Node not being correctly removed from ParentToChildren Map

diff --git a/Engine/Source/Runtime/Experimental/Chaos/Private/Chaos/PBDRigidClustering.cpp b/Engine/Source/Runtime/Experimental/Chaos/Private/Chaos/PBDRigidClustering.cpp
--- a/Engine/Source/Runtime/Experimental/Chaos/Private/Chaos/PBDRigidClustering.cpp
+++ b/Engine/Source/Runtime/Experimental/Chaos/Private/Chaos/PBDRigidClustering.cpp
@@ -1677,8 +1677,7 @@ namespace Chaos
 		}
 
 		TopLevelClusterParents.Remove(ClusteredParticle);
-		TArray<TPBDRigidParticleHandle<T, d>*>& ParentToChildren = MChildren[ClusteredParticle];
-		ParentToChildren.Remove(ClusteredParticle);
+		GetChildrenMap().Remove(ClusteredParticle);
 		ClusteredParticle->ClusterIds() = ClusterId();
 		ClusteredParticle->ClusterGroupIndex() = 0;
 		MActiveRemovalIndices.Remove(ClusteredParticle);
