Chaos - Joint Batching (WIP - Disabled)

diff --git a/Engine/Source/Runtime/Core/Public/Math/Vector.isph b/Engine/Source/Runtime/Core/Public/Math/Vector.isph
--- a/Engine/Source/Runtime/Core/Public/Math/Vector.isph
+++ b/Engine/Source/Runtime/Core/Public/Math/Vector.isph
@@ -153,6 +153,18 @@ inline FVector4 operator*(const FVector4 &A, const FVector4 &B)
 	return Result;
 }
 
+inline FVector4 operator*(const uniform FVector4 &A, const FVector4 &B)
+{
+	FVector4 Result;
+
+	for(uniform int i = 0; i < 4; i++)
+	{
+		Result.V[i] = A.V[i] * B.V[i];
+	}
+
+	return Result;
+}
+
 inline FVector4 operator*(const FVector4 &A, const uniform FVector4 &B)
 {
 	FVector4 Result;
@@ -336,6 +348,12 @@ inline uniform FVector4 SetVector4(const uniform float F)
 	return Result;
 }
 
+inline varying FVector4 SetVector4(const varying float F)
+{
+	const varying FVector4 Result = {{F, F, F, F}};
+	return Result;
+}
+
 inline FVector operator+(const FVector &A, const FVector &B)
 {
 	FVector Result;
@@ -1571,6 +1589,13 @@ inline uniform float VectorDot( const uniform FVector &Vec1, const uniform FVect
 	return Temp.V[0] + Temp.V[1] + Temp.V[2];
 }
 
+
+inline varying float VectorDot( const varying FVector &Vec1, const varying FVector &Vec2 )
+{
+	const varying FVector Temp = Vec1 * Vec2;
+	return Temp.V[0] + Temp.V[1] + Temp.V[2];
+}
+
 /**
  * Does a bitwise vector selection based on a mask (e.g., created from VectorCompareXX)
  *
diff --git a/Engine/Source/Runtime/Engine/Private/PhysicsEngine/ImmediatePhysics/ImmediatePhysicsChaos/ImmediatePhysicsJointHandle_Chaos.cpp b/Engine/Source/Runtime/Engine/Private/PhysicsEngine/ImmediatePhysics/ImmediatePhysicsChaos/ImmediatePhysicsJointHandle_Chaos.cpp
--- a/Engine/Source/Runtime/Engine/Private/PhysicsEngine/ImmediatePhysics/ImmediatePhysicsChaos/ImmediatePhysicsJointHandle_Chaos.cpp
+++ b/Engine/Source/Runtime/Engine/Private/PhysicsEngine/ImmediatePhysics/ImmediatePhysicsChaos/ImmediatePhysicsJointHandle_Chaos.cpp
@@ -218,11 +218,6 @@ namespace ImmediatePhysics_Chaos
 		return reinterpret_cast<const Chaos::TVector<const FActorHandle*, 2>&>(ActorHandles);
 	}
 
-	void FJointHandle::UpdateLevels()
-	{
-		ConstraintHandle->SetParticleLevels({ ActorHandles[0]->GetLevel(), ActorHandles[1]->GetLevel() });
-	}
-
 	void FJointHandle::SetSoftLinearSettings(bool bLinearSoft, FReal LinearStiffness, FReal LinearDamping)
 	{
 		using namespace Chaos;
diff --git a/Engine/Source/Runtime/Engine/Private/PhysicsEngine/ImmediatePhysics/ImmediatePhysicsChaos/ImmediatePhysicsSimulation_Chaos.cpp b/Engine/Source/Runtime/Engine/Private/PhysicsEngine/ImmediatePhysics/ImmediatePhysicsChaos/ImmediatePhysicsSimulation_Chaos.cpp
--- a/Engine/Source/Runtime/Engine/Private/PhysicsEngine/ImmediatePhysics/ImmediatePhysicsChaos/ImmediatePhysicsSimulation_Chaos.cpp
+++ b/Engine/Source/Runtime/Engine/Private/PhysicsEngine/ImmediatePhysics/ImmediatePhysicsChaos/ImmediatePhysicsSimulation_Chaos.cpp
@@ -170,7 +170,6 @@ namespace ImmediatePhysics_Chaos
 		, NumRollingAverageStepTimes(1)
 		, MaxNumRollingAverageStepTimes(ChaosImmediate_Evolution_DeltaTimeCount)
 		, bActorsDirty(false)
-		, bJointsDirty(false)
 	{
 		using namespace Chaos;
 
@@ -332,9 +331,6 @@ namespace ImmediatePhysics_Chaos
 	{
 		FJointHandle* JointHandle = new FJointHandle(&Joints, ConstraintInstance, Body1, Body2);
 		JointHandles.Add(JointHandle);
-
-		bJointsDirty = true;
-
 		return JointHandle;
 	}
 
@@ -343,8 +339,6 @@ namespace ImmediatePhysics_Chaos
 		// @todo(ccaulfield): FJointHandle could remember its index to optimize this
 		JointHandles.Remove(JointHandle);
 		delete JointHandle;
-
-		bJointsDirty = true;
 	}
 
 	void FSimulation::SetNumActiveBodies(int32 InNumActiveActorHandles)
@@ -457,65 +451,6 @@ namespace ImmediatePhysics_Chaos
 		}
 	}
 
-	void FSimulation::ConditionConstraints()
-	{
-		// Assign levels to actors based on connection-distance to a non-dynamic actor
-		// This is used to sort constraints and set parent/child relationship on the constrained particles in a constraint
-		// @todo(ccaulfield): this should use the constraint graph and should only update when constraint connectivity changes or particles change type
-		using namespace Chaos;
-
-		TMap<FActorHandle*, TArray<FJointHandle*>> ActorJoints;
-		TArray<FActorHandle*> ActorQueue;
-		ActorQueue.Reserve(ActorHandles.Num());
-
-		// Reset all actor levels
-		for (FActorHandle* ActorHandle : ActorHandles)
-		{
-			ActorHandle->SetLevel(INDEX_NONE);
-			ActorJoints.Emplace(ActorHandle);
-
-			if (ActorHandle->GetParticle()->ObjectState() != EObjectStateType::Dynamic)
-			{
-				ActorHandle->SetLevel(0);
-				ActorQueue.Add(ActorHandle);
-			}
-		}
-
-		// Build a list of joints per actor
-		for (FJointHandle* JointHandle : JointHandles)
-		{
-			TVector<FActorHandle*, 2> JointActorHandles = JointHandle->GetActorHandles();
-			ActorJoints[JointActorHandles[0]].Add(JointHandle);
-			ActorJoints[JointActorHandles[1]].Add(JointHandle);
-		}
-
-		// Breadth-first assign level
-		for (int32 ActorQueueIndex = 0; ActorQueueIndex < ActorQueue.Num(); ++ActorQueueIndex)
-		{
-			FActorHandle* ActorHandle = ActorQueue[ActorQueueIndex];
-			for (FJointHandle* JointHandle : ActorJoints[ActorHandle])
-			{
-				TVector<FActorHandle*, 2> JointActorHandles = JointHandle->GetActorHandles();
-				if (JointActorHandles[0]->GetLevel() == INDEX_NONE)
-				{
-					JointActorHandles[0]->SetLevel(ActorHandle->GetLevel() + 1);
-					ActorQueue.Add(JointActorHandles[0]);
-				}
-				if (JointActorHandles[1]->GetLevel() == INDEX_NONE)
-				{
-					JointActorHandles[1]->SetLevel(ActorHandle->GetLevel() + 1);
-					ActorQueue.Add(JointActorHandles[1]);
-				}
-			}
-		}
-
-		// Update constraint levels
-		for (FJointHandle* JointHandle : JointHandles)
-		{
-			JointHandle->UpdateLevels();
-		}
-	}
-
 	void FSimulation::SetSimulationSpaceTransform(const FTransform& Transform)
 	{ 
 		SimulationSpaceTransform = Transform;
@@ -620,12 +555,6 @@ namespace ImmediatePhysics_Chaos
 		DebugDrawDynamicParticles(2, 2, FColor(192, 192, 0));
 		DebugDrawConstraints(2, 2, 0.7f);
 
-		if (bJointsDirty)
-		{
-			ConditionConstraints();
-			bJointsDirty = false;
-		}
-
 		if (bActorsDirty)
 		{
 			UpdateActivePotentiallyCollidingPairs();
diff --git a/Engine/Source/Runtime/Engine/Public/Physics/ImmediatePhysics/ImmediatePhysicsChaos/ImmediatePhysicsJointHandle_Chaos.h b/Engine/Source/Runtime/Engine/Public/Physics/ImmediatePhysics/ImmediatePhysicsChaos/ImmediatePhysicsJointHandle_Chaos.h
--- a/Engine/Source/Runtime/Engine/Public/Physics/ImmediatePhysics/ImmediatePhysicsChaos/ImmediatePhysicsJointHandle_Chaos.h
+++ b/Engine/Source/Runtime/Engine/Public/Physics/ImmediatePhysics/ImmediatePhysicsChaos/ImmediatePhysicsJointHandle_Chaos.h
@@ -27,8 +27,6 @@ namespace ImmediatePhysics_Chaos
 		const Chaos::TVector<FActorHandle*, 2>& GetActorHandles();
 		const Chaos::TVector<const FActorHandle*, 2>& GetActorHandles() const;
 
-		void UpdateLevels();
-
 		void SetSoftLinearSettings(bool bLinearSoft, FReal LinearStiffness, FReal LinearDamping);
 
 	private:
diff --git a/Engine/Source/Runtime/Engine/Public/Physics/ImmediatePhysics/ImmediatePhysicsChaos/ImmediatePhysicsSimulation_Chaos.h b/Engine/Source/Runtime/Engine/Public/Physics/ImmediatePhysics/ImmediatePhysicsChaos/ImmediatePhysicsSimulation_Chaos.h
--- a/Engine/Source/Runtime/Engine/Public/Physics/ImmediatePhysics/ImmediatePhysicsChaos/ImmediatePhysicsSimulation_Chaos.h
+++ b/Engine/Source/Runtime/Engine/Public/Physics/ImmediatePhysics/ImmediatePhysicsChaos/ImmediatePhysicsSimulation_Chaos.h
@@ -73,7 +73,6 @@ namespace ImmediatePhysics_Chaos
 		void SetSolverIterations(const int32 SolverIts, const int32 JointIts, const int32 CollisionIts, const int32 SolverPushOutIts, const int32 JointPushOutIts, const int32 CollisionPushOutIts);
 
 	private:
-		void ConditionConstraints();
 		void UpdateActivePotentiallyCollidingPairs();
 		FReal UpdateStepTime(const FReal DeltaTime, const FReal MaxStepTime);
 
@@ -124,7 +123,6 @@ namespace ImmediatePhysics_Chaos
 		int32 MaxNumRollingAverageStepTimes;
 
 		bool bActorsDirty;
-		bool bJointsDirty;
 	};
 
 }
diff --git a/Engine/Source/Runtime/Experimental/Chaos/Private/Chaos/ChaosDebugDraw.cpp b/Engine/Source/Runtime/Experimental/Chaos/Private/Chaos/ChaosDebugDraw.cpp
--- a/Engine/Source/Runtime/Experimental/Chaos/Private/Chaos/ChaosDebugDraw.cpp
+++ b/Engine/Source/Runtime/Experimental/Chaos/Private/Chaos/ChaosDebugDraw.cpp
@@ -230,7 +230,7 @@ namespace Chaos
 			DrawCollisionImpl(SpaceTransform, ConstraintHandle->GetContact(), ColorScale);
 		}
 
-		void DrawJointConstraintImpl(const FRigidTransform3& SpaceTransform, const FVec3& InPa, const FVec3& InCa, const FVec3& InXa, const FMatrix33& Ra, const FVec3& InPb, const FVec3& InCb, const FVec3& InXb, const FMatrix33& Rb, int32 Level, int32 Index, FReal ColorScale, uint32 FeatureMask)
+		void DrawJointConstraintImpl(const FRigidTransform3& SpaceTransform, const FVec3& InPa, const FVec3& InCa, const FVec3& InXa, const FMatrix33& Ra, const FVec3& InPb, const FVec3& InCb, const FVec3& InXb, const FMatrix33& Rb, int32 IslandIndex, int32 LevelIndex, int32 ColorIndex, int32 BatchIndex, int32 Index, FReal ColorScale, uint32 FeatureMask)
 		{
 			using namespace Chaos::DebugDraw;
 			FColor R = (ColorScale * FColor::Red).ToFColor(false);
@@ -304,13 +304,31 @@ namespace Chaos
 				FDebugDrawQueue::GetInstance().DrawDebugDirectionalArrow(Xb, Xb + DrawScale * ConstraintAxisLen * SpaceTransform.TransformVector(Rb.GetAxis(1)), DrawScale * ArrowSize, M, false, KINDA_SMALL_NUMBER, DrawPriority, LineThickness);
 				FDebugDrawQueue::GetInstance().DrawDebugDirectionalArrow(Xb, Xb + DrawScale * ConstraintAxisLen * SpaceTransform.TransformVector(Rb.GetAxis(2)), DrawScale * ArrowSize, Y, false, KINDA_SMALL_NUMBER, DrawPriority, LineThickness);
 			}
-			if ((FeatureMask & (uint32)EDebugDrawJointFeature::Level) && (Level >= 0))
+			FVec3 TextPos = Xb;
+			if ((FeatureMask & (uint32)EDebugDrawJointFeature::Level) && (LevelIndex >= 0))
 			{
-				FDebugDrawQueue::GetInstance().DrawDebugString(Xb + FontHeight * FVec3(0, 0, 1), FString::Format(TEXT("{0}"), { Level }), nullptr, FColor::Red, KINDA_SMALL_NUMBER, false, FontScale);
+				FDebugDrawQueue::GetInstance().DrawDebugString(TextPos, FString::Format(TEXT("{0}"), { LevelIndex }), nullptr, FColor::Red, KINDA_SMALL_NUMBER, false, FontScale);
+				TextPos += FontHeight * FVec3(0, 0, 1);
 			}
-			if ((FeatureMask & (uint32)EDebugDrawJointFeature::Index) && (Level >= 0))
+			if ((FeatureMask & (uint32)EDebugDrawJointFeature::Index) && (Index >= 0))
 			{
-				FDebugDrawQueue::GetInstance().DrawDebugString(Xb + FontHeight * FVec3(0, 0, 1), FString::Format(TEXT("{0}"), { Index }), nullptr, FColor::Red, KINDA_SMALL_NUMBER, false, FontScale);
+				FDebugDrawQueue::GetInstance().DrawDebugString(TextPos, FString::Format(TEXT("{0}"), { Index }), nullptr, FColor::Red, KINDA_SMALL_NUMBER, false, FontScale);
+				TextPos += FontHeight * FVec3(0, 0, 1);
+			}
+			if ((FeatureMask & (uint32)EDebugDrawJointFeature::Color) && (ColorIndex >= 0))
+			{
+				FDebugDrawQueue::GetInstance().DrawDebugString(TextPos, FString::Format(TEXT("{0}"), { ColorIndex }), nullptr, FColor::Red, KINDA_SMALL_NUMBER, false, FontScale);
+				TextPos += FontHeight * FVec3(0, 0, 1);
+			}
+			if ((FeatureMask & (uint32)EDebugDrawJointFeature::Batch) && (BatchIndex >= 0))
+			{
+				FDebugDrawQueue::GetInstance().DrawDebugString(TextPos, FString::Format(TEXT("{0}"), { BatchIndex }), nullptr, FColor::Red, KINDA_SMALL_NUMBER, false, FontScale);
+				TextPos += FontHeight * FVec3(0, 0, 1);
+			}
+			if ((FeatureMask & (uint32)EDebugDrawJointFeature::Island) && (IslandIndex >= 0))
+			{
+				FDebugDrawQueue::GetInstance().DrawDebugString(TextPos, FString::Format(TEXT("{0}"), { IslandIndex }), nullptr, FColor::Red, KINDA_SMALL_NUMBER, false, FontScale);
+				TextPos += FontHeight * FVec3(0, 0, 1);
 			}
 		}
 
@@ -328,7 +346,7 @@ namespace Chaos
 				FVec3 Xa, Xb;
 				FMatrix33 Ra, Rb;
 				ConstraintHandle->CalculateConstraintSpace(Xa, Ra, Xb, Rb);
-				DrawJointConstraintImpl(SpaceTransform, Pa, Ca, Xa, Ra, Pb, Cb, Xb, Rb, ConstraintHandle->GetConstraintLevel(), ConstraintHandle->GetConstraintIndex(), ColorScale, FeatureMask);
+				DrawJointConstraintImpl(SpaceTransform, Pa, Ca, Xa, Ra, Pb, Cb, Xb, Rb, ConstraintHandle->GetConstraintIsland(), ConstraintHandle->GetConstraintLevel(), ConstraintHandle->GetConstraintColor(), ConstraintHandle->GetConstraintBatch(), ConstraintHandle->GetConstraintIndex(), ColorScale, FeatureMask);
 			}
 		}
 
diff --git a/Engine/Source/Runtime/Experimental/Chaos/Private/Chaos/Evolution/PBDMinEvolution.cpp b/Engine/Source/Runtime/Experimental/Chaos/Private/Chaos/Evolution/PBDMinEvolution.cpp
--- a/Engine/Source/Runtime/Experimental/Chaos/Private/Chaos/Evolution/PBDMinEvolution.cpp
+++ b/Engine/Source/Runtime/Experimental/Chaos/Private/Chaos/Evolution/PBDMinEvolution.cpp
@@ -170,6 +170,8 @@ namespace Chaos
 					// We could add the AABB at X/R here, but I'm avoiding another call to TransformedAABB. Hopefully this is good enough.
 					WorldSpaceBounds.GrowByVector(Particle.X() - Particle.P());
 
+					WorldSpaceBounds.ThickenSymmetrically(FVec3(CollisionDetector.GetBroadPhase().GetCullDistance()));
+
 					Particle.SetWorldSpaceInflatedBounds(WorldSpaceBounds);
 				}
 			}
diff --git a/Engine/Source/Runtime/Experimental/Chaos/Private/Chaos/Joint/ColoringGraph.cpp b/Engine/Source/Runtime/Experimental/Chaos/Private/Chaos/Joint/ColoringGraph.cpp
new file mode 100644
--- /dev/null
+++ b/Engine/Source/Runtime/Experimental/Chaos/Private/Chaos/Joint/ColoringGraph.cpp
@@ -0,0 +1,186 @@
+// Copyright Epic Games, Inc. All Rights Reserved.
+#include "Chaos/Joint/ColoringGraph.h"
+
+//PRAGMA_DISABLE_OPTIMIZATION
+
+namespace Chaos
+{
+	void FColoringGraph::Islandize()
+	{
+		Islands.Reset();
+
+		for (int32 VertexIndex = 0; VertexIndex < Vertices.Num(); ++VertexIndex)
+		{
+			int32 IslandIndex = Islands.AddDefaulted();
+			Vertices[VertexIndex].Island = IslandIndex;
+			Islands[IslandIndex].VertexIndices.Add(VertexIndex);
+		}
+
+		for (int32 EdgeIndex = 0; EdgeIndex < Edges.Num(); ++EdgeIndex)
+		{
+			int32 VertexIndex0 = Edges[EdgeIndex].VertexIndices[0];
+			int32 VertexIndex1 = Edges[EdgeIndex].VertexIndices[1];
+			int32 IslandIndex0 = Vertices[VertexIndex0].Island;
+			int32 IslandIndex1 = Vertices[VertexIndex1].Island;
+			if (IslandIndex0 != IslandIndex1)
+			{
+				MergeIslands(IslandIndex0, IslandIndex1);
+			}
+		}
+
+		PackIslands();
+	}
+
+	void FColoringGraph::Levelize()
+	{
+		TArray<int32> LevelQueue;
+		LevelQueue.Reserve(Vertices.Num());
+		for (int32 VertexIndex = 0; VertexIndex < Vertices.Num(); ++VertexIndex)
+		{
+			if (Vertices[VertexIndex].Level == 0)
+			{
+				LevelQueue.Add(VertexIndex);
+			}
+		}
+
+		for (int32 LevelQueueIndex = 0; LevelQueueIndex < LevelQueue.Num(); ++LevelQueueIndex)
+		{
+			int32 VertexIndex = LevelQueue[LevelQueueIndex];
+			int32 VertexLevel = Vertices[VertexIndex].Level;
+
+			for (int32 EdgeIndex : Vertices[VertexIndex].EdgeIndices)
+			{
+				int32 ChildVertexIndex0 = Edges[EdgeIndex].VertexIndices[0];
+				int32 ChildVertexIndex1 = Edges[EdgeIndex].VertexIndices[1];
+				if (Vertices[ChildVertexIndex0].Level == INDEX_NONE)
+				{
+					Vertices[ChildVertexIndex0].Level = VertexLevel + 1;
+					LevelQueue.Add(ChildVertexIndex0);
+				}
+				if (Vertices[ChildVertexIndex1].Level == INDEX_NONE)
+				{
+					Vertices[ChildVertexIndex1].Level = VertexLevel + 1;
+					LevelQueue.Add(ChildVertexIndex1);
+				}
+			}
+		}
+	}
+
+	void FColoringGraph::Colorize()
+	{
+		// Get the Vertices sorted by level for better coloring result
+		TArray<int32> VertexIndices;
+		VertexIndices.Reserve(Vertices.Num());
+		for (int32 VertexIndex = 0; VertexIndex < Vertices.Num(); ++VertexIndex)
+		{
+			VertexIndices.Add(VertexIndex);
+		}
+
+		VertexIndices.StableSort([this](int32 L, int32 R)
+			{
+				int32 LevelL = Vertices[L].Level;
+				int32 LevelR = Vertices[R].Level;
+				return LevelL < LevelR;
+			});
+
+
+		for (int32 VertexIndex : VertexIndices)
+		{
+			AssignColor(VertexIndex);
+		}
+	}
+
+	void FColoringGraph::MergeIslands(int32 IslandIndex0, int32 IslandIndex1)
+	{
+		// We will put the Vertices from the smaller island into the larger one
+		FIsland* Island0 = &Islands[IslandIndex0];
+		FIsland* Island1 = &Islands[IslandIndex1];
+		if (Island1->VertexIndices.Num() > Island0->VertexIndices.Num())
+		{
+			Swap(Island0, Island1);
+			Swap(IslandIndex0, IslandIndex1);
+		}
+
+		if (Island1->VertexIndices.Num() > 0)
+		{
+			// Add Vertices to the first island
+			Island0->VertexIndices.Append(Island1->VertexIndices);
+
+			// Update the Vertices' island indices
+			for (int32 VertexIndex1 : Island1->VertexIndices)
+			{
+				Vertices[VertexIndex1].Island = IslandIndex0;
+			}
+
+			// Remove Vertices from the second island
+			Island1->VertexIndices.Empty();
+		}
+	}
+
+	void FColoringGraph::PackIslands()
+	{
+		// Sort islands by number of vertices
+		Islands.Sort([](const FIsland& L, const FIsland& R)
+			{
+				return L.VertexIndices.Num() > R.VertexIndices.Num();
+			});
+
+		// Fix all island indices and clip array to remove empty islands
+		for (int32 IslandIndex = 0; IslandIndex < Islands.Num(); ++IslandIndex)
+		{
+			// When we hit an empty island we are done
+			if (Islands[IslandIndex].VertexIndices.Num() == 0)
+			{
+				Islands.SetNum(IslandIndex);
+				break;
+			}
+
+			// Set islands index on all verts in the island
+			for (int32 VertexIndex : Islands[IslandIndex].VertexIndices)
+			{
+				Vertices[VertexIndex].Island = IslandIndex;
+			}
+		}
+	}
+
+	void FColoringGraph::AssignColor(int32 VertexIndex)
+	{
+		if (Vertices[VertexIndex].Color == INDEX_NONE)
+		{
+			int32 IslandIndex = Vertices[VertexIndex].Island;
+
+			// We cannot use any color already used by our neighbors
+			// @todo(ccaulfield): eliminate allocations
+			TArray<bool> UsedColors;
+			UsedColors.SetNumZeroed(Islands[IslandIndex].NextColorIndex);
+			for (int32 EdgeIndex : Vertices[VertexIndex].EdgeIndices)
+			{
+				int32 OtherVertexIndex = (Edges[EdgeIndex].VertexIndices[0] == VertexIndex) ? Edges[EdgeIndex].VertexIndices[1] : Edges[EdgeIndex].VertexIndices[0];
+				int32 OtherColorIndex = Vertices[OtherVertexIndex].Color;
+				if (OtherColorIndex != INDEX_NONE)
+				{
+					UsedColors[OtherColorIndex] = true;
+				}
+			}
+
+			// Find the first unused Color
+			int32 VertexColor = INDEX_NONE;
+			for (int32 ColorIndex = 0; ColorIndex < UsedColors.Num(); ++ColorIndex)
+			{
+				if (!UsedColors[ColorIndex])
+				{
+					VertexColor = ColorIndex;
+				}
+			}
+
+			// If all colors are used, assign a new one
+			if (VertexColor == INDEX_NONE)
+			{
+				VertexColor = Islands[IslandIndex].NextColorIndex++;
+			}
+
+			Vertices[VertexIndex].Color = VertexColor;
+		}
+	}
+
+}
\ No newline at end of file
diff --git a/Engine/Source/Runtime/Experimental/Chaos/Private/Chaos/Joint/JointSolverConstraints.cpp b/Engine/Source/Runtime/Experimental/Chaos/Private/Chaos/Joint/JointSolverConstraints.cpp
new file mode 100644
--- /dev/null
+++ b/Engine/Source/Runtime/Experimental/Chaos/Private/Chaos/Joint/JointSolverConstraints.cpp
@@ -0,0 +1,809 @@
+// Copyright Epic Games, Inc. All Rights Reserved.
+#include "Chaos/Joint/JointSolverConstraints.h"
+
+
+#include "Chaos/Joint/ChaosJointLog.h"
+#include "Chaos/PBDJointConstraintTypes.h"
+#include "Chaos/PBDJointConstraintUtilities.h"
+#include "Chaos/Utilities.h"
+
+//PRAGMA_DISABLE_OPTIMIZATION
+
+int32 bChaos_Joint_MultiDimension = true;
+FAutoConsoleVariableRef CVarChaosJointMultiDimension(TEXT("p.Chaos.Joint.MultiDimension"), bChaos_Joint_MultiDimension, TEXT(""));
+
+namespace Chaos
+{
+
+	//
+	//////////////////////////////////////////////////////////////////////////
+	//
+
+	FJointSolverConstraints::FJointSolverConstraints()
+		: JointIndex(INDEX_NONE)
+		, LinearRowIndexBegin(INDEX_NONE)
+		, LinearRowIndexEnd(INDEX_NONE)
+		, AngularRowIndexBegin(INDEX_NONE)
+		, AngularRowIndexEnd(INDEX_NONE)
+		, PositionTolerance(0.0f)
+		, AngleTolerance(0.0f)
+		, bNeedSwingTwist(false)
+		, bNeedLockedAxes(false)
+	{
+	}
+
+	void FJointSolverConstraints::AddPositionConstraints(
+		TArray<FJointSolverConstraintRowData>& RowDatas,
+		const FPBDJointSolverSettings& SolverSettings,
+		const FPBDJointSettings& JointSettings)
+	{
+		PositionTolerance = SolverSettings.PositionTolerance;
+
+		LinearRowIndexBegin = RowDatas.Num();
+		LinearRowIndexEnd = LinearRowIndexBegin;
+
+		const TVector<EJointMotionType, 3>& LinearMotion = JointSettings.LinearMotionTypes;
+		const TVector<bool, 3> bLinearLocked =
+		{
+			(LinearMotion[0] == EJointMotionType::Locked),
+			(LinearMotion[1] == EJointMotionType::Locked),
+			(LinearMotion[2] == EJointMotionType::Locked),
+		};
+		const TVector<bool, 3> bLinearLimted =
+		{
+			(LinearMotion[0] == EJointMotionType::Limited),
+			(LinearMotion[1] == EJointMotionType::Limited),
+			(LinearMotion[2] == EJointMotionType::Limited),
+		};
+
+		if (bLinearLocked[0] && bLinearLocked[1] && bLinearLocked[2])
+		{
+			// Hard point constraint (most common case)
+			AddPointPositionConstraint(RowDatas, SolverSettings, JointSettings);
+		}
+		else if (bLinearLimted[0] && bLinearLimted[1] && bLinearLimted[2])
+		{
+			// Spherical constraint
+			AddSphericalPositionConstraint(RowDatas, SolverSettings, JointSettings);
+		}
+		else if (bLinearLocked[1] && bLinearLocked[2] && !bLinearLocked[0])
+		{
+			// Line constraint along X axis
+			AddCylindricalPositionConstraint(RowDatas, 0, LinearMotion[0], EJointMotionType::Locked, SolverSettings, JointSettings);
+		}
+		else if (bLinearLocked[0] && bLinearLocked[2] && !bLinearLocked[1])
+		{
+			// Line constraint along Y axis
+			AddCylindricalPositionConstraint(RowDatas, 1, LinearMotion[1], EJointMotionType::Locked, SolverSettings, JointSettings);
+		}
+		else if (bLinearLocked[0] && bLinearLocked[1] && !bLinearLocked[2])
+		{
+			// Line constraint along Z axis
+			AddCylindricalPositionConstraint(RowDatas, 2, LinearMotion[2], EJointMotionType::Locked, SolverSettings, JointSettings);
+		}
+		else if (bLinearLimted[1] && bLinearLimted[2] && !bLinearLimted[0])
+		{
+			// Cylindrical constraint along X axis
+			AddCylindricalPositionConstraint(RowDatas, 0, LinearMotion[0], EJointMotionType::Limited, SolverSettings, JointSettings);
+		}
+		else if (bLinearLimted[0] && bLinearLimted[2] && !bLinearLimted[1])
+		{
+			// Cylindrical constraint along Y axis
+			AddCylindricalPositionConstraint(RowDatas, 1, LinearMotion[1], EJointMotionType::Limited, SolverSettings, JointSettings);
+		}
+		else if (bLinearLimted[0] && bLinearLimted[1] && !bLinearLimted[2])
+		{
+			// Cylindrical constraint along Z axis
+			AddCylindricalPositionConstraint(RowDatas, 2, LinearMotion[2], EJointMotionType::Limited, SolverSettings, JointSettings);
+		}
+		else if (bLinearLocked[0] || bLinearLimted[0])
+		{
+			// Planar constraint along X axis
+			AddPlanarPositionConstraint(RowDatas, 0, LinearMotion[0], SolverSettings, JointSettings);
+		}
+		else if (bLinearLocked[1] || bLinearLimted[1])
+		{
+			// Planar constraint along Y axis
+			AddPlanarPositionConstraint(RowDatas, 1, LinearMotion[1], SolverSettings, JointSettings);
+		}
+		else if (bLinearLocked[2] || bLinearLimted[2])
+		{
+			// Planar constraint along Z axis
+			AddPlanarPositionConstraint(RowDatas, 2, LinearMotion[2], SolverSettings, JointSettings);
+		}
+
+		if (SolverSettings.bEnableDrives)
+		{
+			TVector<bool, 3> bDriven =
+			{
+				(JointSettings.bLinearPositionDriveEnabled[0] || JointSettings.bLinearVelocityDriveEnabled[0]) && (JointSettings.LinearMotionTypes[0] != EJointMotionType::Locked),
+				(JointSettings.bLinearPositionDriveEnabled[1] || JointSettings.bLinearVelocityDriveEnabled[1]) && (JointSettings.LinearMotionTypes[1] != EJointMotionType::Locked),
+				(JointSettings.bLinearPositionDriveEnabled[2] || JointSettings.bLinearVelocityDriveEnabled[2]) && (JointSettings.LinearMotionTypes[2] != EJointMotionType::Locked),
+			};
+
+			if (bDriven[0] && bDriven[1] && bDriven[2])
+			{
+				AddPositionDrive(RowDatas, SolverSettings, JointSettings, EJointSolverConstraintUpdateType::Linear_SphericalDrive, INDEX_NONE);
+			}
+			else if (bDriven[1] && bDriven[2])
+			{
+				AddPositionDrive(RowDatas, SolverSettings, JointSettings, EJointSolverConstraintUpdateType::Linear_CircularDrive, 0);
+			}
+			else if (bDriven[0] && bDriven[2])
+			{
+				AddPositionDrive(RowDatas, SolverSettings, JointSettings, EJointSolverConstraintUpdateType::Linear_CircularDrive, 1);
+			}
+			else if (bDriven[0] && bDriven[1])
+			{
+				AddPositionDrive(RowDatas, SolverSettings, JointSettings, EJointSolverConstraintUpdateType::Linear_CircularDrive, 2);
+			}
+			else if (bDriven[0])
+			{
+				AddPositionDrive(RowDatas, SolverSettings, JointSettings, EJointSolverConstraintUpdateType::Linear_AxialDrive, 0);
+			}
+			else if (bDriven[1])
+			{
+				AddPositionDrive(RowDatas, SolverSettings, JointSettings, EJointSolverConstraintUpdateType::Linear_AxialDrive, 1);
+			}
+			else if (bDriven[2])
+			{
+				AddPositionDrive(RowDatas, SolverSettings, JointSettings, EJointSolverConstraintUpdateType::Linear_AxialDrive, 2);
+			}
+		}
+	}
+
+	void FJointSolverConstraints::UpdatePositionConstraints(
+		const TArray<FJointSolverConstraintRowData>& RowDatas,
+		TArray<FJointSolverConstraintRowState>& RowStates,
+		const FJointSolverJointState& JointState,
+		const FPBDJointSettings& JointSettings)
+	{
+		const FRotation3& R0 = JointState.Rs[0];
+		const FVec3& X0 = JointState.Xs[0];
+		const FVec3& X1 = JointState.Xs[1];
+
+		// NOTE: second loop may jump rows, so reset all here
+		for (int32 RowIndex = LinearRowIndexBegin; RowIndex < LinearRowIndexEnd; ++RowIndex)
+		{
+			RowStates[RowIndex].IterationReset();
+		}
+
+		// Calculate all the rotation constraint axes etc. Some of the update functions actually
+		// update multple rows so we need to skip them in the loop if so.
+		for (int32 RowIndex = LinearRowIndexBegin; RowIndex < LinearRowIndexEnd; /* no op */)
+		{
+			switch (RowDatas[RowIndex].UpdateType)
+			{
+			case EJointSolverConstraintUpdateType::Linear_Point:
+				RowIndex += UpdatePointPositionConstraint(RowIndex, RowDatas, RowStates, JointSettings, X0, X1);
+				break;
+			case EJointSolverConstraintUpdateType::Linear_Spherical:
+				RowIndex += UpdateSphericalPositionConstraint(RowIndex, RowDatas, RowStates, JointSettings, X0, X1);
+				break;
+			case EJointSolverConstraintUpdateType::Linear_Cylindrical:
+				RowIndex += UpdateCylindricalPositionConstraint(RowIndex, RowDatas, RowStates, JointSettings, R0, X0, X1);
+				break;
+			case EJointSolverConstraintUpdateType::Linear_Planar:
+				RowIndex += UpdatePlanarPositionConstraint(RowIndex, RowDatas, RowStates, JointSettings, R0, X0, X1);
+				break;
+			case EJointSolverConstraintUpdateType::Linear_SphericalDrive:
+				RowIndex += UpdateSphericalPositionDrive(RowIndex, RowDatas, RowStates, JointSettings, R0, X0, X1);
+				break;
+			case EJointSolverConstraintUpdateType::Linear_CircularDrive:
+				RowIndex += UpdateCircularPositionDrive(RowIndex, RowDatas, RowStates, JointSettings, R0, X0, X1);
+				break;
+			case EJointSolverConstraintUpdateType::Linear_AxialDrive:
+				RowIndex += UpdateAxialPositionDrive(RowIndex, RowDatas, RowStates, JointSettings, R0, X0, X1);
+				break;
+			}
+		}
+	}
+
+	void FJointSolverConstraints::AddRotationConstraints(
+		TArray<FJointSolverConstraintRowData>& RowDatas,
+		const FPBDJointSolverSettings& SolverSettings,
+		const FPBDJointSettings& JointSettings)
+	{
+		AngleTolerance = SolverSettings.AngleTolerance;
+
+		AngularRowIndexBegin = RowDatas.Num();
+		AngularRowIndexEnd = AngularRowIndexBegin;
+
+		EJointMotionType TwistMotion = JointSettings.AngularMotionTypes[(int32)EJointAngularConstraintIndex::Twist];
+		EJointMotionType Swing1Motion = JointSettings.AngularMotionTypes[(int32)EJointAngularConstraintIndex::Swing1];
+		EJointMotionType Swing2Motion = JointSettings.AngularMotionTypes[(int32)EJointAngularConstraintIndex::Swing2];
+
+		// Add twist constraint
+		if (SolverSettings.bEnableTwistLimits)
+		{
+			bool bTwistSoft = JointSettings.bSoftTwistLimitsEnabled;
+			if (TwistMotion == EJointMotionType::Limited)
+			{
+				AddTwistConstraint(RowDatas, SolverSettings, JointSettings, EJointSolverConstraintUpdateType::Angular_Twist, bTwistSoft);
+			}
+			else if (TwistMotion == EJointMotionType::Locked)
+			{
+				AddTwistConstraint(RowDatas, SolverSettings, JointSettings, EJointSolverConstraintUpdateType::Angular_Locked, false);
+			}
+		}
+
+		// Add swing constraints
+		if (SolverSettings.bEnableSwingLimits)
+		{
+			bool bSwingSoft = JointSettings.bSoftSwingLimitsEnabled;
+			if ((Swing1Motion == EJointMotionType::Limited) && (Swing2Motion == EJointMotionType::Limited))
+			{
+				AddSwingConstraint(RowDatas, SolverSettings, JointSettings, EJointSolverConstraintUpdateType::Angular_Cone, EJointAngularConstraintIndex::Swing1, bSwingSoft);
+			}
+			else if ((Swing1Motion == EJointMotionType::Limited) && (Swing2Motion == EJointMotionType::Locked))
+			{
+				AddSwingConstraint(RowDatas, SolverSettings, JointSettings, EJointSolverConstraintUpdateType::Angular_SingleLockedSwing, EJointAngularConstraintIndex::Swing2, false);
+				AddSwingConstraint(RowDatas, SolverSettings, JointSettings, EJointSolverConstraintUpdateType::Angular_SingleLimitedSwing, EJointAngularConstraintIndex::Swing1, bSwingSoft);
+			}
+			else if ((Swing1Motion == EJointMotionType::Limited) && (Swing2Motion == EJointMotionType::Free))
+			{
+				AddSwingConstraint(RowDatas, SolverSettings, JointSettings, EJointSolverConstraintUpdateType::Angular_DualConeSwing, EJointAngularConstraintIndex::Swing1, bSwingSoft);
+			}
+			else if ((Swing1Motion == EJointMotionType::Locked) && (Swing2Motion == EJointMotionType::Limited))
+			{
+				AddSwingConstraint(RowDatas, SolverSettings, JointSettings, EJointSolverConstraintUpdateType::Angular_SingleLockedSwing, EJointAngularConstraintIndex::Swing1, false);
+				AddSwingConstraint(RowDatas, SolverSettings, JointSettings, EJointSolverConstraintUpdateType::Angular_SingleLimitedSwing, EJointAngularConstraintIndex::Swing2, bSwingSoft);
+			}
+			else if ((Swing1Motion == EJointMotionType::Locked) && (Swing2Motion == EJointMotionType::Locked))
+			{
+				AddSwingConstraint(RowDatas, SolverSettings, JointSettings, EJointSolverConstraintUpdateType::Angular_Locked, EJointAngularConstraintIndex::Swing1, false);
+				AddSwingConstraint(RowDatas, SolverSettings, JointSettings, EJointSolverConstraintUpdateType::Angular_Locked, EJointAngularConstraintIndex::Swing2, false);
+			}
+			else if ((Swing1Motion == EJointMotionType::Locked) && (Swing2Motion == EJointMotionType::Free))
+			{
+				AddSwingConstraint(RowDatas, SolverSettings, JointSettings, EJointSolverConstraintUpdateType::Angular_SingleLockedSwing, EJointAngularConstraintIndex::Swing1, false);
+			}
+			else if ((Swing1Motion == EJointMotionType::Free) && (Swing2Motion == EJointMotionType::Limited))
+			{
+				AddSwingConstraint(RowDatas, SolverSettings, JointSettings, EJointSolverConstraintUpdateType::Angular_DualConeSwing, EJointAngularConstraintIndex::Swing2, bSwingSoft);
+			}
+			else if ((Swing1Motion == EJointMotionType::Free) && (Swing2Motion == EJointMotionType::Locked))
+			{
+				AddSwingConstraint(RowDatas, SolverSettings, JointSettings, EJointSolverConstraintUpdateType::Angular_SingleLockedSwing, EJointAngularConstraintIndex::Swing2, false);
+			}
+		}
+
+		// Add Drives
+		if (SolverSettings.bEnableDrives)
+		{
+			bool bTwistLocked = TwistMotion == EJointMotionType::Locked;
+			bool bSwing1Locked = Swing1Motion == EJointMotionType::Locked;
+			bool bSwing2Locked = Swing2Motion == EJointMotionType::Locked;
+
+			// No SLerp drive if we have a locked rotation (it will be grayed out in the editor in this case, but could still have been set before the rotation was locked)
+			if ((JointSettings.bAngularSLerpPositionDriveEnabled || JointSettings.bAngularSLerpVelocityDriveEnabled) && !bTwistLocked && !bSwing1Locked && !bSwing2Locked)
+			{
+				AddSLerpDrive(RowDatas, SolverSettings, JointSettings);
+			}
+			else
+			{
+				if ((JointSettings.bAngularTwistPositionDriveEnabled || JointSettings.bAngularTwistVelocityDriveEnabled) && !bTwistLocked)
+				{
+					AddTwistDrive(RowDatas, SolverSettings, JointSettings);
+				}
+
+				const bool bSwingDriveEnabled = (JointSettings.bAngularSwingPositionDriveEnabled || JointSettings.bAngularSwingVelocityDriveEnabled);
+				if (bSwingDriveEnabled && !bSwing1Locked && !bSwing2Locked)
+				{
+					AddSwingDrive(RowDatas, SolverSettings, JointSettings, EJointSolverConstraintUpdateType::Angular_ConeDrive, EJointAngularConstraintIndex::Swing1);
+				}
+				else if (bSwingDriveEnabled && !bSwing1Locked)
+				{
+					AddSwingDrive(RowDatas, SolverSettings, JointSettings, EJointSolverConstraintUpdateType::Angular_SwingDrive, EJointAngularConstraintIndex::Swing1);
+				}
+				else if (bSwingDriveEnabled && !bSwing2Locked)
+				{
+					AddSwingDrive(RowDatas, SolverSettings, JointSettings, EJointSolverConstraintUpdateType::Angular_SwingDrive, EJointAngularConstraintIndex::Swing2);
+				}
+			}
+		}
+
+		// Determine what intermediate data needs to be calculated in the update step
+		bNeedLockedAxes = false;
+		bNeedSwingTwist = false;
+		for (int32 RowIndex = AngularRowIndexBegin; RowIndex < AngularRowIndexEnd; ++RowIndex)
+		{
+			FJointSolverConstraintRowData& Constraint = RowDatas[RowIndex];
+			switch (Constraint.UpdateType)
+			{
+			case EJointSolverConstraintUpdateType::Angular_Twist:
+			case EJointSolverConstraintUpdateType::Angular_Cone:
+			case EJointSolverConstraintUpdateType::Angular_SingleLimitedSwing:
+				bNeedSwingTwist = true;
+				break;
+			case EJointSolverConstraintUpdateType::Angular_Locked:
+				bNeedLockedAxes = true;
+				break;
+			}
+		}
+	}
+
+	void FJointSolverConstraints::UpdateRotationConstraints(
+		const TArray<FJointSolverConstraintRowData>& RowDatas,
+		TArray<FJointSolverConstraintRowState>& RowStates,
+		const FJointSolverJointState& JointState,
+		const FPBDJointSettings& JointSettings)
+	{
+		const FRotation3& R0 = JointState.Rs[0];
+		const FRotation3& R1 = JointState.Rs[1];
+		const FRotation3 R01 = R0.Inverse() * R1;
+		FRotation3 R01Twist = FRotation3::FromIdentity();
+		FRotation3 R01Swing = FRotation3::FromIdentity();
+		if (bNeedSwingTwist)
+		{
+			R01.ToSwingTwistX(R01Swing, R01Twist);
+		}
+
+		FVec3 Axes[3] = { FVec3(0), FVec3(0), FVec3(0) };
+		if (bNeedLockedAxes)
+		{
+			FPBDJointUtilities::GetLockedRotationAxes(R0, R1, Axes[0], Axes[1], Axes[2]);
+		}
+
+		for (int32 RowIndex = AngularRowIndexBegin; RowIndex < AngularRowIndexEnd; ++RowIndex)
+		{
+			const FJointSolverConstraintRowData& RowData = RowDatas[RowIndex];
+			FJointSolverConstraintRowState& RowState = RowStates[RowIndex];
+
+			RowState.IterationReset();
+
+			switch (RowData.UpdateType)
+			{
+			case EJointSolverConstraintUpdateType::Angular_Twist:
+				check(bNeedSwingTwist);
+				UpdateTwistConstraint(RowData, RowState, JointSettings, R1, R01Twist);
+				break;
+			case EJointSolverConstraintUpdateType::Angular_Cone:
+				check(bNeedSwingTwist);
+				UpdateConeSwingConstraint(RowData, RowState, JointSettings, R0, R01Swing);
+				break;
+			case EJointSolverConstraintUpdateType::Angular_SingleLockedSwing:
+				UpdateSingleLockedSwingConstraint(RowData, RowState, JointSettings, R0, R1);
+				break;
+			case EJointSolverConstraintUpdateType::Angular_SingleLimitedSwing:
+				check(bNeedSwingTwist);
+				UpdateSingleLimitedSwingConstraint(RowData, RowState, JointSettings, R0, R01Swing);
+				break;
+			case EJointSolverConstraintUpdateType::Angular_DualConeSwing:
+				UpdateDualConeSwingConstraint(RowData, RowState, JointSettings, R0, R1);
+				break;
+			case EJointSolverConstraintUpdateType::Angular_Locked:
+			{
+				check(bNeedLockedAxes);
+				int32 AxisIndex = FJointConstants::AxisIndex((EJointAngularConstraintIndex)RowData.ConstraintIndex);
+				RowState.Axis = Axes[AxisIndex];
+				FReal Position = (AxisIndex == 0) ? R01.X : ((AxisIndex == 1) ? R01.Y : R01.Z); // Can't index Quat elements...
+				RowState.CalculateError(Position, 0.0f);
+				break;
+			}
+			case EJointSolverConstraintUpdateType::Angular_TwistDrive:
+				UpdateTwistDrive(RowData, RowState, JointSettings, R1, R01Twist);
+				break;
+			case EJointSolverConstraintUpdateType::Angular_ConeDrive:
+				UpdateConeSwingDrive(RowData, RowState, JointSettings, R0, R01Swing);
+				break;
+			case EJointSolverConstraintUpdateType::Angular_SwingDrive:
+				UpdateSwingDrive(RowData, RowState, JointSettings, R0, R01Swing);
+				break;
+			case EJointSolverConstraintUpdateType::Angular_SLerpDrive:
+				UpdateSLerpDrive(RowData, RowState, JointSettings, R0, R1);
+				break;
+			}
+		}
+	}
+
+	//
+	//////////////////////////////////////////////////////////////////////////
+	//
+
+	// TODO: Handle multi-dimensional constraints better
+	void FJointSolverConstraints::AddPointPositionConstraint(
+		TArray<FJointSolverConstraintRowData>& RowDatas,
+		const FPBDJointSolverSettings& SolverSettings,
+		const FPBDJointSettings& JointSettings)
+	{
+		for (int32 AxisIndex = 0; AxisIndex < 3; ++AxisIndex)
+		{
+			int32 Index = RowDatas.AddDefaulted();
+			LinearRowIndexEnd = Index + 1;
+
+			FJointSolverConstraintRowData& RowData = RowDatas[Index];
+			RowData.UpdateType = EJointSolverConstraintUpdateType::Linear_Point;
+			RowData.NumRows = 1;
+			RowData.Stiffness = FPBDJointUtilities::GetLinearStiffness(SolverSettings, JointSettings);
+			RowData.JointIndex = JointIndex;
+			RowData.ConstraintIndex = AxisIndex;
+			RowData.Limit = 0.0f;
+		}
+
+		if (bChaos_Joint_MultiDimension)
+		{
+			RowDatas.Last(2).NumRows = 3;
+			RowDatas.Last(1).NumRows = 0;
+			RowDatas.Last(0).NumRows = 0;
+		}
+	}
+
+	void FJointSolverConstraints::AddSphericalPositionConstraint(
+		TArray<FJointSolverConstraintRowData>& RowDatas,
+		const FPBDJointSolverSettings& SolverSettings,
+		const FPBDJointSettings& JointSettings)
+	{
+		int32 Index = RowDatas.AddDefaulted();
+		LinearRowIndexEnd = Index + 1;
+
+		// This is the soft version of 
+		FJointSolverConstraintRowData& RowData = RowDatas[Index];
+		RowData.UpdateType = EJointSolverConstraintUpdateType::Linear_Spherical;
+		RowData.NumRows = 1;
+		RowData.JointIndex = JointIndex;
+		RowData.Limit = JointSettings.LinearLimit;
+		if (JointSettings.bSoftLinearLimitsEnabled)
+		{
+			RowData.Stiffness = FPBDJointUtilities::GetSoftLinearStiffness(SolverSettings, JointSettings);
+			RowData.Damping = FPBDJointUtilities::GetSoftLinearDamping(SolverSettings, JointSettings);
+			RowData.bIsAccelerationMode = FPBDJointUtilities::GetLinearSoftAccelerationMode(SolverSettings, JointSettings);
+			RowData.bIsSoft = true;
+		}
+		else
+		{
+			RowData.Stiffness = FPBDJointUtilities::GetLinearStiffness(SolverSettings, JointSettings);
+
+		}
+	}
+
+	void FJointSolverConstraints::AddCylindricalPositionConstraint(
+		TArray<FJointSolverConstraintRowData>& RowDatas,
+		const int32 AxisIndex,
+		const EJointMotionType AxialMotion,
+		const EJointMotionType RadialMotion,
+		const FPBDJointSolverSettings& SolverSettings,
+		const FPBDJointSettings& JointSettings)
+	{
+		{
+			int32 AxialIndex = RowDatas.AddDefaulted(2);
+			int32 RadialIndex = AxialIndex + 1;
+			LinearRowIndexEnd = RadialIndex + 1;
+
+			FJointSolverConstraintRowData& AxisRowData = RowDatas[AxialIndex];
+			AxisRowData.UpdateType = EJointSolverConstraintUpdateType::Linear_Cylindrical;
+			AxisRowData.NumRows = 1;
+			AxisRowData.JointIndex = JointIndex;
+			AxisRowData.ConstraintIndex = AxisIndex;
+			if ((AxialMotion == EJointMotionType::Limited) && JointSettings.bSoftLinearLimitsEnabled)
+			{
+				AxisRowData.Stiffness = FPBDJointUtilities::GetSoftLinearStiffness(SolverSettings, JointSettings);
+				AxisRowData.Damping = FPBDJointUtilities::GetSoftLinearDamping(SolverSettings, JointSettings);
+				AxisRowData.bIsAccelerationMode = FPBDJointUtilities::GetLinearSoftAccelerationMode(SolverSettings, JointSettings);
+				AxisRowData.bIsSoft = true;
+				AxisRowData.Limit = JointSettings.LinearLimit;
+			}
+			else if (AxialMotion == EJointMotionType::Limited)
+			{
+				AxisRowData.Stiffness = FPBDJointUtilities::GetLinearStiffness(SolverSettings, JointSettings);
+				AxisRowData.Limit = JointSettings.LinearLimit;
+			}
+			else if (AxialMotion == EJointMotionType::Locked)
+			{
+				AxisRowData.Stiffness = FPBDJointUtilities::GetLinearStiffness(SolverSettings, JointSettings);
+				AxisRowData.Limit = 0.0f;
+			}
+			else
+			{
+				AxisRowData.Stiffness = 0.0f;
+				AxisRowData.Limit = FLT_MAX;
+			}
+
+			FJointSolverConstraintRowData& RadialRowData = RowDatas[RadialIndex];
+			RadialRowData.UpdateType = EJointSolverConstraintUpdateType::None;	// Handled by Cylindrical update on previous constraint
+			RadialRowData.NumRows = 1;
+			RadialRowData.JointIndex = JointIndex;
+			RadialRowData.ConstraintIndex = AxisIndex;
+			if ((RadialMotion == EJointMotionType::Limited) && JointSettings.bSoftLinearLimitsEnabled)
+			{
+				RadialRowData.Stiffness = FPBDJointUtilities::GetSoftLinearStiffness(SolverSettings, JointSettings);
+				RadialRowData.Damping = FPBDJointUtilities::GetSoftLinearDamping(SolverSettings, JointSettings);
+				RadialRowData.bIsAccelerationMode = FPBDJointUtilities::GetLinearSoftAccelerationMode(SolverSettings, JointSettings);
+				RadialRowData.bIsSoft = true;
+				RadialRowData.Limit = JointSettings.LinearLimit;
+			}
+			else if (RadialMotion == EJointMotionType::Limited)
+			{
+				RadialRowData.Stiffness = FPBDJointUtilities::GetLinearStiffness(SolverSettings, JointSettings);
+				RadialRowData.Limit = JointSettings.LinearLimit;
+			}
+			else if (RadialMotion == EJointMotionType::Locked)
+			{
+				RadialRowData.Stiffness = FPBDJointUtilities::GetLinearStiffness(SolverSettings, JointSettings);
+				RadialRowData.Limit = 0.0f;
+			}
+			else
+			{
+				RadialRowData.Stiffness = 0.0f;
+				RadialRowData.Limit = FLT_MAX;
+			}
+		}
+	}
+
+	void FJointSolverConstraints::AddPlanarPositionConstraint(
+		TArray<FJointSolverConstraintRowData>& RowDatas,
+		const int32 AxisIndex,
+		const EJointMotionType AxialMotion,
+		const FPBDJointSolverSettings& SolverSettings,
+		const FPBDJointSettings& JointSettings)
+	{
+		int32 Index = RowDatas.AddDefaulted();
+		LinearRowIndexEnd = Index + 1;
+
+		FJointSolverConstraintRowData& RowData = RowDatas[Index];
+		RowData.UpdateType = EJointSolverConstraintUpdateType::Linear_Planar;
+		RowData.NumRows = 1;
+		RowData.JointIndex = JointIndex;
+		RowData.Limit = JointSettings.LinearLimit;
+
+		if ((AxialMotion == EJointMotionType::Limited) && JointSettings.bSoftLinearLimitsEnabled)
+		{
+			RowData.Stiffness = FPBDJointUtilities::GetSoftLinearStiffness(SolverSettings, JointSettings);
+			RowData.Damping = FPBDJointUtilities::GetSoftLinearDamping(SolverSettings, JointSettings);
+			RowData.bIsAccelerationMode = FPBDJointUtilities::GetLinearSoftAccelerationMode(SolverSettings, JointSettings);
+			RowData.bIsSoft = true;
+		}
+		else
+		{
+			RowData.Stiffness = FPBDJointUtilities::GetLinearStiffness(SolverSettings, JointSettings);
+		}
+	}
+
+	void FJointSolverConstraints::AddPositionDrive(
+		TArray<FJointSolverConstraintRowData>& RowDatas,
+		const FPBDJointSolverSettings& SolverSettings,
+		const FPBDJointSettings& JointSettings,
+		const EJointSolverConstraintUpdateType UpdateType,
+		const int32 AxisIndex)
+	{
+		int32 Index = RowDatas.AddDefaulted();
+		LinearRowIndexEnd = Index + 1;
+
+		FJointSolverConstraintRowData& RowData = RowDatas[Index];
+		RowData.UpdateType = UpdateType;
+		RowData.NumRows = 1;
+		RowData.JointIndex = JointIndex;
+		RowData.ConstraintIndex = AxisIndex;
+		RowData.Limit = 0.0f;
+		RowData.Stiffness = FPBDJointUtilities::GetLinearDriveStiffness(SolverSettings, JointSettings);
+		RowData.Damping = FPBDJointUtilities::GetLinearDriveDamping(SolverSettings, JointSettings);
+		RowData.bIsAccelerationMode = FPBDJointUtilities::GetDriveAccelerationMode(SolverSettings, JointSettings);
+		RowData.bIsSoft = true;
+	}
+
+	//
+	//////////////////////////////////////////////////////////////////////////
+	//
+
+
+	void FJointSolverConstraints::AddTwistConstraint(
+		TArray<FJointSolverConstraintRowData>& RowDatas,
+		const FPBDJointSolverSettings& SolverSettings,
+		const FPBDJointSettings& JointSettings,
+		const EJointSolverConstraintUpdateType Type,
+		const bool bUseSoftLimit)
+	{
+		int32 Index = RowDatas.AddDefaulted();
+		FJointSolverConstraintRowData& RowData = RowDatas[Index];
+		AngularRowIndexEnd = Index + 1;
+
+		RowData.JointIndex = JointIndex;
+		RowData.UpdateType = Type;
+		RowData.ConstraintIndex = (int32)EJointAngularConstraintIndex::Twist;
+		RowData.NumRows = 1;
+		RowData.Limit = JointSettings.AngularLimits[(int32)EJointAngularConstraintIndex::Twist];
+
+		if (bUseSoftLimit)
+		{
+			RowData.Stiffness = FPBDJointUtilities::GetSoftTwistStiffness(SolverSettings, JointSettings);
+			RowData.Damping = FPBDJointUtilities::GetSoftTwistDamping(SolverSettings, JointSettings);
+			RowData.bIsAccelerationMode = FPBDJointUtilities::GetAngularSoftAccelerationMode(SolverSettings, JointSettings);
+			RowData.bIsSoft = true;
+		}
+		else
+		{
+			RowData.Stiffness = FPBDJointUtilities::GetTwistStiffness(SolverSettings, JointSettings);
+		}
+	}
+
+	void FJointSolverConstraints::AddSwingConstraint(
+		TArray<FJointSolverConstraintRowData>& RowDatas,
+		const FPBDJointSolverSettings& SolverSettings,
+		const FPBDJointSettings& JointSettings,
+		const EJointSolverConstraintUpdateType Type,
+		const EJointAngularConstraintIndex ConstraintIndex,
+		const bool bUseSoftLimit)
+	{
+		int32 Index = RowDatas.AddDefaulted();
+		FJointSolverConstraintRowData& RowData = RowDatas[Index];
+		AngularRowIndexEnd = Index + 1;
+
+		RowData.JointIndex = JointIndex;
+		RowData.UpdateType = Type;
+		RowData.ConstraintIndex = (int32)ConstraintIndex;
+		RowData.NumRows = 1;
+		RowData.Limit = JointSettings.AngularLimits[(int32)ConstraintIndex];
+
+		if (bUseSoftLimit)
+		{
+			RowData.Stiffness = FPBDJointUtilities::GetSoftSwingStiffness(SolverSettings, JointSettings);
+			RowData.Damping = FPBDJointUtilities::GetSoftSwingDamping(SolverSettings, JointSettings);
+			RowData.bIsAccelerationMode = FPBDJointUtilities::GetAngularSoftAccelerationMode(SolverSettings, JointSettings);
+			RowData.bIsSoft = true;
+		}
+		else
+		{
+			RowData.Stiffness = FPBDJointUtilities::GetSwingStiffness(SolverSettings, JointSettings);
+		}
+	}
+
+	void FJointSolverConstraints::AddTwistDrive(
+		TArray<FJointSolverConstraintRowData>& RowDatas,
+		const FPBDJointSolverSettings& SolverSettings,
+		const FPBDJointSettings& JointSettings)
+	{
+		int32 Index = RowDatas.AddDefaulted();
+		FJointSolverConstraintRowData& RowData = RowDatas[Index];
+		AngularRowIndexEnd = Index + 1;
+
+		RowData.JointIndex = JointIndex;
+		RowData.UpdateType = EJointSolverConstraintUpdateType::Angular_TwistDrive;
+		RowData.NumRows = 1;
+		RowData.Limit = 0.0f;
+
+		RowData.Stiffness = FPBDJointUtilities::GetAngularTwistDriveStiffness(SolverSettings, JointSettings);
+		RowData.Damping = FPBDJointUtilities::GetAngularTwistDriveDamping(SolverSettings, JointSettings);
+		RowData.bIsAccelerationMode = FPBDJointUtilities::GetDriveAccelerationMode(SolverSettings, JointSettings);
+		RowData.bIsSoft = true;
+	}
+
+	void FJointSolverConstraints::AddSwingDrive(
+		TArray<FJointSolverConstraintRowData>& RowDatas,
+		const FPBDJointSolverSettings& SolverSettings,
+		const FPBDJointSettings& JointSettings,
+		const EJointSolverConstraintUpdateType Type,
+		const EJointAngularConstraintIndex ConstraintIndex)
+	{
+		int32 Index = RowDatas.AddDefaulted();
+		FJointSolverConstraintRowData& RowData = RowDatas[Index];
+		AngularRowIndexEnd = Index + 1;
+
+		RowData.JointIndex = JointIndex;
+		RowData.UpdateType = Type;
+		RowData.ConstraintIndex = (int32)ConstraintIndex;
+		RowData.NumRows = 1;
+		RowData.Limit = 0.0f;
+
+		RowData.Stiffness = FPBDJointUtilities::GetAngularSwingDriveStiffness(SolverSettings, JointSettings);
+		RowData.Damping = FPBDJointUtilities::GetAngularSwingDriveDamping(SolverSettings, JointSettings);
+		RowData.bIsAccelerationMode = FPBDJointUtilities::GetDriveAccelerationMode(SolverSettings, JointSettings);
+		RowData.bIsSoft = true;
+	}
+
+	void FJointSolverConstraints::AddSLerpDrive(
+		TArray<FJointSolverConstraintRowData>& RowDatas,
+		const FPBDJointSolverSettings& SolverSettings,
+		const FPBDJointSettings& JointSettings)
+	{
+		int32 Index = RowDatas.AddDefaulted();
+		FJointSolverConstraintRowData& RowData = RowDatas[Index];
+		AngularRowIndexEnd = Index + 1;
+
+		RowData.JointIndex = JointIndex;
+		RowData.UpdateType = EJointSolverConstraintUpdateType::Angular_SLerpDrive;
+		RowData.NumRows = 1;
+		RowData.Limit = 0.0f;
+
+		RowData.Stiffness = FPBDJointUtilities::GetAngularSLerpDriveStiffness(SolverSettings, JointSettings);
+		RowData.Damping = FPBDJointUtilities::GetAngularSLerpDriveDamping(SolverSettings, JointSettings);
+		RowData.bIsAccelerationMode = FPBDJointUtilities::GetDriveAccelerationMode(SolverSettings, JointSettings);
+		RowData.bIsSoft = true;
+	}
+
+	//
+	//////////////////////////////////////////////////////////////////////////
+	//
+
+	int32 FJointSolver::ApplyPositionConstraints(
+		const FReal Dt,
+		TArray<FJointSolverJointState>& JointStates,
+		const TArray <FJointSolverConstraintRowData>& RowDatas,
+		TArray <FJointSolverConstraintRowState>& RowStates,
+		int32 JointIndexBegin,
+		int32 JointIndexEnd,
+		int32 RowIndexBegin,
+		int32 RowIndexEnd)
+	{
+		int32 NumActiveConstraints = 0;
+
+		// Calculate the delta position and rotation to fix the constraint errors
+		// NOTE: Rows are in Joint order, so we can loop over them directly rather than looping over joints and then their rows.
+		for (int32 RowIndex = RowIndexBegin; RowIndex < RowIndexEnd; /* no op */)
+		{
+			const int32 RowJointIndex = RowDatas[RowIndex].JointIndex;
+			if (RowDatas[RowIndex].NumRows == 3)
+			{
+				FVec3 Error = FVec3(RowStates[RowIndex].Error, RowStates[RowIndex + 1].Error, RowStates[RowIndex + 2].Error);
+				if (Error.Size() > JointStates[RowJointIndex].PositionTolerance)
+				{
+					ApplyPositionConstraint3(Dt, JointStates[RowJointIndex], RowDatas[RowIndex], RowStates[RowIndex], RowStates[RowIndex + 1], RowStates[RowIndex + 2]);
+					JointStates[RowJointIndex].DPs[0] += RowStates[RowIndex].DPs[0];
+					JointStates[RowJointIndex].DPs[1] += RowStates[RowIndex].DPs[1];
+					JointStates[RowJointIndex].DRs[0] += RowStates[RowIndex].DRs[0];
+					JointStates[RowJointIndex].DRs[1] += RowStates[RowIndex].DRs[1];
+					NumActiveConstraints += 3;
+				}
+				RowIndex += 3;
+			}
+			else if (RowDatas[RowIndex].NumRows == 1)
+			{
+				if (FMath::Abs(RowStates[RowIndex].Error) > JointStates[RowJointIndex].PositionTolerance)
+				{
+					ApplyPositionConstraint1(Dt, JointStates[RowJointIndex], RowDatas[RowIndex], RowStates[RowIndex]);
+					JointStates[RowJointIndex].DPs[0] += RowStates[RowIndex].DPs[0];
+					JointStates[RowJointIndex].DPs[1] += RowStates[RowIndex].DPs[1];
+					JointStates[RowJointIndex].DRs[0] += RowStates[RowIndex].DRs[0];
+					JointStates[RowJointIndex].DRs[1] += RowStates[RowIndex].DRs[1];
+					++NumActiveConstraints;
+				}
+				++RowIndex;
+			}
+		}
+
+		// Accumulate and apply the deltas for each joint
+		for (int32 JointIndex = JointIndexBegin; JointIndex < JointIndexEnd; ++JointIndex)
+		{
+			JointStates[JointIndex].ApplyDelta(JointStates[JointIndex].DPs[0], JointStates[JointIndex].DRs[0], JointStates[JointIndex].DPs[1], JointStates[JointIndex].DRs[1]);
+		}
+
+		return NumActiveConstraints;
+	}
+
+
+	int32 FJointSolver::ApplyRotationConstraints(
+		const FReal Dt,
+		TArray<FJointSolverJointState>& JointStates,
+		const TArray <FJointSolverConstraintRowData>& RowDatas,
+		TArray <FJointSolverConstraintRowState>& RowStates,
+		int32 JointIndexBegin,
+		int32 JointIndexEnd,
+		int32 RowIndexBegin,
+		int32 RowIndexEnd)
+	{
+		int32 NumActiveConstraints = 0;
+
+		// Calculate the delta position and rotation to fix the constraint errors
+		// NOTE: Rows are in Joint order, so we can loop over them directly rather than looping over joints and then their rows.
+		for (int32 RowIndex = RowIndexBegin; RowIndex < RowIndexEnd; ++RowIndex)
+		{
+			const int32 RowJointIndex = RowDatas[RowIndex].JointIndex;
+
+			if (FMath::Abs(RowStates[RowIndex].Error) > JointStates[RowJointIndex].AngleTolerance)
+			{
+				ApplyRotationConstraint(Dt, JointStates[RowJointIndex], RowDatas[RowIndex], RowStates[RowIndex]);
+
+				JointStates[RowJointIndex].DRs[0] += RowStates[RowIndex].DRs[0];
+				JointStates[RowJointIndex].DRs[1] += RowStates[RowIndex].DRs[1];
+				++NumActiveConstraints;
+			}
+		}
+
+		// Accumulate and apply the deltas for each joint
+		for (int32 JointIndex = JointIndexBegin; JointIndex < JointIndexEnd; ++JointIndex)
+		{
+			JointStates[JointIndex].ApplyRotationDelta(JointStates[JointIndex].DRs[0], JointStates[JointIndex].DRs[1]);
+		}
+
+		return NumActiveConstraints;
+	}
+}
\ No newline at end of file
diff --git a/Engine/Source/Runtime/Experimental/Chaos/Private/Chaos/Joint/PBDJointSolverGaussSeidel.cpp b/Engine/Source/Runtime/Experimental/Chaos/Private/Chaos/Joint/PBDJointSolverGaussSeidel.cpp
--- a/Engine/Source/Runtime/Experimental/Chaos/Private/Chaos/Joint/PBDJointSolverGaussSeidel.cpp
+++ b/Engine/Source/Runtime/Experimental/Chaos/Private/Chaos/Joint/PBDJointSolverGaussSeidel.cpp
@@ -1,6 +1,7 @@
 // Copyright Epic Games, Inc. All Rights Reserved.
 #include "Chaos/Joint/PBDJointSolverGaussSeidel.h"
 #include "Chaos/Joint/ChaosJointLog.h"
+#include "Chaos/Joint/JointSolverConstraints.h"
 #include "Chaos/ParticleHandle.h"
 #include "Chaos/PBDJointConstraintUtilities.h"
 #include "Chaos/Utilities.h"
@@ -25,94 +26,6 @@ FAutoConsoleVariableRef CVarChaosJointDegenerateRotationLimit(TEXT("p.Chaos.Join
 
 namespace Chaos
 {
-	void GetSphericalAxisDelta(
-		const FVec3& X0,
-		const FVec3& X1,
-		FVec3& Axis,
-		FReal& Delta)
-	{
-		Axis = FVec3(0);
-		Delta = 0;
-
-		const FVec3 DX = X1 - X0;
-		const FReal DXLen = DX.Size();
-		if (DXLen > KINDA_SMALL_NUMBER)
-		{
-			Axis = DX / DXLen;
-			Delta = DXLen;
-		}
-	}
-
-
-	void GetCylindricalAxesDeltas(
-		const FRotation3& R0,
-		const FVec3& X0,
-		const FVec3& X1,
-		const int32 CylinderAxisIndex,
-		FVec3& CylinderAxis,
-		FReal& CylinderDelta,
-		FVec3& RadialAxis,
-		FReal& RadialDelta)
-	{
-		CylinderAxis = FVec3(0);
-		CylinderDelta = 0;
-		RadialAxis = FVec3(0);
-		RadialDelta = 0;
-
-		FVec3 DX = X1 - X0;
-
-		const FMatrix33 R0M = R0.ToMatrix();
-		CylinderAxis = R0M.GetAxis(CylinderAxisIndex);
-		CylinderDelta = FVec3::DotProduct(DX, CylinderAxis);
-
-		DX = DX - CylinderDelta * CylinderAxis;
-		const FReal DXLen = DX.Size();
-		if (DXLen > KINDA_SMALL_NUMBER)
-		{
-			RadialAxis = DX / DXLen;
-			RadialDelta = DXLen;
-		}
-	}
-
-
-	void GetPlanarAxisDelta(
-		const FRotation3& R0,
-		const FVec3& X0,
-		const FVec3& X1,
-		const int32 PlaneAxisIndex,
-		FVec3& Axis,
-		FReal& Delta)
-	{
-		const FMatrix33 R0M = R0.ToMatrix();
-		Axis = R0M.GetAxis(PlaneAxisIndex);
-		Delta = FVec3::DotProduct(X1 - X0, Axis);
-	}
-
-
-	FReal GetConeAngleLimit(
-		const FPBDJointSettings& JointSettings,
-		const FVec3& SwingAxisLocal,
-		const FReal SwingAngle)
-	{
-		// Calculate swing limit for the current swing axis
-		const FReal Swing1Limit = JointSettings.AngularLimits[(int32)EJointAngularConstraintIndex::Swing1];
-		const FReal Swing2Limit = JointSettings.AngularLimits[(int32)EJointAngularConstraintIndex::Swing2];
-
-		// Circular swing limit
-		FReal SwingAngleMax = Swing1Limit;
-
-		// Elliptical swing limit
-		// @todo(ccaulfield): do elliptical constraints properly (axis is still for circular limit)
-		if (!FMath::IsNearlyEqual(Swing1Limit, Swing2Limit, KINDA_SMALL_NUMBER))
-		{
-			// Map swing axis to ellipse and calculate limit for this swing axis
-			const FReal DotSwing1 = FMath::Abs(FVec3::DotProduct(SwingAxisLocal, FJointConstants::Swing1Axis()));
-			const FReal DotSwing2 = FMath::Abs(FVec3::DotProduct(SwingAxisLocal, FJointConstants::Swing2Axis()));
-			SwingAngleMax = FMath::Sqrt(Swing1Limit * DotSwing1 * Swing1Limit * DotSwing1 + Swing2Limit * DotSwing2 * Swing2Limit * DotSwing2);
-		}
-
-		return SwingAngleMax;
-	}
 
 
 	//
@@ -132,20 +45,47 @@ namespace Chaos
 	}
 
 
+	void FJointSolverGaussSeidel::InitDerivedState()
+	{
+		// Really we should need to do this for Kinematics since Dynamics are updated each iteration
+		Xs[0] = PrevPs[0] + PrevQs[0] * XLs[0].GetTranslation();
+		Rs[0] = PrevQs[0] * XLs[0].GetRotation();
+		InvIs[0] = (InvMs[0] > 0.0f) ? Utilities::ComputeWorldSpaceInertia(PrevQs[0], InvILs[0]) : FMatrix33(0, 0, 0);
+
+		Xs[1] = PrevPs[1] + PrevQs[1] * XLs[1].GetTranslation();
+		Rs[1] = PrevQs[1] * XLs[1].GetRotation();
+		InvIs[1] = (InvMs[1] > 0.0f) ? Utilities::ComputeWorldSpaceInertia(PrevQs[1], InvILs[1]) : FMatrix33(0, 0, 0);
+
+		Rs[1].EnforceShortestArcWith(Rs[0]);
+	}
+
+
 	void FJointSolverGaussSeidel::UpdateDerivedState()
 	{
-		Xs[0] = Ps[0] + Qs[0] * XLs[0].GetTranslation();
-		Xs[1] = Ps[1] + Qs[1] * XLs[1].GetTranslation();
-		Rs[0] = Qs[0] * XLs[0].GetRotation();
-		Rs[1] = Qs[1] * XLs[1].GetRotation();
+		// Kinematic bodies will not be moved, so we don't update derived state during iterations
+		if (InvMs[0] > 0.0f)
+		{
+			Xs[0] = Ps[0] + Qs[0] * XLs[0].GetTranslation();
+			Rs[0] = Qs[0] * XLs[0].GetRotation();
+			InvIs[0] = Utilities::ComputeWorldSpaceInertia(Qs[0], InvILs[0]);
+		}
+		if (InvMs[1] > 0.0f)
+		{
+			Xs[1] = Ps[1] + Qs[1] * XLs[1].GetTranslation();
+			Rs[1] = Qs[1] * XLs[1].GetRotation();
+			InvIs[1] = Utilities::ComputeWorldSpaceInertia(Qs[1], InvILs[1]);
+		}
 		Rs[1].EnforceShortestArcWith(Rs[0]);
 	}
 
+
 	void FJointSolverGaussSeidel::UpdateDerivedState(const int32 BodyIndex)
 	{
 		Xs[BodyIndex] = Ps[BodyIndex] + Qs[BodyIndex] * XLs[BodyIndex].GetTranslation();
 		Rs[BodyIndex] = Qs[BodyIndex] * XLs[BodyIndex].GetRotation();
 		Rs[1].EnforceShortestArcWith(Rs[0]);
+	
+		InvIs[BodyIndex] = Utilities::ComputeWorldSpaceInertia(Qs[BodyIndex], InvILs[BodyIndex]);
 	}
 
 
@@ -188,9 +128,10 @@ namespace Chaos
 		TwistDriveLambda = (FReal)0;
 		SwingDriveLambda = (FReal)0;
 
-		AngularPositionCorrection = FPBDJointUtilities::GetAngularPositionCorrection(SolverSettings, JointSettings);
 		PositionTolerance = SolverSettings.PositionTolerance;
 		AngleTolerance = SolverSettings.AngleTolerance;
+
+		InitDerivedState();
 	}
 
 
@@ -220,12 +161,12 @@ namespace Chaos
 	}
 
 
-	FJointSolverResult FJointSolverGaussSeidel::ApplyConstraints(
+	int32 FJointSolverGaussSeidel::ApplyConstraints(
 		const FReal Dt,
 		const FPBDJointSolverSettings& SolverSettings,
 		const FPBDJointSettings& JointSettings)
 	{
-		FJointSolverResult NetResult;
+		int32 NumActive = 0;
 
 		bool bHasPositionConstraints =
 			(JointSettings.LinearMotionTypes[0] != EJointMotionType::Free)
@@ -239,24 +180,24 @@ namespace Chaos
 
 		if (bHasPositionConstraints)
 		{
-			NetResult += ApplyPositionConstraints(Dt, SolverSettings, JointSettings);
+			NumActive += ApplyPositionConstraints(Dt, SolverSettings, JointSettings);
 		}
 
 		if (bHasRotationConstraints)
 		{
-			NetResult += ApplyRotationConstraints(Dt, SolverSettings, JointSettings);
+			NumActive += ApplyRotationConstraints(Dt, SolverSettings, JointSettings);
 		}
 
-		return NetResult;
+		return NumActive;
 	}
 
 
-	FJointSolverResult FJointSolverGaussSeidel::ApplyDrives(
+	int32 FJointSolverGaussSeidel::ApplyDrives(
 		const FReal Dt,
 		const FPBDJointSolverSettings& SolverSettings,
 		const FPBDJointSettings& JointSettings)
 	{
-		FJointSolverResult NetResult;
+		int32 NumActive = 0;
 
 		bool bHasPositionDrives = 
 			JointSettings.bLinearPositionDriveEnabled[0]
@@ -276,40 +217,40 @@ namespace Chaos
 
 		if (bHasPositionDrives)
 		{
-			NetResult += ApplyPositionDrives(Dt, SolverSettings, JointSettings);
+			NumActive += ApplyPositionDrives(Dt, SolverSettings, JointSettings);
 		}
 
 		if (bHasRotationDrives)
 		{
-			NetResult += ApplyRotationDrives(Dt, SolverSettings, JointSettings);
+			NumActive += ApplyRotationDrives(Dt, SolverSettings, JointSettings);
 		}
 
-		return NetResult;
+		return NumActive;
 	}
 
 
-	FJointSolverResult FJointSolverGaussSeidel::ApplyProjections(
+	int32 FJointSolverGaussSeidel::ApplyProjections(
 		const FReal Dt,
 		const FPBDJointSolverSettings& SolverSettings,
 		const FPBDJointSettings& JointSettings)
 	{
-		FJointSolverResult NetResult;
+		int32 NumActive = 0;
 
-		//NetResult += ApplyRotationProjection(Dt, SolverSettings, JointSettings);
+		//NumActive += ApplyRotationProjection(Dt, SolverSettings, JointSettings);
 
-		NetResult += ApplyPositionProjection(Dt, SolverSettings, JointSettings);
+		NumActive += ApplyPositionProjection(Dt, SolverSettings, JointSettings);
 
-		return NetResult;
+		return NumActive;
 	}
 
 
-	FJointSolverResult FJointSolverGaussSeidel::ApplyRotationConstraints(
+	int32 FJointSolverGaussSeidel::ApplyRotationConstraints(
 		const FReal Dt,
 		const FPBDJointSolverSettings& SolverSettings,
 		const FPBDJointSettings& JointSettings)
 	{
 		// Locked axes always use hard constraints. Limited axes use hard or soft depending on settings
-		FJointSolverResult NetResult;
+		int32 NumActive = 0;
 
 		EJointMotionType TwistMotion = JointSettings.AngularMotionTypes[(int32)EJointAngularConstraintIndex::Twist];
 		EJointMotionType Swing1Motion = JointSettings.AngularMotionTypes[(int32)EJointAngularConstraintIndex::Swing1];
@@ -333,7 +274,7 @@ namespace Chaos
 		{
 			if (TwistMotion == EJointMotionType::Limited)
 			{
-				NetResult += ApplyTwistConstraint(Dt, SolverSettings, JointSettings, bTwistSoft);
+				NumActive += ApplyTwistConstraint(Dt, SolverSettings, JointSettings, bTwistSoft);
 			}
 			else if (TwistMotion == EJointMotionType::Locked)
 			{
@@ -350,29 +291,29 @@ namespace Chaos
 		{
 			if ((Swing1Motion == EJointMotionType::Limited) && (Swing2Motion == EJointMotionType::Limited))
 			{
-				NetResult += ApplyConeConstraint(Dt, SolverSettings, JointSettings, bSwingSoft);
+				NumActive += ApplyConeConstraint(Dt, SolverSettings, JointSettings, bSwingSoft);
 			}
 			else if ((Swing1Motion == EJointMotionType::Limited) && (Swing2Motion == EJointMotionType::Locked))
 			{
-				NetResult += ApplySingleLockedSwingConstraint(Dt, SolverSettings, JointSettings, EJointAngularConstraintIndex::Swing2, false);
+				NumActive += ApplySingleLockedSwingConstraint(Dt, SolverSettings, JointSettings, EJointAngularConstraintIndex::Swing2, false);
 				if (!bDegenerate)
 				{
-					NetResult += ApplySwingConstraint(Dt, SolverSettings, JointSettings, EJointAngularConstraintIndex::Swing1, bSwingSoft);
+					NumActive += ApplySwingConstraint(Dt, SolverSettings, JointSettings, EJointAngularConstraintIndex::Swing1, bSwingSoft);
 				}
 			}
 			else if ((Swing1Motion == EJointMotionType::Limited) && (Swing2Motion == EJointMotionType::Free))
 			{
 				if (!bDegenerate)
 				{
-					NetResult += ApplyDualConeSwingConstraint(Dt, SolverSettings, JointSettings, EJointAngularConstraintIndex::Swing1, bSwingSoft);
+					NumActive += ApplyDualConeSwingConstraint(Dt, SolverSettings, JointSettings, EJointAngularConstraintIndex::Swing1, bSwingSoft);
 				}
 			}
 			else if ((Swing1Motion == EJointMotionType::Locked) && (Swing2Motion == EJointMotionType::Limited))
 			{
-				NetResult += ApplySingleLockedSwingConstraint(Dt, SolverSettings, JointSettings, EJointAngularConstraintIndex::Swing1, false);
+				NumActive += ApplySingleLockedSwingConstraint(Dt, SolverSettings, JointSettings, EJointAngularConstraintIndex::Swing1, false);
 				if (!bDegenerate)
 				{
-					NetResult += ApplySwingConstraint(Dt, SolverSettings, JointSettings, EJointAngularConstraintIndex::Swing2, bSwingSoft);
+					NumActive += ApplySwingConstraint(Dt, SolverSettings, JointSettings, EJointAngularConstraintIndex::Swing2, bSwingSoft);
 				}
 			}
 			else if ((Swing1Motion == EJointMotionType::Locked) && (Swing2Motion == EJointMotionType::Locked))
@@ -381,18 +322,18 @@ namespace Chaos
 			}
 			else if ((Swing1Motion == EJointMotionType::Locked) && (Swing2Motion == EJointMotionType::Free))
 			{
-				NetResult += ApplySingleLockedSwingConstraint(Dt, SolverSettings, JointSettings, EJointAngularConstraintIndex::Swing1, false);
+				NumActive += ApplySingleLockedSwingConstraint(Dt, SolverSettings, JointSettings, EJointAngularConstraintIndex::Swing1, false);
 			}
 			else if ((Swing1Motion == EJointMotionType::Free) && (Swing2Motion == EJointMotionType::Limited))
 			{
 				if (!bDegenerate)
 				{
-					NetResult += ApplyDualConeSwingConstraint(Dt, SolverSettings, JointSettings, EJointAngularConstraintIndex::Swing2, bSwingSoft);
+					NumActive += ApplyDualConeSwingConstraint(Dt, SolverSettings, JointSettings, EJointAngularConstraintIndex::Swing2, bSwingSoft);
 				}
 			}
 			else if ((Swing1Motion == EJointMotionType::Free) && (Swing2Motion == EJointMotionType::Locked))
 			{
-				NetResult += ApplySingleLockedSwingConstraint(Dt, SolverSettings, JointSettings, EJointAngularConstraintIndex::Swing2, false);
+				NumActive += ApplySingleLockedSwingConstraint(Dt, SolverSettings, JointSettings, EJointAngularConstraintIndex::Swing2, false);
 			}
 			else if ((Swing1Motion == EJointMotionType::Free) && (Swing2Motion == EJointMotionType::Free))
 			{
@@ -407,19 +348,19 @@ namespace Chaos
 			&& (JointSettings.AngularMotionTypes[(int32)EJointAngularConstraintIndex::Swing2] == EJointMotionType::Locked);
 		if (bLockedTwist || bLockedSwing)
 		{
-			NetResult += ApplyLockedRotationConstraints(Dt, SolverSettings, JointSettings, bLockedTwist, bLockedSwing);
+			NumActive += ApplyLockedRotationConstraints(Dt, SolverSettings, JointSettings, bLockedTwist, bLockedSwing);
 		}
 
-		return NetResult;
+		return NumActive;
 	}
 
 
-	FJointSolverResult FJointSolverGaussSeidel::ApplyRotationDrives(
+	int32 FJointSolverGaussSeidel::ApplyRotationDrives(
 		const FReal Dt,
 		const FPBDJointSolverSettings& SolverSettings,
 		const FPBDJointSettings& JointSettings)
 	{
-		FJointSolverResult NetResult;
+		int32 NumActive = 0;
 
 		EJointMotionType TwistMotion = JointSettings.AngularMotionTypes[(int32)EJointAngularConstraintIndex::Twist];
 		EJointMotionType Swing1Motion = JointSettings.AngularMotionTypes[(int32)EJointAngularConstraintIndex::Swing1];
@@ -434,41 +375,41 @@ namespace Chaos
 			// No SLerp drive if we have a locked rotation (it will be grayed out in the editor in this case, but could still have been set before the rotation was locked)
 			if ((JointSettings.bAngularSLerpPositionDriveEnabled || JointSettings.bAngularSLerpVelocityDriveEnabled) && !bTwistLocked && !bSwing1Locked && !bSwing2Locked)
 			{
-				NetResult += ApplySLerpDrive(Dt, SolverSettings, JointSettings);
+				NumActive += ApplySLerpDrive(Dt, SolverSettings, JointSettings);
 			}
 			else
 			{
 				if ((JointSettings.bAngularTwistPositionDriveEnabled || JointSettings.bAngularTwistVelocityDriveEnabled) && !bTwistLocked)
 				{
-					NetResult += ApplyTwistDrive(Dt, SolverSettings, JointSettings);
+					NumActive += ApplyTwistDrive(Dt, SolverSettings, JointSettings);
 				}
 
 				const bool bSwingDriveEnabled = (JointSettings.bAngularSwingPositionDriveEnabled || JointSettings.bAngularSwingVelocityDriveEnabled);
 				if (bSwingDriveEnabled && !bSwing1Locked && !bSwing2Locked)
 				{
-					NetResult += ApplyConeDrive(Dt, SolverSettings, JointSettings);
+					NumActive += ApplyConeDrive(Dt, SolverSettings, JointSettings);
 				}
 				else if (bSwingDriveEnabled && !bSwing1Locked)
 				{
-					NetResult += ApplySwingDrive(Dt, SolverSettings, JointSettings, EJointAngularConstraintIndex::Swing1);
+					NumActive += ApplySwingDrive(Dt, SolverSettings, JointSettings, EJointAngularConstraintIndex::Swing1);
 				}
 				else if (bSwingDriveEnabled && !bSwing2Locked)
 				{
-					NetResult += ApplySwingDrive(Dt, SolverSettings, JointSettings, EJointAngularConstraintIndex::Swing2);
+					NumActive += ApplySwingDrive(Dt, SolverSettings, JointSettings, EJointAngularConstraintIndex::Swing2);
 				}
 			}
 		}
 
-		return NetResult;
+		return NumActive;
 	}
 
 
-	FJointSolverResult FJointSolverGaussSeidel::ApplyRotationProjection(
+	int32 FJointSolverGaussSeidel::ApplyRotationProjection(
 		const FReal Dt,
 		const FPBDJointSolverSettings& SolverSettings,
 		const FPBDJointSettings& JointSettings)
 	{
-		FJointSolverResult NetResult;
+		int32 NumActive = 0;
 
 		EJointMotionType TwistMotion = JointSettings.AngularMotionTypes[(int32)EJointAngularConstraintIndex::Twist];
 		EJointMotionType Swing1Motion = JointSettings.AngularMotionTypes[(int32)EJointAngularConstraintIndex::Swing1];
@@ -480,7 +421,7 @@ namespace Chaos
 		{
 			if ((TwistMotion == EJointMotionType::Locked) || ((TwistMotion == EJointMotionType::Limited) && !bTwistSoft))
 			{
-				NetResult += ApplyTwistProjection(Dt, SolverSettings, JointSettings);
+				NumActive += ApplyTwistProjection(Dt, SolverSettings, JointSettings);
 			}
 		}
 
@@ -488,32 +429,32 @@ namespace Chaos
 		{
 			if ((Swing1Motion == EJointMotionType::Limited) && (Swing2Motion == EJointMotionType::Limited) && !bSwingSoft)
 			{
-				NetResult += ApplyConeProjection(Dt, SolverSettings, JointSettings);
+				NumActive += ApplyConeProjection(Dt, SolverSettings, JointSettings);
 			}
 			else
 			{
 				if ((Swing1Motion == EJointMotionType::Locked) || ((Swing1Motion == EJointMotionType::Limited) && !bSwingSoft))
 				{
-					NetResult += ApplySwingProjection(Dt, SolverSettings, JointSettings, EJointAngularConstraintIndex::Swing1);
+					NumActive += ApplySwingProjection(Dt, SolverSettings, JointSettings, EJointAngularConstraintIndex::Swing1);
 				}
 				if ((Swing2Motion == EJointMotionType::Locked) || ((Swing2Motion == EJointMotionType::Limited) && !bSwingSoft))
 				{
-					NetResult += ApplySwingProjection(Dt, SolverSettings, JointSettings, EJointAngularConstraintIndex::Swing2);
+					NumActive += ApplySwingProjection(Dt, SolverSettings, JointSettings, EJointAngularConstraintIndex::Swing2);
 				}
 			}
 		}
 
-		return NetResult;
+		return NumActive;
 	}
 
 
-	FJointSolverResult FJointSolverGaussSeidel::ApplyPositionConstraints(
+	int32 FJointSolverGaussSeidel::ApplyPositionConstraints(
 		const FReal Dt,
 		const FPBDJointSolverSettings& SolverSettings,
 		const FPBDJointSettings& JointSettings)
 	{
 		// @todo(ccaulfield): the branch logic is all constant for a joint - move it to initialization and turn it into a table or something (also for Hard/Soft login in apply functions)
-		FJointSolverResult NetResult;
+		int32 NumActive = 0;
 
 		const TVector<EJointMotionType, 3>& LinearMotion = JointSettings.LinearMotionTypes;
 		const TVector<bool, 3> bLinearLocked =
@@ -534,73 +475,73 @@ namespace Chaos
 			// Hard point constraint (most common case)
 			if (InvMs[0] == 0)
 			{
-				NetResult += ApplyPointPositionConstraintKD(0, 1, Dt, SolverSettings, JointSettings);
+				NumActive += ApplyPointPositionConstraintKD(0, 1, Dt, SolverSettings, JointSettings);
 			}
 			else if (InvMs[1] == 0)
 			{
-				NetResult += ApplyPointPositionConstraintKD(1, 0, Dt, SolverSettings, JointSettings);
+				NumActive += ApplyPointPositionConstraintKD(1, 0, Dt, SolverSettings, JointSettings);
 			}
 			else
 			{
-				NetResult += ApplyPointPositionConstraintDD(Dt, SolverSettings, JointSettings);
+				NumActive += ApplyPointPositionConstraintDD(Dt, SolverSettings, JointSettings);
 			}
 		}
 		else if (bLinearLimted[0] && bLinearLimted[1] && bLinearLimted[2])
 		{
 			// Spherical constraint
-			NetResult += ApplySphericalPositionConstraint(Dt, SolverSettings, JointSettings);
+			NumActive += ApplySphericalPositionConstraint(Dt, SolverSettings, JointSettings);
 		}
 		else if (bLinearLocked[1] && bLinearLocked[2] && !bLinearLocked[0])
 		{
 			// Line constraint along X axis
-			NetResult += ApplyCylindricalPositionConstraint(Dt, 0, LinearMotion[0], EJointMotionType::Locked, SolverSettings, JointSettings);
+			NumActive += ApplyCylindricalPositionConstraint(Dt, 0, LinearMotion[0], EJointMotionType::Locked, SolverSettings, JointSettings);
 		}
 		else if (bLinearLocked[0] && bLinearLocked[2] && !bLinearLocked[1])
 		{
 			// Line constraint along Y axis
-			NetResult += ApplyCylindricalPositionConstraint(Dt, 1, LinearMotion[1], EJointMotionType::Locked, SolverSettings, JointSettings);
+			NumActive += ApplyCylindricalPositionConstraint(Dt, 1, LinearMotion[1], EJointMotionType::Locked, SolverSettings, JointSettings);
 		}
 		else if (bLinearLocked[0] && bLinearLocked[1] && !bLinearLocked[2])
 		{
 			// Line constraint along Z axis
-			NetResult += ApplyCylindricalPositionConstraint(Dt, 2, LinearMotion[2], EJointMotionType::Locked, SolverSettings, JointSettings);
+			NumActive += ApplyCylindricalPositionConstraint(Dt, 2, LinearMotion[2], EJointMotionType::Locked, SolverSettings, JointSettings);
 		}
 		else if (bLinearLimted[1] && bLinearLimted[2] && !bLinearLimted[0])
 		{
 			// Cylindrical constraint along X axis
-			NetResult += ApplyCylindricalPositionConstraint(Dt, 0, LinearMotion[0], EJointMotionType::Limited, SolverSettings, JointSettings);
+			NumActive += ApplyCylindricalPositionConstraint(Dt, 0, LinearMotion[0], EJointMotionType::Limited, SolverSettings, JointSettings);
 		}
 		else if (bLinearLimted[0] && bLinearLimted[2] && !bLinearLimted[1])
 		{
 			// Cylindrical constraint along Y axis
-			NetResult += ApplyCylindricalPositionConstraint(Dt, 1, LinearMotion[1], EJointMotionType::Limited, SolverSettings, JointSettings);
+			NumActive += ApplyCylindricalPositionConstraint(Dt, 1, LinearMotion[1], EJointMotionType::Limited, SolverSettings, JointSettings);
 		}
 		else if (bLinearLimted[0] && bLinearLimted[1] && !bLinearLimted[2])
 		{
 			// Cylindrical constraint along Z axis
-			NetResult += ApplyCylindricalPositionConstraint(Dt, 2, LinearMotion[2], EJointMotionType::Limited, SolverSettings, JointSettings);
+			NumActive += ApplyCylindricalPositionConstraint(Dt, 2, LinearMotion[2], EJointMotionType::Limited, SolverSettings, JointSettings);
 		}
 		else if (bLinearLocked[0] || bLinearLimted[0])
 		{
 			// Planar constraint along X axis
-			NetResult += ApplyPlanarPositionConstraint(Dt, 0, LinearMotion[0], SolverSettings, JointSettings);
+			NumActive += ApplyPlanarPositionConstraint(Dt, 0, LinearMotion[0], SolverSettings, JointSettings);
 		}
 		else if (bLinearLocked[1] || bLinearLimted[1])
 		{
 			// Planar constraint along Y axis
-			NetResult += ApplyPlanarPositionConstraint(Dt, 1, LinearMotion[1], SolverSettings, JointSettings);
+			NumActive += ApplyPlanarPositionConstraint(Dt, 1, LinearMotion[1], SolverSettings, JointSettings);
 		}
 		else if (bLinearLocked[2] || bLinearLimted[2])
 		{
 			// Planar constraint along Z axis
-			NetResult += ApplyPlanarPositionConstraint(Dt, 2, LinearMotion[2], SolverSettings, JointSettings);
+			NumActive += ApplyPlanarPositionConstraint(Dt, 2, LinearMotion[2], SolverSettings, JointSettings);
 		}
 
-		return NetResult;
+		return NumActive;
 	}
 
 
-	FJointSolverResult FJointSolverGaussSeidel::ApplyPositionDrives(
+	int32 FJointSolverGaussSeidel::ApplyPositionDrives(
 		const FReal Dt,
 		const FPBDJointSolverSettings& SolverSettings,
 		const FPBDJointSettings& JointSettings)
@@ -609,7 +550,7 @@ namespace Chaos
 		// E.g., if X/Y have position drives and Z is a velocity drive, it should be a circular position-drive and axial velocity-drive rather than a spherical one.
 		// E.g., if X/Y/Z have position drives and only 1 or 2 axes have velocity drives, we need to apply 2 drives. Etc etc.
 		// Basically we need to split the axes by those that have the same drive settings.
-		FJointSolverResult NetResult;
+		int32 NumActive = 0;
 
 		if (SolverSettings.bEnableDrives)
 		{
@@ -622,35 +563,35 @@ namespace Chaos
 
 			if (bDriven[0] && bDriven[1] && bDriven[2])
 			{
-				NetResult += ApplySphericalPositionDrive(Dt, SolverSettings, JointSettings);
+				NumActive += ApplySphericalPositionDrive(Dt, SolverSettings, JointSettings);
 			}
 			else if (bDriven[1] && bDriven[2])
 			{
-				NetResult += ApplyCircularPositionDrive(Dt, 0, SolverSettings, JointSettings);
+				NumActive += ApplyCircularPositionDrive(Dt, 0, SolverSettings, JointSettings);
 			}
 			else if (bDriven[0] && bDriven[2])
 			{
-				NetResult += ApplyCircularPositionDrive(Dt, 1, SolverSettings, JointSettings);
+				NumActive += ApplyCircularPositionDrive(Dt, 1, SolverSettings, JointSettings);
 			}
 			else if (bDriven[0] && bDriven[1])
 			{
-				NetResult += ApplyCircularPositionDrive(Dt, 2, SolverSettings, JointSettings);
+				NumActive += ApplyCircularPositionDrive(Dt, 2, SolverSettings, JointSettings);
 			}
 			else if (bDriven[0])
 			{
-				NetResult += ApplyAxialPositionDrive(Dt, 0, SolverSettings, JointSettings);
+				NumActive += ApplyAxialPositionDrive(Dt, 0, SolverSettings, JointSettings);
 			}
 			else if (bDriven[1])
 			{
-				NetResult += ApplyAxialPositionDrive(Dt, 1, SolverSettings, JointSettings);
+				NumActive += ApplyAxialPositionDrive(Dt, 1, SolverSettings, JointSettings);
 			}
 			else if (bDriven[2])
 			{
-				NetResult += ApplyAxialPositionDrive(Dt, 2, SolverSettings, JointSettings);
+				NumActive += ApplyAxialPositionDrive(Dt, 2, SolverSettings, JointSettings);
 			}
 		}
 
-		return NetResult;
+		return NumActive;
 	}
 
 
@@ -695,6 +636,8 @@ namespace Chaos
 		const FRotation3 DQ = (FRotation3::FromElements(Stiffness * DR, 0) * Qs[BodyIndex]) * (FReal)0.5;
 		Qs[BodyIndex] = (Qs[BodyIndex] + DQ).GetNormalized();
 		Qs[1].EnforceShortestArcWith(Qs[0]);
+
+		UpdateDerivedState(BodyIndex);
 	}
 
 
@@ -706,11 +649,28 @@ namespace Chaos
 		//UE_LOG(LogChaosJoint, VeryVerbose, TEXT("      Apply DR%d %f %f %f"), 0, DR0.X, DR0.Y, DR0.Z);
 		//UE_LOG(LogChaosJoint, VeryVerbose, TEXT("      Apply DR%d %f %f %f"), 1, DR1.X, DR1.Y, DR1.Z);
 
-		const FRotation3 DQ0 = (FRotation3::FromElements(Stiffness * DR0, 0) * Qs[0]) * (FReal)0.5;
-		const FRotation3 DQ1 = (FRotation3::FromElements(Stiffness * DR1, 0) * Qs[1]) * (FReal)0.5;
-		Qs[0] = (Qs[0] + DQ0).GetNormalized();
-		Qs[1] = (Qs[1] + DQ1).GetNormalized();
-		Qs[1].EnforceShortestArcWith(Qs[0]);
+		if (bChaos_Joint_ISPC_Enabled)
+		{
+#if INTEL_ISPC
+			ispc::ApplyRotationDelta2((ispc::FJointSolverGaussSeidel*)this, Stiffness, (ispc::FVector&)DR0, (ispc::FVector&)DR1);
+#endif
+		}
+		else
+		{
+			if (InvMs[0] > 0.0f)
+			{
+				const FRotation3 DQ0 = (FRotation3::FromElements(Stiffness * DR0, 0) * Qs[0]) * (FReal)0.5;
+				Qs[0] = (Qs[0] + DQ0).GetNormalized();
+			}
+			if (InvMs[1] > 0.0f)
+			{
+				const FRotation3 DQ1 = (FRotation3::FromElements(Stiffness * DR1, 0) * Qs[1]) * (FReal)0.5;
+				Qs[1] = (Qs[1] + DQ1).GetNormalized();
+			}
+			Qs[1].EnforceShortestArcWith(Qs[0]);
+
+			UpdateDerivedState();
+		}
 	}
 
 	void FJointSolverGaussSeidel::ApplyDelta(
@@ -726,6 +686,8 @@ namespace Chaos
 		const FRotation3 DQ = (FRotation3::FromElements(Stiffness * DR, 0) * Qs[BodyIndex]) * (FReal)0.5;
 		Qs[BodyIndex] = (Qs[BodyIndex] + DQ).GetNormalized();
 		Qs[1].EnforceShortestArcWith(Qs[0]);
+
+		UpdateDerivedState(BodyIndex);
 	}
 
 
@@ -781,7 +743,6 @@ namespace Chaos
 
 		ApplyPositionDelta(Stiffness, DP0, DP1);
 		ApplyRotationDelta(Stiffness, DR0, DR1);
-		UpdateDerivedState();
 	}
 
 
@@ -837,7 +798,6 @@ namespace Chaos
 			Lambda += DLambda;
 			ApplyPositionDelta((FReal)1, DP0, DP1);
 			ApplyRotationDelta((FReal)1, DR0, DR1);
-			UpdateDerivedState();
 		}
 	}
 	
@@ -863,7 +823,6 @@ namespace Chaos
 		//const FVec3 DR1 = Axis * -(Angle * II1 / (II0 + II1));
 
 		ApplyRotationDelta(Stiffness, DR0, DR1);
-		UpdateDerivedState();
 	}
 
 
@@ -890,8 +849,7 @@ namespace Chaos
 		else
 		{
 			// World-space inverse mass
-			const FMatrix33 InvI1 = Utilities::ComputeWorldSpaceInertia(Qs[DIndex], InvILs[DIndex]);
-			const FVec3 IA1 = Utilities::Multiply(InvI1, Axis);
+			const FVec3 IA1 = Utilities::Multiply(InvIs[1], Axis);
 
 			// Joint-space inverse mass
 			FReal II1 = FVec3::DotProduct(Axis, IA1);
@@ -916,7 +874,6 @@ namespace Chaos
 
 			Lambda += DLambda;
 			ApplyRotationDelta(DIndex, 1.0f, DR1);
-			UpdateDerivedState(DIndex);
 		}
 	}
 
@@ -942,10 +899,8 @@ namespace Chaos
 		else
 		{
 			// World-space inverse mass
-			const FMatrix33 InvI0 = Utilities::ComputeWorldSpaceInertia(Qs[0], InvILs[0]);
-			const FMatrix33 InvI1 = Utilities::ComputeWorldSpaceInertia(Qs[1], InvILs[1]);
-			const FVec3 IA0 = Utilities::Multiply(InvI0, Axis);
-			const FVec3 IA1 = Utilities::Multiply(InvI1, Axis);
+			const FVec3 IA0 = Utilities::Multiply(InvIs[0], Axis);
+			const FVec3 IA1 = Utilities::Multiply(InvIs[1], Axis);
 
 			// Joint-space inverse mass
 			FReal II0 = FVec3::DotProduct(Axis, IA0);
@@ -974,7 +929,6 @@ namespace Chaos
 
 			Lambda += DLambda;
 			ApplyRotationDelta(1.0f, DR0, DR1);
-			UpdateDerivedState();
 		}
 	}
 
@@ -1007,7 +961,7 @@ namespace Chaos
 	//
 	//
 
-	FJointSolverResult FJointSolverGaussSeidel::ApplyLockedRotationConstraints(
+	int32 FJointSolverGaussSeidel::ApplyLockedRotationConstraints(
 		const FReal Dt,
 		const FPBDJointSolverSettings& SolverSettings,
 		const FPBDJointSettings& JointSettings,
@@ -1015,16 +969,16 @@ namespace Chaos
 		const bool bApplySwing)
 	{
 		FVec3 Axis0, Axis1, Axis2;
-		FPBDJointUtilities::GetLockedAxes(Rs[0], Rs[1], Axis0, Axis1, Axis2);
+		FPBDJointUtilities::GetLockedRotationAxes(Rs[0], Rs[1], Axis0, Axis1, Axis2);
 
 		const FRotation3 R01 = Rs[0].Inverse() * Rs[1];
 
-		FJointSolverResult NetResult;
+		int32 NumActive = 0;
 		if (bApplyTwist)
 		{
 			FReal TwistStiffness = FPBDJointUtilities::GetTwistStiffness(SolverSettings, JointSettings);
 			ApplyRotationConstraint(TwistStiffness, Axis0, R01.X);
-			NetResult += FJointSolverResult::MakeActive();
+			++NumActive;
 		}
 
 		if (bApplySwing)
@@ -1032,14 +986,13 @@ namespace Chaos
 			FReal SwingStiffness = FPBDJointUtilities::GetSwingStiffness(SolverSettings, JointSettings);
 			ApplyRotationConstraint(SwingStiffness, Axis1, R01.Y);
 			ApplyRotationConstraint(SwingStiffness, Axis2, R01.Z);
-			NetResult += FJointSolverResult::MakeActive();
-			NetResult += FJointSolverResult::MakeActive();
+			++NumActive;
 		}
 
-		return NetResult;
+		return NumActive;
 	}
 
-	FJointSolverResult FJointSolverGaussSeidel::ApplyTwistConstraint(
+	int32 FJointSolverGaussSeidel::ApplyTwistConstraint(
 		const FReal Dt,
 		const FPBDJointSolverSettings& SolverSettings,
 		const FPBDJointSettings& JointSettings,
@@ -1061,7 +1014,7 @@ namespace Chaos
 			DTwistAngle = TwistAngle + TwistAngleMax;
 		}
 
-		UE_LOG(LogChaosJoint, VeryVerbose, TEXT("    SoftTwist Angle %f [Limit %f]"), FMath::RadiansToDegrees(TwistAngle), FMath::RadiansToDegrees(TwistAngleMax));
+		UE_LOG(LogChaosJoint, VeryVerbose, TEXT("    Twist Angle %f [Limit %f]"), FMath::RadiansToDegrees(TwistAngle), FMath::RadiansToDegrees(TwistAngleMax));
 
 		// Apply twist correction
 		if (FMath::Abs(DTwistAngle) > 0)
@@ -1078,13 +1031,13 @@ namespace Chaos
 				FReal TwistStiffness = FPBDJointUtilities::GetTwistStiffness(SolverSettings, JointSettings);
 				ApplyRotationConstraint(TwistStiffness, TwistAxis, DTwistAngle);
 			}
-			return FJointSolverResult::MakeActive();
+			return 1;
 		}
-		return FJointSolverResult::MakeSolved();
+		return 0;
 	}
 
 
-	FJointSolverResult FJointSolverGaussSeidel::ApplyTwistDrive(
+	int32 FJointSolverGaussSeidel::ApplyTwistDrive(
 		const FReal Dt,
 		const FPBDJointSolverSettings& SolverSettings,
 		const FPBDJointSettings& JointSettings)
@@ -1103,22 +1056,23 @@ namespace Chaos
 			const FReal AngularDriveDamping = FPBDJointUtilities::GetAngularTwistDriveDamping(SolverSettings, JointSettings);
 			const bool bAccelerationMode = FPBDJointUtilities::GetDriveAccelerationMode(SolverSettings, JointSettings);
 			ApplyRotationConstraintSoft(Dt, AngularDriveStiffness, AngularDriveDamping, bAccelerationMode, TwistAxis, DTwistAngle, TwistDriveLambda);
+			return 1;
 		}
-		return FJointSolverResult::MakeSolved();
+		return 0;
 	}
 
 
-	FJointSolverResult FJointSolverGaussSeidel::ApplyTwistProjection(
+	int32 FJointSolverGaussSeidel::ApplyTwistProjection(
 		const FReal Dt,
 		const FPBDJointSolverSettings& SolverSettings,
 		const FPBDJointSettings& JointSettings)
 	{
 		// @todo(ccaulfield): implement twist projection
-		return FJointSolverResult::MakeSolved();
+		return 0;
 	}
 
 
-	FJointSolverResult FJointSolverGaussSeidel::ApplyConeConstraint(
+	int32 FJointSolverGaussSeidel::ApplyConeConstraint(
 		const FReal Dt,
 		const FPBDJointSolverSettings& SolverSettings,
 		const FPBDJointSettings& JointSettings,
@@ -1131,7 +1085,7 @@ namespace Chaos
 		const FVec3 SwingAxis = Rs[0] * SwingAxisLocal;
 
 		// Calculate swing angle error
-		FReal SwingAngleMax = GetConeAngleLimit(JointSettings, SwingAxisLocal, SwingAngle) + AngleTolerance;
+		FReal SwingAngleMax = FPBDJointUtilities::GetConeAngleLimit(JointSettings, SwingAxisLocal, SwingAngle) + AngleTolerance;
 		FReal DSwingAngle = 0.0f;
 		if (SwingAngle > SwingAngleMax)
 		{
@@ -1159,13 +1113,13 @@ namespace Chaos
 				FReal SwingStiffness = FPBDJointUtilities::GetSwingStiffness(SolverSettings, JointSettings);
 				ApplyRotationConstraint(SwingStiffness, SwingAxis, DSwingAngle);
 			}
-			return FJointSolverResult::MakeActive();
+			return 1;
 		}
-		return FJointSolverResult::MakeSolved();
+		return 0;
 	}
 
 
-	FJointSolverResult FJointSolverGaussSeidel::ApplyConeDrive(
+	int32 FJointSolverGaussSeidel::ApplyConeDrive(
 		const FReal Dt,
 		const FPBDJointSolverSettings& SolverSettings,
 		const FPBDJointSettings& JointSettings)
@@ -1190,21 +1144,22 @@ namespace Chaos
 			const FReal AngularDriveDamping = FPBDJointUtilities::GetAngularSwingDriveDamping(SolverSettings, JointSettings);
 			const bool bAccelerationMode = FPBDJointUtilities::GetDriveAccelerationMode(SolverSettings, JointSettings);
 			ApplyRotationConstraintSoft(Dt, AngularDriveStiffness, AngularDriveDamping, bAccelerationMode, SwingAxis, DSwingAngle, SwingDriveLambda);
+			return 1;
 		}
-		return FJointSolverResult::MakeSolved();
+		return 0;
 	}
 
 
-	FJointSolverResult FJointSolverGaussSeidel::ApplyConeProjection(
+	int32 FJointSolverGaussSeidel::ApplyConeProjection(
 		const FReal Dt,
 		const FPBDJointSolverSettings& SolverSettings,
 		const FPBDJointSettings& JointSettings)
 	{
-		return FJointSolverResult::MakeSolved();
+		return 0;
 	}
 
 
-	FJointSolverResult FJointSolverGaussSeidel::ApplySingleLockedSwingConstraint(
+	int32 FJointSolverGaussSeidel::ApplySingleLockedSwingConstraint(
 		const FReal Dt,
 		const FPBDJointSolverSettings& SolverSettings,
 		const FPBDJointSettings& JointSettings,
@@ -1235,13 +1190,13 @@ namespace Chaos
 				const FReal SwingStiffness = FPBDJointUtilities::GetSwingStiffness(SolverSettings, JointSettings);
 				ApplyRotationConstraint(SwingStiffness, SwingAxis, DSwingAngle);
 			}
-			return FJointSolverResult::MakeActive();
+			return 1;
 		}
-		return FJointSolverResult::MakeSolved();
+		return 0;
 	}
 
 
-	FJointSolverResult FJointSolverGaussSeidel::ApplyDualConeSwingConstraint(
+	int32 FJointSolverGaussSeidel::ApplyDualConeSwingConstraint(
 		const FReal Dt,
 		const FPBDJointSolverSettings& SolverSettings,
 		const FPBDJointSettings& JointSettings,
@@ -1281,13 +1236,13 @@ namespace Chaos
 				const FReal SwingStiffness = FPBDJointUtilities::GetSwingStiffness(SolverSettings, JointSettings);
 				ApplyRotationConstraint(SwingStiffness, SwingAxis, DSwingAngle);
 			}
-			return FJointSolverResult::MakeActive();
+			return 1;
 		}
-		return FJointSolverResult::MakeSolved();
+		return 0;
 	}
 
 
-	FJointSolverResult FJointSolverGaussSeidel::ApplySwingConstraint(
+	int32 FJointSolverGaussSeidel::ApplySwingConstraint(
 		const FReal Dt,
 		const FPBDJointSolverSettings& SolverSettings,
 		const FPBDJointSettings& JointSettings,
@@ -1327,13 +1282,13 @@ namespace Chaos
 				const FReal SwingStiffness = FPBDJointUtilities::GetSwingStiffness(SolverSettings, JointSettings);
 				ApplyRotationConstraint(SwingStiffness, SwingAxis, DSwingAngle);
 			}
-			return FJointSolverResult::MakeActive();
+			return 1;
 		}
-		return FJointSolverResult::MakeSolved();
+		return 0;
 	}
 
 
-	FJointSolverResult FJointSolverGaussSeidel::ApplySwingDrive(
+	int32 FJointSolverGaussSeidel::ApplySwingDrive(
 		const FReal Dt,
 		const FPBDJointSolverSettings& SolverSettings,
 		const FPBDJointSettings& JointSettings,
@@ -1353,23 +1308,24 @@ namespace Chaos
 			const FReal AngularDriveDamping = FPBDJointUtilities::GetAngularSwingDriveDamping(SolverSettings, JointSettings);
 			const bool bAccelerationMode = FPBDJointUtilities::GetDriveAccelerationMode(SolverSettings, JointSettings);
 			ApplyRotationConstraintSoft(Dt, AngularDriveStiffness, AngularDriveDamping, bAccelerationMode, SwingAxis, DSwingAngle, SwingDriveLambda);
+			return 1;
 		}
-		return FJointSolverResult::MakeSolved();
+		return 0;
 	}
 
 
-	FJointSolverResult FJointSolverGaussSeidel::ApplySwingProjection(
+	int32 FJointSolverGaussSeidel::ApplySwingProjection(
 		const FReal Dt,
 		const FPBDJointSolverSettings& SolverSettings,
 		const FPBDJointSettings& JointSettings,
 		const EJointAngularConstraintIndex SwingConstraintIndex)
 	{
 		// @todo(ccaulfield): implement swing projection
-		return FJointSolverResult::MakeSolved();
+		return 0;
 	}
 
 
-	FJointSolverResult FJointSolverGaussSeidel::ApplySLerpDrive(
+	int32 FJointSolverGaussSeidel::ApplySLerpDrive(
 		const FReal Dt,
 		const FPBDJointSolverSettings& SolverSettings,
 		const FPBDJointSettings& JointSettings)
@@ -1386,20 +1342,23 @@ namespace Chaos
 				SLerpAngle = SLerpAngle - (FReal)2 * PI;
 			}
 
+			UE_LOG(LogChaosJoint, VeryVerbose, TEXT("      SLerpDrive Pos: %f Axis: %f %f %f"), -SLerpAngle, SLerpAxis.X, SLerpAxis.Y, SLerpAxis.Z);
+
 			if (FMath::Abs(SLerpAngle) > AngleTolerance)
 			{
 				const FReal AngularDriveStiffness = FPBDJointUtilities::GetAngularSLerpDriveStiffness(SolverSettings, JointSettings);
 				const FReal AngularDriveDamping = FPBDJointUtilities::GetAngularSLerpDriveDamping(SolverSettings, JointSettings);
 				const bool bAccelerationMode = FPBDJointUtilities::GetDriveAccelerationMode(SolverSettings, JointSettings);
 				ApplyRotationConstraintSoft(Dt, AngularDriveStiffness, AngularDriveDamping, bAccelerationMode, SLerpAxis, -SLerpAngle, SwingDriveLambda);
+				return 1;
 			}
 		}
-		return FJointSolverResult::MakeSolved();
+		return 0;
 	}
 
 
 	// Kinematic-Dynamic bodies
-	FJointSolverResult FJointSolverGaussSeidel::ApplyPointPositionConstraintKD(
+	int32 FJointSolverGaussSeidel::ApplyPointPositionConstraintKD(
 		const int32 KIndex,
 		const int32 DIndex,
 		const FReal Dt,
@@ -1424,26 +1383,24 @@ namespace Chaos
 			else
 			{
 				// Calculate constraint correction
-				FMatrix33 InvI1 = Utilities::ComputeWorldSpaceInertia(Qs[DIndex], InvILs[DIndex]);
-				FMatrix33 M = Utilities::ComputeJointFactorMatrix(Xs[DIndex] - Ps[DIndex], InvI1, InvMs[DIndex]);
+				FMatrix33 M = Utilities::ComputeJointFactorMatrix(Xs[DIndex] - Ps[DIndex], InvIs[DIndex], InvMs[DIndex]);
 				FMatrix33 MI = M.Inverse();
 				const FVec3 DX = Utilities::Multiply(MI, CX);
 
 				// Apply constraint correction
 				const FVec3 DP1 = -InvMs[DIndex] * DX;
-				const FVec3 DR1 = Utilities::Multiply(InvI1, FVec3::CrossProduct(Xs[DIndex] - Ps[DIndex], -DX));
+				const FVec3 DR1 = Utilities::Multiply(InvIs[DIndex], FVec3::CrossProduct(Xs[DIndex] - Ps[DIndex], -DX));
 
 				ApplyDelta(DIndex, LinearStiffness, DP1, DR1);
-				UpdateDerivedState(DIndex);
 			}
-			return FJointSolverResult::MakeActive();
+			return 1;
 		}
-		return FJointSolverResult::MakeSolved();
+		return 0;
 	}
 
 
 	// Dynamic-Dynamic bodies
-	FJointSolverResult FJointSolverGaussSeidel::ApplyPointPositionConstraintDD(
+	int32 FJointSolverGaussSeidel::ApplyPointPositionConstraintDD(
 		const FReal Dt,
 		const FPBDJointSolverSettings& SolverSettings,
 		const FPBDJointSettings& JointSettings)
@@ -1454,7 +1411,7 @@ namespace Chaos
 		FReal LinearStiffness = FPBDJointUtilities::GetLinearStiffness(SolverSettings, JointSettings);
 		const FVec3 CX = Xs[1] - Xs[0];
 
-		UE_LOG(LogChaosJoint, VeryVerbose, TEXT("    PointKD Delta %f [Limit %f]"), CX.Size(), PositionTolerance);
+		UE_LOG(LogChaosJoint, VeryVerbose, TEXT("    PointDD Delta %f [Limit %f]"), CX.Size(), PositionTolerance);
 
 		if (CX.SizeSquared() > PositionTolerance * PositionTolerance)
 		{
@@ -1467,37 +1424,34 @@ namespace Chaos
 			else
 			{
 				// Calculate constraint correction
-				FMatrix33 InvI0 = Utilities::ComputeWorldSpaceInertia(Qs[0], InvILs[0]);
-				FMatrix33 InvI1 = Utilities::ComputeWorldSpaceInertia(Qs[1], InvILs[1]);
-				FMatrix33 M0 = Utilities::ComputeJointFactorMatrix(Xs[0] - Ps[0], InvI0, InvMs[0]);
-				FMatrix33 M1 = Utilities::ComputeJointFactorMatrix(Xs[1] - Ps[1], InvI1, InvMs[1]);
+				FMatrix33 M0 = Utilities::ComputeJointFactorMatrix(Xs[0] - Ps[0], InvIs[0], InvMs[0]);
+				FMatrix33 M1 = Utilities::ComputeJointFactorMatrix(Xs[1] - Ps[1], InvIs[1], InvMs[1]);
 				FMatrix33 MI = (M0 + M1).Inverse();
 				const FVec3 DX = Utilities::Multiply(MI, CX);
 
 				// Apply constraint correction
 				const FVec3 DP0 = InvMs[0] * DX;
 				const FVec3 DP1 = -InvMs[1] * DX;
-				const FVec3 DR0 = Utilities::Multiply(InvI0, FVec3::CrossProduct(Xs[0] - Ps[0], DX));
-				const FVec3 DR1 = Utilities::Multiply(InvI1, FVec3::CrossProduct(Xs[1] - Ps[1], -DX));
+				const FVec3 DR0 = Utilities::Multiply(InvIs[0], FVec3::CrossProduct(Xs[0] - Ps[0], DX));
+				const FVec3 DR1 = Utilities::Multiply(InvIs[1], FVec3::CrossProduct(Xs[1] - Ps[1], -DX));
 
 				ApplyPositionDelta(LinearStiffness, DP0, DP1);
 				ApplyRotationDelta(LinearStiffness, DR0, DR1);
-				UpdateDerivedState();
 			}
-			return FJointSolverResult::MakeActive();
+			return 1;
 		}
-		return FJointSolverResult::MakeSolved();
+		return 0;
 	}
 
 
-	FJointSolverResult FJointSolverGaussSeidel::ApplySphericalPositionConstraint(
+	int32 FJointSolverGaussSeidel::ApplySphericalPositionConstraint(
 		const FReal Dt,
 		const FPBDJointSolverSettings& SolverSettings,
 		const FPBDJointSettings& JointSettings)
 	{
 		FVec3 Axis;
 		FReal Delta;
-		GetSphericalAxisDelta(Xs[0], Xs[1], Axis, Delta);
+		FPBDJointUtilities::GetSphericalAxisDelta(Xs[0], Xs[1], Axis, Delta);
 
 		const FReal Error = FMath::Max((FReal)0, Delta - JointSettings.LinearLimit);
 		if (FMath::Abs(Error) > PositionTolerance)
@@ -1514,13 +1468,13 @@ namespace Chaos
 				const bool bAccelerationMode = FPBDJointUtilities::GetLinearSoftAccelerationMode(SolverSettings, JointSettings);
 				ApplyPositionConstraintSoft(Dt, Stiffness, Damping, bAccelerationMode, Axis, Error, LinearSoftLambda);
 			}
-			return FJointSolverResult::MakeActive();
+			return 1;
 		}
-		return FJointSolverResult::MakeSolved();
+		return 0;
 	}
 
 
-	FJointSolverResult FJointSolverGaussSeidel::ApplySphericalPositionDrive(
+	int32 FJointSolverGaussSeidel::ApplySphericalPositionDrive(
 		const FReal Dt,
 		const FPBDJointSolverSettings& SolverSettings,
 		const FPBDJointSettings& JointSettings)
@@ -1528,20 +1482,20 @@ namespace Chaos
 		const FVec3 XTarget = Xs[0] + Rs[0] * JointSettings.LinearDriveTarget;
 		FVec3 Axis;
 		FReal Delta;
-		GetSphericalAxisDelta(XTarget, Xs[1], Axis, Delta);
+		FPBDJointUtilities::GetSphericalAxisDelta(XTarget, Xs[1], Axis, Delta);
 		if (FMath::Abs(Delta) > PositionTolerance)
 		{
 			const FReal Stiffness = FPBDJointUtilities::GetLinearDriveStiffness(SolverSettings, JointSettings);
 			const FReal Damping = FPBDJointUtilities::GetLinearDriveDamping(SolverSettings, JointSettings);
 			const bool bAccelerationMode = FPBDJointUtilities::GetDriveAccelerationMode(SolverSettings, JointSettings);
 			ApplyPositionConstraintSoft(Dt, Stiffness, Damping, bAccelerationMode, Axis, Delta, LinearDriveLambda);
-			return FJointSolverResult::MakeActive();
+			return 1;
 		}
-		return FJointSolverResult::MakeSolved();
+		return 0;
 	}
 
 
-	FJointSolverResult FJointSolverGaussSeidel::ApplyCylindricalPositionConstraint(
+	int32 FJointSolverGaussSeidel::ApplyCylindricalPositionConstraint(
 		const FReal Dt,
 		const int32 AxisIndex,
 		const EJointMotionType AxialMotion,
@@ -1553,9 +1507,9 @@ namespace Chaos
 
 		FVec3 Axis, RadialAxis;
 		FReal AxialDelta, RadialDelta;
-		GetCylindricalAxesDeltas(Rs[0], Xs[0], Xs[1], AxisIndex, Axis, AxialDelta, RadialAxis, RadialDelta);
+		FPBDJointUtilities::GetCylindricalAxesDeltas(Rs[0], Xs[0], Xs[1], AxisIndex, Axis, AxialDelta, RadialAxis, RadialDelta);
 		
-		FJointSolverResult NetResult;
+		int32 NumActive = 0;
 
 		if ((AxialMotion == EJointMotionType::Limited) && JointSettings.bSoftLinearLimitsEnabled)
 		{
@@ -1568,7 +1522,7 @@ namespace Chaos
 				const FReal Damping = FPBDJointUtilities::GetSoftLinearDamping(SolverSettings, JointSettings);
 				const bool bAccelerationMode = FPBDJointUtilities::GetLinearSoftAccelerationMode(SolverSettings, JointSettings);
 				ApplyPositionConstraintSoft(Dt, Stiffness, Damping, bAccelerationMode, Axis, AxialError, LinearSoftLambda);
-				NetResult += FJointSolverResult::MakeActive();
+				++NumActive;
 			}
 		}
 		else if (AxialMotion != EJointMotionType::Free)
@@ -1580,7 +1534,7 @@ namespace Chaos
 				const FReal AxialError = (AxialDelta > 0) ? AxialDelta - AxialLimit : AxialDelta + AxialLimit;
 				const FReal Stiffness = FPBDJointUtilities::GetLinearStiffness(SolverSettings, JointSettings);
 				ApplyPositionConstraint(Stiffness, Axis, AxialError);
-				NetResult += FJointSolverResult::MakeActive();
+				++NumActive;
 			}
 		}
 
@@ -1595,7 +1549,7 @@ namespace Chaos
 				const FReal Damping = FPBDJointUtilities::GetSoftLinearDamping(SolverSettings, JointSettings);
 				const bool bAccelerationMode = FPBDJointUtilities::GetLinearSoftAccelerationMode(SolverSettings, JointSettings);
 				ApplyPositionConstraintSoft(Dt, Stiffness, Damping, bAccelerationMode, RadialAxis, RadialError, LinearSoftLambda);
-				NetResult += FJointSolverResult::MakeActive();
+				++NumActive;
 			}
 		}
 		else if (RadialMotion != EJointMotionType::Free)
@@ -1607,15 +1561,15 @@ namespace Chaos
 				const FReal RadialError = FMath::Max((FReal)0, RadialDelta - RadialLimit);
 				const FReal Stiffness = FPBDJointUtilities::GetLinearStiffness(SolverSettings, JointSettings);
 				ApplyPositionConstraint(Stiffness, RadialAxis, RadialError);
-				NetResult += FJointSolverResult::MakeActive();
+				++NumActive;
 			}
 		}
 
-		return NetResult;
+		return NumActive;
 	}
 
 
-	FJointSolverResult FJointSolverGaussSeidel::ApplyCircularPositionDrive(
+	int32 FJointSolverGaussSeidel::ApplyCircularPositionDrive(
 		const FReal Dt,
 		const int32 AxisIndex,
 		const FPBDJointSolverSettings& SolverSettings,
@@ -1624,19 +1578,21 @@ namespace Chaos
 		const FVec3 XTarget = Xs[0] + Rs[0] * JointSettings.LinearDriveTarget;
 		FVec3 Axis, RadialAxis;
 		FReal AxialDelta, RadialDelta;
-		GetCylindricalAxesDeltas(Rs[0], XTarget, Xs[1], AxisIndex, Axis, AxialDelta, RadialAxis, RadialDelta);
+		FPBDJointUtilities::GetCylindricalAxesDeltas(Rs[0], XTarget, Xs[1], AxisIndex, Axis, AxialDelta, RadialAxis, RadialDelta);
 		if (RadialDelta > PositionTolerance)
 		{
 			const FReal Stiffness = FPBDJointUtilities::GetLinearDriveStiffness(SolverSettings, JointSettings);
 			const FReal Damping = FPBDJointUtilities::GetLinearDriveDamping(SolverSettings, JointSettings);
 			const bool bAccelerationMode = FPBDJointUtilities::GetDriveAccelerationMode(SolverSettings, JointSettings);
 			ApplyPositionConstraintSoft(Dt, Stiffness, Damping, bAccelerationMode, RadialAxis, RadialDelta, LinearDriveLambda);
+
+			return 1;
 		}
-		return FJointSolverResult::MakeSolved();
+		return 0;
 	}
 
 
-	FJointSolverResult FJointSolverGaussSeidel::ApplyPlanarPositionConstraint(
+	int32 FJointSolverGaussSeidel::ApplyPlanarPositionConstraint(
 		const FReal Dt,
 		const int32 AxisIndex,
 		const EJointMotionType AxialMotion,
@@ -1645,7 +1601,7 @@ namespace Chaos
 	{
 		FVec3 Axis;
 		FReal Delta;
-		GetPlanarAxisDelta(Rs[0], Xs[0], Xs[1], AxisIndex, Axis, Delta);
+		FPBDJointUtilities::GetPlanarAxisDelta(Rs[0], Xs[0], Xs[1], AxisIndex, Axis, Delta);
 
 		const FReal Limit = (AxialMotion == EJointMotionType::Locked) ? 0 : JointSettings.LinearLimit;
 		if (FMath::Abs(Delta) > Limit + PositionTolerance)
@@ -1663,13 +1619,13 @@ namespace Chaos
 				const FReal Stiffness = FPBDJointUtilities::GetLinearStiffness(SolverSettings, JointSettings);
 				ApplyPositionConstraint(Stiffness, Axis, Error);
 			}
-			return FJointSolverResult::MakeActive();
+			return 1;
 		}
-		return FJointSolverResult::MakeSolved();
+		return 0;
 	}
 
 
-	FJointSolverResult FJointSolverGaussSeidel::ApplyAxialPositionDrive(
+	int32 FJointSolverGaussSeidel::ApplyAxialPositionDrive(
 		const FReal Dt,
 		const int32 AxisIndex,
 		const FPBDJointSolverSettings& SolverSettings,
@@ -1678,19 +1634,21 @@ namespace Chaos
 		const FVec3 XTarget = Xs[0] + Rs[0] * JointSettings.LinearDriveTarget;
 		FVec3 Axis;
 		FReal Delta;
-		GetPlanarAxisDelta(Rs[0], XTarget, Xs[1], AxisIndex, Axis, Delta);
+		FPBDJointUtilities::GetPlanarAxisDelta(Rs[0], XTarget, Xs[1], AxisIndex, Axis, Delta);
 		if (FMath::Abs(Delta) > PositionTolerance)
 		{
 			const FReal Stiffness = FPBDJointUtilities::GetLinearDriveStiffness(SolverSettings, JointSettings);
 			const FReal Damping = FPBDJointUtilities::GetLinearDriveDamping(SolverSettings, JointSettings);
 			const bool bAccelerationMode = FPBDJointUtilities::GetDriveAccelerationMode(SolverSettings, JointSettings);
 			ApplyPositionConstraintSoft(Dt, Stiffness, Damping, bAccelerationMode, Axis, Delta, LinearDriveLambda);
+
+			return 1;
 		}
-		return FJointSolverResult::MakeSolved();
+		return 0;
 	}
 
 
-	FJointSolverResult FJointSolverGaussSeidel::ApplyPositionProjection(
+	int32 FJointSolverGaussSeidel::ApplyPositionProjection(
 		const FReal Dt,
 		const FPBDJointSolverSettings& SolverSettings,
 		const FPBDJointSettings& JointSettings)
@@ -1722,10 +1680,8 @@ namespace Chaos
 					const FReal IM1 = InvMs[1];
 					const FVec3 IIL0 = ParentMassScale * InvILs[0];
 					const FVec3& IIL1 = InvILs[1];
-					FMatrix33 II0 = Utilities::ComputeWorldSpaceInertia(Qs[0], IIL0);
-					FMatrix33 II1 = Utilities::ComputeWorldSpaceInertia(Qs[1], IIL1);
-					FMatrix33 J0 = (IM0 > 0) ? Utilities::ComputeJointFactorMatrix(Xs[0] - Ps[0], II0, IM0) : FMatrix33(0, 0, 0);
-					FMatrix33 J1 = Utilities::ComputeJointFactorMatrix(Xs[1] - Ps[1], II1, IM1);
+					FMatrix33 J0 = (IM0 > 0) ? Utilities::ComputeJointFactorMatrix(Xs[0] - Ps[0], InvIs[0], IM0) : FMatrix33(0, 0, 0);
+					FMatrix33 J1 = Utilities::ComputeJointFactorMatrix(Xs[1] - Ps[1], InvIs[1], IM1);
 					FMatrix33 IJ = (J0 + J1).Inverse();
 
 					const FVec3 DX = Utilities::Multiply(IJ, CX);
@@ -1733,23 +1689,22 @@ namespace Chaos
 
 					const FVec3 DP0 = IM0 * DX;
 					const FVec3 DP1 = -IM1 * DX;
-					const FVec3 DR0 = Utilities::Multiply(II0, FVec3::CrossProduct(Xs[0] - Ps[0], DX));
-					const FVec3 DR1 = Utilities::Multiply(II1, FVec3::CrossProduct(Xs[1] - Ps[1], -DX));
+					const FVec3 DR0 = Utilities::Multiply(InvIs[0], FVec3::CrossProduct(Xs[0] - Ps[0], DX));
+					const FVec3 DR1 = Utilities::Multiply(InvIs[1], FVec3::CrossProduct(Xs[1] - Ps[1], -DX));
 
 					const FVec3 DV0 = IM0 * DV;
 					const FVec3 DV1 = -IM1 * DV;
-					const FVec3 DW0 = Utilities::Multiply(II0, FVec3::CrossProduct(Xs[0] - Ps[0], DV));
-					const FVec3 DW1 = Utilities::Multiply(II1, FVec3::CrossProduct(Xs[1] - Ps[1], -DV));
+					const FVec3 DW0 = Utilities::Multiply(InvIs[0], FVec3::CrossProduct(Xs[0] - Ps[0], DV));
+					const FVec3 DW1 = Utilities::Multiply(InvIs[1], FVec3::CrossProduct(Xs[1] - Ps[1], -DV));
 
 					const FReal Stiffness = FPBDJointUtilities::GetLinearStiffness(SolverSettings, JointSettings);
 					ApplyPositionDelta(Stiffness, DP0, DP1);
 					ApplyRotationDelta(Stiffness, DR0, DR1);
 					ApplyVelocityDelta(Stiffness, DV0, DW0, DV1, DW1);
-					UpdateDerivedState();
 				}
-				return FJointSolverResult::MakeActive();
+				return 1;
 			}
 		}
-		return FJointSolverResult::MakeSolved();
+		return 0;
 	}
 }
diff --git a/Engine/Source/Runtime/Experimental/Chaos/Private/Chaos/Joint/PBDJointSolverGaussSeidel.ispc b/Engine/Source/Runtime/Experimental/Chaos/Private/Chaos/Joint/PBDJointSolverGaussSeidel.ispc
--- a/Engine/Source/Runtime/Experimental/Chaos/Private/Chaos/Joint/PBDJointSolverGaussSeidel.ispc
+++ b/Engine/Source/Runtime/Experimental/Chaos/Private/Chaos/Joint/PBDJointSolverGaussSeidel.ispc
@@ -35,6 +35,7 @@ struct FJointSolverGaussSeidel
 	FVector4 Qs[MaxConstrainedBodies];			// World-space particle CoM rotations
 	FVector Vs[MaxConstrainedBodies];			// World-space particle CoM velocities
 	FVector Ws[MaxConstrainedBodies];			// World-space particle CoM angular velocities
+	FMatrix33 InvIs[MaxConstrainedBodies];		// World-space inverse inertias
 
 	// XPBD Previous iteration world-space body state
 	FVector PrevPs[MaxConstrainedBodies];		// World-space particle CoM positions
@@ -42,14 +43,16 @@ struct FJointSolverGaussSeidel
 	float SwingDriveLambda;						// XPBD constraint multipliers (net applied constraint-space deltas)
 	FVector4 PrevQs[MaxConstrainedBodies];		// World-space particle CoM rotations
 	FVector PrevXs[MaxConstrainedBodies];		// World-space joint connector positions
-
-	// Post-angular constraint position fixup (to reduce iterations required for stiff angular constraints)
-	float AngularPositionCorrection;
-
+	
 	float PositionTolerance;					// Distance error below which we consider a constraint or drive solved
 	float AngleTolerance;						// Angle error below which we consider a constraint or drive solved
 };
 
+export uniform int32 SizeofFJointSolverGaussSeidel()
+{
+	return sizeof(uniform FJointSolverGaussSeidel);
+}
+
 static inline uniform FMatrix33 ComputeWorldSpaceInertia(const uniform FVector4 &Q, const uniform FVector &I)
 {
 	const uniform FMatrix33 QM = MakeQuatRotationTranslationMatrix(Q);
@@ -105,28 +108,6 @@ static inline uniform FVector CalculateVelocity(const uniform FVector& P0, const
 	return (P1 - P0) / Dt;
 }
 
-
-static inline void UpdateDerivedState(uniform FJointSolverGaussSeidel * uniform M, const uniform int Index)
-{
-	M->Xs[Index] = M->Ps[Index] + VectorQuaternionRotateVector(M->Qs[Index], M->XLs[Index].Translation);
-	M->Rs[Index] = VectorQuaternionMultiply2(M->Qs[Index], M->XLs[Index].Rotation);
-
-	const uniform FVector4 Rs1Shortest = VectorQuaternionEnforceShortestArcWith(M->Rs[1], M->Rs[0]);
-	M->Rs[1] = Rs1Shortest;
-}
-
-static inline void UpdateDerivedState(uniform FJointSolverGaussSeidel * uniform M)
-{
-	M->Xs[0] = M->Ps[0] + VectorQuaternionRotateVector(M->Qs[0], M->XLs[0].Translation);
-	M->Xs[1] = M->Ps[1] + VectorQuaternionRotateVector(M->Qs[1], M->XLs[1].Translation);
-
-	const uniform FVector4 Rs0 = VectorQuaternionMultiply2(M->Qs[0], M->XLs[0].Rotation);
-	const uniform FVector4 Rs1 = VectorQuaternionMultiply2(M->Qs[1], M->XLs[1].Rotation);
-	const uniform FVector4 Rs1Shortest = VectorQuaternionEnforceShortestArcWith(Rs1, Rs0);
-	M->Rs[0] = Rs0;
-	M->Rs[1] = Rs1Shortest;
-}
-
 static inline void ApplyPositionDelta(
 	uniform FJointSolverGaussSeidel * uniform M,
 	const uniform int Index,
@@ -163,6 +144,14 @@ static inline void ApplyRotationDelta(
 
 	const uniform FVector4 Qs1Shortest = VectorQuaternionEnforceShortestArcWith(M->Qs[1], M->Qs[0]);
 	M->Qs[1] = Qs1Shortest;
+
+	M->Xs[Index] = M->Ps[Index] + VectorQuaternionRotateVector(M->Qs[Index], M->XLs[Index].Translation);
+	M->Rs[Index] = VectorQuaternionMultiply2(M->Qs[Index], M->XLs[Index].Rotation);
+
+	const uniform FVector4 Rs1Shortest = VectorQuaternionEnforceShortestArcWith(M->Rs[1], M->Rs[0]);
+	M->Rs[1] = Rs1Shortest;
+
+	M->InvIs[Index] = ComputeWorldSpaceInertia(M->Qs[Index], M->InvILs[Index]);
 }
 
 static inline void ApplyRotationDelta(
@@ -171,21 +160,41 @@ static inline void ApplyRotationDelta(
 	const uniform FVector& DR0,
 	const uniform FVector& DR1)
 {
-	const uniform FVector4 Qs0 = M->Qs[0];
-	const uniform FVector4 Qs1 = M->Qs[1];
-	const uniform FVector4 DR0Stiffness = SetVector4(DR0 * Stiffness, 0);
-	const uniform FVector4 DR1Stiffness = SetVector4(DR1 * Stiffness, 0);
-
-	const uniform FVector4 DQ0 = VectorQuaternionMultiply2(DR0Stiffness, Qs0) * 0.5f;
-	const uniform FVector4 DQ1 = VectorQuaternionMultiply2(DR1Stiffness, Qs1) * 0.5f;
+	if (M->InvMs[0] > 0.0f)
+	{
+		const uniform FVector4 Qs0 = M->Qs[0];
+		const uniform FVector4 DR0Stiffness = SetVector4(DR0 * Stiffness, 0);
+		const uniform FVector4 DQ0 = VectorQuaternionMultiply2(DR0Stiffness, Qs0) * 0.5f;
+		const uniform FVector4 Q0 = VectorNormalizeQuaternion(Qs0 + DQ0);
+		M->Qs[0] = Q0;
+		M->Xs[0] = M->Ps[0] + VectorQuaternionRotateVector(Q0, M->XLs[0].Translation);
+		M->Rs[0] = VectorQuaternionMultiply2(Q0, M->XLs[0].Rotation);
+		M->InvIs[0] = ComputeWorldSpaceInertia(Q0, M->InvILs[0]);
+	}
 
-	const uniform FVector4 Qs0Norm = VectorNormalizeQuaternion(Qs0 + DQ0);
-	const uniform FVector4 Qs1Norm = VectorNormalizeQuaternion(Qs1 + DQ1);
+	if (M->InvMs[1] > 0.0f)
+	{
+		const uniform FVector4 Qs1 = M->Qs[1];
+		const uniform FVector4 DR1Stiffness = SetVector4(DR1 * Stiffness, 0);
+		const uniform FVector4 DQ1 = VectorQuaternionMultiply2(DR1Stiffness, Qs1) * 0.5f;
+		const uniform FVector4 Q1 = VectorNormalizeQuaternion(Qs1 + DQ1);
+		M->Qs[1] = Q1;
+		M->Xs[1] = M->Ps[1] + VectorQuaternionRotateVector(Q1, M->XLs[1].Translation);
+		M->Rs[1] = VectorQuaternionMultiply2(Q1, M->XLs[1].Rotation);
+		M->InvIs[1] = ComputeWorldSpaceInertia(Q1, M->InvILs[1]);
+	}
 
-	const uniform FVector4 Qs1ShortestArcWith = VectorQuaternionEnforceShortestArcWith(Qs1Norm, Qs0Norm);
+	M->Qs[1] = VectorQuaternionEnforceShortestArcWith(M->Qs[1], M->Qs[0]);
+	M->Rs[1] = VectorQuaternionEnforceShortestArcWith(M->Rs[1], M->Rs[0]);
+}
 
-	M->Qs[0] = Qs0Norm;
-	M->Qs[1] = Qs1ShortestArcWith;
+export void ApplyRotationDelta2(
+	uniform FJointSolverGaussSeidel * uniform M,
+	const uniform float Stiffness,
+	const uniform FVector& DR0,
+	const uniform FVector& DR1)
+{
+	ApplyRotationDelta(M, Stiffness, DR0, DR1);
 }
 
 static inline void ApplyDelta(
@@ -213,11 +222,6 @@ static inline void ApplyVelocityDelta(
 	M->Ws[1] = M->Ws[1] + Stiffness * DW1;
 }
 
-export uniform int32 SizeofFJointSolverGaussSeidel()
-{
-	return sizeof(uniform FJointSolverGaussSeidel);
-}
-
 export void ApplyPositionConstraintSoft(
 	uniform FJointSolverGaussSeidel * uniform M,
 	const uniform float Dt,
@@ -228,15 +232,11 @@ export void ApplyPositionConstraintSoft(
 	const uniform float Delta,
 	uniform float& Lambda)
 {
-	// World-space inverse mass
-	const uniform FMatrix33 InvI0 = ComputeWorldSpaceInertia(M->Qs[0], M->InvILs[0]);
-	const uniform FMatrix33 InvI1 = ComputeWorldSpaceInertia(M->Qs[1], M->InvILs[1]);
-	
 	// Joint-space inverse mass
 	const uniform FVector AngularAxis0 = VectorCross(M->Xs[0] - M->Ps[0], Axis);
 	const uniform FVector AngularAxis1 = VectorCross(M->Xs[1] - M->Ps[1], Axis);
-	const uniform FVector IA0 = Multiply(InvI0, AngularAxis0);
-	const uniform FVector IA1 = Multiply(InvI1, AngularAxis1);
+	const uniform FVector IA0 = Multiply(M->InvIs[0], AngularAxis0);
+	const uniform FVector IA1 = Multiply(M->InvIs[1], AngularAxis1);
 	const uniform float II0 = VectorDot(AngularAxis0, IA0);
 	const uniform float II1 = VectorDot(AngularAxis1, IA1);
 	const uniform float II = (M->InvMs[0] + II0 + M->InvMs[1] + II1);
@@ -257,13 +257,12 @@ export void ApplyPositionConstraintSoft(
 	
 	const uniform FVector DP0 = (M->InvMs[0] * DLambda) * Axis;
 	const uniform FVector DP1 = (-M->InvMs[1] * DLambda) * Axis;
-	const uniform FVector DR0 = DLambda * Multiply(InvI0, AngularAxis0);
-	const uniform FVector DR1 = -DLambda * Multiply(InvI1, AngularAxis1);
+	const uniform FVector DR0 = DLambda * Multiply(M->InvIs[0], AngularAxis0);
+	const uniform FVector DR1 = -DLambda * Multiply(M->InvIs[1], AngularAxis1);
 
 	Lambda += DLambda;
 	ApplyPositionDelta(M, 1.0f, DP0, DP1);
 	ApplyRotationDelta(M, 1.0f, DR0, DR1);
-	UpdateDerivedState(M);
 }
 
 // See "XPBD: Position-Based Simulation of Compliant Constrained Dynamics"
@@ -280,32 +279,19 @@ export void ApplyRotationConstraintSoftKD(
 	uniform float& Lambda)
 {
 	// World-space inverse mass
-	const uniform FMatrix33 InvI1 = ComputeWorldSpaceInertia(M->Qs[DIndex], M->InvILs[DIndex]);
-	const uniform FVector IA1 = Multiply(InvI1, Axis);
+	const uniform FVector IA1 = Multiply(M->InvIs[DIndex], Axis);
 
 	// Joint-space inverse mass
 	uniform float II1 = VectorDot(Axis, IA1);
-
-	// If we are correcting the position, we need to adjust the constraint effective mass using parallel axis theorem
-	// @todo(ccaulfield): the IMs here are constant per constraint...pre-build and cache it
-	////if (M->AngularPositionCorrection > 0)
-	//{
-	//	const uniform FVector LinearAxis1 = VectorCross(M->Xs[DIndex] - M->Ps[DIndex], Axis);
-	//	const uniform float LinearAxisLen1Sq = VectorSizeSquared(LinearAxis1);
-	//	if (LinearAxisLen1Sq > KINDA_SMALL_NUMBER)
-	//	{
-	//		const uniform float IM1 = M->InvMs[DIndex] / LinearAxisLen1Sq;
-	//		II1 = II1 * IM1 / (II1 + IM1);
-	//	}
-	//}
 	const uniform float II = II1;
 
 	// Damping angular velocity
 	uniform float AngVelDt = 0;
 	if (Damping > KINDA_SMALL_NUMBER)
 	{
+		const uniform FVector W0 = CalculateAngularVelocity(M->PrevQs[KIndex], M->Qs[KIndex], 1.0f);
 		const uniform FVector W1 = CalculateAngularVelocity(M->PrevQs[DIndex], M->Qs[DIndex], 1.0f);
-		AngVelDt = -VectorDot(Axis, W1);
+		AngVelDt = VectorDot(Axis, W0) - VectorDot(Axis, W1);
 	}
 
 	const uniform float SpringMassScale = (bAccelerationMode) ? 1.0f / II : 1.0f;
@@ -319,23 +305,6 @@ export void ApplyRotationConstraintSoftKD(
 
 	Lambda += DLambda;
 	ApplyRotationDelta(M, DIndex, 1.0f, DR1);
-
-	// Correct the positional error that was introduced by the rotation correction. This is
-	// correct when the position dofs are all locked, but not otherwise (fixable).
-	// This significantly improves angular stiffness at lower iterations, although the same effect
-	// is achieved by increasing iterations.
-	// @todo(ccaulfield): this position correction needs to have components in direction of inactive position constraints removed
-	//if (M->AngularPositionCorrection > 0)
-	//{
-	//	const uniform FVector PrevX0 = M->Xs[KIndex];
-	//	const uniform FVector PrevX1 = M->Xs[DIndex];
-	//	UpdateDerivedState(M);
-	//
-	//	const uniform FVector DP1 = (PrevX1 - PrevX0) - (M->Xs[DIndex] - M->Xs[KIndex]);
-	//	ApplyPositionDelta(M, DIndex, M->AngularPositionCorrection, DP1);
-	//}
-
-	UpdateDerivedState(M, DIndex);
 }
 
 export void ApplyRotationConstraintSoftDD(
@@ -349,34 +318,12 @@ export void ApplyRotationConstraintSoftDD(
 	uniform float& Lambda)
 {
 	// World-space inverse mass
-	const uniform FMatrix33 InvI0 = ComputeWorldSpaceInertia(M->Qs[0], M->InvILs[0]);
-	const uniform FMatrix33 InvI1 = ComputeWorldSpaceInertia(M->Qs[1], M->InvILs[1]);
-	const uniform FVector IA0 = Multiply(InvI0, Axis);
-	const uniform FVector IA1 = Multiply(InvI1, Axis);
+	const uniform FVector IA0 = Multiply(M->InvIs[0], Axis);
+	const uniform FVector IA1 = Multiply(M->InvIs[1], Axis);
 
 	// Joint-space inverse mass
 	uniform float II0 = VectorDot(Axis, IA0);
 	uniform float II1 = VectorDot(Axis, IA1);
-
-	// If we are correcting the position, we need to adjust the constraint effective mass using parallel axis theorem
-	// @todo(ccaulfield): the IMs here are constant per constraint...pre-build and cache it
-	//if (M->AngularPositionCorrection > 0)
-	//{
-	//	const uniform FVector LinearAxis0 = VectorCross(M->Xs[0] - M->Ps[0], Axis);
-	//	const uniform float LinearAxisLen0Sq = VectorSizeSquared(LinearAxis0);
-	//	if (LinearAxisLen0Sq > KINDA_SMALL_NUMBER)
-	//	{
-	//		const uniform float IM0 = M->InvMs[0] / LinearAxisLen0Sq;
-	//		II0 = II0 * IM0 / (II0 + IM0);
-	//	}
-	//	const uniform FVector LinearAxis1 = VectorCross(M->Xs[1] - M->Ps[1], Axis);
-	//	const uniform float LinearAxisLen1Sq = VectorSizeSquared(LinearAxis1);
-	//	if (LinearAxisLen1Sq > KINDA_SMALL_NUMBER)
-	//	{
-	//		const uniform float IM1 = M->InvMs[1] / LinearAxisLen1Sq;
-	//		II1 = II1 * IM1 / (II1 + IM1);
-	//	}
-	//}
 	const uniform float II = (II0 + II1);
 
 	// Damping angular velocity
@@ -401,27 +348,9 @@ export void ApplyRotationConstraintSoftDD(
 
 	Lambda += DLambda;
 	ApplyRotationDelta(M, 1.0f, DR0, DR1);
-
-	// Correct the positional error that was introduced by the rotation correction. This is
-	// correct when the position dofs are all locked, but not otherwise (fixable).
-	// This significantly improves angular stiffness at lower iterations, although the same effect
-	// is achieved by increasing iterations.
-	// @todo(ccaulfield): this position correction needs to have components in direction of inactive position constraints removed
-	//if (M->AngularPositionCorrection > 0)
-	//{
-	//	const uniform FVector PrevX0 = M->Xs[0];
-	//	const uniform FVector PrevX1 = M->Xs[1];
-	//	UpdateDerivedState(M);
-	//
-	//	const uniform FVector DX = (M->Xs[1] - M->Xs[0]) - (PrevX1 - PrevX0);
-	//	const uniform FVector DP0 = DX * (M->InvMs[0] / (M->InvMs[0] + M->InvMs[1]));
-	//	const uniform FVector DP1 = DX * (-M->InvMs[1] / (M->InvMs[0] + M->InvMs[1]));
-	//	ApplyPositionDelta(M, M->AngularPositionCorrection, DP0, DP1);
-	//}
-
-	UpdateDerivedState(M);
 }
 
+
 export void ApplyPointPositionConstraintKD(
 	uniform FJointSolverGaussSeidel * uniform M,
 	const uniform int KIndex,
@@ -431,17 +360,15 @@ export void ApplyPointPositionConstraintKD(
 {
 
 	// Calculate constraint correction
-	const uniform FMatrix33 InvI1 = ComputeWorldSpaceInertia(M->Qs[DIndex], M->InvILs[DIndex]);
-	const uniform FMatrix33 M1 = ComputeJointFactorMatrix(M->Xs[DIndex] - M->Ps[DIndex], InvI1, M->InvMs[DIndex]);
+	const uniform FMatrix33 M1 = ComputeJointFactorMatrix(M->Xs[DIndex] - M->Ps[DIndex], M->InvIs[DIndex], M->InvMs[DIndex]);
 	const uniform FMatrix33 MI = MatrixInverse(M1);
 	const uniform FVector DX = Multiply(MI, CX);
 
 	// Apply constraint correction
 	const uniform FVector DP1 = DX * -M->InvMs[DIndex];
-	const uniform FVector DR1 = Multiply(InvI1, VectorCross(M->Xs[DIndex] - M->Ps[DIndex], VectorNegate(DX)));
+	const uniform FVector DR1 = Multiply(M->InvIs[DIndex], VectorCross(M->Xs[DIndex] - M->Ps[DIndex], VectorNegate(DX)));
 
 	ApplyDelta(M, DIndex, LinearStiffness, DP1, DR1);
-	UpdateDerivedState(M, DIndex);
 }
 
 export void ApplyPointPositionConstraintDD(
@@ -450,22 +377,19 @@ export void ApplyPointPositionConstraintDD(
 	const uniform float LinearStiffness)
 {
 	// Calculate constraint correction
-	const uniform FMatrix33 InvI0 = ComputeWorldSpaceInertia(M->Qs[0], M->InvILs[0]);
-	const uniform FMatrix33 InvI1 = ComputeWorldSpaceInertia(M->Qs[1], M->InvILs[1]);
-	uniform FMatrix33 M0 = ComputeJointFactorMatrix(M->Xs[0] - M->Ps[0], InvI0, M->InvMs[0]);
-	uniform FMatrix33 M1 = ComputeJointFactorMatrix(M->Xs[1] - M->Ps[1], InvI1, M->InvMs[1]);
+	uniform FMatrix33 M0 = ComputeJointFactorMatrix(M->Xs[0] - M->Ps[0], M->InvIs[0], M->InvMs[0]);
+	uniform FMatrix33 M1 = ComputeJointFactorMatrix(M->Xs[1] - M->Ps[1], M->InvIs[1], M->InvMs[1]);
 	const uniform FMatrix33 MI = MatrixInverse(AddAB(M0, M1));
 	const uniform FVector DX = Multiply(MI, CX);
 
 	// Apply constraint correction
 	const uniform FVector DP0 = DX * M->InvMs[0];
 	const uniform FVector DP1 = DX * -M->InvMs[1];
-	const uniform FVector DR0 = Multiply(InvI0, VectorCross(M->Xs[0] - M->Ps[0], DX));
-	const uniform FVector DR1 = Multiply(InvI1, VectorCross(M->Xs[1] - M->Ps[1], VectorNegate(DX)));
+	const uniform FVector DR0 = Multiply(M->InvIs[0], VectorCross(M->Xs[0] - M->Ps[0], DX));
+	const uniform FVector DR1 = Multiply(M->InvIs[1], VectorCross(M->Xs[1] - M->Ps[1], VectorNegate(DX)));
 
 	ApplyPositionDelta(M, LinearStiffness, DP0, DP1);
 	ApplyRotationDelta(M, LinearStiffness, DR0, DR1);
-	UpdateDerivedState(M);
 }
 
 export void ApplyPositionProjection(
@@ -484,14 +408,12 @@ export void ApplyPositionProjection(
 	const uniform float IM1 = M->InvMs[1];
 	const uniform FVector IIL0 = ParentMassScale * M->InvILs[0];
 	const uniform FVector& IIL1 = M->InvILs[1];
-	const uniform FMatrix33 II0 = ComputeWorldSpaceInertia(M->Qs[0], IIL0);
-	const uniform FMatrix33 II1 = ComputeWorldSpaceInertia(M->Qs[1], IIL1);
 	uniform FMatrix33 J0 = Matrix33Zero;
 	if(IM0 > 0)
 	{
-		J0 = ComputeJointFactorMatrix(M->Xs[0] - M->Ps[0], II0, IM0);
+		J0 = ComputeJointFactorMatrix(M->Xs[0] - M->Ps[0], M->InvIs[0], IM0);
 	}
-	const uniform FMatrix33 J1 = ComputeJointFactorMatrix(M->Xs[1] - M->Ps[1], II1, IM1);
+	const uniform FMatrix33 J1 = ComputeJointFactorMatrix(M->Xs[1] - M->Ps[1], M->InvIs[1], IM1);
 	const uniform FMatrix33 IJ = MatrixInverse(AddAB(J0, J1));
 
 	const uniform FVector DX = Multiply(IJ, CX);
@@ -499,16 +421,235 @@ export void ApplyPositionProjection(
 
 	const uniform FVector DP0 = IM0 * DX;
 	const uniform FVector DP1 = -IM1 * DX;
-	const uniform FVector DR0 = Multiply(II0, VectorCross(M->Xs[0] - M->Ps[0], DX));
-	const uniform FVector DR1 = Multiply(II1, VectorCross(M->Xs[1] - M->Ps[1], VectorNegate(DX)));
+	const uniform FVector DR0 = Multiply(M->InvIs[0], VectorCross(M->Xs[0] - M->Ps[0], DX));
+	const uniform FVector DR1 = Multiply(M->InvIs[1], VectorCross(M->Xs[1] - M->Ps[1], VectorNegate(DX)));
 
 	const uniform FVector DV0 = IM0 * DV;
 	const uniform FVector DV1 = -IM1 * DV;
-	const uniform FVector DW0 = Multiply(II0, VectorCross(M->Xs[0] - M->Ps[0], DV));
-	const uniform FVector DW1 = Multiply(II1, VectorCross(M->Xs[1] - M->Ps[1], VectorNegate(DV)));
+	const uniform FVector DW0 = Multiply(M->InvIs[0], VectorCross(M->Xs[0] - M->Ps[0], DV));
+	const uniform FVector DW1 = Multiply(M->InvIs[1], VectorCross(M->Xs[1] - M->Ps[1], VectorNegate(DV)));
 
 	ApplyPositionDelta(M, Stiffness, DP0, DP1);
 	ApplyRotationDelta(M, Stiffness, DR0, DR1);
 	ApplyVelocityDelta(M, Stiffness, DV0, DW0, DV1, DW1);
-	UpdateDerivedState(M);
+}
+
+//
+//////////////////////////////////////////////////////
+//
+// Batched Mode Stuff
+//
+
+struct FJointSolverJointState
+{
+	// Local-space constraint settings
+	FRigidTransform XLs[MaxConstrainedBodies];
+	FVector InvILs[MaxConstrainedBodies];
+	float InvMs[MaxConstrainedBodies];
+
+	// World-space constraint state
+	FVector Xs[MaxConstrainedBodies];
+	FVector4 Rs[MaxConstrainedBodies];
+
+	// World-space body state
+	FVector Ps[MaxConstrainedBodies];
+	FVector4 Qs[MaxConstrainedBodies];
+	FMatrix33 InvIs[MaxConstrainedBodies];
+
+	// XPBD initial world-space body state (start of each tick, not each sub-tick iteration)
+	FVector PrevPs[MaxConstrainedBodies];
+	FVector4 PrevQs[MaxConstrainedBodies];
+	FVector PrevXs[MaxConstrainedBodies];
+
+	FVector DPs[MaxConstrainedBodies];
+	FVector DRs[MaxConstrainedBodies];
+
+	float PositionTolerance;
+	float AngleTolerance;
+};
+
+export uniform int32 SizeofFJointSolverJointState()
+{
+	return sizeof(uniform FJointSolverJointState);
+}
+
+struct FJointSolverConstraintRowState
+{
+	// Per-Iteration Transient
+	FVector DPs[MaxConstrainedBodies];
+	FVector DRs[MaxConstrainedBodies];
+	FVector Axis;
+	float Error;
+
+	// Per-Tick Transient
+	float Lambda;
+};
+
+export uniform int32 SizeofFJointSolverConstraintRowState()
+{
+	return sizeof(uniform FJointSolverConstraintRowState);
+}
+
+struct FJointSolverConstraintRowData
+{
+	int32 UpdateType;			// Method to use to calculate axis, errors, etc.
+	int32 JointIndex;			// Index into the outer container's joint array
+	int32 ConstraintIndex;		// Context dependent: X,Y,Z or Twist,Swing1,Swing2 for linear and angular rows
+	int32 NumRows;				// Number of rows (including this one) to solve simultaneously
+	float Stiffness;			// PBD or XPBD stiffness
+	float Damping;				// XPBD damping
+	float Limit;
+	bool bIsAccelerationMode;
+	bool bIsSoft;
+};
+
+export uniform int32 SizeofFJointSolverConstraintRowData()
+{
+	return sizeof(uniform FJointSolverConstraintRowData);
+}
+
+static inline void BatchUpdateDerivedStateImpl(uniform FJointSolverJointState Joints[], uniform int JointIndexBegin, uniform int JointIndexEnd)
+{
+	for(uniform int JointIndex = JointIndexBegin; JointIndex < JointIndexEnd; JointIndex++)
+	{
+		if (Joints[JointIndex].InvMs[0] > 0.0f)
+		{
+			Joints[JointIndex].Xs[0] = Joints[JointIndex].Ps[0] + VectorQuaternionRotateVector(Joints[JointIndex].Qs[0], Joints[JointIndex].XLs[0].Translation);
+			Joints[JointIndex].Rs[0] = VectorQuaternionMultiply2(Joints[JointIndex].Qs[0], Joints[JointIndex].XLs[0].Rotation);
+			Joints[JointIndex].InvIs[0] = ComputeWorldSpaceInertia(Joints[JointIndex].Qs[0], Joints[JointIndex].InvILs[0]);
+			Joints[JointIndex].DPs[0] = ZeroVector;
+			Joints[JointIndex].DRs[0] = ZeroVector;
+		}
+
+		if (Joints[JointIndex].InvMs[1] > 0.0f)
+		{
+			Joints[JointIndex].Xs[1] = Joints[JointIndex].Ps[1] + VectorQuaternionRotateVector(Joints[JointIndex].Qs[1], Joints[JointIndex].XLs[1].Translation);
+			Joints[JointIndex].Rs[1] = VectorQuaternionMultiply2(Joints[JointIndex].Qs[1], Joints[JointIndex].XLs[1].Rotation);
+			Joints[JointIndex].InvIs[1] = ComputeWorldSpaceInertia(Joints[JointIndex].Qs[1], Joints[JointIndex].InvILs[1]);
+			Joints[JointIndex].DPs[1] = ZeroVector;
+			Joints[JointIndex].DRs[1] = ZeroVector;
+		}
+
+		Joints[JointIndex].Rs[1] = VectorQuaternionEnforceShortestArcWith(Joints[JointIndex].Rs[1], Joints[JointIndex].Rs[0]);
+	}
+}
+
+export void BatchUpdateDerivedState(uniform FJointSolverJointState Joints[], uniform int JointIndexBegin, uniform int JointIndexEnd)
+{
+	BatchUpdateDerivedStateImpl(Joints, JointIndexBegin, JointIndexEnd);
+}
+
+static inline void BatchApplyRotationConstraint(
+	const uniform float Dt,
+	uniform FJointSolverJointState& JointState,
+	uniform FJointSolverConstraintRowData& RowData,
+	uniform FJointSolverConstraintRowState& RowState)
+{
+	// Joint-space inverse mass
+	uniform FVector IA0 = ZeroVector;
+	uniform FVector IA1 = ZeroVector;
+	uniform float II0 = 0.0f;
+	uniform float II1 = 0.0f;
+	if (JointState.InvMs[0] > 0)
+	{
+		IA0 = Multiply(JointState.InvIs[0], RowState.Axis);
+		II0 = VectorDot(RowState.Axis, IA0);
+	}
+	if (JointState.InvMs[1] > 0)
+	{
+		IA1 = Multiply(JointState.InvIs[1], RowState.Axis);
+		II1 = VectorDot(RowState.Axis, IA1);
+	}
+	const uniform float II = (II0 + II1);
+
+	if (RowData.bIsSoft)
+	{
+		// Damping angular velocity
+		uniform float AngVelDt = 0;
+		if (RowData.Damping > 0.0f)
+		{
+			const uniform FVector W0 = CalculateAngularVelocity(JointState.PrevQs[0], JointState.Qs[0], 1.0f);
+			const uniform FVector W1 = CalculateAngularVelocity(JointState.PrevQs[1], JointState.Qs[1], 1.0f);
+			AngVelDt = VectorDot(RowState.Axis, W0) - VectorDot(RowState.Axis, W1);
+		}
+
+		const uniform float SpringMassScale = (RowData.bIsAccelerationMode) ? 1.0f / II : 1.0f;
+		const uniform float S = SpringMassScale * RowData.Stiffness * Dt * Dt;
+		const uniform float D = SpringMassScale * RowData.Damping * Dt;
+		const uniform float Multiplier = 1.0f / ((S + D) * II + 1.0f);
+		const uniform float DLambda = Multiplier * (S * RowState.Error - D * AngVelDt - RowState.Lambda);
+
+		//const FVector DR0 = IA0 * DLambda;
+		//const FVector DR1 = IA1 * -DLambda;
+		RowState.DRs[0] = RowState.Axis * (DLambda * II0);
+		RowState.DRs[1] = RowState.Axis * -(DLambda * II1);
+		RowState.Lambda += DLambda;
+	}
+	else
+	{
+		RowState.DRs[0] = IA0 * (RowData.Stiffness * RowState.Error / II);
+		RowState.DRs[1] = IA1 * -(RowData.Stiffness * RowState.Error / II);
+	}
+}
+
+static inline void BatchApplyRotationDelta(
+	uniform FJointSolverJointState* uniform JointStates, 
+	uniform int JointIndex)
+{
+	if (JointStates[JointIndex].InvMs[0] > 0.0f)
+	{
+		const uniform FVector4 Qs0 = JointStates[JointIndex].Qs[0];
+		const uniform FVector4 DR0 = SetVector4(JointStates[JointIndex].DRs[0], 0);
+		const uniform FVector4 DQ0 = VectorQuaternionMultiply2(DR0, Qs0) * 0.5f;
+		const uniform FVector4 Q0 = VectorNormalizeQuaternion(Qs0 + DQ0);
+		JointStates[JointIndex].Qs[0] = Q0;
+	}
+
+	if (JointStates[JointIndex].InvMs[1] > 0.0f)
+	{
+		const uniform FVector4 Qs1 = JointStates[JointIndex].Qs[1];
+		const uniform FVector4 DR1 = SetVector4(JointStates[JointIndex].DRs[1], 0);
+		const uniform FVector4 DQ1 = VectorQuaternionMultiply2(DR1, Qs1) * 0.5f;
+		const uniform FVector4 Q1 = VectorNormalizeQuaternion(Qs1 + DQ1);
+		JointStates[JointIndex].Qs[1] = Q1;
+	}
+
+	JointStates[JointIndex].Qs[1] = VectorQuaternionEnforceShortestArcWith(JointStates[JointIndex].Qs[1], JointStates[JointIndex].Qs[0]);
+}
+
+export uniform int BatchApplyRotationConstraints(
+	uniform float Dt, 
+	uniform FJointSolverJointState* uniform JointStates, 
+	uniform FJointSolverConstraintRowData* uniform RowDatas, 
+	uniform FJointSolverConstraintRowState* uniform RowStates, 
+	uniform int JointIndexBegin, 
+	uniform int JointIndexEnd, 
+	uniform int RowIndexBegin, 
+	uniform int RowIndexEnd)
+{
+	uniform int NumActiveConstraints = 0;
+
+	// Calculate the delta position and rotation to fix the constraint errors
+	// NOTE: Rows are in Joint order, so we can loop over them directly rather than looping over joints and then their rows.
+	for (uniform int RowIndex = RowIndexBegin; RowIndex < RowIndexEnd; ++RowIndex)
+	{
+		const uniform int RowJointIndex = RowDatas[RowIndex].JointIndex;
+
+		if (abs(RowStates[RowIndex].Error) > JointStates[RowJointIndex].AngleTolerance)
+		{
+			BatchApplyRotationConstraint(Dt, JointStates[RowJointIndex], RowDatas[RowIndex], RowStates[RowIndex]);
+
+			JointStates[RowJointIndex].DRs[0] = JointStates[RowJointIndex].DRs[0] + RowStates[RowIndex].DRs[0];
+			JointStates[RowJointIndex].DRs[1] = JointStates[RowJointIndex].DRs[1] + RowStates[RowIndex].DRs[1];
+			++NumActiveConstraints;
+		}
+	}
+
+	// Accumulate and apply the deltas for each joint
+	for (uniform int32 JointIndex = JointIndexBegin; JointIndex < JointIndexEnd; ++JointIndex)
+	{
+		BatchApplyRotationDelta(JointStates, JointIndex);
+	}
+
+	return NumActiveConstraints;
 }
diff --git a/Engine/Source/Runtime/Experimental/Chaos/Private/Chaos/PBDJointConstraintUtilities.cpp b/Engine/Source/Runtime/Experimental/Chaos/Private/Chaos/PBDJointConstraintUtilities.cpp
--- a/Engine/Source/Runtime/Experimental/Chaos/Private/Chaos/PBDJointConstraintUtilities.cpp
+++ b/Engine/Source/Runtime/Experimental/Chaos/Private/Chaos/PBDJointConstraintUtilities.cpp
@@ -8,6 +8,69 @@
 
 namespace Chaos
 {
+	void FPBDJointUtilities::GetSphericalAxisDelta(
+		const FVec3& X0,
+		const FVec3& X1,
+		FVec3& Axis,
+		FReal& Delta)
+	{
+		Axis = FVec3(0);
+		Delta = 0;
+
+		const FVec3 DX = X1 - X0;
+		const FReal DXLen = DX.Size();
+		if (DXLen > KINDA_SMALL_NUMBER)
+		{
+			Axis = DX / DXLen;
+			Delta = DXLen;
+		}
+	}
+
+
+	void FPBDJointUtilities::GetCylindricalAxesDeltas(
+		const FRotation3& R0,
+		const FVec3& X0,
+		const FVec3& X1,
+		const int32 CylinderAxisIndex,
+		FVec3& CylinderAxis,
+		FReal& CylinderDelta,
+		FVec3& RadialAxis,
+		FReal& RadialDelta)
+	{
+		CylinderAxis = FVec3(0);
+		CylinderDelta = 0;
+		RadialAxis = FVec3(0);
+		RadialDelta = 0;
+
+		FVec3 DX = X1 - X0;
+
+		const FMatrix33 R0M = R0.ToMatrix();
+		CylinderAxis = R0M.GetAxis(CylinderAxisIndex);
+		CylinderDelta = FVec3::DotProduct(DX, CylinderAxis);
+
+		DX = DX - CylinderDelta * CylinderAxis;
+		const FReal DXLen = DX.Size();
+		if (DXLen > KINDA_SMALL_NUMBER)
+		{
+			RadialAxis = DX / DXLen;
+			RadialDelta = DXLen;
+		}
+	}
+
+
+	void FPBDJointUtilities::GetPlanarAxisDelta(
+		const FRotation3& R0,
+		const FVec3& X0,
+		const FVec3& X1,
+		const int32 PlaneAxisIndex,
+		FVec3& Axis,
+		FReal& Delta)
+	{
+		const FMatrix33 R0M = R0.ToMatrix();
+		Axis = R0M.GetAxis(PlaneAxisIndex);
+		Delta = FVec3::DotProduct(X1 - X0, Axis);
+	}
+
 	void FPBDJointUtilities::DecomposeSwingTwistLocal(const FRotation3& R0, const FRotation3& R1, FRotation3& R01Swing, FRotation3& R01Twist)
 	{
 		const FRotation3 R01 = R0.Inverse() * R1;
@@ -126,7 +189,7 @@ namespace Chaos
 	}
 
 
-	void FPBDJointUtilities::GetLockedAxes(const FRotation3& R0, const FRotation3& R1, FVec3& Axis0, FVec3& Axis1, FVec3& Axis2)
+	void FPBDJointUtilities::GetLockedRotationAxes(const FRotation3& R0, const FRotation3& R1, FVec3& Axis0, FVec3& Axis1, FVec3& Axis2)
 	{
 		const FReal W0 = R0.W;
 		const FReal W1 = R1.W;
@@ -152,6 +215,32 @@ namespace Chaos
 		}
 	}
 	
+
+	FReal FPBDJointUtilities::GetConeAngleLimit(
+		const FPBDJointSettings& JointSettings,
+		const FVec3& SwingAxisLocal,
+		const FReal SwingAngle)
+	{
+		// Calculate swing limit for the current swing axis
+		const FReal Swing1Limit = JointSettings.AngularLimits[(int32)EJointAngularConstraintIndex::Swing1];
+		const FReal Swing2Limit = JointSettings.AngularLimits[(int32)EJointAngularConstraintIndex::Swing2];
+
+		// Circular swing limit
+		FReal SwingAngleMax = Swing1Limit;
+
+		// Elliptical swing limit
+		// @todo(ccaulfield): do elliptical constraints properly (axis is still for circular limit)
+		if (!FMath::IsNearlyEqual(Swing1Limit, Swing2Limit, KINDA_SMALL_NUMBER))
+		{
+			// Map swing axis to ellipse and calculate limit for this swing axis
+			const FReal DotSwing1 = FMath::Abs(FVec3::DotProduct(SwingAxisLocal, FJointConstants::Swing1Axis()));
+			const FReal DotSwing2 = FMath::Abs(FVec3::DotProduct(SwingAxisLocal, FJointConstants::Swing2Axis()));
+			SwingAngleMax = FMath::Sqrt(Swing1Limit * DotSwing1 * Swing1Limit * DotSwing1 + Swing2Limit * DotSwing2 * Swing2Limit * DotSwing2);
+		}
+
+		return SwingAngleMax;
+	}
+
 	
 	// @todo(ccaulfield): separate linear soft and stiff
 	FReal FPBDJointUtilities::GetLinearStiffness(
diff --git a/Engine/Source/Runtime/Experimental/Chaos/Private/Chaos/PBDJointConstraints.cpp b/Engine/Source/Runtime/Experimental/Chaos/Private/Chaos/PBDJointConstraints.cpp
--- a/Engine/Source/Runtime/Experimental/Chaos/Private/Chaos/PBDJointConstraints.cpp
+++ b/Engine/Source/Runtime/Experimental/Chaos/Private/Chaos/PBDJointConstraints.cpp
@@ -3,7 +3,7 @@
 #include "Chaos/ChaosDebugDraw.h"
 #include "Chaos/DebugDrawQueue.h"
 #include "Chaos/Joint/ChaosJointLog.h"
-#include "Chaos/Joint/PBDJointSolverGaussSeidel.h"
+#include "Chaos/Joint/ColoringGraph.h"
 #include "Chaos/Particle/ParticleUtilities.h"
 #include "Chaos/ParticleHandle.h"
 #include "Chaos/PBDJointConstraintUtilities.h"
@@ -13,15 +13,36 @@
 
 #include "HAL/IConsoleManager.h"
 
+#if INTEL_ISPC
+#include "PBDJointSolverGaussSeidel.ispc.generated.h"
+#endif
+
 //PRAGMA_DISABLE_OPTIMIZATION
 
 bool bChaos_Joint_EarlyOut_Enabled = true;
 FAutoConsoleVariableRef CVarChaosJointEarlyOutEnabled(TEXT("p.Chaos.Joint.EarlyOut"), bChaos_Joint_EarlyOut_Enabled, TEXT("Whether to iterating when joints report being solved"));
 
+bool bChaos_Joint_Batching = false;
+FAutoConsoleVariableRef CVarChaosJointBatching(TEXT("p.Chaos.Joint.Batching"), bChaos_Joint_Batching, TEXT(""));
+
+int32 bChaos_Joint_MaxBatchSize = 1000;
+FAutoConsoleVariableRef CVarChaosJointBatchSize(TEXT("p.Chaos.Joint.MaxBatchSize"), bChaos_Joint_MaxBatchSize, TEXT(""));
+
+#if !INTEL_ISPC
+const bool bChaos_Joint_ISPC_Enabled = false;
+#elif UE_BUILD_SHIPPING
+const bool bChaos_Joint_ISPC_Enabled = true;
+#else
+extern bool bChaos_Joint_ISPC_Enabled;
+#endif
+
 namespace Chaos
 {
 	DECLARE_CYCLE_STAT(TEXT("Joints::Sort"), STAT_Joints_Sort, STATGROUP_ChaosJoint);
+	DECLARE_CYCLE_STAT(TEXT("Joints::Prepare"), STAT_Joints_Prepare, STATGROUP_ChaosJoint);
+	DECLARE_CYCLE_STAT(TEXT("Joints::Unprepare"), STAT_Joints_Unprepare, STATGROUP_ChaosJoint);
 	DECLARE_CYCLE_STAT(TEXT("Joints::Apply"), STAT_Joints_Apply, STATGROUP_ChaosJoint);
+	DECLARE_CYCLE_STAT(TEXT("Joints::ApplyBatched"), STAT_Joints_ApplyBatched, STATGROUP_ChaosJoint);
 	DECLARE_CYCLE_STAT(TEXT("Joints::ApplyPushOut"), STAT_Joints_ApplyPushOut, STATGROUP_ChaosJoint);
 
 	//
@@ -45,18 +66,30 @@ namespace Chaos
 		ConstraintContainer->CalculateConstraintSpace(ConstraintIndex, OutXa, OutRa, OutXb, OutRb);
 	}
 
-	
-	void FPBDJointConstraintHandle::SetParticleLevels(const TVector<int32, 2>& ParticleLevels)
+
+	int32 FPBDJointConstraintHandle::GetConstraintIsland() const
 	{
-		ConstraintContainer->SetParticleLevels(ConstraintIndex, ParticleLevels);
+		return ConstraintContainer->GetConstraintIsland(ConstraintIndex);
 	}
 
-	
+
 	int32 FPBDJointConstraintHandle::GetConstraintLevel() const
 	{
 		return ConstraintContainer->GetConstraintLevel(ConstraintIndex);
 	}
 
+
+	int32 FPBDJointConstraintHandle::GetConstraintColor() const
+	{
+		return ConstraintContainer->GetConstraintColor(ConstraintIndex);
+	}
+
+
+	int32 FPBDJointConstraintHandle::GetConstraintBatch() const
+	{
+		return ConstraintContainer->GetConstraintBatch(ConstraintIndex);
+	}
+
 	
 	const FPBDJointSettings& FPBDJointConstraintHandle::GetSettings() const
 	{
@@ -117,6 +150,15 @@ namespace Chaos
 		, AngularDriveStiffness(0)
 		, AngularDriveDamping(0)
 	{
+		if (bChaos_Joint_ISPC_Enabled)
+		{
+#if INTEL_ISPC
+			check(sizeof(FJointSolverJointState) == ispc::SizeofFJointSolverJointState());
+			check(sizeof(FJointSolverConstraintRowState) == ispc::SizeofFJointSolverConstraintRowState());
+			check(sizeof(FJointSolverConstraintRowData) == ispc::SizeofFJointSolverConstraintRowData());
+			check(sizeof(FJointSolverJointState) == ispc::SizeofFJointSolverJointState());
+#endif
+		}
 	}
 
 
@@ -144,11 +186,15 @@ namespace Chaos
 
 	
 	FPBDJointState::FPBDJointState()
-		: Level(INDEX_NONE)
-		, ParticleLevels({ INDEX_NONE, INDEX_NONE })
+		: Batch(INDEX_NONE)
+		, Island(INDEX_NONE)
+		, Level(INDEX_NONE)
+		, Color(INDEX_NONE)
+		, IslandSize(0)
 	{
 	}
 
+
 	//
 	// Solver Settings
 	//
@@ -182,6 +228,7 @@ namespace Chaos
 	{
 	}
 
+
 	//
 	// Constraint Container
 	//
@@ -189,7 +236,8 @@ namespace Chaos
 	
 	FPBDJointConstraints::FPBDJointConstraints(const FPBDJointSolverSettings& InSettings)
 		: Settings(InSettings)
-		, bRequiresSort(false)
+		, bJointsDirty(false)
+		, bIsBatched(false)
 		, PreApplyCallback(nullptr)
 		, PostApplyCallback(nullptr)
 	{
@@ -259,6 +307,8 @@ namespace Chaos
 	
 	typename FPBDJointConstraints::FConstraintContainerHandle* FPBDJointConstraints::AddConstraint(const FParticlePair& InConstrainedParticles, const FTransformPair& InConstraintFrames, const FPBDJointSettings& InConstraintSettings)
 	{
+		bJointsDirty = true;
+
 		int ConstraintIndex = Handles.Num();
 		Handles.Add(HandleAllocator.AllocHandle(this, ConstraintIndex));
 		ConstraintParticles.Add(InConstrainedParticles);
@@ -271,6 +321,8 @@ namespace Chaos
 	
 	void FPBDJointConstraints::RemoveConstraint(int ConstraintIndex)
 	{
+		bJointsDirty = true;
+
 		FConstraintContainerHandle* ConstraintHandle = Handles[ConstraintIndex];
 		if (ConstraintHandle != nullptr)
 		{
@@ -309,22 +361,40 @@ namespace Chaos
 		FHandles SortedHandles = Handles;
 		SortedHandles.StableSort([](const FConstraintContainerHandle& L, const FConstraintContainerHandle& R)
 			{
-				return L.GetConstraintLevel() < R.GetConstraintLevel();
+				if (L.GetConstraintBatch() != R.GetConstraintBatch())
+				{
+					return L.GetConstraintBatch() < R.GetConstraintBatch();
+				}
+				else if (L.GetConstraintIsland() != R.GetConstraintIsland())
+				{
+					return L.GetConstraintIsland() < R.GetConstraintIsland();
+				}
+				else if (L.GetConstraintLevel() != R.GetConstraintLevel())
+				{
+					return L.GetConstraintLevel() < R.GetConstraintLevel();
+				}
+				return L.GetConstraintColor() < R.GetConstraintColor();
 			});
 
 		TArray<FPBDJointSettings> SortedConstraintSettings;
 		TArray<FTransformPair> SortedConstraintFrames;
 		TArray<FParticlePair> SortedConstraintParticles;
 		TArray<FPBDJointState> SortedConstraintStates;
-		int32 SortedConstraintIndex = 0;
-		for (FConstraintContainerHandle* Handle : SortedHandles)
+		SortedConstraintSettings.Reserve(SortedHandles.Num());
+		SortedConstraintFrames.Reserve(SortedHandles.Num());
+		SortedConstraintParticles.Reserve(SortedHandles.Num());
+		SortedConstraintStates.Reserve(SortedHandles.Num());
+
+		for (int32 SortedConstraintIndex = 0; SortedConstraintIndex < SortedHandles.Num(); ++SortedConstraintIndex)
 		{
-			int32 UnsortedIndex = Handle->GetConstraintIndex();
-			SortedConstraintSettings.Add(ConstraintSettings[UnsortedIndex]);
-			SortedConstraintFrames.Add(ConstraintFrames[UnsortedIndex]);
-			SortedConstraintParticles.Add(ConstraintParticles[UnsortedIndex]);
-			SortedConstraintStates.Add(ConstraintStates[UnsortedIndex]);
-			SetConstraintIndex(Handle, SortedConstraintIndex++);
+			FConstraintContainerHandle* Handle = SortedHandles[SortedConstraintIndex];
+			int32 UnsortedConstraintIndex = Handle->GetConstraintIndex();
+
+			SortedConstraintSettings.Add(ConstraintSettings[UnsortedConstraintIndex]);
+			SortedConstraintFrames.Add(ConstraintFrames[UnsortedConstraintIndex]);
+			SortedConstraintParticles.Add(ConstraintParticles[UnsortedConstraintIndex]);
+			SortedConstraintStates.Add(ConstraintStates[UnsortedConstraintIndex]);
+			SetConstraintIndex(Handle, SortedConstraintIndex);
 		}
 
 		Swap(ConstraintSettings, SortedConstraintSettings);
@@ -401,68 +471,102 @@ namespace Chaos
 		ConstraintSettings[ConstraintIndex] = InConstraintSettings;
 	}
 
-	
+
+	int32 FPBDJointConstraints::GetConstraintIsland(int32 ConstraintIndex) const
+	{
+		return ConstraintStates[ConstraintIndex].Island;
+	}
+
+
 	int32 FPBDJointConstraints::GetConstraintLevel(int32 ConstraintIndex) const
 	{
 		return ConstraintStates[ConstraintIndex].Level;
 	}
 
-	
-	void FPBDJointConstraints::SetParticleLevels(int32 ConstraintIndex, const TVector<int32, 2>& ParticleLevels)
+
+	int32 FPBDJointConstraints::GetConstraintColor(int32 ConstraintIndex) const
 	{
-		int32 NewLevel = FMath::Min(ParticleLevels[0], ParticleLevels[1]);
-		int32 PreviousLevel = ConstraintStates[ConstraintIndex].Level;
-		ConstraintStates[ConstraintIndex].Level = NewLevel;
-		ConstraintStates[ConstraintIndex].ParticleLevels = ParticleLevels;
-		bRequiresSort = bRequiresSort || (NewLevel != PreviousLevel);
+		return ConstraintStates[ConstraintIndex].Color;
 	}
-	
+
+
+	int32 FPBDJointConstraints::GetConstraintBatch(int32 ConstraintIndex) const
+	{
+		return ConstraintStates[ConstraintIndex].Batch;
+	}
+
 
 	void FPBDJointConstraints::UpdatePositionBasedState(const FReal Dt)
 	{
-		if (bRequiresSort)
+		if (bJointsDirty || (bIsBatched != bChaos_Joint_Batching))
 		{
-			SortConstraints();
-			bRequiresSort = false;
+			DeinitSolverJointData();
+
+			BatchConstraints();
+
+			InitSolverJointData();
+
+			bIsBatched = bChaos_Joint_Batching;
+			bJointsDirty = false;
 		}
 	}
 
 
 	void FPBDJointConstraints::PrepareConstraints(FReal Dt)
 	{
-		ConstraintSolvers.SetNum(NumConstraints());
-		for (int32 ConstraintIndex = 0; ConstraintIndex < NumConstraints(); ++ConstraintIndex)
-		{
-			const FPBDJointSettings& JointSettings = ConstraintSettings[ConstraintIndex];
-			const FTransformPair& JointFrames = ConstraintFrames[ConstraintIndex];
-			FJointSolverGaussSeidel& Solver = ConstraintSolvers[ConstraintIndex];
+		SCOPE_CYCLE_COUNTER(STAT_Joints_Prepare);
 
-			int32 Index0, Index1;
-			GetConstrainedParticleIndices(ConstraintIndex, Index0, Index1);
-			TGenericParticleHandle<FReal, 3> Particle0 = TGenericParticleHandle<FReal, 3>(ConstraintParticles[ConstraintIndex][Index0]);
-			TGenericParticleHandle<FReal, 3> Particle1 = TGenericParticleHandle<FReal, 3>(ConstraintParticles[ConstraintIndex][Index1]);
+		if (bChaos_Joint_Batching)
+		{
+			InitSolverJointState();
+		}
+		else
+		{
+			ConstraintSolvers.SetNum(NumConstraints());
 
-			Solver.Init(
-				Dt,
-				Settings,
-				JointSettings,
-				FParticleUtilitiesXR::GetCoMWorldPosition(Particle0),	// Prev position
-				FParticleUtilitiesXR::GetCoMWorldPosition(Particle1),	// Prev position
-				FParticleUtilitiesXR::GetCoMWorldRotation(Particle0),	// Prev rotation
-				FParticleUtilitiesXR::GetCoMWorldRotation(Particle1),	// Prev rotation
-				Particle0->InvM(),
-				Particle0->InvI().GetDiagonal(),
-				Particle1->InvM(),
-				Particle1->InvI().GetDiagonal(),
-				FParticleUtilities::ParticleLocalToCoMLocal(Particle0, JointFrames[Index0]),
-				FParticleUtilities::ParticleLocalToCoMLocal(Particle1, JointFrames[Index1]));
+			for (int32 ConstraintIndex = 0; ConstraintIndex < NumConstraints(); ++ConstraintIndex)
+			{
+				const FPBDJointSettings& JointSettings = ConstraintSettings[ConstraintIndex];
+
+				const FTransformPair& JointFrames = ConstraintFrames[ConstraintIndex];
+				FJointSolverGaussSeidel& Solver = ConstraintSolvers[ConstraintIndex];
+
+				int32 Index0, Index1;
+				GetConstrainedParticleIndices(ConstraintIndex, Index0, Index1);
+				TGenericParticleHandle<FReal, 3> Particle0 = TGenericParticleHandle<FReal, 3>(ConstraintParticles[ConstraintIndex][Index0]);
+				TGenericParticleHandle<FReal, 3> Particle1 = TGenericParticleHandle<FReal, 3>(ConstraintParticles[ConstraintIndex][Index1]);
+
+				Solver.Init(
+					Dt,
+					Settings,
+					JointSettings,
+					FParticleUtilitiesXR::GetCoMWorldPosition(Particle0),	// Prev position
+					FParticleUtilitiesXR::GetCoMWorldPosition(Particle1),	// Prev position
+					FParticleUtilitiesXR::GetCoMWorldRotation(Particle0),	// Prev rotation
+					FParticleUtilitiesXR::GetCoMWorldRotation(Particle1),	// Prev rotation
+					Particle0->InvM(),
+					Particle0->InvI().GetDiagonal(),
+					Particle1->InvM(),
+					Particle1->InvI().GetDiagonal(),
+					FParticleUtilities::ParticleLocalToCoMLocal(Particle0, JointFrames[Index0]),
+					FParticleUtilities::ParticleLocalToCoMLocal(Particle1, JointFrames[Index1]));
+			}
 		}
 	}
 
 
 	void FPBDJointConstraints::UnprepareConstraints(FReal Dt)
 	{
-		ConstraintSolvers.Empty();
+		SCOPE_CYCLE_COUNTER(STAT_Joints_Unprepare);
+
+		if (bChaos_Joint_Batching)
+		{
+			DeinitSolverJointState();
+		}
+		else
+		{
+			ConstraintSolvers.Empty();
+		}
 	}
 
 	
@@ -493,53 +597,63 @@ namespace Chaos
 
 	bool FPBDJointConstraints::Apply(const FReal Dt, const int32 It, const int32 NumIts)
 	{
-		SCOPE_CYCLE_COUNTER(STAT_Joints_Apply);
-
 		if (PreApplyCallback != nullptr)
 		{
 			PreApplyCallback(Dt, Handles);
 		}
 
-		FJointSolverResult NetResult;
+		int32 NumActive = 0;
 		if (Settings.ApplyPairIterations > 0)
 		{
-			for (int32 ConstraintIndex = 0; ConstraintIndex < NumConstraints(); ++ConstraintIndex)
+			if (bChaos_Joint_Batching)
+			{
+				SCOPE_CYCLE_COUNTER(STAT_Joints_ApplyBatched);
+				for (int32 BatchIndex = 0; BatchIndex < JointBatches.Num(); ++BatchIndex)
+				{
+					NumActive += ApplyBatch(Dt, BatchIndex, Settings.ApplyPairIterations, It, NumIts);
+				}
+			}
+			else
 			{
-				NetResult += SolvePosition_GaussSiedel(Dt, ConstraintIndex, Settings.ApplyPairIterations, It, NumIts);
+				SCOPE_CYCLE_COUNTER(STAT_Joints_Apply);
+				for (int32 ConstraintIndex = 0; ConstraintIndex < NumConstraints(); ++ConstraintIndex)
+				{
+					NumActive += ApplySingle(Dt, ConstraintIndex, Settings.ApplyPairIterations, It, NumIts);
+				}
 			}
 		}
 
-		UE_LOG(LogChaosJoint, Verbose, TEXT("Apply Iteration: %d / %d; Active: %d / %d"), It, NumIts, NetResult.GetNumActive(), NetResult.GetNumActive() + NetResult.GetNumSolved());
+		UE_LOG(LogChaosJoint, Verbose, TEXT("Apply Iteration: %d / %d; Active: %d"), It, NumIts, NumActive);
 
 		if (PostApplyCallback != nullptr)
 		{
 			PostApplyCallback(Dt, Handles);
 		}
 
-		return (NetResult.GetNumActive() > 0);
+		return (NumActive > 0);
 	}
 
 	bool FPBDJointConstraints::ApplyPushOut(const FReal Dt, const int32 It, const int32 NumIts)
 	{
 		SCOPE_CYCLE_COUNTER(STAT_Joints_ApplyPushOut);
 
-		FJointSolverResult NetResult;
+		int32 NumActive = 0;
 		if (Settings.ApplyPushOutPairIterations > 0)
 		{
 			for (int32 ConstraintIndex = 0; ConstraintIndex < NumConstraints(); ++ConstraintIndex)
 			{
-				NetResult += ProjectPosition_GaussSiedel(Dt, ConstraintIndex, Settings.ApplyPushOutPairIterations, It, NumIts);
+				NumActive += ApplyPushOutSingle(Dt, ConstraintIndex, Settings.ApplyPushOutPairIterations, It, NumIts);
 			}
 		}
 
-		UE_LOG(LogChaosJoint, Verbose, TEXT("PushOut Iteration: %d / %d; Active: %d / %d"), It, NumIts, NetResult.GetNumActive(), NetResult.GetNumActive() + NetResult.GetNumSolved());
+		UE_LOG(LogChaosJoint, Verbose, TEXT("PushOut Iteration: %d / %d; Active: %d"), It, NumIts, NumActive);
 
 		if (PostProjectCallback != nullptr)
 		{
 			PostProjectCallback(Dt, Handles);
 		}
 
-		return (NetResult.GetNumActive() > 0);
+		return (NumActive > 0);
 	}
 
 	//////////////////////////////////////////////////////////////////////////
@@ -577,12 +691,12 @@ namespace Chaos
 		}
 
 
-		FJointSolverResult NetResult;
+		int32 NumActive = 0;
 		if (Settings.ApplyPairIterations > 0)
 		{
 			for (FConstraintContainerHandle* ConstraintHandle : SortedConstraintHandles)
 			{
-				NetResult += SolvePosition_GaussSiedel(Dt, ConstraintHandle->GetConstraintIndex(), Settings.ApplyPairIterations, It, NumIts);
+				NumActive += ApplySingle(Dt, ConstraintHandle->GetConstraintIndex(), Settings.ApplyPairIterations, It, NumIts);
 			}
 		}
 
@@ -591,7 +705,7 @@ namespace Chaos
 			PostApplyCallback(Dt, SortedConstraintHandles);
 		}
 
-		return (NetResult.GetNumActive() > 0);
+		return (NumActive > 0);
 	}
 
 	
@@ -606,12 +720,12 @@ namespace Chaos
 				return L.GetConstraintLevel() < R.GetConstraintLevel();
 			});
 
-		FJointSolverResult NetResult;
+		int32 NumActive = 0;
 		if (Settings.ApplyPushOutPairIterations > 0)
 		{
 			for (FConstraintContainerHandle* ConstraintHandle : SortedConstraintHandles)
 			{
-				NetResult += ProjectPosition_GaussSiedel(Dt, ConstraintHandle->GetConstraintIndex(), Settings.ApplyPushOutPairIterations, It, NumIts);
+				NumActive += ApplyPushOutSingle(Dt, ConstraintHandle->GetConstraintIndex(), Settings.ApplyPushOutPairIterations, It, NumIts);
 			}
 		}
 
@@ -620,7 +734,7 @@ namespace Chaos
 			PostProjectCallback(Dt, SortedConstraintHandles);
 		}
 
-		return (NetResult.GetNumActive() > 0);
+		return (NumActive > 0);
 	}
 
 	//////////////////////////////////////////////////////////////////////////
@@ -643,10 +757,10 @@ namespace Chaos
 			FParticleUtilities::SetCoMWorldTransform(Rigid, P, Q);
 			if (bUpdateVelocity && (Dt > SMALL_NUMBER))
 			{
-				const FVec3 DV = FVec3::CalculateVelocity(PrevP, P, Dt);
-				const FVec3 DW = FRotation3::CalculateAngularVelocity(PrevQ, Q, Dt);
-				Rigid->SetV(Rigid->V() + DV);
-				Rigid->SetW(Rigid->W() + DW);
+				const FVec3 V = FVec3::CalculateVelocity(PrevP, P, Dt);
+				const FVec3 W = FRotation3::CalculateAngularVelocity(PrevQ, Q, Dt);
+				Rigid->SetV(V);
+				Rigid->SetW(W);
 			}
 		}
 	}
@@ -662,10 +776,234 @@ namespace Chaos
 		}
 	}
 
+	void FPBDJointConstraints::InitSolverJointData()
+	{
+		SolverConstraints.SetNum(NumConstraints());
+		for (int32 JointIndex = 0; JointIndex < NumConstraints(); ++JointIndex)
+		{
+			const FPBDJointSettings& JointSettings = ConstraintSettings[JointIndex];
+			SolverConstraints[JointIndex].SetJointIndex(JointIndex);
+			SolverConstraints[JointIndex].AddPositionConstraints(SolverConstraintRowDatas, Settings, JointSettings);
+		}
+		for (int32 JointIndex = 0; JointIndex < NumConstraints(); ++JointIndex)
+		{
+			const FPBDJointSettings& JointSettings = ConstraintSettings[JointIndex];
+			SolverConstraints[JointIndex].AddRotationConstraints(SolverConstraintRowDatas, Settings, JointSettings);
+		}
+	}
+
+	void FPBDJointConstraints::DeinitSolverJointData()
+	{
+		SolverConstraints.Empty();
+		SolverConstraintRowDatas.Empty();
+	}
+
+	void FPBDJointConstraints::InitSolverJointState()
+	{
+		SolverConstraintRowStates.SetNum(SolverConstraintRowDatas.Num());
+		SolverConstraintStates.SetNum(NumConstraints());
+
+		for (int32 JointIndex = 0; JointIndex < NumConstraints(); ++JointIndex)
+		{
+			FJointSolverJointState& JointState = SolverConstraintStates[JointIndex];
+			const FPBDJointSettings& JointSettings = ConstraintSettings[JointIndex];
+		
+			const FTransformPair& JointFrames = ConstraintFrames[JointIndex];
+			int32 Index0, Index1;
+			GetConstrainedParticleIndices(JointIndex, Index0, Index1);
+			TGenericParticleHandle<FReal, 3> Particle0 = TGenericParticleHandle<FReal, 3>(ConstraintParticles[JointIndex][Index0]);
+			TGenericParticleHandle<FReal, 3> Particle1 = TGenericParticleHandle<FReal, 3>(ConstraintParticles[JointIndex][Index1]);
+
+			JointState.Init(
+				Settings,
+				JointSettings,
+				FParticleUtilitiesXR::GetCoMWorldPosition(Particle0),
+				FParticleUtilitiesXR::GetCoMWorldRotation(Particle0),
+				FParticleUtilitiesXR::GetCoMWorldPosition(Particle1),
+				FParticleUtilitiesXR::GetCoMWorldRotation(Particle1),
+				Particle0->InvM(),
+				Particle0->InvI().GetDiagonal(),
+				Particle1->InvM(),
+				Particle1->InvI().GetDiagonal(),
+				FParticleUtilities::ParticleLocalToCoMLocal(Particle0, JointFrames[Index0]),
+				FParticleUtilities::ParticleLocalToCoMLocal(Particle1, JointFrames[Index1]));
+		}
+	}
+
+	void FPBDJointConstraints::DeinitSolverJointState()
+	{
+		SolverConstraintRowStates.Empty();
+		SolverConstraintStates.Empty();
+	}
+
+	void FPBDJointConstraints::GatherSolverJointState(int32 JointIndex)
+	{
+		FJointSolverJointState& JointState = SolverConstraintStates[JointIndex];
+
+		int32 Index0, Index1;
+		GetConstrainedParticleIndices(JointIndex, Index0, Index1);
+		TGenericParticleHandle<FReal, 3> Particle0 = TGenericParticleHandle<FReal, 3>(ConstraintParticles[JointIndex][Index0]);
+		TGenericParticleHandle<FReal, 3> Particle1 = TGenericParticleHandle<FReal, 3>(ConstraintParticles[JointIndex][Index1]);
+
+		JointState.Update(
+			FParticleUtilities::GetCoMWorldPosition(Particle0),
+			FParticleUtilities::GetCoMWorldRotation(Particle0),
+			FParticleUtilities::GetCoMWorldPosition(Particle1),
+			FParticleUtilities::GetCoMWorldRotation(Particle1));
+	}
+
+	void FPBDJointConstraints::ScatterSolverJointState(const FReal Dt, int32 JointIndex)
+	{
+		int32 Index0, Index1;
+		GetConstrainedParticleIndices(JointIndex, Index0, Index1);
+		TPBDRigidParticleHandle<FReal, 3>* Particle0 = ConstraintParticles[JointIndex][Index0]->CastToRigidParticle();
+		TPBDRigidParticleHandle<FReal, 3>* Particle1 = ConstraintParticles[JointIndex][Index1]->CastToRigidParticle();
+
+		FJointSolverJointState& JointState = SolverConstraintStates[JointIndex];
+		UpdateParticleState(Particle0, Dt, JointState.PrevPs[0], JointState.PrevQs[0], JointState.Ps[0], JointState.Qs[0]);
+		UpdateParticleState(Particle1, Dt, JointState.PrevPs[1], JointState.PrevQs[1], JointState.Ps[1], JointState.Qs[1]);
+	}
+
+	int32 FPBDJointConstraints::ApplyBatch(const FReal Dt, const int32 BatchIndex, const int32 NumPairIts, const int32 It, const int32 NumIts)
+	{
+		UE_LOG(LogChaosJoint, VeryVerbose, TEXT("Solve Joint Batch %d %d-%d (dt = %f; it = %d / %d)"), BatchIndex, JointBatches[BatchIndex][0], JointBatches[BatchIndex][1], Dt, It, NumIts);
+
+		int32 NumActive = 0;
+
+		// The range of joints in the batch
+		const int32 JointIndexBegin = JointBatches[BatchIndex][0];
+		const int32 JointIndexEnd = JointBatches[BatchIndex][1];
+		if (JointIndexEnd <= JointIndexBegin)
+		{
+			return NumActive;
+		}
+
+		// Initialize the state for each joint in the batch (body CoM position, inertias, etc)
+		for (int32 JointIndex = JointIndexBegin; JointIndex < JointIndexEnd; ++JointIndex)
+		{
+			GatherSolverJointState(JointIndex);
+		}
+
+		FJointSolverResult NetResult;
+		for (int32 PairIt = 0; PairIt < NumPairIts; ++PairIt)
+		{
+			UE_LOG(LogChaosJoint, VeryVerbose, TEXT("  Pair Iteration %d / %d"), PairIt, NumPairIts);
+
+			// Reset accumulators and update derived state
+			if (bChaos_Joint_ISPC_Enabled)
+			{
+#if INTEL_ISPC
+				ispc::BatchUpdateDerivedState(
+					(ispc::FJointSolverJointState*)SolverConstraintStates.GetData(), 
+					JointBatches[BatchIndex][0], 
+					JointBatches[BatchIndex][1]);
+#endif
+			}
+			else
+			{
+				for (int32 JointIndex = JointIndexBegin; JointIndex < JointIndexEnd; ++JointIndex)
+				{
+					SolverConstraintStates[JointIndex].UpdateDerivedState();
+				}
+			}
+
+			// Update the position constraint axes and errors for all Joints in the batch
+			for (int32 JointIndex = JointIndexBegin; JointIndex < JointIndexEnd; ++JointIndex)
+			{
+				SolverConstraints[JointIndex].UpdatePositionConstraints(
+					SolverConstraintRowDatas, 
+					SolverConstraintRowStates, 
+					SolverConstraintStates[JointIndex], 
+					ConstraintSettings[JointIndex]);
+			}
+
+			// Solve and apply the position constraints for all Joints in the batch
+			if (JointIndexEnd > JointIndexBegin)
+			{
+				const int32 RowIndexBegin = SolverConstraints[JointIndexBegin].GetLinearRowIndexBegin();
+				const int32 RowIndexEnd = SolverConstraints[JointIndexEnd - 1].GetLinearRowIndexEnd();
+				NumActive += FJointSolver::ApplyPositionConstraints(
+					Dt, 
+					SolverConstraintStates, 
+					SolverConstraintRowDatas, 
+					SolverConstraintRowStates, 
+					JointIndexBegin,
+					JointIndexEnd,
+					RowIndexBegin,
+					RowIndexEnd);
+			}
+
+			// Reset accumulators and update derived state
+			if (bChaos_Joint_ISPC_Enabled)
+			{
+#if INTEL_ISPC
+				ispc::BatchUpdateDerivedState(
+					(ispc::FJointSolverJointState*)SolverConstraintStates.GetData(), 
+					JointBatches[BatchIndex][0], 
+					JointBatches[BatchIndex][1]);
+#endif
+			}
+			else
+			{
+				for (int32 JointIndex = JointIndexBegin; JointIndex < JointIndexEnd; ++JointIndex)
+				{
+					SolverConstraintStates[JointIndex].UpdateDerivedState();
+				}
+			}
+
+			// Update the rotation constraint axes and errors for all Joints in the batch
+			for (int32 JointIndex = JointIndexBegin; JointIndex < JointIndexEnd; ++JointIndex)
+			{
+				SolverConstraints[JointIndex].UpdateRotationConstraints(
+					SolverConstraintRowDatas, 
+					SolverConstraintRowStates, 
+					SolverConstraintStates[JointIndex], 
+					ConstraintSettings[JointIndex]);
+			}
+
+			// Solve and apply the rotation constraints for all Joints in the batch
+			const int32 RowIndexBegin = SolverConstraints[JointIndexBegin].GetAngularRowIndexBegin();
+			const int32 RowIndexEnd = SolverConstraints[JointIndexEnd - 1].GetAngularRowIndexEnd();
+			if (bChaos_Joint_ISPC_Enabled)
+			{
+#if INTEL_ISPC
+				ispc::BatchApplyRotationConstraints(
+					Dt, 
+					(ispc::FJointSolverJointState*)SolverConstraintStates.GetData(), 
+					(ispc::FJointSolverConstraintRowData*)SolverConstraintRowDatas.GetData(),
+					(ispc::FJointSolverConstraintRowState*)SolverConstraintRowStates.GetData(),
+					JointIndexBegin,
+					JointIndexEnd,
+					RowIndexBegin,
+					RowIndexEnd);
+#endif
+			}
+			else
+			{
+				NumActive += FJointSolver::ApplyRotationConstraints(
+					Dt, 
+					SolverConstraintStates, 
+					SolverConstraintRowDatas, 
+					SolverConstraintRowStates, 
+					JointIndexBegin,
+					JointIndexEnd,
+					RowIndexBegin,
+					RowIndexEnd);
+			}
+		}
+
+		// Copy the updated state back to the bodies
+		for (int32 JointIndex = JointIndexBegin; JointIndex < JointIndexEnd; ++JointIndex)
+		{
+			ScatterSolverJointState(Dt, JointIndex);
+		}
+
+		return NumActive;
+	}
 
 	// This position solver iterates over each of the inner constraints (position, twist, swing) and solves them independently.
 	// This will converge slowly in some cases, particularly where resolving angular constraints violates position constraints and vice versa.
-	FJointSolverResult FPBDJointConstraints::SolvePosition_GaussSiedel(const FReal Dt, const int32 ConstraintIndex, const int32 NumPairIts, const int32 It, const int32 NumIts)
+	int32 FPBDJointConstraints::ApplySingle(const FReal Dt, const int32 ConstraintIndex, const int32 NumPairIts, const int32 It, const int32 NumIts)
 	{
 		const TVector<TGeometryParticleHandle<FReal, 3>*, 2>& Constraint = ConstraintParticles[ConstraintIndex];
 		UE_LOG(LogChaosJoint, VeryVerbose, TEXT("Solve Joint Constraint %d %s %s (dt = %f; it = %d / %d)"), ConstraintIndex, *Constraint[0]->ToString(), *Constraint[1]->ToString(), Dt, It, NumIts);
@@ -694,27 +1032,27 @@ namespace Chaos
 			Particle1->V(),
 			Particle1->W());
 
-		FJointSolverResult NetResult;
+		int32 NumActive = 0;
 		for (int32 PairIt = 0; PairIt < NumPairIts; ++PairIt)
 		{
-			UE_LOG(LogChaosJoint, VeryVerbose, TEXT("  Pair Iteration %d / %d"), PairIt, NumIts);
+			UE_LOG(LogChaosJoint, VeryVerbose, TEXT("  Pair Iteration %d / %d"), PairIt, NumPairIts);
 
-			NetResult += Solver.ApplyConstraints(Dt, Settings, JointSettings);
-			NetResult +=  Solver.ApplyDrives(Dt, Settings, JointSettings);
+			NumActive += Solver.ApplyConstraints(Dt, Settings, JointSettings);
+			NumActive +=  Solver.ApplyDrives(Dt, Settings, JointSettings);
 
-			if ((NetResult.GetNumActive() == 0) && bChaos_Joint_EarlyOut_Enabled)
+			if ((NumActive == 0) && bChaos_Joint_EarlyOut_Enabled)
 			{
 				break;
 			}
 		}
 
-		UpdateParticleState(Particle0->CastToRigidParticle(), Dt, P0, Q0, Solver.GetP(0), Solver.GetQ(0));
-		UpdateParticleState(Particle1->CastToRigidParticle(), Dt, P1, Q1, Solver.GetP(1), Solver.GetQ(1));
+		UpdateParticleState(Particle0->CastToRigidParticle(), Dt, Solver.GetPrevP(0), Solver.GetPrevQ(0), Solver.GetP(0), Solver.GetQ(0));
+		UpdateParticleState(Particle1->CastToRigidParticle(), Dt, Solver.GetPrevP(1), Solver.GetPrevQ(1), Solver.GetP(1), Solver.GetQ(1));
 
-		return NetResult;
+		return NumActive;
 	}
 
-	FJointSolverResult FPBDJointConstraints::ProjectPosition_GaussSiedel(const FReal Dt, const int32 ConstraintIndex, const int32 NumPairIts, const int32 It, const int32 NumIts)
+	int32 FPBDJointConstraints::ApplyPushOutSingle(const FReal Dt, const int32 ConstraintIndex, const int32 NumPairIts, const int32 It, const int32 NumIts)
 	{
 		const TVector<TGeometryParticleHandle<FReal, 3>*, 2>& Constraint = ConstraintParticles[ConstraintIndex];
 		UE_LOG(LogChaosJoint, VeryVerbose, TEXT("Project Joint Constraint %d %s %s (dt = %f; it = %d / %d)"), ConstraintIndex, *Constraint[0]->ToString(), *Constraint[1]->ToString(), Dt, It, NumIts);
@@ -738,12 +1076,12 @@ namespace Chaos
 			Particle1->V(),
 			Particle1->W());
 
-		FJointSolverResult NetResult;
+		int32 NumActive = 0;
 		for (int32 PairIt = 0; PairIt < NumPairIts; ++PairIt)
 		{
-			NetResult = Solver.ApplyProjections(Dt, Settings, JointSettings);
+			NumActive = Solver.ApplyProjections(Dt, Settings, JointSettings);
 			
-			if ((NetResult.GetNumActive() == 0) && bChaos_Joint_EarlyOut_Enabled)
+			if ((NumActive == 0) && bChaos_Joint_EarlyOut_Enabled)
 			{
 				break;
 			}
@@ -752,7 +1090,261 @@ namespace Chaos
 		UpdateParticleStateExplicit(Particle0->CastToRigidParticle(), Dt, Solver.GetP(0), Solver.GetQ(0), Solver.GetV(0), Solver.GetW(0));
 		UpdateParticleStateExplicit(Particle1->CastToRigidParticle(), Dt, Solver.GetP(1), Solver.GetQ(1), Solver.GetV(1), Solver.GetW(1));
 
-		return NetResult;
+		return NumActive;
+	}
+
+
+	// Assign an Island, Level and Color to each constraint. Constraints must be processed in Level order, but
+	// constraints of the same color are independent and can be processed in parallel (SIMD or Task)
+	// NOTE: Constraints are the Vertices, and Edges connect constraints sharing a Particle
+	void FPBDJointConstraints::ColorConstraints()
+	{
+		// Add a Vertex for all constraints involving at least one dynamic body
+		// Maintain a map from Constraint Index to Vertex Index
+		FColoringGraph Graph;
+		TArray<int32> ConstraintVertices; // Map of ConstraintIndex -> VertexIndex
+		Graph.ReserveVertices(NumConstraints());
+		ConstraintVertices.SetNumZeroed(NumConstraints());
+		for (int32 ConstraintIndex = 0; ConstraintIndex < NumConstraints(); ++ConstraintIndex)
+		{
+			TPBDRigidParticleHandle<FReal, 3>* Particle0 = ConstraintParticles[ConstraintIndex][0]->CastToRigidParticle();
+			TPBDRigidParticleHandle<FReal, 3>* Particle1 = ConstraintParticles[ConstraintIndex][1]->CastToRigidParticle();
+
+			bool bContainsDynamic = (Particle0 != nullptr) || (Particle1 != nullptr);
+			if (bContainsDynamic)
+			{
+				ConstraintVertices[ConstraintIndex] = Graph.AddVertex();
+
+				// Set kinematic-connected constraints to level 0 to initialize level calculation
+				bool bContainsKinematic = (Particle0 == nullptr) || (Particle1 == nullptr);
+				if (bContainsKinematic)
+				{
+					Graph.SetVertexLevel(ConstraintVertices[ConstraintIndex], 0);
+				}
+			}
+			else
+			{
+				ConstraintVertices[ConstraintIndex] = INDEX_NONE;
+			}
+		}
+
+		// Also build a map of particles to constraint indices. We only care about dynamic particles since
+		// two constraints that share only a kinematic particle will not interact.
+		TMap<TPBDRigidParticleHandle<FReal, 3>*, TArray<int32>> ParticleConstraints; // Map of ParticleHandle -> Constraint Indices involving the particle
+		for (int32 ConstraintIndex = 0; ConstraintIndex < NumConstraints(); ++ConstraintIndex)
+		{
+			TPBDRigidParticleHandle<FReal, 3>* Particle0 = ConstraintParticles[ConstraintIndex][0]->CastToRigidParticle();
+			TPBDRigidParticleHandle<FReal, 3>* Particle1 = ConstraintParticles[ConstraintIndex][1]->CastToRigidParticle();
+			if (Particle0 != nullptr)
+			{
+				ParticleConstraints.FindOrAdd(Particle0).Add(ConstraintIndex);
+			}
+			if (Particle1 != nullptr)
+			{
+				ParticleConstraints.FindOrAdd(Particle1).Add(ConstraintIndex);
+			}
+		}
+
+		// Connect constraints that share a dynamic particle
+		Graph.ReserveEdges((ParticleConstraints.Num() * (ParticleConstraints.Num() - 1)) / 2);
+		for (auto& ParticleConstraintsElement : ParticleConstraints)
+		{
+			TArray<int32>& ParticleConstraintIndices = ParticleConstraintsElement.Value;
+			int32 NumParticleConstraintIndices = ParticleConstraintIndices.Num();
+			for (int32 ParticleConstraintIndex0 = 0; ParticleConstraintIndex0 < NumParticleConstraintIndices; ++ParticleConstraintIndex0)
+			{
+				int32 ConstraintIndex0 = ParticleConstraintIndices[ParticleConstraintIndex0];
+				int32 VertexIndex0 = ConstraintVertices[ConstraintIndex0];
+				for (int32 ParticleConstraintIndex1 = ParticleConstraintIndex0 + 1; ParticleConstraintIndex1 < NumParticleConstraintIndices; ++ParticleConstraintIndex1)
+				{
+					int32 ConstraintIndex1 = ParticleConstraintIndices[ParticleConstraintIndex1];
+					int32 VertexIndex1 = ConstraintVertices[ConstraintIndex1];
+					Graph.AddEdge(VertexIndex0, VertexIndex1);
+				}
+			}
+		}
+
+		// Colorize the graph
+		Graph.Islandize();
+		Graph.Levelize();
+		Graph.Colorize();
+
+		// Set the constraint colors
+		for (int32 ConstraintIndex = 0; ConstraintIndex < NumConstraints(); ++ConstraintIndex)
+		{
+			int32 VertexIndex = ConstraintVertices[ConstraintIndex];
+			ConstraintStates[ConstraintIndex].Island = Graph.GetVertexIsland(VertexIndex);
+			ConstraintStates[ConstraintIndex].IslandSize = Graph.GetVertexIslandSize(VertexIndex);
+			ConstraintStates[ConstraintIndex].Level = Graph.GetVertexLevel(VertexIndex);
+			ConstraintStates[ConstraintIndex].Color = Graph.GetVertexColor(VertexIndex);
+		}
+	}
+
+	// Assign constraints to batches based on Level and Color. A batch is all constraints that shared the same Level-Color and so may be processed in parallel.
+	// NOTE: some constraints may have no dynamic bodies and therefore should be ignored (They will have Level = 0 and Color = -1).
+	// @todo(ccaulfield): eliminate all the sorting (just use indices until we have the final batch ordering and then sort the actual constraint list)
+	void FPBDJointConstraints::BatchConstraints()
+	{
+		// Reset
+		for (int32 ConstraintIndex = 0; ConstraintIndex < NumConstraints(); ++ConstraintIndex)
+		{
+			ConstraintStates[ConstraintIndex].Island = INDEX_NONE;
+			ConstraintStates[ConstraintIndex].Level = INDEX_NONE;
+			ConstraintStates[ConstraintIndex].Color = INDEX_NONE;
+			ConstraintStates[ConstraintIndex].Batch = INDEX_NONE;
+			ConstraintStates[ConstraintIndex].IslandSize = 0;
+		}
+
+		// Assign all constraints to islands and set colors
+		ColorConstraints();
+
+		// If batching is disabled, just sort and put in one batch
+		if (!bChaos_Joint_Batching)
+		{
+			for (int32 ConstraintIndex = 0; ConstraintIndex < NumConstraints(); ++ConstraintIndex)
+			{
+				ConstraintStates[ConstraintIndex].Batch = 0;
+			}
+			JointBatches.Reset();
+			JointBatches.Add(TVector<int32, 2>(0, NumConstraints()));
+			SortConstraints();
+			return;
+		}
+
+		// Build the list of constraints per island
+		TArray<TArray<int32>> IslandConstraints;
+		for (int32 ConstraintIndex = 0; ConstraintIndex < NumConstraints(); ++ConstraintIndex)
+		{
+			int32 IslandIndex = ConstraintStates[ConstraintIndex].Island;
+			if (IslandIndex >= IslandConstraints.Num())
+			{
+				IslandConstraints.SetNum(IslandIndex + 1);
+			}
+			IslandConstraints[IslandIndex].Add(ConstraintIndex);
+		}
+
+		// For each island, sort the constraints so that the ones to process first are at the end of the list
+		// Also ensure that constraints of same color are adjacent
+		for (int32 IslandIndex = 0; IslandIndex < IslandConstraints.Num(); ++IslandIndex)
+		{
+			IslandConstraints[IslandIndex].StableSort([this](int32 L, int32 R)
+				{
+					int32 LevelL = ConstraintStates[L].Level;
+					int32 LevelR = ConstraintStates[R].Level;
+					if (LevelL != LevelR)
+					{
+						return LevelL > LevelR;
+					}
+
+					int32 ColorL = ConstraintStates[L].Color;
+					int32 ColorR = ConstraintStates[R].Color;
+					return ColorL < ColorR;
+				});
+		}
+
+		// Now assign constraints to batches of BatchSize, taking the first same-colored items from each island (which will be at the end of the island's array).
+		// This way we depopulate the larger islands first, filling batches with items from smaller islands.
+		int32 BatchSize = bChaos_Joint_MaxBatchSize;
+		int32 NumBatches = 0;
+		int32 NumItemsToBatch = NumConstraints();
+		while (NumItemsToBatch > 0)
+		{
+			// Sort the islands so that the larger ones are first
+			// @todo(ccaulfield): optimize (does it use MoveTemp?)
+			IslandConstraints.StableSort([](const TArray<int32>& L, const TArray<int32>& R)
+				{
+					return L.Num() > R.Num();
+				});
+
+			int32 NumBatchItems = 0;
+			for (int32 IslandIndex = 0; (IslandIndex < IslandConstraints.Num()) && (NumBatchItems < BatchSize); ++IslandIndex)
+			{
+				if (IslandConstraints[IslandIndex].Num() == 0)
+				{
+					// Once we hit an empty island we are done (we have sorted on island size)
+					break;
+				}
+
+				// Take all the constraints of the same level and color from this island (up to batch size).
+				int32 ConstraintIndex = IslandConstraints[IslandIndex].Last();
+				int32 IslandBatchColor = ConstraintStates[ConstraintIndex].Color;
+				while ((IslandConstraints[IslandIndex].Num() > 0) && (NumBatchItems < BatchSize))
+				{
+					ConstraintIndex = IslandConstraints[IslandIndex].Last();
+					int32 ConstraintColor = ConstraintStates[ConstraintIndex].Color;
+					if (ConstraintColor == IslandBatchColor)
+					{
+						IslandConstraints[IslandIndex].Pop(false);
+
+						ConstraintStates[ConstraintIndex].Batch = NumBatches;
+						++NumBatchItems;
+						--NumItemsToBatch;
+					}
+					else
+					{
+						break;
+					}
+				}
+			}
+			if (NumBatchItems > 0)
+			{
+				++NumBatches;
+			}
+		}
+		check(NumItemsToBatch == 0);
+
+		// Sort constraints by batch
+		SortConstraints();
+
+		// Set up the batch begin/end indices
+		JointBatches.SetNum(NumBatches);
+		int32 BatchIndex = INDEX_NONE;
+		for (int32 ConstraintIndex = 0; ConstraintIndex < NumConstraints(); ++ConstraintIndex)
+		{
+			int32 ConstraintBatchIndex = ConstraintStates[ConstraintIndex].Batch;
+			if (ConstraintBatchIndex != BatchIndex)
+			{
+				if (BatchIndex != INDEX_NONE)
+				{
+					JointBatches[BatchIndex][1] = ConstraintIndex;
+				}
+				++BatchIndex;
+				JointBatches[BatchIndex][0] = ConstraintIndex;
+			}
+		}
+		if (BatchIndex != INDEX_NONE)
+		{
+			JointBatches[BatchIndex][1] = NumConstraints();
+		}
+
+		CheckBatches();
+	}
+
+	void FPBDJointConstraints::CheckBatches()
+	{
+#if DO_CHECK
+		for (const TVector<int32, 2>& BatchRange : JointBatches)
+		{
+			// No two Constraints in a batch should operate on the same dynamic particle
+			// TODO: validate Level (i.e., all lower level particles in same Island are in a prior batch)
+			TArray<const TPBDRigidParticleHandle<FReal, 3>*> UsedParticles;
+			for (int32 ConstraintIndex = BatchRange[0]; ConstraintIndex < BatchRange[1]; ++ConstraintIndex)
+			{
+				const TPBDRigidParticleHandle<FReal, 3>* Particle0 = ConstraintParticles[ConstraintIndex][0]->CastToRigidParticle();
+				const TPBDRigidParticleHandle<FReal, 3>* Particle1 = ConstraintParticles[ConstraintIndex][1]->CastToRigidParticle();
+				if (Particle0 != nullptr)
+				{
+					ensure(!UsedParticles.Contains(Particle0));
+					UsedParticles.Add(Particle0);
+				}
+				if (Particle1 != nullptr)
+				{
+					ensure(!UsedParticles.Contains(Particle1));
+					UsedParticles.Add(Particle1);
+				}
+			}
+		}
+#endif
 	}
 }
 
diff --git a/Engine/Source/Runtime/Experimental/Chaos/Public/Chaos/ChaosDebugDraw.h b/Engine/Source/Runtime/Experimental/Chaos/Public/Chaos/ChaosDebugDraw.h
--- a/Engine/Source/Runtime/Experimental/Chaos/Public/Chaos/ChaosDebugDraw.h
+++ b/Engine/Source/Runtime/Experimental/Chaos/Public/Chaos/ChaosDebugDraw.h
@@ -16,12 +16,15 @@ namespace Chaos
 		enum class EDebugDrawJointFeature
 		{
 			None = 0,
-			CoMConnector = 1 << 0,
-			ActorConnector = 1 << 1,
-			Stretch = 1 << 2,
-			Axes = 1 << 3,
-			Level = 1 << 4,
-			Index = 1 << 5,
+			CoMConnector = 1 << 0,		// 1
+			ActorConnector = 1 << 1,	// 2
+			Stretch = 1 << 2,			// 4
+			Axes = 1 << 3,				// 8
+			Level = 1 << 4,				// 16
+			Index = 1 << 5,				// 32
+			Color = 1 << 6,				// 64
+			Batch = 1 << 7,				// 128
+			Island = 1 << 8,			// 256
 
 			Default = ActorConnector | Stretch
 		};
diff --git a/Engine/Source/Runtime/Experimental/Chaos/Public/Chaos/Collision/ParticlePairBroadPhase.h b/Engine/Source/Runtime/Experimental/Chaos/Public/Chaos/Collision/ParticlePairBroadPhase.h
--- a/Engine/Source/Runtime/Experimental/Chaos/Public/Chaos/Collision/ParticlePairBroadPhase.h
+++ b/Engine/Source/Runtime/Experimental/Chaos/Public/Chaos/Collision/ParticlePairBroadPhase.h
@@ -29,6 +29,11 @@ namespace Chaos
 		{
 		}
 
+		FReal GetCullDistance() const 
+		{
+			return CullDistance;
+		}
+
 		void SetCullDustance(const FReal InCullDistance)
 		{
 			CullDistance = InCullDistance;
diff --git a/Engine/Source/Runtime/Experimental/Chaos/Public/Chaos/Joint/ColoringGraph.h b/Engine/Source/Runtime/Experimental/Chaos/Public/Chaos/Joint/ColoringGraph.h
new file mode 100644
--- /dev/null
+++ b/Engine/Source/Runtime/Experimental/Chaos/Public/Chaos/Joint/ColoringGraph.h
@@ -0,0 +1,162 @@
+// Copyright Epic Games, Inc. All Rights Reserved.
+#pragma once
+
+#include "CoreMinimal.h"
+
+namespace Chaos
+{
+
+	/**
+	 * A graph used by the constraints system.
+	 * Supports 
+	 *	- partitioning vertices into connected groups (islands)
+	 *	- assigning a level to each vertex based on distance from "level 0" vertices
+	 *	- colorizing vertices
+	 */
+	class FColoringGraph
+	{
+	public:
+
+		int32 GetNumIslands() const
+		{
+			return Islands.Num();
+		}
+
+		int32 GetVertexIsland(int32 VertexIndex) const
+		{
+			if (VertexIndex != INDEX_NONE)
+			{
+				return Vertices[VertexIndex].Island;
+			}
+			return INDEX_NONE;
+		}
+
+		int32 GetVertexIslandSize(int32 VertexIndex) const
+		{
+			if (VertexIndex != INDEX_NONE)
+			{
+				int32 IslandIndex = Vertices[VertexIndex].Island;
+				return Islands[IslandIndex].VertexIndices.Num();
+			}
+			return INDEX_NONE;
+		}
+
+		int32 GetVertexLevel(int32 VertexIndex) const
+		{
+			if (VertexIndex != INDEX_NONE)
+			{
+				return Vertices[VertexIndex].Level;
+			}
+			return INDEX_NONE;
+		}
+
+		void SetVertexLevel(int32 VertexIndex, int32 Level)
+		{
+			if (VertexIndex != INDEX_NONE)
+			{
+				Vertices[VertexIndex].Level = Level;
+			}
+		}
+
+		int32 GetVertexColor(int32 VertexIndex) const
+		{
+			if (VertexIndex != INDEX_NONE)
+			{
+				return Vertices[VertexIndex].Color;
+			}
+			return INDEX_NONE;
+		}
+
+		void ReserveVertices(int32 InNumVertices)
+		{
+			Vertices.Reserve(InNumVertices);
+		}
+
+		void ReserveEdges(int32 InNumEdges)
+		{
+			Edges.Reserve(InNumEdges);
+		}
+
+		int32 AddVertex()
+		{
+			return Vertices.AddDefaulted();
+		}
+
+		int32 AddVertices(int32 InNumVertices)
+		{
+			Vertices.SetNum(Vertices.Num() + InNumVertices, false);
+			return Vertices.Num() - InNumVertices;
+		}
+
+		int32 AddEdge(int32 VertexIndex0, int32 VertexIndex1)
+		{
+			int32 EdgeIndex = Edges.Emplace(FEdge(VertexIndex0, VertexIndex1));
+			Vertices[VertexIndex0].EdgeIndices.Add(EdgeIndex);
+			Vertices[VertexIndex1].EdgeIndices.Add(EdgeIndex);
+			return EdgeIndex;
+		}
+
+		/**
+		 * Assign each vertex to an island. Vertcies in different islands are isolated from each other - thee
+		 * is no path which connectes them.
+		 */
+		void Islandize();
+
+		/**
+		 * Assign a "level" to all vertices. Level is the distance from a Vertex with Level = 0 (Level 0 vertices
+		 * are specified by the user with SetVertexLevel)
+		 */
+		void Levelize();
+
+		/**
+		 * Assign a color to all vertices. Vertices of the same color do not share any edges (i.e., are not directly connected).
+		 */
+		void Colorize();
+
+	private:
+		struct FVertex
+		{
+			FVertex() : Island(INDEX_NONE), Level(INDEX_NONE), Color(INDEX_NONE) {}
+
+			// @todo(ccaulfield): eliminate allocations
+			TArray<int32> EdgeIndices;
+			int32 Island;
+			int32 Level;
+			int32 Color;
+		};
+
+		struct FEdge
+		{
+			FEdge(int32 VertexIndex0, int32 VertexIndex1)
+			{
+				VertexIndices[0] = VertexIndex0;
+				VertexIndices[1] = VertexIndex1;
+			}
+
+			int32 VertexIndices[2];
+		};
+
+		struct FIsland
+		{
+			FIsland() : NextColorIndex(0) {}
+
+			// @todo(ccaulfield): eliminate allocations
+			TArray<int32> VertexIndices;
+			int32 NextColorIndex;
+		};
+
+		// Move vertices from one island into the other
+		void MergeIslands(int32 IslandIndex0, int32 IslandIndex1);
+
+		// Move all non-empty islands to lowest island indices and clip array
+		void PackIslands();
+
+		// Set the color on the vertex so it is not the same as any adjacent vertices
+		void AssignColor(int32 VertexIndex);
+
+		TArray<FVertex> Vertices;
+		TArray<FEdge> Edges;
+		TArray<FIsland> Islands;
+	};
+
+}
\ No newline at end of file
diff --git a/Engine/Source/Runtime/Experimental/Chaos/Public/Chaos/Joint/JointSolverConstraints.h b/Engine/Source/Runtime/Experimental/Chaos/Public/Chaos/Joint/JointSolverConstraints.h
new file mode 100644
--- /dev/null
+++ b/Engine/Source/Runtime/Experimental/Chaos/Public/Chaos/Joint/JointSolverConstraints.h
@@ -0,0 +1,496 @@
+// Copyright Epic Games, Inc. All Rights Reserved.
+#pragma once
+
+#include "CoreMinimal.h"
+
+#include "Chaos/PBDJointConstraintTypes.h"
+
+namespace Chaos
+{
+	class FJointSolverGaussSeidel;
+
+	// The method used to calculate the axis and angle for one of the rotation constraints in the joint
+	enum class EJointSolverConstraintUpdateType
+	{
+		None,
+		Linear_Point,
+		Linear_Spherical,
+		Linear_Cylindrical,
+		Linear_Planar,
+
+		Linear_SphericalDrive,
+		Linear_CircularDrive,
+		Linear_AxialDrive,
+
+		Angular_Twist,
+		Angular_Cone,
+		Angular_SingleLockedSwing,
+		Angular_SingleLimitedSwing,
+		Angular_DualConeSwing,
+		Angular_Locked,
+
+		Angular_TwistDrive,
+		Angular_ConeDrive,
+		Angular_SwingDrive,
+		Angular_SLerpDrive,
+	};
+
+	/**
+	 * Body and joint state required during a joint's solve (used by all sub-constraints in the joint)
+	 */
+	class FJointSolverJointState
+	{
+	public:
+		static const int32 MaxConstrainedBodies = 2;
+
+		void Init(
+			const FPBDJointSolverSettings& SolverSettings,
+			const FPBDJointSettings& JointSettings,
+			const FVec3& PrevP0,
+			const FRotation3& PrevQ0,
+			const FVec3& PrevP1,
+			const FRotation3& PrevQ1,
+			const FReal InvM0,
+			const FVec3& InvIL0,
+			const FReal InvM1,
+			const FVec3& InvIL1,
+			const FRigidTransform3& XL0,
+			const FRigidTransform3& XL1);
+
+		FORCEINLINE void Update(
+			const FVec3& P0,
+			const FRotation3& Q0,
+			const FVec3& P1,
+			const FRotation3& Q1);
+
+		FORCEINLINE void InitDerivedState();
+
+		FORCEINLINE void UpdateDerivedState();
+
+		FORCEINLINE void ApplyDelta(
+			const FVec3& DP0,
+			const FVec3& DR0,
+			const FVec3& DP1,
+			const FVec3& DR1);
+
+		FORCEINLINE void ApplyRotationDelta(
+			const FVec3& DR0,
+			const FVec3& DR1);
+
+		// Local-space constraint settings
+		FRigidTransform3 XLs[MaxConstrainedBodies];	// Local-space joint connector transforms
+		FVec3 InvILs[MaxConstrainedBodies];			// Local-space inverse inertias
+		FReal InvMs[MaxConstrainedBodies];			// Inverse masses
+
+		// World-space constraint state
+		FVec3 Xs[MaxConstrainedBodies];				// World-space joint connector positions
+		FRotation3 Rs[MaxConstrainedBodies];		// World-space joint connector rotations
+
+		// World-space body state
+		FVec3 Ps[MaxConstrainedBodies];				// World-space particle CoM positions
+		FRotation3 Qs[MaxConstrainedBodies];		// World-space particle CoM rotations
+		FMatrix33 InvIs[MaxConstrainedBodies];		// World-space inverse inertias
+
+		// XPBD initial world-space body state (start of each tick, not each sub-tick iteration)
+		FVec3 PrevPs[MaxConstrainedBodies];			// World-space particle CoM positions
+		FRotation3 PrevQs[MaxConstrainedBodies];	// World-space particle CoM rotations
+		FVec3 PrevXs[MaxConstrainedBodies];			// World-space joint connector positions
+
+		FVec3 DPs[MaxConstrainedBodies];
+		FVec3 DRs[MaxConstrainedBodies];
+
+		FReal PositionTolerance;					// Distance error below which we consider a constraint or drive solved
+		FReal AngleTolerance;						// Angle error below which we consider a constraint or drive solved
+	};
+
+	/**
+	 * Transient state for a single sub-constraint in a joint
+	 */
+	class FJointSolverConstraintRowState
+	{
+	public:
+		static const int32 MaxConstrainedBodies = 2;
+
+		FJointSolverConstraintRowState()
+			: DPs{ FVec3(0), FVec3(0) }
+			, DRs{ FVec3(0), FVec3(0) }
+			, Axis(FVec3(0))
+			, Error(0)
+			, Lambda(0)
+		{
+		}
+
+		// Reset calculated values ready for next iteration. Note: Lambda is not reset here,
+		// it accumulates over the whole timestep.
+		FORCEINLINE void IterationReset();
+
+		// Calculate the error to correct, given the current position and limit
+		FORCEINLINE void CalculateError(FReal Position, FReal Limit);
+
+		// Per-Iteration Transient
+		FVec3 DPs[MaxConstrainedBodies];
+		FVec3 DRs[MaxConstrainedBodies];
+		FVec3 Axis;
+		FReal Error;
+
+		// Per-Tick Transient
+		FReal Lambda;
+	};
+
+	/**
+	 * A single sub-constraint in a Joint (e.g., Twist, Swing, etc.)
+	 */
+	class FJointSolverConstraintRowData
+	{
+	public:
+		static const int32 MaxConstrainedBodies = 2;
+
+		FJointSolverConstraintRowData()
+			: UpdateType(EJointSolverConstraintUpdateType::None)
+			, JointIndex(INDEX_NONE)
+			, ConstraintIndex(INDEX_NONE)
+			, NumRows(0)
+			, Stiffness(0)
+			, Damping(0)
+			, Limit(0)
+			, bIsAccelerationMode(false)
+			, bIsSoft(false)
+		{
+		}
+
+		EJointSolverConstraintUpdateType UpdateType;	// Method to use to calculate axis, errors, etc.
+		int32 JointIndex;			// Index into the outer container's joint array
+		int32 ConstraintIndex;		// Context dependent: X,Y,Z or Twist,Swing1,Swing2 for linear and angular rows
+		int32 NumRows;				// Number of rows (including this one) to solve simultaneously
+		FReal Stiffness;			// PBD or XPBD stiffness
+		FReal Damping;				// XPBD damping
+		FReal Limit;
+		bool bIsAccelerationMode;
+		bool bIsSoft;
+	};
+
+
+	/**
+	 * All the constraints in a joint
+	 */
+	class FJointSolverConstraints
+	{
+	public:
+		friend class FJointSolverGaussSeidel;
+
+		FJointSolverConstraints();
+
+		FORCEINLINE void SetJointIndex(int32 InJointIndex)
+		{
+			JointIndex = InJointIndex;
+		}
+
+		FORCEINLINE int32 GetJointIndex() const
+		{
+			return JointIndex;
+		}
+
+		FORCEINLINE int32 GetLinearRowIndexBegin() const
+		{
+			return LinearRowIndexBegin;
+		}
+
+		FORCEINLINE int32 GetLinearRowIndexEnd() const
+		{
+			return LinearRowIndexEnd;
+		}
+
+		FORCEINLINE int32 GetAngularRowIndexBegin() const
+		{
+			return AngularRowIndexBegin;
+		}
+
+		FORCEINLINE int32 GetAngularRowIndexEnd() const
+		{
+			return AngularRowIndexEnd;
+		}
+
+		FORCEINLINE int32 NumLinearConstraints() const 
+		{
+			return LinearRowIndexEnd - LinearRowIndexBegin;
+		}
+
+		FORCEINLINE int32 NumAngularConstraints() const
+		{
+			return AngularRowIndexEnd - AngularRowIndexBegin;
+		}
+
+		void AddPositionConstraints(
+			TArray<FJointSolverConstraintRowData>& RowDatas,
+			const FPBDJointSolverSettings& SolverSettings,
+			const FPBDJointSettings& JointSettings);
+
+		void AddRotationConstraints(
+			TArray<FJointSolverConstraintRowData>& RowDatas,
+			const FPBDJointSolverSettings& SolverSettings,
+			const FPBDJointSettings& JointSettings);
+
+		// Update the axes and positions of all position constraints
+		void UpdatePositionConstraints(
+			const TArray<FJointSolverConstraintRowData>& RowDatas,
+			TArray<FJointSolverConstraintRowState>& RowStates,
+			const FJointSolverJointState& JointState,
+			const FPBDJointSettings& JointSettings);
+
+		// Update the axes and positions of all rotation constraints
+		void UpdateRotationConstraints(
+			const TArray<FJointSolverConstraintRowData>& RowDatas,
+			TArray<FJointSolverConstraintRowState>& RowStates,
+			const FJointSolverJointState& JointState,
+			const FPBDJointSettings& JointSettings);
+
+	private:
+
+		void AddPointPositionConstraint(
+			TArray<FJointSolverConstraintRowData>& RowDatas,
+			const FPBDJointSolverSettings& SolverSettings,
+			const FPBDJointSettings& JointSettings);
+
+		void AddSphericalPositionConstraint(
+			TArray<FJointSolverConstraintRowData>& RowDatas,
+			const FPBDJointSolverSettings& SolverSettings,
+			const FPBDJointSettings& JointSettings);
+
+		void AddCylindricalPositionConstraint(
+			TArray<FJointSolverConstraintRowData>& RowDatas,
+			const int32 AxisIndex,
+			const EJointMotionType AxialMotion,
+			const EJointMotionType RadialMotion,
+			const FPBDJointSolverSettings& SolverSettings,
+			const FPBDJointSettings& JointSettings);
+
+		void AddPlanarPositionConstraint(
+			TArray<FJointSolverConstraintRowData>& RowDatas,
+			const int32 AxisIndex,
+			const EJointMotionType AxialMotion,
+			const FPBDJointSolverSettings& SolverSettings,
+			const FPBDJointSettings& JointSettings);
+
+		void AddPositionDrive(
+			TArray<FJointSolverConstraintRowData>& RowDatas,
+			const FPBDJointSolverSettings& SolverSettings,
+			const FPBDJointSettings& JointSettings,
+			const EJointSolverConstraintUpdateType UpdateType,
+			const int32 AxisIndex);
+
+		void AddTwistConstraint(
+			TArray<FJointSolverConstraintRowData>& RowDatas,
+			const FPBDJointSolverSettings& SolverSettings,
+			const FPBDJointSettings& JointSettings,
+			const EJointSolverConstraintUpdateType Type,
+			const bool bUseSoftLimit);
+
+		void AddSwingConstraint(
+			TArray<FJointSolverConstraintRowData>& RowDatas,
+			const FPBDJointSolverSettings& SolverSettings,
+			const FPBDJointSettings& JointSettings,
+			const EJointSolverConstraintUpdateType Type,
+			const EJointAngularConstraintIndex ConstraintIndex,
+			const bool bUseSoftLimit);
+
+		void AddTwistDrive(
+			TArray<FJointSolverConstraintRowData>& RowDatas,
+			const FPBDJointSolverSettings& SolverSettings,
+			const FPBDJointSettings& JointSettings);
+
+		void AddSwingDrive(
+			TArray<FJointSolverConstraintRowData>& RowDatas,
+			const FPBDJointSolverSettings& SolverSettings,
+			const FPBDJointSettings& JointSettings,
+			const EJointSolverConstraintUpdateType Type,
+			const EJointAngularConstraintIndex ConstraintIndex);
+
+		void AddSLerpDrive(
+			TArray<FJointSolverConstraintRowData>& RowDatas,
+			const FPBDJointSolverSettings& SolverSettings,
+			const FPBDJointSettings& JointSettings);
+
+
+
+		FORCEINLINE int32 UpdatePointPositionConstraint(
+			int32 RowIndex,
+			const TArray<FJointSolverConstraintRowData>& RowDatas,
+			TArray<FJointSolverConstraintRowState>& RowStates,
+			const FPBDJointSettings& JointSettings,
+			const FVec3& X0,
+			const FVec3& X1);
+
+		FORCEINLINE int32 UpdateSphericalPositionConstraint(
+			int32 RowIndex,
+			const TArray<FJointSolverConstraintRowData>& RowDatas,
+			TArray<FJointSolverConstraintRowState>& RowStates,
+			const FPBDJointSettings& JointSettings,
+			const FVec3& X0,
+			const FVec3& X1);
+
+		FORCEINLINE int32 UpdateCylindricalPositionConstraint(
+			int32 RowIndex,
+			const TArray<FJointSolverConstraintRowData>& RowDatas,
+			TArray<FJointSolverConstraintRowState>& RowStates,
+			const FPBDJointSettings& JointSettings,
+			const FRotation3& R0,
+			const FVec3& X0,
+			const FVec3& X1);
+
+		FORCEINLINE int32 UpdatePlanarPositionConstraint(
+			int32 RowIndex,
+			const TArray<FJointSolverConstraintRowData>& RowDatas,
+			TArray<FJointSolverConstraintRowState>& RowStates,
+			const FPBDJointSettings& JointSettings,
+			const FRotation3& R0,
+			const FVec3& X0,
+			const FVec3& X1);
+
+		FORCEINLINE int32 UpdateSphericalPositionDrive(
+			int32 RowIndex,
+			const TArray<FJointSolverConstraintRowData>& RowDatas,
+			TArray<FJointSolverConstraintRowState>& RowStates,
+			const FPBDJointSettings& JointSettings,
+			const FRotation3& R0,
+			const FVec3& X0,
+			const FVec3& X1);
+
+		FORCEINLINE int32 UpdateCircularPositionDrive(
+			int32 RowIndex,
+			const TArray<FJointSolverConstraintRowData>& RowDatas,
+			TArray<FJointSolverConstraintRowState>& RowStates,
+			const FPBDJointSettings& JointSettings,
+			const FRotation3& R0,
+			const FVec3& X0,
+			const FVec3& X1);
+
+		FORCEINLINE int32 UpdateAxialPositionDrive(
+			int32 RowIndex,
+			const TArray<FJointSolverConstraintRowData>& RowDatas,
+			TArray<FJointSolverConstraintRowState>& RowStates,
+			const FPBDJointSettings& JointSettings,
+			const FRotation3& R0,
+			const FVec3& X0,
+			const FVec3& X1);
+
+		FORCEINLINE void UpdateTwistConstraint(
+			const FJointSolverConstraintRowData& RowData,
+			FJointSolverConstraintRowState& RowState,
+			const FPBDJointSettings& JointSettings,
+			const FRotation3& R1,
+			const FRotation3& RTwist);
+
+		FORCEINLINE void UpdateConeSwingConstraint(
+			const FJointSolverConstraintRowData& RowData,
+			FJointSolverConstraintRowState& RowState,
+			const FPBDJointSettings& JointSettings,
+			const FRotation3& R0,
+			const FRotation3& RSwing);
+
+		FORCEINLINE void UpdateSingleLockedSwingConstraint(
+			const FJointSolverConstraintRowData& RowData,
+			FJointSolverConstraintRowState& RowState,
+			const FPBDJointSettings& JointSettings,
+			const FRotation3& R0,
+			const FRotation3& R1);
+
+		FORCEINLINE void UpdateSingleLimitedSwingConstraint(
+			const FJointSolverConstraintRowData& RowData,
+			FJointSolverConstraintRowState& RowState,
+			const FPBDJointSettings& JointSettings,
+			const FRotation3& R0,
+			const FRotation3& RSwing);
+
+		FORCEINLINE void UpdateDualConeSwingConstraint(
+			const FJointSolverConstraintRowData& RowData,
+			FJointSolverConstraintRowState& RowState,
+			const FPBDJointSettings& JointSettings,
+			const FRotation3& R0,
+			const FRotation3& R1);
+
+		FORCEINLINE void UpdateTwistDrive(
+			const FJointSolverConstraintRowData& RowData,
+			FJointSolverConstraintRowState& RowState,
+			const FPBDJointSettings& JointSettings,
+			const FRotation3& R1,
+			const FRotation3& RTwist);
+
+		FORCEINLINE void UpdateConeSwingDrive(
+			const FJointSolverConstraintRowData& RowData,
+			FJointSolverConstraintRowState& RowState,
+			const FPBDJointSettings& JointSettings,
+			const FRotation3& R0,
+			const FRotation3& RSwing);
+
+		FORCEINLINE void UpdateSwingDrive(
+			const FJointSolverConstraintRowData& RowData,
+			FJointSolverConstraintRowState& RowState,
+			const FPBDJointSettings& JointSettings,
+			const FRotation3& R0,
+			const FRotation3& RSwing);
+
+		FORCEINLINE void UpdateSLerpDrive(
+			const FJointSolverConstraintRowData& RowData,
+			FJointSolverConstraintRowState& RowState,
+			const FPBDJointSettings& JointSettings,
+			const FRotation3& R0,
+			const FRotation3& R1);
+
+		int32 JointIndex;
+		int32 LinearRowIndexBegin;
+		int32 LinearRowIndexEnd;
+		int32 AngularRowIndexBegin;
+		int32 AngularRowIndexEnd;
+		FReal PositionTolerance;
+		FReal AngleTolerance;
+		bool bNeedSwingTwist;	// True if any of the constraints need the decomposed swing/twist to calculate angles
+		bool bNeedLockedAxes;	// True if any of the constraints need the fixed axes to calculate angles
+	};
+
+	class FJointSolver
+	{
+	public:
+		static int32 ApplyPositionConstraints(
+			const FReal Dt,
+			TArray<FJointSolverJointState>& JointStates,
+			const TArray <FJointSolverConstraintRowData>& RowDatas,
+			TArray <FJointSolverConstraintRowState>& RowStates,
+			int32 JointIndexBegin,
+			int32 JointIndexEnd,
+			int32 RowIndexBegin,
+			int32 RowIndexEnd);
+
+		static int32 ApplyRotationConstraints(
+			const FReal Dt,
+			TArray<FJointSolverJointState>& JointStates,
+			const TArray <FJointSolverConstraintRowData>& RowDatas,
+			TArray <FJointSolverConstraintRowState>& RowStates,
+			int32 JointIndexBegin,
+			int32 JointIndexEnd,
+			int32 RowIndexBegin,
+			int32 RowIndexEnd);
+
+	private:
+		static FORCEINLINE void ApplyPositionConstraint1(
+			const FReal Dt,
+			const FJointSolverJointState& JointState,
+			const FJointSolverConstraintRowData& RowData,
+			FJointSolverConstraintRowState& RowState);
+
+		static FORCEINLINE void ApplyPositionConstraint3(
+			const FReal Dt,
+			const FJointSolverJointState& JointState,
+			const FJointSolverConstraintRowData& RowData,
+			FJointSolverConstraintRowState& RowState0,
+			FJointSolverConstraintRowState& RowState1,
+			FJointSolverConstraintRowState& RowState2);
+
+		static FORCEINLINE void ApplyRotationConstraint(
+			const FReal Dt,
+			const FJointSolverJointState& JointState,
+			const FJointSolverConstraintRowData& RowData,
+			FJointSolverConstraintRowState& RowState);
+	};
+}
+
+#include "Chaos/Joint/JointSolverConstraints.inl"
diff --git a/Engine/Source/Runtime/Experimental/Chaos/Public/Chaos/Joint/JointSolverConstraints.inl b/Engine/Source/Runtime/Experimental/Chaos/Public/Chaos/Joint/JointSolverConstraints.inl
new file mode 100644
--- /dev/null
+++ b/Engine/Source/Runtime/Experimental/Chaos/Public/Chaos/Joint/JointSolverConstraints.inl
@@ -0,0 +1,667 @@
+// Copyright Epic Games, Inc. All Rights Reserved.
+#pragma once
+
+#include "CoreMinimal.h"
+
+#include "Chaos/Joint/JointSolverConstraints.h"
+#include "Chaos/PBDJointConstraintTypes.h"
+#include "Chaos/PBDJointConstraintUtilities.h"
+#include "Chaos/Utilities.h"
+
+namespace Chaos
+{
+	FORCEINLINE void FJointSolverJointState::Init(
+		const FPBDJointSolverSettings& SolverSettings,
+		const FPBDJointSettings& JointSettings,
+		const FVec3& PrevP0,
+		const FRotation3& PrevQ0,
+		const FVec3& PrevP1,
+		const FRotation3& PrevQ1,
+		const FReal InvM0,
+		const FVec3& InvIL0,
+		const FReal InvM1,
+		const FVec3& InvIL1,
+		const FRigidTransform3& XL0,
+		const FRigidTransform3& XL1)
+	{
+		XLs[0] = XL0;
+		XLs[1] = XL1;
+
+		InvILs[0] = JointSettings.ParentInvMassScale * InvIL0;
+		InvILs[1] = InvIL1;
+		InvMs[0] = JointSettings.ParentInvMassScale * InvM0;
+		InvMs[1] = InvM1;
+
+		FPBDJointUtilities::ConditionInverseMassAndInertia(InvMs[0], InvMs[1], InvILs[0], InvILs[1], SolverSettings.MinParentMassRatio, SolverSettings.MaxInertiaRatio);
+
+		PrevPs[0] = PrevP0;
+		PrevPs[1] = PrevP1;
+		PrevQs[0] = PrevQ0;
+		PrevQs[1] = PrevQ1;
+		PrevXs[0] = PrevP0 + PrevQ0 * XL0.GetTranslation();
+		PrevXs[1] = PrevP1 + PrevQ1 * XL1.GetTranslation();
+
+		PositionTolerance = SolverSettings.PositionTolerance;
+		AngleTolerance = SolverSettings.AngleTolerance;
+
+		InitDerivedState();
+	}
+
+	FORCEINLINE void FJointSolverJointState::Update(
+		const FVec3& P0,
+		const FRotation3& Q0,
+		const FVec3& P1,
+		const FRotation3& Q1)
+	{
+		Ps[0] = P0;
+		Ps[1] = P1;
+		Qs[0] = Q0;
+		Qs[1] = Q1;
+		Qs[1].EnforceShortestArcWith(Qs[0]);
+	}
+
+	FORCEINLINE void FJointSolverJointState::InitDerivedState()
+	{
+		// Really we should need to do this for Kinematics since Dynamics are updated each iteration
+		Xs[0] = PrevPs[0] + PrevQs[0] * XLs[0].GetTranslation();
+		Rs[0] = PrevQs[0] * XLs[0].GetRotation();
+		InvIs[0] = (InvMs[0] > 0.0f) ? Utilities::ComputeWorldSpaceInertia(PrevQs[0], InvILs[0]) : FMatrix33(0, 0, 0);
+
+		Xs[1] = PrevPs[1] + PrevQs[1] * XLs[1].GetTranslation();
+		Rs[1] = PrevQs[1] * XLs[1].GetRotation();
+		InvIs[1] = (InvMs[1] > 0.0f) ? Utilities::ComputeWorldSpaceInertia(PrevQs[1], InvILs[1]) : FMatrix33(0, 0, 0);
+
+		Rs[1].EnforceShortestArcWith(Rs[0]);
+
+		PrevXs[0] = Xs[0];
+		PrevXs[1] = Xs[1];
+
+		DPs[0] = FVec3(0);
+		DRs[0] = FVec3(0);
+		DPs[1] = FVec3(0);
+		DRs[1] = FVec3(0);
+	}
+
+	FORCEINLINE void FJointSolverJointState::UpdateDerivedState()
+	{
+		if (InvMs[0] > 0.0f)
+		{
+			Xs[0] = Ps[0] + Qs[0] * XLs[0].GetTranslation();
+			Rs[0] = Qs[0] * XLs[0].GetRotation();
+			InvIs[0] = Utilities::ComputeWorldSpaceInertia(Qs[0], InvILs[0]);
+			DPs[0] = FVec3(0);
+			DRs[0] = FVec3(0);
+		}
+		if (InvMs[1] > 0.0f)
+		{
+			Xs[1] = Ps[1] + Qs[1] * XLs[1].GetTranslation();
+			Rs[1] = Qs[1] * XLs[1].GetRotation();
+			InvIs[1] = Utilities::ComputeWorldSpaceInertia(Qs[1], InvILs[1]);
+			DPs[1] = FVec3(0);
+			DRs[1] = FVec3(0);
+		}
+		Rs[1].EnforceShortestArcWith(Rs[0]);
+	}
+
+	FORCEINLINE void FJointSolverJointState::ApplyDelta(
+		const FVec3& DP0,
+		const FVec3& DR0,
+		const FVec3& DP1,
+		const FVec3& DR1)
+	{
+		if (InvMs[0] > 0.0f)
+		{
+			//UE_LOG(LogChaosJoint, VeryVerbose, TEXT("      Apply DP%d %f %f %f"), 0, DP0.X, DP0.Y, DP0.Z);
+			//UE_LOG(LogChaosJoint, VeryVerbose, TEXT("      Apply DR%d %f %f %f"), 0, DR0.X, DR0.Y, DR0.Z);
+
+			Ps[0] += DP0;
+			const FRotation3 DQ0 = (FRotation3::FromElements(DR0, 0) * Qs[0]) * (FReal)0.5;
+			Qs[0] = (Qs[0] + DQ0).GetNormalized();
+		}
+		if (InvMs[1] > 0.0f)
+		{
+			//UE_LOG(LogChaosJoint, VeryVerbose, TEXT("      Apply DP%d %f %f %f"), 1, DP1.X, DP1.Y, DP1.Z);
+			//UE_LOG(LogChaosJoint, VeryVerbose, TEXT("      Apply DR%d %f %f %f"), 1, DR1.X, DR1.Y, DR1.Z);
+
+			Ps[1] += DP1;
+			const FRotation3 DQ1 = (FRotation3::FromElements(DR1, 0) * Qs[1]) * (FReal)0.5;
+			Qs[1] = (Qs[1] + DQ1).GetNormalized();
+		}
+		Qs[1].EnforceShortestArcWith(Qs[0]);
+	}
+
+	FORCEINLINE void FJointSolverJointState::ApplyRotationDelta(
+		const FVec3& DR0,
+		const FVec3& DR1)
+	{
+		if (InvMs[0] > 0.0f)
+		{
+			//UE_LOG(LogChaosJoint, VeryVerbose, TEXT("      Apply DR%d %f %f %f"), 0, DR0.X, DR0.Y, DR0.Z);
+
+			const FRotation3 DQ0 = (FRotation3::FromElements(DR0, 0) * Qs[0]) * (FReal)0.5;
+			Qs[0] = (Qs[0] + DQ0).GetNormalized();
+		}
+		if (InvMs[1] > 0.0f)
+		{
+			//UE_LOG(LogChaosJoint, VeryVerbose, TEXT("      Apply DR%d %f %f %f"), 1, DR1.X, DR1.Y, DR1.Z);
+
+			const FRotation3 DQ1 = (FRotation3::FromElements(DR1, 0) * Qs[1]) * (FReal)0.5;
+			Qs[1] = (Qs[1] + DQ1).GetNormalized();
+		}
+		Qs[1].EnforceShortestArcWith(Qs[0]);
+	}
+
+	//
+	///////////////////////////////////////////////////////////////////////
+	//
+
+	// Reset calculated values ready for next iteration. Note: Lambda is not reset here,
+	// it accumulates over the whole timestep.
+	FORCEINLINE void FJointSolverConstraintRowState::IterationReset()
+	{
+		DPs[0] = FVec3(0);
+		DPs[1] = FVec3(0);
+		DRs[0] = FVec3(0);
+		DRs[1] = FVec3(0);
+		Axis = FVec3(0);
+		Error = 0.0f;
+		// Do not reset Lambda
+	}
+
+	FORCEINLINE void FJointSolverConstraintRowState::CalculateError(FReal Position, FReal Limit)
+	{
+		Error = 0.0f;
+		if (Position >= Limit)
+		{
+			Error = Position - Limit;
+		}
+		else if (Position <= -Limit)
+		{
+			Error = Position + Limit;
+		}
+	}
+
+	//
+	//////////////////////////////////////////////////////////////////////////
+	//
+
+
+	FORCEINLINE int32 FJointSolverConstraints::UpdatePointPositionConstraint(
+		int32 RowIndexBegin,
+		const TArray<FJointSolverConstraintRowData>& RowDatas,
+		TArray<FJointSolverConstraintRowState>& RowStates,
+		const FPBDJointSettings& JointSettings,
+		const FVec3& X0,
+		const FVec3& X1)
+	{
+		const FVec3 CX = X1 - X0;
+		for (int32 AxisIndex = 0; AxisIndex < 3; ++AxisIndex)
+		{
+			FVec3 Axis = FVec3(0.0f);
+			Axis[AxisIndex] = 1.0f;
+			RowStates[RowIndexBegin + AxisIndex].Axis = Axis;
+			RowStates[RowIndexBegin + AxisIndex].CalculateError(CX[AxisIndex], RowDatas[RowIndexBegin + AxisIndex].Limit);
+		}
+
+		return 3;
+	}
+
+	FORCEINLINE int32 FJointSolverConstraints::UpdateSphericalPositionConstraint(
+		int32 RowIndexBegin,
+		const TArray<FJointSolverConstraintRowData>& RowDatas,
+		TArray<FJointSolverConstraintRowState>& RowStates,
+		const FPBDJointSettings& JointSettings,
+		const FVec3& X0,
+		const FVec3& X1)
+	{
+		const FVec3 CX = X1 - X0;
+		const FReal CXLen = CX.Size();
+		if (CXLen > KINDA_SMALL_NUMBER)
+		{
+			RowStates[RowIndexBegin].Axis = CX / CXLen;
+			RowStates[RowIndexBegin].CalculateError(CXLen, RowDatas[RowIndexBegin].Limit);
+		}
+
+		return 1;
+	}
+
+	FORCEINLINE int32 FJointSolverConstraints::UpdateCylindricalPositionConstraint(
+		int32 RowIndexBegin,
+		const TArray<FJointSolverConstraintRowData>& RowDatas,
+		TArray<FJointSolverConstraintRowState>& RowStates,
+		const FPBDJointSettings& JointSettings,
+		const FRotation3& R0,
+		const FVec3& X0,
+		const FVec3& X1)
+	{
+		const FJointSolverConstraintRowData& AxisData = RowDatas[RowIndexBegin + 0];
+		const FJointSolverConstraintRowData& RadialData = RowDatas[RowIndexBegin + 1];
+		check(RadialData.UpdateType == EJointSolverConstraintUpdateType::None);
+
+		FVec3 Axis, RadialAxis;
+		FReal AxialDelta, RadialDelta;
+		FPBDJointUtilities::GetCylindricalAxesDeltas(R0, X0, X1, AxisData.ConstraintIndex, Axis, AxialDelta, RadialAxis, RadialDelta);
+
+		FJointSolverConstraintRowState& AxisState = RowStates[RowIndexBegin + 0];
+		AxisState.Axis = Axis;
+		AxisState.CalculateError(AxialDelta, AxisData.Limit);
+
+		FJointSolverConstraintRowState& RadialState = RowStates[RowIndexBegin + 1];
+		RadialState.Axis = RadialAxis;
+		RadialState.CalculateError(RadialDelta, RadialData.Limit);
+
+		return 2;
+	}
+
+	FORCEINLINE int32 FJointSolverConstraints::UpdatePlanarPositionConstraint(
+		int32 RowIndexBegin,
+		const TArray<FJointSolverConstraintRowData>& RowDatas,
+		TArray<FJointSolverConstraintRowState>& RowStates,
+		const FPBDJointSettings& JointSettings,
+		const FRotation3& R0,
+		const FVec3& X0,
+		const FVec3& X1)
+	{
+		const FJointSolverConstraintRowData& RowData = RowDatas[RowIndexBegin];
+
+		FVec3 Axis;
+		FReal Delta;
+		FPBDJointUtilities::GetPlanarAxisDelta(R0, X0, X1, RowData.ConstraintIndex, Axis, Delta);
+
+		FJointSolverConstraintRowState& RowState = RowStates[RowIndexBegin];
+		RowState.Axis = Axis;
+		RowState.CalculateError(Delta, RowData.Limit);
+
+		return 1;
+	}
+
+	FORCEINLINE int32 FJointSolverConstraints::UpdateSphericalPositionDrive(
+		int32 RowIndex,
+		const TArray<FJointSolverConstraintRowData>& RowDatas,
+		TArray<FJointSolverConstraintRowState>& RowStates,
+		const FPBDJointSettings& JointSettings,
+		const FRotation3& R0,
+		const FVec3& X0,
+		const FVec3& X1)
+	{
+		const FJointSolverConstraintRowData& RowData = RowDatas[RowIndex];
+		FJointSolverConstraintRowState& RowState = RowStates[RowIndex];
+
+		const FVec3 XTarget = X0 + R0 * JointSettings.LinearDriveTarget;
+		FVec3 Axis;
+		FReal Delta;
+		FPBDJointUtilities::GetSphericalAxisDelta(XTarget, X1, Axis, Delta);
+
+		RowState.Axis = Axis;
+		RowState.CalculateError(Delta, RowData.Limit);
+
+		return 1;
+	}
+
+	FORCEINLINE int32 FJointSolverConstraints::UpdateCircularPositionDrive(
+		int32 RowIndex,
+		const TArray<FJointSolverConstraintRowData>& RowDatas,
+		TArray<FJointSolverConstraintRowState>& RowStates,
+		const FPBDJointSettings& JointSettings,
+		const FRotation3& R0,
+		const FVec3& X0,
+		const FVec3& X1)
+	{
+		const FJointSolverConstraintRowData& RowData = RowDatas[RowIndex];
+		FJointSolverConstraintRowState& RowState = RowStates[RowIndex];
+
+		const FVec3 XTarget = X0 + R0 * JointSettings.LinearDriveTarget;
+		FVec3 Axis, RadialAxis;
+		FReal AxialDelta, RadialDelta;
+		FPBDJointUtilities::GetCylindricalAxesDeltas(R0, XTarget, X1, RowData.ConstraintIndex, Axis, AxialDelta, RadialAxis, RadialDelta);
+
+		RowState.Axis = RadialAxis;
+		RowState.CalculateError(RadialDelta, RowData.Limit);
+
+		return 1;
+	}
+
+	FORCEINLINE int32 FJointSolverConstraints::UpdateAxialPositionDrive(
+		int32 RowIndex,
+		const TArray<FJointSolverConstraintRowData>& RowDatas,
+		TArray<FJointSolverConstraintRowState>& RowStates,
+		const FPBDJointSettings& JointSettings,
+		const FRotation3& R0,
+		const FVec3& X0,
+		const FVec3& X1)
+	{
+		const FJointSolverConstraintRowData& RowData = RowDatas[RowIndex];
+		FJointSolverConstraintRowState& RowState = RowStates[RowIndex];
+
+		const FVec3 XTarget = X0 + R0 * JointSettings.LinearDriveTarget;
+		FVec3 Axis;
+		FReal Delta;
+		FPBDJointUtilities::GetPlanarAxisDelta(R0, XTarget, X1, RowData.ConstraintIndex, Axis, Delta);
+
+		RowState.Axis = Axis;
+		RowState.CalculateError(Delta, RowData.Limit);
+
+		return 1;
+	}
+
+	//
+	//////////////////////////////////////////////////////////////////////////
+	//
+
+	FORCEINLINE void FJointSolverConstraints::UpdateTwistConstraint(
+		const FJointSolverConstraintRowData& RowData,
+		FJointSolverConstraintRowState& RowState,
+		const FPBDJointSettings& JointSettings,
+		const FRotation3& R1,
+		const FRotation3& RTwist)
+	{
+		RowState.Axis = R1 * FJointConstants::TwistAxis();
+		RowState.CalculateError(FPBDJointUtilities::GetTwistAngle(RTwist), RowData.Limit);
+	}
+
+	FORCEINLINE void FJointSolverConstraints::UpdateConeSwingConstraint(
+		const FJointSolverConstraintRowData& RowData,
+		FJointSolverConstraintRowState& RowState,
+		const FPBDJointSettings& JointSettings,
+		const FRotation3& R0,
+		const FRotation3& RSwing)
+	{
+		FReal SwingAngle;
+		FVec3 SwingAxisLocal;
+		RSwing.ToAxisAndAngleSafe(SwingAxisLocal, SwingAngle, FJointConstants::Swing1Axis(), 1.e-6f);
+		if (SwingAngle > PI)
+		{
+			SwingAngle = SwingAngle - (FReal)2 * PI;
+		}
+
+		RowState.Axis = R0 * SwingAxisLocal;
+		RowState.CalculateError(SwingAngle, FPBDJointUtilities::GetConeAngleLimit(JointSettings, SwingAxisLocal, SwingAngle));
+	}
+
+	FORCEINLINE void FJointSolverConstraints::UpdateSingleLockedSwingConstraint(
+		const FJointSolverConstraintRowData& RowData,
+		FJointSolverConstraintRowState& RowState,
+		const FPBDJointSettings& JointSettings,
+		const FRotation3& R0,
+		const FRotation3& R1)
+	{
+		FVec3 Twist1 = R1 * FJointConstants::TwistAxis();
+		FVec3 Swing0 = R0 * FJointConstants::OtherSwingAxis((EJointAngularConstraintIndex)RowData.ConstraintIndex);
+
+		RowState.Axis = FVec3::CrossProduct(Swing0, Twist1);
+		RowState.CalculateError(-FVec3::DotProduct(Swing0, Twist1), RowData.Limit);
+	}
+
+	FORCEINLINE void FJointSolverConstraints::UpdateSingleLimitedSwingConstraint(
+		const FJointSolverConstraintRowData& RowData,
+		FJointSolverConstraintRowState& RowState,
+		const FPBDJointSettings& JointSettings,
+		const FRotation3& R0,
+		const FRotation3& RSwing)
+	{
+		const FReal R01SwingYorZ = (FJointConstants::AxisIndex((EJointAngularConstraintIndex)RowData.ConstraintIndex) == 2) ? RSwing.Z : RSwing.Y;	// Can't index a quat :(
+		const FReal Angle = 4.0f * FMath::Atan2(R01SwingYorZ, 1.0f + RSwing.W);
+		const FVec3& AxisLocal = (RowData.ConstraintIndex == (int32)EJointAngularConstraintIndex::Swing1) ? FJointConstants::Swing1Axis() : FJointConstants::Swing2Axis();
+
+		RowState.Axis = R0 * AxisLocal;
+		RowState.CalculateError(Angle, RowData.Limit);
+	}
+
+	FORCEINLINE void FJointSolverConstraints::UpdateDualConeSwingConstraint(
+		const FJointSolverConstraintRowData& RowData,
+		FJointSolverConstraintRowState& RowState,
+		const FPBDJointSettings& JointSettings,
+		const FRotation3& R0,
+		const FRotation3& R1)
+	{
+		FVec3 Twist1 = R1 * FJointConstants::TwistAxis();
+		FVec3 Swing0 = R0 * FJointConstants::OtherSwingAxis((EJointAngularConstraintIndex)RowData.ConstraintIndex);
+
+		FVec3 Axis = FVec3::CrossProduct(Swing0, Twist1);
+		if (Utilities::NormalizeSafe(Axis, KINDA_SMALL_NUMBER))
+		{
+			FReal SwingTwistDot = FVec3::DotProduct(Swing0, Twist1);
+			FReal Position = FMath::Asin(FMath::Clamp(-SwingTwistDot, -1.0f, 1.0f));
+			RowState.Axis = Axis;
+			RowState.CalculateError(Position, RowData.Limit);
+		}
+	}
+
+	FORCEINLINE void FJointSolverConstraints::UpdateTwistDrive(
+		const FJointSolverConstraintRowData& RowData,
+		FJointSolverConstraintRowState& RowState,
+		const FPBDJointSettings& JointSettings,
+		const FRotation3& R1,
+		const FRotation3& RTwist)
+	{
+		const FVec3 Axis = R1 * FJointConstants::TwistAxis();
+		const FReal Angle = FPBDJointUtilities::GetTwistAngle(RTwist);
+		const FReal TwistAngleTarget = JointSettings.AngularDriveTargetAngles[(int32)EJointAngularConstraintIndex::Twist];
+		const FReal DTwistAngle = Angle - TwistAngleTarget;
+
+		RowState.Axis = Axis;
+		RowState.CalculateError(DTwistAngle, RowData.Limit);
+	}
+
+	FORCEINLINE void FJointSolverConstraints::UpdateConeSwingDrive(
+		const FJointSolverConstraintRowData& RowData,
+		FJointSolverConstraintRowState& RowState,
+		const FPBDJointSettings& JointSettings,
+		const FRotation3& R0,
+		const FRotation3& RSwing)
+	{
+		FReal SwingAngle;
+		FVec3 SwingAxisLocal;
+		RSwing.ToAxisAndAngleSafe(SwingAxisLocal, SwingAngle, FJointConstants::Swing1Axis(), AngleTolerance);
+		if (SwingAngle > PI)
+		{
+			SwingAngle = SwingAngle - (FReal)2 * PI;
+		}
+
+		// Circular swing target (max of Swing1, Swing2 targets)
+		// @todo(ccaulfield): what should cone target really do?
+		const FReal Swing1Target = JointSettings.AngularDriveTargetAngles[(int32)EJointAngularConstraintIndex::Swing1];
+		const FReal Swing2Target = JointSettings.AngularDriveTargetAngles[(int32)EJointAngularConstraintIndex::Swing2];
+		const FReal SwingAngleTarget = FMath::Max(Swing1Target, Swing2Target);
+		const FReal DSwingAngle = SwingAngle - SwingAngleTarget;
+
+		RowState.Axis = R0 * SwingAxisLocal;
+		RowState.CalculateError(DSwingAngle, RowData.Limit);
+	}
+
+	FORCEINLINE void FJointSolverConstraints::UpdateSwingDrive(
+		const FJointSolverConstraintRowData& RowData,
+		FJointSolverConstraintRowState& RowState,
+		const FPBDJointSettings& JointSettings,
+		const FRotation3& R0,
+		const FRotation3& RSwing)
+	{
+		const FReal R01SwingYorZ = (FJointConstants::AxisIndex((EJointAngularConstraintIndex)RowData.ConstraintIndex) == 2) ? RSwing.Z : RSwing.Y;	// Can't index a quat :(
+		const FReal SwingAngle = 4.0f * FMath::Atan2(R01SwingYorZ, 1.0f + RSwing.W);
+		const FReal SwingAngleTarget = JointSettings.AngularDriveTargetAngles[(int32)RowData.ConstraintIndex];
+		const FReal DSwingAngle = SwingAngle - SwingAngleTarget;
+		const FVec3& AxisLocal = ((EJointAngularConstraintIndex)RowData.ConstraintIndex == EJointAngularConstraintIndex::Swing1) ? FJointConstants::Swing1Axis() : FJointConstants::Swing2Axis();
+
+		RowState.Axis = R0 * AxisLocal;
+		RowState.CalculateError(DSwingAngle, RowData.Limit);
+	}
+
+	FORCEINLINE void FJointSolverConstraints::UpdateSLerpDrive(
+		const FJointSolverConstraintRowData& RowData,
+		FJointSolverConstraintRowState& RowState,
+		const FPBDJointSettings& JointSettings,
+		const FRotation3& R0,
+		const FRotation3& R1)
+	{
+		const FRotation3 TargetR1 = R0 * JointSettings.AngularDrivePositionTarget;
+		const FRotation3 DR = TargetR1 * R1.Inverse();
+
+		FVec3 SLerpAxis;
+		FReal SLerpAngle;
+		if (DR.ToAxisAndAngleSafe(SLerpAxis, SLerpAngle, FVec3(1, 0, 0)))
+		{
+			if (SLerpAngle > (FReal)PI)
+			{
+				SLerpAngle = SLerpAngle - (FReal)2 * PI;
+			}
+
+			//UE_LOG(LogChaosJoint, VeryVerbose, TEXT("      SLerpDrive Pos: %f Axis: %f %f %f"), -SLerpAngle, SLerpAxis.X, SLerpAxis.Y, SLerpAxis.Z);
+
+			RowState.Axis = SLerpAxis;
+			RowState.CalculateError(-SLerpAngle, RowData.Limit);
+		}
+	}
+
+	//
+	///////////////////////////////////////////////////////////////////////
+	//
+
+	FORCEINLINE void FJointSolver::ApplyPositionConstraint1(
+		const FReal Dt,
+		const FJointSolverJointState& JointState,
+		const FJointSolverConstraintRowData& RowData,
+		FJointSolverConstraintRowState& RowState)
+	{
+		//UE_LOG(LogChaosJoint, VeryVerbose, TEXT("    Position Error %f"), RowState.Error);
+
+		FReal II = 0.0f;
+		FVec3 AngularAxis0 = FVec3(0);
+		FVec3 AngularAxis1 = FVec3(0);
+
+		if (JointState.InvMs[0] > 0.0f)
+		{
+			AngularAxis0 = FVec3::CrossProduct(JointState.Xs[0] - JointState.Ps[0], RowState.Axis);
+			const FVec3 IA0 = Utilities::Multiply(JointState.InvIs[0], AngularAxis0);
+			const FReal II0 = FVec3::DotProduct(AngularAxis0, IA0);
+			II += JointState.InvMs[0] + II0;
+		}
+		if (JointState.InvMs[1] > 0.0f)
+		{
+			AngularAxis1 = FVec3::CrossProduct(JointState.Xs[1] - JointState.Ps[1], RowState.Axis);
+			const FVec3 IA1 = Utilities::Multiply(JointState.InvIs[1], AngularAxis1);
+			const FReal II1 = FVec3::DotProduct(AngularAxis1, IA1);
+			II += JointState.InvMs[1] + II1;
+		}
+
+		if (RowData.bIsSoft)
+		{
+			FReal VelDt = 0;
+			if (RowData.Damping > KINDA_SMALL_NUMBER)
+			{
+				const FVec3 V0 = FVec3::CalculateVelocity(JointState.PrevXs[0], JointState.Xs[0], 1.0f);
+				const FVec3 V1 = FVec3::CalculateVelocity(JointState.PrevXs[1], JointState.Xs[1], 1.0f);
+				VelDt = FVec3::DotProduct(V0 - V1, RowState.Axis);
+			}
+
+			const FReal SpringMassScale = (RowData.bIsAccelerationMode) ? 1.0f / (JointState.InvMs[0] + JointState.InvMs[1]) : 1.0f;
+			const FReal S = SpringMassScale * RowData.Stiffness * Dt * Dt;
+			const FReal D = SpringMassScale * RowData.Damping * Dt;
+			const FReal Multiplier = (FReal)1 / ((S + D) * II + (FReal)1);
+			const FReal DLambda = Multiplier * (S * RowState.Error - D * VelDt - RowState.Lambda);
+
+			RowState.DPs[0] = (JointState.InvMs[0] * DLambda) * RowState.Axis;
+			RowState.DPs[1] = (-JointState.InvMs[1] * DLambda) * RowState.Axis;
+			RowState.DRs[0] = DLambda * Utilities::Multiply(JointState.InvIs[0], AngularAxis0);
+			RowState.DRs[1] = -DLambda * Utilities::Multiply(JointState.InvIs[1], AngularAxis1);
+			RowState.Lambda += DLambda;
+		}
+		else
+		{
+			const FVec3 DX = RowState.Axis * (RowData.Stiffness * RowState.Error / II);
+
+			// Apply constraint correction
+			RowState.DPs[0] = JointState.InvMs[0] * DX;
+			RowState.DPs[1] = -JointState.InvMs[1] * DX;
+			RowState.DRs[0] = Utilities::Multiply(JointState.InvIs[0], FVec3::CrossProduct(JointState.Xs[0] - JointState.Ps[0], DX));
+			RowState.DRs[1] = Utilities::Multiply(JointState.InvIs[1], FVec3::CrossProduct(JointState.Xs[1] - JointState.Ps[1], -DX));
+		}
+	}
+
+	// NOTE: Currently assumes row axes are unit axes
+	FORCEINLINE void FJointSolver::ApplyPositionConstraint3(
+		const FReal Dt,
+		const FJointSolverJointState& JointState,
+		const FJointSolverConstraintRowData& RowData,
+		FJointSolverConstraintRowState& RowState0,
+		FJointSolverConstraintRowState& RowState1,
+		FJointSolverConstraintRowState& RowState2)
+	{
+		FVec3 CX = FVec3(RowState0.Error, RowState1.Error, RowState2.Error);
+
+		//UE_LOG(LogChaosJoint, VeryVerbose, TEXT("    Position Error %f"), CX.Size());
+
+		FMatrix33 M = FMatrix33(0, 0, 0);
+		if (JointState.InvMs[0] > 0)
+		{
+			FMatrix33 M0 = Utilities::ComputeJointFactorMatrix(JointState.Xs[0] - JointState.Ps[0], JointState.InvIs[0], JointState.InvMs[0]);
+			M += M0;
+		}
+		if (JointState.InvMs[1] > 0)
+		{
+			FMatrix33 M1 = Utilities::ComputeJointFactorMatrix(JointState.Xs[1] - JointState.Ps[1], JointState.InvIs[1], JointState.InvMs[1]);
+			M += M1;
+		}
+		FMatrix33 MI = M.Inverse();
+		const FVec3 DX = RowData.Stiffness * Utilities::Multiply(MI, CX);
+
+		// Apply constraint correction
+		RowState0.DPs[0] = JointState.InvMs[0] * DX;
+		RowState0.DPs[1] = -JointState.InvMs[1] * DX;
+		RowState0.DRs[0] = Utilities::Multiply(JointState.InvIs[0], FVec3::CrossProduct(JointState.Xs[0] - JointState.Ps[0], DX));
+		RowState0.DRs[1] = Utilities::Multiply(JointState.InvIs[1], FVec3::CrossProduct(JointState.Xs[1] - JointState.Ps[1], -DX));
+	}
+
+	FORCEINLINE void FJointSolver::ApplyRotationConstraint(
+		const FReal Dt,
+		const FJointSolverJointState& JointState,
+		const FJointSolverConstraintRowData& RowData,
+		FJointSolverConstraintRowState& RowState)
+	{
+		// Joint-space inverse mass
+		FVec3 IA0 = FVec3(0);
+		FVec3 IA1 = FVec3(0);
+		FReal II0 = 0.0f;
+		FReal II1 = 0.0f;
+		if (JointState.InvMs[0] > 0)
+		{
+			IA0 = Utilities::Multiply(JointState.InvIs[0], RowState.Axis);
+			II0 = FVec3::DotProduct(RowState.Axis, IA0);
+		}
+		if (JointState.InvMs[1] > 0)
+		{
+			IA1 = Utilities::Multiply(JointState.InvIs[1], RowState.Axis);
+			II1 = FVec3::DotProduct(RowState.Axis, IA1);
+		}
+		const FReal II = (II0 + II1);
+
+		if (RowData.bIsSoft)
+		{
+			// Damping angular velocity
+			FReal AngVelDt = 0;
+			if (RowData.Damping > 0.0f)
+			{
+				const FVec3 W0 = FRotation3::CalculateAngularVelocity(JointState.PrevQs[0], JointState.Qs[0], 1.0f);
+				const FVec3 W1 = FRotation3::CalculateAngularVelocity(JointState.PrevQs[1], JointState.Qs[1], 1.0f);
+				AngVelDt = FVec3::DotProduct(RowState.Axis, W0) - FVec3::DotProduct(RowState.Axis, W1);
+			}
+
+			const FReal SpringMassScale = (RowData.bIsAccelerationMode) ? 1.0f / II : 1.0f;
+			const FReal S = SpringMassScale * RowData.Stiffness * Dt * Dt;
+			const FReal D = SpringMassScale * RowData.Damping * Dt;
+			const FReal Multiplier = (FReal)1 / ((S + D) * II + (FReal)1);
+			const FReal DLambda = Multiplier * (S * RowState.Error - D * AngVelDt - RowState.Lambda);
+
+			//const FVec3 DR0 = IA0 * DLambda;
+			//const FVec3 DR1 = IA1 * -DLambda;
+			RowState.DRs[0] = RowState.Axis * (DLambda * II0);
+			RowState.DRs[1] = RowState.Axis * -(DLambda * II1);
+			RowState.Lambda += DLambda;
+		}
+		else
+		{
+			RowState.DRs[0] = IA0 * (RowData.Stiffness * RowState.Error / II);
+			RowState.DRs[1] = IA1 * -(RowData.Stiffness * RowState.Error / II);
+		}
+	}
+
+}
\ No newline at end of file
diff --git a/Engine/Source/Runtime/Experimental/Chaos/Public/Chaos/Joint/PBDJointSolverGaussSeidel.h b/Engine/Source/Runtime/Experimental/Chaos/Public/Chaos/Joint/PBDJointSolverGaussSeidel.h
--- a/Engine/Source/Runtime/Experimental/Chaos/Public/Chaos/Joint/PBDJointSolverGaussSeidel.h
+++ b/Engine/Source/Runtime/Experimental/Chaos/Public/Chaos/Joint/PBDJointSolverGaussSeidel.h
@@ -11,44 +11,6 @@
 
 namespace Chaos
 {
-	struct FJointSolverResult
-	{
-	public:
-		static FJointSolverResult MakeActive() { return { 1, 0 }; }
-		static FJointSolverResult MakeSolved() { return { 0, 1 }; }
-
-		FJointSolverResult()
-			: NumActive(0)
-			, NumSolved(0)
-		{
-		}
-		
-		int32 GetNumActive() const { return NumActive; }
-		int32 GetNumSolved() const { return NumSolved; }
-		
-		FJointSolverResult& operator+=(const FJointSolverResult& R)
-		{
-			NumActive += R.NumActive;
-			NumSolved += R.NumSolved;
-			return *this;
-		}
-
-		friend FJointSolverResult operator+(const FJointSolverResult& L, const FJointSolverResult& R)
-		{
-			return { L.NumActive + R.NumActive, L.NumSolved + R.NumSolved };
-		}
-
-	private:
-		FJointSolverResult(int32 InActive, int32 InSolved)
-			: NumActive(InActive)
-			, NumSolved(InSolved)
-		{
-		}
-
-		int32 NumActive;
-		int32 NumSolved;
-	};
-
 	/**
 	 * Calculate new positions and rotations for a pair of bodies connected by a joint.
 	 *
@@ -74,6 +36,18 @@ namespace Chaos
 			return Qs[Index];
 		}
 
+		FORCEINLINE const FVec3& GetPrevP(const int32 Index) const
+		{
+			checkSlow(Index < MaxConstrainedBodies);
+			return PrevPs[Index];
+		}
+
+		FORCEINLINE const FRotation3& GetPrevQ(const int32 Index) const
+		{
+			checkSlow(Index < MaxConstrainedBodies);
+			return PrevQs[Index];
+		}
+
 		FORCEINLINE const FVec3& GetV(const int32 Index) const
 		{
 			checkSlow(Index < MaxConstrainedBodies);
@@ -114,47 +88,49 @@ namespace Chaos
 			const FVec3& V1,
 			const FVec3& W1);
 
-		FJointSolverResult ApplyConstraints(
+		int32 ApplyConstraints(
 			const FReal Dt,
 			const FPBDJointSolverSettings& SolverSettings,
 			const FPBDJointSettings& JointSettings);
 
-		FJointSolverResult ApplyDrives(
+		int32 ApplyDrives(
 			const FReal Dt,
 			const FPBDJointSolverSettings& SolverSettings,
 			const FPBDJointSettings& JointSettings);
 
-		FJointSolverResult ApplyProjections(
+		int32 ApplyProjections(
 			const FReal Dt,
 			const FPBDJointSolverSettings& SolverSettings,
 			const FPBDJointSettings& JointSettings);
 
+
 	private:
 
+		void InitDerivedState();
 		void UpdateDerivedState(const int32 BodyIndex);
 		void UpdateDerivedState();
 
-		FJointSolverResult ApplyPositionConstraints(
+		int32 ApplyPositionConstraints(
 			const FReal Dt,
 			const FPBDJointSolverSettings& SolverSettings,
 			const FPBDJointSettings& JointSettings);
 
-		FJointSolverResult ApplyRotationConstraints(
+		int32 ApplyRotationConstraints(
 			const FReal Dt,
 			const FPBDJointSolverSettings& SolverSettings,
 			const FPBDJointSettings& JointSettings);
 
-		FJointSolverResult ApplyPositionDrives(
+		int32 ApplyPositionDrives(
 			const FReal Dt,
 			const FPBDJointSolverSettings& SolverSettings,
 			const FPBDJointSettings& JointSettings);
 
-		FJointSolverResult ApplyRotationDrives(
+		int32 ApplyRotationDrives(
 			const FReal Dt,
 			const FPBDJointSolverSettings& SolverSettings,
 			const FPBDJointSettings& JointSettings);
 
-		FJointSolverResult ApplyRotationProjection(
+		int32 ApplyRotationProjection(
 			const FReal Dt,
 			const FPBDJointSolverSettings& SolverSettings,
 			const FPBDJointSettings& JointSettings);
@@ -247,47 +223,47 @@ namespace Chaos
 			const FReal Angle,
 			FReal& Lambda);
 
-		FJointSolverResult ApplyLockedRotationConstraints(
+		int32 ApplyLockedRotationConstraints(
 			const FReal Dt,
 			const FPBDJointSolverSettings& SolverSettings,
 			const FPBDJointSettings& JointSettings,
 			const bool bApplyTwist,
 			const bool bApplySwing);
 
-		FJointSolverResult ApplyTwistConstraint(
+		int32 ApplyTwistConstraint(
 			const FReal Dt,
 			const FPBDJointSolverSettings& SolverSettings,
 			const FPBDJointSettings& JointSettings,
 			const bool bUseSoftLimit);
 
-		FJointSolverResult ApplyTwistDrive(
+		int32 ApplyTwistDrive(
 			const FReal Dt,
 			const FPBDJointSolverSettings& SolverSettings,
 			const FPBDJointSettings& JointSettings);
 
-		FJointSolverResult ApplyTwistProjection(
+		int32 ApplyTwistProjection(
 			const FReal Dt,
 			const FPBDJointSolverSettings& SolverSettings,
 			const FPBDJointSettings& JointSettings);
 
-		FJointSolverResult ApplyConeConstraint(
+		int32 ApplyConeConstraint(
 			const FReal Dt,
 			const FPBDJointSolverSettings& SolverSettings,
 			const FPBDJointSettings& JointSettings,
 			const bool bUseSoftLimit);
 
-		FJointSolverResult ApplyConeDrive(
+		int32 ApplyConeDrive(
 			const FReal Dt,
 			const FPBDJointSolverSettings& SolverSettings,
 			const FPBDJointSettings& JointSettings);
 
-		FJointSolverResult ApplyConeProjection(
+		int32 ApplyConeProjection(
 			const FReal Dt,
 			const FPBDJointSolverSettings& SolverSettings,
 			const FPBDJointSettings& JointSettings);
 
 		// One Swing axis is free, and the other locked. This applies the lock: Body1 Twist axis is confined to a plane.
-		FJointSolverResult ApplySingleLockedSwingConstraint(
+		int32 ApplySingleLockedSwingConstraint(
 			const FReal Dt,
 			const FPBDJointSolverSettings& SolverSettings,
 			const FPBDJointSettings& JointSettings,
@@ -295,7 +271,7 @@ namespace Chaos
 			const bool bUseSoftLimit);
 
 		// One Swing axis is free, and the other limited. This applies the limit: Body1 Twist axis is confined to space between two cones.
-		FJointSolverResult ApplyDualConeSwingConstraint(
+		int32 ApplyDualConeSwingConstraint(
 			const FReal Dt,
 			const FPBDJointSolverSettings& SolverSettings,
 			const FPBDJointSettings& JointSettings,
@@ -303,53 +279,53 @@ namespace Chaos
 			const bool bUseSoftLimit);
 
 		// One swing axis is locked, the other limited or locked. This applies the Limited axis (ApplyDualConeSwingConstraint is used for the locked axis).
-		FJointSolverResult ApplySwingConstraint(
+		int32 ApplySwingConstraint(
 			const FReal Dt,
 			const FPBDJointSolverSettings& SolverSettings,
 			const FPBDJointSettings& JointSettings,
 			const EJointAngularConstraintIndex SwingConstraintIndex,
 			const bool bUseSoftLimit);
 
-		FJointSolverResult ApplySwingDrive(
+		int32 ApplySwingDrive(
 			const FReal Dt,
 			const FPBDJointSolverSettings& SolverSettings,
 			const FPBDJointSettings& JointSettings,
 			const EJointAngularConstraintIndex SwingConstraintIndex);
 
-		FJointSolverResult ApplySwingProjection(
+		int32 ApplySwingProjection(
 			const FReal Dt,
 			const FPBDJointSolverSettings& SolverSettings,
 			const FPBDJointSettings& JointSettings,
 			const EJointAngularConstraintIndex SwingConstraintIndex);
 
-		FJointSolverResult ApplySLerpDrive(
+		int32 ApplySLerpDrive(
 			const FReal Dt,
 			const FPBDJointSolverSettings& SolverSettings,
 			const FPBDJointSettings& JointSettings);
 
-		FJointSolverResult ApplyPointPositionConstraintKD(
+		int32 ApplyPointPositionConstraintKD(
 			const int32 KIndex,
 			const int32 DIndex,
 			const FReal Dt,
 			const FPBDJointSolverSettings& SolverSettings,
 			const FPBDJointSettings& JointSettings);
 
-		FJointSolverResult ApplyPointPositionConstraintDD(
+		int32 ApplyPointPositionConstraintDD(
 			const FReal Dt,
 			const FPBDJointSolverSettings& SolverSettings,
 			const FPBDJointSettings& JointSettings);
 
-		FJointSolverResult ApplySphericalPositionConstraint(
+		int32 ApplySphericalPositionConstraint(
 			const FReal Dt,
 			const FPBDJointSolverSettings& SolverSettings,
 			const FPBDJointSettings& JointSettings);
 
-		FJointSolverResult ApplySphericalPositionDrive(
+		int32 ApplySphericalPositionDrive(
 			const FReal Dt,
 			const FPBDJointSolverSettings& SolverSettings,
 			const FPBDJointSettings& JointSettings);
 
-		FJointSolverResult ApplyCylindricalPositionConstraint(
+		int32 ApplyCylindricalPositionConstraint(
 			const FReal Dt,
 			const int32 AxisIndex,
 			const EJointMotionType AxialMotion,
@@ -357,26 +333,26 @@ namespace Chaos
 			const FPBDJointSolverSettings& SolverSettings,
 			const FPBDJointSettings& JointSettings);
 
-		FJointSolverResult ApplyCircularPositionDrive(
+		int32 ApplyCircularPositionDrive(
 			const FReal Dt,
 			const int32 AxisIndex,
 			const FPBDJointSolverSettings& SolverSettings,
 			const FPBDJointSettings& JointSettings);
 
-		FJointSolverResult ApplyPlanarPositionConstraint(
+		int32 ApplyPlanarPositionConstraint(
 			const FReal Dt,
 			const int32 AxisIndex,
 			const EJointMotionType AxialMotion,
 			const FPBDJointSolverSettings& SolverSettings,
 			const FPBDJointSettings& JointSettings);
 
-		FJointSolverResult ApplyAxialPositionDrive(
+		int32 ApplyAxialPositionDrive(
 			const FReal Dt,
 			const int32 AxisIndex,
 			const FPBDJointSolverSettings& SolverSettings,
 			const FPBDJointSettings& JointSettings);
 
-		FJointSolverResult ApplyPositionProjection(
+		int32 ApplyPositionProjection(
 			const FReal Dt,
 			const FPBDJointSolverSettings& SolverSettings,
 			const FPBDJointSettings& JointSettings);
@@ -400,6 +376,7 @@ namespace Chaos
 		FRotation3 Qs[MaxConstrainedBodies];		// World-space particle CoM rotations
 		FVec3 Vs[MaxConstrainedBodies];				// World-space particle CoM velocities
 		FVec3 Ws[MaxConstrainedBodies];				// World-space particle CoM angular velocities
+		FMatrix33 InvIs[MaxConstrainedBodies];		// World-space inverse inertias
 
 		// XPBD Previous iteration world-space body state
 		FVec3 PrevPs[MaxConstrainedBodies];			// World-space particle CoM positions
@@ -407,9 +384,6 @@ namespace Chaos
 		FReal SwingDriveLambda;						// XPBD constraint multipliers (net applied constraint-space deltas)
 		FRotation3 PrevQs[MaxConstrainedBodies];	// World-space particle CoM rotations
 		FVec3 PrevXs[MaxConstrainedBodies];			// World-space joint connector positions
-		
-		// Post-angular constraint position fixup (to reduce iterations required for stiff angular constraints)
-		FReal AngularPositionCorrection;
 
 		FReal PositionTolerance;					// Distance error below which we consider a constraint or drive solved
 		FReal AngleTolerance;						// Angle error below which we consider a constraint or drive solved
diff --git a/Engine/Source/Runtime/Experimental/Chaos/Public/Chaos/PBDJointConstraintTypes.h b/Engine/Source/Runtime/Experimental/Chaos/Public/Chaos/PBDJointConstraintTypes.h
--- a/Engine/Source/Runtime/Experimental/Chaos/Public/Chaos/PBDJointConstraintTypes.h
+++ b/Engine/Source/Runtime/Experimental/Chaos/Public/Chaos/PBDJointConstraintTypes.h
@@ -33,6 +33,7 @@ namespace Chaos
 	/**
 	 * The order of the angular constraints (for settings held in vectors etc)
 	 */
+	// @todo(ccaulfield): rename EJointAngularConstraintType
 	enum class EJointAngularConstraintIndex : int32
 	{
 		Twist,
@@ -185,4 +186,44 @@ namespace Chaos
 		FReal SoftSwingStiffness;
 		FReal SoftSwingDamping;
 	};
+
+	class FJointSolverResult
+	{
+	public:
+		static FJointSolverResult MakeActive() { return { 1, 0 }; }
+		static FJointSolverResult MakeInactive() { return { 0, 1 }; }
+		static FJointSolverResult Make(int32 NumActive, int32 NumInactive) { return { NumActive, NumInactive }; }
+
+		FJointSolverResult()
+			: NumActive(0)
+			, NumInactive(0)
+		{
+		}
+
+		int32 GetNumActive() const { return NumActive; }
+		int32 GetNumInactive() const { return NumInactive; }
+
+		FJointSolverResult& operator+=(const FJointSolverResult& R)
+		{
+			NumActive += R.NumActive;
+			NumInactive += R.NumInactive;
+			return *this;
+		}
+
+		friend FJointSolverResult operator+(const FJointSolverResult& L, const FJointSolverResult& R)
+		{
+			return { L.NumActive + R.NumActive, L.NumInactive + R.NumInactive };
+		}
+
+	private:
+		FJointSolverResult(int32 InActive, int32 InSolved)
+			: NumActive(InActive)
+			, NumInactive(InSolved)
+		{
+		}
+
+		int32 NumActive;
+		int32 NumInactive;
+	};
+
 }
\ No newline at end of file
diff --git a/Engine/Source/Runtime/Experimental/Chaos/Public/Chaos/PBDJointConstraintUtilities.h b/Engine/Source/Runtime/Experimental/Chaos/Public/Chaos/PBDJointConstraintUtilities.h
--- a/Engine/Source/Runtime/Experimental/Chaos/Public/Chaos/PBDJointConstraintUtilities.h
+++ b/Engine/Source/Runtime/Experimental/Chaos/Public/Chaos/PBDJointConstraintUtilities.h
@@ -12,6 +12,30 @@ namespace Chaos
 	class FPBDJointUtilities
 	{
 	public:
+		static void GetSphericalAxisDelta(
+			const FVec3& X0,
+			const FVec3& X1,
+			FVec3& Axis,
+			FReal& Delta);
+
+		static void GetCylindricalAxesDeltas(
+			const FRotation3& R0,
+			const FVec3& X0,
+			const FVec3& X1,
+			const int32 CylinderAxisIndex,
+			FVec3& CylinderAxis,
+			FReal& CylinderDelta,
+			FVec3& RadialAxis,
+			FReal& RadialDelta);
+
+		static void GetPlanarAxisDelta(
+			const FRotation3& R0,
+			const FVec3& X0,
+			const FVec3& X1,
+			const int32 PlaneAxisIndex,
+			FVec3& Axis,
+			FReal& Delta);
+
 		static void DecomposeSwingTwistLocal(
 			const FRotation3& R0, 
 			const FRotation3& R1, 
@@ -63,13 +87,18 @@ namespace Chaos
 			FVec3& Axis,
 			FReal& Angle);
 
-		static void GetLockedAxes(
+		static void GetLockedRotationAxes(
 			const FRotation3& R0, 
 			const FRotation3& R1, 
 			FVec3& Axis0, 
 			FVec3& Axis1, 
 			FVec3& Axis2);
 
+		static FReal GetConeAngleLimit(
+			const FPBDJointSettings& JointSettings,
+			const FVec3& SwingAxisLocal,
+			const FReal SwingAngle);
+
 		/**
 		 * Increase the lower inertia components to ensure that the maximum ratio between any pair of elements is MaxRatio.
 		 *
diff --git a/Engine/Source/Runtime/Experimental/Chaos/Public/Chaos/PBDJointConstraints.h b/Engine/Source/Runtime/Experimental/Chaos/Public/Chaos/PBDJointConstraints.h
--- a/Engine/Source/Runtime/Experimental/Chaos/Public/Chaos/PBDJointConstraints.h
+++ b/Engine/Source/Runtime/Experimental/Chaos/Public/Chaos/PBDJointConstraints.h
@@ -9,6 +9,7 @@
 #include "Chaos/Vector.h"
 
 #include "Chaos/ConstraintHandle.h"
+#include "Chaos/Joint/JointSolverConstraints.h"
 #include "Chaos/Joint/PBDJointSolverGaussSeidel.h"
 #include "Chaos/ParticleHandleFwd.h"
 #include "Chaos/PBDConstraintContainer.h"
@@ -16,22 +17,28 @@
 
 namespace Chaos
 {
+	class FJointSolverConstraints;
+	class FJointSolverGaussSeidel;
+	class FJointSolverResult;
 
-	class FPBDJointConstraintHandle : public TContainerConstraintHandle<FPBDJointConstraints>
+	class CHAOS_API FPBDJointConstraintHandle : public TContainerConstraintHandle<FPBDJointConstraints>
 	{
 	public:
 		using Base = TContainerConstraintHandle<FPBDJointConstraints>;
 		using FConstraintContainer = FPBDJointConstraints;
 
-		CHAOS_API FPBDJointConstraintHandle();
-		CHAOS_API FPBDJointConstraintHandle(FConstraintContainer* InConstraintContainer, int32 InConstraintIndex);
+		FPBDJointConstraintHandle();
+		FPBDJointConstraintHandle(FConstraintContainer* InConstraintContainer, int32 InConstraintIndex);
 
-		CHAOS_API void CalculateConstraintSpace(FVec3& OutXa, FMatrix33& OutRa, FVec3& OutXb, FMatrix33& OutRb) const;
-		CHAOS_API void SetParticleLevels(const TVector<int32, 2>& ParticleLevels);
-		CHAOS_API int32 GetConstraintLevel() const;
-		CHAOS_API const FPBDJointSettings& GetSettings() const;
-		CHAOS_API void SetSettings(const FPBDJointSettings& Settings);
-		CHAOS_API TVector<TGeometryParticleHandle<float,3>*, 2> GetConstrainedParticles() const;
+		void CalculateConstraintSpace(FVec3& OutXa, FMatrix33& OutRa, FVec3& OutXb, FMatrix33& OutRb) const;
+		int32 GetConstraintIsland() const;
+		int32 GetConstraintLevel() const;
+		int32 GetConstraintColor() const;
+		int32 GetConstraintBatch() const;
+
+		const FPBDJointSettings& GetSettings() const;
+		void SetSettings(const FPBDJointSettings& Settings);
+		TVector<TGeometryParticleHandle<float,3>*, 2> GetConstrainedParticles() const;
 
 	protected:
 		using Base::ConstraintIndex;
@@ -43,9 +50,11 @@ namespace Chaos
 	public:
 		FPBDJointState();
 
-		// Priorities used for ordering, mass conditioning, projection, and freezing
+		int32 Batch;
+		int32 Island;
 		int32 Level;
-		TVector<int32, 2> ParticleLevels;
+		int32 Color;
+		int32 IslandSize;
 	};
 
 	/**
@@ -129,8 +138,10 @@ namespace Chaos
 		const FPBDJointSettings& GetConstraintSettings(int32 ConstraintIndex) const;
 		void SetConstraintSettings(int32 ConstraintIndex, const FPBDJointSettings& InConstraintSettings);
 
+		int32 GetConstraintIsland(int32 ConstraintIndex) const;
 		int32 GetConstraintLevel(int32 ConstraintIndex) const;
-		void SetParticleLevels(int32 ConstraintIndex, const TVector<int32, 2>& ParticleLevels);
+		int32 GetConstraintColor(int32 ConstraintIndex) const;
+		int32 GetConstraintBatch(int32 ConstraintIndex) const;
 
 		//
 		// General Rule API
@@ -167,10 +178,22 @@ namespace Chaos
 		void CalculateConstraintSpace(int32 ConstraintIndex, FVec3& OutX0, FMatrix33& OutR0, FVec3& OutX1, FMatrix33& OutR1) const;
 		void UpdateParticleState(TPBDRigidParticleHandle<FReal, 3>* Rigid, const FReal Dt, const FVec3& PrevP, const FRotation3& PrevQ, const FVec3& P, const FRotation3& Q, const bool bUpdateVelocity = true);
 		void UpdateParticleStateExplicit(TPBDRigidParticleHandle<FReal, 3>* Rigid, const FReal Dt, const FVec3& P, const FRotation3& Q, const FVec3& V, const FVec3& W);
+		
+		void InitSolverJointData();
+		void DeinitSolverJointData();
+		void InitSolverJointState();
+		void DeinitSolverJointState();
+		void GatherSolverJointState(int32 ConstraintIndex);
+		void ScatterSolverJointState(const FReal Dt, int32 ConstraintIndex);
+
+		void ColorConstraints();
 		void SortConstraints();
+		void BatchConstraints();
+		void CheckBatches();
 
-		FJointSolverResult SolvePosition_GaussSiedel(const FReal Dt, const int32 ConstraintIndex, const int32 NumPairIts, const int32 It, const int32 NumIts);
-		FJointSolverResult ProjectPosition_GaussSiedel(const FReal Dt, const int32 ConstraintIndex, const int32 NumPairIts, const int32 It, const int32 NumIts);
+		int32 ApplyBatch(const FReal Dt, const int32 BatchIndex, const int32 NumPairIts, const int32 It, const int32 NumIts);
+		int32 ApplySingle(const FReal Dt, const int32 ConstraintIndex, const int32 NumPairIts, const int32 It, const int32 NumIts);
+		int32 ApplyPushOutSingle(const FReal Dt, const int32 ConstraintIndex, const int32 NumPairIts, const int32 It, const int32 NumIts);
 
 		FPBDJointSolverSettings Settings;
 
@@ -181,7 +204,8 @@ namespace Chaos
 
 		FHandles Handles;
 		FConstraintHandleAllocator HandleAllocator;
-		bool bRequiresSort;
+		bool bJointsDirty;
+		bool bIsBatched;
 
 		FJointPreApplyCallback PreApplyCallback;
 		FJointPostApplyCallback PostApplyCallback;
@@ -189,6 +213,12 @@ namespace Chaos
 
 		// @todo(ccaulfield): optimize storage for joint solver
 		TArray<FJointSolverGaussSeidel> ConstraintSolvers;
+
+		TArray<FJointSolverConstraints> SolverConstraints;
+		TArray<TVector<int32, 2>> JointBatches;
+		TArray< FJointSolverJointState> SolverConstraintStates;
+		TArray<FJointSolverConstraintRowData> SolverConstraintRowDatas;
+		TArray<FJointSolverConstraintRowState> SolverConstraintRowStates;
 	};
 
 }
diff --git a/Engine/Source/Runtime/Experimental/ChaosCore/Public/Chaos/Matrix33.isph b/Engine/Source/Runtime/Experimental/ChaosCore/Public/Chaos/Matrix33.isph
--- a/Engine/Source/Runtime/Experimental/ChaosCore/Public/Chaos/Matrix33.isph
+++ b/Engine/Source/Runtime/Experimental/ChaosCore/Public/Chaos/Matrix33.isph
@@ -81,6 +81,17 @@ inline uniform FVector Multiply(const uniform FMatrix33& L, const uniform FVecto
 	return SetVector(Result.V[0], Result.V[1], Result.V[2]);
 }
 
+inline varying FVector Multiply(const uniform FMatrix33& L, const varying FVector& R)
+{
+	const uniform FVector4 L0 = SetVector4(L.M[0], L.M[1], L.M[2], L.M[3]);
+	const uniform FVector4 L1 = SetVector4(L.M[4], L.M[5], L.M[6], L.M[7]);
+	const uniform FVector4 L2 = SetVector4(L.M[8], L.M[9], L.M[10], L.M[11]);
+
+	const varying FVector4 Result = L0 * SetVector4(R.V[0]) + L1 * SetVector4(R.V[1]) + L2 * SetVector4(R.V[2]);
+
+	return SetVector(Result.V[0], Result.V[1], Result.V[2]);
+}
+
 inline uniform FMatrix33 AddAB(const uniform FMatrix33& L, const uniform FMatrix33& R)
 {
 	const uniform FVector4 L0 = SetVector4(L.M[0], L.M[1], L.M[2], L.M[3]);
