Fix CIS - dupe symbol in joints

diff --git a/Engine/Source/Runtime/Experimental/Chaos/Private/Chaos/Joint/PBDJointSolverGaussSeidel.cpp b/Engine/Source/Runtime/Experimental/Chaos/Private/Chaos/Joint/PBDJointSolverGaussSeidel.cpp
--- a/Engine/Source/Runtime/Experimental/Chaos/Private/Chaos/Joint/PBDJointSolverGaussSeidel.cpp
+++ b/Engine/Source/Runtime/Experimental/Chaos/Private/Chaos/Joint/PBDJointSolverGaussSeidel.cpp
@@ -12,14 +12,6 @@
 
 //PRAGMA_DISABLE_OPTIMIZATION
 
-#if !INTEL_ISPC
-const bool bChaos_Joint_ISPC_Enabled = false;
-#elif UE_BUILD_SHIPPING
-const bool bChaos_Joint_ISPC_Enabled = true;
-#else
-bool bChaos_Joint_ISPC_Enabled = true;
-FAutoConsoleVariableRef CVarChaosJointISPCEnabled(TEXT("p.Chaos.Joint.ISPC"), bChaos_Joint_ISPC_Enabled, TEXT("Whether to use ISPC optimizations in the Joint Solver"));
-#endif
 
 float Chaos_Joint_DegenerateRotationLimit = -0.998f;	// Cos(176deg)
 FAutoConsoleVariableRef CVarChaosJointDegenerateRotationLimit(TEXT("p.Chaos.Joint.DegenerateRotationLimit"), Chaos_Joint_DegenerateRotationLimit, TEXT("Cosine of the swing angle that is considered degerenerate (default Cos(176deg))"));
diff --git a/Engine/Source/Runtime/Experimental/Chaos/Private/Chaos/PBDJointConstraints.cpp b/Engine/Source/Runtime/Experimental/Chaos/Private/Chaos/PBDJointConstraints.cpp
--- a/Engine/Source/Runtime/Experimental/Chaos/Private/Chaos/PBDJointConstraints.cpp
+++ b/Engine/Source/Runtime/Experimental/Chaos/Private/Chaos/PBDJointConstraints.cpp
@@ -28,12 +28,9 @@ FAutoConsoleVariableRef CVarChaosJointBatching(TEXT("p.Chaos.Joint.Batching"), b
 int32 bChaos_Joint_MaxBatchSize = 1000;
 FAutoConsoleVariableRef CVarChaosJointBatchSize(TEXT("p.Chaos.Joint.MaxBatchSize"), bChaos_Joint_MaxBatchSize, TEXT(""));
 
-#if !INTEL_ISPC
-const bool bChaos_Joint_ISPC_Enabled = false;
-#elif UE_BUILD_SHIPPING
-const bool bChaos_Joint_ISPC_Enabled = true;
-#else
-extern bool bChaos_Joint_ISPC_Enabled;
+#if INTEL_ISPC && !UE_BUILD_SHIPPING
+bool bChaos_Joint_ISPC_Enabled = true;
+FAutoConsoleVariableRef CVarChaosJointISPCEnabled(TEXT("p.Chaos.Joint.ISPC"), bChaos_Joint_ISPC_Enabled, TEXT("Whether to use ISPC optimizations in the Joint Solver"));
 #endif
 
 namespace Chaos
diff --git a/Engine/Source/Runtime/Experimental/Chaos/Public/Chaos/PBDJointConstraintUtilities.h b/Engine/Source/Runtime/Experimental/Chaos/Public/Chaos/PBDJointConstraintUtilities.h
--- a/Engine/Source/Runtime/Experimental/Chaos/Public/Chaos/PBDJointConstraintUtilities.h
+++ b/Engine/Source/Runtime/Experimental/Chaos/Public/Chaos/PBDJointConstraintUtilities.h
@@ -234,4 +234,13 @@ namespace Chaos
 		static FVec3 GetLineLimitedPositionError(const FVec3& CX, const FVec3& Axis, const FReal Limit, const EJointMotionType AxisMotion);
 		static FVec3 GetLimitedPositionError(const FPBDJointSettings& JointSettings, const FRotation3& R0, const FVec3& CX);
 	};
-}
\ No newline at end of file
+}
+
+// Support ISPC enable/disable in non-shipping builds
+#if !INTEL_ISPC
+const bool bChaos_Joint_ISPC_Enabled = false;
+#elif UE_BUILD_SHIPPING
+const bool bChaos_Joint_ISPC_Enabled = true;
+#else
+extern bool bChaos_Joint_ISPC_Enabled;
+#endif
