Fix non-unity CIS (JointSolverConstraints.cpp)

diff --git a/Engine/Source/Runtime/Experimental/Chaos/Private/Chaos/Joint/JointSolverConstraints.cpp b/Engine/Source/Runtime/Experimental/Chaos/Private/Chaos/Joint/JointSolverConstraints.cpp
--- a/Engine/Source/Runtime/Experimental/Chaos/Private/Chaos/Joint/JointSolverConstraints.cpp
+++ b/Engine/Source/Runtime/Experimental/Chaos/Private/Chaos/Joint/JointSolverConstraints.cpp
@@ -7,6 +7,8 @@
 #include "Chaos/PBDJointConstraintUtilities.h"
 #include "Chaos/Utilities.h"
 
+#include "HAL/IConsoleManager.h"
+
 //PRAGMA_DISABLE_OPTIMIZATION
 
 int32 bChaos_Joint_MultiDimension = true;
