Chaos - fix collision jitter for multi-shape particle pairs with multiple collisions between them

diff --git a/Engine/Source/Runtime/Experimental/Chaos/Private/Chaos/PBDCollisionConstraintsContact.cpp b/Engine/Source/Runtime/Experimental/Chaos/Private/Chaos/PBDCollisionConstraintsContact.cpp
--- a/Engine/Source/Runtime/Experimental/Chaos/Private/Chaos/PBDCollisionConstraintsContact.cpp
+++ b/Engine/Source/Runtime/Experimental/Chaos/Private/Chaos/PBDCollisionConstraintsContact.cpp
@@ -339,14 +339,8 @@ namespace Chaos
 
 			for (int32 PairIt = 0; PairIt < IterationParameters.NumPairIterations; ++PairIt)
 			{
-				// Collision is already up-to-date on first iteration (we either just detected it, or updated it in DetectCollisions)
-				// @todo(ccaulfield): this is not great - try to do something nicer like a dirty flag on the constraint?
-				// In particular it is not right if the Collisions are not the first constraints to be solved...
-				const bool bNeedCollisionUpdate = (PairIt > 0) || (IterationParameters.Iteration > 0) || (*IterationParameters.AlwaysUpdateManifold);
-				if (bNeedCollisionUpdate)
-				{
-					Collisions::Update(Constraint, ParticleParameters.CullDistance);
-				}
+				// Update the contact information based on current particles' positions
+				Collisions::Update(Constraint, ParticleParameters.CullDistance);
 
 				if (Constraint.GetPhi() >= ParticleParameters.ShapePadding)
 				{
