#UE4 Fix NaN that may occur in TCapsule::RaycastFast

diff --git a/Engine/Source/Runtime/Experimental/Chaos/Public/Chaos/Capsule.h b/Engine/Source/Runtime/Experimental/Chaos/Public/Chaos/Capsule.h
--- a/Engine/Source/Runtime/Experimental/Chaos/Public/Chaos/Capsule.h
+++ b/Engine/Source/Runtime/Experimental/Chaos/Public/Chaos/Capsule.h
@@ -193,12 +193,12 @@ namespace Chaos
 					const bool bSingleHit = QuarterUnderRoot < Epsilon;
 					if (bSingleHit)
 					{
-						Time = -HalfB / A;
+						Time = (A == 0) ? 0 : (-HalfB / A);
 
 					}
 					else
 					{
-						Time = (-HalfB - FMath::Sqrt(QuarterUnderRoot)) / A; //we already checked for initial overlap so just take smallest time
+						Time = (A == 0) ? 0 : ((-HalfB - FMath::Sqrt(QuarterUnderRoot)) / A); //we already checked for initial overlap so just take smallest time
 						if (Time < 0)	//we must have passed the cylinder
 						{
 							return false;
