Chaos : Ensuring more of the fields system is actually implemented (moved to Handles)

diff --git a/Engine/Source/Runtime/Experimental/ChaosSolvers/Private/PhysicsProxy/FieldSystemPhysicsProxy.cpp b/Engine/Source/Runtime/Experimental/ChaosSolvers/Private/PhysicsProxy/FieldSystemPhysicsProxy.cpp
--- a/Engine/Source/Runtime/Experimental/ChaosSolvers/Private/PhysicsProxy/FieldSystemPhysicsProxy.cpp
+++ b/Engine/Source/Runtime/Experimental/ChaosSolvers/Private/PhysicsProxy/FieldSystemPhysicsProxy.cpp
@@ -63,10 +63,6 @@ void FFieldSystemPhysicsProxy::FieldParameterUpdateCallback(
 
 	if (Commands.Num() && InSolver)
 	{
-#if TODO_REIMPLEMENT_RIGID_CLUSTERING
-		const Chaos::TArrayCollectionArray<Chaos::ClusterId> & ClusterIDs = CurrentSolver->GetRigidClustering().GetClusterIdsArray();
-#endif // TODO_REIMPLEMENT_RIGID_CLUSTERING
-
 		TArray<FFieldSystemCommand>* CommandListPtr = GetSolverCommandList(InSolver);
 		if(!CommandListPtr)
 		{
@@ -241,34 +237,38 @@ void FFieldSystemPhysicsProxy::FieldParameterUpdateCallback(
 						CurrentSolver->GetParticles().UpdateGeometryCollectionViews(true);
 					}
 
-#if TODO_REIMPLEMENT_RIGID_CLUSTERING
-					//  Update all cluster bodies based on the changes in the kinematic state.
-					const Chaos::TArrayCollectionArray<Chaos::ClusterId>& ClusterIdArray = CurrentSolver->GetRigidClustering().GetClusterIdsArray();
-					for (int32 ActiveParticleIndex : CurrentSolver->ActiveIndices())
+					const Chaos::TPBDRigidsSOAs<float, 3>& SolverParticles = CurrentSolver->GetParticles();
+					auto& Clustering = CurrentSolver->GetEvolution()->GetRigidClustering();
+					const auto& ClusterMap = Clustering.GetChildrenMap();
+
+					const Chaos::TParticleView<Chaos::TGeometryParticles<float, 3>>& ParticleView =
+						SolverParticles.GetNonDisabledView();
+
+					for (Chaos::TParticleIterator<Chaos::TGeometryParticles<float, 3>> It = ParticleView.Begin(), ItEnd = ParticleView.End();
+						It != ItEnd; ++It)
 					{
-						if (ClusterIdArray[ActiveParticleIndex].NumChildren)
+						const auto* Clustered = It->Handle()->CastToClustered();
+						if (Clustered && Clustered->ClusterIds().NumChildren)
 						{
-							CurrentSolver->GetRigidClustering().UpdateKinematicProperties(ActiveParticleIndex);
+							Chaos::TPBDRigidParticleHandle<float, 3>* RigidHandle = It->Handle()->CastToRigidParticle();
+							check(RigidHandle);
+							Clustering.UpdateKinematicProperties(RigidHandle);
 						}
 					}
-#endif // TODO_REIMPLEMENT_RIGID_CLUSTERING
 
 				}
 				CommandsToRemove.Add(CommandIndex);
 			}
 			else if (Command.TargetAttribute == GetFieldPhysicsName(EFieldPhysicsType::Field_ActivateDisabled))
 			{
-				FFieldSystemPhysicsProxy::ContiguousIndices(IndicesArray, CurrentSolver, ResolutionType, IndicesArray.Num() != Particles.Size());
-				if (IndicesArray.Num())
+				if (Handles.Num())
 				{
-					TArrayView<ContextIndex> IndexView(&(IndicesArray[0]), IndicesArray.Num());
-
-					FVector * tptr = &(Particles.X(0));
-					TArrayView<FVector> SamplesView(tptr, int32(Particles.Size()));
+					TArrayView<FVector> SamplePointsView(&(SamplePoints[0]), SamplePoints.Num());
+					TArrayView<ContextIndex> SampleIndicesView(&(SampleIndices[0]), SampleIndices.Num());
 
 					FFieldContext Context{
-						IndexView, // @todo(chaos) important: an empty index array should evaluate everything
-						SamplesView,
+						SampleIndicesView, // @todo(chaos) important: an empty index array should evaluate everything
+						SamplePointsView,
 						Command.MetaData
 					};
 
@@ -304,19 +304,15 @@ void FFieldSystemPhysicsProxy::FieldParameterUpdateCallback(
 
 #if TODO_REIMPLEMENT_RIGID_CLUSTERING
 					// transfer results to rigid system.
-					int32 FloorIndex = CurrentSolver->GetFloorIndex();
 					int32 NumSamples = Context.SampleIndices.Num();
 					for (int32 SampleIndex = 0; SampleIndex < NumSamples; SampleIndex++)
 					{
 						int32 RigidBodyIndex = Context.SampleIndices[SampleIndex].Result;
-						if (RigidBodyIndex != FloorIndex) // ignore the floor
-						{ 
-							if (DynamicStateView[RigidBodyIndex] == 0 && Particles.Disabled(RigidBodyIndex))
-							{
-								ensure(CurrentSolver->GetRigidClustering().GetClusterIdsArray()[RigidBodyIndex].Id == INDEX_NONE);
-								CurrentSolver->GetEvolution()->EnableParticle(RigidBodyIndex, INDEX_NONE);
-								Particles.SetObjectState(RigidBodyIndex, Chaos::EObjectStateType::Dynamic);
-							}
+						if (DynamicStateView[RigidBodyIndex] == 0 && Particles.Disabled(RigidBodyIndex))
+						{
+							ensure(CurrentSolver->GetRigidClustering().GetClusterIdsArray()[RigidBodyIndex].Id == INDEX_NONE);
+							CurrentSolver->GetEvolution()->EnableParticle(RigidBodyIndex, INDEX_NONE);
+							Particles.SetObjectState(RigidBodyIndex, Chaos::EObjectStateType::Dynamic);
 						}
 					}
 #endif
@@ -330,17 +326,14 @@ void FFieldSystemPhysicsProxy::FieldParameterUpdateCallback(
 				if (ensureMsgf(Command.RootNode->Type() == FFieldNode<float>::StaticType(),
 					TEXT("Field based evaluation of the simulations 'Strain' parameter expects float field inputs.")))
 				{
-					FFieldSystemPhysicsProxy::ContiguousIndices(IndicesArray, CurrentSolver, ResolutionType, IndicesArray.Num() != Particles.Size());
-					if (IndicesArray.Num())
+					if (Handles.Num())
 					{
-						TArrayView<ContextIndex> IndexView(&(IndicesArray[0]), IndicesArray.Num());
-
-						FVector * tptr = &(Particles.X(0));
-						TArrayView<FVector> SamplesView(tptr, int32(Particles.Size()));
+						TArrayView<FVector> SamplePointsView(&(SamplePoints[0]), SamplePoints.Num());
+						TArrayView<ContextIndex> SampleIndicesView(&(SampleIndices[0]), SampleIndices.Num());
 
 						FFieldContext Context{
-							IndexView,
-							SamplesView,
+							SampleIndicesView, // @todo(chaos) important: an empty index array should evaluate everything
+							SamplePointsView,
 							Command.MetaData
 						};
 
@@ -358,7 +351,7 @@ void FFieldSystemPhysicsProxy::FieldParameterUpdateCallback(
 						// cost of initializing them all, or it's cheaper and potentially 
 						// more cache coherent to just initialize them all.  I'm thinking
 						// the latter may be more likely...
-						StrainSamples.AddZeroed(SamplesView.Num());
+						StrainSamples.AddZeroed(SamplePointsView.Num());
 
 						TArrayView<float> FloatBuffer(&StrainSamples[0], StrainSamples.Num());
 						static_cast<const FFieldNode<float> *>(Command.RootNode.Get())->Evaluate(Context, FloatBuffer);
@@ -385,15 +378,16 @@ void FFieldSystemPhysicsProxy::FieldParameterUpdateCallback(
 				if (ensureMsgf(Command.RootNode->Type() == FFieldNode<float>::StaticType(),
 					TEXT("Field based evaluation of the simulations 'Disabled' parameter expects float field inputs.")))
 				{
-					FFieldSystemPhysicsProxy::ContiguousIndices(IndicesArray, CurrentSolver, ResolutionType, IndicesArray.Num() != Particles.Size());
-					if (IndicesArray.Num())
+					if (Handles.Num())
 					{
 						TArrayView<FVector> SamplePointsView(&(SamplePoints[0]), SamplePoints.Num());
 						TArrayView<ContextIndex> SampleIndicesView(&(SampleIndices[0]), SampleIndices.Num());
-						FFieldContext Context(
-							SampleIndicesView,
+
+						FFieldContext Context{
+							SampleIndicesView, // @todo(chaos) important: an empty index array should evaluate everything
 							SamplePointsView,
-							Command.MetaData);
+							Command.MetaData
+						};
 
 						TArray<float> LocalResults;
 						LocalResults.AddZeroed(Handles.Num());
@@ -462,17 +456,14 @@ void FFieldSystemPhysicsProxy::FieldParameterUpdateCallback(
 				if (ensureMsgf(Command.RootNode->Type() == FFieldNode<FVector>::StaticType(),
 					TEXT("Field based evaluation of the simulations 'AngularVelocity' parameter expects FVector field inputs.")))
 				{
-					FFieldSystemPhysicsProxy::ContiguousIndices(IndicesArray, CurrentSolver, ResolutionType, IndicesArray.Num() != Particles.Size());
-					if (IndicesArray.Num())
+					if (Handles.Num())
 					{
-						TArrayView<ContextIndex> IndexView(&(IndicesArray[0]), IndicesArray.Num());
-
-						FVector * tptr = &(Particles.X(0));
-						TArrayView<FVector> SamplesView(tptr, int32(Particles.Size()));
+						TArrayView<FVector> SamplePointsView(&(SamplePoints[0]), SamplePoints.Num());
+						TArrayView<ContextIndex> SampleIndicesView(&(SampleIndices[0]), SampleIndices.Num());
 
 						FFieldContext Context{
-							IndexView,
-							SamplesView,
+							SampleIndicesView, // @todo(chaos) important: an empty index array should evaluate everything
+							SamplePointsView,
 							Command.MetaData
 						};
 
@@ -607,17 +598,14 @@ void FFieldSystemPhysicsProxy::FieldParameterUpdateCallback(
 				if (ensureMsgf(Command.RootNode->Type() == FFieldNode<int32>::StaticType(),
 					TEXT("Field based evaluation of the simulations 'CollisionGroup' parameter expects int field inputs.")))
 				{
-					FFieldSystemPhysicsProxy::ContiguousIndices(IndicesArray, CurrentSolver, ResolutionType, IndicesArray.Num() != Particles.Size());
-					if (IndicesArray.Num())
+					if (Handles.Num())
 					{
-						TArrayView<ContextIndex> IndexView(&(IndicesArray[0]), IndicesArray.Num());
-
-						FVector * tptr = &(Particles.X(0));
-						TArrayView<FVector> SamplesView(tptr, int32(Particles.Size()));
+						TArrayView<FVector> SamplePointsView(&(SamplePoints[0]), SamplePoints.Num());
+						TArrayView<ContextIndex> SampleIndicesView(&(SampleIndices[0]), SampleIndices.Num());
 
 						FFieldContext Context{
-							IndexView,
-							SamplesView,
+							SampleIndicesView, // @todo(chaos) important: an empty index array should evaluate everything
+							SamplePointsView,
 							Command.MetaData
 						};
 
@@ -632,21 +620,17 @@ void FFieldSystemPhysicsProxy::FieldParameterUpdateCallback(
 			{
 				SCOPE_CYCLE_COUNTER(STAT_ParamUpdateField_PositionStatic);
 
-#if TODO_REIMPLEMENT_FIELDS_TO_USE_PARTICLEHANDLES
 				if (ensureMsgf(Command.RootNode->Type() == FFieldNode<int32>::StaticType(),
 					TEXT("Field based evaluation of the simulations 'Position' parameter expects integer field inputs.")))
 				{
-					FFieldSystemPhysicsProxy::ContiguousIndices(IndicesArray, CurrentSolver, ResolutionType, IndicesArray.Num() != Particles.Size());
-					if (IndicesArray.Num())
+					if (Handles.Num())
 					{
-						TArrayView<ContextIndex> IndexView(&(IndicesArray[0]), IndicesArray.Num());
-
-						FVector * tptr = &(Particles.X(0));
-						TArrayView<FVector> SamplesView(tptr, int32(Particles.Size()));
+						TArrayView<FVector> SamplePointsView(&(SamplePoints[0]), SamplePoints.Num());
+						TArrayView<ContextIndex> SampleIndicesView(&(SampleIndices[0]), SampleIndices.Num());
 
 						FFieldContext Context{
-							IndexView,
-							SamplesView,
+							SampleIndicesView, // @todo(chaos) important: an empty index array should evaluate everything
+							SamplePointsView,
 							Command.MetaData
 						};
 
@@ -659,6 +643,7 @@ void FFieldSystemPhysicsProxy::FieldParameterUpdateCallback(
 						TArrayView<int32> ResultsView(&(Results[0]), Results.Num());
 						static_cast<const FFieldNode<int32> *>(Command.RootNode.Get())->Evaluate(Context, ResultsView);
 
+#if TODO_REIMPLEMENT_KINEMATIC_PROXY
 						for (const ContextIndex& CIndex : IndicesArray)
 						{
 							const int32 i = CIndex.Result;
@@ -671,36 +656,32 @@ void FFieldSystemPhysicsProxy::FieldParameterUpdateCallback(
 								}
 								else
 								{
-									int32 Index = PositionTarget.Add(i, Particles.X(i));
+									int32 Index = PositionTarget.AddConstraint(Particles.Handle(i), Particles.X(i)); //??
 									PositionTargetedParticles.Add(i, Index);
 								}
 							}
 						}
+#endif
 					}
 				}
-#endif
+
 				CommandsToRemove.Add(CommandIndex);
 			}
 			else if (Command.TargetAttribute == GetFieldPhysicsName(EFieldPhysicsType::Field_PositionTarget))
 			{
 				SCOPE_CYCLE_COUNTER(STAT_ParamUpdateField_PositionTarget);
 
-#if TODO_REIMPLEMENT_FIELDS_TO_USE_PARTICLEHANDLES
-
 				if (ensureMsgf(Command.RootNode->Type() == FFieldNode<FVector>::StaticType(),
 					TEXT("Field based evaluation of the simulations 'PositionTarget' parameter expects vector field inputs.")))
 				{
-					FFieldSystemPhysicsProxy::ContiguousIndices(IndicesArray, CurrentSolver, ResolutionType, IndicesArray.Num() != Particles.Size());
-					if (IndicesArray.Num())
+					if (Handles.Num())
 					{
-						TArrayView<ContextIndex> IndexView(&(IndicesArray[0]), IndicesArray.Num());
-
-						FVector * tptr = &(Particles.X(0));
-						TArrayView<FVector> SamplesView(tptr, int32(Particles.Size()));
+						TArrayView<FVector> SamplePointsView(&(SamplePoints[0]), SamplePoints.Num());
+						TArrayView<ContextIndex> SampleIndicesView(&(SampleIndices[0]), SampleIndices.Num());
 
 						FFieldContext Context{
-							IndexView,
-							SamplesView,
+							SampleIndicesView, // @todo(chaos) important: an empty index array should evaluate everything
+							SamplePointsView,
 							Command.MetaData
 						};
 
@@ -713,6 +694,7 @@ void FFieldSystemPhysicsProxy::FieldParameterUpdateCallback(
 						TArrayView<FVector> ResultsView(&(Results[0]), Results.Num());
 						static_cast<const FFieldNode<FVector> *>(Command.RootNode.Get())->Evaluate(Context, ResultsView);
 
+#if TODO_REIMPLEMENT_KINEMATIC_PROXY
 						for (const ContextIndex& CIndex : IndicesArray)
 						{
 							const int32 i = CIndex.Result;
@@ -730,9 +712,9 @@ void FFieldSystemPhysicsProxy::FieldParameterUpdateCallback(
 								}
 							}
 						}
+#endif
 					}
 				}
-#endif
 
 				CommandsToRemove.Add(CommandIndex);
 			}
@@ -743,17 +725,14 @@ void FFieldSystemPhysicsProxy::FieldParameterUpdateCallback(
 				if (ensureMsgf(Command.RootNode->Type() == FFieldNode<int32>::StaticType(),
 					TEXT("Field based evaluation of the simulations 'Position' parameter expects integer field inputs.")))
 				{
-					FFieldSystemPhysicsProxy::ContiguousIndices(IndicesArray, CurrentSolver, ResolutionType, IndicesArray.Num() != Particles.Size());
-					if (IndicesArray.Num())
+					if (Handles.Num())
 					{
-						TArrayView<ContextIndex> IndexView(&(IndicesArray[0]), IndicesArray.Num());
-
-						FVector * tptr = &(Particles.X(0));
-						TArrayView<FVector> SamplesView(tptr, int32(Particles.Size()));
+						TArrayView<FVector> SamplePointsView(&(SamplePoints[0]), SamplePoints.Num());
+						TArrayView<ContextIndex> SampleIndicesView(&(SampleIndices[0]), SampleIndices.Num());
 
 						FFieldContext Context{
-							IndexView,
-							SamplesView,
+							SampleIndicesView, // @todo(chaos) important: an empty index array should evaluate everything
+							SamplePointsView,
 							Command.MetaData
 						};
 
@@ -1041,122 +1020,118 @@ void FFieldSystemPhysicsProxy::BufferCommand(Chaos::FPhysicsSolver* InSolver, co
 	(*ExistingList)->Add(InCommand);
 }
 
+//void FFieldSystemPhysicsProxy::GetParticleHandles(
+//	TArray<Chaos::TGeometryParticleHandle<float,3>*>& Handles,
+//	const Chaos::FPhysicsSolver* RigidSolver,
+//	const EFieldResolutionType ResolutionType,
+//	const bool bForce)
+//{
+//	Handles.SetNum(0, false);
+//	if (!bForce)
+//		return;
+//
+//	const Chaos::TPBDRigidsSOAs<float, 3>& SolverParticles = RigidSolver->GetParticles();
+//
+//	if (ResolutionType == EFieldResolutionType::Field_Resolution_Maximum)
+//	{
+//		// const TParticleView<TGeometryParticles<T, d>>& TPBDRigidSOAs<T,d>::GetAllParticlesView()
+//		const Chaos::TParticleView<Chaos::TGeometryParticles<float, 3>> &ParticleView = 
+//			SolverParticles.GetAllParticlesView();
+//		Handles.Reserve(ParticleView.Num());
+//
+//		// TParticleIterator<TSOA> Begin() const, TSOA = TGeometryParticles<T, d>
+//		for (Chaos::TParticleIterator<Chaos::TGeometryParticles<float, 3>> It = ParticleView.Begin(), ItEnd = ParticleView.End();
+//			It != ItEnd; ++It)
+//		{
+//			const Chaos::TTransientGeometryParticleHandle<float,3> *Handle = &(*It);
+//			// PBDRigidsSOAs.h only has a const version of GetAllParticlesView() - is that wrong?
+//			Handles.Add(GetHandleHelper(const_cast<Chaos::TTransientGeometryParticleHandle<float,3>*>(Handle)));
+//		}
+//	}
+//	else if (ResolutionType == EFieldResolutionType::Field_Resolution_Minimal)
+//	{
+//		const Chaos::TParticleView<Chaos::TGeometryParticles<float, 3>> &ParticleView = 
+//			SolverParticles.GetNonDisabledView();
+//		Handles.Reserve(ParticleView.Num());
+//		for (Chaos::TParticleIterator<Chaos::TGeometryParticles<float, 3>> It = ParticleView.Begin(), ItEnd = ParticleView.End();
+//			It != ItEnd; ++It)
+//		{
+//			const Chaos::TTransientGeometryParticleHandle<float,3> *Handle = &(*It);
+//			Handles.Add(GetHandleHelper(const_cast<Chaos::TTransientGeometryParticleHandle<float,3>*>(Handle)));
+//		}
+//	}
+//	else if (ResolutionType == EFieldResolutionType::Field_Resolution_DisabledParents)
+//	{
+//		check(false); // unimplemented
+//	}
+//}
+
 void FFieldSystemPhysicsProxy::GetParticleHandles(
-	TArray<Chaos::TGeometryParticleHandle<float,3>*>& Handles,
+	TArray<Chaos::TGeometryParticleHandle<float, 3>*>& Handles,
 	const Chaos::FPhysicsSolver* RigidSolver,
-	const EFieldResolutionType ResolutionType,
+	const EFieldResolutionType ResolutionType, 
 	const bool bForce)
 {
 	Handles.SetNum(0, false);
 	if (!bForce)
+	{
 		return;
+	}
 
-	const Chaos::TPBDRigidsSOAs<float, 3>& SolverParticles = RigidSolver->GetParticles(); // const?
+	const Chaos::TPBDRigidsSOAs<float, 3>& SolverParticles = RigidSolver->GetParticles();
 
-	if (ResolutionType == EFieldResolutionType::Field_Resolution_Maximum)
+	if (ResolutionType == EFieldResolutionType::Field_Resolution_Minimal)
 	{
-		// const TParticleView<TGeometryParticles<T, d>>& TPBDRigidSOAs<T,d>::GetAllParticlesView()
-		const Chaos::TParticleView<Chaos::TGeometryParticles<float, 3>> &ParticleView = 
-			SolverParticles.GetAllParticlesView();
-		Handles.Reserve(ParticleView.Num());
+		const auto& Clustering = RigidSolver->GetEvolution()->GetRigidClustering();
+		const auto& ClusterMap = Clustering.GetChildrenMap();
 
-		// TParticleIterator<TSOA> Begin() const, TSOA = TGeometryParticles<T, d>
-		for (Chaos::TParticleIterator<Chaos::TGeometryParticles<float, 3>> It = ParticleView.Begin(), ItEnd = ParticleView.End();
-			It != ItEnd; ++It)
-		{
-			const Chaos::TTransientGeometryParticleHandle<float,3> *Handle = &(*It);
-			// PBDRigidsSOAs.h only has a const version of GetAllParticlesView() - is that wrong?
-			Handles.Add(GetHandleHelper(const_cast<Chaos::TTransientGeometryParticleHandle<float,3>*>(Handle)));
-		}
-	}
-	else if (ResolutionType == EFieldResolutionType::Field_Resolution_Minimal)
-	{
-		const Chaos::TParticleView<Chaos::TGeometryParticles<float, 3>> &ParticleView = 
+		const Chaos::TParticleView<Chaos::TGeometryParticles<float, 3>>& ParticleView =
 			SolverParticles.GetNonDisabledView();
-		Handles.Reserve(ParticleView.Num());
+		Handles.Reserve(ParticleView.Num()); // ?? what about additional number of children added
 		for (Chaos::TParticleIterator<Chaos::TGeometryParticles<float, 3>> It = ParticleView.Begin(), ItEnd = ParticleView.End();
 			It != ItEnd; ++It)
 		{
-			const Chaos::TTransientGeometryParticleHandle<float,3> *Handle = &(*It);
-			Handles.Add(GetHandleHelper(const_cast<Chaos::TTransientGeometryParticleHandle<float,3>*>(Handle)));
-
-		}
-	}
-	else if (ResolutionType == EFieldResolutionType::Field_Resolution_DisabledParents)
-	{
-		check(false); // unimplemented
-	}
-}
-
-void FFieldSystemPhysicsProxy::ContiguousIndices(
-	TArray<ContextIndex>& Array, 
-	const Chaos::FPhysicsSolver* RigidSolver, 
-	const EFieldResolutionType ResolutionType, 
-	const bool bForce)
-{
-#if TODO_REIMPLEMENT_RIGID_CLUSTERING
-	if (!bForce)
-	{
-		return;
-	}
-
-	if (ResolutionType == EFieldResolutionType::Field_Resolution_Minimal)
-	{
-		Array.SetNum(0, false);
-
-		int32 IndexCount = 0;
-		int32 FloorIndex = RigidSolver->GetFloorIndex();
+			const Chaos::TTransientGeometryParticleHandle<float, 3>* Handle = &(*It);
+			Handles.Add(GetHandleHelper(const_cast<Chaos::TTransientGeometryParticleHandle<float, 3>*>(Handle)));
 
-#if TODO_REIMPLEMENT_RIGID_CLUSTERING
-		const Chaos::FPhysicsSolver::FClusteringType& Clustering = RigidSolver->GetRigidClustering();
-		const Chaos::TArrayCollectionArray<Chaos::ClusterId>& ClusterIdArray = Clustering.GetClusterIdsArray();
-		const Chaos::FPhysicsSolver::FClusteringType::FClusterMap &  ClusterMap = Clustering.GetChildrenMap();
-#endif // TODO_REIMPLEMENT_RIGID_CLUSTERING
-
-		for (int32 ActiveParticleIndex : RigidSolver->NonDisabledIndices())
-		{
-			if (ClusterIdArray[ActiveParticleIndex].NumChildren)
+			const auto* Clustered = Handle->CastToClustered();
+			if (Clustered && Clustered->ClusterIds().NumChildren)
 			{
-				for (uint32 ClusterChild : *ClusterMap[ActiveParticleIndex])
+				Chaos::TPBDRigidParticleHandle<float, 3>* RigidHandle = (*It).Handle()->CastToRigidParticle();
+				if (ClusterMap.Contains(RigidHandle))
 				{
-					Array.Add(ContextIndex(ClusterChild, ClusterChild));
+					for (Chaos::TPBDRigidParticleHandle<float, 3>* Child : ClusterMap[RigidHandle])
+					{
+						Handles.Add(Child);
+					}
 				}
 			}
-
-			if( ActiveParticleIndex != FloorIndex )
-			{
-				Array.Add(ContextIndex(ActiveParticleIndex, ActiveParticleIndex));
-			}
 		}
 	}
 	else if (ResolutionType == EFieldResolutionType::Field_Resolution_DisabledParents)
 	{
-		Array.SetNum(0, false);
+		const auto& Clustering = RigidSolver->GetEvolution()->GetRigidClustering();
+		const auto& ClusterMap = Clustering.GetChildrenMap();
+		Handles.Reserve(Clustering.GetTopLevelClusterParents().Num());
 
-		int32 IndexCount = 0;
-		int32 FloorIndex = RigidSolver->GetFloorIndex();
-		const Chaos::FPhysicsSolver::FClusteringType& Clustering = RigidSolver->GetRigidClustering();
-		const Chaos::TArrayCollectionArray<Chaos::ClusterId>& ClusterIdArray = Clustering.GetClusterIdsArray();
-
-		for (int32 TopLevelParent : Clustering.GetTopLevelClusterParents())
+		for (Chaos::TPBDRigidClusteredParticleHandle<float, 3>* TopLevelParent : Clustering.GetTopLevelClusterParents())
 		{
-			if( TopLevelParent != FloorIndex )
-			{
-				Array.Add(ContextIndex(TopLevelParent, TopLevelParent));
-			}
+			Handles.Add(TopLevelParent);
 		}
 	}
 	else if (ResolutionType == EFieldResolutionType::Field_Resolution_Maximum)
 	{
-		const Chaos::FPhysicsSolver::FParticlesType & Particles = RigidSolver->GetRigidParticles();
-		Array.SetNum(Particles.Size());
-		for (int32 i = 0; i < Array.Num(); ++i)
+		const Chaos::TParticleView<Chaos::TGeometryParticles<float, 3>>& ParticleView =
+			SolverParticles.GetAllParticlesView();
+		Handles.Reserve(ParticleView.Num());
+
+		for (Chaos::TParticleIterator<Chaos::TGeometryParticles<float, 3>> It = ParticleView.Begin(), ItEnd = ParticleView.End();
+			It != ItEnd; ++It)
 		{
-			Array[i].Sample = i;
-			Array[i].Result = i;
+			const Chaos::TTransientGeometryParticleHandle<float, 3>* Handle = &(*It);
+			Handles.Add(GetHandleHelper(const_cast<Chaos::TTransientGeometryParticleHandle<float, 3>*>(Handle)));
 		}
 	}
-#endif
 }
 
 TArray<FFieldSystemCommand>* FFieldSystemPhysicsProxy::GetSolverCommandList(const Chaos::FPhysicsSolver* InSolver)
diff --git a/Engine/Source/Runtime/Experimental/ChaosSolvers/Public/PhysicsProxy/FieldSystemPhysicsProxy.h b/Engine/Source/Runtime/Experimental/ChaosSolvers/Public/PhysicsProxy/FieldSystemPhysicsProxy.h
--- a/Engine/Source/Runtime/Experimental/ChaosSolvers/Public/PhysicsProxy/FieldSystemPhysicsProxy.h
+++ b/Engine/Source/Runtime/Experimental/ChaosSolvers/Public/PhysicsProxy/FieldSystemPhysicsProxy.h
@@ -105,11 +105,6 @@ public:
 	 * When \c Minimum is set only the ActiveIndices and the direct children of 
 	 * the active clusters are set in the \p IndicesArray.
 	 */
-	static void ContiguousIndices(
-		TArray<ContextIndex>& IndicesArray, 
-		const Chaos::FPhysicsSolver* RigidSolver, 
-		const EFieldResolutionType ResolutionType, 
-		const bool bForce = true);
 
 	static void GetParticleHandles(
 		TArray<Chaos::TGeometryParticleHandle<float,3>*>& Handles,
