Added the ability for sim-writable properties to be recorded, even if the GT side is completely clean

diff --git a/Engine/Source/Programs/HeadlessChaos/Private/HeadlessChaosTestRewind.cpp b/Engine/Source/Programs/HeadlessChaos/Private/HeadlessChaosTestRewind.cpp
--- a/Engine/Source/Programs/HeadlessChaos/Private/HeadlessChaosTestRewind.cpp
+++ b/Engine/Source/Programs/HeadlessChaos/Private/HeadlessChaosTestRewind.cpp
@@ -28,7 +28,7 @@ namespace ChaosTest {
 		Solver->UpdateGameThreadStructures();
 	}
 
-	GTEST_TEST(RewindTest,DataCapture)
+	GTEST_TEST(RewindTest,DataCapture1)
 	{
 		auto Sphere = TSharedPtr<FImplicitObject,ESPMode::ThreadSafe>(new TSphere<float,3>(TVector<float,3>(0),10));
 		auto Box = TSharedPtr<FImplicitObject,ESPMode::ThreadSafe>(new TBox<float,3>(FVec3(0),FVec3(1)));
@@ -50,16 +50,10 @@ namespace ChaosTest {
 		Particle->SetGeometry(Sphere);
 		Solver->RegisterObject(Particle.Get());
 
-		auto Dynamic = TPBDRigidParticle<float,3>::CreateParticle();
-		Dynamic->SetGeometry(Sphere);
-		Solver->RegisterObject(Dynamic.Get());
-		
-
 		for(int Step = 0; Step < 11; ++Step)
 		{
 			//property that changes every step
 			Particle->SetX(FVec3(0,0,100 - Step));
-			Dynamic->SetF(FVec3(0,0,Step + 1));
 
 			//property that changes once half way through
 			if(Step == 3)
@@ -114,18 +108,133 @@ namespace ChaosTest {
 					EXPECT_EQ(ParticleState->Geometry().Get(),Box2.Get());
 				}
 			}
+		}
+		
+		//no data at head
+		EXPECT_EQ(RewindData->GetStateAtFrame(*Particle,10), nullptr);
+
+		// Throw out the proxy
+		Solver->UnregisterObject(Particle.Get());
+
+		Module->DestroySolver(Solver);
+	}
+
+
+	GTEST_TEST(RewindTest,DataCapture2)
+	{
+		auto Sphere = TSharedPtr<FImplicitObject,ESPMode::ThreadSafe>(new TSphere<float,3>(TVector<float,3>(0),10));
+
+		FChaosSolversModule* Module = FChaosSolversModule::GetModule();
+		Module->ChangeThreadingMode(EChaosThreadingMode::SingleThread);
+
+		// Make a solver
+		FPhysicsSolver* Solver = Module->CreateSolver(nullptr,ESolverFlags::Standalone);
+		Solver->SetEnabled(true);
+
+		Solver->EnableRewindCapture(20);
+
+
+		// Make particles
+		auto Particle = TPBDRigidParticle<float,3>::CreateParticle();
+
+		Particle->SetGeometry(Sphere);
+		Solver->RegisterObject(Particle.Get());
+
+		for(int Step = 0; Step < 11; ++Step)
+		{
+			//sim-writable property that changes every step
+			Particle->SetF(FVec3(0,0,Step + 1));
+
+
+			TickSolverHelper(Module,Solver);
+		}
+
+		const FRewindData* RewindData = Solver->GetRewindData();
+
+		//check state at every step except latest
+		for(int Step = 0; Step < 10; ++Step)
+		{
+			const auto ParticleState = RewindData->GetStateAtFrame(*Particle,Step);
+
+			EXPECT_NE(ParticleState, nullptr);
+			if(ParticleState)
+			{
+				EXPECT_EQ(ParticleState->F()[2],Step+1);
+			}
+		}
+
+		// Throw out the proxy
+		Solver->UnregisterObject(Particle.Get());
+
+		Module->DestroySolver(Solver);
+	}
+
+	GTEST_TEST(RewindTest,DataCapture3)
+	{
+		auto Sphere = TSharedPtr<FImplicitObject,ESPMode::ThreadSafe>(new TSphere<float,3>(TVector<float,3>(0),10));
+
+		FChaosSolversModule* Module = FChaosSolversModule::GetModule();
+		Module->ChangeThreadingMode(EChaosThreadingMode::SingleThread);
+
+		// Make a solver
+		FPhysicsSolver* Solver = Module->CreateSolver(nullptr,ESolverFlags::Standalone);
+		Solver->SetEnabled(true);
+
+		Solver->EnableRewindCapture(20);
 
-			const auto DynamicState = RewindData->GetStateAtFrame(*Dynamic,Step);
 
-			EXPECT_TRUE(DynamicState != nullptr);
-			if(DynamicState)
+		// Make particles
+		auto Particle = TPBDRigidParticle<float,3>::CreateParticle();
+
+		Particle->SetGeometry(Sphere);
+		Solver->RegisterObject(Particle.Get());
+
+		for(int Step = 0; Step < 11; ++Step)
+		{	
+			//sim-writable property that changes infrequently and not at beginning
+			if(Step == 3)
+			{
+				Particle->SetF(FVec3(0,0,Step));
+			}
+
+			if(Step == 5)
 			{
-				EXPECT_EQ(DynamicState->F()[2],Step+1);
+				Particle->SetF(FVec3(0,0,Step));
 			}
+
+			TickSolverHelper(Module,Solver);
 		}
 		
-		//no data at head because we always save the previous frame
-		EXPECT_TRUE(RewindData->GetStateAtFrame(*Particle,10) == nullptr);
+		const FRewindData* RewindData = Solver->GetRewindData();
+
+		//check state at every step except latest
+		for(int Step = 0; Step < 10; ++Step)
+		{
+			const auto ParticleState = RewindData->GetStateAtFrame(*Particle,Step);
+
+			//Should have data up to frame 5 and then nothing (includes 5 because dynamic is record on current frame)
+			if(Step <= 5)
+			{
+				EXPECT_NE(ParticleState, nullptr);
+			}
+			else
+			{
+				EXPECT_EQ(ParticleState,nullptr);
+			}
+			
+
+			if(ParticleState)
+			{
+				if(Step == 3)
+				{
+					EXPECT_EQ(ParticleState->F()[2],3);
+				}
+				else if(Step == 5)
+				{
+					EXPECT_EQ(ParticleState->F()[2],5);
+				}
+			}
+		}
 
 		// Throw out the proxy
 		Solver->UnregisterObject(Particle.Get());
@@ -133,4 +242,90 @@ namespace ChaosTest {
 		Module->DestroySolver(Solver);
 	}
 
+	GTEST_TEST(RewindTest,DataCapture4)
+	{
+		auto Sphere = TSharedPtr<FImplicitObject,ESPMode::ThreadSafe>(new TSphere<float,3>(TVector<float,3>(0),10));
+		auto Box = TSharedPtr<FImplicitObject,ESPMode::ThreadSafe>(new TBox<float,3>(FVec3(0),FVec3(1)));
+		auto Box2 = TSharedPtr<FImplicitObject,ESPMode::ThreadSafe>(new TBox<float,3>(FVec3(2),FVec3(3)));
+
+		FChaosSolversModule* Module = FChaosSolversModule::GetModule();
+		Module->ChangeThreadingMode(EChaosThreadingMode::SingleThread);
+
+		// Make a solver
+		FPhysicsSolver* Solver = Module->CreateSolver(nullptr,ESolverFlags::Standalone);
+		Solver->SetEnabled(true);
+
+		Solver->EnableRewindCapture(20);
+
+
+		// Make particles
+		auto Particle = TKinematicGeometryParticle<float,3>::CreateParticle();
+
+		Particle->SetGeometry(Sphere);
+		Solver->RegisterObject(Particle.Get());
+
+		for(int Step = 0; Step < 11; ++Step)
+		{
+			//property that changes once half way through
+			if(Step == 3)
+			{
+				Particle->SetGeometry(Box);
+			}
+
+			if(Step == 5)
+			{
+				Particle->SetGeometry(Box2);
+			}
+
+			if(Step == 7)
+			{
+				Particle->SetGeometry(Box);
+			}
+
+			TickSolverHelper(Module,Solver);
+		}
+
+		const FRewindData* RewindData = Solver->GetRewindData();
+
+		//check state at every step except latest
+		for(int Step = 0; Step < 10; ++Step)
+		{
+			const auto ParticleState = RewindData->GetStateAtFrame(*Particle,Step);
+
+			//Should have data up to frame 6 and then nothing
+			if(Step < 7)
+			{
+				EXPECT_NE(ParticleState,nullptr);
+			} else
+			{
+				EXPECT_EQ(ParticleState,nullptr);
+			}
+
+			if(ParticleState)
+			{
+				if(Step < 3)
+				{
+					//was sphere
+					EXPECT_EQ(ParticleState->Geometry().Get(),Sphere.Get());
+				} else if(Step < 5)
+				{
+					//then became box
+					EXPECT_EQ(ParticleState->Geometry().Get(),Box.Get());
+				} else
+				{
+					//second box
+					EXPECT_EQ(ParticleState->Geometry().Get(),Box2.Get());
+				}
+			}
+			else
+			{
+				EXPECT_EQ(Particle->Geometry().Get(),Box.Get());
+			}
+		}
+
+		// Throw out the proxy
+		Solver->UnregisterObject(Particle.Get());
+
+		Module->DestroySolver(Solver);
+	}
 }
diff --git a/Engine/Source/Runtime/Experimental/Chaos/Private/Chaos/PBDRigidsEvolutionGBF.cpp b/Engine/Source/Runtime/Experimental/Chaos/Private/Chaos/PBDRigidsEvolutionGBF.cpp
--- a/Engine/Source/Runtime/Experimental/Chaos/Private/Chaos/PBDRigidsEvolutionGBF.cpp
+++ b/Engine/Source/Runtime/Experimental/Chaos/Private/Chaos/PBDRigidsEvolutionGBF.cpp
@@ -418,6 +418,7 @@ void FPBDRigidsEvolutionGBF::AdvanceOneTimeStep(const FReal Dt, const FReal Step
 
 	Clustering.AdvanceClustering(Dt, GetCollisionConstraints());
 
+	CaptureRewindData(Particles.GetActiveParticlesView());
 	ParticleUpdatePosition(Particles.GetActiveParticlesView(), Dt);
 }
 
diff --git a/Engine/Source/Runtime/Experimental/Chaos/Public/Chaos/PBDRigidsEvolution.h b/Engine/Source/Runtime/Experimental/Chaos/Public/Chaos/PBDRigidsEvolution.h
--- a/Engine/Source/Runtime/Experimental/Chaos/Public/Chaos/PBDRigidsEvolution.h
+++ b/Engine/Source/Runtime/Experimental/Chaos/Public/Chaos/PBDRigidsEvolution.h
@@ -223,6 +223,7 @@ class FPBDRigidsEvolutionBase
 	typedef TFunction<void(const TArray<TGeometryParticleHandle<FReal, 3>*>&, const FReal)> FUpdateVelocityRule;
 	typedef TFunction<void(const TParticleView<TPBDRigidParticles<FReal, 3>>&, const FReal)> FUpdatePositionRule;
 	typedef TFunction<void(TPBDRigidParticles<FReal, 3>&, const FReal, const FReal, const int32)> FKinematicUpdateRule;
+	typedef TFunction<void(TParticleView<TPBDRigidParticles<FReal,3>>&)> FCaptureRewindRule;
 
 	friend void ChaosTest::TestPendingSpatialDataHandlePointerConflict();
 
@@ -283,6 +284,7 @@ class FPBDRigidsEvolutionBase
 	CHAOS_API void AddImpulseFunction(FForceRule ImpulseFunction) { ImpulseRules.Add(ImpulseFunction); }
 	CHAOS_API void SetParticleUpdateVelocityFunction(FUpdateVelocityRule ParticleUpdate) { ParticleUpdateVelocity = ParticleUpdate; }
 	CHAOS_API void SetParticleUpdatePositionFunction(FUpdatePositionRule ParticleUpdate) { ParticleUpdatePosition = ParticleUpdate; }
+	CHAOS_API void SetCaptureRewindDataFunction(FCaptureRewindRule Rule){ CaptureRewindData = Rule; }
 
 	CHAOS_API TGeometryParticleHandles<FReal, 3>& GetParticleHandles() { return Particles.GetParticleHandles(); }
 	CHAOS_API const TGeometryParticleHandles<FReal, 3>& GetParticleHandles() const { return Particles.GetParticleHandles(); }
@@ -692,6 +694,7 @@ protected:
 	FUpdateVelocityRule ParticleUpdateVelocity;
 	FUpdatePositionRule ParticleUpdatePosition;
 	FKinematicUpdateRule KinematicUpdate;
+	FCaptureRewindRule CaptureRewindData;
 	TArray<FPBDConstraintGraphRule*> ConstraintRules;
 	TArray<FPBDConstraintGraphRule*> PrioritizedConstraintRules;
 	FPBDConstraintGraph ConstraintGraph;
diff --git a/Engine/Source/Runtime/Experimental/Chaos/Public/Chaos/ParticleDirtyFlags.h b/Engine/Source/Runtime/Experimental/Chaos/Public/Chaos/ParticleDirtyFlags.h
--- a/Engine/Source/Runtime/Experimental/Chaos/Public/Chaos/ParticleDirtyFlags.h
+++ b/Engine/Source/Runtime/Experimental/Chaos/Public/Chaos/ParticleDirtyFlags.h
@@ -57,17 +57,13 @@ struct FParticleDynamics
 	FVec3 Torque;
 	FVec3 LinearImpulse;
 	FVec3 AngularImpulse;
-	FReal LinearEtherDrag;
-	FReal AngularEtherDrag;
 
 	void Serialize(FChaosArchive& Ar)
 	{
 		Ar << F;
 		Ar << Torque;
 		Ar << LinearImpulse;
-		Ar << AngularImpulse;
-		Ar << LinearEtherDrag;
-		Ar << AngularEtherDrag;		
+		Ar << AngularImpulse;	
 	}
 };
 
@@ -143,6 +139,9 @@ struct FParticleNonFrequentData
 	TSharedPtr<FImplicitObject,ESPMode::ThreadSafe> Geometry;
 	void* UserData;
 	FUniqueIdx UniqueIdx;
+	FReal LinearEtherDrag;
+	FReal AngularEtherDrag;
+
 #if CHAOS_CHECKED
 	FName DebugName;
 #endif
@@ -400,6 +399,11 @@ struct TPropertyTypeTrait<Type>\
 			Elements.SetNum(Num);
 		}
 
+		int32 Num() const
+		{
+			return Elements.Num();
+		}
+
 	private:
 
 		TArray<T> Elements;
@@ -716,6 +720,12 @@ public:
 #undef PARTICLE_PROPERTY
 	}
 
+	int32 GetNumParticles() const
+	{
+		//assume this property exists, if it gets renamed just pick any property
+		return XRPool.Num();
+	}
+
 	void SetNumShapes(int32 NumShapes)
 	{
 #define SHAPE_PROPERTY(PropName, Type) PropName##ShapePool.SetNum(NumShapes);
@@ -845,6 +855,11 @@ public:
 		Flags = InFlags;
 	}
 
+	FParticleDirtyFlags GetFlags() const
+	{
+		return Flags;
+	}
+
 	template <typename T, EParticleProperty PropName>
 	void SyncRemote(FDirtyPropertiesManager& Manager, int32 Idx, const T& Val) const
 	{
diff --git a/Engine/Source/Runtime/Experimental/Chaos/Public/Chaos/ParticleHandle.h b/Engine/Source/Runtime/Experimental/Chaos/Public/Chaos/ParticleHandle.h
--- a/Engine/Source/Runtime/Experimental/Chaos/Public/Chaos/ParticleHandle.h
+++ b/Engine/Source/Runtime/Experimental/Chaos/Public/Chaos/ParticleHandle.h
@@ -1622,9 +1622,9 @@ protected:
 private:
 
 	TParticleProperty<FParticlePositionRotation, EParticleProperty::XR> MXR;
-	TParticleProperty<FParticleNonFrequentData, EParticleProperty::NonFrequentData> MNonFrequentData;
 
 protected:
+	TParticleProperty<FParticleNonFrequentData,EParticleProperty::NonFrequentData> MNonFrequentData;
 	TParticleProperty<FParticleMisc,EParticleProperty::Misc> MMiscData;
 private:
 
@@ -2039,16 +2039,16 @@ public:
 		MMassProps.Modify(true,MDirtyFlags,Proxy,[InInvM](auto& Data){ Data.InvM = InInvM;});
 	}
 
-	T LinearEtherDrag() const { return MDynamics.Read().LinearEtherDrag; }
+	T LinearEtherDrag() const { return this->MNonFrequentData.Read().LinearEtherDrag; }
 	void SetLinearEtherDrag(const T& InLinearEtherDrag)
 	{
-		MDynamics.Modify(true,MDirtyFlags,Proxy,[&InLinearEtherDrag](auto& Data){ Data.LinearEtherDrag = InLinearEtherDrag;});
+		this->MNonFrequentData.Modify(true,MDirtyFlags,Proxy,[&InLinearEtherDrag](auto& Data){ Data.LinearEtherDrag = InLinearEtherDrag;});
 	}
 
-	T AngularEtherDrag() const { return MDynamics.Read().AngularEtherDrag; }
+	T AngularEtherDrag() const { return this->MNonFrequentData.Read().AngularEtherDrag; }
 	void SetAngularEtherDrag(const T& InAngularEtherDrag)
 	{
-		MDynamics.Modify(true,MDirtyFlags,Proxy,[&InAngularEtherDrag](auto& Data){ Data.AngularEtherDrag = InAngularEtherDrag;});
+		this->MNonFrequentData.Modify(true,MDirtyFlags,Proxy,[&InAngularEtherDrag](auto& Data){ Data.AngularEtherDrag = InAngularEtherDrag;});
 	}
 
 	int32 Island() const { return MIsland; }
diff --git a/Engine/Source/Runtime/Experimental/ChaosSolvers/Private/PBDRigidsSolver.cpp b/Engine/Source/Runtime/Experimental/ChaosSolvers/Private/PBDRigidsSolver.cpp
--- a/Engine/Source/Runtime/Experimental/ChaosSolvers/Private/PBDRigidsSolver.cpp
+++ b/Engine/Source/Runtime/Experimental/ChaosSolvers/Private/PBDRigidsSolver.cpp
@@ -97,6 +97,12 @@ namespace Chaos
 				//MSolver->StartFrameCallback(MDeltaTime, MSolver->GetSolverTime());
 			}
 
+
+			if(FRewindData* RewindData = MSolver->GetRewindData())
+			{
+				RewindData->AdvanceFrame();
+			}
+
 			{
 				SCOPE_CYCLE_COUNTER(STAT_EvolutionAndKinematicUpdate);
 
@@ -519,6 +525,11 @@ namespace Chaos
 
 		DirtyPropertiesManager = MakeUnique<FDoubleBuffer<FDirtyPropertiesManager>>();
 
+		MEvolution->SetCaptureRewindDataFunction([this](const TParticleView<TPBDRigidParticles<FReal,3>>& ActiveParticles)
+		{
+			FinalizeRewindData(ActiveParticles);
+		});
+
 		FEventDefaults::RegisterSystemEvents(*GetEventManager());
 	}
 
@@ -827,11 +838,6 @@ namespace Chaos
 				}
 			});
 
-			if(RewindData)
-			{
-				RewindData->AdvanceFrame();
-			}
-
 			DirtyProxiesData->Reset();
 		});
 	}
@@ -1017,4 +1023,27 @@ namespace Chaos
 		MRewindData = MakeUnique<FRewindData>(NumFrames);
 	}
 
+	void FPBDRigidsSolver::FinalizeRewindData(const TParticleView<TPBDRigidParticles<FReal,3>>& ActiveParticles)
+	{
+		using namespace Chaos;
+		//Simulated objects must have their properties captured for rewind
+		if(MRewindData && ActiveParticles.Num())
+		{
+			QUICK_SCOPE_CYCLE_COUNTER(RecordRewindData);
+
+			int32 DataIdx = MRewindData->PrepareFrameForPTDirty(ActiveParticles.Num());
+			
+			for(TPBDRigidParticleHandleImp<float,3,false>& ActiveObject : ActiveParticles)
+			{
+				const IPhysicsProxyBase* Proxy = GetProxy(ActiveObject.Handle());
+				{
+					if(ActiveObject.GetParticleType() == EParticleType::Rigid)
+					{
+						MRewindData->PushPTDirtyData(*static_cast<const FRigidParticlePhysicsProxy*>(Proxy)->GetHandle(), DataIdx++);
+					}
+				}
+			}
+		}
+	}
+
 }; // namespace Chaos
diff --git a/Engine/Source/Runtime/Experimental/ChaosSolvers/Private/PhysicsProxy/SingleParticlePhysicsProxy.cpp b/Engine/Source/Runtime/Experimental/ChaosSolvers/Private/PhysicsProxy/SingleParticlePhysicsProxy.cpp
--- a/Engine/Source/Runtime/Experimental/ChaosSolvers/Private/PhysicsProxy/SingleParticlePhysicsProxy.cpp
+++ b/Engine/Source/Runtime/Experimental/ChaosSolvers/Private/PhysicsProxy/SingleParticlePhysicsProxy.cpp
@@ -72,6 +72,13 @@ void PushToPhysicsStateImp(const Chaos::FDirtyPropertiesManager& Manager, Chaos:
 			Handle->SetSharedGeometry(NewNonFrequentData->Geometry);
 			Handle->SetUniqueIdx(NewNonFrequentData->UniqueIdx);
 			Handle->SetUserData(NewNonFrequentData->UserData);
+
+			if(bHasDynamicData)
+			{
+				RigidHandle->SetLinearEtherDrag(NewNonFrequentData->LinearEtherDrag);
+				RigidHandle->SetAngularEtherDrag(NewNonFrequentData->AngularEtherDrag);
+			}
+			
 #if CHAOS_CHECKED
 			Handle->SetDebugName(NewNonFrequentData->DebugName);
 #endif
@@ -142,9 +149,6 @@ void PushToPhysicsStateImp(const Chaos::FDirtyPropertiesManager& Manager, Chaos:
 				RigidHandle->SetTorque(NewData->Torque);
 				RigidHandle->SetLinearImpulse(NewData->LinearImpulse);
 				RigidHandle->SetAngularImpulse(NewData->AngularImpulse);
-
-				RigidHandle->SetLinearEtherDrag(NewData->LinearEtherDrag);
-				RigidHandle->SetAngularEtherDrag(NewData->AngularEtherDrag);
 			}
 		}
 
diff --git a/Engine/Source/Runtime/Experimental/ChaosSolvers/Public/PBDRigidsSolver.h b/Engine/Source/Runtime/Experimental/ChaosSolvers/Public/PBDRigidsSolver.h
--- a/Engine/Source/Runtime/Experimental/ChaosSolvers/Public/PBDRigidsSolver.h
+++ b/Engine/Source/Runtime/Experimental/ChaosSolvers/Public/PBDRigidsSolver.h
@@ -329,6 +329,8 @@ namespace Chaos
 		/** Copy the simulation material list to the query material list, to be done when the SQ commits an update */
 		void SyncQueryMaterials();
 
+		void FinalizeRewindData(const TParticleView<TPBDRigidParticles<FReal,3>>& ActiveParticles);
+
 	private:
 
 		template<typename ParticleType>
diff --git a/Engine/Source/Runtime/Experimental/ChaosSolvers/Public/RewindData.h b/Engine/Source/Runtime/Experimental/ChaosSolvers/Public/RewindData.h
--- a/Engine/Source/Runtime/Experimental/ChaosSolvers/Public/RewindData.h
+++ b/Engine/Source/Runtime/Experimental/ChaosSolvers/Public/RewindData.h
@@ -62,43 +62,6 @@ private:
 	int32 Idx;
 };
 
-//Set of properties that can be written to from the simulation
-class FSimWritableData
-{
-public:
-
-	FSimWritableData()
-	{
-	}
-
-	const FParticlePositionRotation& GetXR() const { return XR; }
-	const FParticleVelocities& GetVelocities() const { return Velocities; }
-	const FParticleDynamics& GetDynamics() const { return Dynamics; }
-
-	void SyncData(const FDirtyPropertiesManager& Manager, const int32 DataIdx, const FParticleDirtyData& Dirty)
-	{
-		if(const auto Data = Dirty.FindXR(Manager, DataIdx))
-		{
-			XR = *Data;
-		}
-
-		if(const auto Data = Dirty.FindVelocities(Manager, DataIdx))
-		{
-			Velocities = *Data;
-		}
-
-		if(const auto Data = Dirty.FindDynamics(Manager, DataIdx))
-		{
-			Dynamics = *Data;
-		}
-	}
-
-private:
-	FParticlePositionRotation XR;
-	FParticleVelocities Velocities;
-	FParticleDynamics Dynamics;
-};
-
 inline bool SimWritablePropsMayChange(const TGeometryParticleHandle<FReal,3>& Handle)
 {
 	if(const auto Rigid = Handle.CastToRigidParticle())
@@ -133,23 +96,55 @@ public:
 		return Dynamics.IsSet() ? Dynamics.Read().F : Particle.CastToRigidParticle()->F();
 	}
 
-	void SyncRemoteData(FDirtyPropertiesManager& Manager,int32 Idx,const FDirtyProxy& Dirty, const FSimWritableData* SimWritableData)
+	void SyncSimWritablePropsFromSim(FDirtyPropertiesManager& Manager,int32 Idx, const TPBDRigidParticleHandle<FReal,3>& Rigid)
+	{
+		FParticleDirtyFlags Flags;
+		Flags.MarkDirty(EParticleFlags::XR);
+		Flags.MarkDirty(EParticleFlags::Velocities);
+		FParticleDirtyData Dirty;
+		Dirty.SetFlags(Flags);
+
+		ParticlePositionRotation.SyncRemoteData(Manager,Idx,Dirty,[&Rigid](auto& Data)
+		{
+			Data.X = Rigid.X();
+			Data.R = Rigid.R();
+		});
+
+		//todo: velocities
+	}
+
+	void SyncDirtyDynamics(FDirtyPropertiesManager& DestManager,int32 DataIdx, const FParticleDirtyData& Dirty, const FDirtyPropertiesManager& SrcManager)
+	{
+		FParticleDirtyData DirtyFlags;
+		DirtyFlags.SetFlags(Dirty.GetFlags());
+
+		Dynamics.SyncRemoteData(DestManager,DataIdx,DirtyFlags,[&Dirty, &SrcManager, DataIdx](auto& Data)
+		{
+			Data = Dirty.GetDynamics(SrcManager, DataIdx);
+		});
+	}
+	
+	void SyncPrevFrame(FDirtyPropertiesManager& Manager,int32 Idx,const FDirtyProxy& Dirty)
 	{
+		//syncs the data before it was made dirty
+		//for sim-writable props this is only possible if those props are immutable from the sim side (sleeping, not simulated, etc...)
+
 		const auto Proxy = static_cast<const FGeometryParticlePhysicsProxy*>(Dirty.Proxy);
 		const auto Handle = Proxy->GetHandle();
-		ParticlePositionRotation.SyncRemoteData(Manager,Idx,Dirty.ParticleData, [Handle, SimWritableData](FParticlePositionRotation& Data)
+
+		const bool bSyncSimWritable = !SimWritablePropsMayChange(*Handle);
+
+		//note: there is a potential bug here if in one frame we put an object to sleep and change its position
+		//not sure if that's a valid operation, we probably need to catch sleep/awake and handle it in a special way
+		if(bSyncSimWritable)
 		{
-			if(SimWritableData)
-			{
-				Data = SimWritableData->GetXR();
-			}
-			else
+			ParticlePositionRotation.SyncRemoteData(Manager,Idx,Dirty.ParticleData,[Handle](FParticlePositionRotation& Data)
 			{
 				Data.X = Handle->X();
 				Data.R = Handle->R();
-			}
-		});
-
+			});
+		}
+		
 		NonFrequentData.SyncRemoteData(Manager,Idx,Dirty.ParticleData, [Handle](FParticleNonFrequentData& Data)
 		{
 			Data.Geometry = Handle->SharedGeometryLowLevel();
@@ -162,26 +157,6 @@ public:
 			Data.DebugName = Handle->DebugName();
 #endif
 		});
-
-		if(auto PBDRigid = Handle->CastToRigidParticle())
-		{
-			Dynamics.SyncRemoteData(Manager,Idx,Dirty.ParticleData,[PBDRigid, SimWritableData](FParticleDynamics& Data)
-			{
-				if(SimWritableData)
-				{
-					Data = SimWritableData->GetDynamics();
-				}
-				else
-				{
-					Data.F = PBDRigid->F();
-					Data.Torque = PBDRigid->Torque();
-					Data.LinearImpulse = PBDRigid->LinearImpulse();
-					Data.AngularImpulse = PBDRigid->AngularImpulse();
-					Data.LinearEtherDrag = PBDRigid->LinearEtherDrag();
-					Data.AngularEtherDrag = PBDRigid->AngularEtherDrag();
-				}
-			});
-		}
 	}
 
 	bool CoalesceState(const FGeometryParticleState& LatestState)
@@ -199,12 +174,8 @@ public:
 			bCoalesced = true;
 		}
 
-		if(!Dynamics.IsSet() && LatestState.Dynamics.IsSet())
-		{
-			Dynamics = LatestState.Dynamics;
-			bCoalesced = true;
-		}
-
+		//dynamics do not coalesce since they are always written when dirty
+		
 		return bCoalesced;
 	}
 
@@ -236,15 +207,28 @@ public:
 		if(const FParticleRewindInfo* Info = ParticleToRewindInfo.Find(Particle.UniqueIdx()))
 		{
 			//is it worth doing binary search?
-			for(const FFrameInfo& FrameInfo : Info->Frames)
+			const int32 NumFrames = Info->Frames.Num();
+			if(NumFrames > 0)
 			{
-				if(FrameInfo.Frame == Frame)
+				//If frame is before first capture, use first capture
+				if(Frame <= Info->Frames[0].Frame)
+				{
+					return &Info->Frames[0].State;
+				}
+
+				//If frame is between two captures, use later capture, because we always store the last data before a change
+				//We can never use an earlier capture, because the fact that we captured at all implies _something_ is different from proceeding frames
+				for(int32 FrameIdx = 1; FrameIdx < NumFrames; ++FrameIdx)
 				{
-					return &FrameInfo.State;
+					if(Frame <= Info->Frames[FrameIdx].Frame)
+					{
+						return &Info->Frames[FrameIdx].State;
+					}
 				}
 			}
 		}
 
+		//If no data, or past capture, just use head
 		return nullptr;
 	}
 
@@ -255,25 +239,48 @@ public:
 
 	void PrepareFrame(int32 NumDirtyParticles)
 	{
-		Managers.Emplace(MakeUnique<FDirtyPropertiesManager>());
-		Managers.Last()->SetNumParticles(NumDirtyParticles);
+		Managers.Emplace(FFrameManagerInfo{MakeUnique<FDirtyPropertiesManager>(), CurFrame});
+		Managers.Last().Manager->SetNumParticles(NumDirtyParticles);
+	}
+
+	int32 PrepareFrameForPTDirty(int32 NumActiveParticles)
+	{
+		//If manager already exists for previous frame, use it
+		if(Managers.Num())
+		{
+			if(Managers.Last().FrameCreatedFor == (CurFrame-1))
+			{
+				FDirtyPropertiesManager& Manager = *Managers.Last().Manager;
+				const int32 NumDirtyAlready = Manager.GetNumParticles();
+				Manager.SetNumParticles(NumDirtyAlready + NumActiveParticles);
+				return NumDirtyAlready;
+			}
+		}
+
+		//No manager for previous frame so create a new one
+		PrepareFrame(NumActiveParticles);
+		return 0;
 	}
 
-	void PushGTDirtyData(const FDirtyPropertiesManager& Manager, int32 DataIdx, const FDirtyProxy& Dirty)
+	void PushGTDirtyData(const FDirtyPropertiesManager& SrcManager, int32 DataIdx, const FDirtyProxy& Dirty)
 	{
-		auto ProcessProxy = [this,Manager, DataIdx, Dirty](const auto Proxy)
+		//This records changes enqueued by GT.
+		//Most new particles do not change, so to avoid useless writes we wait until the next frame's dirty flag
+		//This is possible because most properties are const on the physics thread
+		//For sim-writable properties (forces, position, velocities, etc...) we must immediately write the data because there is no way to know what the previous data was next frame
+		//Some sim-writable properties can change without the GT knowing about it, see PushPTDirtyData
+		
+		auto ProcessProxy = [this,&SrcManager, DataIdx, Dirty](const auto Proxy)
 		{
 			const auto PTParticle = Proxy->GetHandle();
 			FParticleRewindInfo& Info = ParticleToRewindInfo.FindOrAdd(PTParticle->UniqueIdx());
+			FDirtyPropertiesManager& DestManager = *Managers.Last().Manager;
 
-			//Any previous frames that are pointing at head of a dirty property must be written out (since head is changing)
-
-			//Since we are updating previous frames, the proxy must have been initialized already
-			if(Proxy->IsInitialized())
+			//Most properties are always a frame behind
+			if(Proxy->IsInitialized())	//Frame delay so proxy must be initialized
 			{
 				FGeometryParticleState& LatestState = Info.AddFrame(*Proxy->GetParticle(),CurFrame-1);
-
-				LatestState.SyncRemoteData(*Managers.Last(),DataIdx,Dirty, Info.SimWritableData.Get());
+				LatestState.SyncPrevFrame(DestManager,DataIdx,Dirty);
 
 				//for frames further back a simply copy is enough
 				for(int32 FrameIdx = Info.Frames.Num() - 2; FrameIdx >= 0; --FrameIdx)
@@ -287,18 +294,11 @@ public:
 				}
 			}
 
-			//if particle is sim-writable make sure to record appropriate dirty properties
-			if(SimWritablePropsMayChange(*PTParticle))
-			{
-				if(!Info.SimWritableData)
-				{
-					Info.SimWritableData = MakeUnique<FSimWritableData>();
-				}
-				Info.SimWritableData->SyncData(Manager,DataIdx,Dirty.ParticleData);
-			}
-			else
+			//If dynamics are dirty we must record them immediately because the sim resets them to 0
+			if(Dirty.ParticleData.IsDirty(EParticleFlags::Dynamics))
 			{
-				Info.SimWritableData = nullptr;
+				FGeometryParticleState& LatestState = Info.AddFrame(*Proxy->GetParticle(),CurFrame);
+				LatestState.SyncDirtyDynamics(DestManager,DataIdx,Dirty.ParticleData,SrcManager);
 			}
 		};
 
@@ -327,6 +327,23 @@ public:
 		}
 	}
 
+	void PushPTDirtyData(const TPBDRigidParticleHandle<FReal,3>& Rigid, int32 DataIdx)
+	{
+		//todo: is this check needed? why do we pass sleeping rigids into this function?
+		if(SimWritablePropsMayChange(Rigid))
+		{
+			FParticleRewindInfo& Info = ParticleToRewindInfo.FindOrAdd(Rigid.UniqueIdx());
+			FDirtyPropertiesManager& Manager = *Managers.Last().Manager;
+
+			//sim-writable properties changed at head, so we must write down what they were
+			FGeometryParticleState& LatestState = Info.AddFrame(*Rigid.GTGeometryParticle(),CurFrame-1);
+			LatestState.SyncSimWritablePropsFromSim(Manager,DataIdx,Rigid);
+
+			//update any previous frames that were pointing at head
+			CoalesceBack(Info.Frames);
+		}
+	}
+
 private:
 
 	struct FFrameInfo
@@ -338,17 +355,45 @@ private:
 	struct FParticleRewindInfo
 	{
 		TArray<FFrameInfo> Frames;
-		TUniquePtr<FSimWritableData> SimWritableData;
 
 		FGeometryParticleState& AddFrame(TGeometryParticle<FReal,3>& GTParticleUnsafe, int32 FrameIdx)
 		{
+			if(Frames.Num() && Frames.Last().Frame == FrameIdx)
+			{
+				return Frames.Last().State;
+			}
+
 			Frames.Add(FFrameInfo{FGeometryParticleState(GTParticleUnsafe),FrameIdx});
 			return Frames.Last().State;
 		}
 	};
 
+	void CoalesceBack(TArray<FFrameInfo>& Frames)
+	{
+		const FGeometryParticleState& LatestState = Frames.Last().State;
+		for(int32 FrameIdx = Frames.Num() - 2; FrameIdx >= 0; --FrameIdx)
+		{
+			FFrameInfo& Frame = Frames[FrameIdx];
+			if(Frame.State.CoalesceState(LatestState) == false)
+			{
+				//nothing to coalesce so no need to check earlier frames
+				break;
+			}
+		}
+	}
+
+	struct FFrameManagerInfo
+	{
+		TUniquePtr<FDirtyPropertiesManager> Manager;
+
+		//Note that this is not exactly the same as which frame this manager represents. 
+		//A manager can have data for two frames at once, the important part is just knowing which frame it was created on so we know whether the physics data can rely on it
+		//Consider the case where nothing is dirty from GT and then an object moves from the simulation, in that case it needs a manager to record the data into
+		int32 FrameCreatedFor;
+	};
+
 	TArrayAsMap<FUniqueIdx,FParticleRewindInfo> ParticleToRewindInfo;
-	TArray<TUniquePtr<FDirtyPropertiesManager>> Managers;
+	TArray<FFrameManagerInfo> Managers;
 	int32 CurFrame;
 };
 }
