Chaos : GeoemtryCollectionTest

- Analytic only collisions test.

diff --git a/Engine/Source/Runtime/Experimental/ChaosSolvers/Private/PhysicsProxy/GeometryCollectionPhysicsProxy.cpp b/Engine/Source/Runtime/Experimental/ChaosSolvers/Private/PhysicsProxy/GeometryCollectionPhysicsProxy.cpp
--- a/Engine/Source/Runtime/Experimental/ChaosSolvers/Private/PhysicsProxy/GeometryCollectionPhysicsProxy.cpp
+++ b/Engine/Source/Runtime/Experimental/ChaosSolvers/Private/PhysicsProxy/GeometryCollectionPhysicsProxy.cpp
@@ -295,9 +295,10 @@ void PopulateSimulatedParticle(
 		Handle->SetWorldSpaceInflatedBounds(TransformedBBox);
 	}
 
-	Handle->CollisionParticlesInitIfNeeded();
-	if (Simplicial)
+	if (Simplicial && SharedParams.SizeSpecificData[0].CollisionType==ECollisionTypeEnum::Chaos_Surface_Volumetric)
 	{
+		Handle->CollisionParticlesInitIfNeeded();
+
 		TUniquePtr<Chaos::TBVHParticles<float, 3>>& CollisionParticles = Handle->CollisionParticles();
 		if (Simplicial->Size())
 		{
@@ -617,7 +618,8 @@ void FGeometryCollectionPhysicsProxy::InitializeDynamicCollection(FGeometryDynam
 
 	// process simplicials
 	{
-		if (RestCollection.HasAttribute(DynamicCollection.SimplicialsAttribute, FTransformCollection::TransformGroup))
+		if (RestCollection.HasAttribute(DynamicCollection.SimplicialsAttribute, FTransformCollection::TransformGroup)
+			&& Params.Shared.SizeSpecificData[0].CollisionType == ECollisionTypeEnum::Chaos_Surface_Volumetric)
 		{
 			const auto& RestSimplicials = RestCollection.GetAttribute<TUniquePtr<FSimplicial>>(
 				DynamicCollection.SimplicialsAttribute, FTransformCollection::TransformGroup);
@@ -773,20 +775,16 @@ void FGeometryCollectionPhysicsProxy::InitializeBodiesPT(
 				}
 
 				TUniquePtr<Chaos::TBVHParticles<float, 3>>& CollisionParticles = Handle->CollisionParticles();
-				// @todo break everything CollisionParticles.Reset(Simplicials[TransformGroupIndex].Release()); // Steals!
-				//Particles.CollisionParticles(RigidBodyIndex).Reset(Simplicials[TransformGroupIndex].Release());
+				CollisionParticles.Reset(Simplicials[TransformGroupIndex]?Simplicials[TransformGroupIndex]->NewCopy():nullptr); // @chaos(optimize) : maybe just move this memory instead. 
 				if (CollisionParticles)
-					//if (Particles.CollisionParticles(RigidBodyIndex))
 				{
 					int32 NumCollisionParticles = CollisionParticles->Size();
-					//int32 NumCollisionParticles = Particles.CollisionParticles(RigidBodyIndex)->Size();
 					int32 CollisionParticlesSize = FMath::Max(0, FMath::Min(int(NumCollisionParticles * CollisionParticlesPerObjectFraction), NumCollisionParticles));
 					CollisionParticles->Resize(CollisionParticlesSize); // Truncates!
-					//Particles.CollisionParticles(RigidBodyIndex)->Resize(CollisionParticlesSize);
 				}
 
+
 				RigidsSolver->GetEvolution()->SetPhysicsMaterial(Handle, Parameters.PhysicalMaterial);
-				//GetSolver()->SetPhysicsMaterial(RigidBodyIndex, Parameters.PhysicalMaterial);
 			}
 		});
 
