Added rewind to frame functionality

diff --git a/Engine/Source/Programs/HeadlessChaos/Private/HeadlessChaosTestRewind.cpp b/Engine/Source/Programs/HeadlessChaos/Private/HeadlessChaosTestRewind.cpp
--- a/Engine/Source/Programs/HeadlessChaos/Private/HeadlessChaosTestRewind.cpp
+++ b/Engine/Source/Programs/HeadlessChaos/Private/HeadlessChaosTestRewind.cpp
@@ -385,4 +385,55 @@ namespace ChaosTest {
 
 		Module->DestroySolver(Solver);
 	}
+
+	GTEST_TEST(RewindTest,ApplyRewind)
+	{
+		auto Sphere = TSharedPtr<FImplicitObject,ESPMode::ThreadSafe>(new TSphere<float,3>(TVector<float,3>(0),10));
+
+		FChaosSolversModule* Module = FChaosSolversModule::GetModule();
+		Module->ChangeThreadingMode(EChaosThreadingMode::SingleThread);
+
+		// Make a solver
+		FPhysicsSolver* Solver = Module->CreateSolver(nullptr,ESolverFlags::Standalone);
+		Solver->SetEnabled(true);
+
+		Solver->EnableRewindCapture(20);
+
+
+		// Make particles
+		auto Particle = TPBDRigidParticle<float,3>::CreateParticle();
+
+		Particle->SetGeometry(Sphere);
+		Solver->RegisterObject(Particle.Get());
+		Particle->SetGravityEnabled(true);
+		Particle->SetX(FVec3(0,0,100));
+
+		TArray<FVec3> X;
+		TArray<FVec3> V;
+
+		for(int Step = 0; Step < 10; ++Step)
+		{
+			//teleport from GT
+			if(Step == 5)
+			{
+				Particle->SetX(FVec3(0,0,10));
+				Particle->SetV(FVec3(0,0,1));
+			}
+
+			X.Add(Particle->X());
+			V.Add(Particle->V());
+			TickSolverHelper(Module,Solver);
+		}
+
+		FRewindData* RewindData = Solver->GetRewindData();
+		RewindData->RewindToFrame(3);
+		
+		EXPECT_EQ(Particle->X()[2],X[3][2]);
+		EXPECT_EQ(Particle->V()[2],V[3][2]);
+
+		// Throw out the proxy
+		Solver->UnregisterObject(Particle.Get());
+
+		Module->DestroySolver(Solver);
+	}
 }
diff --git a/Engine/Source/Runtime/Experimental/ChaosSolvers/Public/RewindData.h b/Engine/Source/Runtime/Experimental/ChaosSolvers/Public/RewindData.h
--- a/Engine/Source/Runtime/Experimental/ChaosSolvers/Public/RewindData.h
+++ b/Engine/Source/Runtime/Experimental/ChaosSolvers/Public/RewindData.h
@@ -31,7 +31,7 @@ public:
 	}
 
 	template <typename LambdaWrite>
-	void SyncToParticle(LambdaWrite& WriteFunc) const
+	void SyncToParticle(const LambdaWrite& WriteFunc) const
 	{
 		if(Manager)
 		{
@@ -76,7 +76,7 @@ class FGeometryParticleState
 {
 public:
 
-	FGeometryParticleState(const TGeometryParticle<FReal,3>& InParticle)
+	FGeometryParticleState(TGeometryParticle<FReal,3>& InParticle)
 		: Particle(InParticle)
 	{
 	}
@@ -142,6 +142,54 @@ public:
 			Data = Dirty.GetDynamics(SrcManager, DataIdx);
 		});
 	}
+
+	void SyncToParticle() const
+	{
+		//todo: set properties directly instead of assigning sub-properties
+
+		ParticlePositionRotation.SyncToParticle([this](const auto& Data)
+		{
+			Particle.SetX(Data.X);
+			Particle.SetR(Data.R);
+		});
+
+		if(auto Kinematic = Particle.CastToKinematicParticle())
+		{
+			Velocities.SyncToParticle([Kinematic](const auto& Data)
+			{
+				Kinematic->SetV(Data.V);
+				Kinematic->SetW(Data.W);
+			});
+		}
+
+
+		NonFrequentData.SyncToParticle([this](const auto& Data)
+		{
+			Particle.SetGeometry(Data.Geometry);
+			Particle.SetUserData(Data.UserData);
+			ensure(Data.UniqueIdx == Particle.UniqueIdx());	//this should never change
+
+#if CHAOS_CHECKED
+			Particle.SetDebugName(Data.DebugName);
+#endif
+			if(auto Rigid = Particle.CastToRigidParticle())
+			{
+				Rigid->SetLinearEtherDrag(Data.LinearEtherDrag);
+				Rigid->SetAngularEtherDrag(Data.AngularEtherDrag);
+			}
+		});
+
+		if(auto Rigid = Particle.CastToRigidParticle())
+		{
+			Dynamics.SyncToParticle([Rigid](const auto& Data)
+			{
+				Rigid->SetF(Data.F);
+				Rigid->SetTorque(Data.Torque);
+				Rigid->SetLinearImpulse(Data.LinearImpulse);
+				Rigid->SetAngularImpulse(Data.AngularImpulse);
+			});
+		}
+	}
 	
 	void SyncPrevFrame(FDirtyPropertiesManager& Manager,int32 Idx,const FDirtyProxy& Dirty)
 	{
@@ -184,6 +232,12 @@ public:
 #if CHAOS_CHECKED
 			Data.DebugName = Handle->DebugName();
 #endif
+
+			if(auto Rigid = Handle->CastToRigidParticle())
+			{
+				Data.LinearEtherDrag = Rigid->LinearEtherDrag();
+				Data.AngularEtherDrag = Rigid->AngularEtherDrag();
+			}
 		});
 	}
 
@@ -214,7 +268,7 @@ public:
 	}
 
 protected:
-	const TGeometryParticle<FReal,3>& Particle;
+	TGeometryParticle<FReal,3>& Particle;
 private:
 	TParticleStateProperty<FParticlePositionRotation,EParticleProperty::XR> ParticlePositionRotation;
 	TParticleStateProperty<FParticleNonFrequentData, EParticleProperty::NonFrequentData> NonFrequentData;
@@ -233,10 +287,26 @@ class FRewindData
 {
 public:
 	FRewindData(int32 NumFrames)
-		: CurFrame(0)
+	: CurFrame(0)
 	{
 	}
 
+	void RewindToFrame(int32 Frame)
+	{
+		ensure(IsInGameThread());
+		//todo: parallel for
+		for(TGeometryParticle<FReal,3>* Particle : AllDirtyParticles)
+		{
+			if(const FGeometryParticleState* State = GetStateAtFrame(*Particle, Frame))
+			{
+				State->SyncToParticle();
+			}
+		}
+		
+		Reset();
+		CurFrame = Frame;
+	}
+
 	const FGeometryParticleState* GetStateAtFrame(const TGeometryParticle<FReal,3>& Particle,int32 Frame) const
 	{
 		if(const FParticleRewindInfo* Info = ParticleToRewindInfo.Find(Particle.UniqueIdx()))
@@ -308,7 +378,7 @@ public:
 		auto ProcessProxy = [this,&SrcManager, DataIdx, Dirty](const auto Proxy)
 		{
 			const auto PTParticle = Proxy->GetHandle();
-			FParticleRewindInfo& Info = ParticleToRewindInfo.FindOrAdd(PTParticle->UniqueIdx());
+			FParticleRewindInfo& Info = FindOrAddParticle(Proxy->GetParticle(),PTParticle->UniqueIdx());
 			FDirtyPropertiesManager& DestManager = *Managers.Last().Manager;
 
 			//Most properties are always a frame behind
@@ -367,7 +437,7 @@ public:
 		//todo: is this check needed? why do we pass sleeping rigids into this function?
 		if(SimWritablePropsMayChange(Rigid))
 		{
-			FParticleRewindInfo& Info = ParticleToRewindInfo.FindOrAdd(Rigid.UniqueIdx());
+			FParticleRewindInfo& Info = FindOrAddParticle(Rigid.GTGeometryParticle(), Rigid.UniqueIdx());
 			FDirtyPropertiesManager& Manager = *Managers.Last().Manager;
 
 			//sim-writable properties changed at head, so we must write down what they were
@@ -390,6 +460,7 @@ private:
 	struct FParticleRewindInfo
 	{
 		TArray<FFrameInfo> Frames;
+		int32 AllDirtyParticlesIdx;
 
 		FGeometryParticleState& AddFrame(TGeometryParticle<FReal,3>& GTParticleUnsafe, int32 FrameIdx)
 		{
@@ -427,8 +498,29 @@ private:
 		int32 FrameCreatedFor;
 	};
 
+	void Reset()
+	{
+		ParticleToRewindInfo.Reset();
+		Managers.Reset();
+		AllDirtyParticles.Reset();
+	}
+
+	FParticleRewindInfo& FindOrAddParticle(TGeometryParticle<FReal,3>* UnsafeGTParticle, FUniqueIdx UniqueIdx)
+	{
+		if(FParticleRewindInfo* Info = ParticleToRewindInfo.Find(UniqueIdx))
+		{
+			return *Info;
+		}
+
+
+		FParticleRewindInfo& Info = ParticleToRewindInfo.FindOrAdd(UniqueIdx);
+		Info.AllDirtyParticlesIdx = AllDirtyParticles.Add(UnsafeGTParticle);
+		return Info;
+	}
+
 	TArrayAsMap<FUniqueIdx,FParticleRewindInfo> ParticleToRewindInfo;
 	TArray<FFrameManagerInfo> Managers;
+	TArray<TGeometryParticle<FReal,3>*> AllDirtyParticles;
 	int32 CurFrame;
 };
 }
