Chaos Joint - RBAN opt: don't update velocity in apply for joints or collisions (collisions calculate it when needed)

diff --git a/Engine/Source/Runtime/Experimental/Chaos/Private/Chaos/PBDCollisionConstraintsContact.cpp b/Engine/Source/Runtime/Experimental/Chaos/Private/Chaos/PBDCollisionConstraintsContact.cpp
--- a/Engine/Source/Runtime/Experimental/Chaos/Private/Chaos/PBDCollisionConstraintsContact.cpp
+++ b/Engine/Source/Runtime/Experimental/Chaos/Private/Chaos/PBDCollisionConstraintsContact.cpp
@@ -235,12 +235,8 @@ namespace Chaos
 			bool bIsRigidDynamic0 = PBDRigid0 && PBDRigid0->ObjectState() == EObjectStateType::Dynamic;
 			bool bIsRigidDynamic1 = PBDRigid1 && PBDRigid1->ObjectState() == EObjectStateType::Dynamic;
 
-			FVec3 X0 = FParticleUtilitiesXR::GetCoMWorldPosition(Particle0);
-			FVec3 X1 = FParticleUtilitiesXR::GetCoMWorldPosition(Particle1);
 			FVec3 P0 = FParticleUtilities::GetCoMWorldPosition(Particle0);
 			FVec3 P1 = FParticleUtilities::GetCoMWorldPosition(Particle1);
-			FRotation3 R0 = FParticleUtilitiesXR::GetCoMWorldRotation(Particle0);
-			FRotation3 R1 = FParticleUtilitiesXR::GetCoMWorldRotation(Particle1);
 			FRotation3 Q0 = FParticleUtilities::GetCoMWorldRotation(Particle0);
 			FRotation3 Q1 = FParticleUtilities::GetCoMWorldRotation(Particle1);
 			FVec3 VectorToPoint0 = Contact.Location - P0;
@@ -271,6 +267,10 @@ namespace Chaos
 				if (Contact.Friction > 0)
 				{
 					// Get contact velocity
+					FVec3 X0 = FParticleUtilitiesXR::GetCoMWorldPosition(Particle0);
+					FVec3 X1 = FParticleUtilitiesXR::GetCoMWorldPosition(Particle1);
+					FRotation3 R0 = FParticleUtilitiesXR::GetCoMWorldRotation(Particle0);
+					FRotation3 R1 = FParticleUtilitiesXR::GetCoMWorldRotation(Particle1);
 					FVec3 V0 = FVec3::CalculateVelocity(X0, P0, IterationParameters.Dt);
 					FVec3 W0 = FRotation3::CalculateAngularVelocity(R0, Q0, IterationParameters.Dt);
 					FVec3 V1 = FVec3::CalculateVelocity(X1, P1, IterationParameters.Dt);
@@ -308,8 +308,6 @@ namespace Chaos
 				{
 					FVec3 DP0 = InvM0 * DX;
 					FVec3 DR0 = Utilities::Multiply(InvI0, FVec3::CrossProduct(VectorToPoint0, DX));
-					PBDRigid0->V() += DP0 / IterationParameters.Dt;
-					PBDRigid0->W() += DR0 / IterationParameters.Dt;
 					P0 += DP0;
 					Q0 += FRotation3::FromElements(DR0, 0.f) * Q0 * FReal(0.5);
 					Q0.Normalize();
@@ -319,8 +317,6 @@ namespace Chaos
 				{
 					FVec3 DP1 = InvM1 * -DX;
 					FVec3 DR1 = Utilities::Multiply(InvI1, FVec3::CrossProduct(VectorToPoint1, -DX));
-					PBDRigid1->V() += DP1 / IterationParameters.Dt;
-					PBDRigid1->W() += DR1 / IterationParameters.Dt;
 					P1 += DP1;
 					Q1 += FRotation3::FromElements(DR1, 0.f) * Q1 * FReal(0.5);
 					Q1.Normalize();
diff --git a/Engine/Source/Runtime/Experimental/Chaos/Private/Chaos/PBDJointConstraints.cpp b/Engine/Source/Runtime/Experimental/Chaos/Private/Chaos/PBDJointConstraints.cpp
--- a/Engine/Source/Runtime/Experimental/Chaos/Private/Chaos/PBDJointConstraints.cpp
+++ b/Engine/Source/Runtime/Experimental/Chaos/Private/Chaos/PBDJointConstraints.cpp
@@ -857,8 +857,9 @@ namespace Chaos
 		TPBDRigidParticleHandle<FReal, 3>* Particle1 = ConstraintParticles[JointIndex][Index1]->CastToRigidParticle();
 
 		FJointSolverJointState& JointState = SolverConstraintStates[JointIndex];
-		UpdateParticleState(Particle0, Dt, JointState.PrevPs[0], JointState.PrevQs[0], JointState.Ps[0], JointState.Qs[0]);
-		UpdateParticleState(Particle1, Dt, JointState.PrevPs[1], JointState.PrevQs[1], JointState.Ps[1], JointState.Qs[1]);
+		bool bUpdateVelocity = false;	// Position-based collision solver does not need V() and W()
+		UpdateParticleState(Particle0, Dt, JointState.PrevPs[0], JointState.PrevQs[0], JointState.Ps[0], JointState.Qs[0], bUpdateVelocity);
+		UpdateParticleState(Particle1, Dt, JointState.PrevPs[1], JointState.PrevQs[1], JointState.Ps[1], JointState.Qs[1], bUpdateVelocity);
 	}
 
 	int32 FPBDJointConstraints::ApplyBatch(const FReal Dt, const int32 BatchIndex, const int32 NumPairIts, const int32 It, const int32 NumIts)
@@ -1043,8 +1044,9 @@ namespace Chaos
 			}
 		}
 
-		UpdateParticleState(Particle0->CastToRigidParticle(), Dt, Solver.GetPrevP(0), Solver.GetPrevQ(0), Solver.GetP(0), Solver.GetQ(0));
-		UpdateParticleState(Particle1->CastToRigidParticle(), Dt, Solver.GetPrevP(1), Solver.GetPrevQ(1), Solver.GetP(1), Solver.GetQ(1));
+		bool bUpdateVelocity = false;	// Position-based collision solver does not need V() and W()
+		UpdateParticleState(Particle0->CastToRigidParticle(), Dt, Solver.GetPrevP(0), Solver.GetPrevQ(0), Solver.GetP(0), Solver.GetQ(0), bUpdateVelocity);
+		UpdateParticleState(Particle1->CastToRigidParticle(), Dt, Solver.GetPrevP(1), Solver.GetPrevQ(1), Solver.GetP(1), Solver.GetQ(1), bUpdateVelocity);
 
 		return NumActive;
 	}
