ToolsFramework: ToolManager needs to keep track of whether it is active or not to be able to properly expire change events.
#rnx

diff --git a/Engine/Plugins/MeshPainting/Source/MeshPaintingToolset/Private/MeshPaintHelpers.cpp b/Engine/Plugins/MeshPainting/Source/MeshPaintingToolset/Private/MeshPaintHelpers.cpp
--- a/Engine/Plugins/MeshPainting/Source/MeshPaintingToolset/Private/MeshPaintHelpers.cpp
+++ b/Engine/Plugins/MeshPainting/Source/MeshPaintingToolset/Private/MeshPaintHelpers.cpp
@@ -1423,6 +1423,8 @@ UMeshToolManager::UMeshToolManager()
 
 void UMeshToolManager::Shutdown()
 {
+	UInteractiveToolManager::Shutdown();
+
 	ComponentToAdapterMap.Empty();
 	SelectedMeshComponents.Empty();
 }
diff --git a/Engine/Source/Runtime/Experimental/InteractiveToolsFramework/Private/InteractiveToolManager.cpp b/Engine/Source/Runtime/Experimental/InteractiveToolsFramework/Private/InteractiveToolManager.cpp
--- a/Engine/Source/Runtime/Experimental/InteractiveToolsFramework/Private/InteractiveToolManager.cpp
+++ b/Engine/Source/Runtime/Experimental/InteractiveToolsFramework/Private/InteractiveToolManager.cpp
@@ -30,6 +30,8 @@ void UInteractiveToolManager::Initialize(IToolsContextQueriesAPI* queriesAPI, IT
 	this->QueriesAPI = queriesAPI;
 	this->TransactionsAPI = transactionsAPI;
 	this->InputRouter = InputRouterIn;
+
+	bIsActive = true;
 }
 
 
@@ -47,6 +49,8 @@ void UInteractiveToolManager::Shutdown()
 	}
 
 	this->TransactionsAPI = nullptr;
+
+	bIsActive = false;
 }
 
 
@@ -363,7 +367,7 @@ void FBeginToolChange::Revert(UObject* Object)
 bool FBeginToolChange::HasExpired( UObject* Object ) const
 {
 	UInteractiveToolManager* ToolManager = CastChecked<UInteractiveToolManager>(Object);
-	return (ToolManager == nullptr) || (ToolManager->HasAnyActiveTool() == false);
+	return (ToolManager == nullptr) || (ToolManager->IsActive() == false) || (ToolManager->HasAnyActiveTool() == false);
 }
 
 FString FBeginToolChange::ToString() const
@@ -412,7 +416,7 @@ void FActivateToolChange::Revert(UObject* Object)
 bool FActivateToolChange::HasExpired(UObject* Object) const
 {
 	UInteractiveToolManager* ToolManager = CastChecked<UInteractiveToolManager>(Object);
-	return (ToolManager == nullptr);
+	return (ToolManager == nullptr) || (ToolManager->IsActive() == false);
 }
 
 FString FActivateToolChange::ToString() const
diff --git a/Engine/Source/Runtime/Experimental/InteractiveToolsFramework/Public/InteractiveToolManager.h b/Engine/Source/Runtime/Experimental/InteractiveToolsFramework/Public/InteractiveToolManager.h
--- a/Engine/Source/Runtime/Experimental/InteractiveToolsFramework/Public/InteractiveToolManager.h
+++ b/Engine/Source/Runtime/Experimental/InteractiveToolsFramework/Public/InteractiveToolManager.h
@@ -75,6 +75,13 @@ protected:
 
 public:
 
+
+	/**
+	 * @return true if ToolManager is currently active, ie between Initialize() and Shutdown() 
+	 */
+	bool IsActive() const { return bIsActive; }
+
+
 	//
 	// Tool registration and Current Tool state
 	//
@@ -250,6 +257,9 @@ protected:
 	/** Pointer to current InputRouter (Context owns this) */
 	UInputRouter* InputRouter;
 
+	/** This flag is set to true on Initialize() and false on Shutdown(). */
+	bool bIsActive = false;
+
 	/** Current set of named ToolBuilders */
 	UPROPERTY()
 	TMap<FString, UInteractiveToolBuilder*> ToolBuilders;
