Implement CollisionResolution_SimplicialSphereToImplicitSphere().  Do a pass at simplifying the geometry collection physics proxy's CreateTriangleMesh() (by eliminating unused parameters), and IntializeSharedCollisionStructures().  Add support in the testing infracture for simplicial type none.

diff --git a/Engine/Source/Runtime/Experimental/ChaosSolvers/Private/PhysicsProxy/GeometryCollectionPhysicsProxy.cpp b/Engine/Source/Runtime/Experimental/ChaosSolvers/Private/PhysicsProxy/GeometryCollectionPhysicsProxy.cpp
--- a/Engine/Source/Runtime/Experimental/ChaosSolvers/Private/PhysicsProxy/GeometryCollectionPhysicsProxy.cpp
+++ b/Engine/Source/Runtime/Experimental/ChaosSolvers/Private/PhysicsProxy/GeometryCollectionPhysicsProxy.cpp
@@ -104,37 +104,26 @@ bool IsMultithreaded()
 }
 
 Chaos::TTriangleMesh<float>* CreateTriangleMesh(
+	const int32 FaceStart,
 	const int32 FaceCount, 
-	const int32 VertexOffset, 
-	const int32 StartIndex, 
-	const TManagedArray<FVector>& Vertex, 
 	const TManagedArray<bool>& Visible, 
-	const TManagedArray<FIntVector>& Indices, 
-	TSet<int32>& VertsAdded)
+	const TManagedArray<FIntVector>& Indices)
 {
 	TArray<Chaos::TVector<int32, 3>> Faces;
+	Faces.Reserve(FaceCount);
+	
+	const int32 FaceEnd = FaceStart + FaceCount;
+	for (int Idx = FaceStart; Idx < FaceEnd; ++Idx)
 	{
-		Faces.Reserve(FaceCount);
-		for (int j = 0; j < FaceCount; j++)
+		// Note: This function used to cull small triangles.  As one of the purposes 
+		// of the tri mesh this function creates is for level set rasterization, we 
+		// don't want to do that.  Keep the mesh intact, which hopefully is water tight.
+		if (Visible[Idx])
 		{
-			if (!Visible[j + StartIndex])
-			{
-				continue;
-			}
-
-			// Note: This function used to cull small triangles.  As one of the purposes 
-			// of the tri mesh this function creates is for level set rasterization, we 
-			// don't want to do that.  Keep the mesh intact, which hopefully is water tight.
-
-			const FIntVector& Tri = Indices[j + StartIndex];
+			const FIntVector& Tri = Indices[Idx];
 			Faces.Add(Chaos::TVector<int32, 3>(Tri.X, Tri.Y, Tri.Z));
-			for (int Axis = 0; Axis < 3; ++Axis)
-			{
-				VertsAdded.Add(Tri[Axis]);
-			}
 		}
 	}
-
 	return new Chaos::TTriangleMesh<float>(MoveTemp(Faces)); // Culls geometrically degenerate faces
 }
 
@@ -1835,7 +1824,8 @@ void FGeometryCollectionPhysicsProxy::InitializeSharedCollisionStructures(
 		RestCollection.AddAttribute<TUniquePtr<FSimplicial>>(
 			FGeometryDynamicCollection::SimplicialsAttribute, FTransformCollection::TransformGroup);
 
-	RestCollection.RemoveAttribute(FGeometryDynamicCollection::ImplicitsAttribute, FTransformCollection::TransformGroup);
+	RestCollection.RemoveAttribute(
+		FGeometryDynamicCollection::ImplicitsAttribute, FTransformCollection::TransformGroup);
 	TManagedArray<FGeometryDynamicCollection::FSharedImplicit>& CollectionImplicits =
 		RestCollection.AddAttribute<FGeometryDynamicCollection::FSharedImplicit>(
 			FGeometryDynamicCollection::ImplicitsAttribute, FTransformCollection::TransformGroup);
@@ -1851,19 +1841,20 @@ void FGeometryCollectionPhysicsProxy::InitializeSharedCollisionStructures(
 	// VerticesGroup
 	const TManagedArray<FVector>& Vertex = RestCollection.Vertex;
 
-	// Faces
+	// FacesGroup
 	const TManagedArray<bool>& Visible = RestCollection.Visible;
+	const TManagedArray<FIntVector>& Indices = RestCollection.Indices;
 
 	// GeometryGroup
+	const TManagedArray<int32>& TransformIndex = RestCollection.TransformIndex;
 	const TManagedArray<FBox>& BoundingBox = RestCollection.BoundingBox;
 	TManagedArray<float>& InnerRadius = RestCollection.InnerRadius;
 	TManagedArray<float>& OuterRadius = RestCollection.OuterRadius;
-	const TManagedArray<int32>& VertexCount = RestCollection.VertexCount;
 	const TManagedArray<int32>& VertexStart = RestCollection.VertexStart;
-	const TManagedArray<int32>& FaceCount = RestCollection.FaceCount;
+	const TManagedArray<int32>& VertexCount = RestCollection.VertexCount;
 	const TManagedArray<int32>& FaceStart = RestCollection.FaceStart;
-	const TManagedArray<int32>& TransformIndex = RestCollection.TransformIndex;
-	const TManagedArray<FIntVector>& Indices = RestCollection.Indices;
+	const TManagedArray<int32>& FaceCount = RestCollection.FaceCount;
+
 
 	TArray<FTransform> CollectionSpaceTransforms;
 	{ // tmp scope
@@ -1877,8 +1868,6 @@ void FGeometryCollectionPhysicsProxy::InitializeSharedCollisionStructures(
 	TArray<TUniquePtr<TTriangleMesh<float>>> TriangleMeshesArray;	//use to union trimeshes in cluster case
 	TriangleMeshesArray.AddDefaulted(NumTransforms);
 
-	TArray<TSet<int32>> VertsAddedArray;
-	VertsAddedArray.AddDefaulted(NumGeometries);
 
 	TParticles<float, 3> MassSpaceParticles;
 	MassSpaceParticles.AddParticles(Vertex.Num());
@@ -1909,13 +1898,10 @@ void FGeometryCollectionPhysicsProxy::InitializeSharedCollisionStructures(
 		{
 			TUniquePtr<TTriangleMesh<float>> TriMesh(
 				CreateTriangleMesh(
-					FaceCount[GeometryIndex],
-					VertexStart[GeometryIndex],
 					FaceStart[GeometryIndex],
-					Vertex,
+					FaceCount[GeometryIndex],
 					Visible,
-					Indices,
-					VertsAddedArray[GeometryIndex]));
+					Indices));
 
 			TMassProperties<float, 3>& MassProperties = MassPropertiesArray[GeometryIndex];
 
@@ -1951,15 +1937,6 @@ void FGeometryCollectionPhysicsProxy::InitializeSharedCollisionStructures(
 					else
 					{
 						InertiaComputationNeeded[GeometryIndex] = true;
-
-						if (!FMath::IsNearlyZero(MassProperties.CenterOfMass.SizeSquared()))
-						{
-							for (int32 Idx = 0; Idx < Vertex.Num(); ++Idx)
-							{
-								MassSpaceParticles.X(Idx) -= MassProperties.CenterOfMass;
-							}
-						}
-
 						CollectionMassToLocal[TransformGroupIndex] = FTransform(FQuat::Identity, MassProperties.CenterOfMass);
 					}
 				}
@@ -1967,27 +1944,31 @@ void FGeometryCollectionPhysicsProxy::InitializeSharedCollisionStructures(
 
 			if (InnerRadius[GeometryIndex] == 0.0f || OuterRadius[GeometryIndex] == 0.0f)
 			{
-				const FVector Center = BoundingBox[GeometryIndex].GetCenter();
-				const int32 VStart = VertexStart[GeometryIndex];
 				const int32 VCount = VertexCount[GeometryIndex];
-				InnerRadius[GeometryIndex] = VCount ? TNumericLimits<float>::Max() : 0.0f;
-				OuterRadius[GeometryIndex] = 0.0f;
-				for (int32 VIdx = 0; VIdx < VCount; ++VIdx)
+				if (VCount != 0)
 				{
-					const int32 PtIdx = VStart + VIdx;
-					const FVector& Pt = Vertex[PtIdx];
-					const float DistSq = FVector::DistSquared(Pt, Center);
-					if (InnerRadius[GeometryIndex] > DistSq)
-					{
-						InnerRadius[GeometryIndex] = DistSq;
-					}
-					if (OuterRadius[GeometryIndex] < DistSq)
+					const FVector Center = BoundingBox[GeometryIndex].GetCenter();
+					const int32 VStart = VertexStart[GeometryIndex];
+
+					InnerRadius[GeometryIndex] = VCount ? TNumericLimits<float>::Max() : 0.0f;
+					OuterRadius[GeometryIndex] = 0.0f;
+					for (int32 VIdx = 0; VIdx < VCount; ++VIdx)
 					{
-						OuterRadius[GeometryIndex] = DistSq;
+						const int32 PtIdx = VStart + VIdx;
+						const FVector& Pt = Vertex[PtIdx];
+						const float DistSq = FVector::DistSquared(Pt, Center);
+						if (InnerRadius[GeometryIndex] > DistSq)
+						{
+							InnerRadius[GeometryIndex] = DistSq;
+						}
+						if (OuterRadius[GeometryIndex] < DistSq)
+						{
+							OuterRadius[GeometryIndex] = DistSq;
+						}
 					}
+					InnerRadius[GeometryIndex] = FMath::Sqrt(InnerRadius[GeometryIndex]);
+					OuterRadius[GeometryIndex] = FMath::Sqrt(OuterRadius[GeometryIndex]);
 				}
-				InnerRadius[GeometryIndex] = FMath::Sqrt(InnerRadius[GeometryIndex]);
-				OuterRadius[GeometryIndex] = FMath::Sqrt(OuterRadius[GeometryIndex]);
 			}
 
 			TotalVolume += MassProperties.Volume;
@@ -2016,7 +1997,7 @@ void FGeometryCollectionPhysicsProxy::InitializeSharedCollisionStructures(
 			if (DesiredDensity * Volume_i > SharedParams.MaximumMassClamp)
 			{
 				// For rigid bodies outside of range just defaut to a clamped bounding box, and warn the user.
-				ErrorReporter.ReportError(*FString::Printf(TEXT("Geometry has invalid mass (to large)")));
+				ErrorReporter.ReportError(*FString::Printf(TEXT("Geometry has invalid mass (too large)")));
 				ErrorReporter.HandleLatestError();
 
 				CollectionSimulatableParticles[TransformGroupIndex] = false;
@@ -2034,9 +2015,29 @@ void FGeometryCollectionPhysicsProxy::InitializeSharedCollisionStructures(
 
 			if (InertiaComputationNeeded[GeometryIndex])
 			{
+				if (!FMath::IsNearlyZero(MassProperties.CenterOfMass.SizeSquared()))
+				{
+					const int32 IdxStart = VertexStart[GeometryIndex];
+					const int32 IdxEnd = IdxStart + VertexCount[GeometryIndex];
+					for (int32 Idx = IdxStart; Idx < IdxEnd; ++Idx)
+					{
+						MassSpaceParticles.X(Idx) -= MassProperties.CenterOfMass;
+					}
+				}
+
 				CalculateInertiaAndRotationOfMass(MassSpaceParticles, TriMesh->GetSurfaceElements(), Density_i, MassProperties.CenterOfMass, MassProperties.InertiaTensor, MassProperties.RotationOfMass);
 				CollectionMassToLocal[TransformGroupIndex] = FTransform(MassProperties.RotationOfMass, MassProperties.CenterOfMass);
 				CollectionInertiaTensor[TransformGroupIndex] = TVector<float, 3>(MassProperties.InertiaTensor.M[0][0], MassProperties.InertiaTensor.M[1][1], MassProperties.InertiaTensor.M[2][2]);
+
+				if (!FMath::IsNearlyZero(MassProperties.CenterOfMass.SizeSquared()))
+				{
+					const int32 IdxStart = VertexStart[GeometryIndex];
+					const int32 IdxEnd = IdxStart + VertexCount[GeometryIndex];
+					for (int32 Idx = IdxStart; Idx < IdxEnd; ++Idx)
+					{
+						MassSpaceParticles.X(Idx) += MassProperties.CenterOfMass;
+					}
+				}
 			}
 			else
 			{
@@ -2045,15 +2046,27 @@ void FGeometryCollectionPhysicsProxy::InitializeSharedCollisionStructures(
 			}
 
 			FBox InstanceBoundingBox(EForceInit::ForceInitToZero);
-			const TSet<int32>& VertsAdded = VertsAddedArray[GeometryIndex];
-			for (int32 VertIdx = VertexStart[GeometryIndex]; VertIdx < VertexStart[GeometryIndex] + VertexCount[GeometryIndex]; ++VertIdx)
+			if (TriMesh->GetElements().Num())
 			{
-				MassSpaceParticles.X(VertIdx) += MassProperties.CenterOfMass;
-				if (VertsAdded.Contains(VertIdx))	//only consider verts from the trimesh
+				const TSet<int32> MeshVertices = TriMesh->GetVertices();
+				for (const int32 Idx : MeshVertices)
 				{
-					InstanceBoundingBox += MassSpaceParticles.X(VertIdx);	//build bounding box for visible verts in mass space
+					InstanceBoundingBox += MassSpaceParticles.X(Idx);
 				}
 			}
+			else if(VertexCount[GeometryIndex])
+			{
+				const int32 IdxStart = VertexStart[GeometryIndex];
+				const int32 IdxEnd = IdxStart + VertexCount[GeometryIndex];
+				for (int32 Idx = IdxStart; Idx < IdxEnd; ++Idx)
+				{
+					InstanceBoundingBox += MassSpaceParticles.X(Idx);
+				}
+			}
+			else
+			{
+				InstanceBoundingBox = FBox(MassProperties.CenterOfMass, MassProperties.CenterOfMass);
+			}
 
 			const int32 SizeSpecificIdx = FindSizeSpecificIdx(SharedParams.SizeSpecificData, InstanceBoundingBox);
 			const FSharedSimulationSizeSpecificData& SizeSpecificData = SharedParams.SizeSpecificData[SizeSpecificIdx];
@@ -2395,59 +2408,63 @@ void FGeometryCollectionPhysicsProxy::InitRemoveOnFracture(FGeometryCollection&
 void IdentifySimulatableElements(Chaos::FErrorReporter& ErrorReporter, FGeometryCollection& GeometryCollection)
 {
 	// Determine which collection particles to simulate
+
+	// Geometry group
+	const TManagedArray<int32>& TransformIndex = GeometryCollection.TransformIndex;
 	const TManagedArray<FBox>& BoundingBox = GeometryCollection.BoundingBox;
 	const TManagedArray<int32>& VertexCount = GeometryCollection.VertexCount;
-	const TManagedArray<int32>& TransformIndex = GeometryCollection.TransformIndex;
-	const int32 NumTransforms = GeometryCollection.NumElements(FGeometryCollection::TransformGroup);
 
+	const int32 NumTransforms = GeometryCollection.NumElements(FGeometryCollection::TransformGroup);
 	const int32 NumTransformMappings = TransformIndex.Num();
 
-	// Do not simulate hidden geometry
-	TArray<bool> HiddenObject;
-	HiddenObject.Init(true, NumTransforms);
+	// Faces group
+	const TManagedArray<FIntVector>& Indices = GeometryCollection.Indices;
 	const TManagedArray<bool>& Visible = GeometryCollection.Visible;
+	// Vertices group
 	const TManagedArray<int32>& BoneMap = GeometryCollection.BoneMap;
 
-	const TManagedArray<FIntVector>& Indices = GeometryCollection.Indices;
-	int32 PrevObject = -1;
+	// Do not simulate hidden geometry
+	TArray<bool> HiddenObject;
+	HiddenObject.Init(true, NumTransforms);
+	int32 PrevObject = INDEX_NONE;
 	bool bContiguous = true;
 	for(int32 i = 0; i < Indices.Num(); i++)
 	{
-		if(Visible[i])
+		if(Visible[i]) // Face index i is visible
 		{
-			int32 ObjIdx = BoneMap[Indices[i][0]];
+			const int32 ObjIdx = BoneMap[Indices[i][0]]; // Look up associated bone to the faces X coord.
 			HiddenObject[ObjIdx] = false;
-			if (!ensureMsgf(ObjIdx >= PrevObject, TEXT("Objects are not contiguous. This breaks assumptions later in the pipeline")))
+
+			if (!ensure(ObjIdx >= PrevObject))
 			{
 				bContiguous = false;
 			}
-			
+
 			PrevObject = ObjIdx;
 		}
 	}
 
 	if (!bContiguous)
 	{
+		// What assumptions???  How are we ever going to know if this is still the case?
 		ErrorReporter.ReportError(TEXT("Objects are not contiguous. This breaks assumptions later in the pipeline"));
 		ErrorReporter.HandleLatestError();
 	}
 
-
 	//For now all simulation data is a non compiled attribute. Not clear what we want for simulated vs kinematic collections
 	TManagedArray<bool>& SimulatableParticles = 
 		GeometryCollection.AddAttribute<bool>(
 			FGeometryCollection::SimulatableParticlesAttribute, FTransformCollection::TransformGroup);
-
-	for(int32 TransformIdx = 0; TransformIdx < NumTransforms; TransformIdx++)
-	{
-		SimulatableParticles[TransformIdx] = false;
-	}
+	SimulatableParticles.Fill(false);
 
 	for(int i = 0; i < NumTransformMappings; i++)
 	{
 		int32 Tdx = TransformIndex[i];
 		checkSlow(0 <= Tdx && Tdx < NumTransforms);
-		if(GeometryCollection.IsGeometry(Tdx) && VertexCount[i] && 0.f < BoundingBox[i].GetSize().SizeSquared() && !HiddenObject[Tdx])
+		if (GeometryCollection.IsGeometry(Tdx) && // checks that TransformToGeometryIndex[Tdx] != INDEX_NONE
+			VertexCount[i] &&					 // must have vertices to be simulated?
+			0.f < BoundingBox[i].GetSize().SizeSquared() && // must have a non-zero bbox to be simulated?  No single point?
+			!HiddenObject[Tdx])					 // must have 1 associated face
 		{
 			SimulatableParticles[Tdx] = true;
 		}
diff --git a/Engine/Source/Runtime/Experimental/ChaosSolvers/Public/PhysicsProxy/GeometryCollectionPhysicsProxy.h b/Engine/Source/Runtime/Experimental/ChaosSolvers/Public/PhysicsProxy/GeometryCollectionPhysicsProxy.h
--- a/Engine/Source/Runtime/Experimental/ChaosSolvers/Public/PhysicsProxy/GeometryCollectionPhysicsProxy.h
+++ b/Engine/Source/Runtime/Experimental/ChaosSolvers/Public/PhysicsProxy/GeometryCollectionPhysicsProxy.h
@@ -303,5 +303,5 @@ private:
 
 };
 
-CHAOSSOLVERS_API Chaos::TTriangleMesh<float>* CreateTriangleMesh(const int32 FaceCount,const int32 VertexOffset,const int32 StartIndex,const TManagedArray<FVector>& Vertex,const TManagedArray<bool>& Visible,const TManagedArray<FIntVector>& Indices,TSet<int32>& VertsAdded);
+CHAOSSOLVERS_API Chaos::TTriangleMesh<float>* CreateTriangleMesh(const int32 FaceStart,const int32 FaceCount,const TManagedArray<bool>& Visible,const TManagedArray<FIntVector>& Indices);
 CHAOSSOLVERS_API void BuildSimulationData(Chaos::FErrorReporter& ErrorReporter, FGeometryCollection& GeometryCollection, const FSharedSimulationParameters& SharedParams);
