Chaos: Bug fix clustering, Num children was being reset and lost from child when child was being removed from parent
This prevented any further fractuing on asset after first fracturing.

diff --git a/Engine/Source/Runtime/Experimental/Chaos/Private/Chaos/PBDRigidClustering.cpp b/Engine/Source/Runtime/Experimental/Chaos/Private/Chaos/PBDRigidClustering.cpp
--- a/Engine/Source/Runtime/Experimental/Chaos/Private/Chaos/PBDRigidClustering.cpp
+++ b/Engine/Source/Runtime/Experimental/Chaos/Private/Chaos/PBDRigidClustering.cpp
@@ -727,7 +727,7 @@ namespace Chaos
 			//make sure to remove multi child proxy if it exists
 			ClusteredChild->MultiChildProxyData().Reset();
 			ClusteredChild->MultiChildProxyId().Id = nullptr;
-			ClusteredChild->SetClusterId(ClusterId());
+			ClusteredChild->SetClusterId(ClusterId(nullptr, ClusteredChild->ClusterIds().NumChildren)); // clear Id but retain number of children
 
 			const TRigidTransform<T, d> ChildFrame = ClusteredChild->ChildToParent() * PreSolveTM;
 			Child->SetX(ChildFrame.GetTranslation());
