Fix crash in unit tests

diff --git a/Engine/Source/Runtime/Experimental/Chaos/Private/Chaos/PBDRigidsEvolutionGBF.cpp b/Engine/Source/Runtime/Experimental/Chaos/Private/Chaos/PBDRigidsEvolutionGBF.cpp
--- a/Engine/Source/Runtime/Experimental/Chaos/Private/Chaos/PBDRigidsEvolutionGBF.cpp
+++ b/Engine/Source/Runtime/Experimental/Chaos/Private/Chaos/PBDRigidsEvolutionGBF.cpp
@@ -418,7 +418,11 @@ void FPBDRigidsEvolutionGBF::AdvanceOneTimeStep(const FReal Dt, const FReal Step
 
 	Clustering.AdvanceClustering(Dt, GetCollisionConstraints());
 
-	CaptureRewindData(Particles.GetActiveParticlesView());
+	if(CaptureRewindData)
+	{
+		CaptureRewindData(Particles.GetActiveParticlesView());
+	}
+
 	ParticleUpdatePosition(Particles.GetActiveParticlesView(), Dt);
 }
 
