Avoid holding ptr copies. Makes API easier to use

diff --git a/Engine/Source/Programs/HeadlessChaos/Private/HeadlessChaosTestRewind.cpp b/Engine/Source/Programs/HeadlessChaos/Private/HeadlessChaosTestRewind.cpp
--- a/Engine/Source/Programs/HeadlessChaos/Private/HeadlessChaosTestRewind.cpp
+++ b/Engine/Source/Programs/HeadlessChaos/Private/HeadlessChaosTestRewind.cpp
@@ -86,33 +86,25 @@ namespace ChaosTest {
 		for(int Step = 0; Step < 10; ++Step)
 		{
 			const auto ParticleState = RewindData->GetStateAtFrame(*Particle,Step);
+			EXPECT_EQ(ParticleState.X()[2],100 - Step);
 
-			EXPECT_TRUE(ParticleState != nullptr);
-			if(ParticleState)
+			if(Step < 3)
 			{
-				EXPECT_EQ(ParticleState->X()[2],100 - Step);
-
-				if(Step < 3)
-				{
-					//was sphere
-					EXPECT_EQ(ParticleState->Geometry().Get(),Sphere.Get());
-				}
-				else if(Step < 5 || Step >= 7)
-				{
-					//then became box
-					EXPECT_EQ(ParticleState->Geometry().Get(),Box.Get());
-				}
-				else
-				{
-					//second box
-					EXPECT_EQ(ParticleState->Geometry().Get(),Box2.Get());
-				}
+				//was sphere
+				EXPECT_EQ(ParticleState.Geometry().Get(),Sphere.Get());
+			}
+			else if(Step < 5 || Step >= 7)
+			{
+				//then became box
+				EXPECT_EQ(ParticleState.Geometry().Get(),Box.Get());
+			}
+			else
+			{
+				//second box
+				EXPECT_EQ(ParticleState.Geometry().Get(),Box2.Get());
 			}
 		}
 		
-		//no data at head
-		EXPECT_EQ(RewindData->GetStateAtFrame(*Particle,10), nullptr);
-
 		// Throw out the proxy
 		Solver->UnregisterObject(Particle.Get());
 
@@ -155,12 +147,7 @@ namespace ChaosTest {
 		for(int Step = 0; Step < 10; ++Step)
 		{
 			const auto ParticleState = RewindData->GetStateAtFrame(*Particle,Step);
-
-			EXPECT_NE(ParticleState, nullptr);
-			if(ParticleState)
-			{
-				EXPECT_EQ(ParticleState->F()[2],Step+1);
-			}
+			EXPECT_EQ(ParticleState.F()[2],Step+1);
 		}
 
 		// Throw out the proxy
@@ -212,27 +199,17 @@ namespace ChaosTest {
 		{
 			const auto ParticleState = RewindData->GetStateAtFrame(*Particle,Step);
 
-			//Should have data up to frame 5 and then nothing (includes 5 because dynamic is record on current frame)
-			if(Step <= 5)
+			if(Step == 3)
 			{
-				EXPECT_NE(ParticleState, nullptr);
+				EXPECT_EQ(ParticleState.F()[2],3);
 			}
-			else
+			else if(Step == 5)
 			{
-				EXPECT_EQ(ParticleState,nullptr);
+				EXPECT_EQ(ParticleState.F()[2],5);
 			}
-			
-
-			if(ParticleState)
+			else
 			{
-				if(Step == 3)
-				{
-					EXPECT_EQ(ParticleState->F()[2],3);
-				}
-				else if(Step == 5)
-				{
-					EXPECT_EQ(ParticleState->F()[2],5);
-				}
+				EXPECT_EQ(ParticleState.F()[2],0);
 			}
 		}
 
@@ -292,34 +269,21 @@ namespace ChaosTest {
 		{
 			const auto ParticleState = RewindData->GetStateAtFrame(*Particle,Step);
 
-			//Should have data up to frame 6 and then nothing
-			if(Step < 7)
-			{
-				EXPECT_NE(ParticleState,nullptr);
-			} else
+			
+			if(Step < 3)
 			{
-				EXPECT_EQ(ParticleState,nullptr);
+				//was sphere
+				EXPECT_EQ(ParticleState.Geometry().Get(),Sphere.Get());
 			}
-
-			if(ParticleState)
+			else if(Step < 5 || Step >= 7)
 			{
-				if(Step < 3)
-				{
-					//was sphere
-					EXPECT_EQ(ParticleState->Geometry().Get(),Sphere.Get());
-				} else if(Step < 5)
-				{
-					//then became box
-					EXPECT_EQ(ParticleState->Geometry().Get(),Box.Get());
-				} else
-				{
-					//second box
-					EXPECT_EQ(ParticleState->Geometry().Get(),Box2.Get());
-				}
+				//then became box
+				EXPECT_EQ(ParticleState.Geometry().Get(),Box.Get());
 			}
 			else
 			{
-				EXPECT_EQ(Particle->Geometry().Get(),Box.Get());
+				//second box
+				EXPECT_EQ(ParticleState.Geometry().Get(),Box2.Get());
 			}
 		}
 
@@ -374,10 +338,9 @@ namespace ChaosTest {
 		for(int Step = 0; Step < 9; ++Step)
 		{
 			const auto ParticleState = RewindData->GetStateAtFrame(*Particle,Step);
-			EXPECT_NE(ParticleState,nullptr);
 			
-			EXPECT_EQ(ParticleState->X()[2],X[Step][2]);
-			EXPECT_EQ(ParticleState->V()[2],V[Step][2]);
+			EXPECT_EQ(ParticleState.X()[2],X[Step][2]);
+			EXPECT_EQ(ParticleState.V()[2],V[Step][2]);
 		}
 
 		// Throw out the proxy
diff --git a/Engine/Source/Runtime/Experimental/ChaosSolvers/Public/RewindData.h b/Engine/Source/Runtime/Experimental/ChaosSolvers/Public/RewindData.h
--- a/Engine/Source/Runtime/Experimental/ChaosSolvers/Public/RewindData.h
+++ b/Engine/Source/Runtime/Experimental/ChaosSolvers/Public/RewindData.h
@@ -2,6 +2,7 @@
 #pragma once
 #include "Chaos/Core.h"
 #include "Chaos/ParticleHandle.h"
+#include "PhysicsProxy/SingleParticlePhysicsProxy.h"
 #include "PhysicsProxy/SingleParticlePhysicsProxyFwd.h"
 #include "Chaos/Framework/PhysicsSolverBase.h"
 
@@ -72,46 +73,40 @@ inline bool SimWritablePropsMayChange(const TGeometryParticleHandle<FReal,3>& Ha
 	return false;
 }
 
-class FGeometryParticleState
+class FGeometryParticleStateBase
 {
 public:
-
-	FGeometryParticleState(TGeometryParticle<FReal,3>& InParticle)
-		: Particle(InParticle)
-	{
-	}
-
-	const FVec3& X() const
+	const FVec3& X(const TGeometryParticle<FReal,3>& Particle) const
 	{
 		return ParticlePositionRotation.IsSet() ? ParticlePositionRotation.Read().X : Particle.X();
 	}
 
-	const FRotation3& R() const
+	const FRotation3& R(const TGeometryParticle<FReal,3>& Particle) const
 	{
 		return ParticlePositionRotation.IsSet() ? ParticlePositionRotation.Read().R : Particle.R();
 	}
 
-	const FVec3& V() const
+	const FVec3& V(const TGeometryParticle<FReal,3>& Particle) const
 	{
 		return Velocities.IsSet() ? Velocities.Read().V : Particle.CastToKinematicParticle()->V();
 	}
 
-	const FVec3& W() const
+	const FVec3& W(const TGeometryParticle<FReal,3>& Particle) const
 	{
 		return Velocities.IsSet() ? Velocities.Read().W : Particle.CastToKinematicParticle()->W();
 	}
 
-	TSerializablePtr<FImplicitObject> Geometry() const
+	TSerializablePtr<FImplicitObject> Geometry(const TGeometryParticle<FReal,3>& Particle) const
 	{
 		return NonFrequentData.IsSet() ? MakeSerializable(NonFrequentData.Read().Geometry) : Particle.Geometry();
 	}
 
-	const FVec3& F() const
+	const FVec3& F(const TGeometryParticle<FReal,3>& Particle) const
 	{
 		return Dynamics.IsSet() ? Dynamics.Read().F : Particle.CastToRigidParticle()->F();
 	}
 
-	void SyncSimWritablePropsFromSim(FDirtyPropertiesManager& Manager,int32 Idx, const TPBDRigidParticleHandle<FReal,3>& Rigid)
+	void SyncSimWritablePropsFromSim(FDirtyPropertiesManager& Manager,int32 Idx,const TPBDRigidParticleHandle<FReal,3>& Rigid)
 	{
 		FParticleDirtyFlags Flags;
 		Flags.MarkDirty(EParticleFlags::XR);
@@ -132,22 +127,22 @@ public:
 		});
 	}
 
-	void SyncDirtyDynamics(FDirtyPropertiesManager& DestManager,int32 DataIdx, const FParticleDirtyData& Dirty, const FDirtyPropertiesManager& SrcManager)
+	void SyncDirtyDynamics(FDirtyPropertiesManager& DestManager,int32 DataIdx,const FParticleDirtyData& Dirty,const FDirtyPropertiesManager& SrcManager)
 	{
 		FParticleDirtyData DirtyFlags;
 		DirtyFlags.SetFlags(Dirty.GetFlags());
 
-		Dynamics.SyncRemoteData(DestManager,DataIdx,DirtyFlags,[&Dirty, &SrcManager, DataIdx](auto& Data)
+		Dynamics.SyncRemoteData(DestManager,DataIdx,DirtyFlags,[&Dirty,&SrcManager,DataIdx](auto& Data)
 		{
-			Data = Dirty.GetDynamics(SrcManager, DataIdx);
+			Data = Dirty.GetDynamics(SrcManager,DataIdx);
 		});
 	}
 
-	void SyncToParticle() const
+	void SyncToParticle(TGeometryParticle<FReal,3>& Particle) const
 	{
 		//todo: set properties directly instead of assigning sub-properties
 
-		ParticlePositionRotation.SyncToParticle([this](const auto& Data)
+		ParticlePositionRotation.SyncToParticle([&Particle](const auto& Data)
 		{
 			Particle.SetX(Data.X);
 			Particle.SetR(Data.R);
@@ -163,7 +158,7 @@ public:
 		}
 
 
-		NonFrequentData.SyncToParticle([this](const auto& Data)
+		NonFrequentData.SyncToParticle([&Particle](const auto& Data)
 		{
 			Particle.SetGeometry(Data.Geometry);
 			Particle.SetUserData(Data.UserData);
@@ -190,7 +185,7 @@ public:
 			});
 		}
 	}
-	
+
 	void SyncPrevFrame(FDirtyPropertiesManager& Manager,int32 Idx,const FDirtyProxy& Dirty)
 	{
 		//syncs the data before it was made dirty
@@ -220,8 +215,8 @@ public:
 				});
 			}
 		}
-		
-		NonFrequentData.SyncRemoteData(Manager,Idx,Dirty.ParticleData, [Handle](FParticleNonFrequentData& Data)
+
+		NonFrequentData.SyncRemoteData(Manager,Idx,Dirty.ParticleData,[Handle](FParticleNonFrequentData& Data)
 		{
 			Data.Geometry = Handle->SharedGeometryLowLevel();
 			Data.UserData = Handle->UserData();
@@ -241,7 +236,7 @@ public:
 		});
 	}
 
-	bool CoalesceState(const FGeometryParticleState& LatestState)
+	bool CoalesceState(const FGeometryParticleStateBase& LatestState)
 	{
 		bool bCoalesced = false;
 		if(!ParticlePositionRotation.IsSet() && LatestState.ParticlePositionRotation.IsSet())
@@ -263,17 +258,15 @@ public:
 		}
 
 		//dynamics do not coalesce since they are always written when dirty
-		
+
 		return bCoalesced;
 	}
-
-protected:
-	TGeometryParticle<FReal,3>& Particle;
 private:
+
 	TParticleStateProperty<FParticlePositionRotation,EParticleProperty::XR> ParticlePositionRotation;
-	TParticleStateProperty<FParticleNonFrequentData, EParticleProperty::NonFrequentData> NonFrequentData;
+	TParticleStateProperty<FParticleNonFrequentData,EParticleProperty::NonFrequentData> NonFrequentData;
 	TParticleStateProperty<FParticleVelocities,EParticleProperty::Velocities> Velocities;
-	TParticleStateProperty<FParticleDynamics, EParticleProperty::Dynamics> Dynamics;
+	TParticleStateProperty<FParticleDynamics,EParticleProperty::Dynamics> Dynamics;
 	/*
 	PARTICLE_PROPERTY(XR,FParticlePositionRotation)
 		PARTICLE_PROPERTY(Velocities,FParticleVelocities)
@@ -283,6 +276,55 @@ private:
 		PARTICLE_PROPERTY(MassProps,FParticleMassProps)*/
 };
 
+class FGeometryParticleState : private FGeometryParticleStateBase
+{
+public:
+
+	FGeometryParticleState(const TGeometryParticle<FReal,3>& InParticle)
+	: Particle(InParticle)
+	{
+	}
+
+	FGeometryParticleState(const FGeometryParticleStateBase& State, const TGeometryParticle<FReal,3>& InParticle)
+	: FGeometryParticleStateBase(State)
+	, Particle(InParticle)
+	{
+	}
+
+	const FVec3& X() const
+	{
+		return FGeometryParticleStateBase::X(Particle);
+	}
+
+	const FRotation3& R() const
+	{
+		return FGeometryParticleStateBase::R(Particle);
+	}
+
+	const FVec3& V() const
+	{
+		return FGeometryParticleStateBase::V(Particle);
+	}
+
+	const FVec3& W() const
+	{
+		return FGeometryParticleStateBase::W(Particle);
+	}
+
+	TSerializablePtr<FImplicitObject> Geometry() const
+	{
+		return FGeometryParticleStateBase::Geometry(Particle);
+	}
+
+	const FVec3& F() const
+	{
+		return FGeometryParticleStateBase::F(Particle);
+	}
+
+private:
+	const TGeometryParticle<FReal,3>& Particle;
+};
+
 class FRewindData
 {
 public:
@@ -297,9 +339,9 @@ public:
 		//todo: parallel for
 		for(TGeometryParticle<FReal,3>* Particle : AllDirtyParticles)
 		{
-			if(const FGeometryParticleState* State = GetStateAtFrame(*Particle, Frame))
+			if(const FGeometryParticleStateBase* State = GetStateAtFrameImp(*Particle, Frame))
 			{
-				State->SyncToParticle();
+				State->SyncToParticle(*Particle);
 			}
 		}
 		
@@ -307,34 +349,15 @@ public:
 		CurFrame = Frame;
 	}
 
-	const FGeometryParticleState* GetStateAtFrame(const TGeometryParticle<FReal,3>& Particle,int32 Frame) const
+	FGeometryParticleState GetStateAtFrame(const TGeometryParticle<FReal,3>& Particle,int32 Frame) const
 	{
-		if(const FParticleRewindInfo* Info = ParticleToRewindInfo.Find(Particle.UniqueIdx()))
+		if(const FGeometryParticleStateBase* State = GetStateAtFrameImp(Particle, Frame))
 		{
-			//is it worth doing binary search?
-			const int32 NumFrames = Info->Frames.Num();
-			if(NumFrames > 0)
-			{
-				//If frame is before first capture, use first capture
-				if(Frame <= Info->Frames[0].Frame)
-				{
-					return &Info->Frames[0].State;
-				}
-
-				//If frame is between two captures, use later capture, because we always store the last data before a change
-				//We can never use an earlier capture, because the fact that we captured at all implies _something_ is different from proceeding frames
-				for(int32 FrameIdx = 1; FrameIdx < NumFrames; ++FrameIdx)
-				{
-					if(Frame <= Info->Frames[FrameIdx].Frame)
-					{
-						return &Info->Frames[FrameIdx].State;
-					}
-				}
-			}
+			return FGeometryParticleState(*State,Particle);
 		}
 
 		//If no data, or past capture, just use head
-		return nullptr;
+		return FGeometryParticleState(Particle);
 	}
 
 	void AdvanceFrame()
@@ -384,7 +407,7 @@ public:
 			//Most properties are always a frame behind
 			if(Proxy->IsInitialized())	//Frame delay so proxy must be initialized
 			{
-				FGeometryParticleState& LatestState = Info.AddFrame(*Proxy->GetParticle(),CurFrame-1);
+				FGeometryParticleStateBase& LatestState = Info.AddFrame(CurFrame-1);
 				LatestState.SyncPrevFrame(DestManager,DataIdx,Dirty);
 
 				//for frames further back a simply copy is enough
@@ -402,7 +425,7 @@ public:
 			//If dynamics are dirty we must record them immediately because the sim resets them to 0
 			if(Dirty.ParticleData.IsDirty(EParticleFlags::Dynamics))
 			{
-				FGeometryParticleState& LatestState = Info.AddFrame(*Proxy->GetParticle(),CurFrame);
+				FGeometryParticleStateBase& LatestState = Info.AddFrame(CurFrame);
 				LatestState.SyncDirtyDynamics(DestManager,DataIdx,Dirty.ParticleData,SrcManager);
 			}
 		};
@@ -441,7 +464,7 @@ public:
 			FDirtyPropertiesManager& Manager = *Managers.Last().Manager;
 
 			//sim-writable properties changed at head, so we must write down what they were
-			FGeometryParticleState& LatestState = Info.AddFrame(*Rigid.GTGeometryParticle(),CurFrame-1);
+			FGeometryParticleStateBase& LatestState = Info.AddFrame(CurFrame-1);
 			LatestState.SyncSimWritablePropsFromSim(Manager,DataIdx,Rigid);
 
 			//update any previous frames that were pointing at head
@@ -453,7 +476,7 @@ private:
 
 	struct FFrameInfo
 	{
-		FGeometryParticleState State;
+		FGeometryParticleStateBase State;
 		int32 Frame;
 	};
 
@@ -462,21 +485,21 @@ private:
 		TArray<FFrameInfo> Frames;
 		int32 AllDirtyParticlesIdx;
 
-		FGeometryParticleState& AddFrame(TGeometryParticle<FReal,3>& GTParticleUnsafe, int32 FrameIdx)
+		FGeometryParticleStateBase& AddFrame(int32 FrameIdx)
 		{
 			if(Frames.Num() && Frames.Last().Frame == FrameIdx)
 			{
 				return Frames.Last().State;
 			}
 
-			Frames.Add(FFrameInfo{FGeometryParticleState(GTParticleUnsafe),FrameIdx});
+			Frames.Add(FFrameInfo{FGeometryParticleStateBase(),FrameIdx});
 			return Frames.Last().State;
 		}
 	};
 
 	void CoalesceBack(TArray<FFrameInfo>& Frames)
 	{
-		const FGeometryParticleState& LatestState = Frames.Last().State;
+		const FGeometryParticleStateBase& LatestState = Frames.Last().State;
 		for(int32 FrameIdx = Frames.Num() - 2; FrameIdx >= 0; --FrameIdx)
 		{
 			FFrameInfo& Frame = Frames[FrameIdx];
@@ -518,6 +541,36 @@ private:
 		return Info;
 	}
 
+	const FGeometryParticleStateBase* GetStateAtFrameImp(const TGeometryParticle<FReal,3>& Particle,int32 Frame) const
+	{
+		if(const FParticleRewindInfo* Info = ParticleToRewindInfo.Find(Particle.UniqueIdx()))
+		{
+			//is it worth doing binary search?
+			const int32 NumFrames = Info->Frames.Num();
+			if(NumFrames > 0)
+			{
+				//If frame is before first capture, use first capture
+				if(Frame <= Info->Frames[0].Frame)
+				{
+					return &Info->Frames[0].State;
+				}
+
+				//If frame is between two captures, use later capture, because we always store the last data before a change
+				//We can never use an earlier capture, because the fact that we captured at all implies _something_ is different from proceeding frames
+				for(int32 FrameIdx = 1; FrameIdx < NumFrames; ++FrameIdx)
+				{
+					if(Frame <= Info->Frames[FrameIdx].Frame)
+					{
+						return &Info->Frames[FrameIdx].State;
+					}
+				}
+			}
+		}
+
+		//If no data, or past capture, just use head
+		return nullptr;
+	}
+
 	TArrayAsMap<FUniqueIdx,FParticleRewindInfo> ParticleToRewindInfo;
 	TArray<FFrameManagerInfo> Managers;
 	TArray<TGeometryParticle<FReal,3>*> AllDirtyParticles;
