Move wake-checks when setting properties that affect sleep state to GT.

diff --git a/Engine/Source/Runtime/Experimental/Chaos/Public/Chaos/ParticleHandle.h b/Engine/Source/Runtime/Experimental/Chaos/Public/Chaos/ParticleHandle.h
--- a/Engine/Source/Runtime/Experimental/Chaos/Public/Chaos/ParticleHandle.h
+++ b/Engine/Source/Runtime/Experimental/Chaos/Public/Chaos/ParticleHandle.h
@@ -1400,10 +1400,7 @@ public:
 	static TGeometryParticle<T, d>* SerializationFactory(FChaosArchive& Ar, TGeometryParticle<T, d>* Serializable);
 
 	const TVector<T, d>& X() const { return MXR.Read().X; }
-	void SetX(const TVector<T, d>& InX, bool bInvalidate = true)
-	{
-		MXR.Modify(bInvalidate,MDirtyFlags,Proxy,[&InX](auto& Data){ Data.X = InX;});
-	}
+	void SetX(const TVector<T, d>& InX, bool bInvalidate = true);
 
 	FUniqueIdx UniqueIdx() const { return MNonFrequentData.Read().UniqueIdx; }
 	void SetUniqueIdx(const FUniqueIdx UniqueIdx, bool bInvalidate = true)
@@ -1412,10 +1409,7 @@ public:
 	}
 
 	const TRotation<T, d>& R() const { return MXR.Read().R; }
-	void SetR(const TRotation<T, d>& InR, bool bInvalidate = true)
-	{
-		MXR.Modify(bInvalidate,MDirtyFlags,Proxy,[&InR](auto& Data){ Data.R = InR;});
-	}
+	void SetR(const TRotation<T, d>& InR, bool bInvalidate = true);
 
 	//todo: geometry should not be owned by particle
 	void SetGeometry(TUniquePtr<FImplicitObject>&& UniqueGeometry)
@@ -1831,16 +1825,10 @@ public:
 	}
 
 	const TVector<T, d>& V() const { return MVelocities.Read().V; }
-	void SetV(const TVector<T, d>& InV, bool bInvalidate = true)
-	{
-		MVelocities.Modify(bInvalidate,MDirtyFlags,Proxy,[&InV](auto& Data){ Data.V = InV;});
-	}
+	void SetV(const TVector<T, d>& InV, bool bInvalidate = true);
 
 	const TVector<T, d>& W() const { return MVelocities.Read().W; }
-	void SetW(const TVector<T, d>& InW, bool bInvalidate = true)
-	{
-		MVelocities.Modify(bInvalidate,MDirtyFlags,Proxy,[&InW](auto& Data){ Data.W = InW;});
-	}
+	void SetW(const TVector<T, d>& InW, bool bInvalidate = true);
 
 	const TVector<T, d>& CenterOfMass() const { return MMassProps.Read().CenterOfMass; }
 	void SetCenterOfMass(const TVector<T, d>& InCenterOfMass, bool bInvalidate = true)
@@ -1994,24 +1982,40 @@ public:
 	const TVector<T, d>& F() const { return MDynamics.Read().F; }
 	void SetF(const TVector<T, d>& InF, bool bInvalidate = true)
 	{
+		if (bInvalidate)
+		{
+			SetObjectState(EObjectStateType::Dynamic, true);
+		}
 		MDynamics.Modify(bInvalidate,MDirtyFlags,Proxy,[&InF](auto& Data){ Data.F = InF;});
 	}
 
 	const TVector<T, d>& Torque() const { return MDynamics.Read().Torque; }
 	void SetTorque(const TVector<T, d>& InTorque, bool bInvalidate=true)
 	{
+		if (bInvalidate)
+		{
+			SetObjectState(EObjectStateType::Dynamic, true);
+		}
 		MDynamics.Modify(bInvalidate,MDirtyFlags,Proxy,[&InTorque](auto& Data){ Data.Torque = InTorque;});
 	}
 
 	const TVector<T, d>& LinearImpulse() const { return MDynamics.Read().LinearImpulse; }
 	void SetLinearImpulse(const TVector<T, d>& InLinearImpulse, bool bInvalidate = true)
 	{
+		if (bInvalidate)
+		{
+			SetObjectState(EObjectStateType::Dynamic, true);
+		}
 		MDynamics.Modify(bInvalidate,MDirtyFlags,Proxy,[&InLinearImpulse](auto& Data){ Data.LinearImpulse = InLinearImpulse;});
 	}
 
 	const TVector<T, d>& AngularImpulse() const { return MDynamics.Read().AngularImpulse; }
 	void SetAngularImpulse(const TVector<T, d>& InAngularImpulse, bool bInvalidate = true)
 	{
+		if (bInvalidate)
+		{
+			SetObjectState(EObjectStateType::Dynamic, true);
+		}
 		MDynamics.Modify(bInvalidate,MDirtyFlags,Proxy,[&InAngularImpulse](auto& Data){ Data.AngularImpulse = InAngularImpulse;});
 	}
 
@@ -2307,6 +2311,34 @@ void TGeometryParticle<T, d>::SetObjectState(const EObjectStateType InState, boo
 	}
 }
 
+template <typename T, int d>
+void TGeometryParticle<T, d>::SetX(const TVector<T, d>& InX, bool bInvalidate)
+{
+	if (bInvalidate)
+	{
+		TPBDRigidParticle<T, d>* Dyn = CastToRigidParticle();
+		if (Dyn && Dyn->ObjectState() == EObjectStateType::Sleeping)
+		{
+			Dyn->SetObjectState(EObjectStateType::Dynamic, true);
+		}
+	}
+	MXR.Modify(bInvalidate, MDirtyFlags, Proxy, [&InX](auto& Data) { Data.X = InX; });
+}
+
+template <typename T, int d>
+void TGeometryParticle<T, d>::SetR(const TRotation<T, d>& InR, bool bInvalidate)
+{
+	if (bInvalidate)
+	{
+		TPBDRigidParticle<T, d>* Dyn = CastToRigidParticle();
+		if (Dyn && Dyn->ObjectState() == EObjectStateType::Sleeping)
+		{
+			Dyn->SetObjectState(EObjectStateType::Dynamic, true);
+		}
+	}
+	MXR.Modify(bInvalidate, MDirtyFlags, Proxy, [&InR](auto& Data) { Data.R = InR; });
+}
+
 template <typename T, int d>
 EObjectStateType TGeometryParticle<T, d>::ObjectState() const
 {
@@ -2321,6 +2353,34 @@ EObjectStateType TKinematicGeometryParticle<T, d>::ObjectState() const
 	return Dyn ? Dyn->ObjectState() : EObjectStateType::Kinematic;
 }
 
+template <typename T, int d>
+void TKinematicGeometryParticle<T, d>::SetV(const TVector<T, d>& InV, bool bInvalidate)
+{
+	if (bInvalidate)
+	{
+		TPBDRigidParticle<T, d>* Dyn = CastToRigidParticle();
+		if (Dyn && Dyn->ObjectState() == EObjectStateType::Sleeping && !InV.IsNearlyZero())
+		{
+			Dyn->SetObjectState(EObjectStateType::Dynamic, true);
+		}
+	}
+	MVelocities.Modify(bInvalidate, MDirtyFlags, Proxy, [&InV](auto& Data) { Data.V = InV; });
+}
+
+template <typename T, int d>
+void TKinematicGeometryParticle<T, d>::SetW(const TVector<T, d>& InW, bool bInvalidate)
+{
+	if (bInvalidate)
+	{
+		TPBDRigidParticle<T, d>* Dyn = CastToRigidParticle();
+		if (Dyn && Dyn->ObjectState() == EObjectStateType::Sleeping && !InW.IsNearlyZero())
+		{
+			Dyn->SetObjectState(EObjectStateType::Dynamic, true);
+		}
+	}
+	MVelocities.Modify(bInvalidate, MDirtyFlags, Proxy, [&InW](auto& Data) { Data.W = InW; });
+}
+
 template <typename T, int d>
 TGeometryParticle<T, d>* TGeometryParticle<T, d>::SerializationFactory(FChaosArchive& Ar, TGeometryParticle<T, d>* Serializable)
 {
diff --git a/Engine/Source/Runtime/Experimental/ChaosSolvers/Private/PhysicsProxy/SingleParticlePhysicsProxy.cpp b/Engine/Source/Runtime/Experimental/ChaosSolvers/Private/PhysicsProxy/SingleParticlePhysicsProxy.cpp
--- a/Engine/Source/Runtime/Experimental/ChaosSolvers/Private/PhysicsProxy/SingleParticlePhysicsProxy.cpp
+++ b/Engine/Source/Runtime/Experimental/ChaosSolvers/Private/PhysicsProxy/SingleParticlePhysicsProxy.cpp
@@ -47,11 +47,6 @@ void PushToPhysicsStateImp(const Chaos::FDirtyPropertiesManager& Manager, Chaos:
 	const FParticleDirtyData& ParticleData = Dirty.ParticleData;
 	// move the copied game thread data into the handle
 	{
-		//  DynamicProperties are anything that would wake up a sleeping
-		//  particle. For example, if we set position on a sleeping particle
-		//  we will need to wake it up
-		bool bDynamicPropertyUpdated = false;
-
 		auto NewXR = ParticleData.FindXR(Manager, DataIdx);
 		auto NewNonFrequentData = ParticleData.FindNonFrequentData(Manager, DataIdx);
 
@@ -63,7 +58,6 @@ void PushToPhysicsStateImp(const Chaos::FDirtyPropertiesManager& Manager, Chaos:
 			{
 				RigidHandle->SetP(NewXR->X);
 				RigidHandle->SetQ(NewXR->R);
-				bDynamicPropertyUpdated = true;
 			}
 		}
 
@@ -89,10 +83,6 @@ void PushToPhysicsStateImp(const Chaos::FDirtyPropertiesManager& Manager, Chaos:
 		{
 			KinematicHandle->SetV(NewVelocities->V);
 			KinematicHandle->SetW(NewVelocities->W);
-			if(bHasDynamicData)
-			{
-				bDynamicPropertyUpdated = true;
-			}
 		}
 
 		if(NewXR || NewNonFrequentData || NewVelocities)
@@ -152,16 +142,6 @@ void PushToPhysicsStateImp(const Chaos::FDirtyPropertiesManager& Manager, Chaos:
 			}
 		}
 
-		//if dynamic property changed trigger a wakeup
-		//question: should this be here?
-		if(bDynamicPropertyUpdated)
-		{
-			if(bInitialized && RigidHandle->ObjectState() == Chaos::EObjectStateType::Sleeping)
-			{
-				Solver->GetEvolution()->SetParticleObjectState(RigidHandle,Chaos::EObjectStateType::Dynamic);
-			}
-		}
-
 		//shape properties
 		for(int32 ShapeDataIdx : Dirty.ShapeDataIndices)
 		{
@@ -457,22 +437,15 @@ void FSingleParticlePhysicsProxy<Chaos::TPBDRigidParticle<float, 3>>::PushToPhys
 		RigidHandle->SetDebugName(Data->DebugName);
 #endif
 
-		//  DynamicProperties are anything that would wake up a sleeping
-		//  particle. For example, if we set position on a sleeping particle
-		//  we will need to wake it up
-		bool bDynamicPropertyUpdated = false;
-
 		if (Data->DirtyFlags.IsDirty(Chaos::EParticleFlags::X))
 		{
 			RigidHandle->SetX(Data->X);
 			RigidHandle->SetP(Data->X);
-			bDynamicPropertyUpdated = true;
 		}
 		if (Data->DirtyFlags.IsDirty(Chaos::EParticleFlags::R))
 		{
 			RigidHandle->SetR(Data->R);
 			RigidHandle->SetQ(Data->R);
-			bDynamicPropertyUpdated = true;
 		}
 		if (Data->DirtyFlags.IsDirty(Chaos::EParticleFlags::Geometry))
 		{
@@ -481,20 +454,10 @@ void FSingleParticlePhysicsProxy<Chaos::TPBDRigidParticle<float, 3>>::PushToPhys
 		if (Data->DirtyFlags.IsDirty(Chaos::EParticleFlags::V))
 		{
 			RigidHandle->SetV(Data->MV);
-			//TODO: This is gross, we should clean it up
-			if (!(Data->MV.IsNearlyZero() && Data->MObjectState == Chaos::EObjectStateType::Sleeping))
-			{
-				bDynamicPropertyUpdated = true;
-			}
 		}
 		if (Data->DirtyFlags.IsDirty(Chaos::EParticleFlags::W))
 		{
 			RigidHandle->SetW(Data->MW);
-			//TODO: Same as above
-			if (!(Data->MW.IsNearlyZero() && Data->MObjectState == Chaos::EObjectStateType::Sleeping))
-			{
-				bDynamicPropertyUpdated = true;
-			}
 		}
 
 		RigidHandle->SetCenterOfMass(Data->MCenterOfMass);
@@ -509,24 +472,20 @@ void FSingleParticlePhysicsProxy<Chaos::TPBDRigidParticle<float, 3>>::PushToPhys
 		if (Data->DirtyFlags.IsDirty(Chaos::EParticleFlags::F))
 		{
 			RigidHandle->SetF(Data->MF);
-			bDynamicPropertyUpdated = true;
 		}
 		if (Data->DirtyFlags.IsDirty(Chaos::EParticleFlags::Torque))
 		{
 			RigidHandle->SetTorque(Data->MTorque);
-			bDynamicPropertyUpdated = true;
 		}
 
 		if (Data->DirtyFlags.IsDirty(Chaos::EParticleFlags::LinearImpulse))
 		{
 			RigidHandle->SetLinearImpulse(Data->MLinearImpulse);
-			bDynamicPropertyUpdated = true;
 		}
 
 		if (Data->DirtyFlags.IsDirty(Chaos::EParticleFlags::AngularImpulse))
 		{
 			RigidHandle->SetAngularImpulse(Data->MAngularImpulse);
-			bDynamicPropertyUpdated = true;
 		}
 
 		if (Data->DirtyFlags.IsDirty(Chaos::EParticleFlags::ObjectState))
@@ -587,14 +546,6 @@ void FSingleParticlePhysicsProxy<Chaos::TPBDRigidParticle<float, 3>>::PushToPhys
 				Shape->Materials = Data->ShapeMaterials[CurrShape++];
 			}
 		}
-
-		if (bDynamicPropertyUpdated)
-		{
-			if (bInitialized && RigidHandle->ObjectState() == Chaos::EObjectStateType::Sleeping)
-			{
-				GetSolver()->GetEvolution()->SetParticleObjectState(RigidHandle, Chaos::EObjectStateType::Dynamic);
-			}
-		}
 		
 		if(Data->MInitialized)
 		{
