Properly clean up removed particles from rewind data

diff --git a/Engine/Source/Programs/HeadlessChaos/Private/HeadlessChaosTestRewind.cpp b/Engine/Source/Programs/HeadlessChaos/Private/HeadlessChaosTestRewind.cpp
--- a/Engine/Source/Programs/HeadlessChaos/Private/HeadlessChaosTestRewind.cpp
+++ b/Engine/Source/Programs/HeadlessChaos/Private/HeadlessChaosTestRewind.cpp
@@ -399,4 +399,55 @@ namespace ChaosTest {
 
 		Module->DestroySolver(Solver);
 	}
+
+	GTEST_TEST(RewindTest,Remove)
+	{
+		auto Sphere = TSharedPtr<FImplicitObject,ESPMode::ThreadSafe>(new TSphere<float,3>(TVector<float,3>(0),10));
+
+		FChaosSolversModule* Module = FChaosSolversModule::GetModule();
+		Module->ChangeThreadingMode(EChaosThreadingMode::SingleThread);
+
+		// Make a solver
+		FPhysicsSolver* Solver = Module->CreateSolver(nullptr,ESolverFlags::Standalone);
+		Solver->SetEnabled(true);
+
+		Solver->EnableRewindCapture(20);
+
+
+		// Make particles
+		auto Particle = TPBDRigidParticle<float,3>::CreateParticle();
+
+		Particle->SetGeometry(Sphere);
+		Solver->RegisterObject(Particle.Get());
+		Particle->SetGravityEnabled(true);
+		Particle->SetX(FVec3(0,0,100));
+
+		TArray<FVec3> X;
+		TArray<FVec3> V;
+
+		for(int Step = 0; Step < 10; ++Step)
+		{
+			X.Add(Particle->X());
+			V.Add(Particle->V());
+			TickSolverHelper(Module,Solver);
+		}
+
+		FRewindData* RewindData = Solver->GetRewindData();
+
+		{
+			const FGeometryParticleState State = RewindData->GetStateAtFrame(*Particle,5);
+			EXPECT_EQ(State.X(),X[5]);
+		}
+
+		// Throw out the proxy
+		Solver->UnregisterObject(Particle.Get());
+
+		// State should be the same as being at head because we removed it from solver
+		{
+			const FGeometryParticleState State = RewindData->GetStateAtFrame(*Particle,5);
+			EXPECT_EQ(Particle->X(), State.X());
+		}
+
+		Module->DestroySolver(Solver);
+	}
 }
diff --git a/Engine/Source/Runtime/Experimental/ChaosSolvers/Private/PBDRigidsSolver.cpp b/Engine/Source/Runtime/Experimental/ChaosSolvers/Private/PBDRigidsSolver.cpp
--- a/Engine/Source/Runtime/Experimental/ChaosSolvers/Private/PBDRigidsSolver.cpp
+++ b/Engine/Source/Runtime/Experimental/ChaosSolvers/Private/PBDRigidsSolver.cpp
@@ -418,14 +418,20 @@ namespace Chaos
 
 			if(Handle)
 			{
-			// Remove game thread particle from ActiveGameThreadParticles so we won't crash when pulling physics state
-			// if this particle was deleted after buffering results. 
-			Solver->GetActiveParticlesBuffer()->RemoveActiveParticleFromConsumerBuffer(Handle->GTGeometryParticle());
+				// Remove from rewind data
+				if(FRewindData* RewindData = Solver->GetRewindData())
+				{
+					RewindData->RemoveParticle(*Handle);
+				}
+
+			  // Remove game thread particle from ActiveGameThreadParticles so we won't crash when pulling physics state
+			  // if this particle was deleted after buffering results. 
+			  Solver->GetActiveParticlesBuffer()->RemoveActiveParticleFromConsumerBuffer(Handle->GTGeometryParticle());
   
-			Solver->MParticleToProxy.Remove(Handle);
+			  Solver->MParticleToProxy.Remove(Handle);
   
-			// Use the handle to destroy the particle data
-			Solver->GetEvolution()->DestroyParticle(Handle);
+			  // Use the handle to destroy the particle data
+			  Solver->GetEvolution()->DestroyParticle(Handle);
 			}
 
 		});
@@ -510,6 +516,9 @@ namespace Chaos
 		return false;
 	}
 
+	int32 RewindCaptureNumFrames = -1;
+	FAutoConsoleVariableRef CVarRewindCaptureNumFrames(TEXT("p.RewindCaptureNumFrames"),RewindCaptureNumFrames,TEXT("The number of frames to capture rewind for. Requires restart of solver"));
+	
 	void FPBDRigidsSolver::Reset()
 	{
 		UE_LOG(LogPBDRigidsSolver, Verbose, TEXT("PBDRigidsSolver::Reset()"));
@@ -525,6 +534,11 @@ namespace Chaos
 
 		DirtyPropertiesManager = MakeUnique<FDoubleBuffer<FDirtyPropertiesManager>>();
 
+		if(RewindCaptureNumFrames >= 0)
+		{
+			EnableRewindCapture(20);
+		}
+
 		MEvolution->SetCaptureRewindDataFunction([this](const TParticleView<TPBDRigidParticles<FReal,3>>& ActiveParticles)
 		{
 			FinalizeRewindData(ActiveParticles);
@@ -746,22 +760,22 @@ namespace Chaos
 			switch(Dirty.Proxy->GetType())
 			{
 			case EPhysicsProxyType::SingleRigidParticleType:
-				{
+			{
 				auto Proxy = static_cast<FRigidParticlePhysicsProxy*>(Dirty.Proxy);
 				ProcessProxyGT(Proxy,DataIdx,Dirty);
-					break;
+				break;
 			}
 			case EPhysicsProxyType::SingleKinematicParticleType:
 			{
 				auto Proxy = static_cast<FKinematicGeometryParticlePhysicsProxy*>(Dirty.Proxy);
 				ProcessProxyGT(Proxy,DataIdx,Dirty);
-					break;
+				break;
 			}
 			case EPhysicsProxyType::SingleGeometryParticleType:
 			{
 				auto Proxy = static_cast<FGeometryParticlePhysicsProxy*>(Dirty.Proxy);
 				ProcessProxyGT(Proxy,DataIdx,Dirty);
-					break;
+				break;
 			}
 			default:
 			ensure("Unknown proxy type in physics solver.");
@@ -834,7 +848,7 @@ namespace Chaos
 					break;
 				}
 				default:
-					ensure("Unknown proxy type in physics solver.");
+				ensure("Unknown proxy type in physics solver.");
 				}
 			});
 
diff --git a/Engine/Source/Runtime/Experimental/ChaosSolvers/Public/RewindData.h b/Engine/Source/Runtime/Experimental/ChaosSolvers/Public/RewindData.h
--- a/Engine/Source/Runtime/Experimental/ChaosSolvers/Public/RewindData.h
+++ b/Engine/Source/Runtime/Experimental/ChaosSolvers/Public/RewindData.h
@@ -337,11 +337,11 @@ public:
 	{
 		ensure(IsInGameThread());
 		//todo: parallel for
-		for(TGeometryParticle<FReal,3>* Particle : AllDirtyParticles)
+		for(const FDirtyParticleInfo& Info : AllDirtyParticles)
 		{
-			if(const FGeometryParticleStateBase* State = GetStateAtFrameImp(*Particle, Frame))
+			if(const FGeometryParticleStateBase* State = GetStateAtFrameImp(*Info.Particle, Frame))
 			{
-				State->SyncToParticle(*Particle);
+				State->SyncToParticle(*Info.Particle);
 			}
 		}
 		
@@ -349,6 +349,22 @@ public:
 		CurFrame = Frame;
 	}
 
+	void RemoveParticle(const TGeometryParticleHandle<FReal,3>& Particle)
+	{
+		if(const FParticleRewindInfo* Info = ParticleToRewindInfo.Find(Particle.UniqueIdx()))
+		{
+			const int32 Idx = Info->AllDirtyParticlesIdx;
+			AllDirtyParticles.RemoveAtSwap(Idx);
+			if(Idx< AllDirtyParticles.Num())
+			{
+				//update particle in new position
+				ParticleToRewindInfo.FindChecked(AllDirtyParticles[Idx].CachedUniqueIdx).AllDirtyParticlesIdx = Idx;
+			}
+
+			ParticleToRewindInfo.RemoveChecked(Particle.UniqueIdx());
+		}
+	}
+
 	FGeometryParticleState GetStateAtFrame(const TGeometryParticle<FReal,3>& Particle,int32 Frame) const
 	{
 		if(const FGeometryParticleStateBase* State = GetStateAtFrameImp(Particle, Frame))
@@ -537,7 +553,7 @@ private:
 
 
 		FParticleRewindInfo& Info = ParticleToRewindInfo.FindOrAdd(UniqueIdx);
-		Info.AllDirtyParticlesIdx = AllDirtyParticles.Add(UnsafeGTParticle);
+		Info.AllDirtyParticlesIdx = AllDirtyParticles.Add(FDirtyParticleInfo{UnsafeGTParticle,UniqueIdx});
 		return Info;
 	}
 
@@ -571,9 +587,15 @@ private:
 		return nullptr;
 	}
 
+	struct FDirtyParticleInfo
+	{
+		TGeometryParticle<FReal,3>* Particle;
+		FUniqueIdx CachedUniqueIdx;	//Needed when manipulating on physics thread and Particle data cannot be read
+	};
+
 	TArrayAsMap<FUniqueIdx,FParticleRewindInfo> ParticleToRewindInfo;
 	TArray<FFrameManagerInfo> Managers;
-	TArray<TGeometryParticle<FReal,3>*> AllDirtyParticles;
+	TArray<FDirtyParticleInfo> AllDirtyParticles;
 	int32 CurFrame;
 };
 }
