Drop warning for contiguous particles, which Brice says is no longer an issue.

diff --git a/Engine/Source/Runtime/Experimental/ChaosSolvers/Private/PhysicsProxy/GeometryCollectionPhysicsProxy.cpp b/Engine/Source/Runtime/Experimental/ChaosSolvers/Private/PhysicsProxy/GeometryCollectionPhysicsProxy.cpp
--- a/Engine/Source/Runtime/Experimental/ChaosSolvers/Private/PhysicsProxy/GeometryCollectionPhysicsProxy.cpp
+++ b/Engine/Source/Runtime/Experimental/ChaosSolvers/Private/PhysicsProxy/GeometryCollectionPhysicsProxy.cpp
@@ -2427,30 +2427,16 @@ void IdentifySimulatableElements(Chaos::FErrorReporter& ErrorReporter, FGeometry
 	TArray<bool> HiddenObject;
 	HiddenObject.Init(true, NumTransforms);
 	int32 PrevObject = INDEX_NONE;
-	bool bContiguous = true;
 	for(int32 i = 0; i < Indices.Num(); i++)
 	{
 		if(Visible[i]) // Face index i is visible
 		{
 			const int32 ObjIdx = BoneMap[Indices[i][0]]; // Look up associated bone to the faces X coord.
 			HiddenObject[ObjIdx] = false;
-
-			if (!ensure(ObjIdx >= PrevObject))
-			{
-				bContiguous = false;
-			}
-
 			PrevObject = ObjIdx;
 		}
 	}
 
-	if (!bContiguous)
-	{
-		// What assumptions???  How are we ever going to know if this is still the case?
-		ErrorReporter.ReportError(TEXT("Objects are not contiguous. This breaks assumptions later in the pipeline"));
-		ErrorReporter.HandleLatestError();
-	}
-
 	//For now all simulation data is a non compiled attribute. Not clear what we want for simulated vs kinematic collections
 	TManagedArray<bool>& SimulatableParticles = 
 		GeometryCollection.AddAttribute<bool>(
