Chaos: Clustering fixes + enabling more tests

diff --git a/Engine/Source/Runtime/Experimental/Chaos/Private/Chaos/PBDRigidClustering.cpp b/Engine/Source/Runtime/Experimental/Chaos/Private/Chaos/PBDRigidClustering.cpp
--- a/Engine/Source/Runtime/Experimental/Chaos/Private/Chaos/PBDRigidClustering.cpp
+++ b/Engine/Source/Runtime/Experimental/Chaos/Private/Chaos/PBDRigidClustering.cpp
@@ -696,7 +696,7 @@ namespace Chaos
 	TSet<TPBDRigidParticleHandle<T, d>*> 
 	TPBDRigidClustering<FPBDRigidsEvolution, FPBDCollisionConstraint, T, d>::ReleaseClusterParticles(
 		TPBDRigidClusteredParticleHandle<T,d>* ClusteredParticle,
-		const TArrayView<T>* ExternalStrainArray)
+		const TMap<TGeometryParticleHandle<T, d>*, float>* ExternalStrainMap)
 	{
 		SCOPE_CYCLE_COUNTER(STAT_ReleaseClusterParticles_STRAIN);
 
@@ -762,8 +762,8 @@ namespace Chaos
 			TPBDRigidClusteredParticleHandle<T, d>* Child = Children[ChildIdx]->CastToClustered();
 			if (!Child)
 				continue;
-			if ((ExternalStrainArray && (*ExternalStrainArray)[ChildIdx]) ||
-			   (!ExternalStrainArray && Child->CollisionImpulses() >= Child->Strain()))
+			if ((ExternalStrainMap && (ExternalStrainMap->Find(Child))) ||
+			   (!ExternalStrainMap && Child->CollisionImpulses() >= Child->Strain()))
 			{
 				// The piece that hits just breaks off - we may want more control 
 				// by looking at the edges of this piece which would give us cleaner 
@@ -1096,7 +1096,7 @@ namespace Chaos
 	template<class FPBDRigidsEvolution, class FPBDCollisionConstraint, class T, int d>
 	TMap<TPBDRigidClusteredParticleHandle<T, d>*, TSet<TPBDRigidParticleHandle<T, d>*>> 
 	TPBDRigidClustering<FPBDRigidsEvolution, FPBDCollisionConstraint, T, d>::BreakingModel(
-		TArrayView<T>* ExternalStrain)
+		TMap<TGeometryParticleHandle<T, d>*, float>* ExternalStrainMap)
 	{
 		SCOPE_CYCLE_COUNTER(STAT_BreakingModel);
 
@@ -1109,7 +1109,7 @@ namespace Chaos
 			{
 				AllActivatedChildren.Add(
 					ClusteredParticle, 
-					ReleaseClusterParticles(ClusteredParticle, ExternalStrain));
+					ReleaseClusterParticles(ClusteredParticle, ExternalStrainMap));
 			}
 			else
 			{
diff --git a/Engine/Source/Runtime/Experimental/Chaos/Public/Chaos/PBDRigidClustering.h b/Engine/Source/Runtime/Experimental/Chaos/Public/Chaos/PBDRigidClustering.h
--- a/Engine/Source/Runtime/Experimental/Chaos/Public/Chaos/PBDRigidClustering.h
+++ b/Engine/Source/Runtime/Experimental/Chaos/Public/Chaos/PBDRigidClustering.h
@@ -102,7 +102,7 @@ public:
 	*    that have a strain value less than this valid will be released from the 
 	*    cluster.
 	*/
-	TSet<TPBDRigidParticleHandle<T, d>*> ReleaseClusterParticles(TPBDRigidClusteredParticleHandle<T, d>* ClusteredParticle, const TArrayView<T>* ExternalStrainArray=nullptr);
+	TSet<TPBDRigidParticleHandle<T, d>*> ReleaseClusterParticles(TPBDRigidClusteredParticleHandle<T, d>* ClusteredParticle, const TMap<TGeometryParticleHandle<T, d>*, float>* ExternalStrainMap = nullptr);
 
 #if 0 // Not called currently
 	/*
@@ -133,7 +133,7 @@ public:
 	*    the children clusters.
 	*/
 	TMap<TPBDRigidClusteredParticleHandle<T, d>*, TSet<TPBDRigidParticleHandle<T, d>*>> BreakingModel(
-		TArrayView<T>* ExternalStrain=nullptr);
+		TMap<TGeometryParticleHandle<T, d>*, float>* ExternalStrainMap = nullptr);
 
 	/**
 	*  PromoteStrains
diff --git a/Engine/Source/Runtime/Experimental/ChaosSolvers/Private/PhysicsProxy/FieldSystemPhysicsProxy.cpp b/Engine/Source/Runtime/Experimental/ChaosSolvers/Private/PhysicsProxy/FieldSystemPhysicsProxy.cpp
--- a/Engine/Source/Runtime/Experimental/ChaosSolvers/Private/PhysicsProxy/FieldSystemPhysicsProxy.cpp
+++ b/Engine/Source/Runtime/Experimental/ChaosSolvers/Private/PhysicsProxy/FieldSystemPhysicsProxy.cpp
@@ -365,7 +365,16 @@ void FFieldSystemPhysicsProxy::FieldParameterUpdateCallback(
 
 						if (StrainSamples.Num())
 						{
-							CurrentSolver->GetEvolution()->GetRigidClustering().BreakingModel(&FloatBuffer);
+							TMap<TGeometryParticleHandle<float, 3>*, float> Map;
+							
+							for (int Index=0; Index < StrainSamples.Num(); Index++)
+							{
+								if (StrainSamples[Index] > 0)
+								{
+									Map.Add(Handles[Index], StrainSamples[Index]);
+								}
+							}
+							CurrentSolver->GetEvolution()->GetRigidClustering().BreakingModel(&Map);
 						}
 					}
 				}
diff --git a/Engine/Source/Runtime/Experimental/ChaosSolvers/Private/PhysicsProxy/GeometryCollectionPhysicsProxy.cpp b/Engine/Source/Runtime/Experimental/ChaosSolvers/Private/PhysicsProxy/GeometryCollectionPhysicsProxy.cpp
--- a/Engine/Source/Runtime/Experimental/ChaosSolvers/Private/PhysicsProxy/GeometryCollectionPhysicsProxy.cpp
+++ b/Engine/Source/Runtime/Experimental/ChaosSolvers/Private/PhysicsProxy/GeometryCollectionPhysicsProxy.cpp
@@ -1002,7 +1002,7 @@ FGeometryCollectionPhysicsProxy::BuildClusters(
 
 	//If we are a root particle use the world transform, otherwise set the relative transform
 	const FTransform CollectionSpaceTransform = GeometryCollectionAlgo::GlobalMatrix(Transform, ParentIndex, CollectionClusterIndex);
-	const Chaos::TRigidTransform<float, 3> ParticleTM = CollectionSpaceTransform * Parameters.WorldTransform;
+	const Chaos::TRigidTransform<float, 3> ParticleTM = MassToLocal[CollectionClusterIndex] * CollectionSpaceTransform * Parameters.WorldTransform;
 
 	//create new cluster particle
 	//The reason we need to pass in a mass orientation override is as follows:
