Another attempt to fix static analysis.

diff --git a/Engine/Source/Runtime/Experimental/Chaos/Private/Chaos/CollisionResolution.cpp b/Engine/Source/Runtime/Experimental/Chaos/Private/Chaos/CollisionResolution.cpp
--- a/Engine/Source/Runtime/Experimental/Chaos/Private/Chaos/CollisionResolution.cpp
+++ b/Engine/Source/Runtime/Experimental/Chaos/Private/Chaos/CollisionResolution.cpp
@@ -2535,6 +2535,10 @@ namespace Chaos
 
 		bool GetPairTOIHack(const TPBDRigidParticleHandle<FReal, 3>* Particle0, const TGeometryParticleHandle<FReal, 3>* Particle1, const FImplicitObject* Implicit0, const FImplicitObject* Implicit1, const FRigidTransform3& StartTransform0, const FRigidTransform3& Transform1, FReal& OutTOI, FVec3& OutNormal, FReal& OutPhi)
 		{
+			// It's impossible to access an null Implicit0 or Implicit1 as we check types.
+			CA_ASSUME(Implicit0);
+			CA_ASSUME(Implicit1);
+
 			if (!Implicit0 || !Implicit1)
 			{
 				return false;
@@ -2554,8 +2558,6 @@ namespace Chaos
 			}
 			Dir /= Length;
 
-			CA_SUPPRESS(6011); // It's impossible to have an null Implicit0 or Implicit1 as we check types.
-
 			if (Implicit0Type == TBox<FReal, 3>::StaticType() && Implicit1Type == THeightField<FReal>::StaticType())
 			{
 				return GetPairTOIHackImpl<TBox<FReal, 3>>(*Implicit0->template GetObject<TBox<FReal, 3>>(), StartTransform0, *Implicit1->template GetObject<THeightField<FReal>>(), Transform1, Dir, Length, OutTOI, OutNormal, OutPhi);
