Add custom collision type.

diff --git a/Engine/Source/Runtime/Engine/Private/PhysicsEngine/BodyInstance.cpp b/Engine/Source/Runtime/Engine/Private/PhysicsEngine/BodyInstance.cpp
--- a/Engine/Source/Runtime/Engine/Private/PhysicsEngine/BodyInstance.cpp
+++ b/Engine/Source/Runtime/Engine/Private/PhysicsEngine/BodyInstance.cpp
@@ -1767,7 +1767,7 @@ bool FBodyInstance::UpdateBodyScale(const FVector& InScale3D, bool bForceUpdate)
 		for (FPhysicsShapeHandle& ShapeHandle : Shapes)
 		{
 			const Chaos::FImplicitObject& ImplicitObject = ShapeHandle.GetGeometry();
-			EImplicitObjectType ImplicitType = ImplicitObject.GetType(true);
+			EImplicitObjectType ImplicitType = ImplicitObject.GetType();
 			EImplicitObjectType GeomType = GetInnerType(ImplicitType);
 
 			const FTransform& RelativeTM = GetRelativeBodyTransform(ShapeHandle);
@@ -1779,7 +1779,7 @@ bool FBodyInstance::UpdateBodyScale(const FVector& InScale3D, bool bForceUpdate)
 
 				// Get GeomType that is transformed
 				const TImplicitObjectTransformed<FReal, 3>& ImplicitObjectTransformed = static_cast<const TImplicitObjectTransformed<FReal, 3>&>(ImplicitObject);
-				GeomType = ImplicitObjectTransformed.GetTransformedObject()->GetType(true);
+				GeomType = ImplicitObjectTransformed.GetTransformedObject()->GetType();
 				CHAOS_ENSURE(!IsInstanced(GeomType));
 			}
 
diff --git a/Engine/Source/Runtime/Experimental/Chaos/Private/Chaos/ChaosDebugDraw.cpp b/Engine/Source/Runtime/Experimental/Chaos/Private/Chaos/ChaosDebugDraw.cpp
--- a/Engine/Source/Runtime/Experimental/Chaos/Private/Chaos/ChaosDebugDraw.cpp
+++ b/Engine/Source/Runtime/Experimental/Chaos/Private/Chaos/ChaosDebugDraw.cpp
@@ -56,7 +56,7 @@ namespace Chaos
 		void DrawShapesImpl(const FRigidTransform3& ShapeTransform, const FImplicitObject* Shape, FColor Color)
 		{
 			// @todo(ccaulfield): handle scale throughout
-			switch (Shape->GetType(false))
+			switch (Shape->GetType())
 			{
 			case ImplicitObjectType::Sphere:
 			{
diff --git a/Engine/Source/Runtime/Experimental/Chaos/Private/Chaos/CollisionResolution.cpp b/Engine/Source/Runtime/Experimental/Chaos/Private/Chaos/CollisionResolution.cpp
--- a/Engine/Source/Runtime/Experimental/Chaos/Private/Chaos/CollisionResolution.cpp
+++ b/Engine/Source/Runtime/Experimental/Chaos/Private/Chaos/CollisionResolution.cpp
@@ -308,7 +308,7 @@ namespace Chaos
 			const TRigidTransform<T, d> AToBTM = AStartTransform.GetRelativeTransform(BTransform);
 			const TVector<T, d> LocalDir = BTransform.InverseTransformVectorNoScale(Dir);
 
-			if (!ensure(B.GetType(true) & ImplicitObjectType::TriangleMesh) || !ensure(!IsInstanced(B.GetType(true))))
+			if (!ensure(B.GetType() & ImplicitObjectType::TriangleMesh) || !ensure(!IsInstanced(B.GetType())))
 			{
 				return TContactPoint<T>();
 			}
diff --git a/Engine/Source/Runtime/Experimental/Chaos/Private/Chaos/CollisionResolutionUtil.cpp b/Engine/Source/Runtime/Experimental/Chaos/Private/Chaos/CollisionResolutionUtil.cpp
--- a/Engine/Source/Runtime/Experimental/Chaos/Private/Chaos/CollisionResolutionUtil.cpp
+++ b/Engine/Source/Runtime/Experimental/Chaos/Private/Chaos/CollisionResolutionUtil.cpp
@@ -179,7 +179,7 @@ namespace Chaos
 				{
 					SCOPE_CYCLE_COUNTER(STAT_UpdateLevelsetSignedDistance);
 
-					if (Object.GetType(true) == ImplicitObjectType::LevelSet && PotentialParticles.Num() > 0)
+					if (Object.GetType() == ImplicitObjectType::LevelSet && PotentialParticles.Num() > 0)
 					{
 						//QUICK_SCOPE_CYCLE_COUNTER(STAT_LevelSet);
 						const TLevelSet<float, 3>* LevelSet = Object.GetObject<TLevelSet<float, 3>>();
@@ -223,7 +223,7 @@ namespace Chaos
 							}
 						}
 					}
-					else if (Object.GetType(true) == ImplicitObjectType::Box && PotentialParticles.Num() > 0)
+					else if (Object.GetType() == ImplicitObjectType::Box && PotentialParticles.Num() > 0)
 					{
 						//QUICK_SCOPE_CYCLE_COUNTER(STAT_Box);
 						const TBox<float, 3>* Box = Object.GetObject<Chaos::TBox<float, 3>>();
@@ -288,7 +288,7 @@ namespace Chaos
 			else
 			{
 				SCOPE_CYCLE_COUNTER(STAT_UpdateLevelsetAll);
-				if (Object.GetType(true) == ImplicitObjectType::LevelSet && NumParticles > 0)
+				if (Object.GetType() == ImplicitObjectType::LevelSet && NumParticles > 0)
 				{
 					const TLevelSet<float, 3>* LevelSet = Object.GetObject<Chaos::TLevelSet<float, 3>>();
 					const TUniformGrid<float, 3>& Grid = LevelSet->GetGrid();
@@ -329,7 +329,7 @@ namespace Chaos
 						}
 					}
 				}
-				else if (Object.GetType(true) == ImplicitObjectType::Plane && NumParticles > 0)
+				else if (Object.GetType() == ImplicitObjectType::Plane && NumParticles > 0)
 				{
 					const TPlane<float, 3>* Plane = Object.GetObject<Chaos::TPlane<float, 3>>();
 
@@ -363,7 +363,7 @@ namespace Chaos
 						}
 					}
 				}
-				else if (Object.GetType(true) == ImplicitObjectType::Box && NumParticles > 0)
+				else if (Object.GetType() == ImplicitObjectType::Box && NumParticles > 0)
 				{
 					const TBox<float, 3>* Box = Object.GetObject<Chaos::TBox<float, 3>>();
 
diff --git a/Engine/Source/Runtime/Experimental/Chaos/Private/Chaos/ImplicitObject.cpp b/Engine/Source/Runtime/Experimental/Chaos/Private/Chaos/ImplicitObject.cpp
--- a/Engine/Source/Runtime/Experimental/Chaos/Private/Chaos/ImplicitObject.cpp
+++ b/Engine/Source/Runtime/Experimental/Chaos/Private/Chaos/ImplicitObject.cpp
@@ -16,11 +16,13 @@
 #include "Chaos/TriangleMeshImplicitObject.h"
 #include "HAL/IConsoleManager.h"
 #include "UObject/DestructionObjectVersion.h"
+#include "UObject/ReleaseObjectVersion.h"
 
 using namespace Chaos;
 
 FImplicitObject::FImplicitObject(int32 Flags, EImplicitObjectType InType)
     : Type(InType)
+	, CollisionType(InType)
     , bIsConvex(!!(Flags & EImplicitObject::IsConvex))
     , bDoCollide(!(Flags & EImplicitObject::DisableCollisions))
     , bHasBoundingBox(!!(Flags & EImplicitObject::HasBoundingBox))
@@ -56,11 +58,16 @@ void FImplicitObject::Track(TSerializablePtr<FImplicitObject> This, const FStrin
 
 #endif
 
-EImplicitObjectType FImplicitObject::GetType(bool bGetTrueType) const
+EImplicitObjectType FImplicitObject::GetType() const
 {
 	return Type;
 }
 
+EImplicitObjectType FImplicitObject::GetCollisionType() const
+{
+	return CollisionType;
+}
+
 bool FImplicitObject::IsValidGeometry() const
 {
 	return true;
@@ -321,6 +328,16 @@ void FImplicitObject::SerializeImp(FArchive& Ar)
 	{
 		bDoCollide = true;
 	}
+
+	Ar.UsingCustomVersion(FReleaseObjectVersion::GUID);
+	if (Ar.CustomVer(FReleaseObjectVersion::GUID) > FReleaseObjectVersion::CustomImplicitCollisionType)
+	{
+		Ar << CollisionType;
+	}
+	else
+	{
+		CollisionType = Type;
+	}
 }
 
 void FImplicitObject::Serialize(FChaosArchive& Ar)
diff --git a/Engine/Source/Runtime/Experimental/Chaos/Public/Chaos/CastingUtilities.h b/Engine/Source/Runtime/Experimental/Chaos/Public/Chaos/CastingUtilities.h
--- a/Engine/Source/Runtime/Experimental/Chaos/Public/Chaos/CastingUtilities.h
+++ b/Engine/Source/Runtime/Experimental/Chaos/Public/Chaos/CastingUtilities.h
@@ -18,7 +18,7 @@ namespace Chaos
 		template <typename Lambda>
 		FORCEINLINE_DEBUGGABLE auto CastHelper(const FImplicitObject& Geom, const FRigidTransform3& TM, const Lambda& Func)
 		{
-			const EImplicitObjectType Type = Geom.GetType(true);
+			const EImplicitObjectType Type = Geom.GetType();
 			switch (Type)
 			{
 			case ImplicitObjectType::Sphere: return Func(Geom.template GetObjectChecked<TSphere<FReal, 3>>(), TM);
diff --git a/Engine/Source/Runtime/Experimental/Chaos/Public/Chaos/Collision/SpatialAccelerationBroadPhase.h b/Engine/Source/Runtime/Experimental/Chaos/Public/Chaos/Collision/SpatialAccelerationBroadPhase.h
--- a/Engine/Source/Runtime/Experimental/Chaos/Public/Chaos/Collision/SpatialAccelerationBroadPhase.h
+++ b/Engine/Source/Runtime/Experimental/Chaos/Public/Chaos/Collision/SpatialAccelerationBroadPhase.h
@@ -204,7 +204,7 @@ namespace Chaos
 					const bool bIsParticle2Kinematic = Particle2.CastToKinematicParticle() &&
 						(Particle2.ObjectState() == EObjectStateType::Kinematic &&
 							(Particle2.CastToKinematicParticle()->V().SizeSquared() > 1e-4 ||
-								Particle2.Geometry()->GetType(true) == TCapsule<float>::StaticType()));
+								Particle2.Geometry()->GetType() == TCapsule<float>::StaticType()));
 					if (Particle1.ObjectState() == EObjectStateType::Sleeping && !bIsParticle2Kinematic)
 					{
 						continue;
diff --git a/Engine/Source/Runtime/Experimental/Chaos/Public/Chaos/GeometryQueries.h b/Engine/Source/Runtime/Experimental/Chaos/Public/Chaos/GeometryQueries.h
--- a/Engine/Source/Runtime/Experimental/Chaos/Public/Chaos/GeometryQueries.h
+++ b/Engine/Source/Runtime/Experimental/Chaos/Public/Chaos/GeometryQueries.h
@@ -31,7 +31,7 @@ namespace Chaos
 	template <typename QueryGeometry>
 	bool OverlapQuery(const FImplicitObject& A, const FRigidTransform3& ATM, const QueryGeometry& B, const FRigidTransform3& BTM, const FReal Thickness = 0, FMTDInfo* OutMTD=nullptr)
 	{
-		const EImplicitObjectType AType = A.GetType(true);
+		const EImplicitObjectType AType = A.GetType();
 		constexpr EImplicitObjectType BType = QueryGeometry::StaticType();
 		
 		if (AType == ImplicitObjectType::Transformed)
@@ -121,7 +121,7 @@ namespace Chaos
 	template <typename SweptGeometry>
 	bool SweepQuery(const FImplicitObject& A, const FRigidTransform3& ATM, const SweptGeometry& B, const FRigidTransform3& BTM, const FVec3& Dir, const FReal Length, FReal& OutTime, FVec3& OutPosition, FVec3& OutNormal, int32& OutFaceIndex, const FReal Thickness, const bool bComputeMTD)
 	{
-		const EImplicitObjectType AType = A.GetType(true);
+		const EImplicitObjectType AType = A.GetType();
 		constexpr EImplicitObjectType BType = SweptGeometry::StaticType();
 
 		bool bResult = false;
diff --git a/Engine/Source/Runtime/Experimental/Chaos/Public/Chaos/ImplicitObject.h b/Engine/Source/Runtime/Experimental/Chaos/Public/Chaos/ImplicitObject.h
--- a/Engine/Source/Runtime/Experimental/Chaos/Public/Chaos/ImplicitObject.h
+++ b/Engine/Source/Runtime/Experimental/Chaos/Public/Chaos/ImplicitObject.h
@@ -160,7 +160,9 @@ public:
 		return static_cast<const T_DERIVED&>(*this);
 	}
 
-	EImplicitObjectType GetType(bool bGetTrueType = false) const;
+	EImplicitObjectType GetType() const;
+	EImplicitObjectType GetCollisionType() const;
+	void SetCollsionType(EImplicitObjectType InCollsiionType) { CollisionType = InCollsiionType; }
 
 	virtual bool IsValidGeometry() const;
 
@@ -303,6 +305,7 @@ public:
 
 protected:
 	EImplicitObjectType Type;
+	EImplicitObjectType CollisionType;
 	bool bIsConvex;
 	bool bDoCollide;
 	bool bHasBoundingBox;
diff --git a/Engine/Source/Runtime/Experimental/Chaos/Public/Chaos/ImplicitObjectScaled.h b/Engine/Source/Runtime/Experimental/Chaos/Public/Chaos/ImplicitObjectScaled.h
--- a/Engine/Source/Runtime/Experimental/Chaos/Public/Chaos/ImplicitObjectScaled.h
+++ b/Engine/Source/Runtime/Experimental/Chaos/Public/Chaos/ImplicitObjectScaled.h
@@ -36,7 +36,7 @@ public:
 		: FImplicitObject(EImplicitObject::HasBoundingBox, Object->GetType() | ImplicitObjectType::IsInstanced)
 		, MObject(MoveTemp(Object))
 	{
-		ensure(IsInstanced(MObject->GetType(true)) == false);	//cannot have an instance of an instance
+		ensure(IsInstanced(MObject->GetType()) == false);	//cannot have an instance of an instance
 		this->bIsConvex = MObject->IsConvex();
 		this->bDoCollide = MObject->GetDoCollide();
 	}
@@ -45,7 +45,7 @@ public:
 		: FImplicitObject(EImplicitObject::HasBoundingBox,Object->GetType() | ImplicitObjectType::IsInstanced)
 		,MObject(Object)
 	{
-		ensure(IsInstanced(MObject->GetType(true)) == false);	//cannot have an instance of an instance
+		ensure(IsInstanced(MObject->GetType()) == false);	//cannot have an instance of an instance
 		this->bIsConvex = MObject->IsConvex();
 		this->bDoCollide = MObject->GetDoCollide();
 	}
@@ -176,9 +176,9 @@ public:
 		, MSharedPtrForRefCount(nullptr)
 		, MInternalThickness(Thickness)
 	{
-		ensureMsgf((IsScaled(MObject->GetType(true)) == false), TEXT("Scaled objects should not contain each other."));
-		ensureMsgf((IsInstanced(MObject->GetType(true)) == false), TEXT("Scaled objects should not contain instances."));
-		switch (MObject->GetType(true))
+		ensureMsgf((IsScaled(MObject->GetType()) == false), TEXT("Scaled objects should not contain each other."));
+		ensureMsgf((IsInstanced(MObject->GetType()) == false), TEXT("Scaled objects should not contain instances."));
+		switch (MObject->GetType())
 		{
 		case ImplicitObjectType::Transformed:
 		case ImplicitObjectType::Union:
@@ -197,9 +197,9 @@ public:
 		, MSharedPtrForRefCount(Object)
 		, MInternalThickness(Thickness)
 	{
-		ensureMsgf((IsScaled(MObject->GetType(true)) == false), TEXT("Scaled objects should not contain each other."));
-		ensureMsgf((IsInstanced(MObject->GetType(true)) == false), TEXT("Scaled objects should not contain instances."));
-		switch (MObject->GetType(true))
+		ensureMsgf((IsScaled(MObject->GetType()) == false), TEXT("Scaled objects should not contain each other."));
+		ensureMsgf((IsInstanced(MObject->GetType()) == false), TEXT("Scaled objects should not contain instances."));
+		switch (MObject->GetType())
 		{
 		case ImplicitObjectType::Transformed:
 		case ImplicitObjectType::Union:
@@ -235,8 +235,8 @@ public:
 		, MInternalThickness(Other.MInternalThickness)
 	    , MLocalBoundingBox(MoveTemp(Other.MLocalBoundingBox))
 	{
-		ensureMsgf((IsScaled(MObject->GetType(true)) == false), TEXT("Scaled objects should not contain each other."));
-		ensureMsgf((IsInstanced(MObject->GetType(true)) == false), TEXT("Scaled objects should not contain instances."));
+		ensureMsgf((IsScaled(MObject->GetType()) == false), TEXT("Scaled objects should not contain each other."));
+		ensureMsgf((IsInstanced(MObject->GetType()) == false), TEXT("Scaled objects should not contain instances."));
 		this->bIsConvex = Other.MObject->IsConvex();
 		this->bDoCollide = Other.MObject->GetDoCollide();
 	}
