Fixed non-threadsafe code in connectivity graph generation for geometry collecitons.

diff --git a/Engine/Source/Runtime/Experimental/ChaosSolvers/Private/PhysicsProxy/GeometryCollectionPhysicsProxy.cpp b/Engine/Source/Runtime/Experimental/ChaosSolvers/Private/PhysicsProxy/GeometryCollectionPhysicsProxy.cpp
--- a/Engine/Source/Runtime/Experimental/ChaosSolvers/Private/PhysicsProxy/GeometryCollectionPhysicsProxy.cpp
+++ b/Engine/Source/Runtime/Experimental/ChaosSolvers/Private/PhysicsProxy/GeometryCollectionPhysicsProxy.cpp
@@ -649,14 +649,10 @@ void FGeometryCollectionPhysicsProxy::InitializeDynamicCollection(FGeometryDynam
 
 }
 
-
-
 int32 ReportTooManyChildrenNum = -1;
 FAutoConsoleVariableRef CVarReportTooManyChildrenNum(TEXT("p.ReportTooManyChildrenNum"), ReportTooManyChildrenNum, TEXT("Issue warning if more than this many children exist in a single cluster"));
 
-void FGeometryCollectionPhysicsProxy::InitializeBodiesPT(
-	Chaos::FPBDRigidsSolver* RigidsSolver,
-	Chaos::FPBDRigidsSolver::FParticlesType& Particles)
+void FGeometryCollectionPhysicsProxy::InitializeBodiesPT(Chaos::FPBDRigidsSolver* RigidsSolver, Chaos::FPBDRigidsSolver::FParticlesType& Particles)
 {
 	const FGeometryCollection* RestCollection = Parameters.RestCollection;
 	const FGeometryDynamicCollection& DynamicCollection = PhysicsThreadCollection;
@@ -699,8 +695,7 @@ void FGeometryCollectionPhysicsProxy::InitializeBodiesPT(
 
 		// Add entries into simulation array
 		RigidsSolver->GetEvolution()->ReserveParticles(NumSimulatedParticles);
-		TArray<Chaos::TPBDGeometryCollectionParticleHandle<float, 3>*> Handles =
-			RigidsSolver->GetEvolution()->CreateGeometryCollectionParticles(NumLeafNodes);
+		TArray<Chaos::TPBDGeometryCollectionParticleHandle<float, 3>*> Handles = RigidsSolver->GetEvolution()->CreateGeometryCollectionParticles(NumLeafNodes);
 
 		int32 NextIdx = 0;
 		for (int32 Idx = 0; Idx < SimulatableParticles.Num(); ++Idx)
@@ -782,13 +777,14 @@ void FGeometryCollectionPhysicsProxy::InitializeBodiesPT(
 
 		for (FFieldSystemCommand& Cmd : Parameters.InitializationCommands)
 		{
-			if (Cmd.MetaData.Contains(FFieldSystemMetaData::EMetaType::ECommandData_ProcessingResolution)) 
+			if(Cmd.MetaData.Contains(FFieldSystemMetaData::EMetaType::ECommandData_ProcessingResolution))
+			{
 				Cmd.MetaData.Remove(FFieldSystemMetaData::EMetaType::ECommandData_ProcessingResolution);
-			FFieldSystemMetaDataProcessingResolution* ResolutionData = 
-				new FFieldSystemMetaDataProcessingResolution(EFieldResolutionType::Field_Resolution_Maximum);
-			Cmd.MetaData.Add(
-				FFieldSystemMetaData::EMetaType::ECommandData_ProcessingResolution, 
-				TUniquePtr< FFieldSystemMetaDataProcessingResolution >(ResolutionData));
+			}
+
+			FFieldSystemMetaDataProcessingResolution* ResolutionData = new FFieldSystemMetaDataProcessingResolution(EFieldResolutionType::Field_Resolution_Maximum);
+
+			Cmd.MetaData.Add(FFieldSystemMetaData::EMetaType::ECommandData_ProcessingResolution, TUniquePtr<FFieldSystemMetaDataProcessingResolution>(ResolutionData));
 			Commands.Add(Cmd);
 		}
 		Parameters.InitializationCommands.Empty();
@@ -889,15 +885,9 @@ void FGeometryCollectionPhysicsProxy::InitializeBodiesPT(
 					}
 
 					Chaos::FClusterCreationParameters<float> CreationParameters;
-					CreationParameters.ClusterParticleHandle = 
-						ClusterHandles.Num() ? ClusterHandles[ClusterHandlesIndex++] : nullptr;
+					CreationParameters.ClusterParticleHandle = ClusterHandles.Num() ? ClusterHandles[ClusterHandlesIndex++] : nullptr;
 
-					Chaos::TPBDRigidClusteredParticleHandle<float, 3>* Handle =
-						BuildClusters(
-							TransformGroupIndex, 
-							RigidChildren, 
-							RigidChildrenTransformGroupIndex, 
-							CreationParameters);
+					Chaos::TPBDRigidClusteredParticleHandle<float, 3>* Handle = BuildClusters(TransformGroupIndex, RigidChildren, RigidChildrenTransformGroupIndex, CreationParameters);
 
 					int32 RigidChildrenIdx = 0;
 					for(const int32 ChildTransformIndex : RigidChildrenTransformGroupIndex)
@@ -923,18 +913,17 @@ void FGeometryCollectionPhysicsProxy::InitializeBodiesPT(
 
 			// Set cluster connectivity.  TPBDRigidClustering::CreateClusterParticle() 
 			// will optionally do this, but we switch that functionality off in BuildClusters().
-			ParallelFor(NumTransforms, [&](int32 TransformGroupIndex)
+			for(int32 TransformGroupIndex = 0; TransformGroupIndex < NumTransforms; ++TransformGroupIndex)
 			{
 				const TManagedArray<TSet<int32>>& RestChildren = RestCollection->Children;
 				if (RestChildren[TransformGroupIndex].Num() > 0)
 				{
 					if (SolverClusterHandles[TransformGroupIndex])
 					{
-						GetSolver()->GetEvolution()->GetRigidClustering().GenerateConnectionGraph(
-							SolverClusterHandles[TransformGroupIndex]);
+						GetSolver()->GetEvolution()->GetRigidClustering().GenerateConnectionGraph(SolverClusterHandles[TransformGroupIndex]);
 					}
 				}
-			});
+			}
 		} // end if EnableClustering
  
 
