Chaos : GeometyCollectionTest

Simplicial Collisions to implicits.

diff --git a/Engine/Source/Runtime/Experimental/Chaos/Private/Chaos/CollisionResolution.cpp b/Engine/Source/Runtime/Experimental/Chaos/Private/Chaos/CollisionResolution.cpp
--- a/Engine/Source/Runtime/Experimental/Chaos/Private/Chaos/CollisionResolution.cpp
+++ b/Engine/Source/Runtime/Experimental/Chaos/Private/Chaos/CollisionResolution.cpp
@@ -2030,8 +2030,8 @@ namespace Chaos
 		void ConstructConstraintsImpl(TGeometryParticleHandle<T, d>* Particle0, TGeometryParticleHandle<T, d>* Particle1, const FImplicitObject* Implicit0, const FImplicitObject* Implicit1, const TRigidTransform<T, d>& Transform0, const TRigidTransform<T, d>& Transform1, const T CullDistance, const FCollisionContext& Context, FCollisionConstraintsArray& NewConstraints)
 		{
 			// @todo(chaos): We use GetInnerType here because TriMeshes are left with their "Instanced" wrapper, unlike all other instanced implicits. Should we strip the instance on Tri Mesh too?
-			EImplicitObjectType Implicit0Type = Implicit0 ? GetInnerType(Implicit0->GetType()) : ImplicitObjectType::Unknown;
-			EImplicitObjectType Implicit1Type = Implicit1 ? GetInnerType(Implicit1->GetType()) : ImplicitObjectType::Unknown;
+			EImplicitObjectType Implicit0Type = Implicit0 ? GetInnerType(Implicit0->GetCollisionType()) : ImplicitObjectType::Unknown;
+			EImplicitObjectType Implicit1Type = Implicit1 ? GetInnerType(Implicit1->GetCollisionType()) : ImplicitObjectType::Unknown;
 			bool bIsConvex0 = Implicit0 && Implicit0->IsConvex();
 			bool bIsConvex1 = Implicit1 && Implicit1->IsConvex();
 
diff --git a/Engine/Source/Runtime/Experimental/ChaosSolvers/Private/PhysicsProxy/GeometryCollectionPhysicsProxy.cpp b/Engine/Source/Runtime/Experimental/ChaosSolvers/Private/PhysicsProxy/GeometryCollectionPhysicsProxy.cpp
--- a/Engine/Source/Runtime/Experimental/ChaosSolvers/Private/PhysicsProxy/GeometryCollectionPhysicsProxy.cpp
+++ b/Engine/Source/Runtime/Experimental/ChaosSolvers/Private/PhysicsProxy/GeometryCollectionPhysicsProxy.cpp
@@ -275,6 +275,7 @@ void PopulateSimulatedParticle(
 	if (Implicit)	//todo(ocohen): this is only needed for cases where clusters have no proxy. Kind of gross though, should refactor
 	{
 		TSharedPtr<Chaos::FImplicitObject, ESPMode::ThreadSafe> SharedImplicitTS(Implicit->DeepCopy().Release());
+		FCollisionStructureManager::UpdateImplicitFlags(SharedImplicitTS.Get(), SharedParams.SizeSpecificData[0].CollisionType); // @todo(chaos) : Implicit constructor clobbers CollisionType
 		Handle->SetSharedGeometry(SharedImplicitTS);
 		Handle->SetHasBounds(true);
 		Handle->SetLocalBounds(SharedImplicitTS->BoundingBox());
@@ -2416,16 +2417,30 @@ void IdentifySimulatableElements(Chaos::FErrorReporter& ErrorReporter, FGeometry
 	TArray<bool> HiddenObject;
 	HiddenObject.Init(true, NumTransforms);
 	int32 PrevObject = INDEX_NONE;
+	bool bContiguous = true;
 	for(int32 i = 0; i < Indices.Num(); i++)
 	{
 		if(Visible[i]) // Face index i is visible
 		{
 			const int32 ObjIdx = BoneMap[Indices[i][0]]; // Look up associated bone to the faces X coord.
 			HiddenObject[ObjIdx] = false;
+
+			if (!ensure(ObjIdx >= PrevObject))
+			{
+				bContiguous = false;
+			}
+
 			PrevObject = ObjIdx;
 		}
 	}
 
+	if (!bContiguous)
+	{
+		// What assumptions???  How are we ever going to know if this is still the case?
+		ErrorReporter.ReportError(TEXT("Objects are not contiguous. This breaks assumptions later in the pipeline"));
+		ErrorReporter.HandleLatestError();
+	}
+
 	//For now all simulation data is a non compiled attribute. Not clear what we want for simulated vs kinematic collections
 	TManagedArray<bool>& SimulatableParticles = 
 		GeometryCollection.AddAttribute<bool>(
diff --git a/Engine/Source/Runtime/Experimental/GeometryCollectionSimulationCore/Private/GeometryCollection/GeometryCollectionCollisionStructureManager.cpp b/Engine/Source/Runtime/Experimental/GeometryCollectionSimulationCore/Private/GeometryCollection/GeometryCollectionCollisionStructureManager.cpp
--- a/Engine/Source/Runtime/Experimental/GeometryCollectionSimulationCore/Private/GeometryCollection/GeometryCollectionCollisionStructureManager.cpp
+++ b/Engine/Source/Runtime/Experimental/GeometryCollectionSimulationCore/Private/GeometryCollection/GeometryCollectionCollisionStructureManager.cpp
@@ -173,6 +173,7 @@ void FCollisionStructureManager::UpdateImplicitFlags(
 	FImplicit* Implicit, 
 	const ECollisionTypeEnum CollisionType)
 {
+	Implicit->SetCollsionType(CollisionType==ECollisionTypeEnum::Chaos_Surface_Volumetric?Chaos::ImplicitObjectType::LevelSet:Implicit->GetType());
 	if (Implicit && (CollisionType == ECollisionTypeEnum::Chaos_Surface_Volumetric) && Implicit->GetType() == Chaos::ImplicitObjectType::LevelSet)
 	{
 		Implicit->SetDoCollide(false);
