Chaos : GeometryCollectionTest

- ispc :: Simplicial to Analytic Sphere support.

diff --git a/Engine/Source/Runtime/Experimental/Chaos/Private/Chaos/CollisionResolutionUtil.cpp b/Engine/Source/Runtime/Experimental/Chaos/Private/Chaos/CollisionResolutionUtil.cpp
--- a/Engine/Source/Runtime/Experimental/Chaos/Private/Chaos/CollisionResolutionUtil.cpp
+++ b/Engine/Source/Runtime/Experimental/Chaos/Private/Chaos/CollisionResolutionUtil.cpp
@@ -363,6 +363,40 @@ namespace Chaos
 						}
 					}
 				}
+				else if (Object.GetType() == ImplicitObjectType::Sphere && NumParticles > 0)
+				{
+					const TSphere<float, 3>* Sphere = Object.GetObject<Chaos::TSphere<float, 3>>();
+
+					if (NormalAveraging && UpdateType != ECollisionUpdateType::Any)
+					{
+						ispc::SampleSphereNormalAverageAll(
+							Sphere->GetRadius(),
+							(ispc::FVector&)Sphere->GetCenter(),
+							(ispc::FTransform&)SampleToObjectTM,
+							(ispc::FVector*) & SampleParticles.XArray()[0],
+							Thickness,
+							TotalThickness,
+							(ispc::FVector&)AvgContact.Location,
+							NumParticles);
+					}
+					else
+					{
+						ispc::SampleSphereNoNormalAll(
+							Sphere->GetRadius(),
+							(ispc::FVector&)Sphere->GetCenter(),
+							(ispc::FTransform&)SampleToObjectTM,
+							(ispc::FVector*) & SampleParticles.XArray()[0],
+							DeepestParticle,
+							AvgContact.Phi,
+							NumParticles);
+
+						if (UpdateType == ECollisionUpdateType::Any)
+						{
+							Contact.Phi = AvgContact.Phi;
+							return;
+						}
+					}
+				}
 				else if (Object.GetType() == ImplicitObjectType::Box && NumParticles > 0)
 				{
 					const TBox<float, 3>* Box = Object.GetObject<Chaos::TBox<float, 3>>();
diff --git a/Engine/Source/Runtime/Experimental/Chaos/Private/Chaos/PBDCollisionConstraints.ispc b/Engine/Source/Runtime/Experimental/Chaos/Private/Chaos/PBDCollisionConstraints.ispc
--- a/Engine/Source/Runtime/Experimental/Chaos/Private/Chaos/PBDCollisionConstraints.ispc
+++ b/Engine/Source/Runtime/Experimental/Chaos/Private/Chaos/PBDCollisionConstraints.ispc
@@ -89,7 +89,7 @@ inline static float SquaredSum(const FVector &Vec)
 	return SquareSum;
 }
 
-inline static float NormalSize(const FVector &Normal)
+inline static float VectorSize(const FVector &Normal)
 {
 	return rcp(rsqrt(SquaredSum(Normal)));
 }
@@ -143,7 +143,7 @@ inline static float BoxPhi(const uniform FVector &BoxMin, const uniform FVector
 				Normal.V[i] = 0;
 			}
 		}
-		Phi = NormalSize(Normal);
+		Phi = VectorSize(Normal);
 	}
 
 	return Phi;
@@ -657,5 +657,78 @@ export uniform bool SamplePlaneNoNormalAll(const uniform FVector &ObjectNormal,
 		}
 	}
 
+	return bPhiFound;
+}
+
+export void SampleSphereNormalAverageAll(const uniform float &ObjectRadius,
+									const uniform FVector &ObjectMX,
+									const uniform FTransform &SampleToObjectTransform, 
+									const uniform FVector SampleParticles[],
+									const uniform float Thickness, 
+									uniform float &TotalThickness, 
+									uniform FVector &ConstraintLocation,
+									const uniform unsigned int NumParticles)
+{
+	FVector ConstraintLocationSum = SetVector(0.0f, 0.0f, 0.0f);
+	float ThicknessSum = 0.0f;
+
+	foreach(i = 0 ... NumParticles)
+	{
+		#pragma ignore warning(perf)
+		FVector Particle = SampleParticles[i];
+
+		FVector LocalPoint = TransformPositionNoScale(SampleToObjectTransform, Particle);
+		float LocalPhi = ObjectRadius - VectorSize(LocalPoint - ObjectMX);
+
+		float LocalThickness = LocalPhi - Thickness;
+		if(LocalThickness < -KINDA_SMALL_NUMBER)
+		{
+			ConstraintLocationSum = ConstraintLocationSum + (LocalPoint * LocalThickness);
+			ThicknessSum += LocalThickness;
+		}
+	}
+
+	ConstraintLocation = ConstraintLocation + SetVector(reduce_add(ConstraintLocationSum.V[0]), reduce_add(ConstraintLocationSum.V[1]), reduce_add(ConstraintLocationSum.V[2]));
+	TotalThickness = TotalThickness + reduce_add(ThicknessSum);
+}
+
+export uniform bool SampleSphereNoNormalAll(const uniform float &ObjectRadius,
+										const uniform FVector &ObjectMX,
+										const uniform FTransform &SampleToObjectTransform, 
+										const uniform FVector SampleParticles[],
+										uniform int &DeepestParticle,
+										uniform float &Phi,
+										const uniform unsigned int NumParticles)
+{
+	uniform bool bPhiFound = false;
+
+	foreach(i = 0 ... NumParticles)
+	{
+		if(!bPhiFound)
+		{
+			#pragma ignore warning(perf)
+			FVector Particle = SampleParticles[i];
+
+			FVector LocalPoint = TransformPositionNoScale(SampleToObjectTransform, Particle);
+			float LocalPhi = ObjectRadius - VectorSize(LocalPoint - ObjectMX);
+
+			if(LocalPhi < Phi)
+			{
+				foreach_active(j)
+				{
+					unmasked
+					{
+						if(!bPhiFound)
+						{
+							Phi = extract(LocalPhi, j);
+							DeepestParticle = extract(i, j);
+							bPhiFound = true;
+						}
+					}
+				}
+			}
+		}
+	}
+
 	return bPhiFound;
 }
\ No newline at end of file
