Captured the results of a geometry collection breaking model run and associated any new particles with the proxy their parent is using to ensure correct binding when updating physics data to the game thread.

diff --git a/Engine/Source/Runtime/Experimental/ChaosSolvers/Private/PhysicsProxy/FieldSystemPhysicsProxy.cpp b/Engine/Source/Runtime/Experimental/ChaosSolvers/Private/PhysicsProxy/FieldSystemPhysicsProxy.cpp
--- a/Engine/Source/Runtime/Experimental/ChaosSolvers/Private/PhysicsProxy/FieldSystemPhysicsProxy.cpp
+++ b/Engine/Source/Runtime/Experimental/ChaosSolvers/Private/PhysicsProxy/FieldSystemPhysicsProxy.cpp
@@ -374,7 +374,28 @@ void FFieldSystemPhysicsProxy::FieldParameterUpdateCallback(
 									Map.Add(Handles[Index], StrainSamples[Index]);
 								}
 							}
-							CurrentSolver->GetEvolution()->GetRigidClustering().BreakingModel(&Map);
+
+							// Capture the results from the breaking model to post-process
+							TMap<TPBDRigidClusteredParticleHandle<FReal, 3>*, TSet<TPBDRigidParticleHandle<FReal, 3>*>> BreakResults = CurrentSolver->GetEvolution()->GetRigidClustering().BreakingModel(&Map);
+							
+							// If clusters broke apart then we'll have activated new particles that have no relationship to the proxy that now owns them
+							// Here we attach each new particle to the proxy of the parent particle that owns it.
+							for(const TPair<TPBDRigidClusteredParticleHandle<FReal, 3>*, TSet<TPBDRigidParticleHandle<FReal, 3>*>>& Iter : BreakResults)
+							{
+								const TSet<TPBDRigidParticleHandle<FReal, 3>*>& Activated = Iter.Value;
+
+								for(TPBDRigidParticleHandle<FReal, 3>* Handle : Activated)
+								{
+									if(!CurrentSolver->GetProxy(Handle))
+									{
+										IPhysicsProxyBase* ParentProxy = CurrentSolver->GetProxy(Iter.Key);
+										if(ensure(ParentProxy))
+										{
+											CurrentSolver->AddParticleToProxy(Handle, ParentProxy);
+										}
+									}
+								}
+							}
 						}
 					}
 				}
