Get level set collisions working (though resting contact has problems).  Get deepest penetration during level set constraint update, not just any penetration.  Initialize Depth from FTriangleMeshImiplicitObject::PhiWithNormal().  Implement tests CollisionResolution_AnalyticCubeToAnalyticCube() and CollisionResolution_SimplicialCubeToImplicitCube().

diff --git a/Engine/Source/Runtime/Experimental/Chaos/Private/Chaos/CollisionResolution.cpp b/Engine/Source/Runtime/Experimental/Chaos/Private/Chaos/CollisionResolution.cpp
--- a/Engine/Source/Runtime/Experimental/Chaos/Private/Chaos/CollisionResolution.cpp
+++ b/Engine/Source/Runtime/Experimental/Chaos/Private/Chaos/CollisionResolution.cpp
@@ -1865,7 +1865,7 @@ namespace Chaos
 				Constraint.SetManifold(Implicit0, Implicit1);
 			}
 
-			UpdateLevelsetLevelsetConstraint<ECollisionUpdateType::Any>(CullDistance, Constraint);
+			UpdateLevelsetLevelsetConstraint<ECollisionUpdateType::Deepest>(CullDistance, Constraint);
 
 			NewConstraints.TryAdd(CullDistance, Constraint);
 		}
diff --git a/Engine/Source/Runtime/Experimental/Chaos/Private/Chaos/TriangleMeshImplicitObject.cpp b/Engine/Source/Runtime/Experimental/Chaos/Private/Chaos/TriangleMeshImplicitObject.cpp
--- a/Engine/Source/Runtime/Experimental/Chaos/Private/Chaos/TriangleMeshImplicitObject.cpp
+++ b/Engine/Source/Runtime/Experimental/Chaos/Private/Chaos/TriangleMeshImplicitObject.cpp
@@ -264,7 +264,7 @@ FReal FTriangleMeshImplicitObject::PhiWithNormal(const FVec3& x, FVec3& Normal)
 	TSphere<FReal, 3> TestSphere(x, 0.0f);
 	FRigidTransform3 TestXf(TVector<float, 3>(0.0), TRotation<float, 3>::FromIdentity());
 	FVec3 TestLocation = x;
-	FReal Depth;
+	FReal Depth = TNumericLimits<FReal>::Max();
 	GJKContactPointImp(TestSphere, TestXf, 0.0f, TestLocation, Normal, Depth);
 	return Depth;
 }
