Chaos : Collisions

- MassToLocal translation of vertices.
- MassToLocal rotation is set to identity.

diff --git a/Engine/Source/Runtime/Experimental/ChaosSolvers/Private/PhysicsProxy/GeometryCollectionPhysicsProxy.cpp b/Engine/Source/Runtime/Experimental/ChaosSolvers/Private/PhysicsProxy/GeometryCollectionPhysicsProxy.cpp
--- a/Engine/Source/Runtime/Experimental/ChaosSolvers/Private/PhysicsProxy/GeometryCollectionPhysicsProxy.cpp
+++ b/Engine/Source/Runtime/Experimental/ChaosSolvers/Private/PhysicsProxy/GeometryCollectionPhysicsProxy.cpp
@@ -234,8 +234,8 @@ void PopulateSimulatedParticle(
 	Handle->SetP(Handle->X());
 	Handle->SetQ(Handle->R());
 	Handle->SetIsland(INDEX_NONE);
-	Handle->SetCenterOfMass(MassOffset.GetTranslation());
-	Handle->SetRotationOfMass(MassOffset.GetRotation());
+	Handle->SetCenterOfMass(FVector::ZeroVector);
+	Handle->SetRotationOfMass(FQuat::Identity);
 
 	//
 	// Setup Mass
@@ -737,7 +737,8 @@ void FGeometryCollectionPhysicsProxy::InitializeBodiesPT(Chaos::FPBDRigidsSolver
 			if (FClusterHandle* Handle = SolverParticleHandles[TransformGroupIndex])
 			{
 				// Mass space -> Composed parent space -> world
-				const FTransform WorldTransform = Transform[TransformGroupIndex] * Parameters.WorldTransform;
+				const FTransform WorldTransform = 
+					MassToLocal[TransformGroupIndex] * Transform[TransformGroupIndex] * Parameters.WorldTransform;
 
 				PopulateSimulatedParticle(
 					Handle,
@@ -1505,7 +1506,8 @@ void FGeometryCollectionPhysicsProxy::BufferPhysicsResults()
 				FTransform& ParticleToWorld = TargetResults.ParticleToWorldTransforms[TransformIndex];
 				ParticleToWorld = FTransform(Handle->R(), Handle->X());
 
-				TargetResults.Transforms[TransformIndex] = ParticleToWorld.GetRelativeTransform(ActorToWorld);
+				TargetResults.Transforms[TransformIndex] = 
+					PTDynamicCollection.MassToLocal[TransformIndex].GetRelativeTransformReverse(ParticleToWorld).GetRelativeTransform(ActorToWorld);
 				TargetResults.Transforms[TransformIndex].NormalizeRotation();
 
 				// If the parent of this NON DISABLED body is set to anything other than INDEX_NONE,
@@ -1564,7 +1566,8 @@ void FGeometryCollectionPhysicsProxy::BufferPhysicsResults()
 								Handle->SetX(ClusterChildToWorld.GetTranslation());
 								Handle->SetR(ClusterChildToWorld.GetRotation());
 							}
-							TargetResults.Transforms[TransformIndex] = ClusterChildToWorld.GetRelativeTransform(ActorToWorld);
+							TargetResults.Transforms[TransformIndex] = 
+								PTDynamicCollection.MassToLocal[TransformIndex].GetRelativeTransformReverse(ClusterChildToWorld).GetRelativeTransform(ActorToWorld);
 							TargetResults.Transforms[TransformIndex].NormalizeRotation();
 						}
 					}
@@ -1694,6 +1697,14 @@ void FGeometryCollectionPhysicsProxy::UpdateGeometryCollection( const FGeometryC
 	TManagedArray<int32>& SimulationType = Collection->SimulationType;
 	TManagedArray<int32>& DynamicState = Collection->GetAttribute<int32>("DynamicState", FGeometryCollection::TransformGroup);
 
+	// Particles X and R are aligned with center of mass and inertia principal axes.
+	// Renderer doesn't know about this and simply does ActorToWorld * GeomToActor * LocalSpaceVerts
+	// In proper math multiplication order:
+	//		ParticleToWorld = ActorToWorld * GeomToActor * MassToLocal
+	//		GeomToWorld = ActorToWorld * GeomToActor
+	//		=> GeomToWorld = ParticleToWorld * MassToLocal.Inv()
+	//		=> GeomToActor = ActorToWorld.Inv() * ParticleToWorld * MassToLocal.Inv()
+
 	const FTransform& ActorToWorld = Parameters.WorldTransform;
 	const int32 TransformSize = Collection->NumElements(FGeometryCollection::TransformGroup);
 	for (int32 TransformGroupIndex = 0; TransformGroupIndex < TransformSize; ++TransformGroupIndex)
@@ -1724,6 +1735,15 @@ void FGeometryCollectionPhysicsProxy::UpdateGeometryCollection( const FGeometryC
 		if(!GCResults.DisabledStates[TransformGroupIndex])
 		{
 #endif
+	// Update the transform of the active body. The active body can be either a single rigid
+	// or a collection of rigidly attached geometries (Clustering). The cluster is represented as a
+	// single transform in the GeometryCollection, and all children are stored in the local space
+	// of the parent cluster.
+	// ... When setting cluster transforms it is expected that the MassToLocal is identity.
+	//     Cluster initialization will set the vertices in the MassSpace of the rigid body.
+	// ... When setting individual rigid bodies that are not clustered, the MassToLocal will be 
+	//     non-Identity, and will reflect the difference between the geometric center of the geometry
+	//     and that corresponding rigid bodies center of mass. 
 
 			Transform[TransformGroupIndex] = GCResults.Transforms[TransformGroupIndex];
 
@@ -1929,6 +1949,17 @@ void FGeometryCollectionPhysicsProxy::InitializeSharedCollisionStructures(
 						InertiaComputationNeeded[GeometryIndex] = true;
 						CollectionMassToLocal[TransformGroupIndex] = FTransform(FQuat::Identity, MassProperties.CenterOfMass);
 					}
+
+					FVector MassTranslation = CollectionMassToLocal[TransformGroupIndex].GetTranslation();
+					if (!FMath::IsNearlyZero(MassTranslation.SizeSquared()))
+					{
+						const int32 IdxStart = VertexStart[GeometryIndex];
+						const int32 IdxEnd = IdxStart + VertexCount[GeometryIndex];
+						for (int32 Idx = IdxStart; Idx < IdxEnd; ++Idx)
+						{
+							MassSpaceParticles.X(Idx) -= MassTranslation;
+						}
+					}
 				}
 			}
 
@@ -2005,29 +2036,23 @@ void FGeometryCollectionPhysicsProxy::InitializeSharedCollisionStructures(
 
 			if (InertiaComputationNeeded[GeometryIndex])
 			{
-				if (!FMath::IsNearlyZero(MassProperties.CenterOfMass.SizeSquared()))
-				{
-					const int32 IdxStart = VertexStart[GeometryIndex];
-					const int32 IdxEnd = IdxStart + VertexCount[GeometryIndex];
-					for (int32 Idx = IdxStart; Idx < IdxEnd; ++Idx)
-					{
-						MassSpaceParticles.X(Idx) -= MassProperties.CenterOfMass;
-					}
-				}
-
 				CalculateInertiaAndRotationOfMass(MassSpaceParticles, TriMesh->GetSurfaceElements(), Density_i, MassProperties.CenterOfMass, MassProperties.InertiaTensor, MassProperties.RotationOfMass);
-				CollectionMassToLocal[TransformGroupIndex] = FTransform(MassProperties.RotationOfMass, MassProperties.CenterOfMass);
 				CollectionInertiaTensor[TransformGroupIndex] = TVector<float, 3>(MassProperties.InertiaTensor.M[0][0], MassProperties.InertiaTensor.M[1][1], MassProperties.InertiaTensor.M[2][2]);
+#if false
+				CollectionMassToLocal[TransformGroupIndex] = FTransform(MassProperties.RotationOfMass, MassProperties.CenterOfMass);
 
-				if (!FMath::IsNearlyZero(MassProperties.CenterOfMass.SizeSquared()))
+
+				if (!MassProperties.RotationOfMass.Equals(FQuat::Identity))
 				{
+					FTransform InverseMassRotation = FTransform(MassProperties.RotationOfMass.Inverse());
 					const int32 IdxStart = VertexStart[GeometryIndex];
 					const int32 IdxEnd = IdxStart + VertexCount[GeometryIndex];
 					for (int32 Idx = IdxStart; Idx < IdxEnd; ++Idx)
 					{
-						MassSpaceParticles.X(Idx) += MassProperties.CenterOfMass;
+						MassSpaceParticles.X(Idx) = InverseMassRotation.TransformPosition(MassSpaceParticles.X(Idx));
 					}
 				}
+#endif
 			}
 			else
 			{
diff --git a/Engine/Source/Runtime/Experimental/ChaosSolvers/Public/PhysicsProxy/GeometryCollectionPhysicsProxy.h b/Engine/Source/Runtime/Experimental/ChaosSolvers/Public/PhysicsProxy/GeometryCollectionPhysicsProxy.h
--- a/Engine/Source/Runtime/Experimental/ChaosSolvers/Public/PhysicsProxy/GeometryCollectionPhysicsProxy.h
+++ b/Engine/Source/Runtime/Experimental/ChaosSolvers/Public/PhysicsProxy/GeometryCollectionPhysicsProxy.h
@@ -99,6 +99,9 @@ public:
 		Chaos::FPBDRigidsSolver* RigidsSolver,
 		Chaos::FPBDRigidsSolver::FParticlesType& Particles);
 
+	/** */
+	static void InitializeDynamicCollection(FGeometryDynamicCollection& DynamicCollection, const FGeometryCollection& RestCollection, const FSimulationParameters& Params);
+
 	/** */
 	bool IsSimulating() const { return Parameters.Simulating; }
 
@@ -220,7 +223,6 @@ protected:
 	 */
 	int32 CalculateHierarchyLevel(const FGeometryDynamicCollection& GeometryCollection, int32 TransformIndex) const;
 
-	static void InitializeDynamicCollection(FGeometryDynamicCollection& DynamicCollection, const FGeometryCollection& RestCollection, const FSimulationParameters& Params);
 	void InitializeRemoveOnFracture(FParticlesType& Particles, const TManagedArray<int32>& DynamicState);
 	void ProcessCommands(FParticlesType& Particles, const float Time);
 
