Rotating geometry collection DDC key missed by 12386066

diff --git a/Engine/Source/Runtime/Experimental/GeometryCollectionEngine/Private/GeometryCollection/DerivedDataGeometryCollectionCooker.cpp b/Engine/Source/Runtime/Experimental/GeometryCollectionEngine/Private/GeometryCollection/DerivedDataGeometryCollectionCooker.cpp
--- a/Engine/Source/Runtime/Experimental/GeometryCollectionEngine/Private/GeometryCollection/DerivedDataGeometryCollectionCooker.cpp
+++ b/Engine/Source/Runtime/Experimental/GeometryCollectionEngine/Private/GeometryCollection/DerivedDataGeometryCollectionCooker.cpp
@@ -59,7 +59,7 @@ const TCHAR* FDerivedDataGeometryCollectionCooker::GetVersionString() const
 		return OverrideVersion;	//force load old ddc if found. Not recommended
 	}
 
-	return TEXT("AE5BE323BD8A445E939DF316EB074D37");
+	return TEXT("655D3B4051AE4A1E99135F2417195C8C");
 }
 
 FString FDerivedDataGeometryCollectionCooker::GetPluginSpecificCacheKeySuffix() const
