Rewinding data creates a snapshot at head

diff --git a/Engine/Source/Programs/HeadlessChaos/Private/HeadlessChaosTestRewind.cpp b/Engine/Source/Programs/HeadlessChaos/Private/HeadlessChaosTestRewind.cpp
--- a/Engine/Source/Programs/HeadlessChaos/Private/HeadlessChaosTestRewind.cpp
+++ b/Engine/Source/Programs/HeadlessChaos/Private/HeadlessChaosTestRewind.cpp
@@ -387,13 +387,24 @@ namespace ChaosTest {
 			V.Add(Particle->V());
 			TickSolverHelper(Module,Solver);
 		}
+		X.Add(Particle->X());
+		V.Add(Particle->V());
 
 		FRewindData* RewindData = Solver->GetRewindData();
 		RewindData->RewindToFrame(3);
-		
+
 		EXPECT_EQ(Particle->X()[2],X[3][2]);
 		EXPECT_EQ(Particle->V()[2],V[3][2]);
 
+		//make sure recorded data is still valid even at head
+		for(int Step = 0; Step < 11; ++Step)
+		{
+			const FGeometryParticleState State = RewindData->GetStateAtFrame(*Particle,Step);
+			EXPECT_EQ(State.X()[2],X[Step][2]);
+			EXPECT_EQ(State.V()[2],V[Step][2]);
+		}
+		
+
 		// Throw out the proxy
 		Solver->UnregisterObject(Particle.Get());
 
diff --git a/Engine/Source/Runtime/Engine/Private/PhysicsEngine/PhysUtils.cpp b/Engine/Source/Runtime/Engine/Private/PhysicsEngine/PhysUtils.cpp
--- a/Engine/Source/Runtime/Engine/Private/PhysicsEngine/PhysUtils.cpp
+++ b/Engine/Source/Runtime/Engine/Private/PhysicsEngine/PhysUtils.cpp
@@ -20,6 +20,7 @@
 #include "Chaos/PBDRigidsEvolutionGBF.h"
 #include "Chaos/ChaosArchive.h"
 #include "Chaos/TrackedGeometryManager.h"
+#include "RewindData.h"
 
 /** Returns false if ModelToHulls operation should halt because of vertex count overflow. */
 static bool AddConvexPrim(FKAggregateGeom* OutGeom, TArray<FPlane> &Planes, UModel* InModel)
@@ -482,6 +483,12 @@ bool FPhysicsInterface::ExecPhysCommands(const TCHAR* Cmd, FOutputDevice* Output
 		Chaos::FTrackedGeometryManager::Get().DumpMemoryUsage(OutputDevice);
 		return true;
 	}
+
+	if(FParse::Command(&Cmd,TEXT("ChaosRewind")))
+	{
+		InWorld->GetPhysicsScene()->GetSolver()->GetRewindData()->RewindToFrame(0);
+		return true;
+	}
 #if CHAOS_MEMORY_TRACKING
 	if (FParse::Command(&Cmd, TEXT("ChaosMemoryDistribution")))
 	{
diff --git a/Engine/Source/Runtime/Experimental/Chaos/Public/Chaos/ParticleHandle.h b/Engine/Source/Runtime/Experimental/Chaos/Public/Chaos/ParticleHandle.h
--- a/Engine/Source/Runtime/Experimental/Chaos/Public/Chaos/ParticleHandle.h
+++ b/Engine/Source/Runtime/Experimental/Chaos/Public/Chaos/ParticleHandle.h
@@ -1437,6 +1437,8 @@ public:
 		check(false);
 	}
 
+	const TSharedPtr<FImplicitObject,ESPMode::ThreadSafe>& SharedGeometryLowLevel() const { return MNonFrequentData.Read().Geometry; }
+
 	void* UserData() const { return MNonFrequentData.Read().UserData; }
 	void SetUserData(void* InUserData)
 	{
diff --git a/Engine/Source/Runtime/Experimental/ChaosSolvers/Public/RewindData.h b/Engine/Source/Runtime/Experimental/ChaosSolvers/Public/RewindData.h
--- a/Engine/Source/Runtime/Experimental/ChaosSolvers/Public/RewindData.h
+++ b/Engine/Source/Runtime/Experimental/ChaosSolvers/Public/RewindData.h
@@ -42,14 +42,20 @@ public:
 		}
 	}
 
+	template <typename LambdaSet>
+	void SyncRemoteDataForced(FDirtyPropertiesManager& InManager,int32 InIdx,const LambdaSet& SetFunc)
+	{
+		Manager = &InManager;
+		Idx = InIdx;
+		SetFunc(Manager->GetParticlePool<T,PropName>().GetElement(Idx));
+	}
+
 	template <typename LambdaSet>
 	void SyncRemoteData(FDirtyPropertiesManager& InManager,int32 InIdx, const FParticleDirtyData& DirtyData, const LambdaSet& SetFunc)
 	{
 		if(DirtyData.IsDirty(ParticlePropToFlag(PropName)))
 		{
-			Manager = &InManager;
-			Idx = InIdx;
-			SetFunc(Manager->GetParticlePool<T,PropName>().GetElement(Idx));
+			SyncRemoteDataForced(InManager,InIdx,SetFunc);
 		}
 	}
 
@@ -236,6 +242,55 @@ public:
 		});
 	}
 
+	void SyncIfDirty(FDirtyPropertiesManager& Manager,int32 Idx,const TGeometryParticle<FReal,3>& InParticle, const FGeometryParticleStateBase& RewindState)
+	{
+		ensure(IsInGameThread());
+		const auto Particle = &InParticle;
+
+		if(RewindState.ParticlePositionRotation.IsSet())
+		{
+			ParticlePositionRotation.SyncRemoteDataForced(Manager,Idx,[Particle](FParticlePositionRotation& Data)
+			{
+				Data.X = Particle->X();
+				Data.R = Particle->R();
+			});
+		}
+		
+		if(const auto Kinematic = Particle->CastToKinematicParticle())
+		{
+			if(RewindState.Velocities.IsSet())
+			{
+				Velocities.SyncRemoteDataForced(Manager,Idx,[Kinematic](auto& Data)
+				{
+					Data.V = Kinematic->V();
+					Data.W = Kinematic->W();
+				});
+			}
+		}
+		
+		if(RewindState.NonFrequentData.IsSet())
+		{
+			NonFrequentData.SyncRemoteDataForced(Manager,Idx,[Particle](FParticleNonFrequentData& Data)
+			{
+				Data.Geometry = Particle->SharedGeometryLowLevel();
+				Data.UserData = Particle->UserData();
+
+				//note: this data is keyed based on unique idx so it's not really possible to change this
+				//but we save it anyway since it's part of a big struct
+				Data.UniqueIdx = Particle->UniqueIdx();
+	#if CHAOS_CHECKED
+				Data.DebugName = Particle->DebugName();
+	#endif
+
+				if(auto Rigid = Particle->CastToRigidParticle())
+				{
+					Data.LinearEtherDrag = Rigid->LinearEtherDrag();
+					Data.AngularEtherDrag = Rigid->AngularEtherDrag();
+				}
+			});
+		}
+	}
+
 	bool CoalesceState(const FGeometryParticleStateBase& LatestState)
 	{
 		bool bCoalesced = false;
@@ -336,16 +391,29 @@ public:
 	void RewindToFrame(int32 Frame)
 	{
 		ensure(IsInGameThread());
+		PrepareFrame(AllDirtyParticles.Num());
+		FDirtyPropertiesManager& DestManager = *Managers.Last().Manager;
+
+		//todo: avoid making head copy if not needed
 		//todo: parallel for
-		for(const FDirtyParticleInfo& Info : AllDirtyParticles)
+		int32 DataIdx = 0;
+		for(FDirtyParticleInfo& DirtyParticleInfo : AllDirtyParticles)
 		{
-			if(const FGeometryParticleStateBase* State = GetStateAtFrameImp(*Info.Particle, Frame))
+			//GetStateAtFrameImp returns a pointer from the TArray that holds state data
+			//But it's possible that we'll need to save state from head, which would grow that TArray
+			//So preallocate just in case
+			FParticleRewindInfo& RewindInfo = ParticleToRewindInfo.FindChecked(DirtyParticleInfo.CachedUniqueIdx);
+			FGeometryParticleStateBase& LatestState = RewindInfo.AddFrame(CurFrame);
+			
+			if(const FGeometryParticleStateBase* RewindState = GetStateAtFrameImp(*DirtyParticleInfo.Particle, Frame))
 			{
-				State->SyncToParticle(*Info.Particle);
+				LatestState.SyncIfDirty(DestManager,DataIdx++,*DirtyParticleInfo.Particle,*RewindState);
+				CoalesceBack(RewindInfo.Frames);
+
+				RewindState->SyncToParticle(*DirtyParticleInfo.Particle);
 			}
 		}
-		
-		Reset();
+
 		CurFrame = Frame;
 	}
 
@@ -425,17 +493,7 @@ public:
 			{
 				FGeometryParticleStateBase& LatestState = Info.AddFrame(CurFrame-1);
 				LatestState.SyncPrevFrame(DestManager,DataIdx,Dirty);
-
-				//for frames further back a simply copy is enough
-				for(int32 FrameIdx = Info.Frames.Num() - 2; FrameIdx >= 0; --FrameIdx)
-				{
-					FFrameInfo& Frame = Info.Frames[FrameIdx];
-					if(Frame.State.CoalesceState(LatestState) == false)
-					{
-						//nothing to coalesce so no need to check earlier frames
-						break;
-					}
-				}
+				CoalesceBack(Info.Frames);
 			}
 
 			//If dynamics are dirty we must record them immediately because the sim resets them to 0
