Move managers to circular buffer. Also made it so you can sync from past to present one frame at a time

diff --git a/Engine/Source/Programs/HeadlessChaos/Private/HeadlessChaosTestRewind.cpp b/Engine/Source/Programs/HeadlessChaos/Private/HeadlessChaosTestRewind.cpp
--- a/Engine/Source/Programs/HeadlessChaos/Private/HeadlessChaosTestRewind.cpp
+++ b/Engine/Source/Programs/HeadlessChaos/Private/HeadlessChaosTestRewind.cpp
@@ -391,11 +391,8 @@ namespace ChaosTest {
 		V.Add(Particle->V());
 
 		FRewindData* RewindData = Solver->GetRewindData();
-		RewindData->RewindToFrame(3);
-
-		EXPECT_EQ(Particle->X()[2],X[3][2]);
-		EXPECT_EQ(Particle->V()[2],V[3][2]);
-
+		EXPECT_TRUE(RewindData->RewindToFrame(0));
+		
 		//make sure recorded data is still valid even at head
 		for(int Step = 0; Step < 11; ++Step)
 		{
@@ -403,7 +400,17 @@ namespace ChaosTest {
 			EXPECT_EQ(State.X()[2],X[Step][2]);
 			EXPECT_EQ(State.V()[2],V[Step][2]);
 		}
-		
+
+		//rewind to each frame and make sure data is recorded
+		for(int Step = 0; Step < 10; ++Step)
+		{
+			EXPECT_TRUE(RewindData->RewindToFrame(Step));
+			EXPECT_EQ(Particle->X()[2],X[Step][2]);
+			EXPECT_EQ(Particle->V()[2],V[Step][2]);
+		}
+
+		//can't rewind earlier than latest rewind
+		EXPECT_FALSE(RewindData->RewindToFrame(5));
 
 		// Throw out the proxy
 		Solver->UnregisterObject(Particle.Get());
@@ -461,4 +468,68 @@ namespace ChaosTest {
 
 		Module->DestroySolver(Solver);
 	}
+
+	GTEST_TEST(RewindTest,BufferLimit)
+	{
+		auto Sphere = TSharedPtr<FImplicitObject,ESPMode::ThreadSafe>(new TSphere<float,3>(TVector<float,3>(0),10));
+
+		FChaosSolversModule* Module = FChaosSolversModule::GetModule();
+		Module->ChangeThreadingMode(EChaosThreadingMode::SingleThread);
+
+		// Make a solver
+		FPhysicsSolver* Solver = Module->CreateSolver(nullptr,ESolverFlags::Standalone);
+		Solver->SetEnabled(true);
+
+		Solver->EnableRewindCapture(5);
+
+
+		// Make particles
+		auto Particle = TPBDRigidParticle<float,3>::CreateParticle();
+
+		Particle->SetGeometry(Sphere);
+		Solver->RegisterObject(Particle.Get());
+		Particle->SetGravityEnabled(true);
+		Particle->SetX(FVec3(0,0,100));
+
+		TArray<FVec3> X;
+		TArray<FVec3> V;
+
+		const int32 NumSteps = 20;
+		for(int Step = 0; Step < NumSteps; ++Step)
+		{
+			//teleport from GT
+			if(Step == 15)
+			{
+				Particle->SetX(FVec3(0,0,10));
+				Particle->SetV(FVec3(0,0,1));
+			}
+
+			X.Add(Particle->X());
+			V.Add(Particle->V());
+			TickSolverHelper(Module,Solver);
+		}
+		X.Add(Particle->X());
+		V.Add(Particle->V());
+
+		FRewindData* RewindData = Solver->GetRewindData();
+		const int32 LastValidStep = NumSteps - 1;
+		const int32 FirstValid = NumSteps - RewindData->Capacity() + 1;	//we lose 1 step because we have to save head
+		for(int Step = 0; Step < FirstValid; ++Step)
+		{
+			//can't go back that far
+			EXPECT_FALSE(RewindData->RewindToFrame(Step));
+		}
+
+		for(int Step = FirstValid; Step <= LastValidStep; ++Step)
+		{
+			EXPECT_TRUE(RewindData->RewindToFrame(Step));
+			EXPECT_EQ(Particle->X()[2],X[Step][2]);
+			EXPECT_EQ(Particle->V()[2],V[Step][2]);
+		}
+
+		// Throw out the proxy
+		Solver->UnregisterObject(Particle.Get());
+
+		Module->DestroySolver(Solver);
+	}
 }
diff --git a/Engine/Source/Runtime/Experimental/ChaosSolvers/Private/PBDRigidsSolver.cpp b/Engine/Source/Runtime/Experimental/ChaosSolvers/Private/PBDRigidsSolver.cpp
--- a/Engine/Source/Runtime/Experimental/ChaosSolvers/Private/PBDRigidsSolver.cpp
+++ b/Engine/Source/Runtime/Experimental/ChaosSolvers/Private/PBDRigidsSolver.cpp
@@ -735,12 +735,6 @@ namespace Chaos
 		QUICK_SCOPE_CYCLE_COUNTER(STAT_PushPhysicsState);
 		FDirtySet* DirtyProxiesData = DirtyProxiesDataBuffer.AccessProducerBuffer();
 
-		if(DirtyProxiesData->NumDirtyProxies() == 0)
-		{
-			//is this early out really necessary?
-			return;
-		}
-
 		FDirtyPropertiesManager* Manager = DirtyPropertiesManager->AccessProducerBuffer();
 		Manager->SetNumParticles(DirtyProxiesData->NumDirtyProxies());
 		Manager->SetNumShapes(DirtyProxiesData->NumDirtyShapes());
diff --git a/Engine/Source/Runtime/Experimental/ChaosSolvers/Public/RewindData.h b/Engine/Source/Runtime/Experimental/ChaosSolvers/Public/RewindData.h
--- a/Engine/Source/Runtime/Experimental/ChaosSolvers/Public/RewindData.h
+++ b/Engine/Source/Runtime/Experimental/ChaosSolvers/Public/RewindData.h
@@ -5,6 +5,7 @@
 #include "PhysicsProxy/SingleParticlePhysicsProxy.h"
 #include "PhysicsProxy/SingleParticlePhysicsProxyFwd.h"
 #include "Chaos/Framework/PhysicsSolverBase.h"
+#include "Containers/CircularBuffer.h"
 
 namespace Chaos
 {
@@ -384,37 +385,72 @@ class FRewindData
 {
 public:
 	FRewindData(int32 NumFrames)
-	: CurFrame(0)
+	: Managers(NumFrames+1)	//give 1 extra for saving at head
+	, CurFrame(0)
+	, EarliestFrame(0)
+	, bNeedsSave(false)
 	{
 	}
 
-	void RewindToFrame(int32 Frame)
+	uint32 Capacity() const { return Managers.Capacity(); }
+
+	bool RewindToFrame(uint32 Frame)
 	{
 		ensure(IsInGameThread());
-		PrepareFrame(AllDirtyParticles.Num());
-		FDirtyPropertiesManager& DestManager = *Managers.Last().Manager;
 
-		//todo: avoid making head copy if not needed
+		//Can't go too far back, also we need 1 entry for saving head
+		if(Frame < EarliestFrame)
+		{
+			return false;
+		}
+
+		//If we need to save and we are right on the edge of the buffer, we can't go back to earliest frame
+		if(Frame == EarliestFrame && bNeedsSave && &Managers[CurFrame] == &Managers[EarliestFrame])
+		{
+			return false;
+		}
+		
+		FDirtyPropertiesManager* DestManager = nullptr;
+		if(bNeedsSave)
+		{
+			PrepareFrame(AllDirtyParticles.Num());
+			DestManager = Managers[CurFrame].Manager.Get();
+		}
+
 		//todo: parallel for
 		int32 DataIdx = 0;
 		for(FDirtyParticleInfo& DirtyParticleInfo : AllDirtyParticles)
 		{
-			//GetStateAtFrameImp returns a pointer from the TArray that holds state data
-			//But it's possible that we'll need to save state from head, which would grow that TArray
-			//So preallocate just in case
-			FParticleRewindInfo& RewindInfo = ParticleToRewindInfo.FindChecked(DirtyParticleInfo.CachedUniqueIdx);
-			FGeometryParticleStateBase& LatestState = RewindInfo.AddFrame(CurFrame);
-			
-			if(const FGeometryParticleStateBase* RewindState = GetStateAtFrameImp(*DirtyParticleInfo.Particle, Frame))
+			if(bNeedsSave)
 			{
-				LatestState.SyncIfDirty(DestManager,DataIdx++,*DirtyParticleInfo.Particle,*RewindState);
-				CoalesceBack(RewindInfo.Frames);
+				//GetStateAtFrameImp returns a pointer from the TArray that holds state data
+				//But it's possible that we'll need to save state from head, which would grow that TArray
+				//So preallocate just in case
+				FParticleRewindInfo& RewindInfo = ParticleToRewindInfo.FindChecked(DirtyParticleInfo.CachedUniqueIdx);
+				FGeometryParticleStateBase& LatestState = RewindInfo.AddFrame(CurFrame);
+			
+				if(const FGeometryParticleStateBase* RewindState = GetStateAtFrameImp(*DirtyParticleInfo.Particle, Frame))
+				{
+					LatestState.SyncIfDirty(*DestManager,DataIdx++,*DirtyParticleInfo.Particle,*RewindState);
+					CoalesceBack(RewindInfo.Frames);
 
-				RewindState->SyncToParticle(*DirtyParticleInfo.Particle);
+					RewindState->SyncToParticle(*DirtyParticleInfo.Particle);
+				}
+			}
+			else
+			{
+				if(const FGeometryParticleStateBase* RewindState = GetStateAtFrameImp(*DirtyParticleInfo.Particle,Frame))
+				{
+					RewindState->SyncToParticle(*DirtyParticleInfo.Particle);
+				}
 			}
 		}
 
 		CurFrame = Frame;
+		bNeedsSave = false;
+		EarliestFrame = CurFrame;	//can't rewind before this point. This simplifies saving the state at head
+
+		return true;
 	}
 
 	void RemoveParticle(const TGeometryParticleHandle<FReal,3>& Particle)
@@ -447,30 +483,31 @@ public:
 	void AdvanceFrame()
 	{
 		++CurFrame;
+		if(CurFrame > Managers.Capacity())
+		{
+			++EarliestFrame;
+		}
 	}
 
 	void PrepareFrame(int32 NumDirtyParticles)
 	{
-		Managers.Emplace(FFrameManagerInfo{MakeUnique<FDirtyPropertiesManager>(), CurFrame});
-		Managers.Last().Manager->SetNumParticles(NumDirtyParticles);
+		PrepareManagerForFrame(CurFrame,NumDirtyParticles);
 	}
 
 	int32 PrepareFrameForPTDirty(int32 NumActiveParticles)
 	{
+		bNeedsSave = true;
+
 		//If manager already exists for previous frame, use it
-		if(Managers.Num())
+		const int32 PrevFrame = CurFrame - 1;
+		FFrameManagerInfo& Info = Managers[PrevFrame];
+		ensure(Info.Manager && Info.FrameCreatedFor == (PrevFrame));
+
 		{
-			if(Managers.Last().FrameCreatedFor == (CurFrame-1))
-			{
-				FDirtyPropertiesManager& Manager = *Managers.Last().Manager;
-				const int32 NumDirtyAlready = Manager.GetNumParticles();
-				Manager.SetNumParticles(NumDirtyAlready + NumActiveParticles);
-				return NumDirtyAlready;
-			}
+			const int32 NumDirtyAlready = Info.Manager->GetNumParticles();
+			Info.Manager->SetNumParticles(NumDirtyAlready + NumActiveParticles);
+			return NumDirtyAlready;
 		}
-
-		//No manager for previous frame so create a new one
-		PrepareFrame(NumActiveParticles);
 		return 0;
 	}
 
@@ -481,12 +518,15 @@ public:
 		//This is possible because most properties are const on the physics thread
 		//For sim-writable properties (forces, position, velocities, etc...) we must immediately write the data because there is no way to know what the previous data was next frame
 		//Some sim-writable properties can change without the GT knowing about it, see PushPTDirtyData
+
+		//User called PrepareManagerForFrame for this frame so use it
+		FDirtyPropertiesManager& DestManager = *Managers[CurFrame].Manager;
+		bNeedsSave = true;
 		
-		auto ProcessProxy = [this,&SrcManager, DataIdx, Dirty](const auto Proxy)
+		auto ProcessProxy = [this,&SrcManager, DataIdx, Dirty, &DestManager](const auto Proxy)
 		{
 			const auto PTParticle = Proxy->GetHandle();
 			FParticleRewindInfo& Info = FindOrAddParticle(Proxy->GetParticle(),PTParticle->UniqueIdx());
-			FDirtyPropertiesManager& DestManager = *Managers.Last().Manager;
 
 			//Most properties are always a frame behind
 			if(Proxy->IsInitialized())	//Frame delay so proxy must be initialized
@@ -535,11 +575,13 @@ public:
 		if(SimWritablePropsMayChange(Rigid))
 		{
 			FParticleRewindInfo& Info = FindOrAddParticle(Rigid.GTGeometryParticle(), Rigid.UniqueIdx());
-			FDirtyPropertiesManager& Manager = *Managers.Last().Manager;
+			
+			//User called PrepareManagerForFrame (or PrepareFrameForPTDirty) for the previous frame, so use it
+			FDirtyPropertiesManager& DestManager = *Managers[CurFrame-1].Manager;
 
 			//sim-writable properties changed at head, so we must write down what they were
 			FGeometryParticleStateBase& LatestState = Info.AddFrame(CurFrame-1);
-			LatestState.SyncSimWritablePropsFromSim(Manager,DataIdx,Rigid);
+			LatestState.SyncSimWritablePropsFromSim(DestManager,DataIdx,Rigid);
 
 			//update any previous frames that were pointing at head
 			CoalesceBack(Info.Frames);
@@ -585,6 +627,18 @@ private:
 		}
 	}
 
+	void PrepareManagerForFrame(int32 Frame, int32 NumDirtyParticles)
+	{
+		FFrameManagerInfo& Info = Managers[Frame];
+		if(Info.Manager == nullptr)
+		{
+			Info.Manager = MakeUnique<FDirtyPropertiesManager>();
+		}
+
+		Info.Manager->SetNumParticles(NumDirtyParticles);
+		Info.FrameCreatedFor = Frame;
+	}
+
 	struct FFrameManagerInfo
 	{
 		TUniquePtr<FDirtyPropertiesManager> Manager;
@@ -595,13 +649,6 @@ private:
 		int32 FrameCreatedFor;
 	};
 
-	void Reset()
-	{
-		ParticleToRewindInfo.Reset();
-		Managers.Reset();
-		AllDirtyParticles.Reset();
-	}
-
 	FParticleRewindInfo& FindOrAddParticle(TGeometryParticle<FReal,3>* UnsafeGTParticle, FUniqueIdx UniqueIdx)
 	{
 		if(FParticleRewindInfo* Info = ParticleToRewindInfo.Find(UniqueIdx))
@@ -652,8 +699,10 @@ private:
 	};
 
 	TArrayAsMap<FUniqueIdx,FParticleRewindInfo> ParticleToRewindInfo;
-	TArray<FFrameManagerInfo> Managers;
+	TCircularBuffer<FFrameManagerInfo> Managers;
 	TArray<FDirtyParticleInfo> AllDirtyParticles;
-	int32 CurFrame;
+	uint32 CurFrame;
+	uint32 EarliestFrame;
+	bool bNeedsSave;	//Indicates that some data is pointing at head and requires saving before a rewind
 };
 }
