Make it so dirty particles get removed from history if they stop being dirty for N frames

diff --git a/Engine/Source/Programs/HeadlessChaos/Private/HeadlessChaosTestRewind.cpp b/Engine/Source/Programs/HeadlessChaos/Private/HeadlessChaosTestRewind.cpp
--- a/Engine/Source/Programs/HeadlessChaos/Private/HeadlessChaosTestRewind.cpp
+++ b/Engine/Source/Programs/HeadlessChaos/Private/HeadlessChaosTestRewind.cpp
@@ -532,4 +532,63 @@ namespace ChaosTest {
 
 		Module->DestroySolver(Solver);
 	}
+
+	GTEST_TEST(RewindTest,NumDirty)
+	{
+		auto Sphere = TSharedPtr<FImplicitObject,ESPMode::ThreadSafe>(new TSphere<float,3>(TVector<float,3>(0),10));
+
+		FChaosSolversModule* Module = FChaosSolversModule::GetModule();
+		Module->ChangeThreadingMode(EChaosThreadingMode::SingleThread);
+
+		// Make a solver
+		FPhysicsSolver* Solver = Module->CreateSolver(nullptr,ESolverFlags::Standalone);
+		Solver->SetEnabled(true);
+
+		Solver->EnableRewindCapture(5);
+
+
+		// Make particles
+		auto Particle = TPBDRigidParticle<float,3>::CreateParticle();
+
+		Particle->SetGeometry(Sphere);
+		Solver->RegisterObject(Particle.Get());
+		Particle->SetGravityEnabled(true);
+		
+		for(int Step = 0; Step < 10; ++Step)
+		{
+			TickSolverHelper(Module,Solver);
+
+			const FRewindData* RewindData = Solver->GetRewindData();
+			EXPECT_EQ(RewindData->GetNumDirtyParticles(),1);
+		}
+
+		//stop movement
+		Particle->SetGravityEnabled(false);
+		Particle->SetV(FVec3(0));
+
+		for(int Step = 0; Step < 10; ++Step)
+		{
+			TickSolverHelper(Module,Solver);
+		}
+
+		{
+			//enough frames with no changes so no longer dirty
+			const FRewindData* RewindData = Solver->GetRewindData();
+			EXPECT_EQ(RewindData->GetNumDirtyParticles(),0);
+		}
+
+		{
+			//single change so back to being dirty
+			Particle->SetGravityEnabled(true);
+			TickSolverHelper(Module,Solver);
+
+			const FRewindData* RewindData = Solver->GetRewindData();
+			EXPECT_EQ(RewindData->GetNumDirtyParticles(),1);
+		}
+
+		// Throw out the proxy
+		Solver->UnregisterObject(Particle.Get());
+
+		Module->DestroySolver(Solver);
+	}
 }
diff --git a/Engine/Source/Runtime/Experimental/ChaosSolvers/Private/PBDRigidsSolver.cpp b/Engine/Source/Runtime/Experimental/ChaosSolvers/Private/PBDRigidsSolver.cpp
--- a/Engine/Source/Runtime/Experimental/ChaosSolvers/Private/PBDRigidsSolver.cpp
+++ b/Engine/Source/Runtime/Experimental/ChaosSolvers/Private/PBDRigidsSolver.cpp
@@ -421,7 +421,7 @@ namespace Chaos
 				// Remove from rewind data
 				if(FRewindData* RewindData = Solver->GetRewindData())
 				{
-					RewindData->RemoveParticle(*Handle);
+					RewindData->RemoveParticle(Handle->UniqueIdx());
 				}
 
 			  // Remove game thread particle from ActiveGameThreadParticles so we won't crash when pulling physics state
diff --git a/Engine/Source/Runtime/Experimental/ChaosSolvers/Public/RewindData.h b/Engine/Source/Runtime/Experimental/ChaosSolvers/Public/RewindData.h
--- a/Engine/Source/Runtime/Experimental/ChaosSolvers/Public/RewindData.h
+++ b/Engine/Source/Runtime/Experimental/ChaosSolvers/Public/RewindData.h
@@ -453,9 +453,9 @@ public:
 		return true;
 	}
 
-	void RemoveParticle(const TGeometryParticleHandle<FReal,3>& Particle)
+	void RemoveParticle(const FUniqueIdx UniqueIdx)
 	{
-		if(const FParticleRewindInfo* Info = ParticleToRewindInfo.Find(Particle.UniqueIdx()))
+		if(const FParticleRewindInfo* Info = ParticleToRewindInfo.Find(UniqueIdx))
 		{
 			const int32 Idx = Info->AllDirtyParticlesIdx;
 			AllDirtyParticles.RemoveAtSwap(Idx);
@@ -465,7 +465,7 @@ public:
 				ParticleToRewindInfo.FindChecked(AllDirtyParticles[Idx].CachedUniqueIdx).AllDirtyParticlesIdx = Idx;
 			}
 
-			ParticleToRewindInfo.RemoveChecked(Particle.UniqueIdx());
+			ParticleToRewindInfo.RemoveChecked(UniqueIdx);
 		}
 	}
 
@@ -487,8 +487,20 @@ public:
 		{
 			++EarliestFrame;
 		}
+
+		//remove any old dirty particles
+		for(int32 DirtyIdx = AllDirtyParticles.Num() - 1; DirtyIdx >= 0; --DirtyIdx)
+		{
+			if(AllDirtyParticles[DirtyIdx].LastDirtyFrame < EarliestFrame)
+			{
+				RemoveParticle(AllDirtyParticles[DirtyIdx].CachedUniqueIdx);
+			}
+		}
 	}
 
+	//Number of particles that we're currently storing history for
+	int32 GetNumDirtyParticles() const { return AllDirtyParticles.Num(); }
+
 	void PrepareFrame(int32 NumDirtyParticles)
 	{
 		PrepareManagerForFrame(CurFrame,NumDirtyParticles);
@@ -527,6 +539,7 @@ public:
 		{
 			const auto PTParticle = Proxy->GetHandle();
 			FParticleRewindInfo& Info = FindOrAddParticle(Proxy->GetParticle(),PTParticle->UniqueIdx());
+			AllDirtyParticles[Info.AllDirtyParticlesIdx].LastDirtyFrame = CurFrame;
 
 			//Most properties are always a frame behind
 			if(Proxy->IsInitialized())	//Frame delay so proxy must be initialized
@@ -575,7 +588,8 @@ public:
 		if(SimWritablePropsMayChange(Rigid))
 		{
 			FParticleRewindInfo& Info = FindOrAddParticle(Rigid.GTGeometryParticle(), Rigid.UniqueIdx());
-			
+			AllDirtyParticles[Info.AllDirtyParticlesIdx].LastDirtyFrame = CurFrame-1;
+
 			//User called PrepareManagerForFrame (or PrepareFrameForPTDirty) for the previous frame, so use it
 			FDirtyPropertiesManager& DestManager = *Managers[CurFrame-1].Manager;
 
@@ -658,7 +672,7 @@ private:
 
 
 		FParticleRewindInfo& Info = ParticleToRewindInfo.FindOrAdd(UniqueIdx);
-		Info.AllDirtyParticlesIdx = AllDirtyParticles.Add(FDirtyParticleInfo{UnsafeGTParticle,UniqueIdx});
+		Info.AllDirtyParticlesIdx = AllDirtyParticles.Add(FDirtyParticleInfo{UnsafeGTParticle,UniqueIdx, CurFrame});
 		return Info;
 	}
 
@@ -696,6 +710,7 @@ private:
 	{
 		TGeometryParticle<FReal,3>* Particle;
 		FUniqueIdx CachedUniqueIdx;	//Needed when manipulating on physics thread and Particle data cannot be read
+		uint32 LastDirtyFrame;	//Track how recently this was made dirty
 	};
 
 	TArrayAsMap<FUniqueIdx,FParticleRewindInfo> ParticleToRewindInfo;
