Chaos Collision
- eliminate initial extraneous Collision Update
- eliminate unnecessary copying of CollisionConstraint data for RBAN
- move Box-Capsule AABB check into Update (from Create)
- add Box-Box AABB check in Update
- disable manifolds

diff --git a/Engine/Source/Runtime/Experimental/Chaos/Private/Chaos/CollisionResolution.cpp b/Engine/Source/Runtime/Experimental/Chaos/Private/Chaos/CollisionResolution.cpp
--- a/Engine/Source/Runtime/Experimental/Chaos/Private/Chaos/CollisionResolution.cpp
+++ b/Engine/Source/Runtime/Experimental/Chaos/Private/Chaos/CollisionResolution.cpp
@@ -42,7 +42,7 @@ FAutoConsoleVariableRef CVarCCDAllowedDepthBoundsScale(TEXT("p.Chaos.CCD.Allowed
 #define PHI_RESAMPLE_THRESHOLD 0.001
 
 
-bool Chaos_Collision_UseManifolds = true;
+bool Chaos_Collision_UseManifolds = false;
 FAutoConsoleVariableRef CVarChaosCollisionUseManifolds(TEXT("p.Chaos.Collision.UseManifolds"), Chaos_Collision_UseManifolds, TEXT("Enable/Disable use of manifoldes in collision."));
 
 bool Chaos_Collision_UseManifolds_Test = false;
@@ -58,6 +58,9 @@ float Chaos_Collision_CapsuleBoxManifoldTolerance = FMath::Sin(FMath::DegreesToR
 FConsoleVariableDelegate Chaos_Collision_CapsuleBoxManifoldDelegate = FConsoleVariableDelegate::CreateLambda([](IConsoleVariable* CVar) { Chaos_Collision_CapsuleBoxManifoldTolerance = FMath::Sin(FMath::DegreesToRadians(Chaos_Collision_CapsuleBoxManifoldAngle)); });
 FAutoConsoleVariableRef CVarChaosCollisionBoxCapsuleManifoldAngle(TEXT("p.Chaos.Collision.CapsuleBoxManifoldAngle"), Chaos_Collision_CapsuleBoxManifoldAngle, TEXT("If a capsule is more than this angle from vertical, do not use a manifold"), Chaos_Collision_CapsuleBoxManifoldDelegate);
 
+bool bChaos_Collision_DeferNarrowPhase = true;
+FAutoConsoleVariableRef CVarChaosCollisionDeferNarrowPhase(TEXT("p.Chaos.Collision.DeferNarrowPhase"), bChaos_Collision_DeferNarrowPhase, TEXT("Create contacts for all broadphase pairs, perform NarrowPhase later."));
+
 namespace Chaos
 {
 	namespace Collisions
@@ -185,14 +188,12 @@ namespace Chaos
 			}
 		}
 
-
 		template <typename T, int d, typename GeometryA, typename GeometryB>
-		TContactPoint<T> GJKContactPoint(const GeometryA& A, const TRigidTransform<T, d>& ATM, const GeometryB& B, const TRigidTransform<T, d>& BTM, const TVector<T, 3>& InitialDir)
+		TContactPoint<T> GJKContactPoint2(const GeometryA& A, const GeometryB& B, const TRigidTransform<T, d>& ATM, const TRigidTransform<T, d>& BToATM, const TVector<T, 3>& InitialDir)
 		{
 			SCOPE_CYCLE_COUNTER_GJK();
 
 			TContactPoint<T> Contact;
-			const TRigidTransform<T, d> BToATM = BTM.GetRelativeTransform(ATM);
 
 			T Penetration;
 			TVec3<T> ClosestA, ClosestB, Normal;
@@ -212,6 +213,13 @@ namespace Chaos
 			return Contact;
 		}
 
+		template <typename T, int d, typename GeometryA, typename GeometryB>
+		TContactPoint<T> GJKContactPoint(const GeometryA& A, const TRigidTransform<T, d>& ATM, const GeometryB& B, const TRigidTransform<T, d>& BTM, const TVector<T, 3>& InitialDir)
+		{
+			const TRigidTransform<T, d> BToATM = BTM.GetRelativeTransform(ATM);
+			return GJKContactPoint2(A, B, ATM, BToATM, InitialDir);
+		}
+
 		template <typename T, int d, typename GeometryA, typename GeometryB>
 		TContactPoint<T> GJKContactPointSwept(const GeometryA& A, const TRigidTransform<T, d>& ATM, const GeometryB& B, const TRigidTransform<T, d>& BTM, const TVector<T, d>& Dir, const T Length, const T CullDistance, T& TOI)
 		{
@@ -404,10 +412,10 @@ namespace Chaos
 		}
 
 		template <typename T, int d>
-		void UpdateSingleShotManifold(TRigidBodyMultiPointContactConstraint<T, d>& Constraint, const TRigidTransform<T, d>& Transform0, const TRigidTransform<T, d>& Transform1, const T CullDistance)
+		void UpdateSingleShotManifold(TRigidBodyMultiPointContactConstraint<T, d>& Constraint, const TRigidTransform<T, d>& WorldTransform0, const TRigidTransform<T, d>& WorldTransform1, const T CullDistance)
 		{
 			// single shot manifolds for TConvex implicit object in the constraints implicit[0] position. 
-			TContactPoint<T> ContactPoint = ConvexConvexContactPoint(*Constraint.Manifold.Implicit[0], Transform0, *Constraint.Manifold.Implicit[1], Transform1, CullDistance);
+			TContactPoint<T> ContactPoint = ConvexConvexContactPoint(*Constraint.Manifold.Implicit[0], WorldTransform0, *Constraint.Manifold.Implicit[1], WorldTransform1, CullDistance);
 
 			// Cache the nearest point as the initial contact
 			Constraint.Manifold.Phi = ContactPoint.Phi;
@@ -422,13 +430,13 @@ namespace Chaos
 			//   ConvexConvexContactPoint that can collect the face indices during evaluation of the support functions. 
 			//   This can be implemented without virtual calls.
 			//
-			int32 FaceIndex = Constraint.Manifold.Implicit[0]->FindClosestFaceAndVertices(Transform0.InverseTransformPosition(ContactPoint.Location), CollisionSamples, 1.f);
+			int32 FaceIndex = Constraint.Manifold.Implicit[0]->FindClosestFaceAndVertices(WorldTransform0.InverseTransformPosition(ContactPoint.Location), CollisionSamples, 1.f);
 
 			bool bNewManifold = (FaceIndex != Constraint.GetManifoldPlaneFaceIndex()) || (Constraint.NumManifoldPoints() == 0);
 			if (bNewManifold)
 			{
-				const FVec3 PlaneNormal = Transform1.InverseTransformVectorNoScale(ContactPoint.Normal);
-				const FVec3 PlanePos = Transform1.InverseTransformPosition(ContactPoint.Location - ContactPoint.Phi*ContactPoint.Normal);
+				const FVec3 PlaneNormal = WorldTransform1.InverseTransformVectorNoScale(ContactPoint.Normal);
+				const FVec3 PlanePos = WorldTransform1.InverseTransformPosition(ContactPoint.Location - ContactPoint.Phi*ContactPoint.Normal);
 				Constraint.SetManifoldPlane(1, FaceIndex, PlaneNormal, PlanePos);
 
 				//
@@ -445,7 +453,7 @@ namespace Chaos
 		}
 
 		template <typename T, int d>
-		void UpdateIterativeManifold(TRigidBodyMultiPointContactConstraint<T, d>&  Constraint, const TRigidTransform<T, d>& Transform0, const TRigidTransform<T, d>& Transform1, const T CullDistance)
+		void UpdateIterativeManifold(TRigidBodyMultiPointContactConstraint<T, d>&  Constraint, const TRigidTransform<T, d>& WorldTransform0, const TRigidTransform<T, d>& WorldTransform1, const T CullDistance)
 		{
 			auto SumSampleData = [&](TRigidBodyMultiPointContactConstraint<T, d>& LambdaConstraint) -> TVector<float, 3>
 			{
@@ -458,7 +466,7 @@ namespace Chaos
 			};
 
 			// iterative manifolds for non TConvex implicit objects that require sampling 
-			TContactPoint<T> ContactPoint = ConvexConvexContactPoint(*Constraint.Manifold.Implicit[0], Transform0, *Constraint.Manifold.Implicit[1], Transform1, CullDistance);
+			TContactPoint<T> ContactPoint = ConvexConvexContactPoint(*Constraint.Manifold.Implicit[0], WorldTransform0, *Constraint.Manifold.Implicit[1], WorldTransform1, CullDistance);
 
 			// Cache the nearest point as the initial contact
 			Constraint.Manifold.Phi = ContactPoint.Phi;
@@ -468,12 +476,12 @@ namespace Chaos
 			if (!ContactPoint.Normal.Equals(Constraint.GetManifoldPlaneNormal()) || !Constraint.NumManifoldPoints())
 			{
 				Constraint.ResetManifoldPoints();
-				FVec3 PlaneNormal = Transform1.InverseTransformVectorNoScale(ContactPoint.Normal);
-				FVec3 PlanePosition = Transform1.InverseTransformPosition(ContactPoint.Location - ContactPoint.Phi*ContactPoint.Normal);
+				FVec3 PlaneNormal = WorldTransform1.InverseTransformVectorNoScale(ContactPoint.Normal);
+				FVec3 PlanePosition = WorldTransform1.InverseTransformPosition(ContactPoint.Location - ContactPoint.Phi*ContactPoint.Normal);
 				Constraint.SetManifoldPlane(1, INDEX_NONE, PlaneNormal, PlanePosition);
 			}
 
-			TVector<T, d> SurfaceSample = Transform0.InverseTransformPosition(ContactPoint.Location);
+			TVector<T, d> SurfaceSample = WorldTransform0.InverseTransformPosition(ContactPoint.Location);
 			if (Constraint.NumManifoldPoints() < 4)
 			{
 				Constraint.AddManifoldPoint(SurfaceSample);
@@ -528,15 +536,21 @@ namespace Chaos
 		// Box - Box
 		//
 		template <typename T, int d>
-		TContactPoint<T> BoxBoxContactPoint(const TAABB<T, d>& Box1, const TRigidTransform<T, d>& ATM, const TAABB<T, d>& Box2, const TRigidTransform<T, d>& BTM, const T CullDistance)
+		TContactPoint<T> BoxBoxContactPoint(const TAABB<T, d>& Box1, const TAABB<T, d>& Box2, const TRigidTransform<T, d>& ATM, const TRigidTransform<T, d>& BToATM, const T CullDistance)
 		{
-			return GJKContactPoint(Box1, ATM, Box2, BTM, TVector<T, d>(1, 0, 0));
+			return GJKContactPoint2(Box1, Box2, ATM, BToATM, TVector<T, d>(1, 0, 0));
 		}
 
 		template <typename T, int d>
 		void UpdateBoxBoxConstraint(const TAABB<T, d>& Box1, const TRigidTransform<T, d>& Box1Transform, const TAABB<T, d>& Box2, const TRigidTransform<T, d>& Box2Transform, const T CullDistance, TRigidBodyPointContactConstraint<T, d>& Constraint)
 		{
-			UpdateContactPoint(Constraint.Manifold, BoxBoxContactPoint(Box1, Box1Transform, Box2, Box2Transform, CullDistance));
+			const TRigidTransform<T, d> Box2ToBox1TM = Box2Transform.GetRelativeTransform(Box1Transform);
+			TAABB<T, d> Box2In1 = Box2.TransformedAABB(Box2ToBox1TM);
+			Box2In1.Thicken(CullDistance);
+			if (Box1.Intersects(Box2In1))
+			{
+				UpdateContactPoint(Constraint.Manifold, BoxBoxContactPoint(Box1, Box2, Box1Transform, Box2ToBox1TM, CullDistance));
+			}
 		}
 
 		template <typename T, int d>
@@ -547,17 +561,24 @@ namespace Chaos
 		}
 
 		template<typename T, int d>
-		void ConstructBoxBoxConstraints(TGeometryParticleHandle<T, d>* Particle0, TGeometryParticleHandle<T, d>* Particle1, const FImplicitObject* Implicit0, const FImplicitObject* Implicit1, const TRigidTransform<T, d>& Transform0, const TRigidTransform<T, d>& Transform1, const T CullDistance, FCollisionConstraintsArray& NewConstraints)
+		void ConstructBoxBoxConstraints(TGeometryParticleHandle<T, d>* Particle0, TGeometryParticleHandle<T, d>* Particle1, const FImplicitObject* Implicit0, const FImplicitObject* Implicit1, const TRigidTransform<T, d>& LocalTransform0, const TRigidTransform<T, d>& LocalTransform1, const T CullDistance, FCollisionConstraintsArray& NewConstraints)
 		{
 			const TBox<T, d> * Object0 = Implicit0->template GetObject<const TBox<T, d> >();
 			const TBox<T, d> * Object1 = Implicit1->template GetObject<const TBox<T, d> >();
 			if (ensure(Object0 && Object1))
 			{
-				TRigidTransform<T, d> ParticleImplicit0TM = Transform0.GetRelativeTransform(Collisions::GetTransform(Particle0));
-				TRigidTransform<T, d> ParticleImplicit1TM = Transform1.GetRelativeTransform(Collisions::GetTransform(Particle1));
-				FRigidBodyPointContactConstraint Constraint = FRigidBodyPointContactConstraint(Particle0, Implicit0, ParticleImplicit0TM, Particle1, Implicit1, ParticleImplicit1TM, EContactShapesType::BoxBox);
-				UpdateBoxBoxConstraint(Object0->BoundingBox(), Transform0, Object1->BoundingBox(), Transform1, CullDistance, Constraint);
-				NewConstraints.TryAdd(CullDistance, Constraint);
+				FRigidBodyPointContactConstraint Constraint = FRigidBodyPointContactConstraint(Particle0, Implicit0, LocalTransform0, Particle1, Implicit1, LocalTransform1, EContactShapesType::BoxBox);
+				if (!bChaos_Collision_DeferNarrowPhase)
+				{
+					TRigidTransform<T, d> WorldTransform0 = LocalTransform0 * Collisions::GetTransform(Particle0);
+					TRigidTransform<T, d> WorldTransform1 = LocalTransform1 * Collisions::GetTransform(Particle1);
+					UpdateBoxBoxConstraint(Object0->BoundingBox(), WorldTransform0, Object1->BoundingBox(), WorldTransform1, CullDistance, Constraint);
+					NewConstraints.TryAdd(CullDistance, Constraint);
+				}
+				else
+				{
+					NewConstraints.Add(Constraint);
+				}
 			}
 		}
 
@@ -585,18 +606,25 @@ namespace Chaos
 		}
 
 		template<typename T, int d>
-		void ConstructBoxHeightFieldConstraints(TGeometryParticleHandle<T, d>* Particle0, TGeometryParticleHandle<T, d>* Particle1, const FImplicitObject* Implicit0, const FImplicitObject* Implicit1, const TRigidTransform<T, d>& Transform0, const TRigidTransform<T, d>& Transform1, const T CullDistance, FCollisionConstraintsArray& NewConstraints)
+		void ConstructBoxHeightFieldConstraints(TGeometryParticleHandle<T, d>* Particle0, TGeometryParticleHandle<T, d>* Particle1, const FImplicitObject* Implicit0, const FImplicitObject* Implicit1, const TRigidTransform<T, d>& LocalTransform0, const TRigidTransform<T, d>& LocalTransform1, const T CullDistance, FCollisionConstraintsArray& NewConstraints)
 		{
 
 			const TBox<T, d> * Object0 = Implicit0->template GetObject<const TBox<T, d> >();
 			const THeightField<T> * Object1 = Implicit1->template GetObject<const THeightField<T> >();
 			if (ensure(Object0 && Object1))
 			{
-				TRigidTransform<T, d> ParticleImplicit0TM = Transform0.GetRelativeTransform(Collisions::GetTransform(Particle0));
-				TRigidTransform<T, d> ParticleImplicit1TM = Transform1.GetRelativeTransform(Collisions::GetTransform(Particle1));
-				FRigidBodyPointContactConstraint Constraint = FRigidBodyPointContactConstraint(Particle0, Implicit0, ParticleImplicit0TM, Particle1, Implicit1, ParticleImplicit1TM, EContactShapesType::BoxHeightField);
-				UpdateBoxHeightFieldConstraint(Object0->BoundingBox(), Transform0, *Object1, Transform1, CullDistance, Constraint);
-				NewConstraints.TryAdd(CullDistance, Constraint);
+				FRigidBodyPointContactConstraint Constraint = FRigidBodyPointContactConstraint(Particle0, Implicit0, LocalTransform0, Particle1, Implicit1, LocalTransform1, EContactShapesType::BoxHeightField);
+				if (!bChaos_Collision_DeferNarrowPhase)
+				{
+					TRigidTransform<T, d> WorldTransform0 = LocalTransform0 * Collisions::GetTransform(Particle0);
+					TRigidTransform<T, d> WorldTransform1 = LocalTransform1 * Collisions::GetTransform(Particle1);
+					UpdateBoxHeightFieldConstraint(Object0->BoundingBox(), WorldTransform0, *Object1, WorldTransform1, CullDistance, Constraint);
+					NewConstraints.TryAdd(CullDistance, Constraint);
+				}
+				else
+				{
+					NewConstraints.Add(Constraint);
+				}
 			}
 		}
 
@@ -673,18 +701,25 @@ namespace Chaos
 
 
 		template<typename T, int d>
-		void ConstructBoxPlaneConstraints(TGeometryParticleHandle<T, d>* Particle0, TGeometryParticleHandle<T, d>* Particle1, const FImplicitObject* Implicit0, const FImplicitObject* Implicit1, const TRigidTransform<T, d>& Transform0, const TRigidTransform<T, d>& Transform1, const T CullDistance, FCollisionConstraintsArray& NewConstraints)
+		void ConstructBoxPlaneConstraints(TGeometryParticleHandle<T, d>* Particle0, TGeometryParticleHandle<T, d>* Particle1, const FImplicitObject* Implicit0, const FImplicitObject* Implicit1, const TRigidTransform<T, d>& LocalTransform0, const TRigidTransform<T, d>& LocalTransform1, const T CullDistance, FCollisionConstraintsArray& NewConstraints)
 		{
 
 			const TBox<T, d> * Object0 = Implicit0->template GetObject<const TBox<T, d> >();
 			const TPlane<T, d> * Object1 = Implicit1->template GetObject<const TPlane<T, d> >();
 			if (ensure(Object0 && Object1))
 			{
-				TRigidTransform<T, d> ParticleImplicit0TM = Transform0.GetRelativeTransform(Collisions::GetTransform(Particle0));
-				TRigidTransform<T, d> ParticleImplicit1TM = Transform1.GetRelativeTransform(Collisions::GetTransform(Particle1));
-				FRigidBodyPointContactConstraint Constraint = FRigidBodyPointContactConstraint(Particle0, Implicit0, ParticleImplicit0TM, Particle1, Implicit1, ParticleImplicit1TM, EContactShapesType::BoxPlane);
-				UpdateBoxPlaneConstraint(Object0->BoundingBox(), Transform0, *Object1, Transform1, CullDistance, Constraint);
-				NewConstraints.TryAdd(CullDistance, Constraint);
+				FRigidBodyPointContactConstraint Constraint = FRigidBodyPointContactConstraint(Particle0, Implicit0, LocalTransform0, Particle1, Implicit1, LocalTransform1, EContactShapesType::BoxPlane);
+				if (!bChaos_Collision_DeferNarrowPhase)
+				{
+					TRigidTransform<T, d> WorldTransform0 = LocalTransform0 * Collisions::GetTransform(Particle0);
+					TRigidTransform<T, d> WorldTransform1 = LocalTransform1 * Collisions::GetTransform(Particle1);
+					UpdateBoxPlaneConstraint(Object0->BoundingBox(), WorldTransform0, *Object1, WorldTransform1, CullDistance, Constraint);
+					NewConstraints.TryAdd(CullDistance, Constraint);
+				}
+				else
+				{
+					NewConstraints.Add(Constraint);
+				}
 			}
 		}
 
@@ -699,9 +734,9 @@ namespace Chaos
 		}
 
 		template <typename TriMeshType, typename T, int d>
-		void UpdateBoxTriangleMeshConstraint(const TAABB<T, d>& Box0, const TRigidTransform<T, d>& Transform0, const TriMeshType& TriangleMesh1, const TRigidTransform<T, d>& Transform1, const T CullDistance, TRigidBodyPointContactConstraint<T, d>& Constraint)
+		void UpdateBoxTriangleMeshConstraint(const TAABB<T, d>& Box0, const TRigidTransform<T, d>& WorldTransform0, const TriMeshType& TriangleMesh1, const TRigidTransform<T, d>& WorldTransform1, const T CullDistance, TRigidBodyPointContactConstraint<T, d>& Constraint)
 		{
-			UpdateContactPoint(Constraint.Manifold, BoxTriangleMeshContactPoint(Box0, Transform0, TriangleMesh1, Transform1, CullDistance));
+			UpdateContactPoint(Constraint.Manifold, BoxTriangleMeshContactPoint(Box0, WorldTransform0, TriangleMesh1, WorldTransform1, CullDistance));
 		}
 
 
@@ -712,26 +747,39 @@ namespace Chaos
 		}
 
 		template<typename T, int d>
-		void ConstructBoxTriangleMeshConstraints(TGeometryParticleHandle<T, d>* Particle0, TGeometryParticleHandle<T, d>* Particle1, const FImplicitObject* Implicit0, const FImplicitObject* Implicit1, const TRigidTransform<T, d>& Transform0, const TRigidTransform<T, d>& Transform1, const T CullDistance, FCollisionConstraintsArray& NewConstraints)
-		{
+		void ConstructBoxTriangleMeshConstraints(TGeometryParticleHandle<T, d>* Particle0, TGeometryParticleHandle<T, d>* Particle1, const FImplicitObject* Implicit0, const FImplicitObject* Implicit1, const TRigidTransform<T, d>& LocalTransform0, const TRigidTransform<T, d>& LocalTransform1, const T CullDistance, FCollisionConstraintsArray& NewConstraints)		{
 			const TBox<T, d> * Object0 = Implicit0->template GetObject<const TBox<T, d> >();
 			if (ensure(Object0))
 			{
 				if (const TImplicitObjectScaled<FTriangleMeshImplicitObject>* ScaledTriangleMesh = Implicit1->template GetObject<const TImplicitObjectScaled<FTriangleMeshImplicitObject>>())
 				{
-					TRigidTransform<T, d> ParticleImplicit0TM = Transform0.GetRelativeTransform(Collisions::GetTransform(Particle0));
-					TRigidTransform<T, d> ParticleImplicit1TM = Transform1.GetRelativeTransform(Collisions::GetTransform(Particle1));
-					FRigidBodyPointContactConstraint Constraint = FRigidBodyPointContactConstraint(Particle0, Implicit0, ParticleImplicit0TM, Particle1, Implicit1, ParticleImplicit1TM, EContactShapesType::BoxTriMesh);
-					UpdateBoxTriangleMeshConstraint(Object0->GetAABB(), Transform0, *ScaledTriangleMesh, Transform1, CullDistance, Constraint);
-					NewConstraints.TryAdd(CullDistance, Constraint);
+					FRigidBodyPointContactConstraint Constraint = FRigidBodyPointContactConstraint(Particle0, Implicit0, LocalTransform0, Particle1, Implicit1, LocalTransform1, EContactShapesType::BoxTriMesh);
+					if (!bChaos_Collision_DeferNarrowPhase)
+					{
+						TRigidTransform<T, d> WorldTransform0 = LocalTransform0 * Collisions::GetTransform(Particle0);
+						TRigidTransform<T, d> WorldTransform1 = LocalTransform1 * Collisions::GetTransform(Particle1);
+						UpdateBoxTriangleMeshConstraint(Object0->GetAABB(), WorldTransform0, *ScaledTriangleMesh, WorldTransform1, CullDistance, Constraint);
+						NewConstraints.TryAdd(CullDistance, Constraint);
+					}
+					else
+					{
+						NewConstraints.Add(Constraint);
+					}
 				}
 				else if(const FTriangleMeshImplicitObject* TriangleMesh = Implicit1->template GetObject<const FTriangleMeshImplicitObject>())
 				{
-					TRigidTransform<T, d> ParticleImplicit0TM = Transform0.GetRelativeTransform(Collisions::GetTransform(Particle0));
-					TRigidTransform<T, d> ParticleImplicit1TM = Transform1.GetRelativeTransform(Collisions::GetTransform(Particle1));
-					FRigidBodyPointContactConstraint Constraint = FRigidBodyPointContactConstraint(Particle0, Implicit0, ParticleImplicit0TM, Particle1, Implicit1, ParticleImplicit1TM, EContactShapesType::BoxTriMesh);
-					UpdateBoxTriangleMeshConstraint(Object0->GetAABB(), Transform0, *TriangleMesh, Transform1, CullDistance, Constraint);
-					NewConstraints.TryAdd(CullDistance, Constraint);
+					FRigidBodyPointContactConstraint Constraint = FRigidBodyPointContactConstraint(Particle0, Implicit0, LocalTransform0, Particle1, Implicit1, LocalTransform1, EContactShapesType::BoxTriMesh);
+					if (!bChaos_Collision_DeferNarrowPhase)
+					{
+						TRigidTransform<T, d> WorldTransform0 = LocalTransform0 * Collisions::GetTransform(Particle0);
+						TRigidTransform<T, d> WorldTransform1 = LocalTransform1 * Collisions::GetTransform(Particle1);
+						UpdateBoxTriangleMeshConstraint(Object0->GetAABB(), WorldTransform0, *TriangleMesh, WorldTransform1, CullDistance, Constraint);
+						NewConstraints.TryAdd(CullDistance, Constraint);
+					}
+					else
+					{
+						NewConstraints.Add(Constraint);
+					}
 				}
 				else
 				{
@@ -775,18 +823,24 @@ namespace Chaos
 		}
 
 		template<typename T, int d>
-		void ConstructSphereSphereConstraints(TGeometryParticleHandle<T, d>* Particle0, TGeometryParticleHandle<T, d>* Particle1, const FImplicitObject* Implicit0, const FImplicitObject* Implicit1, const TRigidTransform<T, d>& Transform0, const TRigidTransform<T, d>& Transform1, const T CullDistance, FCollisionConstraintsArray& NewConstraints)
+		void ConstructSphereSphereConstraints(TGeometryParticleHandle<T, d>* Particle0, TGeometryParticleHandle<T, d>* Particle1, const FImplicitObject* Implicit0, const FImplicitObject* Implicit1, const TRigidTransform<T, d>& LocalTransform0, const TRigidTransform<T, d>& LocalTransform1, const T CullDistance, FCollisionConstraintsArray& NewConstraints)
 		{
-
 			const TSphere<T, d> * Object0 = Implicit0->template GetObject<const TSphere<T, d> >();
 			const TSphere<T, d> * Object1 = Implicit1->template GetObject<const TSphere<T, d> >();
 			if (ensure(Object0 && Object1))
 			{
-				TRigidTransform<T, d> ParticleImplicit0TM = Transform0.GetRelativeTransform(Collisions::GetTransform(Particle0));
-				TRigidTransform<T, d> ParticleImplicit1TM = Transform1.GetRelativeTransform(Collisions::GetTransform(Particle1));
-				FRigidBodyPointContactConstraint Constraint = FRigidBodyPointContactConstraint(Particle0, Implicit0, ParticleImplicit0TM, Particle1, Implicit1, ParticleImplicit1TM, EContactShapesType::SphereSphere);
-				UpdateSphereSphereConstraint(*Object0, Transform0, *Object1, Transform1, CullDistance, Constraint);
-				NewConstraints.TryAdd(CullDistance, Constraint);
+				FRigidBodyPointContactConstraint Constraint = FRigidBodyPointContactConstraint(Particle0, Implicit0, LocalTransform0, Particle1, Implicit1, LocalTransform1, EContactShapesType::SphereSphere);
+				if (!bChaos_Collision_DeferNarrowPhase)
+				{
+					TRigidTransform<T, d> WorldTransform0 = LocalTransform0 * Collisions::GetTransform(Particle0);
+					TRigidTransform<T, d> WorldTransform1 = LocalTransform1 * Collisions::GetTransform(Particle1);
+					UpdateSphereSphereConstraint(*Object0, WorldTransform0, *Object1, WorldTransform1, CullDistance, Constraint);
+					NewConstraints.TryAdd(CullDistance, Constraint);
+				}
+				else
+				{
+					NewConstraints.Add(Constraint);
+				}
 			}
 		}
 
@@ -814,18 +868,25 @@ namespace Chaos
 		}
 
 		template<typename T, int d>
-		void ConstructSphereHeightFieldConstraints(TGeometryParticleHandle<T, d>* Particle0, TGeometryParticleHandle<T, d>* Particle1, const FImplicitObject* Implicit0, const FImplicitObject* Implicit1, const TRigidTransform<T, d>& Transform0, const TRigidTransform<T, d>& Transform1, const T CullDistance, FCollisionConstraintsArray& NewConstraints)
+		void ConstructSphereHeightFieldConstraints(TGeometryParticleHandle<T, d>* Particle0, TGeometryParticleHandle<T, d>* Particle1, const FImplicitObject* Implicit0, const FImplicitObject* Implicit1, const TRigidTransform<T, d>& LocalTransform0, const TRigidTransform<T, d>& LocalTransform1, const T CullDistance, FCollisionConstraintsArray& NewConstraints)
 		{
 
 			const TSphere<T, d> * Object0 = Implicit0->template GetObject<const TSphere<T, d> >();
 			const THeightField<T> * Object1 = Implicit1->template GetObject<const THeightField<T> >();
 			if (ensure(Object0 && Object1))
 			{
-				TRigidTransform<T, d> ParticleImplicit0TM = Transform0.GetRelativeTransform(Collisions::GetTransform(Particle0));
-				TRigidTransform<T, d> ParticleImplicit1TM = Transform1.GetRelativeTransform(Collisions::GetTransform(Particle1));
-				FRigidBodyPointContactConstraint Constraint = FRigidBodyPointContactConstraint(Particle0, Implicit0, ParticleImplicit0TM, Particle1, Implicit1, ParticleImplicit1TM, EContactShapesType::SphereHeightField);
-				UpdateSphereHeightFieldConstraint(*Object0, Transform0, *Object1, Transform1, CullDistance, Constraint);
-				NewConstraints.TryAdd(CullDistance, Constraint);
+				FRigidBodyPointContactConstraint Constraint = FRigidBodyPointContactConstraint(Particle0, Implicit0, LocalTransform0, Particle1, Implicit1, LocalTransform1, EContactShapesType::SphereHeightField);
+				if (!bChaos_Collision_DeferNarrowPhase)
+				{
+					TRigidTransform<T, d> WorldTransform0 = LocalTransform0 * Collisions::GetTransform(Particle0);
+					TRigidTransform<T, d> WorldTransform1 = LocalTransform1 * Collisions::GetTransform(Particle1);
+					UpdateSphereHeightFieldConstraint(*Object0, WorldTransform0, *Object1, WorldTransform1, CullDistance, Constraint);
+					NewConstraints.TryAdd(CullDistance, Constraint);
+				}
+				else
+				{
+					NewConstraints.Add(Constraint);
+				}
 			}
 		}
 
@@ -861,18 +922,25 @@ namespace Chaos
 		}
 
 		template<typename T, int d>
-		void ConstructSpherePlaneConstraints(TGeometryParticleHandle<T, d>* Particle0, TGeometryParticleHandle<T, d>* Particle1, const FImplicitObject* Implicit0, const FImplicitObject* Implicit1, const TRigidTransform<T, d>& Transform0, const TRigidTransform<T, d>& Transform1, const T CullDistance, FCollisionConstraintsArray& NewConstraints)
+		void ConstructSpherePlaneConstraints(TGeometryParticleHandle<T, d>* Particle0, TGeometryParticleHandle<T, d>* Particle1, const FImplicitObject* Implicit0, const FImplicitObject* Implicit1, const TRigidTransform<T, d>& LocalTransform0, const TRigidTransform<T, d>& LocalTransform1, const T CullDistance, FCollisionConstraintsArray& NewConstraints)
 		{
 
 			const TSphere<T, d> * Object0 = Implicit0->template GetObject<const TSphere<T, d> >();
 			const TPlane<T, d> * Object1 = Implicit1->template GetObject<const TPlane<T, d> >();
 			if (ensure(Object0 && Object1))
 			{
-				TRigidTransform<T, d> ParticleImplicit0TM = Transform0.GetRelativeTransform(Collisions::GetTransform(Particle0));
-				TRigidTransform<T, d> ParticleImplicit1TM = Transform1.GetRelativeTransform(Collisions::GetTransform(Particle1));
-				FRigidBodyPointContactConstraint Constraint = FRigidBodyPointContactConstraint(Particle0, Implicit0, ParticleImplicit0TM, Particle1, Implicit1, ParticleImplicit1TM, EContactShapesType::SpherePlane);
-				UpdateSpherePlaneConstraint(*Object0, Transform0, *Object1, Transform1, CullDistance, Constraint);
-				NewConstraints.TryAdd(CullDistance, Constraint);
+				FRigidBodyPointContactConstraint Constraint = FRigidBodyPointContactConstraint(Particle0, Implicit0, LocalTransform0, Particle1, Implicit1, LocalTransform1, EContactShapesType::SpherePlane);
+				if (!bChaos_Collision_DeferNarrowPhase)
+				{
+					TRigidTransform<T, d> WorldTransform0 = LocalTransform0 * Collisions::GetTransform(Particle0);
+					TRigidTransform<T, d> WorldTransform1 = LocalTransform1 * Collisions::GetTransform(Particle1);
+					UpdateSpherePlaneConstraint(*Object0, WorldTransform0, *Object1, WorldTransform1, CullDistance, Constraint);
+					NewConstraints.TryAdd(CullDistance, Constraint);
+				}
+				else
+				{
+					NewConstraints.Add(Constraint);
+				}
 			}
 		}
 
@@ -912,18 +980,25 @@ namespace Chaos
 		}
 
 		template<typename T, int d>
-		void ConstructSphereBoxConstraints(TGeometryParticleHandle<T, d>* Particle0, TGeometryParticleHandle<T, d>* Particle1, const FImplicitObject* Implicit0, const FImplicitObject* Implicit1, const TRigidTransform<T, d>& Transform0, const TRigidTransform<T, d>& Transform1, const T CullDistance, FCollisionConstraintsArray& NewConstraints)
+		void ConstructSphereBoxConstraints(TGeometryParticleHandle<T, d>* Particle0, TGeometryParticleHandle<T, d>* Particle1, const FImplicitObject* Implicit0, const FImplicitObject* Implicit1, const TRigidTransform<T, d>& LocalTransform0, const TRigidTransform<T, d>& LocalTransform1, const T CullDistance, FCollisionConstraintsArray& NewConstraints)
 		{
 
 			const TSphere<T, d> * Object0 = Implicit0->template GetObject<const TSphere<T, d> >();
 			const TBox<T, d> * Object1 = Implicit1->template GetObject<const TBox<T, d> >();
 			if (ensure(Object0 && Object1))
 			{
-				TRigidTransform<T, d> ParticleImplicit0TM = Transform0.GetRelativeTransform(Collisions::GetTransform(Particle0));
-				TRigidTransform<T, d> ParticleImplicit1TM = Transform1.GetRelativeTransform(Collisions::GetTransform(Particle1));
-				FRigidBodyPointContactConstraint Constraint = FRigidBodyPointContactConstraint(Particle0, Implicit0, ParticleImplicit0TM, Particle1, Implicit1, ParticleImplicit1TM, EContactShapesType::SphereBox);
-				UpdateSphereBoxConstraint(*Object0, Transform0, Object1->BoundingBox(), Transform1, CullDistance, Constraint);
-				NewConstraints.TryAdd(CullDistance, Constraint);
+				FRigidBodyPointContactConstraint Constraint = FRigidBodyPointContactConstraint(Particle0, Implicit0, LocalTransform0, Particle1, Implicit1, LocalTransform1, EContactShapesType::SphereBox);
+				if (!bChaos_Collision_DeferNarrowPhase)
+				{
+					TRigidTransform<T, d> WorldTransform0 = LocalTransform0 * Collisions::GetTransform(Particle0);
+					TRigidTransform<T, d> WorldTransform1 = LocalTransform1 * Collisions::GetTransform(Particle1);
+					UpdateSphereBoxConstraint(*Object0, WorldTransform0, Object1->BoundingBox(), WorldTransform1, CullDistance, Constraint);
+					NewConstraints.TryAdd(CullDistance, Constraint);
+				}
+				else
+				{
+					NewConstraints.Add(Constraint);
+				}
 			}
 		}
 
@@ -970,17 +1045,24 @@ namespace Chaos
 		}
 
 		template<typename T, int d>
-		void ConstructSphereCapsuleConstraints(TGeometryParticleHandle<T, d>* Particle0, TGeometryParticleHandle<T, d>* Particle1, const FImplicitObject* Implicit0, const FImplicitObject* Implicit1, const TRigidTransform<T, d>& Transform0, const TRigidTransform<T, d>& Transform1, const T CullDistance, FCollisionConstraintsArray& NewConstraints)
+		void ConstructSphereCapsuleConstraints(TGeometryParticleHandle<T, d>* Particle0, TGeometryParticleHandle<T, d>* Particle1, const FImplicitObject* Implicit0, const FImplicitObject* Implicit1, const TRigidTransform<T, d>& LocalTransform0, const TRigidTransform<T, d>& LocalTransform1, const T CullDistance, FCollisionConstraintsArray& NewConstraints)
 		{
 			const TSphere<T, d>* Object0 = Implicit0->template GetObject<const TSphere<T, d> >();
 			const TCapsule<T>* Object1 = Implicit1->template GetObject<const TCapsule<T> >();
 			if (ensure(Object0 && Object1))
 			{
-				TRigidTransform<T, d> ParticleImplicit0TM = Transform0.GetRelativeTransform(Collisions::GetTransform(Particle0));
-				TRigidTransform<T, d> ParticleImplicit1TM = Transform1.GetRelativeTransform(Collisions::GetTransform(Particle1));
-				FRigidBodyPointContactConstraint Constraint = FRigidBodyPointContactConstraint(Particle0, Implicit0, ParticleImplicit0TM, Particle1, Implicit1, ParticleImplicit1TM, EContactShapesType::SphereCapsule);
-				UpdateSphereCapsuleConstraint(*Object0, Transform0, *Object1, Transform1, CullDistance, Constraint);
-				NewConstraints.TryAdd(CullDistance, Constraint);
+				FRigidBodyPointContactConstraint Constraint = FRigidBodyPointContactConstraint(Particle0, Implicit0, LocalTransform0, Particle1, Implicit1, LocalTransform1, EContactShapesType::SphereCapsule);
+				if (!bChaos_Collision_DeferNarrowPhase)
+				{
+					TRigidTransform<T, d> WorldTransform0 = LocalTransform0 * Collisions::GetTransform(Particle0);
+					TRigidTransform<T, d> WorldTransform1 = LocalTransform1 * Collisions::GetTransform(Particle1);
+					UpdateSphereCapsuleConstraint(*Object0, WorldTransform0, *Object1, WorldTransform1, CullDistance, Constraint);
+					NewConstraints.TryAdd(CullDistance, Constraint);
+				}
+				else
+				{
+					NewConstraints.Add(Constraint);
+				}
 			}
 		}
 
@@ -995,9 +1077,9 @@ namespace Chaos
 		}
 
 		template <typename T, int d, typename TriMeshType>
-		void UpdateSphereTriangleMeshConstraint(const TSphere<T, d>& Sphere0, const TRigidTransform<T, d>& Transform0, const TriMeshType& TriangleMesh1, const TRigidTransform<T, d>& Transform1, const T CullDistance, TRigidBodyPointContactConstraint<T, d>& Constraint)
+		void UpdateSphereTriangleMeshConstraint(const TSphere<T, d>& Sphere0, const TRigidTransform<T, d>& WorldTransform0, const TriMeshType& TriangleMesh1, const TRigidTransform<T, d>& WorldTransform1, const T CullDistance, TRigidBodyPointContactConstraint<T, d>& Constraint)
 		{
-			UpdateContactPoint(Constraint.Manifold, SphereTriangleMeshContactPoint(Sphere0, Transform0, TriangleMesh1, Transform1, CullDistance));
+			UpdateContactPoint(Constraint.Manifold, SphereTriangleMeshContactPoint(Sphere0, WorldTransform0, TriangleMesh1, WorldTransform1, CullDistance));
 		}
 
 
@@ -1008,26 +1090,40 @@ namespace Chaos
 		}
 
 		template<typename T, int d>
-		void ConstructSphereTriangleMeshConstraints(TGeometryParticleHandle<T, d>* Particle0, TGeometryParticleHandle<T, d>* Particle1, const FImplicitObject* Implicit0, const FImplicitObject* Implicit1, const TRigidTransform<T, d>& Transform0, const TRigidTransform<T, d>& Transform1, const T CullDistance, FCollisionConstraintsArray& NewConstraints)
+		void ConstructSphereTriangleMeshConstraints(TGeometryParticleHandle<T, d>* Particle0, TGeometryParticleHandle<T, d>* Particle1, const FImplicitObject* Implicit0, const FImplicitObject* Implicit1, const TRigidTransform<T, d>& LocalTransform0, const TRigidTransform<T, d>& LocalTransform1, const T CullDistance, FCollisionConstraintsArray& NewConstraints)
 		{
 			const TSphere<T, d>* Object0 = Implicit0->template GetObject<const TSphere<T, d> >();
 			if (ensure(Object0))
 			{
 				if (const TImplicitObjectScaled<FTriangleMeshImplicitObject>* ScaledTriangleMesh = Implicit1->template GetObject<const TImplicitObjectScaled<FTriangleMeshImplicitObject>>())
 				{
-					TRigidTransform<T, d> ParticleImplicit0TM = Transform0.GetRelativeTransform(Collisions::GetTransform(Particle0));
-					TRigidTransform<T, d> ParticleImplicit1TM = Transform1.GetRelativeTransform(Collisions::GetTransform(Particle1));
-					FRigidBodyPointContactConstraint Constraint = FRigidBodyPointContactConstraint(Particle0, Implicit0, ParticleImplicit0TM, Particle1, Implicit1, ParticleImplicit1TM, EContactShapesType::SphereTriMesh);
-					UpdateSphereTriangleMeshConstraint(*Object0, Transform0, *ScaledTriangleMesh, Transform1, CullDistance, Constraint);
-					NewConstraints.TryAdd(CullDistance, Constraint);
+					FRigidBodyPointContactConstraint Constraint = FRigidBodyPointContactConstraint(Particle0, Implicit0, LocalTransform0, Particle1, Implicit1, LocalTransform1, EContactShapesType::SphereTriMesh);
+					if (!bChaos_Collision_DeferNarrowPhase)
+					{
+						TRigidTransform<T, d> WorldTransform0 = LocalTransform0 * Collisions::GetTransform(Particle0);
+						TRigidTransform<T, d> WorldTransform1 = LocalTransform1 * Collisions::GetTransform(Particle1);
+						UpdateSphereTriangleMeshConstraint(*Object0, WorldTransform0, *ScaledTriangleMesh, WorldTransform1, CullDistance, Constraint);
+						NewConstraints.TryAdd(CullDistance, Constraint);
+					}
+					else
+					{
+						NewConstraints.Add(Constraint);
+					}
 				}
 				else if (const FTriangleMeshImplicitObject* TriangleMesh = Implicit1->template GetObject<const FTriangleMeshImplicitObject>())
 				{
-					TRigidTransform<T, d> ParticleImplicit0TM = Transform0.GetRelativeTransform(Collisions::GetTransform(Particle0));
-					TRigidTransform<T, d> ParticleImplicit1TM = Transform1.GetRelativeTransform(Collisions::GetTransform(Particle1));
-					FRigidBodyPointContactConstraint Constraint = FRigidBodyPointContactConstraint(Particle0, Implicit0, ParticleImplicit0TM, Particle1, Implicit1, ParticleImplicit1TM, EContactShapesType::SphereTriMesh);
-					UpdateSphereTriangleMeshConstraint(*Object0, Transform0, *TriangleMesh, Transform1, CullDistance, Constraint);
-					NewConstraints.TryAdd(CullDistance, Constraint);
+					FRigidBodyPointContactConstraint Constraint = FRigidBodyPointContactConstraint(Particle0, Implicit0, LocalTransform0, Particle1, Implicit1, LocalTransform1, EContactShapesType::SphereTriMesh);
+					if (!bChaos_Collision_DeferNarrowPhase)
+					{
+						TRigidTransform<T, d> WorldTransform0 = LocalTransform0 * Collisions::GetTransform(Particle0);
+						TRigidTransform<T, d> WorldTransform1 = LocalTransform1 * Collisions::GetTransform(Particle1);
+						UpdateSphereTriangleMeshConstraint(*Object0, WorldTransform0, *TriangleMesh, WorldTransform1, CullDistance, Constraint);
+						NewConstraints.TryAdd(CullDistance, Constraint);
+					}
+					else
+					{
+						NewConstraints.Add(Constraint);
+					}
 				}
 				else
 				{
@@ -1082,17 +1178,24 @@ namespace Chaos
 		}
 
 		template<typename T, int d>
-		void ConstructCapsuleCapsuleConstraints(TGeometryParticleHandle<T, d>* Particle0, TGeometryParticleHandle<T, d>* Particle1, const FImplicitObject* Implicit0, const FImplicitObject* Implicit1, const TRigidTransform<T, d>& Transform0, const TRigidTransform<T, d>& Transform1, const T CullDistance, FCollisionConstraintsArray& NewConstraints)
+		void ConstructCapsuleCapsuleConstraints(TGeometryParticleHandle<T, d>* Particle0, TGeometryParticleHandle<T, d>* Particle1, const FImplicitObject* Implicit0, const FImplicitObject* Implicit1, const TRigidTransform<T, d>& LocalTransform0, const TRigidTransform<T, d>& LocalTransform1, const T CullDistance, FCollisionConstraintsArray& NewConstraints)
 		{
 			const TCapsule<T> * Object0 = Implicit0->template GetObject<const TCapsule<T> >();
 			const TCapsule<T> * Object1 = Implicit1->template GetObject<const TCapsule<T> >();
 			if (ensure(Object0 && Object1))
 			{
-				TRigidTransform<T, d> ParticleImplicit0TM = Transform0.GetRelativeTransform(Collisions::GetTransform(Particle0));
-				TRigidTransform<T, d> ParticleImplicit1TM = Transform1.GetRelativeTransform(Collisions::GetTransform(Particle1));
-				FRigidBodyPointContactConstraint Constraint = FRigidBodyPointContactConstraint(Particle0, Implicit0, ParticleImplicit0TM, Particle1, Implicit1, ParticleImplicit1TM, EContactShapesType::CapsuleCapsule);
-				UpdateCapsuleCapsuleConstraint(*Object0, Transform0, *Object1, Transform1, CullDistance, Constraint);
-				NewConstraints.TryAdd(CullDistance, Constraint);
+				FRigidBodyPointContactConstraint Constraint = FRigidBodyPointContactConstraint(Particle0, Implicit0, LocalTransform0, Particle1, Implicit1, LocalTransform1, EContactShapesType::CapsuleCapsule);
+				if (!bChaos_Collision_DeferNarrowPhase)
+				{
+					TRigidTransform<T, d> WorldTransform0 = LocalTransform0 * Collisions::GetTransform(Particle0);
+					TRigidTransform<T, d> WorldTransform1 = LocalTransform1 * Collisions::GetTransform(Particle1);
+					UpdateCapsuleCapsuleConstraint(*Object0, WorldTransform0, *Object1, WorldTransform1, CullDistance, Constraint);
+					NewConstraints.TryAdd(CullDistance, Constraint);
+				}
+				else
+				{
+					NewConstraints.Add(Constraint);
+				}
 			}
 		}
 
@@ -1109,8 +1212,17 @@ namespace Chaos
 		template <typename T, int d>
 		void UpdateCapsuleBoxConstraint(const TCapsule<T>& A, const TRigidTransform<T, d>& ATransform, const TAABB<T, d>& B, const TRigidTransform<T, d>& BTransform, const T CullDistance, TRigidBodyPointContactConstraint<T, d>& Constraint)
 		{
-			const TVector<T, d> InitialDir = ATransform.GetRotation().Inverse() * -Constraint.GetNormal();
-			UpdateContactPoint(Constraint.Manifold, CapsuleBoxContactPoint(A, ATransform, B, BTransform, InitialDir, CullDistance));
+			// @todo(ccaulfield): the inverse of CapsuleToBoxTM is calculated in GJKContactPoint - try to eliminate this one
+			const FRigidTransform3 CapsuleToBoxTM = ATransform.GetRelativeTransform(BTransform);
+			const FVec3 P1 = CapsuleToBoxTM.TransformPosition(A.GetX1());
+			const FVec3 P2 = CapsuleToBoxTM.TransformPosition(A.GetX2());
+			TAABB<T, 3> CapsuleAABB(P1.ComponentMin(P2), P1.ComponentMax(P2));
+			CapsuleAABB.Thicken(A.GetRadius() + CullDistance);
+			if (CapsuleAABB.Intersects(B))
+			{
+				const TVector<T, d> InitialDir = ATransform.GetRotation().Inverse() * -Constraint.GetNormal();
+				UpdateContactPoint(Constraint.Manifold, CapsuleBoxContactPoint(A, ATransform, B, BTransform, InitialDir, CullDistance));
+			}
 		}
 
 		/**
@@ -1312,47 +1424,46 @@ namespace Chaos
 		}
 
 		template<typename T, int d>
-		void ConstructCapsuleBoxConstraints(TGeometryParticleHandle<T, d>* Particle0, TGeometryParticleHandle<T, d>* Particle1, const FImplicitObject* Implicit0, const FImplicitObject* Implicit1, const TRigidTransform<T, d>& Transform0, const TRigidTransform<T, d>& Transform1, const T CullDistance, const FCollisionContext& Context, FCollisionConstraintsArray& NewConstraints)
+		void ConstructCapsuleBoxConstraints(TGeometryParticleHandle<T, d>* Particle0, TGeometryParticleHandle<T, d>* Particle1, const FImplicitObject* Implicit0, const FImplicitObject* Implicit1, const TRigidTransform<T, d>& LocalTransform0, const TRigidTransform<T, d>& LocalTransform1, const T CullDistance, const FCollisionContext& Context, FCollisionConstraintsArray& NewConstraints)
 		{
-			const TCapsule<T> * Object0 = Implicit0->template GetObject<const TCapsule<T> >();
-			const TBox<T, d> * Object1 = Implicit1->template GetObject<const TBox<T, d> >();
+			const TCapsule<T> * Object0 = Implicit0->template GetObject<const TCapsule<T>>();
+			const TBox<T, d> * Object1 = Implicit1->template GetObject<const TBox<T, d>>();
 			if (ensure(Object0 && Object1))
 			{
-				// Box-space AABB check
-				// @todo(chaos): avoid recalc of relative transform in GJK (will have to switch order of box/capsule)
-				const FRigidTransform3 CapsuleToBoxTM = Transform0.GetRelativeTransform(Transform1);
-				const FVec3 P1 = CapsuleToBoxTM.TransformPosition(Object0->GetX1());
-				const FVec3 P2 = CapsuleToBoxTM.TransformPosition(Object0->GetX2());
-				TAABB<T, 3> CapsuleAABB(P1.ComponentMin(P2), P1.ComponentMax(P2));
-				CapsuleAABB.Thicken(Object0->GetRadius() + CullDistance);
-				if (CapsuleAABB.Intersects(Object1->GetAABB()))
+				bool bAllowManifold = Chaos_Collision_UseManifolds;
+				if (bAllowManifold && (Chaos_Collision_CapsuleBoxManifoldTolerance > KINDA_SMALL_NUMBER))
 				{
-					TRigidTransform<T, d> ParticleImplicit0TM = Transform0.GetRelativeTransform(Collisions::GetTransform(Particle0));
-					TRigidTransform<T, d> ParticleImplicit1TM = Transform1.GetRelativeTransform(Collisions::GetTransform(Particle1));
-
-					bool bAllowManifold = Chaos_Collision_UseManifolds;
-
-					if (bAllowManifold && (Chaos_Collision_CapsuleBoxManifoldTolerance > KINDA_SMALL_NUMBER))
-					{
-						// @todo(ccaulfield): weak sauce - fix capsule-box manifolds.
-						// HACK: Disable manifolds for "horizontal" capsules. Manifolds don't work well when joints are pulling boxes down
-						// (under gravity) when the upper boxes are draped over a horizontal capsule. The box rotations about the manifold
-						// points(line) is too great and we end up with jitter.
-						const FVector CapsuleAxis = Context.SpaceTransform.TransformVectorNoScale(Transform0.TransformVectorNoScale(Object0->GetAxis()));
-						bAllowManifold = (FMath::Abs(CapsuleAxis.Z) > Chaos_Collision_CapsuleBoxManifoldTolerance);
-					}
+					// @todo(ccaulfield): weak sauce - fix capsule-box manifolds.
+					// HACK: Disable manifolds for "horizontal" capsules. Manifolds don't work well when joints are pulling boxes down
+					// (under gravity) when the upper boxes are draped over a horizontal capsule. The box rotations about the manifold
+					// points(line) is too great and we end up with jitter.
+					TRigidTransform<T, d> WorldTransform0 = LocalTransform0 * Collisions::GetTransform(Particle0);
+					const FVector CapsuleAxis = Context.SpaceTransform.TransformVectorNoScale(WorldTransform0.TransformVectorNoScale(Object0->GetAxis()));
+					bAllowManifold = (FMath::Abs(CapsuleAxis.Z) > Chaos_Collision_CapsuleBoxManifoldTolerance);
+				}
 
-					if (bAllowManifold)
+				if (bAllowManifold)
+				{
+					FRigidBodyMultiPointContactConstraint Constraint = FRigidBodyMultiPointContactConstraint(Particle0, Implicit0, LocalTransform0, Particle1, Implicit1, LocalTransform1, EContactShapesType::CapsuleBox);
+					// @todo(ccaulfield): manifold creation should be deferrable same as non-manifold below (only Context is preventing this - Apply does not know about the context yet)
+					TRigidTransform<T, d> WorldTransform0 = LocalTransform0 * Collisions::GetTransform(Particle0);
+					TRigidTransform<T, d> WorldTransform1 = LocalTransform1 * Collisions::GetTransform(Particle1);
+					UpdateCapsuleBoxManifold(*Object0, WorldTransform0, Object1->BoundingBox(), WorldTransform1, CullDistance, Context, Constraint);
+					NewConstraints.TryAdd(CullDistance, Constraint);
+				}
+				else
+				{
+					FRigidBodyPointContactConstraint Constraint = FRigidBodyPointContactConstraint(Particle0, Implicit0, LocalTransform0, Particle1, Implicit1, LocalTransform1, EContactShapesType::CapsuleBox);
+					if (!bChaos_Collision_DeferNarrowPhase)
 					{
-						FRigidBodyMultiPointContactConstraint Constraint = FRigidBodyMultiPointContactConstraint(Particle0, Implicit0, ParticleImplicit0TM, Particle1, Implicit1, ParticleImplicit1TM, EContactShapesType::CapsuleBox);
-						UpdateCapsuleBoxManifold(*Object0, Transform0, Object1->BoundingBox(), Transform1, CullDistance, Context, Constraint);
+						TRigidTransform<T, d> WorldTransform0 = LocalTransform0 * Collisions::GetTransform(Particle0);
+						TRigidTransform<T, d> WorldTransform1 = LocalTransform1 * Collisions::GetTransform(Particle1);
+						UpdateCapsuleBoxConstraint(*Object0, WorldTransform0, Object1->BoundingBox(), WorldTransform1, CullDistance, Constraint);
 						NewConstraints.TryAdd(CullDistance, Constraint);
 					}
 					else
 					{
-						FRigidBodyPointContactConstraint Constraint = FRigidBodyPointContactConstraint(Particle0, Implicit0, ParticleImplicit0TM, Particle1, Implicit1, ParticleImplicit1TM, EContactShapesType::CapsuleBox);
-						UpdateCapsuleBoxConstraint(*Object0, Transform0, Object1->BoundingBox(), Transform1, CullDistance, Constraint);
-						NewConstraints.TryAdd(CullDistance, Constraint);
+						NewConstraints.Add(Constraint);
 					}
 				}
 			}
@@ -1390,34 +1501,39 @@ namespace Chaos
 		}
 
 		template<typename T, int d>
-		void ConstructCapsuleHeightFieldConstraints(TGeometryParticleHandle<T, d>* Particle0, TGeometryParticleHandle<T, d>* Particle1, const FImplicitObject* Implicit0, const FImplicitObject* Implicit1, const TRigidTransform<T, d>& Transform0, const TRigidTransform<T, d>& Transform1, const T CullDistance, FCollisionConstraintsArray& NewConstraints)
+		void ConstructCapsuleHeightFieldConstraints(TGeometryParticleHandle<T, d>* Particle0, TGeometryParticleHandle<T, d>* Particle1, const FImplicitObject* Implicit0, const FImplicitObject* Implicit1, const TRigidTransform<T, d>& LocalTransform0, const TRigidTransform<T, d>& LocalTransform1, const T CullDistance, FCollisionConstraintsArray& NewConstraints)
 		{
 
 			const TCapsule<T> * Object0 = Implicit0->template GetObject<const TCapsule<T> >();
 			const THeightField<T> * Object1 = Implicit1->template GetObject<const THeightField<T> >();
 			if (ensure(Object0 && Object1))
 			{
-				TRigidTransform<T, d> ParticleImplicit0TM = Transform0.GetRelativeTransform(Collisions::GetTransform(Particle0));
-				TRigidTransform<T, d> ParticleImplicit1TM = Transform1.GetRelativeTransform(Collisions::GetTransform(Particle1));
-				FRigidBodyPointContactConstraint Constraint = FRigidBodyPointContactConstraint(Particle0, Implicit0, ParticleImplicit0TM, Particle1, Implicit1, ParticleImplicit1TM, EContactShapesType::CapsuleHeightField);
-				UpdateCapsuleHeightFieldConstraint(*Object0, Transform0, *Object1, Transform1, CullDistance, Constraint);
-				NewConstraints.TryAdd(CullDistance, Constraint);
+				FRigidBodyPointContactConstraint Constraint = FRigidBodyPointContactConstraint(Particle0, Implicit0, LocalTransform0, Particle1, Implicit1, LocalTransform1, EContactShapesType::CapsuleHeightField);
+				if (!bChaos_Collision_DeferNarrowPhase)
+				{
+					TRigidTransform<T, d> WorldTransform0 = LocalTransform0 * Collisions::GetTransform(Particle0);
+					TRigidTransform<T, d> WorldTransform1 = LocalTransform1 * Collisions::GetTransform(Particle1);
+					UpdateCapsuleHeightFieldConstraint(*Object0, WorldTransform0, *Object1, WorldTransform1, CullDistance, Constraint);
+					NewConstraints.TryAdd(CullDistance, Constraint);
+				}
+				else
+				{
+					NewConstraints.Add(Constraint);
+				}
 			}
 		}
 
 		template<typename T, int d>
-		void ConstructCapsuleHeightFieldConstraintsSwept(TGeometryParticleHandle<T, d>* Particle0, TGeometryParticleHandle<T, d>* Particle1, const FImplicitObject* Implicit0, const FImplicitObject* Implicit1, const TRigidTransform<T, d>& Transform0, const TRigidTransform<T, d>& Transform1, const T CullDistance, const TVector<T,d>& Dir, FReal Length, FCollisionConstraintsArray& NewConstraints)
+		void ConstructCapsuleHeightFieldConstraintsSwept(TGeometryParticleHandle<T, d>* Particle0, TGeometryParticleHandle<T, d>* Particle1, const FImplicitObject* Implicit0, const FImplicitObject* Implicit1, const TRigidTransform<T, d>& LocalTransform0, const TRigidTransform<T, d>& LocalTransform1, const T CullDistance, const TVector<T,d>& Dir, FReal Length, FCollisionConstraintsArray& NewConstraints)
 		{
 			const TCapsule<T> * Object0 = Implicit0->template GetObject<const TCapsule<T> >();
 			const THeightField<T> * Object1 = Implicit1->template GetObject<const THeightField<T> >();
 			if (ensure(Object0 && Object1))
 			{
-				const TRigidTransform<T, 3> TransformX0(Particle0->X(), Particle0->R());
-				TRigidTransform<T, d> ParticleImplicit0TM = Transform0.GetRelativeTransform(Collisions::GetTransform(Particle0));
-				TRigidTransform<T, d> ParticleImplicit1TM = Transform1.GetRelativeTransform(Collisions::GetTransform(Particle1));
-
-				FRigidBodySweptPointContactConstraint Constraint = FRigidBodySweptPointContactConstraint(Particle0, Implicit0, ParticleImplicit0TM, Particle1, Implicit1, ParticleImplicit1TM, EContactShapesType::CapsuleHeightField);
-				UpdateCapsuleHeightFieldConstraintSwept(Particle0, *Object0, TransformX0, *Object1, Transform1, Dir, Length, CullDistance, Constraint);
+				FRigidBodySweptPointContactConstraint Constraint = FRigidBodySweptPointContactConstraint(Particle0, Implicit0, LocalTransform0, Particle1, Implicit1, LocalTransform1, EContactShapesType::CapsuleHeightField);
+				TRigidTransform<T, 3> WorldTransformX0 = LocalTransform0 * FRigidTransform3(Particle0->X(), Particle0->R());
+				TRigidTransform<T, d> WorldTransform1 = LocalTransform1 * Collisions::GetTransform(Particle1);
+				UpdateCapsuleHeightFieldConstraintSwept(Particle0, *Object0, WorldTransformX0, *Object1, WorldTransform1, Dir, Length, CullDistance, Constraint);
 				NewConstraints.TryAdd(CullDistance, Constraint);
 			}
 		}
@@ -1451,16 +1567,16 @@ namespace Chaos
 
 
 		template <typename TriMeshType, typename T, int d>
-		void UpdateCapsuleTriangleMeshConstraint(const TCapsule<T>& Capsule0, const TRigidTransform<T, d>& Transform0, const TriMeshType& TriangleMesh1, const TRigidTransform<T, d>& Transform1, const T CullDistance, TRigidBodyPointContactConstraint<T, d>& Constraint)
+		void UpdateCapsuleTriangleMeshConstraint(const TCapsule<T>& Capsule0, const TRigidTransform<T, d>& WorldTransform0, const TriMeshType& TriangleMesh1, const TRigidTransform<T, d>& WorldTransform1, const T CullDistance, TRigidBodyPointContactConstraint<T, d>& Constraint)
 		{
-			UpdateContactPoint(Constraint.Manifold, CapsuleTriangleMeshContactPoint(Capsule0, Transform0, TriangleMesh1, Transform1, CullDistance));
+			UpdateContactPoint(Constraint.Manifold, CapsuleTriangleMeshContactPoint(Capsule0, WorldTransform0, TriangleMesh1, WorldTransform1, CullDistance));
 		}
 
 		template <typename TriMeshType, typename T, int d>
-		void UpdateCapsuleTriangleMeshConstraintSwept(TGeometryParticleHandle<T, d>* Particle0, const TCapsule<T>& Capsule0, const TRigidTransform<T, d>& Transform0, const TriMeshType& TriangleMesh1, const TRigidTransform<T, d>& Transform1, const TVector<T, d>& Dir, const T Length, const T CullDistance, TRigidBodySweptPointContactConstraint<T, d>& Constraint)
+		void UpdateCapsuleTriangleMeshConstraintSwept(TGeometryParticleHandle<T, d>* Particle0, const TCapsule<T>& Capsule0, const TRigidTransform<T, d>& WorldTransform0, const TriMeshType& TriangleMesh1, const TRigidTransform<T, d>& WorldTransform1, const TVector<T, d>& Dir, const T Length, const T CullDistance, TRigidBodySweptPointContactConstraint<T, d>& Constraint)
 		{
 			T TOI = 1.0f;
-			UpdateContactPoint(Constraint.Manifold, CapsuleTriangleMeshSweptContactPoint(Capsule0, Transform0, TriangleMesh1, Transform1, Dir, Length, CullDistance, TOI));
+			UpdateContactPoint(Constraint.Manifold, CapsuleTriangleMeshSweptContactPoint(Capsule0, WorldTransform0, TriangleMesh1, WorldTransform1, Dir, Length, CullDistance, TOI));
 			SetSweptConstraintTOI(Particle0, TOI, Length, Dir, Constraint);
 		}
 
@@ -1471,26 +1587,40 @@ namespace Chaos
 		}
 
 		template<typename T, int d>
-		void ConstructCapsuleTriangleMeshConstraints(TGeometryParticleHandle<T, d>* Particle0, TGeometryParticleHandle<T, d>* Particle1, const FImplicitObject* Implicit0, const FImplicitObject* Implicit1, const TRigidTransform<T, d>& Transform0, const TRigidTransform<T, d>& Transform1, const T CullDistance, FCollisionConstraintsArray& NewConstraints)
+		void ConstructCapsuleTriangleMeshConstraints(TGeometryParticleHandle<T, d>* Particle0, TGeometryParticleHandle<T, d>* Particle1, const FImplicitObject* Implicit0, const FImplicitObject* Implicit1, const TRigidTransform<T, d>& LocalTransform0, const TRigidTransform<T, d>& LocalTransform1, const T CullDistance, FCollisionConstraintsArray& NewConstraints)
 		{
 			const TCapsule<T> * Object0 = Implicit0->template GetObject<const TCapsule<T> >();
 			if (ensure(Object0))
 			{
 				if (const TImplicitObjectScaled<FTriangleMeshImplicitObject>* ScaledTriangleMesh = Implicit1->template GetObject<const TImplicitObjectScaled<FTriangleMeshImplicitObject>>())
 				{
-					TRigidTransform<T, d> ParticleImplicit0TM = Transform0.GetRelativeTransform(Collisions::GetTransform(Particle0));
-					TRigidTransform<T, d> ParticleImplicit1TM = Transform1.GetRelativeTransform(Collisions::GetTransform(Particle1));
-					FRigidBodyPointContactConstraint Constraint = FRigidBodyPointContactConstraint(Particle0, Implicit0, ParticleImplicit0TM, Particle1, Implicit1, ParticleImplicit1TM, EContactShapesType::CapsuleTriMesh);
-					UpdateCapsuleTriangleMeshConstraint(*Object0, Transform0, *ScaledTriangleMesh, Transform1, CullDistance, Constraint);
-					NewConstraints.TryAdd(CullDistance, Constraint);
+					FRigidBodyPointContactConstraint Constraint = FRigidBodyPointContactConstraint(Particle0, Implicit0, LocalTransform0, Particle1, Implicit1, LocalTransform1, EContactShapesType::CapsuleTriMesh);
+					if (!bChaos_Collision_DeferNarrowPhase)
+					{
+						TRigidTransform<T, d> WorldTransform0 = LocalTransform0 * Collisions::GetTransform(Particle0);
+						TRigidTransform<T, d> WorldTransform1 = LocalTransform1 * Collisions::GetTransform(Particle1);
+						UpdateCapsuleTriangleMeshConstraint(*Object0, WorldTransform0, *ScaledTriangleMesh, WorldTransform1, CullDistance, Constraint);
+						NewConstraints.TryAdd(CullDistance, Constraint);
+					}
+					else
+					{
+						NewConstraints.Add(Constraint);
+					}
 				}
 				else if (const FTriangleMeshImplicitObject* TriangleMesh = Implicit1->template GetObject<const FTriangleMeshImplicitObject>())
 				{
-					TRigidTransform<T, d> ParticleImplicit0TM = Transform0.GetRelativeTransform(Collisions::GetTransform(Particle0));
-					TRigidTransform<T, d> ParticleImplicit1TM = Transform1.GetRelativeTransform(Collisions::GetTransform(Particle1));
-					FRigidBodyPointContactConstraint Constraint = FRigidBodyPointContactConstraint(Particle0, Implicit0, ParticleImplicit0TM, Particle1, Implicit1, ParticleImplicit1TM, EContactShapesType::CapsuleTriMesh);
-					UpdateCapsuleTriangleMeshConstraint(*Object0, Transform0, *TriangleMesh, Transform1, CullDistance, Constraint);
-					NewConstraints.TryAdd(CullDistance, Constraint);
+					FRigidBodyPointContactConstraint Constraint = FRigidBodyPointContactConstraint(Particle0, Implicit0, LocalTransform0, Particle1, Implicit1, LocalTransform1, EContactShapesType::CapsuleTriMesh);
+					if (!bChaos_Collision_DeferNarrowPhase)
+					{
+						TRigidTransform<T, d> WorldTransform0 = LocalTransform0 * Collisions::GetTransform(Particle0);
+						TRigidTransform<T, d> WorldTransform1 = LocalTransform1 * Collisions::GetTransform(Particle1);
+						UpdateCapsuleTriangleMeshConstraint(*Object0, WorldTransform0, *TriangleMesh, WorldTransform1, CullDistance, Constraint);
+						NewConstraints.TryAdd(CullDistance, Constraint);
+					}
+					else
+					{
+						NewConstraints.Add(Constraint);
+					}
 				}
 				else
 				{
@@ -1500,25 +1630,25 @@ namespace Chaos
 		}
 
 		template<typename T, int d>
-		void ConstructCapsuleTriangleMeshConstraintsSwept(TGeometryParticleHandle<T, d>* Particle0, TGeometryParticleHandle<T, d>* Particle1, const FImplicitObject* Implicit0, const FImplicitObject* Implicit1, const TRigidTransform<T, d>& Transform0, const TRigidTransform<T, d>& Transform1, const T CullDistance, const TVector<T,d>& Dir, FReal Length, FCollisionConstraintsArray& NewConstraints)
+		void ConstructCapsuleTriangleMeshConstraintsSwept(TGeometryParticleHandle<T, d>* Particle0, TGeometryParticleHandle<T, d>* Particle1, const FImplicitObject* Implicit0, const FImplicitObject* Implicit1, const TRigidTransform<T, d>& LocalTransform0, const TRigidTransform<T, d>& LocalTransform1, const T CullDistance, const TVector<T,d>& Dir, FReal Length, FCollisionConstraintsArray& NewConstraints)
 		{
 			const TCapsule<T> * Object0 = Implicit0->template GetObject<const TCapsule<T> >();
 			if (ensure(Object0))
 			{
-				TRigidTransform<T, 3> TransformX0(Particle0->X(), Particle0->R());
-				TRigidTransform<T, d> ParticleImplicit0TM = Transform0.GetRelativeTransform(Collisions::GetTransform(Particle0));
-				TRigidTransform<T, d> ParticleImplicit1TM = Transform1.GetRelativeTransform(Collisions::GetTransform(Particle1));
-
 				if (const TImplicitObjectScaled<FTriangleMeshImplicitObject>* ScaledTriangleMesh = Implicit1->template GetObject<const TImplicitObjectScaled<FTriangleMeshImplicitObject>>())
 				{
-					FRigidBodySweptPointContactConstraint Constraint = FRigidBodySweptPointContactConstraint(Particle0, Implicit0, ParticleImplicit0TM, Particle1, Implicit1, ParticleImplicit1TM, EContactShapesType::CapsuleTriMesh);
-					UpdateCapsuleTriangleMeshConstraintSwept(Particle0, *Object0, TransformX0, *ScaledTriangleMesh, Transform1, Dir, Length, CullDistance, Constraint);
+					FRigidBodySweptPointContactConstraint Constraint = FRigidBodySweptPointContactConstraint(Particle0, Implicit0, LocalTransform0, Particle1, Implicit1, LocalTransform1, EContactShapesType::CapsuleTriMesh);
+					TRigidTransform<T, 3> WorldTransformX0 = LocalTransform0 * FRigidTransform3(Particle0->X(), Particle0->R());
+					TRigidTransform<T, d> WorldTransform1 = LocalTransform1 * Collisions::GetTransform(Particle1);
+					UpdateCapsuleTriangleMeshConstraintSwept(Particle0, *Object0, WorldTransformX0, *ScaledTriangleMesh, WorldTransform1, Dir, Length, CullDistance, Constraint);
 					NewConstraints.TryAdd(CullDistance, Constraint);
 				}
 				else if (const FTriangleMeshImplicitObject* TriangleMesh = Implicit1->template GetObject<const FTriangleMeshImplicitObject>())
 				{
-					FRigidBodySweptPointContactConstraint Constraint = FRigidBodySweptPointContactConstraint(Particle0, Implicit0, ParticleImplicit0TM, Particle1, Implicit1, ParticleImplicit1TM, EContactShapesType::CapsuleTriMesh);
-					UpdateCapsuleTriangleMeshConstraintSwept(Particle0, *Object0, TransformX0, *TriangleMesh, Transform1, Dir, Length, CullDistance, Constraint);
+					FRigidBodySweptPointContactConstraint Constraint = FRigidBodySweptPointContactConstraint(Particle0, Implicit0, LocalTransform0, Particle1, Implicit1, LocalTransform1, EContactShapesType::CapsuleTriMesh);
+					TRigidTransform<T, 3> WorldTransformX0 = LocalTransform0 * FRigidTransform3(Particle0->X(), Particle0->R());
+					TRigidTransform<T, d> WorldTransform1 = LocalTransform1 * Collisions::GetTransform(Particle1);
+					UpdateCapsuleTriangleMeshConstraintSwept(Particle0, *Object0, WorldTransformX0, *TriangleMesh, WorldTransform1, Dir, Length, CullDistance, Constraint);
 					NewConstraints.TryAdd(CullDistance, Constraint);
 				}
 				else
@@ -1533,24 +1663,24 @@ namespace Chaos
 		//
 
 		template<class T, int d>
-		void UpdateConvexConvexConstraint(const FImplicitObject& Implicit0, const TRigidTransform<T, d>& Transform0, const FImplicitObject& Implicit1, const TRigidTransform<T, d>& Transform1, const T CullDistance, TCollisionConstraintBase<T, d>& ConstraintBase)
+		void UpdateConvexConvexConstraint(const FImplicitObject& Implicit0, const TRigidTransform<T, d>& WorldTransform0, const FImplicitObject& Implicit1, const TRigidTransform<T, d>& WorldTransform1, const T CullDistance, TCollisionConstraintBase<T, d>& ConstraintBase)
 		{
 			TContactPoint<T> ContactPoint;
 
 			if (ConstraintBase.GetType() == FRigidBodyPointContactConstraint::StaticType())
 			{
-				ContactPoint = ConvexConvexContactPoint(Implicit0, Transform0, Implicit1, Transform1, CullDistance);
+				ContactPoint = ConvexConvexContactPoint(Implicit0, WorldTransform0, Implicit1, WorldTransform1, CullDistance);
 			}
 			else if (ConstraintBase.GetType() == FRigidBodySweptPointContactConstraint::StaticType())
 			{
-				ContactPoint = ConvexConvexContactPoint(Implicit0, Transform0, Implicit1, Transform1, CullDistance);
+				ContactPoint = ConvexConvexContactPoint(Implicit0, WorldTransform0, Implicit1, WorldTransform1, CullDistance);
 			}
 			else if (ConstraintBase.GetType() == FRigidBodyMultiPointContactConstraint::StaticType())
 			{
 				TRigidBodyMultiPointContactConstraint<T, d>& Constraint = *ConstraintBase.template As<TRigidBodyMultiPointContactConstraint<T, d>>();
 				ContactPoint.Phi = FLT_MAX;
 
-				const TRigidTransform<T, d> AToBTM = Transform0.GetRelativeTransform(Transform1);
+				const TRigidTransform<T, d> AToBTM = WorldTransform0.GetRelativeTransform(WorldTransform1);
 
 				TPlane<T, d> CollisionPlane(Constraint.GetManifoldPlanePosition(), Constraint.GetManifoldPlaneNormal());
 
@@ -1565,8 +1695,8 @@ namespace Chaos
 					if (ContactPoint.Phi > Phi)
 					{
 						ContactPoint.Phi = Phi;
-						ContactPoint.Normal = Transform1.TransformVectorNoScale(Constraint.GetManifoldPlaneNormal());
-						ContactPoint.Location = Transform0.TransformPosition(Location);
+						ContactPoint.Normal = WorldTransform1.TransformVectorNoScale(Constraint.GetManifoldPlaneNormal());
+						ContactPoint.Location = WorldTransform0.TransformPosition(Location);
 					}
 				}
 			}
@@ -1575,38 +1705,34 @@ namespace Chaos
 		}
 
 		template<class T, int d>
-		void UpdateConvexConvexConstraintSwept(TGeometryParticleHandle<T, d>* Particle0, const FImplicitObject& Implicit0, const TRigidTransform<T, d>& Transform0, const FImplicitObject& Implicit1, const TRigidTransform<T, d>& Transform1, const TVector<T, d>& Dir, const T Length, const T CullDistance, TRigidBodySweptPointContactConstraint<T, d>& Constraint)
+		void UpdateConvexConvexConstraintSwept(TGeometryParticleHandle<T, d>* Particle0, const FImplicitObject& Implicit0, const TRigidTransform<T, d>& WorldTransform0, const FImplicitObject& Implicit1, const TRigidTransform<T, d>& WorldTransform1, const TVector<T, d>& Dir, const T Length, const T CullDistance, TRigidBodySweptPointContactConstraint<T, d>& Constraint)
 		{
 			T TOI = 1.0f;
-			UpdateContactPoint(Constraint.Manifold, ConvexConvexContactPointSwept(Implicit0, Transform0, Implicit1, Transform1, Dir, Length, CullDistance, TOI));
+			UpdateContactPoint(Constraint.Manifold, ConvexConvexContactPointSwept(Implicit0, WorldTransform0, Implicit1, WorldTransform1, Dir, Length, CullDistance, TOI));
 			SetSweptConstraintTOI(Particle0, TOI, Length, Dir, Constraint);
 		}
 
 		template<class T, int d>
-		void UpdateConvexConvexManifold(TCollisionConstraintBase<T, d>&  ConstraintBase, const TRigidTransform<T, d>& Transform0, const TRigidTransform<T, d>& Transform1, const T CullDistance)
+		void UpdateConvexConvexManifold(TCollisionConstraintBase<T, d>&  ConstraintBase, const TRigidTransform<T, d>& WorldTransform0, const TRigidTransform<T, d>& WorldTransform1, const T CullDistance)
 		{
 			if (ConstraintBase.GetType() == FRigidBodyMultiPointContactConstraint::StaticType())
 			{
 				TRigidBodyMultiPointContactConstraint<T, d>* Constraint = ConstraintBase.template As<TRigidBodyMultiPointContactConstraint<T, d>>();
 				if (GetInnerType(ConstraintBase.Manifold.Implicit[0]->GetType()) == ImplicitObjectType::Convex)
 				{
-					UpdateSingleShotManifold(*Constraint, Transform0, Transform1, CullDistance);
+					UpdateSingleShotManifold(*Constraint, WorldTransform0, WorldTransform1, CullDistance);
 				}
 				else
 				{
-					UpdateIterativeManifold(*Constraint, Transform0, Transform1, CullDistance);
+					UpdateIterativeManifold(*Constraint, WorldTransform0, WorldTransform1, CullDistance);
 				}
 			}
 		}
 
 
 		template<class T, int d>
-		void ConstructConvexConvexConstraints(TGeometryParticleHandle<T, d>* Particle0, TGeometryParticleHandle<T, d>* Particle1, const FImplicitObject* Implicit0, const FImplicitObject* Implicit1, const TRigidTransform<T, d>& Transform0, const TRigidTransform<T, d>& Transform1, const T CullDistance, FCollisionConstraintsArray& NewConstraints)
+		void ConstructConvexConvexConstraints(TGeometryParticleHandle<T, d>* Particle0, TGeometryParticleHandle<T, d>* Particle1, const FImplicitObject* Implicit0, const FImplicitObject* Implicit1, const TRigidTransform<T, d>& LocalTransform0, const TRigidTransform<T, d>& LocalTransform1, const T CullDistance, FCollisionConstraintsArray& NewConstraints)
 		{
-			TRigidTransform<T, d> ParticleImplicit0TM = Transform0.GetRelativeTransform(Collisions::GetTransform(Particle0));
-			TRigidTransform<T, d> ParticleImplicit1TM = Transform1.GetRelativeTransform(Collisions::GetTransform(Particle1));
-
-
 			EImplicitObjectType Implicit0Type = Particle0->Geometry()->GetType();
 			EImplicitObjectType Implicit1Type = Particle1->Geometry()->GetType();
 
@@ -1614,28 +1740,44 @@ namespace Chaos
 			// We need to improve iterative manifolds for this case
 			if (Chaos_Collision_UseManifolds && Implicit0Type != TBox<FReal, 3>::StaticType() && Implicit1Type != TBox<FReal, 3>::StaticType())
 			{
-				FRigidBodyMultiPointContactConstraint Constraint = FRigidBodyMultiPointContactConstraint(Particle0, Implicit0, ParticleImplicit0TM, Particle1, Implicit1, ParticleImplicit1TM, EContactShapesType::ConvexConvex);
-				UpdateConvexConvexManifold(Constraint, Transform0, Transform1, CullDistance);
-				UpdateConvexConvexConstraint(*Implicit0, Transform0, *Implicit1, Transform1, CullDistance, Constraint);
-				NewConstraints.TryAdd(CullDistance, Constraint);
+				FRigidBodyMultiPointContactConstraint Constraint = FRigidBodyMultiPointContactConstraint(Particle0, Implicit0, LocalTransform0, Particle1, Implicit1, LocalTransform1, EContactShapesType::ConvexConvex);
+				TRigidTransform<T, d> WorldTransform0 = LocalTransform0 * Collisions::GetTransform(Particle0);
+				TRigidTransform<T, d> WorldTransform1 = LocalTransform1 * Collisions::GetTransform(Particle1);
+				UpdateConvexConvexManifold(Constraint, WorldTransform0, WorldTransform1, CullDistance);
+				if (!bChaos_Collision_DeferNarrowPhase)
+				{
+					UpdateConvexConvexConstraint(*Implicit0, WorldTransform0, *Implicit1, WorldTransform1, CullDistance, Constraint);
+					NewConstraints.TryAdd(CullDistance, Constraint);
+				}
+				else
+				{
+					NewConstraints.Add(Constraint);
+				}
 			}
 			else
 			{
-				FRigidBodyPointContactConstraint Constraint = FRigidBodyPointContactConstraint(Particle0, Implicit0, ParticleImplicit0TM, Particle1, Implicit1, ParticleImplicit1TM, EContactShapesType::ConvexConvex);
-				UpdateConvexConvexConstraint(*Implicit0, Transform0, *Implicit1, Transform1, CullDistance, Constraint);
-				NewConstraints.TryAdd(CullDistance, Constraint);
+				FRigidBodyPointContactConstraint Constraint = FRigidBodyPointContactConstraint(Particle0, Implicit0, LocalTransform0, Particle1, Implicit1, LocalTransform1, EContactShapesType::ConvexConvex);
+				if (!bChaos_Collision_DeferNarrowPhase)
+				{
+					TRigidTransform<T, d> WorldTransform0 = LocalTransform0 * Collisions::GetTransform(Particle0);
+					TRigidTransform<T, d> WorldTransform1 = LocalTransform1 * Collisions::GetTransform(Particle1);
+					UpdateConvexConvexConstraint(*Implicit0, WorldTransform0, *Implicit1, WorldTransform1, CullDistance, Constraint);
+					NewConstraints.TryAdd(CullDistance, Constraint);
+				}
+				else
+				{
+					NewConstraints.Add(Constraint);
+				}
 			}
 		}
 
 		template<class T, int d>
-		void ConstructConvexConvexConstraintsSwept(TGeometryParticleHandle<T, d>* Particle0, TGeometryParticleHandle<T, d>* Particle1, const FImplicitObject* Implicit0, const FImplicitObject* Implicit1, const TRigidTransform<T, d>& Transform0, const TRigidTransform<T, d>& Transform1, const T CullDistance, const TVector<T,d>& Dir, FReal Length, FCollisionConstraintsArray& NewConstraints)
+		void ConstructConvexConvexConstraintsSwept(TGeometryParticleHandle<T, d>* Particle0, TGeometryParticleHandle<T, d>* Particle1, const FImplicitObject* Implicit0, const FImplicitObject* Implicit1, const TRigidTransform<T, d>& LocalTransform0, const TRigidTransform<T, d>& LocalTransform1, const T CullDistance, const TVector<T,d>& Dir, FReal Length, FCollisionConstraintsArray& NewConstraints)
 		{
-			TRigidTransform<T, 3> TransformX0(Particle0->X(), Particle0->R());
-			TRigidTransform<T, d> ParticleImplicit0TM = Transform0.GetRelativeTransform(Collisions::GetTransform(Particle0));
-			TRigidTransform<T, d> ParticleImplicit1TM = Transform1.GetRelativeTransform(Collisions::GetTransform(Particle1));
-
-			FRigidBodySweptPointContactConstraint Constraint = FRigidBodySweptPointContactConstraint(Particle0, Implicit0, ParticleImplicit0TM, Particle1, Implicit1, ParticleImplicit1TM, EContactShapesType::ConvexConvex);
-			UpdateConvexConvexConstraintSwept(Particle0, *Implicit0, TransformX0, *Implicit1, Transform1, Dir, Length, CullDistance, Constraint);
+			FRigidBodySweptPointContactConstraint Constraint = FRigidBodySweptPointContactConstraint(Particle0, Implicit0, LocalTransform0, Particle1, Implicit1, LocalTransform1, EContactShapesType::ConvexConvex);
+			TRigidTransform<T, 3> WorldTransformX0 = LocalTransform0 * FRigidTransform3(Particle0->X(), Particle0->R());
+			TRigidTransform<T, d> WorldTransform1 = LocalTransform1 * Collisions::GetTransform(Particle1);
+			UpdateConvexConvexConstraintSwept(Particle0, *Implicit0, WorldTransformX0, *Implicit1, WorldTransform1, Dir, Length, CullDistance, Constraint);
 			NewConstraints.TryAdd(CullDistance, Constraint);
 		}
 
@@ -1671,31 +1813,36 @@ namespace Chaos
 		}
 
 		template<typename T, int d>
-		void ConstructConvexHeightFieldConstraints(TGeometryParticleHandle<T, d>* Particle0, TGeometryParticleHandle<T, d>* Particle1, const FImplicitObject* Implicit0, const FImplicitObject* Implicit1, const TRigidTransform<T, d>& Transform0, const TRigidTransform<T, d>& Transform1, const T CullDistance, FCollisionConstraintsArray& NewConstraints)
+		void ConstructConvexHeightFieldConstraints(TGeometryParticleHandle<T, d>* Particle0, TGeometryParticleHandle<T, d>* Particle1, const FImplicitObject* Implicit0, const FImplicitObject* Implicit1, const TRigidTransform<T, d>& LocalTransform0, const TRigidTransform<T, d>& LocalTransform1, const T CullDistance, FCollisionConstraintsArray& NewConstraints)
 		{
 			const THeightField<T> * Object1 = Implicit1->template GetObject<const THeightField<T> >();
 			if (ensure(Implicit0->IsConvex() && Object1))
 			{
-				TRigidTransform<T, d> ParticleImplicit0TM = Transform0.GetRelativeTransform(Collisions::GetTransform(Particle0));
-				TRigidTransform<T, d> ParticleImplicit1TM = Transform1.GetRelativeTransform(Collisions::GetTransform(Particle1));
-				FRigidBodyPointContactConstraint Constraint = FRigidBodyPointContactConstraint(Particle0, Implicit0, ParticleImplicit0TM, Particle1, Implicit1, ParticleImplicit1TM, EContactShapesType::ConvexHeightField);
-				UpdateConvexHeightFieldConstraint(*Implicit0, Transform0, *Object1, Transform1, CullDistance, Constraint);
-				NewConstraints.TryAdd(CullDistance, Constraint);
+				FRigidBodyPointContactConstraint Constraint = FRigidBodyPointContactConstraint(Particle0, Implicit0, LocalTransform0, Particle1, Implicit1, LocalTransform1, EContactShapesType::ConvexHeightField);
+				if (!bChaos_Collision_DeferNarrowPhase)
+				{
+					TRigidTransform<T, d> WorldTransform0 = LocalTransform0 * Collisions::GetTransform(Particle0);
+					TRigidTransform<T, d> WorldTransform1 = LocalTransform1 * Collisions::GetTransform(Particle1);
+					UpdateConvexHeightFieldConstraint(*Implicit0, WorldTransform0, *Object1, WorldTransform1, CullDistance, Constraint);
+					NewConstraints.TryAdd(CullDistance, Constraint);
+				}
+				else
+				{
+					NewConstraints.Add(Constraint);
+				}
 			}
 		}
 
 		template<typename T, int d>
-		void ConstructConvexHeightFieldConstraintsSwept(TGeometryParticleHandle<T, d>* Particle0, TGeometryParticleHandle<T, d>* Particle1, const FImplicitObject* Implicit0, const FImplicitObject* Implicit1, const TRigidTransform<T, d>& Transform0, const TRigidTransform<T, d>& Transform1, const T CullDistance, const TVector<T, d>& Dir, FReal Length, FCollisionConstraintsArray& NewConstraints)
+		void ConstructConvexHeightFieldConstraintsSwept(TGeometryParticleHandle<T, d>* Particle0, TGeometryParticleHandle<T, d>* Particle1, const FImplicitObject* Implicit0, const FImplicitObject* Implicit1, const TRigidTransform<T, d>& LocalTransform0, const TRigidTransform<T, d>& LocalTransform1, const T CullDistance, const TVector<T, d>& Dir, FReal Length, FCollisionConstraintsArray& NewConstraints)
 		{
 			const THeightField<T> * Object1 = Implicit1->template GetObject<const THeightField<T> >();
 			if (ensure(Implicit0->IsConvex() && Object1))
 			{
-				TRigidTransform<T, 3> TransformX0(Particle0->X(), Particle0->R());
-				TRigidTransform<T, d> ParticleImplicit0TM = Transform0.GetRelativeTransform(Collisions::GetTransform(Particle0));
-				TRigidTransform<T, d> ParticleImplicit1TM = Transform1.GetRelativeTransform(Collisions::GetTransform(Particle1));
-
-				FRigidBodySweptPointContactConstraint Constraint = FRigidBodySweptPointContactConstraint(Particle0, Implicit0, ParticleImplicit0TM, Particle1, Implicit1, ParticleImplicit1TM, EContactShapesType::ConvexHeightField);
-				UpdateConvexHeightFieldConstraintSwept(Particle0, *Implicit0, TransformX0, *Object1, Transform1, Dir, Length, CullDistance, Constraint);
+				FRigidBodySweptPointContactConstraint Constraint = FRigidBodySweptPointContactConstraint(Particle0, Implicit0, LocalTransform0, Particle1, Implicit1, LocalTransform1, EContactShapesType::ConvexHeightField);
+				TRigidTransform<T, 3> WorldTransformX0 = LocalTransform0 * FRigidTransform3(Particle0->X(), Particle0->R());
+				TRigidTransform<T, d> WorldTransform1 = LocalTransform1 * Collisions::GetTransform(Particle1);
+				UpdateConvexHeightFieldConstraintSwept(Particle0, *Implicit0, WorldTransformX0, *Object1, WorldTransform1, Dir, Length, CullDistance, Constraint);
 				NewConstraints.TryAdd(CullDistance, Constraint);
 			}
 		}
@@ -1727,17 +1874,17 @@ namespace Chaos
 		}
 		
 		template <typename TriMeshType, typename T, int d>
-		void UpdateConvexTriangleMeshConstraint(const FImplicitObject& Convex0, const TRigidTransform<T, d>& Transform0, const TriMeshType& TriangleMesh1, const TRigidTransform<T, d>& Transform1, const T CullDistance, TRigidBodyPointContactConstraint<T, d>& Constraint)
+		void UpdateConvexTriangleMeshConstraint(const FImplicitObject& Convex0, const TRigidTransform<T, d>& WorldTransform0, const TriMeshType& TriangleMesh1, const TRigidTransform<T, d>& WorldTransform1, const T CullDistance, TRigidBodyPointContactConstraint<T, d>& Constraint)
 		{
-			UpdateContactPoint(Constraint.Manifold, ConvexTriangleMeshContactPoint(Convex0, Transform0, TriangleMesh1, Transform1, CullDistance));
+			UpdateContactPoint(Constraint.Manifold, ConvexTriangleMeshContactPoint(Convex0, WorldTransform0, TriangleMesh1, WorldTransform1, CullDistance));
 		}
 
 		// Sweeps convex against trimesh
 		template <typename TriMeshType, typename T, int d>
-		void UpdateConvexTriangleMeshConstraintSwept(TGeometryParticleHandle<T,d>* Particle0, const FImplicitObject& Convex0, const TRigidTransform<T, d>& Transform0, const TriMeshType& TriangleMesh1, const TRigidTransform<T, d>& Transform1, const TVector<T, d>& Dir, const T Length, const T CullDistance, TRigidBodySweptPointContactConstraint<T, d>& Constraint)
+		void UpdateConvexTriangleMeshConstraintSwept(TGeometryParticleHandle<T,d>* Particle0, const FImplicitObject& Convex0, const TRigidTransform<T, d>& WorldTransform0, const TriMeshType& TriangleMesh1, const TRigidTransform<T, d>& WorldTransform1, const TVector<T, d>& Dir, const T Length, const T CullDistance, TRigidBodySweptPointContactConstraint<T, d>& Constraint)
 		{
 			T TOI = 1.0f;
-			UpdateContactPoint(Constraint.Manifold, ConvexTriangleMeshSweptContactPoint(Convex0, Transform0, TriangleMesh1, Transform1, Dir, Length, CullDistance, TOI));
+			UpdateContactPoint(Constraint.Manifold, ConvexTriangleMeshSweptContactPoint(Convex0, WorldTransform0, TriangleMesh1, WorldTransform1, Dir, Length, CullDistance, TOI));
 			SetSweptConstraintTOI(Particle0, TOI, Length, Dir, Constraint);
 		}
 
@@ -1748,25 +1895,39 @@ namespace Chaos
 		}
 
 		template<typename T, int d>
-		void ConstructConvexTriangleMeshConstraints(TGeometryParticleHandle<T, d>* Particle0, TGeometryParticleHandle<T, d>* Particle1, const FImplicitObject* Implicit0, const FImplicitObject* Implicit1, const TRigidTransform<T, d>& Transform0, const TRigidTransform<T, d>& Transform1, const T CullDistance, FCollisionConstraintsArray& NewConstraints)
+		void ConstructConvexTriangleMeshConstraints(TGeometryParticleHandle<T, d>* Particle0, TGeometryParticleHandle<T, d>* Particle1, const FImplicitObject* Implicit0, const FImplicitObject* Implicit1, const TRigidTransform<T, d>& LocalTransform0, const TRigidTransform<T, d>& LocalTransform1, const T CullDistance, FCollisionConstraintsArray& NewConstraints)
 		{
 			if (ensure(Implicit0->IsConvex()))
 			{
 				if (const TImplicitObjectScaled<FTriangleMeshImplicitObject>* ScaledTriangleMesh = Implicit1->template GetObject<const TImplicitObjectScaled<FTriangleMeshImplicitObject>>())
 				{
-					TRigidTransform<T, d> ParticleImplicit0TM = Transform0.GetRelativeTransform(Collisions::GetTransform(Particle0));
-					TRigidTransform<T, d> ParticleImplicit1TM = Transform1.GetRelativeTransform(Collisions::GetTransform(Particle1));
-					FRigidBodyPointContactConstraint Constraint = FRigidBodyPointContactConstraint(Particle0, Implicit0, ParticleImplicit0TM, Particle1, Implicit1, ParticleImplicit1TM, EContactShapesType::ConvexTriMesh);
-					UpdateConvexTriangleMeshConstraint(*Implicit0, Transform0, *ScaledTriangleMesh , Transform1, CullDistance, Constraint);
-					NewConstraints.TryAdd(CullDistance, Constraint);
+					FRigidBodyPointContactConstraint Constraint = FRigidBodyPointContactConstraint(Particle0, Implicit0, LocalTransform0, Particle1, Implicit1, LocalTransform0, EContactShapesType::ConvexTriMesh);
+					if (!bChaos_Collision_DeferNarrowPhase)
+					{
+						TRigidTransform<T, d> WorldTransform0 = LocalTransform0 * Collisions::GetTransform(Particle0);
+						TRigidTransform<T, d> WorldTransform1 = LocalTransform1 * Collisions::GetTransform(Particle1);
+						UpdateConvexTriangleMeshConstraint(*Implicit0, WorldTransform0, *ScaledTriangleMesh , WorldTransform1, CullDistance, Constraint);
+						NewConstraints.TryAdd(CullDistance, Constraint);
+					}
+					else
+					{
+						NewConstraints.Add(Constraint);
+					}
 				}
 				else if (const FTriangleMeshImplicitObject* TriangleMesh = Implicit1->template GetObject<const FTriangleMeshImplicitObject>())
 				{
-					TRigidTransform<T, d> ParticleImplicit0TM = Transform0.GetRelativeTransform(Collisions::GetTransform(Particle0));
-					TRigidTransform<T, d> ParticleImplicit1TM = Transform1.GetRelativeTransform(Collisions::GetTransform(Particle1));
-					FRigidBodyPointContactConstraint Constraint = FRigidBodyPointContactConstraint(Particle0, Implicit0, ParticleImplicit0TM, Particle1, Implicit1, ParticleImplicit1TM, EContactShapesType::ConvexTriMesh);
-					UpdateConvexTriangleMeshConstraint(*Implicit0, Transform0, *TriangleMesh, Transform1, CullDistance, Constraint);
-					NewConstraints.TryAdd(CullDistance, Constraint);
+					FRigidBodyPointContactConstraint Constraint = FRigidBodyPointContactConstraint(Particle0, Implicit0, LocalTransform0, Particle1, Implicit1, LocalTransform1, EContactShapesType::ConvexTriMesh);
+					if (!bChaos_Collision_DeferNarrowPhase)
+					{
+						TRigidTransform<T, d> WorldTransform0 = LocalTransform0 * Collisions::GetTransform(Particle0);
+						TRigidTransform<T, d> WorldTransform1 = LocalTransform1 * Collisions::GetTransform(Particle1);
+						UpdateConvexTriangleMeshConstraint(*Implicit0, WorldTransform0, *TriangleMesh, WorldTransform1, CullDistance, Constraint);
+						NewConstraints.TryAdd(CullDistance, Constraint);
+					}
+					else
+					{
+						NewConstraints.Add(Constraint);
+					}
 				}
 				else
 				{
@@ -1776,23 +1937,24 @@ namespace Chaos
 		}
 
 		template<typename T, int d>
-		void ConstructConvexTriangleMeshConstraintsSwept(TGeometryParticleHandle<T, d>* Particle0, TGeometryParticleHandle<T, d>* Particle1, const FImplicitObject* Implicit0, const FImplicitObject* Implicit1, const TRigidTransform<T, d>& Transform0, const TRigidTransform<T, d>& Transform1, const T CullDistance, const TVector<T, d>& Dir, FReal Length, FCollisionConstraintsArray& NewConstraints)
+		void ConstructConvexTriangleMeshConstraintsSwept(TGeometryParticleHandle<T, d>* Particle0, TGeometryParticleHandle<T, d>* Particle1, const FImplicitObject* Implicit0, const FImplicitObject* Implicit1, const TRigidTransform<T, d>& LocalTransform0, const TRigidTransform<T, d>& LocalTransform1, const T CullDistance, const TVector<T, d>& Dir, FReal Length, FCollisionConstraintsArray& NewConstraints)
 		{
-			TRigidTransform<T, d> ParticleImplicit0TM = Transform0.GetRelativeTransform(Collisions::GetTransform(Particle0));
-			TRigidTransform<T, d> ParticleImplicit1TM = Transform1.GetRelativeTransform(Collisions::GetTransform(Particle1));
 			if (ensure(Implicit0->IsConvex()))
 			{
-				TRigidTransform<T, 3> TransformX0(Particle0->X(), Particle0->R());
 				if (const TImplicitObjectScaled<FTriangleMeshImplicitObject>* ScaledTriangleMesh = Implicit1->template GetObject<const TImplicitObjectScaled<FTriangleMeshImplicitObject>>())
 				{
-					FRigidBodySweptPointContactConstraint Constraint = FRigidBodySweptPointContactConstraint(Particle0, Implicit0, ParticleImplicit0TM, Particle1, Implicit1, ParticleImplicit1TM, EContactShapesType::ConvexTriMesh);
-					UpdateConvexTriangleMeshConstraintSwept(Particle0, *Implicit0, TransformX0, *ScaledTriangleMesh, Transform1, Dir, Length, CullDistance, Constraint);
+					FRigidBodySweptPointContactConstraint Constraint = FRigidBodySweptPointContactConstraint(Particle0, Implicit0, LocalTransform0, Particle1, Implicit1, LocalTransform1, EContactShapesType::ConvexTriMesh);
+					TRigidTransform<T, 3> WorldTransformX0 = LocalTransform0 * FRigidTransform3(Particle0->X(), Particle0->R());
+					TRigidTransform<T, d> WorldTransform1 = LocalTransform1 * Collisions::GetTransform(Particle1);
+					UpdateConvexTriangleMeshConstraintSwept(Particle0, *Implicit0, WorldTransformX0, *ScaledTriangleMesh, WorldTransform1, Dir, Length, CullDistance, Constraint);
 					NewConstraints.TryAdd(CullDistance, Constraint);
 				}
 				else if (const FTriangleMeshImplicitObject* TriangleMesh = Implicit1->template GetObject<const FTriangleMeshImplicitObject>())
 				{
-					FRigidBodySweptPointContactConstraint Constraint = FRigidBodySweptPointContactConstraint(Particle0, Implicit0, ParticleImplicit0TM, Particle1, Implicit1, ParticleImplicit1TM, EContactShapesType::ConvexTriMesh);
-					UpdateConvexTriangleMeshConstraintSwept(Particle0, *Implicit0, TransformX0, *TriangleMesh, Transform1, Dir, Length, CullDistance, Constraint);
+					FRigidBodySweptPointContactConstraint Constraint = FRigidBodySweptPointContactConstraint(Particle0, Implicit0, LocalTransform0, Particle1, Implicit1, LocalTransform1, EContactShapesType::ConvexTriMesh);
+					TRigidTransform<T, 3> WorldTransformX0 = LocalTransform0 * FRigidTransform3(Particle0->X(), Particle0->R());
+					TRigidTransform<T, d> WorldTransform1 = LocalTransform1 * Collisions::GetTransform(Particle1);
+					UpdateConvexTriangleMeshConstraintSwept(Particle0, *Implicit0, WorldTransformX0, *TriangleMesh, WorldTransform1, Dir, Length, CullDistance, Constraint);
 					NewConstraints.TryAdd(CullDistance, Constraint);
 				}
 				else
@@ -1845,11 +2007,9 @@ namespace Chaos
 
 
 		template<typename T, int d>
-		void ConstructLevelsetLevelsetConstraints(TGeometryParticleHandle<T, d>* Particle0, TGeometryParticleHandle<T, d>* Particle1, const FImplicitObject* Implicit0, const FImplicitObject* Implicit1, const TRigidTransform<T, d>& Transform0, const TRigidTransform<T, d>& Transform1, const T CullDistance, FCollisionConstraintsArray& NewConstraints)
+		void ConstructLevelsetLevelsetConstraints(TGeometryParticleHandle<T, d>* Particle0, TGeometryParticleHandle<T, d>* Particle1, const FImplicitObject* Implicit0, const FImplicitObject* Implicit1, const TRigidTransform<T, d>& LocalTransform0, const TRigidTransform<T, d>& LocalTransform1, const T CullDistance, FCollisionConstraintsArray& NewConstraints)
 		{
-			TRigidTransform<T, d> ParticleImplicit0TM = Transform0.GetRelativeTransform(Collisions::GetTransform(Particle0));
-			TRigidTransform<T, d> ParticleImplicit1TM = Transform1.GetRelativeTransform(Collisions::GetTransform(Particle1));
-			FRigidBodyPointContactConstraint Constraint = FRigidBodyPointContactConstraint(Particle0, Implicit0, ParticleImplicit0TM, Particle1, Implicit1, ParticleImplicit1TM, EContactShapesType::LevelSetLevelSet);
+			FRigidBodyPointContactConstraint Constraint = FRigidBodyPointContactConstraint(Particle0, Implicit0, LocalTransform0, Particle1, Implicit1, LocalTransform1, EContactShapesType::LevelSetLevelSet);
 
 			bool bIsParticleDynamic0 = Particle0->CastToRigidParticle() && Particle0->ObjectState() == EObjectStateType::Dynamic;
 			if (!Particle1->Geometry() || (bIsParticleDynamic0 && !Particle0->CastToRigidParticle()->CollisionParticlesSize() && Particle0->Geometry() && !Particle0->Geometry()->IsUnderlyingUnion()))
@@ -1865,9 +2025,15 @@ namespace Chaos
 				Constraint.SetManifold(Implicit0, Implicit1);
 			}
 
-			UpdateLevelsetLevelsetConstraint<ECollisionUpdateType::Deepest>(CullDistance, Constraint);
-
-			NewConstraints.TryAdd(CullDistance, Constraint);
+			if (!bChaos_Collision_DeferNarrowPhase)
+			{
+				UpdateLevelsetLevelsetConstraint<ECollisionUpdateType::Deepest>(CullDistance, Constraint);
+				NewConstraints.TryAdd(CullDistance, Constraint);
+			}
+			else
+			{
+				NewConstraints.Add(Constraint);
+			}
 		}
 
 
@@ -1881,16 +2047,16 @@ namespace Chaos
 		{
 			const FImplicitObject& Implicit0 = *Constraint.Manifold.Implicit[0];
 			const FImplicitObject& Implicit1 = *Constraint.Manifold.Implicit[1];
-			const TRigidTransform<T, d> Transform0 = Constraint.ImplicitTransform[0] * ParticleTransform0;
-			const TRigidTransform<T, d> Transform1 = Constraint.ImplicitTransform[1] * ParticleTransform1;
+			const TRigidTransform<T, d> WorldTransform0 = Constraint.ImplicitTransform[0] * ParticleTransform0;
+			const TRigidTransform<T, d> WorldTransform1 = Constraint.ImplicitTransform[1] * ParticleTransform1;
 
 			switch (Constraint.Manifold.ShapesType)
 			{
 			case EContactShapesType::CapsuleBox:
-				UpdateCapsuleBoxManifold<T, d>(*Implicit0.template GetObject<const TCapsule<T>>(), Transform0, Implicit1.template GetObject<const TBox<T, d>>()->GetAABB(), Transform1, CullDistance, Context, Constraint);
+				UpdateCapsuleBoxManifold<T, d>(*Implicit0.template GetObject<const TCapsule<T>>(), WorldTransform0, Implicit1.template GetObject<const TBox<T, d>>()->GetAABB(), WorldTransform1, CullDistance, Context, Constraint);
 				break;
 			case EContactShapesType::ConvexConvex:
-				UpdateConvexConvexManifold(Constraint, Transform0, Transform1, CullDistance);
+				UpdateConvexConvexManifold(Constraint, WorldTransform0, WorldTransform1, CullDistance);
 				break;
 			default:
 				// This switch needs to contain all pair types that generate a Manifold. ConstructConstrainsImpl
@@ -1906,7 +2072,7 @@ namespace Chaos
 		// @todo(chaos): add the missing cases below
 		// @todo(chaos): see use GetInnerObject below - we should try to uise the leaf types for all (currently only TriMesh needs this)
 		template<ECollisionUpdateType UpdateType, typename T, int d>
-		inline void UpdateConstraintFromGeometryImpl(TRigidBodyPointContactConstraint<T, d>& Constraint, const TRigidTransform<T, d>& Transform0, const TRigidTransform<T, d>& Transform1, const T CullDistance)
+		inline void UpdateConstraintFromGeometryImpl(TRigidBodyPointContactConstraint<T, d>& Constraint, const TRigidTransform<T, d>& WorldTransform0, const TRigidTransform<T, d>& WorldTransform1, const T CullDistance)
 		{
 			const FImplicitObject& Implicit0 = *Constraint.Manifold.Implicit[0];
 			const FImplicitObject& Implicit1 = *Constraint.Manifold.Implicit[1];
@@ -1914,107 +2080,107 @@ namespace Chaos
 			switch (Constraint.Manifold.ShapesType)
 			{
 			case EContactShapesType::SphereSphere:
-				UpdateSphereSphereConstraint(*Implicit0.template GetObject<TSphere<T, d>>(), Transform0, *Implicit1.template GetObject<TSphere<T, d>>(), Transform1, CullDistance, Constraint);
+				UpdateSphereSphereConstraint(*Implicit0.template GetObject<TSphere<T, d>>(), WorldTransform0, *Implicit1.template GetObject<TSphere<T, d>>(), WorldTransform1, CullDistance, Constraint);
 				break;
 			case EContactShapesType::SphereCapsule:
-				UpdateSphereCapsuleConstraint(*Implicit0.template GetObject<TSphere<T, d>>(), Transform0, *Implicit1.template GetObject<TCapsule<T>>(), Transform1, CullDistance, Constraint);
+				UpdateSphereCapsuleConstraint(*Implicit0.template GetObject<TSphere<T, d>>(), WorldTransform0, *Implicit1.template GetObject<TCapsule<T>>(), WorldTransform1, CullDistance, Constraint);
 				break;
 			case EContactShapesType::SphereBox:
-				UpdateSphereBoxConstraint(*Implicit0.template GetObject<TSphere<T, d>>(), Transform0, Implicit1.template GetObject<TBox<T, d>>()->GetAABB(), Transform1, CullDistance, Constraint);
+				UpdateSphereBoxConstraint(*Implicit0.template GetObject<TSphere<T, d>>(), WorldTransform0, Implicit1.template GetObject<TBox<T, d>>()->GetAABB(), WorldTransform1, CullDistance, Constraint);
 				break;
 			case EContactShapesType::SphereConvex:
 				// MISSING CASE!!!
-				UpdateConvexConvexConstraint(Implicit0, Transform0, Implicit1, Transform1, CullDistance, Constraint);
+				UpdateConvexConvexConstraint(Implicit0, WorldTransform0, Implicit1, WorldTransform1, CullDistance, Constraint);
 				break;
 			case EContactShapesType::SphereTriMesh:
 				if (const TImplicitObjectScaled<FTriangleMeshImplicitObject>* ScaledTriMesh = Implicit1.template GetObject<const TImplicitObjectScaled<FTriangleMeshImplicitObject> >())
 				{
-					UpdateSphereTriangleMeshConstraint(*Implicit0.template GetObject<TSphere<T, d>>(), Transform0, *ScaledTriMesh, Transform1, CullDistance, Constraint);
+					UpdateSphereTriangleMeshConstraint(*Implicit0.template GetObject<TSphere<T, d>>(), WorldTransform0, *ScaledTriMesh, WorldTransform1, CullDistance, Constraint);
 					break;
 				}
 				else if (const FTriangleMeshImplicitObject* TriangleMeshImplicit = Implicit1.template GetObject<const FTriangleMeshImplicitObject>())
 				{
-					UpdateSphereTriangleMeshConstraint(*Implicit0.template GetObject<TSphere<T, d>>(), Transform0, *TriangleMeshImplicit, Transform1, CullDistance, Constraint);
+					UpdateSphereTriangleMeshConstraint(*Implicit0.template GetObject<TSphere<T, d>>(), WorldTransform0, *TriangleMeshImplicit, WorldTransform1, CullDistance, Constraint);
 					break;
 				}
 				ensure(false);
 				break;
 			case EContactShapesType::SphereHeightField:
-				UpdateSphereHeightFieldConstraint(*Implicit0.template GetObject<TSphere<T, d>>(), Transform0, *Implicit1.template GetObject< THeightField<T> >(), Transform1, CullDistance, Constraint);
+				UpdateSphereHeightFieldConstraint(*Implicit0.template GetObject<TSphere<T, d>>(), WorldTransform0, *Implicit1.template GetObject< THeightField<T> >(), WorldTransform1, CullDistance, Constraint);
 				break;
 			case EContactShapesType::SpherePlane:
-				UpdateSpherePlaneConstraint(*Implicit0.template GetObject<TSphere<T, d>>(), Transform0, *Implicit1.template GetObject<TPlane<T, d>>(), Transform1, CullDistance, Constraint);
+				UpdateSpherePlaneConstraint(*Implicit0.template GetObject<TSphere<T, d>>(), WorldTransform0, *Implicit1.template GetObject<TPlane<T, d>>(), WorldTransform1, CullDistance, Constraint);
 				break;
 			case EContactShapesType::CapsuleCapsule:
-				UpdateCapsuleCapsuleConstraint(*Implicit0.template GetObject<TCapsule<T>>(), Transform0, *Implicit1.template GetObject<TCapsule<T>>(), Transform1, CullDistance, Constraint);
+				UpdateCapsuleCapsuleConstraint(*Implicit0.template GetObject<TCapsule<T>>(), WorldTransform0, *Implicit1.template GetObject<TCapsule<T>>(), WorldTransform1, CullDistance, Constraint);
 				break;
 			case EContactShapesType::CapsuleBox:
-				UpdateCapsuleBoxConstraint(*Implicit0.template GetObject<TCapsule<T>>(), Transform0, Implicit1.template GetObject<TBox<T, d>>()->GetAABB(), Transform1, CullDistance, Constraint);
+				UpdateCapsuleBoxConstraint(*Implicit0.template GetObject<TCapsule<T>>(), WorldTransform0, Implicit1.template GetObject<TBox<T, d>>()->GetAABB(), WorldTransform1, CullDistance, Constraint);
 				break;
 			case EContactShapesType::CapsuleConvex:
 				// MISSING CASE!!!
-				UpdateConvexConvexConstraint(Implicit0, Transform0, Implicit1, Transform1, CullDistance, Constraint);
+				UpdateConvexConvexConstraint(Implicit0, WorldTransform0, Implicit1, WorldTransform1, CullDistance, Constraint);
 				break;
 			case EContactShapesType::CapsuleTriMesh:
 				if (const TImplicitObjectScaled<FTriangleMeshImplicitObject>* ScaledTriMesh = Implicit1.template GetObject<const TImplicitObjectScaled<FTriangleMeshImplicitObject> >())
 				{
-					UpdateCapsuleTriangleMeshConstraint(*Implicit0.template GetObject<TCapsule<T>>(), Transform0, *ScaledTriMesh, Transform1, CullDistance, Constraint);
+					UpdateCapsuleTriangleMeshConstraint(*Implicit0.template GetObject<TCapsule<T>>(), WorldTransform0, *ScaledTriMesh, WorldTransform1, CullDistance, Constraint);
 					break;
 				}
 				else if (const FTriangleMeshImplicitObject* TriangleMeshImplicit = Implicit1.template GetObject<const FTriangleMeshImplicitObject>())
 				{
-					UpdateCapsuleTriangleMeshConstraint(*Implicit0.template GetObject<TCapsule<T>>(), Transform0, *TriangleMeshImplicit, Transform1, CullDistance, Constraint);
+					UpdateCapsuleTriangleMeshConstraint(*Implicit0.template GetObject<TCapsule<T>>(), WorldTransform0, *TriangleMeshImplicit, WorldTransform1, CullDistance, Constraint);
 					break;
 				}
 				ensure(false);
 				break;
 			case EContactShapesType::CapsuleHeightField:
-				UpdateCapsuleHeightFieldConstraint(*Implicit0.template GetObject<TCapsule<T>>(), Transform0, *Implicit1.template GetObject< THeightField<T> >(), Transform1, CullDistance, Constraint);
+				UpdateCapsuleHeightFieldConstraint(*Implicit0.template GetObject<TCapsule<T>>(), WorldTransform0, *Implicit1.template GetObject< THeightField<T> >(), WorldTransform1, CullDistance, Constraint);
 				break;
 			case EContactShapesType::BoxBox:
-				UpdateBoxBoxConstraint(Implicit0.template GetObject<TBox<T, d>>()->GetAABB(), Transform0, Implicit1.template GetObject<TBox<T, d>>()->GetAABB(), Transform1, CullDistance, Constraint);
+				UpdateBoxBoxConstraint(Implicit0.template GetObject<TBox<T, d>>()->GetAABB(), WorldTransform0, Implicit1.template GetObject<TBox<T, d>>()->GetAABB(), WorldTransform1, CullDistance, Constraint);
 				break;
 			case EContactShapesType::BoxConvex:
 				// MISSING CASE!!!
-				UpdateConvexConvexConstraint(Implicit0, Transform0, Implicit1, Transform1, CullDistance, Constraint);
+				UpdateConvexConvexConstraint(Implicit0, WorldTransform0, Implicit1, WorldTransform1, CullDistance, Constraint);
 				break;
 			case EContactShapesType::BoxTriMesh:
 				if (const TImplicitObjectScaled<FTriangleMeshImplicitObject>* ScaledTriMesh = Implicit1.template GetObject<const TImplicitObjectScaled<FTriangleMeshImplicitObject> >())
 				{
-					UpdateBoxTriangleMeshConstraint(Implicit0.template GetObject<TBox<T, d>>()->GetAABB(), Transform0, *ScaledTriMesh, Transform1, CullDistance, Constraint);
+					UpdateBoxTriangleMeshConstraint(Implicit0.template GetObject<TBox<T, d>>()->GetAABB(), WorldTransform0, *ScaledTriMesh, WorldTransform1, CullDistance, Constraint);
 					break;
 				}
 				else if (const FTriangleMeshImplicitObject* TriangleMeshImplicit = Implicit1.template GetObject<const FTriangleMeshImplicitObject>())
 				{
-					UpdateBoxTriangleMeshConstraint(Implicit0.template GetObject<TBox<T, d>>()->GetAABB(), Transform0, *TriangleMeshImplicit, Transform1, CullDistance, Constraint);
+					UpdateBoxTriangleMeshConstraint(Implicit0.template GetObject<TBox<T, d>>()->GetAABB(), WorldTransform0, *TriangleMeshImplicit, WorldTransform1, CullDistance, Constraint);
 					break;
 				}
 				ensure(false);
 				break;
 			case EContactShapesType::BoxHeightField:
-				UpdateBoxHeightFieldConstraint(Implicit0.template GetObject<TBox<T, d>>()->GetAABB(), Transform0, *Implicit1.template GetObject< THeightField<T> >(), Transform1, CullDistance, Constraint);
+				UpdateBoxHeightFieldConstraint(Implicit0.template GetObject<TBox<T, d>>()->GetAABB(), WorldTransform0, *Implicit1.template GetObject< THeightField<T> >(), WorldTransform1, CullDistance, Constraint);
 				break;
 			case EContactShapesType::BoxPlane:
-				UpdateBoxPlaneConstraint(Implicit0.template GetObject<TBox<T, d>>()->GetAABB(), Transform0, *Implicit1.template GetObject<TPlane<T, d>>(), Transform1, CullDistance, Constraint);
+				UpdateBoxPlaneConstraint(Implicit0.template GetObject<TBox<T, d>>()->GetAABB(), WorldTransform0, *Implicit1.template GetObject<TPlane<T, d>>(), WorldTransform1, CullDistance, Constraint);
 				break;
 			case EContactShapesType::ConvexConvex:
-				UpdateConvexConvexConstraint(Implicit0, Transform0, Implicit1, Transform1, CullDistance, Constraint);
+				UpdateConvexConvexConstraint(Implicit0, WorldTransform0, Implicit1, WorldTransform1, CullDistance, Constraint);
 				break;
 			case EContactShapesType::ConvexTriMesh:
 				if (const TImplicitObjectScaled<FTriangleMeshImplicitObject>* ScaledTriMesh = Implicit1.template GetObject<const TImplicitObjectScaled<FTriangleMeshImplicitObject> >())
 				{
-					UpdateConvexTriangleMeshConstraint(Implicit0, Transform0, *ScaledTriMesh, Transform1, CullDistance, Constraint);
+					UpdateConvexTriangleMeshConstraint(Implicit0, WorldTransform0, *ScaledTriMesh, WorldTransform1, CullDistance, Constraint);
 					break;
 				}
 				else if (const FTriangleMeshImplicitObject* TriangleMeshImplicit = Implicit1.template GetObject<const FTriangleMeshImplicitObject>())
 				{
-					UpdateConvexTriangleMeshConstraint(Implicit0, Transform0, *TriangleMeshImplicit, Transform1, CullDistance, Constraint);
+					UpdateConvexTriangleMeshConstraint(Implicit0, WorldTransform0, *TriangleMeshImplicit, WorldTransform1, CullDistance, Constraint);
 					break;
 				}
 				ensure(false);
 				break;
 			case EContactShapesType::ConvexHeightField:
-				UpdateConvexHeightFieldConstraint(Implicit0, Transform0, *Implicit1.template GetObject< THeightField<T> >(), Transform1, CullDistance, Constraint);
+				UpdateConvexHeightFieldConstraint(Implicit0, WorldTransform0, *Implicit1.template GetObject< THeightField<T> >(), WorldTransform1, CullDistance, Constraint);
 				break;
 			case EContactShapesType::LevelSetLevelSet:
 				UpdateLevelsetLevelsetConstraint<UpdateType>(CullDistance, Constraint);
@@ -2027,7 +2193,7 @@ namespace Chaos
 		}
 
 		template<typename T, int d>
-		void ConstructConstraintsImpl(TGeometryParticleHandle<T, d>* Particle0, TGeometryParticleHandle<T, d>* Particle1, const FImplicitObject* Implicit0, const FImplicitObject* Implicit1, const TRigidTransform<T, d>& Transform0, const TRigidTransform<T, d>& Transform1, const T CullDistance, const FCollisionContext& Context, FCollisionConstraintsArray& NewConstraints)
+		void ConstructConstraintsImpl(TGeometryParticleHandle<T, d>* Particle0, TGeometryParticleHandle<T, d>* Particle1, const FImplicitObject* Implicit0, const FImplicitObject* Implicit1, const TRigidTransform<T, d>& LocalTransform0, const TRigidTransform<T, d>& LocalTransform1, const T CullDistance, const FCollisionContext& Context, FCollisionConstraintsArray& NewConstraints)
 		{
 			// @todo(chaos): We use GetInnerType here because TriMeshes are left with their "Instanced" wrapper, unlike all other instanced implicits. Should we strip the instance on Tri Mesh too?
 			EImplicitObjectType Implicit0Type = Implicit0 ? GetInnerType(Implicit0->GetCollisionType()) : ImplicitObjectType::Unknown;
@@ -2041,181 +2207,181 @@ namespace Chaos
 
 			if (Implicit0Type == TBox<FReal, 3>::StaticType() && Implicit1Type == TBox<FReal, 3>::StaticType())
 			{
-				ConstructBoxBoxConstraints(Particle0, Particle1, Implicit0, Implicit1, Transform0, Transform1, CullDistance, NewConstraints);
+				ConstructBoxBoxConstraints(Particle0, Particle1, Implicit0, Implicit1, LocalTransform0, LocalTransform1, CullDistance, NewConstraints);
 			}
 			else if (Implicit0Type == TBox<FReal, 3>::StaticType() && Implicit1Type == THeightField<FReal>::StaticType())
 			{
-				ConstructBoxHeightFieldConstraints(Particle0, Particle1, Implicit0, Implicit1, Transform0, Transform1, CullDistance, NewConstraints);
+				ConstructBoxHeightFieldConstraints(Particle0, Particle1, Implicit0, Implicit1, LocalTransform0, LocalTransform1, CullDistance, NewConstraints);
 			}
 			else if (Implicit0Type == THeightField<FReal>::StaticType() && Implicit1Type == TBox<FReal, 3>::StaticType())
 			{
-				ConstructBoxHeightFieldConstraints(Particle1, Particle0, Implicit1, Implicit0, Transform1, Transform0, CullDistance, NewConstraints);
+				ConstructBoxHeightFieldConstraints(Particle1, Particle0, Implicit1, Implicit0, LocalTransform1, LocalTransform0, CullDistance, NewConstraints);
 			}
 			else if (Implicit0Type == TBox<FReal, 3>::StaticType() && Implicit1Type == TPlane<FReal, 3>::StaticType())
 			{
-				ConstructBoxPlaneConstraints(Particle0, Particle1, Implicit0, Implicit1, Transform0, Transform1, CullDistance, NewConstraints);
+				ConstructBoxPlaneConstraints(Particle0, Particle1, Implicit0, Implicit1, LocalTransform0, LocalTransform1, CullDistance, NewConstraints);
 			}
 			else if (Implicit0Type == TPlane<FReal, 3>::StaticType() && Implicit1Type == TBox<FReal, 3>::StaticType())
 			{
-				ConstructBoxPlaneConstraints(Particle1, Particle0, Implicit1, Implicit0, Transform1, Transform0, CullDistance, NewConstraints);
+				ConstructBoxPlaneConstraints(Particle1, Particle0, Implicit1, Implicit0, LocalTransform1, LocalTransform0, CullDistance, NewConstraints);
 			}
 			else if (Implicit0Type == TSphere<FReal, 3>::StaticType() && Implicit1Type == TSphere<FReal, 3>::StaticType())
 			{
-				ConstructSphereSphereConstraints(Particle0, Particle1, Implicit0, Implicit1, Transform0, Transform1, CullDistance, NewConstraints);
+				ConstructSphereSphereConstraints(Particle0, Particle1, Implicit0, Implicit1, LocalTransform0, LocalTransform1, CullDistance, NewConstraints);
 			}
 			else if (Implicit0Type == TSphere<FReal, 3>::StaticType() && Implicit1Type == THeightField<FReal>::StaticType())
 			{
-				ConstructSphereHeightFieldConstraints(Particle0, Particle1, Implicit0, Implicit1, Transform0, Transform1, CullDistance, NewConstraints);
+				ConstructSphereHeightFieldConstraints(Particle0, Particle1, Implicit0, Implicit1, LocalTransform0, LocalTransform1, CullDistance, NewConstraints);
 			}
 			else if (Implicit0Type == THeightField<FReal>::StaticType() && Implicit1Type == TSphere<FReal, 3>::StaticType())
 			{
-				ConstructSphereHeightFieldConstraints(Particle1, Particle0, Implicit1, Implicit0, Transform1, Transform0, CullDistance, NewConstraints);
+				ConstructSphereHeightFieldConstraints(Particle1, Particle0, Implicit1, Implicit0, LocalTransform1, LocalTransform0, CullDistance, NewConstraints);
 			}
 			else if (Implicit0Type == TSphere<FReal, 3>::StaticType() && Implicit1Type == TPlane<FReal, 3>::StaticType())
 			{
-				ConstructSpherePlaneConstraints(Particle0, Particle1, Implicit0, Implicit1, Transform0, Transform1, CullDistance, NewConstraints);
+				ConstructSpherePlaneConstraints(Particle0, Particle1, Implicit0, Implicit1, LocalTransform0, LocalTransform1, CullDistance, NewConstraints);
 			}
 			else if (Implicit0Type == TPlane<FReal, 3>::StaticType() && Implicit1Type == TSphere<FReal, 3>::StaticType())
 			{
-				ConstructSpherePlaneConstraints(Particle1, Particle0, Implicit1, Implicit0, Transform1, Transform0, CullDistance, NewConstraints);
+				ConstructSpherePlaneConstraints(Particle1, Particle0, Implicit1, Implicit0, LocalTransform1, LocalTransform0, CullDistance, NewConstraints);
 			}
 			else if (Implicit0Type == TSphere<FReal, 3>::StaticType() && Implicit1Type == TBox<FReal, 3>::StaticType())
 			{
-				ConstructSphereBoxConstraints(Particle0, Particle1, Implicit0, Implicit1, Transform0, Transform1, CullDistance, NewConstraints);
+				ConstructSphereBoxConstraints(Particle0, Particle1, Implicit0, Implicit1, LocalTransform0, LocalTransform1, CullDistance, NewConstraints);
 			}
 			else if (Implicit0Type == TBox<FReal, 3>::StaticType() && Implicit1Type == TSphere<FReal, 3>::StaticType())
 			{
-				ConstructSphereBoxConstraints(Particle1, Particle0, Implicit1, Implicit0, Transform1, Transform0, CullDistance, NewConstraints);
+				ConstructSphereBoxConstraints(Particle1, Particle0, Implicit1, Implicit0, LocalTransform1, LocalTransform0, CullDistance, NewConstraints);
 			}
 			else if (Implicit0Type == TSphere<FReal, 3>::StaticType() && Implicit1Type == TCapsule<FReal>::StaticType())
 			{
-				ConstructSphereCapsuleConstraints(Particle0, Particle1, Implicit0, Implicit1, Transform0, Transform1, CullDistance, NewConstraints);
+				ConstructSphereCapsuleConstraints(Particle0, Particle1, Implicit0, Implicit1, LocalTransform0, LocalTransform1, CullDistance, NewConstraints);
 			}
 			else if (Implicit0Type == TCapsule<FReal>::StaticType() && Implicit1Type == TSphere<FReal, 3>::StaticType())
 			{
-				ConstructSphereCapsuleConstraints(Particle1, Particle0, Implicit1, Implicit0, Transform1, Transform0, CullDistance, NewConstraints);
+				ConstructSphereCapsuleConstraints(Particle1, Particle0, Implicit1, Implicit0, LocalTransform1, LocalTransform0, CullDistance, NewConstraints);
 			}
 			else if (Implicit0Type == TCapsule<FReal>::StaticType() && Implicit1Type == TCapsule<FReal>::StaticType())
 			{
-				ConstructCapsuleCapsuleConstraints(Particle0, Particle1, Implicit0, Implicit1, Transform0, Transform1, CullDistance, NewConstraints);
+				ConstructCapsuleCapsuleConstraints(Particle0, Particle1, Implicit0, Implicit1, LocalTransform0, LocalTransform1, CullDistance, NewConstraints);
 			}
 			else if (Implicit0Type == TCapsule<FReal>::StaticType() && Implicit1Type == TBox<FReal, 3>::StaticType())
 			{
-				ConstructCapsuleBoxConstraints(Particle0, Particle1, Implicit0, Implicit1, Transform0, Transform1, CullDistance, Context, NewConstraints);
+				ConstructCapsuleBoxConstraints(Particle0, Particle1, Implicit0, Implicit1, LocalTransform0, LocalTransform1, CullDistance, Context, NewConstraints);
 			}
 			else if (Implicit0Type == TBox<FReal, 3>::StaticType() && Implicit1Type == TCapsule<FReal>::StaticType())
 			{
-				ConstructCapsuleBoxConstraints(Particle1, Particle0, Implicit1, Implicit0, Transform1, Transform0, CullDistance, Context, NewConstraints);
+				ConstructCapsuleBoxConstraints(Particle1, Particle0, Implicit1, Implicit0, LocalTransform1, LocalTransform0, CullDistance, Context, NewConstraints);
 			}
 			else if (Implicit0Type == TCapsule<FReal>::StaticType() && Implicit1Type == THeightField<FReal>::StaticType())
 			{
 				if (bUseCCD)
 				{
-					ConstructCapsuleHeightFieldConstraintsSwept(Particle0, Particle1, Implicit0, Implicit1, Transform0, Transform1, CullDistance, DirCCD, LengthCCD, NewConstraints);
+					ConstructCapsuleHeightFieldConstraintsSwept(Particle0, Particle1, Implicit0, Implicit1, LocalTransform0, LocalTransform1, CullDistance, DirCCD, LengthCCD, NewConstraints);
 					return;
 				}
 
-				ConstructCapsuleHeightFieldConstraints(Particle0, Particle1, Implicit0, Implicit1, Transform0, Transform1, CullDistance, NewConstraints);
+				ConstructCapsuleHeightFieldConstraints(Particle0, Particle1, Implicit0, Implicit1, LocalTransform0, LocalTransform1, CullDistance, NewConstraints);
 			}
 			else if (Implicit0Type == THeightField<FReal>::StaticType() && Implicit1Type == TCapsule<FReal>::StaticType())
 			{
 				if (bUseCCD)
 				{
-					ConstructCapsuleHeightFieldConstraintsSwept(Particle1, Particle0, Implicit1, Implicit0, Transform1, Transform0, CullDistance, DirCCD, LengthCCD, NewConstraints);
+					ConstructCapsuleHeightFieldConstraintsSwept(Particle1, Particle0, Implicit1, Implicit0, LocalTransform1, LocalTransform0, CullDistance, DirCCD, LengthCCD, NewConstraints);
 					return;
 				}
 
-				ConstructCapsuleHeightFieldConstraints(Particle1, Particle0, Implicit1, Implicit0, Transform1, Transform0, CullDistance, NewConstraints);
+				ConstructCapsuleHeightFieldConstraints(Particle1, Particle0, Implicit1, Implicit0, LocalTransform1, LocalTransform0, CullDistance, NewConstraints);
 			}
 			else if (Implicit0Type == TBox<FReal, 3>::StaticType() && Implicit1Type == FTriangleMeshImplicitObject::StaticType())
 			{
-				ConstructBoxTriangleMeshConstraints(Particle0, Particle1, Implicit0, Implicit1, Transform0, Transform1, CullDistance, NewConstraints);
+				ConstructBoxTriangleMeshConstraints(Particle0, Particle1, Implicit0, Implicit1, LocalTransform0, LocalTransform1, CullDistance, NewConstraints);
 			}
 			else if (Implicit0Type == FTriangleMeshImplicitObject::StaticType() && Implicit1Type == TBox<FReal, 3>::StaticType())
 			{
-				ConstructBoxTriangleMeshConstraints(Particle1, Particle0, Implicit1, Implicit0, Transform1, Transform0, CullDistance, NewConstraints);
+				ConstructBoxTriangleMeshConstraints(Particle1, Particle0, Implicit1, Implicit0, LocalTransform1, LocalTransform0, CullDistance, NewConstraints);
 			}
 			else if (Implicit0Type == TSphere<FReal, 3>::StaticType() && Implicit1Type == FTriangleMeshImplicitObject::StaticType())
 			{
-				ConstructSphereTriangleMeshConstraints(Particle0, Particle1, Implicit0, Implicit1, Transform0, Transform1, CullDistance, NewConstraints);
+				ConstructSphereTriangleMeshConstraints(Particle0, Particle1, Implicit0, Implicit1, LocalTransform0, LocalTransform1, CullDistance, NewConstraints);
 			}
 			else if (Implicit0Type == FTriangleMeshImplicitObject::StaticType() && Implicit1Type == TSphere<FReal, 3>::StaticType())
 			{
-				ConstructSphereTriangleMeshConstraints(Particle1, Particle0, Implicit1, Implicit0, Transform1, Transform0, CullDistance, NewConstraints);
+				ConstructSphereTriangleMeshConstraints(Particle1, Particle0, Implicit1, Implicit0, LocalTransform1, LocalTransform0, CullDistance, NewConstraints);
 			}
 			else if (Implicit0Type == TCapsule<FReal>::StaticType() && Implicit1Type == FTriangleMeshImplicitObject::StaticType())
 			{
 				if (bUseCCD)
 				{
-					ConstructCapsuleTriangleMeshConstraintsSwept(Particle0, Particle1, Implicit0, Implicit1, Transform0, Transform1, CullDistance, DirCCD, LengthCCD, NewConstraints);
+					ConstructCapsuleTriangleMeshConstraintsSwept(Particle0, Particle1, Implicit0, Implicit1, LocalTransform0, LocalTransform1, CullDistance, DirCCD, LengthCCD, NewConstraints);
 					return;
 				}
 
-				ConstructCapsuleTriangleMeshConstraints(Particle0, Particle1, Implicit0, Implicit1, Transform0, Transform1, CullDistance, NewConstraints);
+				ConstructCapsuleTriangleMeshConstraints(Particle0, Particle1, Implicit0, Implicit1, LocalTransform0, LocalTransform1, CullDistance, NewConstraints);
 			}
 			else if (Implicit0Type == FTriangleMeshImplicitObject::StaticType() && Implicit1Type == TCapsule<FReal>::StaticType())
 			{
 				if (bUseCCD)
 				{
-					ConstructCapsuleTriangleMeshConstraintsSwept(Particle0, Particle1, Implicit0, Implicit1, Transform0, Transform1, CullDistance, DirCCD, LengthCCD, NewConstraints);
+					ConstructCapsuleTriangleMeshConstraintsSwept(Particle0, Particle1, Implicit0, Implicit1, LocalTransform0, LocalTransform1, CullDistance, DirCCD, LengthCCD, NewConstraints);
 					return;
 				}
 
-				ConstructCapsuleTriangleMeshConstraints(Particle1, Particle0, Implicit1, Implicit0, Transform1, Transform0, CullDistance, NewConstraints);
+				ConstructCapsuleTriangleMeshConstraints(Particle1, Particle0, Implicit1, Implicit0, LocalTransform1, LocalTransform0, CullDistance, NewConstraints);
 			}
 			else if (bIsConvex0 && Implicit1Type == THeightField<FReal>::StaticType())
 			{
 				if (bUseCCD)
 				{
-					ConstructConvexHeightFieldConstraintsSwept(Particle0, Particle1, Implicit0, Implicit1, Transform0, Transform1, CullDistance, DirCCD, LengthCCD, NewConstraints);
+					ConstructConvexHeightFieldConstraintsSwept(Particle0, Particle1, Implicit0, Implicit1, LocalTransform0, LocalTransform1, CullDistance, DirCCD, LengthCCD, NewConstraints);
 					return;
 				}
 
-				ConstructConvexHeightFieldConstraints(Particle0, Particle1, Implicit0, Implicit1, Transform0, Transform1, CullDistance, NewConstraints);
+				ConstructConvexHeightFieldConstraints(Particle0, Particle1, Implicit0, Implicit1, LocalTransform0, LocalTransform1, CullDistance, NewConstraints);
 			}
 			else if (Implicit0Type == THeightField<FReal>::StaticType() && bIsConvex1)
 			{
 				if (bUseCCD)
 				{
-					ConstructConvexHeightFieldConstraintsSwept(Particle1, Particle0, Implicit1, Implicit0, Transform1, Transform0, CullDistance, DirCCD, LengthCCD, NewConstraints);
+					ConstructConvexHeightFieldConstraintsSwept(Particle1, Particle0, Implicit1, Implicit0, LocalTransform1, LocalTransform0, CullDistance, DirCCD, LengthCCD, NewConstraints);
 					return;
 				}
 
-				ConstructConvexHeightFieldConstraints(Particle1, Particle0, Implicit1, Implicit0, Transform1, Transform0, CullDistance, NewConstraints);
+				ConstructConvexHeightFieldConstraints(Particle1, Particle0, Implicit1, Implicit0, LocalTransform1, LocalTransform0, CullDistance, NewConstraints);
 			}
 			else if (bIsConvex0 && Implicit1Type == FTriangleMeshImplicitObject::StaticType())
 			{
 				if (bUseCCD)
 				{
-					ConstructConvexTriangleMeshConstraintsSwept(Particle0, Particle1, Implicit0, Implicit1, Transform0, Transform1, CullDistance, DirCCD, LengthCCD, NewConstraints);
+					ConstructConvexTriangleMeshConstraintsSwept(Particle0, Particle1, Implicit0, Implicit1, LocalTransform0, LocalTransform1, CullDistance, DirCCD, LengthCCD, NewConstraints);
 					return;
 				}
 
-				ConstructConvexTriangleMeshConstraints(Particle0, Particle1, Implicit0, Implicit1, Transform0, Transform1, CullDistance, NewConstraints);
+				ConstructConvexTriangleMeshConstraints(Particle0, Particle1, Implicit0, Implicit1, LocalTransform0, LocalTransform1, CullDistance, NewConstraints);
 			}
 			else if (Implicit0Type == FTriangleMeshImplicitObject::StaticType() && bIsConvex1)
 			{
 				if (bUseCCD)
 				{
-					ConstructConvexTriangleMeshConstraintsSwept(Particle1, Particle0, Implicit1, Implicit0, Transform1, Transform0, CullDistance, DirCCD, LengthCCD, NewConstraints);
+					ConstructConvexTriangleMeshConstraintsSwept(Particle1, Particle0, Implicit1, Implicit0, LocalTransform1, LocalTransform0, CullDistance, DirCCD, LengthCCD, NewConstraints);
 					return;
 				}
 
-				ConstructConvexTriangleMeshConstraints(Particle1, Particle0, Implicit1, Implicit0, Transform1, Transform0, CullDistance, NewConstraints);
+				ConstructConvexTriangleMeshConstraints(Particle1, Particle0, Implicit1, Implicit0, LocalTransform1, LocalTransform0, CullDistance, NewConstraints);
 			}
 			else if (bIsConvex0 && bIsConvex1)
 			{
 				if (bUseCCD)
 				{
-					ConstructConvexConvexConstraintsSwept(Particle0, Particle1, Implicit0, Implicit1, Transform0, Transform1, CullDistance, DirCCD, LengthCCD, NewConstraints);
+					ConstructConvexConvexConstraintsSwept(Particle0, Particle1, Implicit0, Implicit1, LocalTransform0, LocalTransform1, CullDistance, DirCCD, LengthCCD, NewConstraints);
 					return;
 				}
 
-				ConstructConvexConvexConstraints(Particle0, Particle1, Implicit0, Implicit1, Transform0, Transform1, CullDistance, NewConstraints);
+				ConstructConvexConvexConstraints(Particle0, Particle1, Implicit0, Implicit1, LocalTransform0, LocalTransform1, CullDistance, NewConstraints);
 			}
 			else
 			{
-				ConstructLevelsetLevelsetConstraints(Particle0, Particle1, Implicit0, Implicit1, Transform0, Transform1, CullDistance, NewConstraints);
+				ConstructLevelsetLevelsetConstraints(Particle0, Particle1, Implicit0, Implicit1, LocalTransform0, LocalTransform1, CullDistance, NewConstraints);
 			}
 		}
 
@@ -2225,9 +2391,9 @@ namespace Chaos
 		template<ECollisionUpdateType UpdateType, typename T, int d>
 		void UpdateConstraintFromGeometry(TRigidBodyPointContactConstraint<T, d>& Constraint, const TRigidTransform<T, d>& ParticleTransform0, const TRigidTransform<T, d>& ParticleTransform1, const T CullDistance)
 		{
-			const TRigidTransform<T, d> Transform0 = Constraint.ImplicitTransform[0] * ParticleTransform0;
-			const TRigidTransform<T, d> Transform1 = Constraint.ImplicitTransform[1] * ParticleTransform1;
-			UpdateConstraintFromGeometryImpl<UpdateType>(Constraint, Transform0, Transform1, CullDistance);
+			const TRigidTransform<T, d> WorldTransform0 = Constraint.ImplicitTransform[0] * ParticleTransform0;
+			const TRigidTransform<T, d> WorldTransform1 = Constraint.ImplicitTransform[1] * ParticleTransform1;
+			UpdateConstraintFromGeometryImpl<UpdateType>(Constraint, WorldTransform0, WorldTransform1, CullDistance);
 		}
 
 		// Select the best manifold point as the new contact point.
@@ -2235,20 +2401,20 @@ namespace Chaos
 		template<class T, int d>
 		void UpdateConstraintFromManifold(TRigidBodyMultiPointContactConstraint<T, d>& Constraint, const TRigidTransform<T, d>& ParticleTransform0, const TRigidTransform<T, d>& ParticleTransform1, const T CullDistance)
 		{
-			const FRigidTransform3 Transform0 = Constraint.ImplicitTransform[0] * ParticleTransform0;
-			const FRigidTransform3 Transform1 = Constraint.ImplicitTransform[1] * ParticleTransform1;
+			const FRigidTransform3 WorldTransform0 = Constraint.ImplicitTransform[0] * ParticleTransform0;
+			const FRigidTransform3 WorldTransform1 = Constraint.ImplicitTransform[1] * ParticleTransform1;
 			const int32 NumPoints = Constraint.NumManifoldPoints();
 
 			// Fall back to full collision detection if we have no manifold (or for testing)
 			if ((NumPoints == 0) || Chaos_Collision_UseManifolds_Test)
 			{
-				UpdateConstraintFromGeometryImpl<ECollisionUpdateType::Deepest>(Constraint, Transform0, Transform1, CullDistance);
+				UpdateConstraintFromGeometryImpl<ECollisionUpdateType::Deepest>(Constraint, WorldTransform0, WorldTransform1, CullDistance);
 				return;
 			}
 
 			// Get the plane and point transforms (depends which body owns the plane)
-			const FRigidTransform3& PlaneTransform = (Constraint.GetManifoldPlaneOwnerIndex() == 0) ? Transform0 : Transform1;
-			const FRigidTransform3& PointsTransform = (Constraint.GetManifoldPlaneOwnerIndex() == 0) ? Transform1 : Transform0;
+			const FRigidTransform3& PlaneTransform = (Constraint.GetManifoldPlaneOwnerIndex() == 0) ? WorldTransform0 : WorldTransform1;
+			const FRigidTransform3& PointsTransform = (Constraint.GetManifoldPlaneOwnerIndex() == 0) ? WorldTransform1 : WorldTransform0;
 
 			// World-space manifold plane
 			FVec3 PlaneNormal = PlaneTransform.TransformVectorNoScale(Constraint.GetManifoldPlaneNormal());
@@ -2369,7 +2535,7 @@ namespace Chaos
 
 
 		template<typename T, int d>
-		void ConstructConstraints(TGeometryParticleHandle<T, d>* Particle0, TGeometryParticleHandle<T, d>* Particle1, const FImplicitObject* Implicit0, const FImplicitObject* Implicit1, const TRigidTransform<T, d>& Transform0, const TRigidTransform<T, d>& Transform1, const T CullDistance, const FCollisionContext& Context, FCollisionConstraintsArray& NewConstraints)
+		void ConstructConstraints(TGeometryParticleHandle<T, d>* Particle0, TGeometryParticleHandle<T, d>* Particle1, const FImplicitObject* Implicit0, const FImplicitObject* Implicit1, const TRigidTransform<T, d>& LocalTransform0, const TRigidTransform<T, d>& LocalTransform1, const T CullDistance, const FCollisionContext& Context, FCollisionConstraintsArray& NewConstraints)
 		{
 			EImplicitObjectType Implicit0Type = Implicit0 ? GetInnerType(Implicit0->GetType()) : ImplicitObjectType::Unknown;
 			EImplicitObjectType Implicit1Type = Implicit1 ? GetInnerType(Implicit1->GetType()) : ImplicitObjectType::Unknown;
@@ -2377,7 +2543,7 @@ namespace Chaos
 
 			if (!Implicit0 || !Implicit1)
 			{
-				ConstructLevelsetLevelsetConstraints(Particle0, Particle1, Implicit0, Implicit1, Transform0, Transform1, CullDistance, NewConstraints);
+				ConstructLevelsetLevelsetConstraints(Particle0, Particle1, Implicit0, Implicit1, LocalTransform0, LocalTransform1, CullDistance, NewConstraints);
 				return;
 			}
 
@@ -2389,23 +2555,23 @@ namespace Chaos
 			{
 				const TImplicitObjectTransformed<FReal, 3>* TransformedImplicit0 = Implicit0->template GetObject<const TImplicitObjectTransformed<FReal, 3>>();
 				const TImplicitObjectTransformed<FReal, 3>* TransformedImplicit1 = Implicit1->template GetObject<const TImplicitObjectTransformed<FReal, 3>>();
-				TRigidTransform<T, d> TransformedTransform0 = TransformedImplicit0->GetTransform() * Transform0;
-				TRigidTransform<T, d> TransformedTransform1 = TransformedImplicit1->GetTransform() * Transform1;
+				TRigidTransform<T, d> TransformedTransform0 = TransformedImplicit0->GetTransform() * LocalTransform0;
+				TRigidTransform<T, d> TransformedTransform1 = TransformedImplicit1->GetTransform() * LocalTransform1;
 				ConstructConstraints(Particle0, Particle1, TransformedImplicit0->GetTransformedObject(), TransformedImplicit1->GetTransformedObject(), TransformedTransform0, TransformedTransform1, CullDistance, Context, NewConstraints);
 				return;
 			}
 			else if (Implicit0OuterType == TImplicitObjectTransformed<T, d>::StaticType())
 			{
 				const TImplicitObjectTransformed<FReal, 3>* TransformedImplicit0 = Implicit0->template GetObject<const TImplicitObjectTransformed<FReal, 3>>();
-				TRigidTransform<T, d> TransformedTransform0 = TransformedImplicit0->GetTransform() * Transform0;
-				ConstructConstraints(Particle0, Particle1, TransformedImplicit0->GetTransformedObject(), Implicit1, TransformedTransform0, Transform1, CullDistance, Context, NewConstraints);
+				TRigidTransform<T, d> TransformedTransform0 = TransformedImplicit0->GetTransform() * LocalTransform0;
+				ConstructConstraints(Particle0, Particle1, TransformedImplicit0->GetTransformedObject(), Implicit1, TransformedTransform0, LocalTransform1, CullDistance, Context, NewConstraints);
 				return;
 			}
 			else if (Implicit1OuterType == TImplicitObjectTransformed<T, d>::StaticType())
 			{
 				const TImplicitObjectTransformed<FReal, 3>* TransformedImplicit1 = Implicit1->template GetObject<const TImplicitObjectTransformed<FReal, 3>>();
-				TRigidTransform<T, d> TransformedTransform1 = TransformedImplicit1->GetTransform() * Transform1;
-				ConstructConstraints(Particle0, Particle1, Implicit0, TransformedImplicit1->GetTransformedObject(), Transform0, TransformedTransform1, CullDistance, Context, NewConstraints);
+				TRigidTransform<T, d> TransformedTransform1 = TransformedImplicit1->GetTransform() * LocalTransform1;
+				ConstructConstraints(Particle0, Particle1, Implicit0, TransformedImplicit1->GetTransformedObject(), LocalTransform0, TransformedTransform1, CullDistance, Context, NewConstraints);
 				return;
 			}
 			// Handle Instanced shapes
@@ -2424,17 +2590,17 @@ namespace Chaos
 				}
 				if (InnerImplicit0 && InnerImplicit1)
 				{
-					ConstructConstraints(Particle0, Particle1, InnerImplicit0, InnerImplicit1, Transform0, Transform1, CullDistance, Context, NewConstraints);
+					ConstructConstraints(Particle0, Particle1, InnerImplicit0, InnerImplicit1, LocalTransform0, LocalTransform1, CullDistance, Context, NewConstraints);
 					return;
 				}
 				else if (InnerImplicit0 && !InnerImplicit1)
 				{
-					ConstructConstraints(Particle0, Particle1, InnerImplicit0, Implicit1, Transform0, Transform1, CullDistance, Context, NewConstraints);
+					ConstructConstraints(Particle0, Particle1, InnerImplicit0, Implicit1, LocalTransform0, LocalTransform1, CullDistance, Context, NewConstraints);
 					return;
 				}
 				else if (!InnerImplicit0 && InnerImplicit1)
 				{
-					ConstructConstraints(Particle0, Particle1, Implicit0, InnerImplicit1, Transform0, Transform1, CullDistance, Context, NewConstraints);
+					ConstructConstraints(Particle0, Particle1, Implicit0, InnerImplicit1, LocalTransform0, LocalTransform1, CullDistance, Context, NewConstraints);
 					return;
 				}
 
@@ -2446,7 +2612,7 @@ namespace Chaos
 				const FImplicitObjectUnion* Union0 = Implicit0->template GetObject<FImplicitObjectUnion>();
 				for (const auto& Child0 : Union0->GetObjects())
 				{
-					ConstructConstraints(Particle0, Particle1, Child0.Get(), Implicit1, Transform0, Transform1, CullDistance, Context, NewConstraints);
+					ConstructConstraints(Particle0, Particle1, Child0.Get(), Implicit1, LocalTransform0, LocalTransform1, CullDistance, Context, NewConstraints);
 				}
 				return;
 			}
@@ -2455,7 +2621,7 @@ namespace Chaos
 				const FImplicitObjectUnion* Union1 = Implicit1->template GetObject<FImplicitObjectUnion>();
 				for (const auto& Child1 : Union1->GetObjects())
 				{
-					ConstructConstraints(Particle0, Particle1, Implicit0, Child1.Get(), Transform0, Transform1, CullDistance, Context, NewConstraints);
+					ConstructConstraints(Particle0, Particle1, Implicit0, Child1.Get(), LocalTransform0, LocalTransform1, CullDistance, Context, NewConstraints);
 				}
 				return;
 			}
@@ -2468,7 +2634,12 @@ namespace Chaos
 
 			// If we get here, we have a pair of concrete shapes (i.e., no wrappers or containers)
 			// Create a constraint for the shape pair
-			ConstructConstraintsImpl(Particle0, Particle1, Implicit0, Implicit1, Transform0, Transform1, CullDistance, Context, NewConstraints);
+			FReal ModifiedCullDistance = CullDistance;
+			if (bChaos_Collision_DeferNarrowPhase)
+			{
+				ModifiedCullDistance = FLT_MAX;
+			}
+			ConstructConstraintsImpl(Particle0, Particle1, Implicit0, Implicit1, LocalTransform0, LocalTransform1, ModifiedCullDistance, Context, NewConstraints);
 		}
 
 
diff --git a/Engine/Source/Runtime/Experimental/Chaos/Private/Chaos/PBDCollisionConstraints.cpp b/Engine/Source/Runtime/Experimental/Chaos/Private/Chaos/PBDCollisionConstraints.cpp
--- a/Engine/Source/Runtime/Experimental/Chaos/Private/Chaos/PBDCollisionConstraints.cpp
+++ b/Engine/Source/Runtime/Experimental/Chaos/Private/Chaos/PBDCollisionConstraints.cpp
@@ -59,6 +59,12 @@ namespace Chaos
 	float DefaultCollisionRestitution = 0;
 	FAutoConsoleVariableRef CVarDefaultCollisionRestitution(TEXT("p.DefaultCollisionRestitution"), DefaultCollisionRestitution, TEXT("Collision restitution default value if no materials are found."));
 
+#if INTEL_ISPC
+	bool bChaos_Collision_ISPC_Enabled = false;
+	FAutoConsoleVariableRef CVarChaosCollisionISPCEnabled(TEXT("p.Chaos.Collision.ISPC"), bChaos_Collision_ISPC_Enabled, TEXT("Whether to use ISPC optimizations in the Collision Solver"));
+#endif
+
+
 	DECLARE_CYCLE_STAT(TEXT("Collisions::Reset"), STAT_Collisions_Reset, STATGROUP_ChaosCollision);
 	DECLARE_CYCLE_STAT(TEXT("Collisions::UpdatePointConstraints"), STAT_Collisions_UpdatePointConstraints, STATGROUP_ChaosCollision);
 	DECLARE_CYCLE_STAT(TEXT("Collisions::UpdateManifoldConstraints"), STAT_Collisions_UpdateManifoldConstraints, STATGROUP_ChaosCollision);
@@ -79,6 +85,9 @@ namespace Chaos
 		const T CullDistance /*= (T)0*/,
 		const T ShapePadding /*= (T)0*/)
 		: Particles(InParticles)
+		, NumActivePointConstraints(0)
+		, NumActiveSweptPointConstraints(0)
+		, NumActiveIterativeConstraints(0)
 		, MCollided(Collided)
 		, MPhysicsMaterials(InPerParticleMaterials)
 		, MApplyPairIterations(InApplyPairIterations)
@@ -94,6 +103,12 @@ namespace Chaos
 		, PostApplyCallback(nullptr)
 		, PostApplyPushOutCallback(nullptr)
 	{
+#if INTEL_ISPC
+		if (bChaos_Collision_ISPC_Enabled)
+		{
+			check(sizeof(FCollisionContact) == ispc::SizeofFCollisionContact());
+		}
+#endif
 	}
 
 	template<typename T, int d>
@@ -171,14 +186,14 @@ namespace Chaos
 	template<typename T, int d>
 	void TPBDCollisionConstraints<T, d>::AddConstraint(const TRigidBodyPointContactConstraint<FReal, 3>& InConstraint)
 	{
-		int32 Idx = PointConstraints.Add(InConstraint);
+		int32 Idx = Constraints.SinglePointConstraints.Add(InConstraint);
 
 		if (bHandlesEnabled)
 		{
 			FConstraintContainerHandle* Handle = HandleAllocator.template AllocHandle< TRigidBodyPointContactConstraint<T, d> >(this, Idx);
 			Handle->GetContact().Timestamp = -INT_MAX; // force point constraints to be deleted.
 
-			PointConstraints[Idx].ConstraintHandle = Handle;
+			Constraints.SinglePointConstraints[Idx].ConstraintHandle = Handle;
 
 			check(Handle != nullptr);
 			Handles.Add(Handle);
@@ -189,20 +204,20 @@ namespace Chaos
 #endif
 		}
 
-		UpdateConstraintMaterialProperties(PointConstraints[Idx]);
+		UpdateConstraintMaterialProperties(Constraints.SinglePointConstraints[Idx]);
 	}
 
 	template<typename T, int d>
 	void TPBDCollisionConstraints<T, d>::AddConstraint(const TRigidBodySweptPointContactConstraint<FReal, 3>& InConstraint)
 	{
-		int32 Idx = SweptPointConstraints.Add(InConstraint);
+		int32 Idx = Constraints.SinglePointSweptConstraints.Add(InConstraint);
 
 		if (bHandlesEnabled)
 		{
 			FConstraintContainerHandle* Handle = HandleAllocator.template AllocHandle< TRigidBodySweptPointContactConstraint<T, d> >(this, Idx);
 			Handle->GetContact().Timestamp = -INT_MAX; // force point constraints to be deleted.
 
-			SweptPointConstraints[Idx].ConstraintHandle = Handle;
+			Constraints.SinglePointSweptConstraints[Idx].ConstraintHandle = Handle;
 
 			if(ensure(Handle != nullptr))
 			{			
@@ -219,21 +234,21 @@ namespace Chaos
 		// manifolds with particle may be invalidated, so we should update manifolds with this particle even on first apply iteration.
 		AddParticleAlwaysUpdateManifold(InConstraint.Particle[0]);
 
-		UpdateConstraintMaterialProperties(SweptPointConstraints[Idx]);
+		UpdateConstraintMaterialProperties(Constraints.SinglePointSweptConstraints[Idx]);
 	}
 
 
 	template<typename T, int d>
 	void TPBDCollisionConstraints<T, d>::AddConstraint(const TRigidBodyMultiPointContactConstraint<FReal, 3>& InConstraint)
 	{
-		int32 Idx = IterativeConstraints.Add(InConstraint);
+		int32 Idx = Constraints.MultiPointConstraints.Add(InConstraint);
 
 		if (bHandlesEnabled)
 		{
 			FConstraintContainerHandle* Handle = HandleAllocator.template AllocHandle< TRigidBodyMultiPointContactConstraint<T, d> >(this, Idx);
 			Handle->GetContact().Timestamp = LifespanCounter;
 
-			IterativeConstraints[Idx].ConstraintHandle = Handle;
+			Constraints.MultiPointConstraints[Idx].ConstraintHandle = Handle;
 
 			check(Handle != nullptr);
 			Handles.Add(Handle);
@@ -244,7 +259,7 @@ namespace Chaos
 #endif
 		}
 
-		UpdateConstraintMaterialProperties(IterativeConstraints[Idx]);
+		UpdateConstraintMaterialProperties(Constraints.MultiPointConstraints[Idx]);
 	}
 
 	template<typename T, int d>
@@ -276,9 +291,7 @@ namespace Chaos
 		{
 			HandleAllocator.FreeHandle(Handle);
 		}
-		PointConstraints.Reset();
-		SweptPointConstraints.Reset();
-		IterativeConstraints.Reset();
+		Constraints.Reset();
 		ParticlesAlwaysUpdatingManifold.Reset();
 		Handles.Reset();
 #endif
@@ -334,50 +347,50 @@ namespace Chaos
 		if (ConstraintType == FCollisionConstraintBase::FType::SinglePoint)
 		{
 #if CHAOS_COLLISION_PERSISTENCE_ENABLED
-			if (Idx < PointConstraints.Num() - 1)
+			if (Idx < Constraints.SinglePointConstraints.Num() - 1)
 			{
 				// update the handle
-				FConstraintContainerHandleKey Key = FPBDCollisionConstraintHandle::MakeKey(&PointConstraints.Last());
+				FConstraintContainerHandleKey Key = FPBDCollisionConstraintHandle::MakeKey(&Constraints.SinglePointConstraints.Last());
 				Manifolds[Key]->SetConstraintIndex(Idx, ConstraintType);
 			}
 #endif
-			PointConstraints.RemoveAtSwap(Idx);
-			if (bHandlesEnabled && (Idx < PointConstraints.Num()))
+			Constraints.SinglePointConstraints.RemoveAtSwap(Idx);
+			if (bHandlesEnabled && (Idx < Constraints.SinglePointConstraints.Num()))
 			{
-				PointConstraints[Idx].ConstraintHandle->SetConstraintIndex(Idx, FCollisionConstraintBase::FType::SinglePoint);
+				Constraints.SinglePointConstraints[Idx].ConstraintHandle->SetConstraintIndex(Idx, FCollisionConstraintBase::FType::SinglePoint);
 			}
 
 		}
 		else if (ConstraintType == FCollisionConstraintBase::FType::SinglePointSwept)
 		{
 #if CHAOS_COLLISION_PERSISTENCE_ENABLED
-			if (Idx < SweptPointConstraints.Num() - 1)
+			if (Idx < Constraints.SinglePointSweptConstraints.Num() - 1)
 			{
 				// update the handle
-				FConstraintContainerHandleKey Key = FPBDCollisionConstraintHandle::MakeKey(&SweptPointConstraints.Last());
+				FConstraintContainerHandleKey Key = FPBDCollisionConstraintHandle::MakeKey(&Constraints.SinglePointSweptConstraints.Last());
 				Manifolds[Key]->SetConstraintIndex(Idx, ConstraintType);
 			}
 #endif
-			SweptPointConstraints.RemoveAtSwap(Idx);
-			if (bHandlesEnabled && (Idx < SweptPointConstraints.Num()))
+			Constraints.SinglePointSweptConstraints.RemoveAtSwap(Idx);
+			if (bHandlesEnabled && (Idx < Constraints.SinglePointSweptConstraints.Num()))
 			{
-				SweptPointConstraints[Idx].ConstraintHandle->SetConstraintIndex(Idx, FCollisionConstraintBase::FType::SinglePointSwept);
+				Constraints.SinglePointSweptConstraints[Idx].ConstraintHandle->SetConstraintIndex(Idx, FCollisionConstraintBase::FType::SinglePointSwept);
 			}
 		}
 		else if (ConstraintType == FCollisionConstraintBase::FType::MultiPoint)
 		{
 #if CHAOS_COLLISION_PERSISTENCE_ENABLED
-			if (Idx < IterativeConstraints.Num() - 1)
+			if (Idx < Constraints.MultiPointConstraints.Num() - 1)
 			{
 				// update the handle
-				FConstraintContainerHandleKey Key = FPBDCollisionConstraintHandle::MakeKey(&IterativeConstraints.Last());
+				FConstraintContainerHandleKey Key = FPBDCollisionConstraintHandle::MakeKey(&Constraints.MultiPointConstraints.Last());
 				Manifolds[Key]->SetConstraintIndex(Idx, ConstraintType);
 			}
 #endif
-			IterativeConstraints.RemoveAtSwap(Idx);
-			if (bHandlesEnabled && (Idx < IterativeConstraints.Num()))
+			Constraints.MultiPointConstraints.RemoveAtSwap(Idx);
+			if (bHandlesEnabled && (Idx < Constraints.MultiPointConstraints.Num()))
 			{
-				IterativeConstraints[Idx].ConstraintHandle->SetConstraintIndex(Idx, FCollisionConstraintBase::FType::MultiPoint);
+				Constraints.MultiPointConstraints[Idx].ConstraintHandle->SetConstraintIndex(Idx, FCollisionConstraintBase::FType::MultiPoint);
 			}
 		}
 		else 
@@ -394,7 +407,7 @@ namespace Chaos
 			Manifolds.Remove(KeyToRemove);
 #endif
 			Handles.Remove(Handle);
-			check(Handles.Num() == PointConstraints.Num() + SweptPointConstraints.Num() + IterativeConstraints.Num());
+			check(Handles.Num() == Constraints.SinglePointConstraints.Num() + Constraints.SinglePointSweptConstraints.Num() + Constraints.MultiPointConstraints.Num());
 
 			HandleAllocator.FreeHandle(Handle);
 		}
@@ -427,7 +440,7 @@ namespace Chaos
 		//	}
 		//}, bDisableCollisionParallelFor);
 
-		for (FPointContactConstraint& Contact : PointConstraints)
+		for (FPointContactConstraint& Contact : Constraints.SinglePointConstraints)
 		{
 			Collisions::Update(Contact, MCullDistance);
 			if (Contact.GetPhi() < MCullDistance)
@@ -454,7 +467,7 @@ namespace Chaos
 
 		FCollisionContext Context;
 
-		for (FMultiPointContactConstraint& Contact : IterativeConstraints)
+		for (FMultiPointContactConstraint& Contact : Constraints.MultiPointConstraints)
 		{
 			Collisions::UpdateManifold(Contact, MCullDistance, Context);
 			if (Contact.GetPhi() < MCullDistance)
@@ -488,21 +501,36 @@ namespace Chaos
 			const Collisions::TContactParticleParameters<T> ParticleParameters = { MCullDistance, MShapePadding, &MCollided };
 			const Collisions::TContactIterationParameters<T> IterationParameters = { Dt, Iterations, NumIterations, MApplyPairIterations, ApplyType, &bNeedsAnotherIteration, &bAlwaysUpdateManifold };
 
-			for (FPointContactConstraint& Contact : PointConstraints)
+			NumActivePointConstraints = 0;
+			for (FPointContactConstraint& Contact : Constraints.SinglePointConstraints)
 			{
-				Collisions::ApplySinglePoint(Contact, IterationParameters, ParticleParameters);
+				if (!Contact.GetDisabled())
+				{
+					Collisions::ApplySinglePoint(Contact, IterationParameters, ParticleParameters);
+					++NumActivePointConstraints;
+				}
 			}
 
-			for (FMultiPointContactConstraint& Contact : IterativeConstraints)
+			NumActiveIterativeConstraints = 0;
+			for (FMultiPointContactConstraint& Contact : Constraints.MultiPointConstraints)
 			{
-				Collisions::ApplyMultiPoint(Contact, IterationParameters, ParticleParameters);
+				if (!Contact.GetDisabled())
+				{
+					Collisions::ApplyMultiPoint(Contact, IterationParameters, ParticleParameters);
+					++NumActiveIterativeConstraints;
+				}
 			}
 
 			// Swept apply may significantly change particle position, invalidating other constraint's manifolds.
 			// We don't update manifolds on first apply iteration, so make sure we apply swept constraints last.
-			for (FSweptPointContactConstraint& Contact : SweptPointConstraints)
+			NumActiveSweptPointConstraints = 0;
+			for (FSweptPointContactConstraint& Contact : Constraints.SinglePointSweptConstraints)
 			{
-				Collisions::Apply(Contact, IterationParameters, ParticleParameters);
+				if (!Contact.GetDisabled())
+				{
+					Collisions::Apply(Contact, IterationParameters, ParticleParameters);
+					++NumActiveSweptPointConstraints;
+				}
 			}
 		}
 
@@ -526,19 +554,28 @@ namespace Chaos
 			const Collisions::TContactParticleParameters<T> ParticleParameters = { MCullDistance, MShapePadding, &MCollided };
 			const Collisions::TContactIterationParameters<T> IterationParameters = { Dt, Iterations, NumIterations, MApplyPushOutPairIterations, ECollisionApplyType::None, &bNeedsAnotherIteration };
 
-			for (FPointContactConstraint& Contact : PointConstraints)
+			for (FPointContactConstraint& Contact : Constraints.SinglePointConstraints)
 			{
-				Collisions::ApplyPushOutSinglePoint(Contact, TempStatic, IterationParameters, ParticleParameters);
+				if (!Contact.GetDisabled())
+				{
+					Collisions::ApplyPushOutSinglePoint(Contact, TempStatic, IterationParameters, ParticleParameters);
+				}
 			}
 
-			for (FSweptPointContactConstraint& Contact : SweptPointConstraints)
+			for (FSweptPointContactConstraint& Contact : Constraints.SinglePointSweptConstraints)
 			{
-				Collisions::ApplyPushOut(Contact, TempStatic, IterationParameters, ParticleParameters);
+				if (!Contact.GetDisabled())
+				{
+					Collisions::ApplyPushOut(Contact, TempStatic, IterationParameters, ParticleParameters);
+				}
 			}
 
-			for (FMultiPointContactConstraint& Contact : IterativeConstraints)
+			for (FMultiPointContactConstraint& Contact : Constraints.MultiPointConstraints)
 			{
-				Collisions::ApplyPushOutMultiPoint(Contact, TempStatic, IterationParameters, ParticleParameters);
+				if (!Contact.GetDisabled())
+				{
+					Collisions::ApplyPushOutMultiPoint(Contact, TempStatic, IterationParameters, ParticleParameters);
+				}
 			}
 		}
 
@@ -621,20 +658,20 @@ namespace Chaos
 	{
 		check(Index < NumConstraints());
 		
-		if (Index < PointConstraints.Num())
+		if (Index < Constraints.SinglePointConstraints.Num())
 		{
-			return PointConstraints[Index];
+			return Constraints.SinglePointConstraints[Index];
 		}
-		Index -= PointConstraints.Num();
+		Index -= Constraints.SinglePointConstraints.Num();
 
-		if (Index < SweptPointConstraints.Num())
+		if (Index < Constraints.SinglePointSweptConstraints.Num())
 		{
-			return SweptPointConstraints[Index];
+			return Constraints.SinglePointSweptConstraints[Index];
 		}
-		Index -= SweptPointConstraints.Num();
+		Index -= Constraints.SinglePointSweptConstraints.Num();
 
 		{
-			return IterativeConstraints[Index];
+			return Constraints.MultiPointConstraints[Index];
 		}
 	}
 
diff --git a/Engine/Source/Runtime/Experimental/Chaos/Private/Chaos/PBDCollisionConstraints.ispc b/Engine/Source/Runtime/Experimental/Chaos/Private/Chaos/PBDCollisionConstraints.ispc
--- a/Engine/Source/Runtime/Experimental/Chaos/Private/Chaos/PBDCollisionConstraints.ispc
+++ b/Engine/Source/Runtime/Experimental/Chaos/Private/Chaos/PBDCollisionConstraints.ispc
@@ -1,5 +1,9 @@
 // Copyright Epic Games, Inc. All Rights Reserved.
 
+#define EXPLICIT_VECTOR4 1
+
+#include "Chaos/Matrix33.isph"
+#include "Math/Quat.isph"
 #include "Math/Vector.isph"
 
 #define SMALL_NUMBER		(1.e-8f)
@@ -731,4 +735,232 @@ export uniform bool SampleSphereNoNormalAll(const uniform float &ObjectRadius,
 	}
 
 	return bPhiFound;
-}
\ No newline at end of file
+}
+
+
+//
+//
+///////////////////////////////////////////////////////////////////////////////
+//
+//
+
+struct FCollisionContact
+{
+	bool bDisabled;
+	FVector Normal;
+	FVector Location;
+	float Phi;
+
+	float Friction;
+	float AngularFriction;
+	float Restitution;
+	float InvInertiaScale0;
+	float InvInertiaScale1;
+	int ShapesType;
+	int* Implicit[2];
+};
+
+export uniform int32 SizeofFCollisionContact()
+{
+	return sizeof(uniform FCollisionContact);
+}
+
+// @todo(ccaulfield): move to ispc libs
+inline uniform FMatrix33 Multiply(const uniform FMatrix33& L, const uniform float R)
+{
+	uniform FMatrix33 M;
+	for (uniform int Index = 0; Index < 12; ++Index)
+	{
+		M.M[Index] = L.M[Index] * R;
+	}
+	M.M[12] = 0.0f;
+	M.M[13] = 0.0f;
+	M.M[14] = 0.0f;
+	M.M[15] = 1.0f;
+	return M;
+}
+
+
+// Get the CoM position given the Actor position
+static inline uniform FVector GetCoMWorldPosition(
+	const uniform FVector& PActorWorld,
+	const uniform FVector4& QActorWorld,
+	const uniform FVector& PCoMLocal)
+{
+	return PActorWorld + VectorQuaternionRotateVector(QActorWorld, PCoMLocal);
+}
+
+// Get the CoM rotation given the Actor rotation
+static inline uniform FVector4 GetCoMWorldRotation(
+	const uniform FVector4& QActorWorld,
+	const uniform FVector4& QCoMLocal)
+{
+	return VectorQuaternionMultiply2(QActorWorld, QCoMLocal);
+}
+
+// Set the Actor position and rotation given the CoM position and rotation
+static inline uniform FVector SetCoMWorldTransform(
+	uniform FVector& PActorWorld,
+	uniform FVector4& QActorWorld,
+	const uniform FVector& PCoMWorld,
+	const uniform FVector4& QCoMWorld,
+	const uniform FVector& PCoMLocal,
+	const uniform FVector4& QCoMLocal)
+{
+	uniform FVector4 QCoMLocalInverse = QuatInverse(QCoMLocal);
+	QActorWorld = VectorQuaternionMultiply2(QCoMWorld, QCoMLocalInverse);
+	PActorWorld = PCoMWorld - VectorQuaternionRotateVector(QActorWorld, PCoMLocal);
+}
+
+// @todo(ccaulfield): dedupe
+static inline uniform FMatrix33 ComputeWorldSpaceInertia(const uniform FVector4 &Q, const uniform FVector &I)
+{
+	const uniform FMatrix33 QM = MakeQuatRotationTranslationMatrix(Q);
+	const uniform FMatrix33 L = SetMatrix33(I.V[0], I.V[1], I.V[2]);
+	return MultiplyAB(QM, MultiplyABt(L, QM));
+}
+
+// @todo(ccaulfield): dedupe
+static inline uniform FMatrix33 ComputeJointFactorMatrix(const uniform FVector& V, const uniform FMatrix33& M, const uniform float Im)
+{
+	// Rigid objects rotational contribution to the impulse.
+	// Vx*M*VxT+Im
+	const uniform FVector V221 = SetVector(-V.V[2], V.V[2], -V.V[1]);
+	const uniform FVector V222 = SetVector(-V.V[2], -V.V[2], -V.V[2]);
+	const uniform FVector M544 = SetVector(M.M[5], M.M[4], M.M[4]);
+	const uniform FVector V111 = SetVector(V.V[1], V.V[1], V.V[1]);
+	const uniform FVector M988 = SetVector(M.M[9], M.M[8], M.M[8]);
+	const uniform FVector V100 = SetVector(V.V[1], -V.V[0], V.V[0]);
+	const uniform FVector M995 = SetVector(M.M[9], M.M[9], M.M[5]);
+	const uniform FVector M10109 = SetVector(M.M[10], M.M[10], M.M[9]);
+	const uniform FVector IM = SetVector(Im, 0, 0);
+
+	const uniform FVector R1 = (V221 * (V222 * M544 + V111 * M988)) + (V100 * (V222 * M995 + V111 * M10109)) + IM;
+
+	const uniform FVector V211 = SetVector(V.V[2], -V.V[1], -V.V[1]);
+	const uniform FVector V221_2 = SetVector(V.V[2], V.V[2], -V.V[1]);
+	const uniform FVector M000 = SetVector(M.M[0], M.M[0], M.M[0]);
+	const uniform FVector V000 = SetVector(-V.V[0], -V.V[0], V.V[0]);
+	const uniform FVector M884 = SetVector(M.M[8], M.M[8], M.M[4]);
+	const uniform FVector V000_2 = SetVector(-V.V[0], V.V[0], V.V[0]);
+	const uniform FVector M844 = SetVector(M.M[8], M.M[4], M.M[4]);
+	const uniform FVector M1095 = SetVector(M.M[10], M.M[9], M.M[5]);
+	const uniform FVector IM_2 = SetVector(Im, 0, Im);
+
+	const uniform FVector R2 = (V211 * (V221_2 * M000 + V000 * M884)) + (V000_2 * (V221_2 * M844 + V000 * M1095)) + IM_2;
+
+	return SetMatrix33(R1.V[0], R1.V[1], R1.V[2], R2.V[0], R2.V[1], R2.V[2]);
+}
+
+// Apply the contact force. Update the actors' positions and rotations.
+export void ApplyContact2(
+	const uniform FCollisionContact* uniform Contact,
+	uniform FVector& PActor0,
+	uniform FVector4& QActor0,
+	uniform FVector& PActor1,
+	uniform FVector4& QActor1,
+	const uniform float InvM0,
+	const uniform FVector& InvILocal0,
+	const uniform FVector& PCoMLocal0,
+	const uniform FVector4& QCoMLocal0,
+	const uniform float InvM1,
+	const uniform FVector& InvILocal1,
+	const uniform FVector& PCoMLocal1,
+	const uniform FVector4& QCoMLocal1)
+{
+	uniform FVector P0 = GetCoMWorldPosition(PActor0, QActor0, PCoMLocal0);
+	uniform FVector P1 = GetCoMWorldPosition(PActor1, QActor1, PCoMLocal1);
+	uniform FVector4 Q0 = GetCoMWorldRotation(QActor0, QCoMLocal0);
+	uniform FVector4 Q1 = GetCoMWorldRotation(QActor1, QCoMLocal1);
+	uniform FVector VectorToPoint0 = Contact->Location - P0;
+	uniform FVector VectorToPoint1 = Contact->Location - P1;
+	
+	uniform FMatrix33 ContactInvI = Matrix33Zero;
+	uniform FMatrix33 InvI0 = Matrix33Zero;
+	uniform FMatrix33 InvI1 = Matrix33Zero;
+	if (InvM0 > 0.0f)
+	{
+		InvI0 = Multiply(ComputeWorldSpaceInertia(Q0, InvILocal0), Contact->InvInertiaScale0);
+		ContactInvI = AddAB(ContactInvI, ComputeJointFactorMatrix(VectorToPoint0, InvI0, InvM0));
+	}
+	if (InvM1 > 0.0f)
+	{
+		InvI1 = Multiply(ComputeWorldSpaceInertia(Q1, InvILocal1), Contact->InvInertiaScale1);
+		ContactInvI = AddAB(ContactInvI, ComputeJointFactorMatrix(VectorToPoint1, InvI1, InvM1));
+	}
+	
+	// Calculate the normal correction
+	// @todo(chaos): support restitution in this mode (requires a small-impact cutoff threshold)
+	uniform float Restitution = 0.0f;//Contact->Restitution;
+	uniform FVector NormalError = Contact->Phi * Contact->Normal;
+	uniform float NormalImpulseDenominator = VectorDot(Contact->Normal, Multiply(ContactInvI, Contact->Normal));
+	uniform FVector NormalImpulseNumerator = -(1.0f + Restitution) * NormalError;
+	uniform FVector NormalCorrection = NormalImpulseNumerator / NormalImpulseDenominator;
+	
+	// Calculate the lateral correction
+	uniform FVector LateralCorrection = ZeroVector;
+	//if (Contact->Friction > 0)
+	//{
+	//	// Get contact velocity
+	//	uniform FVector X0 = FParticleUtilitiesXR::GetCoMWorldPosition(Particle0);
+	//	uniform FVector X1 = FParticleUtilitiesXR::GetCoMWorldPosition(Particle1);
+	//	FRotation3 R0 = FParticleUtilitiesXR::GetCoMWorldRotation(Particle0);
+	//	FRotation3 R1 = FParticleUtilitiesXR::GetCoMWorldRotation(Particle1);
+	//	uniform FVector V0 = FVec3::CalculateVelocity(X0, P0, IterationParameters.Dt);
+	//	uniform FVector W0 = FRotation3::CalculateAngularVelocity(R0, Q0, IterationParameters.Dt);
+	//	uniform FVector V1 = FVec3::CalculateVelocity(X1, P1, IterationParameters.Dt);
+	//	uniform FVector W1 = FRotation3::CalculateAngularVelocity(R1, Q1, IterationParameters.Dt);
+	//	uniform FVector CV0 = V0 + FVec3::CrossProduct(W0, VectorToPoint0);
+	//	uniform FVector CV1 = V1 + FVec3::CrossProduct(W1, VectorToPoint1);
+	//	uniform FVector CV = CV0 - CV1;
+	//
+	//	// Calculate lateral correction, clamped to the friction cone. Kinda.
+	//	uniform float CVNormalMag = FVec3::DotProduct(CV, Contact->Normal);
+	//	if (CVNormalMag < 0.0f)
+	//	{
+	//		uniform FVector CVLateral = CV - CVNormalMag * Contact->Normal;
+	//		uniform float CVLateralMag = CVLateral.Size();
+	//		if (CVLateralMag > KINDA_SMALL_NUMBER)
+	//		{
+	//			uniform FVector DirLateral = CVLateral / CVLateralMag;
+	//			uniform FVector LateralImpulseNumerator = -CVLateral * IterationParameters.Dt;
+	//			uniform float LateralImpulseDenominator = FVec3::DotProduct(DirLateral, ContactInvI * DirLateral);
+	//			LateralCorrection = LateralImpulseNumerator / LateralImpulseDenominator;
+	//			uniform float LateralImpulseMag = LateralCorrection.Size();
+	//			uniform float NormalImpulseMag = NormalCorrection.Size();
+	//			if (LateralImpulseMag > Contact->Friction * NormalImpulseMag)
+	//			{
+	//				LateralCorrection *= Contact->Friction * NormalImpulseMag / LateralImpulseMag;
+	//			}
+	//		}
+	//	}
+	//}
+
+	// Net Correction
+	uniform FVector DX = NormalCorrection + LateralCorrection;
+				
+	if (InvM0 > 0.0f)
+	{
+		uniform FVector DP0 = InvM0 * DX;
+		uniform FVector DR0 = Multiply(InvI0, VectorCross(VectorToPoint0, DX));
+		uniform FVector4 DQ0 = VectorQuaternionMultiply2(SetVector4(DR0, 0.0f), Q0) * 0.5f;
+	
+		P0 = P0 + DP0;
+		Q0 = VectorNormalizeQuaternion(Q0 + DQ0);
+		
+		SetCoMWorldTransform(PActor0, QActor0, P0, Q0, PCoMLocal0, QCoMLocal0);
+	}
+
+	if (InvM1 > 0.0f)
+	{
+		uniform FVector DXNeg = VectorNegate(DX);
+		uniform FVector DP1 = InvM1 * DXNeg;
+		uniform FVector DR1 = Multiply(InvI1, VectorCross(VectorToPoint1, DXNeg));
+		uniform FVector4 DQ1 = VectorQuaternionMultiply2(SetVector4(DR1, 0.0f), Q1) * 0.5f;
+		
+		P1 = P1 + DP1;
+		Q1 = VectorNormalizeQuaternion(Q1 + DQ1);
+		
+		SetCoMWorldTransform(PActor1, QActor1, P1, Q1, PCoMLocal1, QCoMLocal1);
+	}
+}
diff --git a/Engine/Source/Runtime/Experimental/Chaos/Private/Chaos/PBDCollisionConstraintsContact.cpp b/Engine/Source/Runtime/Experimental/Chaos/Private/Chaos/PBDCollisionConstraintsContact.cpp
--- a/Engine/Source/Runtime/Experimental/Chaos/Private/Chaos/PBDCollisionConstraintsContact.cpp
+++ b/Engine/Source/Runtime/Experimental/Chaos/Private/Chaos/PBDCollisionConstraintsContact.cpp
@@ -9,8 +9,16 @@
 
 //PRAGMA_DISABLE_OPTIMIZATION
 
+#if INTEL_ISPC
+#include "PBDCollisionConstraints.ispc.generated.h"
+#endif
+
 namespace Chaos
 {
+#if INTEL_ISPC
+	extern bool bChaos_Collision_ISPC_Enabled;
+#endif
+
 	namespace Collisions
 	{
 		int32 Chaos_Collision_EnergyClampEnabled = 1;
@@ -228,13 +236,57 @@ namespace Chaos
 			const TContactParticleParameters<FReal>& ParticleParameters)
 		{
 			FVec3 AccumulatedImpulse(0);
-
 			TPBDRigidParticleHandle<FReal, 3>* PBDRigid0 = Particle0->CastToRigidParticle();
 			TPBDRigidParticleHandle<FReal, 3>* PBDRigid1 = Particle1->CastToRigidParticle();
-
 			bool bIsRigidDynamic0 = PBDRigid0 && PBDRigid0->ObjectState() == EObjectStateType::Dynamic;
 			bool bIsRigidDynamic1 = PBDRigid1 && PBDRigid1->ObjectState() == EObjectStateType::Dynamic;
 
+#if INTEL_ISPC
+			if (bChaos_Collision_ISPC_Enabled)
+			{
+				FVec3 PActor0 = Particle0->P();
+				FRotation3 QActor0 = Particle0->Q();
+				FVec3 PActor1 = Particle1->P();
+				FRotation3 QActor1 = Particle1->Q();
+				const FReal InvM0 = Particle0->InvM();
+				const FVec3 InvI0 = Particle0->InvI().GetDiagonal();
+				const FVec3 PCoM0 = Particle0->CenterOfMass();
+				const FRotation3 QCoM0 = Particle0->RotationOfMass();
+				const FReal InvM1 = Particle1->InvM();
+				const FVec3 InvI1 = Particle1->InvI().GetDiagonal();
+				const FVec3 PCoM1 = Particle1->CenterOfMass();
+				const FRotation3 QCoM1 = Particle1->RotationOfMass();
+				ispc::ApplyContact2(
+					(ispc::FCollisionContact*)&Contact,
+					(ispc::FVector&)PActor0,
+					(ispc::FVector4&)QActor0,
+					(ispc::FVector&)PActor1,
+					(ispc::FVector4&)QActor1,
+					InvM0,
+					(const ispc::FVector&)InvI0,
+					(const ispc::FVector&)PCoM0,
+					(const ispc::FVector4&)QCoM0,
+					InvM1,
+					(const ispc::FVector&)InvI1,
+					(const ispc::FVector&)PCoM1,
+					(const ispc::FVector4&)QCoM1);
+
+				if (bIsRigidDynamic0)
+				{
+					PBDRigid0->SetP(PActor0);
+					PBDRigid0->SetQ(QActor0);
+				}
+				if (bIsRigidDynamic1)
+				{
+					PBDRigid1->SetP(PActor1);
+					PBDRigid1->SetQ(QActor1);
+				}
+
+				*IterationParameters.NeedsAnotherIteration = true;
+				return AccumulatedImpulse;
+			}
+#endif
+
 			FVec3 P0 = FParticleUtilities::GetCoMWorldPosition(Particle0);
 			FVec3 P1 = FParticleUtilities::GetCoMWorldPosition(Particle1);
 			FRotation3 Q0 = FParticleUtilities::GetCoMWorldRotation(Particle0);
@@ -338,6 +390,13 @@ namespace Chaos
 				// Update the contact information based on current particles' positions
 				Collisions::Update(Constraint, ParticleParameters.CullDistance);
 
+				// Permanently disable a constraint that is beyond the cull distance
+				if (Constraint.GetPhi() >= ParticleParameters.CullDistance)
+				{
+					Constraint.SetDisabled(true);
+					return;
+				}
+
 				if (Constraint.GetPhi() >= ParticleParameters.ShapePadding)
 				{
 					return;
diff --git a/Engine/Source/Runtime/Experimental/Chaos/Public/Chaos/Collision/CollisionReceiver.h b/Engine/Source/Runtime/Experimental/Chaos/Public/Chaos/Collision/CollisionReceiver.h
--- a/Engine/Source/Runtime/Experimental/Chaos/Public/Chaos/Collision/CollisionReceiver.h
+++ b/Engine/Source/Runtime/Experimental/Chaos/Public/Chaos/Collision/CollisionReceiver.h
@@ -101,23 +101,20 @@ namespace Chaos
 		{
 		}
 
+		FCollisionConstraintsArray& GetConstraintsArray()
+		{
+			return CollisionConstraints.GetConstraintsArray();
+		}
+
 		/**
-		 * Called by a CollisionDetector (possibly in a task) when it finds collisions.
+		 * Does nothing - we write directly into the Constraints Container
 		 */
 		void ReceiveCollisions(const FCollisionConstraintsArray& Constraints)
 		{
-			for (const TRigidBodyPointContactConstraint<FReal, 3>& Constraint : Constraints.SinglePointConstraints)
-			{
-				CollisionConstraints.AddConstraint(Constraint);
-			}
-			for (const TRigidBodyMultiPointContactConstraint<FReal, 3>& Constraint : Constraints.MultiPointConstraints)
-			{
-				CollisionConstraints.AddConstraint(Constraint);
-			}
 		}
 
 		/**
-		 * No-op
+		 * Does nothing - we write directly into the Constraints Container
 		 */
 		void ProcessCollisions()
 		{
diff --git a/Engine/Source/Runtime/Experimental/Chaos/Public/Chaos/Collision/NarrowPhase.h b/Engine/Source/Runtime/Experimental/Chaos/Public/Chaos/Collision/NarrowPhase.h
--- a/Engine/Source/Runtime/Experimental/Chaos/Public/Chaos/Collision/NarrowPhase.h
+++ b/Engine/Source/Runtime/Experimental/Chaos/Public/Chaos/Collision/NarrowPhase.h
@@ -50,7 +50,7 @@ namespace Chaos
 				//   determine if the constraint is already defined, and then opt out of 
 				//   the creation process. 
 				//
-				Collisions::ConstructConstraints<FReal, 3>(Particle0, Particle1, Particle0->Geometry().Get(), Particle1->Geometry().Get(), Collisions::GetTransform(Particle0), Collisions::GetTransform(Particle1), CullDistance, Context, NewConstraints);
+				Collisions::ConstructConstraints<FReal, 3>(Particle0, Particle1, Particle0->Geometry().Get(), Particle1->Geometry().Get(), FRigidTransform3(), FRigidTransform3(), CullDistance, Context, NewConstraints);
 
 				CHAOS_COLLISION_STAT(if (NewConstraints.Num()) { StatData.IncrementCountNP(NewConstraints.Num()); });
 				CHAOS_COLLISION_STAT(if (!NewConstraints.Num()) { StatData.IncrementRejectedNP(); });
diff --git a/Engine/Source/Runtime/Experimental/Chaos/Public/Chaos/Collision/ParticlePairBroadPhase.h b/Engine/Source/Runtime/Experimental/Chaos/Public/Chaos/Collision/ParticlePairBroadPhase.h
--- a/Engine/Source/Runtime/Experimental/Chaos/Public/Chaos/Collision/ParticlePairBroadPhase.h
+++ b/Engine/Source/Runtime/Experimental/Chaos/Public/Chaos/Collision/ParticlePairBroadPhase.h
@@ -49,6 +49,8 @@ namespace Chaos
 		{
 			SCOPE_CYCLE_COUNTER(STAT_Collisions_BroadPhase);
 
+			FCollisionConstraintsArray& ConstraintsArray = Receiver.GetConstraintsArray();
+
 			int32 NumPairs = ParticlePairs.Num();
 			for (int32 PairIndex = 0; PairIndex < NumPairs; ++PairIndex)
 			{
@@ -62,9 +64,7 @@ namespace Chaos
 					TGeometryParticleHandle<FReal, 3>* Particle0 = const_cast<TGeometryParticleHandle<FReal, 3>*>(ParticlePair[0]);
 					TGeometryParticleHandle<FReal, 3>* Particle1 = const_cast<TGeometryParticleHandle<FReal, 3>*>(ParticlePair[1]);
 
-					NarrowPhase.GenerateCollisions(NewConstraints, Dt, Particle0, Particle1, CullDistance, StatData);
-					Receiver.ReceiveCollisions(NewConstraints);
-					NewConstraints.Empty();
+					NarrowPhase.GenerateCollisions(ConstraintsArray, Dt, Particle0, Particle1, CullDistance, StatData);
 				}
 			}
 		}
@@ -72,6 +72,5 @@ namespace Chaos
 	private:
 		const TArray<FParticlePair>& ParticlePairs;
 		FReal CullDistance;
-		FCollisionConstraintsArray NewConstraints;
 	};
 }
diff --git a/Engine/Source/Runtime/Experimental/Chaos/Public/Chaos/CollisionResolutionTypes.h b/Engine/Source/Runtime/Experimental/Chaos/Public/Chaos/CollisionResolutionTypes.h
--- a/Engine/Source/Runtime/Experimental/Chaos/Public/Chaos/CollisionResolutionTypes.h
+++ b/Engine/Source/Runtime/Experimental/Chaos/Public/Chaos/CollisionResolutionTypes.h
@@ -74,7 +74,7 @@ namespace Chaos
 	{
 	public:
 		TCollisionContact(const FImplicitObject* InImplicit0 = nullptr, const FImplicitObject* InImplicit1 = nullptr)
-		    : bDisabled(true), Normal(0), Location(0), Phi(FLT_MAX), Friction(0), AngularFriction(0), Restitution(0), InvInertiaScale0(1.f), InvInertiaScale1(1.f), ShapesType(EContactShapesType::Unknown)
+		    : bDisabled(false), Normal(0), Location(0), Phi(FLT_MAX), Friction(0), AngularFriction(0), Restitution(0), InvInertiaScale0(1.f), InvInertiaScale1(1.f), ShapesType(EContactShapesType::Unknown)
 		{
 			Implicit[0] = InImplicit0;
 			Implicit[1] = InImplicit1;
@@ -384,42 +384,42 @@ namespace Chaos
 		{ 
 			if (GetType() == FConstraintBase::FType::SinglePoint)
 			{
-				return ConstraintContainer->PointConstraints[ConstraintIndex];
+				return ConstraintContainer->Constraints.SinglePointConstraints[ConstraintIndex];
 			}
 			else if (GetType() == FConstraintBase::FType::SinglePointSwept)
 			{
-				return ConstraintContainer->SweptPointConstraints[ConstraintIndex];
+				return ConstraintContainer->Constraints.SinglePointSweptConstraints[ConstraintIndex];
 			}
 
 			check(GetType() == FConstraintBase::FType::MultiPoint);
-			return ConstraintContainer->IterativeConstraints[ConstraintIndex];
+			return ConstraintContainer->Constraints.MultiPointConstraints[ConstraintIndex];
 		}
 
 		TCollisionConstraintBase<T, d>& GetContact()
 		{
 			if (GetType() == FConstraintBase::FType::SinglePoint)
 			{
-				return ConstraintContainer->PointConstraints[ConstraintIndex];
+				return ConstraintContainer->Constraints.SinglePointConstraints[ConstraintIndex];
 			}
 			else if (GetType() == FConstraintBase::FType::SinglePointSwept)
 			{
-				return ConstraintContainer->SweptPointConstraints[ConstraintIndex];
+				return ConstraintContainer->Constraints.SinglePointSweptConstraints[ConstraintIndex];
 			}
 
 			check(GetType() == FConstraintBase::FType::MultiPoint);
-			return ConstraintContainer->IterativeConstraints[ConstraintIndex];
+			return ConstraintContainer->Constraints.MultiPointConstraints[ConstraintIndex];
 		}
 
 
 
-		const TRigidBodyPointContactConstraint<T, d>& GetPointContact() const { check(GetType() == FConstraintBase::FType::SinglePoint); return ConstraintContainer->PointConstraints[ConstraintIndex]; }
-		TRigidBodyPointContactConstraint<T, d>& GetPointContact() { check(GetType() == FConstraintBase::FType::SinglePoint); return ConstraintContainer->PointConstraints[ConstraintIndex]; }
+		const TRigidBodyPointContactConstraint<T, d>& GetPointContact() const { check(GetType() == FConstraintBase::FType::SinglePoint); return ConstraintContainer->Constraints.SinglePointConstraints[ConstraintIndex]; }
+		TRigidBodyPointContactConstraint<T, d>& GetPointContact() { check(GetType() == FConstraintBase::FType::SinglePoint); return ConstraintContainer->Constraints.SinglePointConstraints[ConstraintIndex]; }
 
-		const TRigidBodySweptPointContactConstraint<T, d>& GetSweptPointContact() const { check(GetType() == FConstraintBase::FType::SinglePointSwept); return ConstraintContainer->SweptPointConstraints[ConstraintIndex]; }
-		TRigidBodySweptPointContactConstraint<T, d>& GetSweptPointContact() { check(GetType() == FConstraintBase::FType::SinglePointSwept); return ConstraintContainer->SweptPointConstraints[ConstraintIndex]; }
+		const TRigidBodySweptPointContactConstraint<T, d>& GetSweptPointContact() const { check(GetType() == FConstraintBase::FType::SinglePointSwept); return ConstraintContainer->Constraints.SinglePointSweptConstraints[ConstraintIndex]; }
+		TRigidBodySweptPointContactConstraint<T, d>& GetSweptPointContact() { check(GetType() == FConstraintBase::FType::SinglePointSwept); return ConstraintContainer->Constraints.SinglePointSweptConstraints[ConstraintIndex]; }
 
-		const TRigidBodyMultiPointContactConstraint<T, d>& GetMultiPointContact() const { check(GetType() == FConstraintBase::FType::MultiPoint); return ConstraintContainer->IterativeConstraints[ConstraintIndex]; }
-		TRigidBodyMultiPointContactConstraint<T, d>& GetMultiPointContact() { check(GetType() == FConstraintBase::FType::MultiPoint); return ConstraintContainer->IterativeConstraints[ConstraintIndex]; }
+		const TRigidBodyMultiPointContactConstraint<T, d>& GetMultiPointContact() const { check(GetType() == FConstraintBase::FType::MultiPoint); return ConstraintContainer->Constraints.MultiPointConstraints[ConstraintIndex]; }
+		TRigidBodyMultiPointContactConstraint<T, d>& GetMultiPointContact() { check(GetType() == FConstraintBase::FType::MultiPoint); return ConstraintContainer->Constraints.MultiPointConstraints[ConstraintIndex]; }
 
 		typename FConstraintBase::FType GetType() const { return ConstraintType; }
 
@@ -470,6 +470,13 @@ namespace Chaos
 
 		int32 Num() const { return SinglePointConstraints.Num() + SinglePointSweptConstraints.Num() + MultiPointConstraints.Num(); }
 
+		void Reset()
+		{
+			SinglePointConstraints.Reset();
+			SinglePointSweptConstraints.Reset();
+			MultiPointConstraints.Reset();
+		}
+
 		void Empty()
 		{
 			SinglePointConstraints.Empty();
@@ -477,6 +484,26 @@ namespace Chaos
 			MultiPointConstraints.Empty();
 		}
 
+
+		TRigidBodyPointContactConstraint<FReal, 3>* Add(const TRigidBodyPointContactConstraint<FReal, 3>& C)
+		{
+			int32 ConstraintIndex = SinglePointConstraints.Add(C);
+			return &SinglePointConstraints[ConstraintIndex];
+		}
+
+		TRigidBodySweptPointContactConstraint<FReal, 3>* Add(const TRigidBodySweptPointContactConstraint<FReal, 3>& C)
+		{
+			int32 ConstraintIndex = SinglePointSweptConstraints.Add(C);
+			return &SinglePointSweptConstraints[ConstraintIndex];
+		}
+
+		TRigidBodyMultiPointContactConstraint<FReal, 3>* Add(const TRigidBodyMultiPointContactConstraint<FReal, 3>& C)
+		{
+			int32 ConstraintIndex = MultiPointConstraints.Add(C);
+			return &MultiPointConstraints[ConstraintIndex];
+		}
+
+
 		TRigidBodyPointContactConstraint<FReal, 3>* TryAdd(FReal MaxPhi, const TRigidBodyPointContactConstraint<FReal, 3>& C)
 		{
 			if (C.GetPhi() < MaxPhi)
diff --git a/Engine/Source/Runtime/Experimental/Chaos/Public/Chaos/PBDCollisionConstraints.h b/Engine/Source/Runtime/Experimental/Chaos/Public/Chaos/PBDCollisionConstraints.h
--- a/Engine/Source/Runtime/Experimental/Chaos/Public/Chaos/PBDCollisionConstraints.h
+++ b/Engine/Source/Runtime/Experimental/Chaos/Public/Chaos/PBDCollisionConstraints.h
@@ -279,7 +279,7 @@ public:
 
 	int32 NumConstraints() const
 	{
-		return PointConstraints.Num() + SweptPointConstraints.Num() + IterativeConstraints.Num();
+		return Constraints.SinglePointConstraints.Num() + Constraints.SinglePointSweptConstraints.Num() + Constraints.MultiPointConstraints.Num();
 	}
 
 	FHandles& GetConstraintHandles()
@@ -294,6 +294,7 @@ public:
 
 	const FConstraintBase& GetConstraint(int32 Index) const;
 
+	FCollisionConstraintsArray& GetConstraintsArray() { return Constraints; }
 
 protected:
 	using Base::GetConstraintIndex;
@@ -305,9 +306,10 @@ private:
 
 	const TPBDRigidsSOAs<T,d>& Particles;
 
-	TArray<FPointContactConstraint> PointConstraints;
-	TArray<FSweptPointContactConstraint> SweptPointConstraints;
-	TArray<FMultiPointContactConstraint> IterativeConstraints;
+	FCollisionConstraintsArray Constraints;
+	int32 NumActivePointConstraints;
+	int32 NumActiveSweptPointConstraints;
+	int32 NumActiveIterativeConstraints;
 
 	TSet<const TGeometryParticleHandle<T, d>*> ParticlesAlwaysUpdatingManifold;
 
