Make it so dirty particles have a circular buffer

diff --git a/Engine/Source/Runtime/Experimental/ChaosSolvers/Public/RewindData.h b/Engine/Source/Runtime/Experimental/ChaosSolvers/Public/RewindData.h
--- a/Engine/Source/Runtime/Experimental/ChaosSolvers/Public/RewindData.h
+++ b/Engine/Source/Runtime/Experimental/ChaosSolvers/Public/RewindData.h
@@ -392,9 +392,9 @@ public:
 	{
 	}
 
-	uint32 Capacity() const { return Managers.Capacity(); }
+	int32 Capacity() const { return Managers.Capacity(); }
 
-	bool RewindToFrame(uint32 Frame)
+	bool RewindToFrame(int32 Frame)
 	{
 		ensure(IsInGameThread());
 
@@ -432,7 +432,7 @@ public:
 				if(const FGeometryParticleStateBase* RewindState = GetStateAtFrameImp(*DirtyParticleInfo.Particle, Frame))
 				{
 					LatestState.SyncIfDirty(*DestManager,DataIdx++,*DirtyParticleInfo.Particle,*RewindState);
-					CoalesceBack(RewindInfo.Frames);
+					CoalesceBack(RewindInfo.GetFrames(), CurFrame);
 
 					RewindState->SyncToParticle(*DirtyParticleInfo.Particle);
 				}
@@ -483,7 +483,7 @@ public:
 	void AdvanceFrame()
 	{
 		++CurFrame;
-		if(CurFrame > Managers.Capacity())
+		if(CurFrame > static_cast<int32>(Managers.Capacity()))
 		{
 			++EarliestFrame;
 		}
@@ -546,7 +546,7 @@ public:
 			{
 				FGeometryParticleStateBase& LatestState = Info.AddFrame(CurFrame-1);
 				LatestState.SyncPrevFrame(DestManager,DataIdx,Dirty);
-				CoalesceBack(Info.Frames);
+				CoalesceBack(Info.GetFrames(), CurFrame-1);
 			}
 
 			//If dynamics are dirty we must record them immediately because the sim resets them to 0
@@ -598,45 +598,112 @@ public:
 			LatestState.SyncSimWritablePropsFromSim(DestManager,DataIdx,Rigid);
 
 			//update any previous frames that were pointing at head
-			CoalesceBack(Info.Frames);
+			CoalesceBack(Info.GetFrames(), CurFrame-1);
 		}
 	}
 
 private:
 
-	struct FFrameInfo
+	class FFrameInfo
 	{
+	public:
+		FFrameInfo()
+		: bSet(false)
+		{
+
+		}
+
+		FGeometryParticleStateBase* GetState(int32 Frame)
+		{
+			return (bSet && Frame == RecordedFrame) ? &State : nullptr;
+		}
+
+		const FGeometryParticleStateBase* GetState(int32 Frame) const
+		{
+			return (bSet && Frame == RecordedFrame) ? &State : nullptr;
+		}
+
+		FGeometryParticleStateBase& GetStateChecked(int32 Frame)
+		{
+			check(bSet && Frame == RecordedFrame);
+			return State;
+		}
+
+		const FGeometryParticleStateBase& GetStateChecked(int32 Frame) const
+		{
+			check(bSet && Frame == RecordedFrame);
+			return State;
+		}
+
+		FGeometryParticleStateBase& NewState(int32 Frame)
+		{
+			RecordedFrame = Frame;
+			bSet = true;
+			State = FGeometryParticleStateBase();
+			return State;
+		}
+
+	private:
 		FGeometryParticleStateBase State;
-		int32 Frame;
+		int32 RecordedFrame;
+		bool bSet;
 	};
 
 	struct FParticleRewindInfo
 	{
-		TArray<FFrameInfo> Frames;
 		int32 AllDirtyParticlesIdx;
 
 		FGeometryParticleStateBase& AddFrame(int32 FrameIdx)
 		{
-			if(Frames.Num() && Frames.Last().Frame == FrameIdx)
+			FFrameInfo& Info = GetFrames()[FrameIdx];
+			if(FGeometryParticleStateBase* State = Info.GetState(FrameIdx))
 			{
-				return Frames.Last().State;
+				return *State;
 			}
 
-			Frames.Add(FFrameInfo{FGeometryParticleStateBase(),FrameIdx});
-			return Frames.Last().State;
+			return Info.NewState(FrameIdx);
+		}
+
+		const TCircularBuffer<FFrameInfo>& GetFrames() const
+		{
+			return Frames.GetValue();
+
 		}
+
+		TCircularBuffer<FFrameInfo>& GetFrames()
+		{
+			return Frames.GetValue();
+
+		}
+
+		bool IsInitialized() const
+		{
+			return Frames.IsSet();
+		}
+
+		void Initialize(int32 NumFrames)
+		{
+			ensure(IsInitialized() == false);
+			Frames = TCircularBuffer<FFrameInfo>(NumFrames);
+		}
+
+	private:
+		TOptional<TCircularBuffer<FFrameInfo>> Frames;
+
 	};
 
-	void CoalesceBack(TArray<FFrameInfo>& Frames)
+	void CoalesceBack(TCircularBuffer<FFrameInfo>& Frames, int32 LatestIdx)
 	{
-		const FGeometryParticleStateBase& LatestState = Frames.Last().State;
-		for(int32 FrameIdx = Frames.Num() - 2; FrameIdx >= 0; --FrameIdx)
+		const FGeometryParticleStateBase& LatestState = Frames[LatestIdx].GetStateChecked(LatestIdx);
+		for(int32 FrameIdx = LatestIdx - 1; FrameIdx >= EarliestFrame; --FrameIdx)
 		{
-			FFrameInfo& Frame = Frames[FrameIdx];
-			if(Frame.State.CoalesceState(LatestState) == false)
+			if(FGeometryParticleStateBase* State = Frames[FrameIdx].GetState(FrameIdx))
 			{
-				//nothing to coalesce so no need to check earlier frames
-				break;
+				if(State->CoalesceState(LatestState) == false)
+				{
+					//nothing to coalesce so no need to check earlier frames
+					break;
+				}
 			}
 		}
 	}
@@ -665,14 +732,13 @@ private:
 
 	FParticleRewindInfo& FindOrAddParticle(TGeometryParticle<FReal,3>* UnsafeGTParticle, FUniqueIdx UniqueIdx)
 	{
-		if(FParticleRewindInfo* Info = ParticleToRewindInfo.Find(UniqueIdx))
+		FParticleRewindInfo& Info = ParticleToRewindInfo.FindOrAdd(UniqueIdx);
+		if(!Info.IsInitialized())
 		{
-			return *Info;
+			Info.Initialize(Managers.Capacity());
+			Info.AllDirtyParticlesIdx = AllDirtyParticles.Add(FDirtyParticleInfo{UnsafeGTParticle,UniqueIdx,CurFrame});
 		}
-
-
-		FParticleRewindInfo& Info = ParticleToRewindInfo.FindOrAdd(UniqueIdx);
-		Info.AllDirtyParticlesIdx = AllDirtyParticles.Add(FDirtyParticleInfo{UnsafeGTParticle,UniqueIdx, CurFrame});
+		
 		return Info;
 	}
 
@@ -680,24 +746,20 @@ private:
 	{
 		if(const FParticleRewindInfo* Info = ParticleToRewindInfo.Find(Particle.UniqueIdx()))
 		{
-			//is it worth doing binary search?
-			const int32 NumFrames = Info->Frames.Num();
-			if(NumFrames > 0)
+			const TCircularBuffer<FFrameInfo>& Frames = Info->GetFrames();
+			if(const FGeometryParticleStateBase* FrameState = Frames[Frame].GetState(Frame))
 			{
-				//If frame is before first capture, use first capture
-				if(Frame <= Info->Frames[0].Frame)
-				{
-					return &Info->Frames[0].State;
-				}
+				return FrameState;
+			}
 
-				//If frame is between two captures, use later capture, because we always store the last data before a change
-				//We can never use an earlier capture, because the fact that we captured at all implies _something_ is different from proceeding frames
-				for(int32 FrameIdx = 1; FrameIdx < NumFrames; ++FrameIdx)
+			//If frame is between two captures, use later capture. We always store the last data before a change
+			//We can never use an earlier capture because the fact that we captured at all implies _something_ is different from proceeding frames
+
+			for(int32 FrameIdx = Frame + 1; FrameIdx <= CurFrame; ++FrameIdx)
+			{
+				if(const FGeometryParticleStateBase* FrameState = Frames[FrameIdx].GetState(FrameIdx))
 				{
-					if(Frame <= Info->Frames[FrameIdx].Frame)
-					{
-						return &Info->Frames[FrameIdx].State;
-					}
+					return FrameState;
 				}
 			}
 		}
@@ -710,14 +772,14 @@ private:
 	{
 		TGeometryParticle<FReal,3>* Particle;
 		FUniqueIdx CachedUniqueIdx;	//Needed when manipulating on physics thread and Particle data cannot be read
-		uint32 LastDirtyFrame;	//Track how recently this was made dirty
+		int32 LastDirtyFrame;	//Track how recently this was made dirty
 	};
 
 	TArrayAsMap<FUniqueIdx,FParticleRewindInfo> ParticleToRewindInfo;
 	TCircularBuffer<FFrameManagerInfo> Managers;
 	TArray<FDirtyParticleInfo> AllDirtyParticles;
-	uint32 CurFrame;
-	uint32 EarliestFrame;
+	int32 CurFrame;
+	int32 EarliestFrame;
 	bool bNeedsSave;	//Indicates that some data is pointing at head and requires saving before a rewind
 };
 }
