Chaos - remove unused "always update manifold" flags

diff --git a/Engine/Source/Runtime/Experimental/Chaos/Private/Chaos/PBDCollisionConstraints.cpp b/Engine/Source/Runtime/Experimental/Chaos/Private/Chaos/PBDCollisionConstraints.cpp
--- a/Engine/Source/Runtime/Experimental/Chaos/Private/Chaos/PBDCollisionConstraints.cpp
+++ b/Engine/Source/Runtime/Experimental/Chaos/Private/Chaos/PBDCollisionConstraints.cpp
@@ -230,10 +230,6 @@ namespace Chaos
 			}
 		}
 
-		// Swept constraint may drastically change particle position during Apply, other constraint's
-		// manifolds with particle may be invalidated, so we should update manifolds with this particle even on first apply iteration.
-		AddParticleAlwaysUpdateManifold(InConstraint.Particle[0]);
-
 		UpdateConstraintMaterialProperties(Constraints.SinglePointSweptConstraints[Idx]);
 	}
 
@@ -292,7 +288,6 @@ namespace Chaos
 			HandleAllocator.FreeHandle(Handle);
 		}
 		Constraints.Reset();
-		ParticlesAlwaysUpdatingManifold.Reset();
 		Handles.Reset();
 #endif
 
@@ -477,29 +472,16 @@ namespace Chaos
 		}
 	}
 
-	template<typename T, int d>
-	void TPBDCollisionConstraints<T, d>::AddParticleAlwaysUpdateManifold(const TGeometryParticleHandle<T, d>* Particle)
-	{
-		ParticlesAlwaysUpdatingManifold.Add(Particle);
-	}
-
-	template<typename T, int d>
-	bool TPBDCollisionConstraints<T, d>::ShouldParticleAlwaysUpdateManifold(const TGeometryParticleHandle<T, d>* Particle) const
-	{
-		return ParticlesAlwaysUpdatingManifold.Contains(Particle);
-	}
-
 	template<typename T, int d>
 	bool TPBDCollisionConstraints<T, d>::Apply(const T Dt, const int32 Iterations, const int32 NumIterations)
 	{
 		SCOPE_CYCLE_COUNTER(STAT_Collisions_Apply);
 
 		bool bNeedsAnotherIteration = false;
-		bool bAlwaysUpdateManifold = false;
 		if (MApplyPairIterations > 0)
 		{
 			const Collisions::TContactParticleParameters<T> ParticleParameters = { MCullDistance, MShapePadding, &MCollided };
-			const Collisions::TContactIterationParameters<T> IterationParameters = { Dt, Iterations, NumIterations, MApplyPairIterations, ApplyType, &bNeedsAnotherIteration, &bAlwaysUpdateManifold };
+			const Collisions::TContactIterationParameters<T> IterationParameters = { Dt, Iterations, NumIterations, MApplyPairIterations, ApplyType, &bNeedsAnotherIteration };
 
 			NumActivePointConstraints = 0;
 			for (FPointContactConstraint& Contact : Constraints.SinglePointConstraints)
@@ -602,10 +584,9 @@ namespace Chaos
 				check(ConstraintHandle != nullptr);
 
 				TVector<const TGeometryParticleHandle<T, d>*, 2> ConstrainedParticles = ConstraintHandle->GetConstrainedParticles();
-				bool bAlwaysUpdateManifold = ShouldParticleAlwaysUpdateManifold(ConstrainedParticles[0]) || ShouldParticleAlwaysUpdateManifold(ConstrainedParticles[1]);
 				bool bNeedsAnotherIteration = false;
 				Collisions::TContactParticleParameters<T> ParticleParameters = { MCullDistance, MShapePadding, &MCollided };
-				Collisions::TContactIterationParameters<T> IterationParameters = { Dt, Iterations, NumIterations, MApplyPairIterations, ApplyType, &bNeedsAnotherIteration, &bAlwaysUpdateManifold };
+				Collisions::TContactIterationParameters<T> IterationParameters = { Dt, Iterations, NumIterations, MApplyPairIterations, ApplyType, &bNeedsAnotherIteration };
 				Collisions::Apply(ConstraintHandle->GetContact(), IterationParameters, ParticleParameters);
 
 				if (bNeedsAnotherIteration)
diff --git a/Engine/Source/Runtime/Experimental/Chaos/Private/Chaos/PBDCollisionConstraintsContact.cpp b/Engine/Source/Runtime/Experimental/Chaos/Private/Chaos/PBDCollisionConstraintsContact.cpp
--- a/Engine/Source/Runtime/Experimental/Chaos/Private/Chaos/PBDCollisionConstraintsContact.cpp
+++ b/Engine/Source/Runtime/Experimental/Chaos/Private/Chaos/PBDCollisionConstraintsContact.cpp
@@ -449,8 +449,8 @@ namespace Chaos
 			const FReal PartialDT = Constraint.TimeOfImpact * IterationParameters.Dt;
 			const FReal RemainingDT = (1 - Constraint.TimeOfImpact) * IterationParameters.Dt;
 			const int32 FakeIteration = FMath::Max(IterationParameters.Iteration, 1); // Force Apply to update constraint, as other constraints could've changed P
-			const FContactIterationParameters IterationParametersPartialDT{ PartialDT, FakeIteration, IterationParameters.NumIterations, IterationParameters.NumPairIterations, IterationParameters.ApplyType, IterationParameters.NeedsAnotherIteration, IterationParameters.AlwaysUpdateManifold };
-			const FContactIterationParameters IterationParametersRemainingDT{ RemainingDT, FakeIteration, IterationParameters.NumIterations, IterationParameters.NumPairIterations, IterationParameters.ApplyType, IterationParameters.NeedsAnotherIteration, IterationParameters.AlwaysUpdateManifold };
+			const FContactIterationParameters IterationParametersPartialDT{ PartialDT, FakeIteration, IterationParameters.NumIterations, IterationParameters.NumPairIterations, IterationParameters.ApplyType, IterationParameters.NeedsAnotherIteration };
+			const FContactIterationParameters IterationParametersRemainingDT{ RemainingDT, FakeIteration, IterationParameters.NumIterations, IterationParameters.NumPairIterations, IterationParameters.ApplyType, IterationParameters.NeedsAnotherIteration };
 
 			// Rewind P to TOI and Apply
 			Particle0->P() = FMath::Lerp(Particle0->X(), Particle0->P(), Constraint.TimeOfImpact);
diff --git a/Engine/Source/Runtime/Experimental/Chaos/Public/Chaos/PBDCollisionConstraints.h b/Engine/Source/Runtime/Experimental/Chaos/Public/Chaos/PBDCollisionConstraints.h
--- a/Engine/Source/Runtime/Experimental/Chaos/Public/Chaos/PBDCollisionConstraints.h
+++ b/Engine/Source/Runtime/Experimental/Chaos/Public/Chaos/PBDCollisionConstraints.h
@@ -165,13 +165,6 @@ public:
 	*/
 	void UpdateManifolds(T Dt);
 
-	/*
-	* Add particle to list of particles that will always update their constraint manifold.
-	*/
-	void AddParticleAlwaysUpdateManifold(const TGeometryParticleHandle<T, d>* Particle);
-
-	bool ShouldParticleAlwaysUpdateManifold(const TGeometryParticleHandle<T, d>* Particle) const;
-
 
 	//
 	// General Rule API
@@ -311,8 +304,6 @@ private:
 	int32 NumActiveSweptPointConstraints;
 	int32 NumActiveIterativeConstraints;
 
-	TSet<const TGeometryParticleHandle<T, d>*> ParticlesAlwaysUpdatingManifold;
-
 #if CHAOS_COLLISION_PERSISTENCE_ENABLED
 	TMap< FConstraintContainerHandleKey, FConstraintContainerHandle* > Manifolds;
 #endif
diff --git a/Engine/Source/Runtime/Experimental/Chaos/Public/Chaos/PBDCollisionConstraintsContact.h b/Engine/Source/Runtime/Experimental/Chaos/Public/Chaos/PBDCollisionConstraintsContact.h
--- a/Engine/Source/Runtime/Experimental/Chaos/Public/Chaos/PBDCollisionConstraintsContact.h
+++ b/Engine/Source/Runtime/Experimental/Chaos/Public/Chaos/PBDCollisionConstraintsContact.h
@@ -28,7 +28,6 @@ namespace Chaos
 			const int32 NumPairIterations;
 			const ECollisionApplyType ApplyType;	// @todo(chaos): a better way to customize the collision solver
 			bool* NeedsAnotherIteration;
-			bool* AlwaysUpdateManifold;
 		};
 		using FContactIterationParameters = TContactIterationParameters<FReal>;
 
