Account for sleeping when changing state to/from kinematic or static and dynamic in TPBDRigidParticles::SetObjectState.

#localization none
#tests PIE

diff --git a/Engine/Source/Runtime/Experimental/Chaos/Public/Chaos/PBDRigidParticles.h b/Engine/Source/Runtime/Experimental/Chaos/Public/Chaos/PBDRigidParticles.h
--- a/Engine/Source/Runtime/Experimental/Chaos/Public/Chaos/PBDRigidParticles.h
+++ b/Engine/Source/Runtime/Experimental/Chaos/Public/Chaos/PBDRigidParticles.h
@@ -115,13 +115,13 @@ class TPBDRigidParticles : public TRigidParticles<T, d>
 			this->InvI(Index) = PMatrix<float, 3, 3>(0);
 		}
 
-		if (CurrentState == EObjectStateType::Dynamic && (InObjectState == EObjectStateType::Kinematic || InObjectState == EObjectStateType::Static))
+		if ((CurrentState == EObjectStateType::Dynamic || CurrentState == EObjectStateType::Sleeping) && (InObjectState == EObjectStateType::Kinematic || InObjectState == EObjectStateType::Static))
 		{
 			// Transitioning from dynamic to static or kinematic, set inverse mass and inertia tensor to zero.
 			this->InvM(Index) = 0.0f;
 			this->InvI(Index) = PMatrix<float, 3, 3>(0);
 		}
-		else if ((CurrentState == EObjectStateType::Kinematic || CurrentState == EObjectStateType::Static || CurrentState == EObjectStateType::Uninitialized) && InObjectState == EObjectStateType::Dynamic)
+		else if ((CurrentState == EObjectStateType::Kinematic || CurrentState == EObjectStateType::Static || CurrentState == EObjectStateType::Uninitialized) && (InObjectState == EObjectStateType::Dynamic || InObjectState == EObjectStateType::Sleeping))
 		{
 			// Transitioning from kinematic or static to dynamic, compute the inverses.
 			checkSlow(this->M(Index) != 0.0);
