Adding support for re running and still recording data. Desync support still missing

diff --git a/Engine/Source/Programs/HeadlessChaos/Private/HeadlessChaosTestRewind.cpp b/Engine/Source/Programs/HeadlessChaos/Private/HeadlessChaosTestRewind.cpp
--- a/Engine/Source/Programs/HeadlessChaos/Private/HeadlessChaosTestRewind.cpp
+++ b/Engine/Source/Programs/HeadlessChaos/Private/HeadlessChaosTestRewind.cpp
@@ -586,6 +586,80 @@ namespace ChaosTest {
 			EXPECT_EQ(RewindData->GetNumDirtyParticles(),1);
 		}
 
+		// Throw out the proxy
+		Solver->UnregisterObject(Particle.Get());
+
+		Module->DestroySolver(Solver);
+	}
+
+	GTEST_TEST(RewindTest,Resim)
+	{
+		auto Sphere = TSharedPtr<FImplicitObject,ESPMode::ThreadSafe>(new TSphere<float,3>(TVector<float,3>(0),10));
+
+		FChaosSolversModule* Module = FChaosSolversModule::GetModule();
+		Module->ChangeThreadingMode(EChaosThreadingMode::SingleThread);
+
+		// Make a solver
+		FPhysicsSolver* Solver = Module->CreateSolver(nullptr,ESolverFlags::Standalone);
+		Solver->SetEnabled(true);
+
+		Solver->EnableRewindCapture(5);
+
+		// Make particles
+		auto Particle = TPBDRigidParticle<float,3>::CreateParticle();
+
+		Particle->SetGeometry(Sphere);
+		Solver->RegisterObject(Particle.Get());
+		Particle->SetGravityEnabled(true);
+
+		auto Kinematic = TKinematicGeometryParticle<float,3>::CreateParticle();
+
+		Kinematic->SetGeometry(Sphere);
+		Solver->RegisterObject(Kinematic.Get());
+		Kinematic->SetX(FVec3(2,2,2));
+
+		TArray<FVec3> X;
+
+		for(int Step = 0; Step < 10; ++Step)
+		{
+			X.Add(Particle->X());
+
+			if(Step == 9)
+			{
+				Kinematic->SetX(FVec3(50,50,50));
+			}
+
+			TickSolverHelper(Module,Solver);
+		}
+
+		const int RewindStep = 7;
+
+		FRewindData* RewindData = Solver->GetRewindData();
+		EXPECT_TRUE(RewindData->RewindToFrame(RewindStep));
+
+		//Move particle and rerun
+		Particle->SetX(FVec3(0,0,100));
+		for(int Step = RewindStep; Step < 10; ++Step)
+		{
+			X[Step] = Particle->X();
+			TickSolverHelper(Module,Solver);
+		}
+
+		EXPECT_EQ(Kinematic->X()[2],2);	//Rewound kinematic and never updated it so back to original value
+
+		//Make sure we recorded the new data
+		for(int Step = RewindStep; Step < 10; ++Step)
+		{
+			const FGeometryParticleState State = RewindData->GetStateAtFrame(*Particle,Step);
+			EXPECT_EQ(State.X()[2],X[Step][2]);
+
+			const FGeometryParticleState KinState = RewindData->GetStateAtFrame(*Kinematic,Step);
+			//TODO: this is a form of desync, need support for that still
+			//EXPECT_EQ(KinState.X()[2],2);	//even though not dirty, data is properly stored as not moving
+		}
+
+
+
 		// Throw out the proxy
 		Solver->UnregisterObject(Particle.Get());
 
diff --git a/Engine/Source/Runtime/Experimental/ChaosSolvers/Private/PBDRigidsSolver.cpp b/Engine/Source/Runtime/Experimental/ChaosSolvers/Private/PBDRigidsSolver.cpp
--- a/Engine/Source/Runtime/Experimental/ChaosSolvers/Private/PBDRigidsSolver.cpp
+++ b/Engine/Source/Runtime/Experimental/ChaosSolvers/Private/PBDRigidsSolver.cpp
@@ -1039,8 +1039,9 @@ namespace Chaos
 		{
 			QUICK_SCOPE_CYCLE_COUNTER(RecordRewindData);
 
-			int32 DataIdx = MRewindData->PrepareFrameForPTDirty(ActiveParticles.Num());
+			MRewindData->PrepareFrameForPTDirty(ActiveParticles.Num());
 			
+			int32 DataIdx = 0;
 			for(TPBDRigidParticleHandleImp<float,3,false>& ActiveObject : ActiveParticles)
 			{
 				const IPhysicsProxyBase* Proxy = GetProxy(ActiveObject.Handle());
diff --git a/Engine/Source/Runtime/Experimental/ChaosSolvers/Public/RewindData.h b/Engine/Source/Runtime/Experimental/ChaosSolvers/Public/RewindData.h
--- a/Engine/Source/Runtime/Experimental/ChaosSolvers/Public/RewindData.h
+++ b/Engine/Source/Runtime/Experimental/ChaosSolvers/Public/RewindData.h
@@ -387,7 +387,8 @@ public:
 	FRewindData(int32 NumFrames)
 	: Managers(NumFrames+1)	//give 1 extra for saving at head
 	, CurFrame(0)
-	, EarliestFrame(0)
+	, FramesSaved(0)
+	, DataIdxOffset(0)
 	, bNeedsSave(false)
 	{
 	}
@@ -398,14 +399,15 @@ public:
 	{
 		ensure(IsInGameThread());
 
-		//Can't go too far back, also we need 1 entry for saving head
+		//Can't go too far back
+		const int32 EarliestFrame = CurFrame - FramesSaved;
 		if(Frame < EarliestFrame)
 		{
 			return false;
 		}
 
 		//If we need to save and we are right on the edge of the buffer, we can't go back to earliest frame
-		if(Frame == EarliestFrame && bNeedsSave && &Managers[CurFrame] == &Managers[EarliestFrame])
+		if(Frame == EarliestFrame && bNeedsSave && FramesSaved == Managers.Capacity())
 		{
 			return false;
 		}
@@ -448,7 +450,7 @@ public:
 
 		CurFrame = Frame;
 		bNeedsSave = false;
-		EarliestFrame = CurFrame;	//can't rewind before this point. This simplifies saving the state at head
+		FramesSaved = 0; //can't rewind before this point. This simplifies saving the state at head
 
 		return true;
 	}
@@ -483,11 +485,9 @@ public:
 	void AdvanceFrame()
 	{
 		++CurFrame;
-		if(CurFrame > static_cast<int32>(Managers.Capacity()))
-		{
-			++EarliestFrame;
-		}
-
+		FramesSaved = FMath::Min(FramesSaved+1,static_cast<int32>(Managers.Capacity()));
+		
+		const int32 EarliestFrame = CurFrame - 1 - FramesSaved;
 		//remove any old dirty particles
 		for(int32 DirtyIdx = AllDirtyParticles.Num() - 1; DirtyIdx >= 0; --DirtyIdx)
 		{
@@ -503,10 +503,18 @@ public:
 
 	void PrepareFrame(int32 NumDirtyParticles)
 	{
-		PrepareManagerForFrame(CurFrame,NumDirtyParticles);
+		FFrameManagerInfo& Info = Managers[CurFrame];
+		if(Info.Manager == nullptr)
+		{
+			Info.Manager = MakeUnique<FDirtyPropertiesManager>();
+		}
+
+		DataIdxOffset = Info.Manager->GetNumParticles();
+		Info.Manager->SetNumParticles(DataIdxOffset + NumDirtyParticles);
+		Info.FrameCreatedFor = CurFrame;
 	}
 
-	int32 PrepareFrameForPTDirty(int32 NumActiveParticles)
+	void PrepareFrameForPTDirty(int32 NumActiveParticles)
 	{
 		bNeedsSave = true;
 
@@ -515,16 +523,13 @@ public:
 		FFrameManagerInfo& Info = Managers[PrevFrame];
 		ensure(Info.Manager && Info.FrameCreatedFor == (PrevFrame));
 
-		{
-			const int32 NumDirtyAlready = Info.Manager->GetNumParticles();
-			Info.Manager->SetNumParticles(NumDirtyAlready + NumActiveParticles);
-			return NumDirtyAlready;
-		}
-		return 0;
+		DataIdxOffset = Info.Manager->GetNumParticles();
+		Info.Manager->SetNumParticles(DataIdxOffset + NumActiveParticles);
 	}
 
-	void PushGTDirtyData(const FDirtyPropertiesManager& SrcManager, int32 DataIdx, const FDirtyProxy& Dirty)
+	void PushGTDirtyData(const FDirtyPropertiesManager& SrcManager, const int32 DataIdxIn, const FDirtyProxy& Dirty)
 	{
+		const int32 DataIdx = DataIdxIn + DataIdxOffset;
 		//This records changes enqueued by GT.
 		//Most new particles do not change, so to avoid useless writes we wait until the next frame's dirty flag
 		//This is possible because most properties are const on the physics thread
@@ -582,8 +587,10 @@ public:
 		}
 	}
 
-	void PushPTDirtyData(const TPBDRigidParticleHandle<FReal,3>& Rigid, int32 DataIdx)
+	void PushPTDirtyData(const TPBDRigidParticleHandle<FReal,3>& Rigid, const int32 DataIdxIn)
 	{
+		const int32 DataIdx = DataIdxIn + DataIdxOffset;
+
 		//todo: is this check needed? why do we pass sleeping rigids into this function?
 		if(SimWritablePropsMayChange(Rigid))
 		{
@@ -695,6 +702,7 @@ private:
 	void CoalesceBack(TCircularBuffer<FFrameInfo>& Frames, int32 LatestIdx)
 	{
 		const FGeometryParticleStateBase& LatestState = Frames[LatestIdx].GetStateChecked(LatestIdx);
+		const int32 EarliestFrame = LatestIdx - FramesSaved;
 		for(int32 FrameIdx = LatestIdx - 1; FrameIdx >= EarliestFrame; --FrameIdx)
 		{
 			if(FGeometryParticleStateBase* State = Frames[FrameIdx].GetState(FrameIdx))
@@ -707,19 +715,7 @@ private:
 			}
 		}
 	}
-
-	void PrepareManagerForFrame(int32 Frame, int32 NumDirtyParticles)
-	{
-		FFrameManagerInfo& Info = Managers[Frame];
-		if(Info.Manager == nullptr)
-		{
-			Info.Manager = MakeUnique<FDirtyPropertiesManager>();
-		}
-
-		Info.Manager->SetNumParticles(NumDirtyParticles);
-		Info.FrameCreatedFor = Frame;
-	}
-
+	
 	struct FFrameManagerInfo
 	{
 		TUniquePtr<FDirtyPropertiesManager> Manager;
@@ -779,7 +775,8 @@ private:
 	TCircularBuffer<FFrameManagerInfo> Managers;
 	TArray<FDirtyParticleInfo> AllDirtyParticles;
 	int32 CurFrame;
-	int32 EarliestFrame;
+	int32 FramesSaved;
+	int32 DataIdxOffset;
 	bool bNeedsSave;	//Indicates that some data is pointing at head and requires saving before a rewind
 };
 }
