Pull the implicit from the constraint's manifold, which has been filtered, not the particle's geometry, which may be a union containing things we don't want to collide against.

diff --git a/Engine/Source/Runtime/Experimental/Chaos/Private/Chaos/CollisionResolution.cpp b/Engine/Source/Runtime/Experimental/Chaos/Private/Chaos/CollisionResolution.cpp
--- a/Engine/Source/Runtime/Experimental/Chaos/Private/Chaos/CollisionResolution.cpp
+++ b/Engine/Source/Runtime/Experimental/Chaos/Private/Chaos/CollisionResolution.cpp
@@ -1825,7 +1825,8 @@ namespace Chaos
 
 			if (SampleParticles)
 			{
-				SampleObject<UpdateType>(*Particle1->Geometry(), LevelsetTM, *SampleParticles, ParticlesTM, CullDistance, Constraint);
+				const FImplicitObject* Obj1 = Constraint.Manifold.Implicit[1];
+				SampleObject<UpdateType>(*Obj1, LevelsetTM, *SampleParticles, ParticlesTM, CullDistance, Constraint);
 			}
 		}
 
