Fixed transforms for activated and disabled clustered particles in geometry collections. Mitigated crash in SQ when querying against cluster parents.

diff --git a/Engine/Source/Runtime/Experimental/ChaosSolvers/Private/PhysicsProxy/GeometryCollectionPhysicsProxy.cpp b/Engine/Source/Runtime/Experimental/ChaosSolvers/Private/PhysicsProxy/GeometryCollectionPhysicsProxy.cpp
--- a/Engine/Source/Runtime/Experimental/ChaosSolvers/Private/PhysicsProxy/GeometryCollectionPhysicsProxy.cpp
+++ b/Engine/Source/Runtime/Experimental/ChaosSolvers/Private/PhysicsProxy/GeometryCollectionPhysicsProxy.cpp
@@ -1506,8 +1506,7 @@ void FGeometryCollectionPhysicsProxy::BufferPhysicsResults()
 				FTransform& ParticleToWorld = TargetResults.ParticleToWorldTransforms[TransformIndex];
 				ParticleToWorld = FTransform(Handle->R(), Handle->X());
 
-				TargetResults.Transforms[TransformIndex] = 
-					PTDynamicCollection.MassToLocal[TransformIndex].GetRelativeTransformReverse(ParticleToWorld).GetRelativeTransform(ActorToWorld);
+				TargetResults.Transforms[TransformIndex] = PTDynamicCollection.MassToLocal[TransformIndex].GetRelativeTransformReverse(ParticleToWorld).GetRelativeTransform(ActorToWorld);
 				TargetResults.Transforms[TransformIndex].NormalizeRotation();
 
 				// If the parent of this NON DISABLED body is set to anything other than INDEX_NONE,
@@ -1519,6 +1518,11 @@ void FGeometryCollectionPhysicsProxy::BufferPhysicsResults()
 					// Children in the hierarchy are stored in a TSet, which is not thread safe.  
 					// So we retain indices to remove afterwards.
 					EndFrameUnparentingBuffer[TransformIndex] = ParentIndex;
+					
+					// Once a child is released there is no longer a useful parent relationship as the child is
+					// now entirely responsible for its own simulation. Sever the link here so we don't attempt
+					// to pull in a parent transform later in the pipeline.
+					PhysicsThreadCollection.Parent[TransformIndex] = INDEX_NONE;
 				}
 
 				// When a leaf node rigid body is removed from a cluster, the rigid
@@ -1546,33 +1550,6 @@ void FGeometryCollectionPhysicsProxy::BufferPhysicsResults()
 					}
 					SolverClusterID[TransformIndex] = Handle->ClusterIds().Id;
 				}
-
-				// Disabled rigid bodies that have valid cluster parents, and have been re-indexed by the
-				// solver (as in, They were re-clustered outside of the geometry collection), These clusters 
-				// will need to be rendered based on the clusters position. 
-
-				if (TPBDRigidParticleHandle<float, 3>* ParentHandle = Handle->ClusterIds().Id)
-				{
-					if (TPBDRigidClusteredParticleHandle<float, 3>* ClusteredParentHandle = ParentHandle->CastToClustered())
-					{
-						if (ClusteredParentHandle->InternalCluster())
-						{
-
-							//const FTransform ClusterChildToWorld = ClusterChildToParentMap[RigidBodyIndex] * FTransform(Particles.R(ClusterParentIndex), Particles.X(ClusterParentIndex));
-							const TRigidTransform<float, 3>& ChildToParent = Handle->ChildToParent();
-							const FTransform ClusterChildToWorld = ChildToParent * FTransform(ParentHandle->R(), ParentHandle->X());
-							if (Parameters.IsCacheRecording())
-							{
-								Handle->SetX(ClusterChildToWorld.GetTranslation());
-								Handle->SetR(ClusterChildToWorld.GetRotation());
-							}
-							TargetResults.Transforms[TransformIndex] = 
-								PTDynamicCollection.MassToLocal[TransformIndex].GetRelativeTransformReverse(ClusterChildToWorld).GetRelativeTransform(ActorToWorld);
-							TargetResults.Transforms[TransformIndex].NormalizeRotation();
-						}
-					}
-				}
-
 			}// end if
 		} // end for
 	} // tmp scope
@@ -1643,8 +1620,11 @@ void FGeometryCollectionPhysicsProxy::PullFromPhysicsState()
 	 */
 
 	const FGeometryCollectionResults* TargetResultPtr = PhysToGameInterchange.GetConsumerBuffer();
-	if (!TargetResultPtr)
+	if(!TargetResultPtr)
+	{
 		return;
+	}
+
 	// Remove const-ness as we're going to do the ExchangeArrays() thing.
 	FGeometryCollectionResults& TR = *const_cast<FGeometryCollectionResults*>(TargetResultPtr);
 
diff --git a/Engine/Source/Runtime/PhysicsSQ/Private/SQAccelerator.cpp b/Engine/Source/Runtime/PhysicsSQ/Private/SQAccelerator.cpp
--- a/Engine/Source/Runtime/PhysicsSQ/Private/SQAccelerator.cpp
+++ b/Engine/Source/Runtime/PhysicsSQ/Private/SQAccelerator.cpp
@@ -188,6 +188,14 @@ private:
 		//todo: add a check to ensure hitbuffer matches SQ type
 		using namespace Chaos;
 		TGeometryParticle<float, 3>* GeometryParticle = Payload.GetExternalGeometryParticle_ExternalThread();
+
+		if(!GeometryParticle)
+		{
+			// This case handles particles created by the physics simulation without the main thread
+			// being made aware of their creation. We have a PT particle but no external particle
+			return true;
+		}
+
 		const TShapesArray<float,3>& Shapes = GeometryParticle->ShapesArray();
 
 		const bool bTestShapeBounds =  Shapes.Num() > 1;
