Fixed crash when attempting to generate a levelset from an empty triangle mesh

diff --git a/Engine/Source/Runtime/Experimental/GeometryCollectionSimulationCore/Private/GeometryCollection/GeometryCollectionCollisionStructureManager.cpp b/Engine/Source/Runtime/Experimental/GeometryCollectionSimulationCore/Private/GeometryCollection/GeometryCollectionCollisionStructureManager.cpp
--- a/Engine/Source/Runtime/Experimental/GeometryCollectionSimulationCore/Private/GeometryCollection/GeometryCollectionCollisionStructureManager.cpp
+++ b/Engine/Source/Runtime/Experimental/GeometryCollectionSimulationCore/Private/GeometryCollection/GeometryCollectionCollisionStructureManager.cpp
@@ -289,6 +289,12 @@ Chaos::TLevelSet<float, 3>* FCollisionStructureManager::NewLevelset(
 	const int32 MaxRes,
 	const ECollisionTypeEnum CollisionType)
 {
+	if(TriMesh.GetNumElements() == 0)
+	{
+		// Empty tri-mesh, can't create a new level set
+		return nullptr;
+	}
+
 	Chaos::TVector<int32, 3> Counts;
 	const FVector Extents = CollisionBounds.GetExtent();
 	if (Extents.X < Extents.Y && Extents.X < Extents.Z)
