Chaos

diff --git a/Engine/Source/Runtime/Experimental/Chaos/Private/Chaos/CollisionResolution.cpp b/Engine/Source/Runtime/Experimental/Chaos/Private/Chaos/CollisionResolution.cpp
--- a/Engine/Source/Runtime/Experimental/Chaos/Private/Chaos/CollisionResolution.cpp
+++ b/Engine/Source/Runtime/Experimental/Chaos/Private/Chaos/CollisionResolution.cpp
@@ -8,6 +8,7 @@
 #include "Chaos/CollisionResolutionTypes.h"
 #include "Chaos/CollisionResolutionUtil.h"
 #include "Chaos/Collision/CollisionContext.h"
+#include "Chaos/Collision/PBDCollisionConstraint.h"
 #include "Chaos/Convex.h"
 #include "Chaos/Defines.h"
 #include "Chaos/HeightField.h"
@@ -65,21 +66,20 @@ namespace Chaos
 {
 	namespace Collisions
 	{
-		template <typename T>
-		struct TContactPoint
+		struct FContactPoint
 		{
-			TVec3<T> Normal;
-			TVec3<T> Location;
-			T Phi;
+			FVec3 Normal;
+			FVec3 Location;
+			FReal Phi;
 
-			TContactPoint() 
+			FContactPoint() 
 			: Normal(1, 0, 0)
-			, Phi(TNumericLimits<T>::Max()) {}
+			, Phi(TNumericLimits<FReal>::Max()) {}
 		};
 
 		// Determines if body should use CCD. If using CCD, computes Dir and Length of sweep.
-		template <typename T, int d>
-		bool UseCCD(const TGeometryParticleHandle<T, d>* SweptParticle, const TGeometryParticleHandle<T,d>* OtherParticle, const FImplicitObject* Implicit, TVector<T,d>& Dir, FReal& Length)
+
+		bool UseCCD(const TGeometryParticleHandle<FReal, 3>* SweptParticle, const TGeometryParticleHandle<FReal, 3>* OtherParticle, const FImplicitObject* Implicit, FVec3& Dir, FReal& Length)
 		{
 			if (OtherParticle->ObjectState() != EObjectStateType::Static)
 			{
@@ -97,8 +97,8 @@ namespace Chaos
 					return false;
 				}
 
-				T MinBoundsAxis = Implicit->BoundingBox().Extents().Min();
-				T LengthCCDThreshold = MinBoundsAxis * CCDEnableThresholdBoundsScale;
+				FReal MinBoundsAxis = Implicit->BoundingBox().Extents().Min();
+				FReal LengthCCDThreshold = MinBoundsAxis * CCDEnableThresholdBoundsScale;
 
 				if (Length > 0.0f && CCDEnableThresholdBoundsScale >= 0.0f && Length >= LengthCCDThreshold)
 				{
@@ -133,15 +133,14 @@ namespace Chaos
 		}
 
 		// Initializes TimeOfImpact and adjusts Phi for TOI.
-		template <typename T, int d>
-		void SetSweptConstraintTOI(TGeometryParticleHandle<T, d>* Particle, const FReal TOI, const FReal Length, const FVec3& Dir, FRigidBodySweptPointContactConstraint& SweptConstraint)
+		void SetSweptConstraintTOI(TGeometryParticleHandle<FReal, 3>* Particle, const FReal TOI, const FReal Length, const FVec3& Dir, FRigidBodySweptPointContactConstraint& SweptConstraint)
 		{
 			if (SweptConstraint.Manifold.Phi > 0.0f)
 			{
 				return;
 			}
 
-			TPBDRigidParticleHandle<T, d>* RigidParticle = Particle->CastToRigidParticle();
+			TPBDRigidParticleHandle<FReal, 3>* RigidParticle = Particle->CastToRigidParticle();
 			if (!CHAOS_ENSURE(RigidParticle))
 			{
 				return;
@@ -176,8 +175,7 @@ namespace Chaos
 		}
 
 
-		template <typename T>
-		void UpdateContactPoint(TCollisionContact<T, 3>& Manifold, const TContactPoint<T>& NewContactPoint)
+		void UpdateContactPoint(FCollisionContact& Manifold, const FContactPoint& NewContactPoint)
 		{
 			//for now just override
 			if (NewContactPoint.Phi < Manifold.Phi)
@@ -188,18 +186,18 @@ namespace Chaos
 			}
 		}
 
-		template <typename T, int d, typename GeometryA, typename GeometryB>
-		TContactPoint<T> GJKContactPoint2(const GeometryA& A, const GeometryB& B, const TRigidTransform<T, d>& ATM, const TRigidTransform<T, d>& BToATM, const TVector<T, 3>& InitialDir)
+		template <typename GeometryA, typename GeometryB>
+		FContactPoint GJKContactPoint2(const GeometryA& A, const GeometryB& B, const FRigidTransform3& ATM, const FRigidTransform3& BToATM, const FVec3& InitialDir)
 		{
 			SCOPE_CYCLE_COUNTER_GJK();
 
-			TContactPoint<T> Contact;
+			FContactPoint Contact;
 
-			T Penetration;
-			TVec3<T> ClosestA, ClosestB, Normal;
+			FReal Penetration;
+			FVec3 ClosestA, ClosestB, Normal;
 			int32 NumIterations = 0;
 
-			if (ensure(GJKPenetration<true>(A, B, BToATM, Penetration, ClosestA, ClosestB, Normal, (T)0, InitialDir, (T)0, &NumIterations)))
+			if (ensure(GJKPenetration<true>(A, B, BToATM, Penetration, ClosestA, ClosestB, Normal, (FReal)0, InitialDir, (FReal)0, &NumIterations)))
 			{
 				Contact.Location = ATM.TransformPosition(ClosestA);
 				Contact.Normal = -ATM.TransformVectorNoScale(Normal);
@@ -213,25 +211,25 @@ namespace Chaos
 			return Contact;
 		}
 
-		template <typename T, int d, typename GeometryA, typename GeometryB>
-		TContactPoint<T> GJKContactPoint(const GeometryA& A, const TRigidTransform<T, d>& ATM, const GeometryB& B, const TRigidTransform<T, d>& BTM, const TVector<T, 3>& InitialDir)
+		template <typename GeometryA, typename GeometryB>
+		FContactPoint GJKContactPoint(const GeometryA& A, const FRigidTransform3& ATM, const GeometryB& B, const FRigidTransform3& BTM, const FVec3& InitialDir)
 		{
-			const TRigidTransform<T, d> BToATM = BTM.GetRelativeTransform(ATM);
+			const FRigidTransform3 BToATM = BTM.GetRelativeTransform(ATM);
 			return GJKContactPoint2(A, B, ATM, BToATM, InitialDir);
 		}
 
-		template <typename T, int d, typename GeometryA, typename GeometryB>
-		TContactPoint<T> GJKContactPointSwept(const GeometryA& A, const TRigidTransform<T, d>& ATM, const GeometryB& B, const TRigidTransform<T, d>& BTM, const TVector<T, d>& Dir, const T Length, const T CullDistance, T& TOI)
+		template <typename GeometryA, typename GeometryB>
+		FContactPoint GJKContactPointSwept(const GeometryA& A, const FRigidTransform3& ATM, const GeometryB& B, const FRigidTransform3& BTM, const FVec3& Dir, const FReal Length, const FReal CullDistance, FReal& TOI)
 		{
-			TContactPoint<T> Contact;
-			const TRigidTransform<T, d> AToBTM = ATM.GetRelativeTransform(BTM);
-			const TVector<T, d> LocalDir = BTM.InverseTransformVectorNoScale(Dir);
+			FContactPoint Contact;
+			const FRigidTransform3 AToBTM = ATM.GetRelativeTransform(BTM);
+			const FVec3 LocalDir = BTM.InverseTransformVectorNoScale(Dir);
 
-			T OutTime;
-			TVec3<T> Location, Normal;
+			FReal OutTime;
+			FVec3 Location, Normal;
 			int32 NumIterations = 0;
 
-			if (GJKRaycast2(B, A, AToBTM, LocalDir, Length, OutTime, Location, Normal, (T)0, true))
+			if (GJKRaycast2(B, A, AToBTM, LocalDir, Length, OutTime, Location, Normal, (FReal)0, true))
 			{
 				Contact.Location = BTM.TransformPosition(Location);
 				Contact.Normal = BTM.TransformVectorNoScale(Normal);
@@ -242,14 +240,14 @@ namespace Chaos
 		}
 
 
-		template <typename GeometryA, typename GeometryB, typename T, int d>
-		TContactPoint<T> GJKImplicitContactPoint(const FImplicitObject& A, const TRigidTransform<T, d>& ATransform, const GeometryB& B, const TRigidTransform<T, d>& BTransform, const T CullDistance)
+		template <typename GeometryA, typename GeometryB>
+		FContactPoint GJKImplicitContactPoint(const FImplicitObject& A, const FRigidTransform3& ATransform, const GeometryB& B, const FRigidTransform3& BTransform, const FReal CullDistance)
 		{
-			TContactPoint<T> Contact;
-			const TRigidTransform<T, d> AToBTM = ATransform.GetRelativeTransform(BTransform);
+			FContactPoint Contact;
+			const FRigidTransform3 AToBTM = ATransform.GetRelativeTransform(BTransform);
 
-			T Penetration = FLT_MAX;
-			TVec3<T> Location, Normal;
+			FReal Penetration = FLT_MAX;
+			FVec3 Location, Normal;
 			if (const TImplicitObjectScaled<GeometryA>* ScaledConvexImplicit = A.template GetObject<const TImplicitObjectScaled<GeometryA> >())
 			{
 				if (B.GJKContactPoint(*ScaledConvexImplicit, AToBTM, CullDistance, Location, Normal, Penetration))
@@ -284,18 +282,18 @@ namespace Chaos
 			return Contact;
 		}
 
-		template <typename GeometryA, typename GeometryB, typename T, int d>
-		TContactPoint<T> GJKImplicitSweptContactPoint(const FImplicitObject& A, const TRigidTransform<T, d>& AStartTransform, const GeometryB& B, const TRigidTransform<T, d>& BTransform, const TVector<T, d>& Dir, const T Length, const T CullDistance, T& TOI)
+		template <typename GeometryA, typename GeometryB>
+		FContactPoint GJKImplicitSweptContactPoint(const FImplicitObject& A, const FRigidTransform3& AStartTransform, const GeometryB& B, const FRigidTransform3& BTransform, const FVec3& Dir, const FReal Length, const FReal CullDistance, FReal& TOI)
 		{
-			TContactPoint<T> Contact;
-			const TRigidTransform<T, d> AToBTM = AStartTransform.GetRelativeTransform(BTransform);
-			const TVector<T, d> LocalDir = BTransform.InverseTransformVectorNoScale(Dir);
+			FContactPoint Contact;
+			const FRigidTransform3 AToBTM = AStartTransform.GetRelativeTransform(BTransform);
+			const FVec3 LocalDir = BTransform.InverseTransformVectorNoScale(Dir);
 
-			T OutTime = FLT_MAX;
+			FReal OutTime = FLT_MAX;
 			int32 FaceIndex = -1;
-			TVec3<T> Location, Normal;
+			FVec3 Location, Normal;
 
-			Utilities::CastHelper(A, AStartTransform, [&](const auto& ADowncast, const TRigidTransform<T,d>& AFullTM)
+			Utilities::CastHelper(A, AStartTransform, [&](const auto& ADowncast, const FRigidTransform3& AFullTM)
 			{
 				if (B.SweepGeom(ADowncast, AToBTM, LocalDir, Length, OutTime, Location, Normal, FaceIndex, 0.0f, true))
 				{
@@ -309,23 +307,23 @@ namespace Chaos
 		}
 
 		// A is the implicit here, we want to return a contact point on B (trimesh)
-		template <typename GeometryA, typename T, int d>
-		TContactPoint<T> GJKImplicitScaledTriMeshSweptContactPoint(const FImplicitObject& A, const TRigidTransform<T, d>& AStartTransform, const TImplicitObjectScaled<FTriangleMeshImplicitObject>& B, const TRigidTransform<T, d>& BTransform, const TVector<T,d>& Dir, const T Length, const T CullDistance, T& TOI)
+		template <typename GeometryA>
+		FContactPoint GJKImplicitScaledTriMeshSweptContactPoint(const FImplicitObject& A, const FRigidTransform3& AStartTransform, const TImplicitObjectScaled<FTriangleMeshImplicitObject>& B, const FRigidTransform3& BTransform, const FVec3& Dir, const FReal Length, const FReal CullDistance, FReal& TOI)
 		{
-			TContactPoint<T> Contact;
-			const TRigidTransform<T, d> AToBTM = AStartTransform.GetRelativeTransform(BTransform);
-			const TVector<T, d> LocalDir = BTransform.InverseTransformVectorNoScale(Dir);
+			FContactPoint Contact;
+			const FRigidTransform3 AToBTM = AStartTransform.GetRelativeTransform(BTransform);
+			const FVec3 LocalDir = BTransform.InverseTransformVectorNoScale(Dir);
 
 			if (!ensure(B.GetType() & ImplicitObjectType::TriangleMesh) || !ensure(!IsInstanced(B.GetType())))
 			{
-				return TContactPoint<T>();
+				return FContactPoint();
 			}
 
-			T OutTime = FLT_MAX;
-			TVec3<T> Location, Normal;
+			FReal OutTime = FLT_MAX;
+			FVec3 Location, Normal;
 			int32 FaceIndex = -1;
 
-			Utilities::CastHelper(A, AStartTransform, [&](const auto& ADowncast, const TRigidTransform<T,d>& AFullTM)
+			Utilities::CastHelper(A, AStartTransform, [&](const auto& ADowncast, const FRigidTransform3& AFullTM)
 			{
 				if(B.LowLevelSweepGeom(ADowncast, AToBTM, LocalDir, Length, OutTime, Location, Normal, FaceIndex, 0.0f, true))
 				{
@@ -375,23 +373,22 @@ namespace Chaos
 		}
 
 
-		template<class T, int d>
-		TContactPoint<T> ConvexConvexContactPoint(const FImplicitObject& A, const TRigidTransform<T, d>& ATM, const FImplicitObject& B, const TRigidTransform<T, d>& BTM, const T CullDistance)
+		FContactPoint ConvexConvexContactPoint(const FImplicitObject& A, const FRigidTransform3& ATM, const FImplicitObject& B, const FRigidTransform3& BTM, const FReal CullDistance)
 		{
-			TContactPoint<T> ContactPoint = Utilities::CastHelper(A, ATM, [&](const auto& ADowncast, const TRigidTransform<T,d>& AFullTM)
+			FContactPoint ContactPoint = Utilities::CastHelper(A, ATM, [&](const auto& ADowncast, const FRigidTransform3& AFullTM)
 			{
-				return Utilities::CastHelper(B, BTM, [&](const auto& BDowncast, const TRigidTransform<T,d>& BFullTM)
+				return Utilities::CastHelper(B, BTM, [&](const auto& BDowncast, const FRigidTransform3& BFullTM)
 				{
-					return GJKContactPoint(ADowncast, AFullTM, BDowncast, BFullTM, TVector<T, d>(1, 0, 0));
+					return GJKContactPoint(ADowncast, AFullTM, BDowncast, BFullTM, FVec3(1, 0, 0));
 				});
 			});
 
-			if (FMath::Abs(ContactPoint.Phi) < (T)(PHI_RESAMPLE_THRESHOLD))
+			if (FMath::Abs(ContactPoint.Phi) < (FReal)(PHI_RESAMPLE_THRESHOLD))
 			{
 				// If GJKPenetration returns a phi of abs value < this number, we use PhiWithNormal to resample phi and normal.
 				// We have observed bad normals coming from GJKPenetration when barely in contact.
 
-				TVector<T, d> ContactLocalB = BTM.InverseTransformPosition(ContactPoint.Location);
+				FVec3 ContactLocalB = BTM.InverseTransformPosition(ContactPoint.Location);
 				ContactPoint.Phi = B.PhiWithNormal(ContactLocalB, ContactPoint.Normal);
 				ContactPoint.Normal = BTM.TransformVectorNoScale(ContactPoint.Normal);
 			}
@@ -399,23 +396,22 @@ namespace Chaos
 			return ContactPoint;
 		}
 
-		template<class T, int d>
-		TContactPoint<T> ConvexConvexContactPointSwept(const FImplicitObject& A, const TRigidTransform<T, d>& ATM, const FImplicitObject& B, const TRigidTransform<T, d>& BTM, const TVector<T, d>& Dir, const T Length, const T CullDistance, T& TOI)
+		FContactPoint ConvexConvexContactPointSwept(const FImplicitObject& A, const FRigidTransform3& ATM, const FImplicitObject& B, const FRigidTransform3& BTM, const FVec3& Dir, const FReal Length, const FReal CullDistance, FReal& TOI)
 		{
-			return Utilities::CastHelper(A, ATM, [&](const auto& ADowncast, const TRigidTransform<T,d>& AFullTM)
+			return Utilities::CastHelper(A, ATM, [&](const auto& ADowncast, const FRigidTransform3& AFullTM)
 			{
-				return Utilities::CastHelper(B, BTM, [&](const auto& BDowncast, const TRigidTransform<T,d>& BFullTM)
+				return Utilities::CastHelper(B, BTM, [&](const auto& BDowncast, const FRigidTransform3& BFullTM)
 				{
 					return GJKContactPointSwept(ADowncast, AFullTM, BDowncast, BFullTM, Dir, Length, CullDistance, TOI);
 				});
 			});
 		}
 
-		template <typename T, int d>
-		void UpdateSingleShotManifold(TRigidBodyMultiPointContactConstraint<T, d>& Constraint, const TRigidTransform<T, d>& WorldTransform0, const TRigidTransform<T, d>& WorldTransform1, const T CullDistance)
+
+		void UpdateSingleShotManifold(FRigidBodyMultiPointContactConstraint& Constraint, const FRigidTransform3& WorldTransform0, const FRigidTransform3& WorldTransform1, const FReal CullDistance)
 		{
 			// single shot manifolds for TConvex implicit object in the constraints implicit[0] position. 
-			TContactPoint<T> ContactPoint = ConvexConvexContactPoint(*Constraint.Manifold.Implicit[0], WorldTransform0, *Constraint.Manifold.Implicit[1], WorldTransform1, CullDistance);
+			FContactPoint ContactPoint = ConvexConvexContactPoint(*Constraint.Manifold.Implicit[0], WorldTransform0, *Constraint.Manifold.Implicit[1], WorldTransform1, CullDistance);
 
 			// Cache the nearest point as the initial contact
 			Constraint.Manifold.Phi = ContactPoint.Phi;
@@ -452,10 +448,10 @@ namespace Chaos
 			}
 		}
 
-		template <typename T, int d>
-		void UpdateIterativeManifold(TRigidBodyMultiPointContactConstraint<T, d>&  Constraint, const TRigidTransform<T, d>& WorldTransform0, const TRigidTransform<T, d>& WorldTransform1, const T CullDistance)
+
+		void UpdateIterativeManifold(FRigidBodyMultiPointContactConstraint&  Constraint, const FRigidTransform3& WorldTransform0, const FRigidTransform3& WorldTransform1, const FReal CullDistance)
 		{
-			auto SumSampleData = [&](TRigidBodyMultiPointContactConstraint<T, d>& LambdaConstraint) -> TVector<float, 3>
+			auto SumSampleData = [&](FRigidBodyMultiPointContactConstraint& LambdaConstraint) -> TVector<float, 3>
 			{
 				TVector<float, 3> ReturnValue(0);
 				for (int i = 0; i < LambdaConstraint.NumManifoldPoints(); i++)
@@ -466,7 +462,7 @@ namespace Chaos
 			};
 
 			// iterative manifolds for non TConvex implicit objects that require sampling 
-			TContactPoint<T> ContactPoint = ConvexConvexContactPoint(*Constraint.Manifold.Implicit[0], WorldTransform0, *Constraint.Manifold.Implicit[1], WorldTransform1, CullDistance);
+			FContactPoint ContactPoint = ConvexConvexContactPoint(*Constraint.Manifold.Implicit[0], WorldTransform0, *Constraint.Manifold.Implicit[1], WorldTransform1, CullDistance);
 
 			// Cache the nearest point as the initial contact
 			Constraint.Manifold.Phi = ContactPoint.Phi;
@@ -481,26 +477,26 @@ namespace Chaos
 				Constraint.SetManifoldPlane(1, INDEX_NONE, PlaneNormal, PlanePosition);
 			}
 
-			TVector<T, d> SurfaceSample = WorldTransform0.InverseTransformPosition(ContactPoint.Location);
+			FVec3 SurfaceSample = WorldTransform0.InverseTransformPosition(ContactPoint.Location);
 			if (Constraint.NumManifoldPoints() < 4)
 			{
 				Constraint.AddManifoldPoint(SurfaceSample);
 			}
 			else if (Constraint.NumManifoldPoints() == 4)
 			{
-				TVector<T, d> Center = SumSampleData(Constraint) / Constraint.NumManifoldPoints();
-				T Delta = (Center - SurfaceSample).SizeSquared();
+				FVec3 Center = SumSampleData(Constraint) / Constraint.NumManifoldPoints();
+				FReal Delta = (Center - SurfaceSample).SizeSquared();
 
 				//
 				// @todo(chaos) : Collision Manifold
 				//    The iterative manifold need to be maximized for area instead of largest 
 				//    distance from center.
 				//
-				T SmallestDelta = FLT_MAX;
+				FReal SmallestDelta = FLT_MAX;
 				int32 SmallestIndex = 0;
 				for (int32 idx = 0; idx < Constraint.NumManifoldPoints(); idx++)
 				{
-					T IdxDelta = (Center - Constraint.GetManifoldPoint(idx)).SizeSquared();
+					FReal IdxDelta = (Center - Constraint.GetManifoldPoint(idx)).SizeSquared();
 					if (IdxDelta < SmallestDelta)
 					{
 						SmallestDelta = IdxDelta;
@@ -535,17 +531,17 @@ namespace Chaos
 		//
 		// Box - Box
 		//
-		template <typename T, int d>
-		TContactPoint<T> BoxBoxContactPoint(const TAABB<T, d>& Box1, const TAABB<T, d>& Box2, const TRigidTransform<T, d>& ATM, const TRigidTransform<T, d>& BToATM, const T CullDistance)
+
+		FContactPoint BoxBoxContactPoint(const FAABB3& Box1, const FAABB3& Box2, const FRigidTransform3& ATM, const FRigidTransform3& BToATM, const FReal CullDistance)
 		{
-			return GJKContactPoint2(Box1, Box2, ATM, BToATM, TVector<T, d>(1, 0, 0));
+			return GJKContactPoint2(Box1, Box2, ATM, BToATM, FVec3(1, 0, 0));
 		}
 
-		template <typename T, int d>
-		void UpdateBoxBoxConstraint(const TAABB<T, d>& Box1, const TRigidTransform<T, d>& Box1Transform, const TAABB<T, d>& Box2, const TRigidTransform<T, d>& Box2Transform, const T CullDistance, TRigidBodyPointContactConstraint<T, d>& Constraint)
+
+		void UpdateBoxBoxConstraint(const FAABB3& Box1, const FRigidTransform3& Box1Transform, const FAABB3& Box2, const FRigidTransform3& Box2Transform, const FReal CullDistance, FRigidBodyPointContactConstraint& Constraint)
 		{
-			const TRigidTransform<T, d> Box2ToBox1TM = Box2Transform.GetRelativeTransform(Box1Transform);
-			TAABB<T, d> Box2In1 = Box2.TransformedAABB(Box2ToBox1TM);
+			const FRigidTransform3 Box2ToBox1TM = Box2Transform.GetRelativeTransform(Box1Transform);
+			FAABB3 Box2In1 = Box2.TransformedAABB(Box2ToBox1TM);
 			Box2In1.Thicken(CullDistance);
 			if (Box1.Intersects(Box2In1))
 			{
@@ -553,25 +549,24 @@ namespace Chaos
 			}
 		}
 
-		template <typename T, int d>
-		void UpdateBoxBoxManifold(TCollisionConstraintBase<T, d>&  Constraint, const TRigidTransform<T, d>& ATM, const TRigidTransform<T, d>& BTM, const T CullDistance)
+
+		void UpdateBoxBoxManifold(FCollisionConstraintBase&  Constraint, const FRigidTransform3& ATM, const FRigidTransform3& BTM, const FReal CullDistance)
 		{
 			// @todo(chaos) : Stub Update Manifold
 			//   Stub function for updating the manifold prior to the Apply and ApplyPushOut
 		}
 
-		template<typename T, int d>
-		void ConstructBoxBoxConstraints(TGeometryParticleHandle<T, d>* Particle0, TGeometryParticleHandle<T, d>* Particle1, const FImplicitObject* Implicit0, const FImplicitObject* Implicit1, const TRigidTransform<T, d>& LocalTransform0, const TRigidTransform<T, d>& LocalTransform1, const T CullDistance, FCollisionConstraintsArray& NewConstraints)
+		void ConstructBoxBoxConstraints(TGeometryParticleHandle<FReal, 3>* Particle0, TGeometryParticleHandle<FReal, 3>* Particle1, const FImplicitObject* Implicit0, const FImplicitObject* Implicit1, const FRigidTransform3& LocalTransform0, const FRigidTransform3& LocalTransform1, const FReal CullDistance, FCollisionConstraintsArray& NewConstraints)
 		{
-			const TBox<T, d> * Object0 = Implicit0->template GetObject<const TBox<T, d> >();
-			const TBox<T, d> * Object1 = Implicit1->template GetObject<const TBox<T, d> >();
+			const TBox<FReal, 3> * Object0 = Implicit0->template GetObject<const TBox<FReal, 3> >();
+			const TBox<FReal, 3> * Object1 = Implicit1->template GetObject<const TBox<FReal, 3> >();
 			if (ensure(Object0 && Object1))
 			{
 				FRigidBodyPointContactConstraint Constraint = FRigidBodyPointContactConstraint(Particle0, Implicit0, LocalTransform0, Particle1, Implicit1, LocalTransform1, EContactShapesType::BoxBox);
 				if (!bChaos_Collision_DeferNarrowPhase)
 				{
-					TRigidTransform<T, d> WorldTransform0 = LocalTransform0 * Collisions::GetTransform(Particle0);
-					TRigidTransform<T, d> WorldTransform1 = LocalTransform1 * Collisions::GetTransform(Particle1);
+					FRigidTransform3 WorldTransform0 = LocalTransform0 * Collisions::GetTransform(Particle0);
+					FRigidTransform3 WorldTransform1 = LocalTransform1 * Collisions::GetTransform(Particle1);
 					UpdateBoxBoxConstraint(Object0->BoundingBox(), WorldTransform0, Object1->BoundingBox(), WorldTransform1, CullDistance, Constraint);
 					NewConstraints.TryAdd(CullDistance, Constraint);
 				}
@@ -586,38 +581,36 @@ namespace Chaos
 		// Box - HeightField
 		//
 
-		template <typename T, int d>
-		TContactPoint<T> BoxHeightFieldContactPoint(const TAABB<T, d>& A, const TRigidTransform<T, d>& ATransform, const THeightField<T>& B, const TRigidTransform<T, d>& BTransform, const T CullDistance)
+
+		FContactPoint BoxHeightFieldContactPoint(const FAABB3& A, const FRigidTransform3& ATransform, const THeightField<FReal>& B, const FRigidTransform3& BTransform, const FReal CullDistance)
 		{
 			return GJKImplicitContactPoint< TBox<float, 3> >(TBox<float, 3>(A), ATransform, B, BTransform, CullDistance);
 		}
 
-		template <typename T, int d>
-		void UpdateBoxHeightFieldConstraint(const TAABB<T, d>& A, const TRigidTransform<T, d>& ATransform, const THeightField<T>& B, const TRigidTransform<T, d>& BTransform, const T CullDistance, TRigidBodyPointContactConstraint<T, d>& Constraint)
+
+		void UpdateBoxHeightFieldConstraint(const FAABB3& A, const FRigidTransform3& ATransform, const THeightField<FReal>& B, const FRigidTransform3& BTransform, const FReal CullDistance, FRigidBodyPointContactConstraint& Constraint)
 		{
 			UpdateContactPoint(Constraint.Manifold, BoxHeightFieldContactPoint(A, ATransform, B, BTransform, CullDistance));
 		}
 
-		template<class T, int d>
-		void UpdateBoxHeightFieldManifold(TCollisionConstraintBase<T, d>&  Constraint, const TRigidTransform<T, d>& ATM, const TRigidTransform<T, d>& BTM, const T CullDistance)
+		void UpdateBoxHeightFieldManifold(FCollisionConstraintBase&  Constraint, const FRigidTransform3& ATM, const FRigidTransform3& BTM, const FReal CullDistance)
 		{
 			// @todo(chaos) : Stub Update Manifold
 			//   Stub function for updating the manifold prior to the Apply and ApplyPushOut
 		}
 
-		template<typename T, int d>
-		void ConstructBoxHeightFieldConstraints(TGeometryParticleHandle<T, d>* Particle0, TGeometryParticleHandle<T, d>* Particle1, const FImplicitObject* Implicit0, const FImplicitObject* Implicit1, const TRigidTransform<T, d>& LocalTransform0, const TRigidTransform<T, d>& LocalTransform1, const T CullDistance, FCollisionConstraintsArray& NewConstraints)
+		void ConstructBoxHeightFieldConstraints(TGeometryParticleHandle<FReal, 3>* Particle0, TGeometryParticleHandle<FReal, 3>* Particle1, const FImplicitObject* Implicit0, const FImplicitObject* Implicit1, const FRigidTransform3& LocalTransform0, const FRigidTransform3& LocalTransform1, const FReal CullDistance, FCollisionConstraintsArray& NewConstraints)
 		{
 
-			const TBox<T, d> * Object0 = Implicit0->template GetObject<const TBox<T, d> >();
-			const THeightField<T> * Object1 = Implicit1->template GetObject<const THeightField<T> >();
+			const TBox<FReal, 3> * Object0 = Implicit0->template GetObject<const TBox<FReal, 3> >();
+			const THeightField<FReal> * Object1 = Implicit1->template GetObject<const THeightField<FReal> >();
 			if (ensure(Object0 && Object1))
 			{
 				FRigidBodyPointContactConstraint Constraint = FRigidBodyPointContactConstraint(Particle0, Implicit0, LocalTransform0, Particle1, Implicit1, LocalTransform1, EContactShapesType::BoxHeightField);
 				if (!bChaos_Collision_DeferNarrowPhase)
 				{
-					TRigidTransform<T, d> WorldTransform0 = LocalTransform0 * Collisions::GetTransform(Particle0);
-					TRigidTransform<T, d> WorldTransform1 = LocalTransform1 * Collisions::GetTransform(Particle1);
+					FRigidTransform3 WorldTransform0 = LocalTransform0 * Collisions::GetTransform(Particle0);
+					FRigidTransform3 WorldTransform1 = LocalTransform1 * Collisions::GetTransform(Particle1);
 					UpdateBoxHeightFieldConstraint(Object0->BoundingBox(), WorldTransform0, *Object1, WorldTransform1, CullDistance, Constraint);
 					NewConstraints.TryAdd(CullDistance, Constraint);
 				}
@@ -634,41 +627,40 @@ namespace Chaos
 		// Box-Plane
 		//
 
-		template <typename T, int d>
-		void UpdateBoxPlaneConstraint(const TAABB<T, d>& Box, const TRigidTransform<T, d>& BoxTransform, const TPlane<T, d>& Plane, const TRigidTransform<T, d>& PlaneTransform, const T CullDistance, TRigidBodyPointContactConstraint<T, d>& Constraint)
+		void UpdateBoxPlaneConstraint(const FAABB3& Box, const FRigidTransform3& BoxTransform, const TPlane<FReal, 3>& Plane, const FRigidTransform3& PlaneTransform, const FReal CullDistance, FRigidBodyPointContactConstraint& Constraint)
 		{
-			TCollisionContact<T, d> & Contact = Constraint.Manifold;
+			FCollisionContact & Contact = Constraint.Manifold;
 
 #if USING_CODE_ANALYSIS
 			MSVC_PRAGMA(warning(push))
 				MSVC_PRAGMA(warning(disable : ALL_CODE_ANALYSIS_WARNINGS))
 #endif	// USING_CODE_ANALYSIS
 
-			const TRigidTransform<T, d> BoxToPlaneTransform(BoxTransform.GetRelativeTransform(PlaneTransform));
-			const TVector<T, d> Extents = Box.Extents();
-			constexpr int32 NumCorners = 2 + 2 * d;
-			constexpr T Epsilon = KINDA_SMALL_NUMBER;
+			const FRigidTransform3 BoxToPlaneTransform(BoxTransform.GetRelativeTransform(PlaneTransform));
+			const FVec3 Extents = Box.Extents();
+			constexpr int32 NumCorners = 2 + 2 * 3;
+			constexpr FReal Epsilon = KINDA_SMALL_NUMBER;
 
-			TVector<T, d> Corners[NumCorners];
+			FVec3 Corners[NumCorners];
 			int32 CornerIdx = 0;
 			Corners[CornerIdx++] = BoxToPlaneTransform.TransformPosition(Box.Max());
 			Corners[CornerIdx++] = BoxToPlaneTransform.TransformPosition(Box.Min());
-			for (int32 j = 0; j < d; ++j)
+			for (int32 j = 0; j < 3; ++j)
 			{
-				Corners[CornerIdx++] = BoxToPlaneTransform.TransformPosition(Box.Min() + TVector<T, d>::AxisVector(j) * Extents);
-				Corners[CornerIdx++] = BoxToPlaneTransform.TransformPosition(Box.Max() - TVector<T, d>::AxisVector(j) * Extents);
+				Corners[CornerIdx++] = BoxToPlaneTransform.TransformPosition(Box.Min() + FVec3::AxisVector(j) * Extents);
+				Corners[CornerIdx++] = BoxToPlaneTransform.TransformPosition(Box.Max() - FVec3::AxisVector(j) * Extents);
 			}
 
 #if USING_CODE_ANALYSIS
 			MSVC_PRAGMA(warning(pop))
 #endif	// USING_CODE_ANALYSIS
 
-			TVector<T, d> PotentialConstraints[NumCorners];
+			FVec3 PotentialConstraints[NumCorners];
 			int32 NumConstraints = 0;
 			for (int32 i = 0; i < NumCorners; ++i)
 			{
-				TVector<T, d> Normal;
-				const T NewPhi = Plane.PhiWithNormal(Corners[i], Normal);
+				FVec3 Normal;
+				const FReal NewPhi = Plane.PhiWithNormal(Corners[i], Normal);
 				if (NewPhi < Contact.Phi + Epsilon)
 				{
 					if (NewPhi <= Contact.Phi - Epsilon)
@@ -683,7 +675,7 @@ namespace Chaos
 			}
 			if (NumConstraints > 1)
 			{
-				TVector<T, d> AverageLocation(0);
+				FVec3 AverageLocation(0);
 				for (int32 ConstraintIdx = 0; ConstraintIdx < NumConstraints; ++ConstraintIdx)
 				{
 					AverageLocation += PotentialConstraints[ConstraintIdx];
@@ -692,27 +684,25 @@ namespace Chaos
 			}
 		}
 
-		template<class T, int d>
-		void UpdateBoxPlaneManifold(TCollisionConstraintBase<T, d>&  Constraint, const TRigidTransform<T, d>& ATM, const TRigidTransform<T, d>& BTM, const T CullDistance)
+		void UpdateBoxPlaneManifold(FCollisionConstraintBase&  Constraint, const FRigidTransform3& ATM, const FRigidTransform3& BTM, const FReal CullDistance)
 		{
 			// @todo(chaos) : Stub Update Manifold
 			//   Stub function for updating the manifold prior to the Apply and ApplyPushOut
 		}
 
 
-		template<typename T, int d>
-		void ConstructBoxPlaneConstraints(TGeometryParticleHandle<T, d>* Particle0, TGeometryParticleHandle<T, d>* Particle1, const FImplicitObject* Implicit0, const FImplicitObject* Implicit1, const TRigidTransform<T, d>& LocalTransform0, const TRigidTransform<T, d>& LocalTransform1, const T CullDistance, FCollisionConstraintsArray& NewConstraints)
+		void ConstructBoxPlaneConstraints(TGeometryParticleHandle<FReal, 3>* Particle0, TGeometryParticleHandle<FReal, 3>* Particle1, const FImplicitObject* Implicit0, const FImplicitObject* Implicit1, const FRigidTransform3& LocalTransform0, const FRigidTransform3& LocalTransform1, const FReal CullDistance, FCollisionConstraintsArray& NewConstraints)
 		{
 
-			const TBox<T, d> * Object0 = Implicit0->template GetObject<const TBox<T, d> >();
-			const TPlane<T, d> * Object1 = Implicit1->template GetObject<const TPlane<T, d> >();
+			const TBox<FReal, 3> * Object0 = Implicit0->template GetObject<const TBox<FReal, 3> >();
+			const TPlane<FReal, 3>* Object1 = Implicit1->template GetObject<const TPlane<FReal, 3> >();
 			if (ensure(Object0 && Object1))
 			{
 				FRigidBodyPointContactConstraint Constraint = FRigidBodyPointContactConstraint(Particle0, Implicit0, LocalTransform0, Particle1, Implicit1, LocalTransform1, EContactShapesType::BoxPlane);
 				if (!bChaos_Collision_DeferNarrowPhase)
 				{
-					TRigidTransform<T, d> WorldTransform0 = LocalTransform0 * Collisions::GetTransform(Particle0);
-					TRigidTransform<T, d> WorldTransform1 = LocalTransform1 * Collisions::GetTransform(Particle1);
+					FRigidTransform3 WorldTransform0 = LocalTransform0 * Collisions::GetTransform(Particle0);
+					FRigidTransform3 WorldTransform1 = LocalTransform1 * Collisions::GetTransform(Particle1);
 					UpdateBoxPlaneConstraint(Object0->BoundingBox(), WorldTransform0, *Object1, WorldTransform1, CullDistance, Constraint);
 					NewConstraints.TryAdd(CullDistance, Constraint);
 				}
@@ -727,28 +717,27 @@ namespace Chaos
 		// Box-TriangleMesh
 		//
 
-		template <typename TriMeshType, typename T, int d>
-		TContactPoint<T> BoxTriangleMeshContactPoint(const TAABB<T, d>& A, const TRigidTransform<T, d>& ATransform, const TriMeshType& B, const TRigidTransform<T, d>& BTransform, const T CullDistance)
+		template <typename TriMeshType>
+		FContactPoint BoxTriangleMeshContactPoint(const FAABB3& A, const FRigidTransform3& ATransform, const TriMeshType& B, const FRigidTransform3& BTransform, const FReal CullDistance)
 		{
 			return GJKImplicitContactPoint< TBox<float, 3> >(TBox<float, 3>(A), ATransform, B, BTransform, CullDistance);
 		}
 
-		template <typename TriMeshType, typename T, int d>
-		void UpdateBoxTriangleMeshConstraint(const TAABB<T, d>& Box0, const TRigidTransform<T, d>& WorldTransform0, const TriMeshType& TriangleMesh1, const TRigidTransform<T, d>& WorldTransform1, const T CullDistance, TRigidBodyPointContactConstraint<T, d>& Constraint)
+		template <typename TriMeshType>
+		void UpdateBoxTriangleMeshConstraint(const FAABB3& Box0, const FRigidTransform3& WorldTransform0, const TriMeshType& TriangleMesh1, const FRigidTransform3& WorldTransform1, const FReal CullDistance, FRigidBodyPointContactConstraint& Constraint)
 		{
 			UpdateContactPoint(Constraint.Manifold, BoxTriangleMeshContactPoint(Box0, WorldTransform0, TriangleMesh1, WorldTransform1, CullDistance));
 		}
 
 
-		template <typename T, int d>
-		void UpdateBoxTriangleMeshManifold(TCollisionConstraintBase<T, d>&  Constraint, const TRigidTransform<T, d>& ATM, const TRigidTransform<T, d>& BTM, const T CullDistance)
+
+		void UpdateBoxTriangleMeshManifold(FCollisionConstraintBase&  Constraint, const FRigidTransform3& ATM, const FRigidTransform3& BTM, const FReal CullDistance)
 		{
 
 		}
 
-		template<typename T, int d>
-		void ConstructBoxTriangleMeshConstraints(TGeometryParticleHandle<T, d>* Particle0, TGeometryParticleHandle<T, d>* Particle1, const FImplicitObject* Implicit0, const FImplicitObject* Implicit1, const TRigidTransform<T, d>& LocalTransform0, const TRigidTransform<T, d>& LocalTransform1, const T CullDistance, FCollisionConstraintsArray& NewConstraints)		{
-			const TBox<T, d> * Object0 = Implicit0->template GetObject<const TBox<T, d> >();
+		void ConstructBoxTriangleMeshConstraints(TGeometryParticleHandle<FReal, 3>* Particle0, TGeometryParticleHandle<FReal, 3>* Particle1, const FImplicitObject* Implicit0, const FImplicitObject* Implicit1, const FRigidTransform3& LocalTransform0, const FRigidTransform3& LocalTransform1, const FReal CullDistance, FCollisionConstraintsArray& NewConstraints)		{
+			const TBox<FReal, 3> * Object0 = Implicit0->template GetObject<const TBox<FReal, 3> >();
 			if (ensure(Object0))
 			{
 				if (const TImplicitObjectScaled<FTriangleMeshImplicitObject>* ScaledTriangleMesh = Implicit1->template GetObject<const TImplicitObjectScaled<FTriangleMeshImplicitObject>>())
@@ -756,8 +745,8 @@ namespace Chaos
 					FRigidBodyPointContactConstraint Constraint = FRigidBodyPointContactConstraint(Particle0, Implicit0, LocalTransform0, Particle1, Implicit1, LocalTransform1, EContactShapesType::BoxTriMesh);
 					if (!bChaos_Collision_DeferNarrowPhase)
 					{
-						TRigidTransform<T, d> WorldTransform0 = LocalTransform0 * Collisions::GetTransform(Particle0);
-						TRigidTransform<T, d> WorldTransform1 = LocalTransform1 * Collisions::GetTransform(Particle1);
+						FRigidTransform3 WorldTransform0 = LocalTransform0 * Collisions::GetTransform(Particle0);
+						FRigidTransform3 WorldTransform1 = LocalTransform1 * Collisions::GetTransform(Particle1);
 						UpdateBoxTriangleMeshConstraint(Object0->GetAABB(), WorldTransform0, *ScaledTriangleMesh, WorldTransform1, CullDistance, Constraint);
 						NewConstraints.TryAdd(CullDistance, Constraint);
 					}
@@ -771,8 +760,8 @@ namespace Chaos
 					FRigidBodyPointContactConstraint Constraint = FRigidBodyPointContactConstraint(Particle0, Implicit0, LocalTransform0, Particle1, Implicit1, LocalTransform1, EContactShapesType::BoxTriMesh);
 					if (!bChaos_Collision_DeferNarrowPhase)
 					{
-						TRigidTransform<T, d> WorldTransform0 = LocalTransform0 * Collisions::GetTransform(Particle0);
-						TRigidTransform<T, d> WorldTransform1 = LocalTransform1 * Collisions::GetTransform(Particle1);
+						FRigidTransform3 WorldTransform0 = LocalTransform0 * Collisions::GetTransform(Particle0);
+						FRigidTransform3 WorldTransform1 = LocalTransform1 * Collisions::GetTransform(Particle1);
 						UpdateBoxTriangleMeshConstraint(Object0->GetAABB(), WorldTransform0, *TriangleMesh, WorldTransform1, CullDistance, Constraint);
 						NewConstraints.TryAdd(CullDistance, Constraint);
 					}
@@ -792,48 +781,46 @@ namespace Chaos
 		// Sphere - Sphere
 		//
 
-		template <typename T, int d>
-		TContactPoint<T> SphereSphereContactPoint(const TSphere<T, d>& Sphere1, const TRigidTransform<T, d>& Sphere1Transform, const TSphere<T, d>& Sphere2, const TRigidTransform<T, d>& Sphere2Transform, const T CullDistance)
+
+		FContactPoint SphereSphereContactPoint(const TSphere<FReal, 3>& Sphere1, const FRigidTransform3& Sphere1Transform, const TSphere<FReal, 3>& Sphere2, const FRigidTransform3& Sphere2Transform, const FReal CullDistance)
 		{
-			TContactPoint<T> Result;
+			FContactPoint Result;
 
-			const TVector<T, d> Center1 = Sphere1Transform.TransformPosition(Sphere1.GetCenter());
-			const TVector<T, d> Center2 = Sphere2Transform.TransformPosition(Sphere2.GetCenter());
-			const TVector<T, d> Direction = Center1 - Center2;
-			const T Size = Direction.Size();
-			const T NewPhi = Size - (Sphere1.GetRadius() + Sphere2.GetRadius());
+			const FVec3 Center1 = Sphere1Transform.TransformPosition(Sphere1.GetCenter());
+			const FVec3 Center2 = Sphere2Transform.TransformPosition(Sphere2.GetCenter());
+			const FVec3 Direction = Center1 - Center2;
+			const FReal Size = Direction.Size();
+			const FReal NewPhi = Size - (Sphere1.GetRadius() + Sphere2.GetRadius());
 			Result.Phi = NewPhi;
-			Result.Normal = Size > SMALL_NUMBER ? Direction / Size : TVector<T, d>(0, 0, 1);
+			Result.Normal = Size > SMALL_NUMBER ? Direction / Size : FVec3(0, 0, 1);
 			Result.Location = Center1 - Sphere1.GetRadius() * Result.Normal;
 
 			return Result;
 		}
 
-		template <typename T, int d>
-		void UpdateSphereSphereConstraint(const TSphere<T, d>& Sphere1, const TRigidTransform<T, d>& Sphere1Transform, const TSphere<T, d>& Sphere2, const TRigidTransform<T, d>& Sphere2Transform, const T CullDistance, TRigidBodyPointContactConstraint<T, d>& Constraint)
+
+		void UpdateSphereSphereConstraint(const TSphere<FReal, 3>& Sphere1, const FRigidTransform3& Sphere1Transform, const TSphere<FReal, 3>& Sphere2, const FRigidTransform3& Sphere2Transform, const FReal CullDistance, FRigidBodyPointContactConstraint& Constraint)
 		{
 			UpdateContactPoint(Constraint.Manifold, SphereSphereContactPoint(Sphere1, Sphere1Transform, Sphere2, Sphere2Transform, CullDistance));
 		}
 
-		template<class T, int d>
-		void UpdateSphereSphereManifold(TCollisionConstraintBase<T, d>&  Constraint, const TRigidTransform<T, d>& ATM, const TRigidTransform<T, d>& BTM, const T CullDistance)
+		void UpdateSphereSphereManifold(FCollisionConstraintBase&  Constraint, const FRigidTransform3& ATM, const FRigidTransform3& BTM, const FReal CullDistance)
 		{
 			// @todo(chaos) : Stub Update Manifold
 			//   Stub function for updating the manifold prior to the Apply and ApplyPushOut
 		}
 
-		template<typename T, int d>
-		void ConstructSphereSphereConstraints(TGeometryParticleHandle<T, d>* Particle0, TGeometryParticleHandle<T, d>* Particle1, const FImplicitObject* Implicit0, const FImplicitObject* Implicit1, const TRigidTransform<T, d>& LocalTransform0, const TRigidTransform<T, d>& LocalTransform1, const T CullDistance, FCollisionConstraintsArray& NewConstraints)
+		void ConstructSphereSphereConstraints(TGeometryParticleHandle<FReal, 3>* Particle0, TGeometryParticleHandle<FReal, 3>* Particle1, const FImplicitObject* Implicit0, const FImplicitObject* Implicit1, const FRigidTransform3& LocalTransform0, const FRigidTransform3& LocalTransform1, const FReal CullDistance, FCollisionConstraintsArray& NewConstraints)
 		{
-			const TSphere<T, d> * Object0 = Implicit0->template GetObject<const TSphere<T, d> >();
-			const TSphere<T, d> * Object1 = Implicit1->template GetObject<const TSphere<T, d> >();
+			const TSphere<FReal, 3>* Object0 = Implicit0->template GetObject<const TSphere<FReal, 3> >();
+			const TSphere<FReal, 3>* Object1 = Implicit1->template GetObject<const TSphere<FReal, 3> >();
 			if (ensure(Object0 && Object1))
 			{
 				FRigidBodyPointContactConstraint Constraint = FRigidBodyPointContactConstraint(Particle0, Implicit0, LocalTransform0, Particle1, Implicit1, LocalTransform1, EContactShapesType::SphereSphere);
 				if (!bChaos_Collision_DeferNarrowPhase)
 				{
-					TRigidTransform<T, d> WorldTransform0 = LocalTransform0 * Collisions::GetTransform(Particle0);
-					TRigidTransform<T, d> WorldTransform1 = LocalTransform1 * Collisions::GetTransform(Particle1);
+					FRigidTransform3 WorldTransform0 = LocalTransform0 * Collisions::GetTransform(Particle0);
+					FRigidTransform3 WorldTransform1 = LocalTransform1 * Collisions::GetTransform(Particle1);
 					UpdateSphereSphereConstraint(*Object0, WorldTransform0, *Object1, WorldTransform1, CullDistance, Constraint);
 					NewConstraints.TryAdd(CullDistance, Constraint);
 				}
@@ -848,38 +835,35 @@ namespace Chaos
 		// Sphere - HeightField
 		//
 
-		template <typename T, int d>
-		TContactPoint<T> SphereHeightFieldContactPoint(const TSphere<T, d>& A, const TRigidTransform<T, d>& ATransform, const THeightField<T>& B, const TRigidTransform<T, d>& BTransform, const T CullDistance)
+		FContactPoint SphereHeightFieldContactPoint(const TSphere<FReal, 3>& A, const FRigidTransform3& ATransform, const THeightField<FReal>& B, const FRigidTransform3& BTransform, const FReal CullDistance)
 		{
 			return GJKImplicitContactPoint< TSphere<float, 3> >(TSphere<float, 3>(A), ATransform, B, BTransform, CullDistance);
 		}
 
-		template <typename T, int d>
-		void UpdateSphereHeightFieldConstraint(const TSphere<T, d>& A, const TRigidTransform<T, d>& ATransform, const THeightField<T>& B, const TRigidTransform<T, d>& BTransform, const T CullDistance, TRigidBodyPointContactConstraint<T, d>& Constraint)
+
+		void UpdateSphereHeightFieldConstraint(const TSphere<FReal, 3>& A, const FRigidTransform3& ATransform, const THeightField<FReal>& B, const FRigidTransform3& BTransform, const FReal CullDistance, FRigidBodyPointContactConstraint& Constraint)
 		{
 			UpdateContactPoint(Constraint.Manifold, SphereHeightFieldContactPoint(A, ATransform, B, BTransform, CullDistance));
 		}
 
-		template<class T, int d>
-		void UpdateSphereHeightFieldManifold(TCollisionConstraintBase<T, d>&  Constraint, const TRigidTransform<T, d>& ATM, const TRigidTransform<T, d>& BTM, const T CullDistance)
+		void UpdateSphereHeightFieldManifold(FCollisionConstraintBase&  Constraint, const FRigidTransform3& ATM, const FRigidTransform3& BTM, const FReal CullDistance)
 		{
 			// @todo(chaos) : Stub Update Manifold
 			//   Stub function for updating the manifold prior to the Apply and ApplyPushOut
 		}
 
-		template<typename T, int d>
-		void ConstructSphereHeightFieldConstraints(TGeometryParticleHandle<T, d>* Particle0, TGeometryParticleHandle<T, d>* Particle1, const FImplicitObject* Implicit0, const FImplicitObject* Implicit1, const TRigidTransform<T, d>& LocalTransform0, const TRigidTransform<T, d>& LocalTransform1, const T CullDistance, FCollisionConstraintsArray& NewConstraints)
+		void ConstructSphereHeightFieldConstraints(TGeometryParticleHandle<FReal, 3>* Particle0, TGeometryParticleHandle<FReal, 3>* Particle1, const FImplicitObject* Implicit0, const FImplicitObject* Implicit1, const FRigidTransform3& LocalTransform0, const FRigidTransform3& LocalTransform1, const FReal CullDistance, FCollisionConstraintsArray& NewConstraints)
 		{
 
-			const TSphere<T, d> * Object0 = Implicit0->template GetObject<const TSphere<T, d> >();
-			const THeightField<T> * Object1 = Implicit1->template GetObject<const THeightField<T> >();
+			const TSphere<FReal, 3>* Object0 = Implicit0->template GetObject<const TSphere<FReal, 3> >();
+			const THeightField<FReal> * Object1 = Implicit1->template GetObject<const THeightField<FReal> >();
 			if (ensure(Object0 && Object1))
 			{
 				FRigidBodyPointContactConstraint Constraint = FRigidBodyPointContactConstraint(Particle0, Implicit0, LocalTransform0, Particle1, Implicit1, LocalTransform1, EContactShapesType::SphereHeightField);
 				if (!bChaos_Collision_DeferNarrowPhase)
 				{
-					TRigidTransform<T, d> WorldTransform0 = LocalTransform0 * Collisions::GetTransform(Particle0);
-					TRigidTransform<T, d> WorldTransform1 = LocalTransform1 * Collisions::GetTransform(Particle1);
+					FRigidTransform3 WorldTransform0 = LocalTransform0 * Collisions::GetTransform(Particle0);
+					FRigidTransform3 WorldTransform1 = LocalTransform1 * Collisions::GetTransform(Particle1);
 					UpdateSphereHeightFieldConstraint(*Object0, WorldTransform0, *Object1, WorldTransform1, CullDistance, Constraint);
 					NewConstraints.TryAdd(CullDistance, Constraint);
 				}
@@ -894,16 +878,15 @@ namespace Chaos
 		//  Sphere-Plane
 		//
 
-		template <typename T, int d>
-		void UpdateSpherePlaneConstraint(const TSphere<T, d>& Sphere, const TRigidTransform<T, d>& SphereTransform, const TPlane<T, d>& Plane, const TRigidTransform<T, d>& PlaneTransform, const T CullDistance, TRigidBodyPointContactConstraint<T, d>& Constraint)
+		void UpdateSpherePlaneConstraint(const TSphere<FReal, 3>& Sphere, const FRigidTransform3& SphereTransform, const TPlane<FReal, 3>& Plane, const FRigidTransform3& PlaneTransform, const FReal CullDistance, FRigidBodyPointContactConstraint& Constraint)
 		{
-			TCollisionContact<T, d> & Contact = Constraint.Manifold;
+			FCollisionContact & Contact = Constraint.Manifold;
 
-			const TRigidTransform<T, d> SphereToPlaneTransform(PlaneTransform.Inverse() * SphereTransform);
-			const TVector<T, d> SphereCenter = SphereToPlaneTransform.TransformPosition(Sphere.GetCenter());
+			const FRigidTransform3 SphereToPlaneTransform(PlaneTransform.Inverse() * SphereTransform);
+			const FVec3 SphereCenter = SphereToPlaneTransform.TransformPosition(Sphere.GetCenter());
 
-			TVector<T, d> NewNormal;
-			T NewPhi = Plane.PhiWithNormal(SphereCenter, NewNormal);
+			FVec3 NewNormal;
+			FReal NewPhi = Plane.PhiWithNormal(SphereCenter, NewNormal);
 			NewPhi -= Sphere.GetRadius();
 
 			if (NewPhi < Contact.Phi)
@@ -914,26 +897,24 @@ namespace Chaos
 			}
 		}
 
-		template<class T, int d>
-		void UpdateSpherePlaneManifold(TCollisionConstraintBase<T, d>&  Constraint, const TRigidTransform<T, d>& ATM, const TRigidTransform<T, d>& BTM, const T CullDistance)
+		void UpdateSpherePlaneManifold(FCollisionConstraintBase&  Constraint, const FRigidTransform3& ATM, const FRigidTransform3& BTM, const FReal CullDistance)
 		{
 			// @todo(chaos) : Stub Update Manifold
 			//   Stub function for updating the manifold prior to the Apply and ApplyPushOut
 		}
 
-		template<typename T, int d>
-		void ConstructSpherePlaneConstraints(TGeometryParticleHandle<T, d>* Particle0, TGeometryParticleHandle<T, d>* Particle1, const FImplicitObject* Implicit0, const FImplicitObject* Implicit1, const TRigidTransform<T, d>& LocalTransform0, const TRigidTransform<T, d>& LocalTransform1, const T CullDistance, FCollisionConstraintsArray& NewConstraints)
+		void ConstructSpherePlaneConstraints(TGeometryParticleHandle<FReal, 3>* Particle0, TGeometryParticleHandle<FReal, 3>* Particle1, const FImplicitObject* Implicit0, const FImplicitObject* Implicit1, const FRigidTransform3& LocalTransform0, const FRigidTransform3& LocalTransform1, const FReal CullDistance, FCollisionConstraintsArray& NewConstraints)
 		{
 
-			const TSphere<T, d> * Object0 = Implicit0->template GetObject<const TSphere<T, d> >();
-			const TPlane<T, d> * Object1 = Implicit1->template GetObject<const TPlane<T, d> >();
+			const TSphere<FReal, 3>* Object0 = Implicit0->template GetObject<const TSphere<FReal, 3> >();
+			const TPlane<FReal, 3>* Object1 = Implicit1->template GetObject<const TPlane<FReal, 3> >();
 			if (ensure(Object0 && Object1))
 			{
 				FRigidBodyPointContactConstraint Constraint = FRigidBodyPointContactConstraint(Particle0, Implicit0, LocalTransform0, Particle1, Implicit1, LocalTransform1, EContactShapesType::SpherePlane);
 				if (!bChaos_Collision_DeferNarrowPhase)
 				{
-					TRigidTransform<T, d> WorldTransform0 = LocalTransform0 * Collisions::GetTransform(Particle0);
-					TRigidTransform<T, d> WorldTransform1 = LocalTransform1 * Collisions::GetTransform(Particle1);
+					FRigidTransform3 WorldTransform0 = LocalTransform0 * Collisions::GetTransform(Particle0);
+					FRigidTransform3 WorldTransform1 = LocalTransform1 * Collisions::GetTransform(Particle1);
 					UpdateSpherePlaneConstraint(*Object0, WorldTransform0, *Object1, WorldTransform1, CullDistance, Constraint);
 					NewConstraints.TryAdd(CullDistance, Constraint);
 				}
@@ -948,16 +929,16 @@ namespace Chaos
 		// Sphere - Box
 		//
 
-		template <typename T, int d>
-		TContactPoint<T> SphereBoxContactPoint(const TSphere<T, d>& Sphere, const TRigidTransform<T, d>& SphereTransform, const TAABB<T, d>& Box, const TRigidTransform<T, d>& BoxTransform, const T CullDistance)
+
+		FContactPoint SphereBoxContactPoint(const TSphere<FReal, 3>& Sphere, const FRigidTransform3& SphereTransform, const FAABB3& Box, const FRigidTransform3& BoxTransform, const FReal CullDistance)
 		{
-			TContactPoint<T> Result;
+			FContactPoint Result;
 
-			const TRigidTransform<T, d> SphereToBoxTransform(SphereTransform * BoxTransform.Inverse());	//todo: this should use GetRelative
-			const TVector<T, d> SphereCenterInBox = SphereToBoxTransform.TransformPosition(Sphere.GetCenter());
+			const FRigidTransform3 SphereToBoxTransform(SphereTransform * BoxTransform.Inverse());	//todo: this should use GetRelative
+			const FVec3 SphereCenterInBox = SphereToBoxTransform.TransformPosition(Sphere.GetCenter());
 
-			TVector<T, d> NewNormal;
-			T NewPhi = Box.PhiWithNormal(SphereCenterInBox, NewNormal);
+			FVec3 NewNormal;
+			FReal NewPhi = Box.PhiWithNormal(SphereCenterInBox, NewNormal);
 			NewPhi -= Sphere.GetRadius();
 
 			Result.Phi = NewPhi;
@@ -966,32 +947,30 @@ namespace Chaos
 			return Result;
 		}
 
-		template <typename T, int d>
-		void UpdateSphereBoxConstraint(const TSphere<T, d>& Sphere, const TRigidTransform<T, d>& SphereTransform, const TAABB<T, d>& Box, const TRigidTransform<T, d>& BoxTransform, const T CullDistance, TRigidBodyPointContactConstraint<T, d>& Constraint)
+
+		void UpdateSphereBoxConstraint(const TSphere<FReal, 3>& Sphere, const FRigidTransform3& SphereTransform, const FAABB3& Box, const FRigidTransform3& BoxTransform, const FReal CullDistance, FRigidBodyPointContactConstraint& Constraint)
 		{
 			UpdateContactPoint(Constraint.Manifold, SphereBoxContactPoint(Sphere, SphereTransform, Box, BoxTransform, CullDistance));
 		}
 
-		template<class T, int d>
-		void UpdateSphereBoxManifold(TCollisionConstraintBase<T, d>&  Constraint, const TRigidTransform<T, d>& ATM, const TRigidTransform<T, d>& BTM, const T CullDistance)
+		void UpdateSphereBoxManifold(FCollisionConstraintBase&  Constraint, const FRigidTransform3& ATM, const FRigidTransform3& BTM, const FReal CullDistance)
 		{
 			// @todo(chaos) : Stub Update Manifold
 			//   Stub function for updating the manifold prior to the Apply and ApplyPushOut
 		}
 
-		template<typename T, int d>
-		void ConstructSphereBoxConstraints(TGeometryParticleHandle<T, d>* Particle0, TGeometryParticleHandle<T, d>* Particle1, const FImplicitObject* Implicit0, const FImplicitObject* Implicit1, const TRigidTransform<T, d>& LocalTransform0, const TRigidTransform<T, d>& LocalTransform1, const T CullDistance, FCollisionConstraintsArray& NewConstraints)
+		void ConstructSphereBoxConstraints(TGeometryParticleHandle<FReal, 3>* Particle0, TGeometryParticleHandle<FReal, 3>* Particle1, const FImplicitObject* Implicit0, const FImplicitObject* Implicit1, const FRigidTransform3& LocalTransform0, const FRigidTransform3& LocalTransform1, const FReal CullDistance, FCollisionConstraintsArray& NewConstraints)
 		{
 
-			const TSphere<T, d> * Object0 = Implicit0->template GetObject<const TSphere<T, d> >();
-			const TBox<T, d> * Object1 = Implicit1->template GetObject<const TBox<T, d> >();
+			const TSphere<FReal, 3>* Object0 = Implicit0->template GetObject<const TSphere<FReal, 3> >();
+			const TBox<FReal, 3> * Object1 = Implicit1->template GetObject<const TBox<FReal, 3> >();
 			if (ensure(Object0 && Object1))
 			{
 				FRigidBodyPointContactConstraint Constraint = FRigidBodyPointContactConstraint(Particle0, Implicit0, LocalTransform0, Particle1, Implicit1, LocalTransform1, EContactShapesType::SphereBox);
 				if (!bChaos_Collision_DeferNarrowPhase)
 				{
-					TRigidTransform<T, d> WorldTransform0 = LocalTransform0 * Collisions::GetTransform(Particle0);
-					TRigidTransform<T, d> WorldTransform1 = LocalTransform1 * Collisions::GetTransform(Particle1);
+					FRigidTransform3 WorldTransform0 = LocalTransform0 * Collisions::GetTransform(Particle0);
+					FRigidTransform3 WorldTransform1 = LocalTransform1 * Collisions::GetTransform(Particle1);
 					UpdateSphereBoxConstraint(*Object0, WorldTransform0, Object1->BoundingBox(), WorldTransform1, CullDistance, Constraint);
 					NewConstraints.TryAdd(CullDistance, Constraint);
 				}
@@ -1007,22 +986,21 @@ namespace Chaos
 		// Sphere - Capsule
 		//
 
-		template <typename T, int d>
-		TContactPoint<T> SphereCapsuleContactPoint(const TSphere<T, d>& A, const TRigidTransform<T, d>& ATransform, const TCapsule<T>& B, const TRigidTransform<T, d>& BTransform, const T CullDistance)
+		FContactPoint SphereCapsuleContactPoint(const TSphere<FReal, 3>& A, const FRigidTransform3& ATransform, const TCapsule<FReal>& B, const FRigidTransform3& BTransform, const FReal CullDistance)
 		{
-			TContactPoint<T> Result;
+			FContactPoint Result;
 
 			FVector A1 = ATransform.TransformPosition(A.GetCenter());
 			FVector B1 = BTransform.TransformPosition(B.GetX1());
 			FVector B2 = BTransform.TransformPosition(B.GetX2());
 			FVector P2 = FMath::ClosestPointOnSegment(A1, B1, B2);
 
-			TVector<T, d> Delta = P2 - A1;
-			T DeltaLen = Delta.Size();
+			FVec3 Delta = P2 - A1;
+			FReal DeltaLen = Delta.Size();
 			if (DeltaLen > KINDA_SMALL_NUMBER)
 			{
-				T NewPhi = DeltaLen - (A.GetRadius() + B.GetRadius());
-				TVector<T, d> Dir = Delta / DeltaLen;
+				FReal NewPhi = DeltaLen - (A.GetRadius() + B.GetRadius());
+				FVec3 Dir = Delta / DeltaLen;
 				Result.Phi = NewPhi;
 				Result.Normal = -Dir;
 				Result.Location = A1 + Dir * A.GetRadius();
@@ -1031,31 +1009,29 @@ namespace Chaos
 			return Result;
 		}
 
-		template <typename T, int d>
-		void UpdateSphereCapsuleConstraint(const TSphere<T, d>& A, const TRigidTransform<T, d>& ATransform, const TCapsule<T>& B, const TRigidTransform<T, d>& BTransform, const T CullDistance, TRigidBodyPointContactConstraint<T, d>& Constraint)
+
+		void UpdateSphereCapsuleConstraint(const TSphere<FReal, 3>& A, const FRigidTransform3& ATransform, const TCapsule<FReal>& B, const FRigidTransform3& BTransform, const FReal CullDistance, FRigidBodyPointContactConstraint& Constraint)
 		{
 			UpdateContactPoint(Constraint.Manifold, SphereCapsuleContactPoint(A, ATransform, B, BTransform, CullDistance));
 		}
 
-		template<class T, int d>
-		void UpdateSphereCapsuleManifold(TCollisionConstraintBase<T, d>&  Constraint, const TRigidTransform<T, d>& ATM, const TRigidTransform<T, d>& BTM, const T CullDistance)
+		void UpdateSphereCapsuleManifold(FCollisionConstraintBase&  Constraint, const FRigidTransform3& ATM, const FRigidTransform3& BTM, const FReal CullDistance)
 		{
 			// @todo(chaos) : Stub Update Manifold
 			//   Stub function for updating the manifold prior to the Apply and ApplyPushOut
 		}
 
-		template<typename T, int d>
-		void ConstructSphereCapsuleConstraints(TGeometryParticleHandle<T, d>* Particle0, TGeometryParticleHandle<T, d>* Particle1, const FImplicitObject* Implicit0, const FImplicitObject* Implicit1, const TRigidTransform<T, d>& LocalTransform0, const TRigidTransform<T, d>& LocalTransform1, const T CullDistance, FCollisionConstraintsArray& NewConstraints)
+		void ConstructSphereCapsuleConstraints(TGeometryParticleHandle<FReal, 3>* Particle0, TGeometryParticleHandle<FReal, 3>* Particle1, const FImplicitObject* Implicit0, const FImplicitObject* Implicit1, const FRigidTransform3& LocalTransform0, const FRigidTransform3& LocalTransform1, const FReal CullDistance, FCollisionConstraintsArray& NewConstraints)
 		{
-			const TSphere<T, d>* Object0 = Implicit0->template GetObject<const TSphere<T, d> >();
-			const TCapsule<T>* Object1 = Implicit1->template GetObject<const TCapsule<T> >();
+			const TSphere<FReal, 3>* Object0 = Implicit0->template GetObject<const TSphere<FReal, 3> >();
+			const TCapsule<FReal>* Object1 = Implicit1->template GetObject<const TCapsule<FReal> >();
 			if (ensure(Object0 && Object1))
 			{
 				FRigidBodyPointContactConstraint Constraint = FRigidBodyPointContactConstraint(Particle0, Implicit0, LocalTransform0, Particle1, Implicit1, LocalTransform1, EContactShapesType::SphereCapsule);
 				if (!bChaos_Collision_DeferNarrowPhase)
 				{
-					TRigidTransform<T, d> WorldTransform0 = LocalTransform0 * Collisions::GetTransform(Particle0);
-					TRigidTransform<T, d> WorldTransform1 = LocalTransform1 * Collisions::GetTransform(Particle1);
+					FRigidTransform3 WorldTransform0 = LocalTransform0 * Collisions::GetTransform(Particle0);
+					FRigidTransform3 WorldTransform1 = LocalTransform1 * Collisions::GetTransform(Particle1);
 					UpdateSphereCapsuleConstraint(*Object0, WorldTransform0, *Object1, WorldTransform1, CullDistance, Constraint);
 					NewConstraints.TryAdd(CullDistance, Constraint);
 				}
@@ -1070,29 +1046,28 @@ namespace Chaos
 		// Sphere-TriangleMesh
 		//
 
-		template <typename T, int d, typename TriMeshType>
-		TContactPoint<T> SphereTriangleMeshContactPoint(const TSphere<T, d>& A, const TRigidTransform<T, d>& ATransform, const TriMeshType& B, const TRigidTransform<T, d>& BTransform, const T CullDistance)
+		template <typename TriMeshType>
+		FContactPoint SphereTriangleMeshContactPoint(const TSphere<FReal, 3>& A, const FRigidTransform3& ATransform, const TriMeshType& B, const FRigidTransform3& BTransform, const FReal CullDistance)
 		{
 			return GJKImplicitContactPoint< TSphere<float, 3> >(TSphere<float, 3>(A), ATransform, B, BTransform, CullDistance);
 		}
 
-		template <typename T, int d, typename TriMeshType>
-		void UpdateSphereTriangleMeshConstraint(const TSphere<T, d>& Sphere0, const TRigidTransform<T, d>& WorldTransform0, const TriMeshType& TriangleMesh1, const TRigidTransform<T, d>& WorldTransform1, const T CullDistance, TRigidBodyPointContactConstraint<T, d>& Constraint)
+		template <typename TriMeshType>
+		void UpdateSphereTriangleMeshConstraint(const TSphere<FReal, 3>& Sphere0, const FRigidTransform3& WorldTransform0, const TriMeshType& TriangleMesh1, const FRigidTransform3& WorldTransform1, const FReal CullDistance, FRigidBodyPointContactConstraint& Constraint)
 		{
 			UpdateContactPoint(Constraint.Manifold, SphereTriangleMeshContactPoint(Sphere0, WorldTransform0, TriangleMesh1, WorldTransform1, CullDistance));
 		}
 
 
-		template <typename T, int d>
-		void UpdateSphereTriangleMeshManifold(TCollisionConstraintBase<T, d>&  Constraint, const TRigidTransform<T, d>& ATM, const TRigidTransform<T, d>& BTM, const T CullDistance)
+
+		void UpdateSphereTriangleMeshManifold(FCollisionConstraintBase&  Constraint, const FRigidTransform3& ATM, const FRigidTransform3& BTM, const FReal CullDistance)
 		{
 
 		}
 
-		template<typename T, int d>
-		void ConstructSphereTriangleMeshConstraints(TGeometryParticleHandle<T, d>* Particle0, TGeometryParticleHandle<T, d>* Particle1, const FImplicitObject* Implicit0, const FImplicitObject* Implicit1, const TRigidTransform<T, d>& LocalTransform0, const TRigidTransform<T, d>& LocalTransform1, const T CullDistance, FCollisionConstraintsArray& NewConstraints)
+		void ConstructSphereTriangleMeshConstraints(TGeometryParticleHandle<FReal, 3>* Particle0, TGeometryParticleHandle<FReal, 3>* Particle1, const FImplicitObject* Implicit0, const FImplicitObject* Implicit1, const FRigidTransform3& LocalTransform0, const FRigidTransform3& LocalTransform1, const FReal CullDistance, FCollisionConstraintsArray& NewConstraints)
 		{
-			const TSphere<T, d>* Object0 = Implicit0->template GetObject<const TSphere<T, d> >();
+			const TSphere<FReal, 3>* Object0 = Implicit0->template GetObject<const TSphere<FReal, 3> >();
 			if (ensure(Object0))
 			{
 				if (const TImplicitObjectScaled<FTriangleMeshImplicitObject>* ScaledTriangleMesh = Implicit1->template GetObject<const TImplicitObjectScaled<FTriangleMeshImplicitObject>>())
@@ -1100,8 +1075,8 @@ namespace Chaos
 					FRigidBodyPointContactConstraint Constraint = FRigidBodyPointContactConstraint(Particle0, Implicit0, LocalTransform0, Particle1, Implicit1, LocalTransform1, EContactShapesType::SphereTriMesh);
 					if (!bChaos_Collision_DeferNarrowPhase)
 					{
-						TRigidTransform<T, d> WorldTransform0 = LocalTransform0 * Collisions::GetTransform(Particle0);
-						TRigidTransform<T, d> WorldTransform1 = LocalTransform1 * Collisions::GetTransform(Particle1);
+						FRigidTransform3 WorldTransform0 = LocalTransform0 * Collisions::GetTransform(Particle0);
+						FRigidTransform3 WorldTransform1 = LocalTransform1 * Collisions::GetTransform(Particle1);
 						UpdateSphereTriangleMeshConstraint(*Object0, WorldTransform0, *ScaledTriangleMesh, WorldTransform1, CullDistance, Constraint);
 						NewConstraints.TryAdd(CullDistance, Constraint);
 					}
@@ -1115,8 +1090,8 @@ namespace Chaos
 					FRigidBodyPointContactConstraint Constraint = FRigidBodyPointContactConstraint(Particle0, Implicit0, LocalTransform0, Particle1, Implicit1, LocalTransform1, EContactShapesType::SphereTriMesh);
 					if (!bChaos_Collision_DeferNarrowPhase)
 					{
-						TRigidTransform<T, d> WorldTransform0 = LocalTransform0 * Collisions::GetTransform(Particle0);
-						TRigidTransform<T, d> WorldTransform1 = LocalTransform1 * Collisions::GetTransform(Particle1);
+						FRigidTransform3 WorldTransform0 = LocalTransform0 * Collisions::GetTransform(Particle0);
+						FRigidTransform3 WorldTransform1 = LocalTransform1 * Collisions::GetTransform(Particle1);
 						UpdateSphereTriangleMeshConstraint(*Object0, WorldTransform0, *TriangleMesh, WorldTransform1, CullDistance, Constraint);
 						NewConstraints.TryAdd(CullDistance, Constraint);
 					}
@@ -1138,10 +1113,10 @@ namespace Chaos
 		// Capsule-Capsule
 		//
 
-		template <typename T, int d>
-		TContactPoint<T> CapsuleCapsuleContactPoint(const TCapsule<T>& A, const TRigidTransform<T, d>& ATransform, const TCapsule<T>& B, const TRigidTransform<T, d>& BTransform, const T CullDistance)
+
+		FContactPoint CapsuleCapsuleContactPoint(const TCapsule<FReal>& A, const FRigidTransform3& ATransform, const TCapsule<FReal>& B, const FRigidTransform3& BTransform, const FReal CullDistance)
 		{
-			TContactPoint<T> Result;
+			FContactPoint Result;
 
 			FVector A1 = ATransform.TransformPosition(A.GetX1());
 			FVector A2 = ATransform.TransformPosition(A.GetX2());
@@ -1150,12 +1125,12 @@ namespace Chaos
 			FVector P1, P2;
 			FMath::SegmentDistToSegmentSafe(A1, A2, B1, B2, P1, P2);
 
-			TVector<T, d> Delta = P2 - P1;
-			T DeltaLen = Delta.Size();
+			FVec3 Delta = P2 - P1;
+			FReal DeltaLen = Delta.Size();
 			if (DeltaLen > KINDA_SMALL_NUMBER)
 			{
-				T NewPhi = DeltaLen - (A.GetRadius() + B.GetRadius());
-				TVector<T, d> Dir = Delta / DeltaLen;
+				FReal NewPhi = DeltaLen - (A.GetRadius() + B.GetRadius());
+				FVec3 Dir = Delta / DeltaLen;
 				Result.Phi = NewPhi;
 				Result.Normal = -Dir;
 				Result.Location = P1 + Dir * A.GetRadius();
@@ -1164,31 +1139,29 @@ namespace Chaos
 			return Result;
 		}
 
-		template <typename T, int d>
-		void UpdateCapsuleCapsuleConstraint(const TCapsule<T>& A, const TRigidTransform<T, d>& ATransform, const TCapsule<T>& B, const TRigidTransform<T, d>& BTransform, const T CullDistance, TRigidBodyPointContactConstraint<T, d>& Constraint)
+
+		void UpdateCapsuleCapsuleConstraint(const TCapsule<FReal>& A, const FRigidTransform3& ATransform, const TCapsule<FReal>& B, const FRigidTransform3& BTransform, const FReal CullDistance, FRigidBodyPointContactConstraint& Constraint)
 		{
 			UpdateContactPoint(Constraint.Manifold, CapsuleCapsuleContactPoint(A, ATransform, B, BTransform, CullDistance));
 		}
 
-		template<class T, int d>
-		void UpdateCapsuleCapsuleManifold(TCollisionConstraintBase<T, d>&  Constraint, const TRigidTransform<T, d>& ATM, const TRigidTransform<T, d>& BTM, const T CullDistance)
+		void UpdateCapsuleCapsuleManifold(FCollisionConstraintBase&  Constraint, const FRigidTransform3& ATM, const FRigidTransform3& BTM, const FReal CullDistance)
 		{
 			// @todo(chaos) : Stub Update Manifold
 			//   Stub function for updating the manifold prior to the Apply and ApplyPushOut
 		}
 
-		template<typename T, int d>
-		void ConstructCapsuleCapsuleConstraints(TGeometryParticleHandle<T, d>* Particle0, TGeometryParticleHandle<T, d>* Particle1, const FImplicitObject* Implicit0, const FImplicitObject* Implicit1, const TRigidTransform<T, d>& LocalTransform0, const TRigidTransform<T, d>& LocalTransform1, const T CullDistance, FCollisionConstraintsArray& NewConstraints)
+		void ConstructCapsuleCapsuleConstraints(TGeometryParticleHandle<FReal, 3>* Particle0, TGeometryParticleHandle<FReal, 3>* Particle1, const FImplicitObject* Implicit0, const FImplicitObject* Implicit1, const FRigidTransform3& LocalTransform0, const FRigidTransform3& LocalTransform1, const FReal CullDistance, FCollisionConstraintsArray& NewConstraints)
 		{
-			const TCapsule<T> * Object0 = Implicit0->template GetObject<const TCapsule<T> >();
-			const TCapsule<T> * Object1 = Implicit1->template GetObject<const TCapsule<T> >();
+			const TCapsule<FReal> * Object0 = Implicit0->template GetObject<const TCapsule<FReal> >();
+			const TCapsule<FReal> * Object1 = Implicit1->template GetObject<const TCapsule<FReal> >();
 			if (ensure(Object0 && Object1))
 			{
 				FRigidBodyPointContactConstraint Constraint = FRigidBodyPointContactConstraint(Particle0, Implicit0, LocalTransform0, Particle1, Implicit1, LocalTransform1, EContactShapesType::CapsuleCapsule);
 				if (!bChaos_Collision_DeferNarrowPhase)
 				{
-					TRigidTransform<T, d> WorldTransform0 = LocalTransform0 * Collisions::GetTransform(Particle0);
-					TRigidTransform<T, d> WorldTransform1 = LocalTransform1 * Collisions::GetTransform(Particle1);
+					FRigidTransform3 WorldTransform0 = LocalTransform0 * Collisions::GetTransform(Particle0);
+					FRigidTransform3 WorldTransform1 = LocalTransform1 * Collisions::GetTransform(Particle1);
 					UpdateCapsuleCapsuleConstraint(*Object0, WorldTransform0, *Object1, WorldTransform1, CullDistance, Constraint);
 					NewConstraints.TryAdd(CullDistance, Constraint);
 				}
@@ -1203,24 +1176,24 @@ namespace Chaos
 		// Capsule - Box
 		//
 
-		template <typename T, int d>
-		TContactPoint<T> CapsuleBoxContactPoint(const TCapsule<T>& A, const TRigidTransform<T, d>& ATransform, const TAABB<T, d>& B, const TRigidTransform<T, d>& BTransform, const TVector<T, 3>& InitialDir, const T CullDistance)
+
+		FContactPoint CapsuleBoxContactPoint(const TCapsule<FReal>& A, const FRigidTransform3& ATransform, const FAABB3& B, const FRigidTransform3& BTransform, const FVec3& InitialDir, const FReal CullDistance)
 		{
 			return GJKContactPoint(A, ATransform, B, BTransform, InitialDir);
 		}
 
-		template <typename T, int d>
-		void UpdateCapsuleBoxConstraint(const TCapsule<T>& A, const TRigidTransform<T, d>& ATransform, const TAABB<T, d>& B, const TRigidTransform<T, d>& BTransform, const T CullDistance, TRigidBodyPointContactConstraint<T, d>& Constraint)
+
+		void UpdateCapsuleBoxConstraint(const TCapsule<FReal>& A, const FRigidTransform3& ATransform, const FAABB3& B, const FRigidTransform3& BTransform, const FReal CullDistance, FRigidBodyPointContactConstraint& Constraint)
 		{
 			// @todo(ccaulfield): the inverse of CapsuleToBoxTM is calculated in GJKContactPoint - try to eliminate this one
 			const FRigidTransform3 CapsuleToBoxTM = ATransform.GetRelativeTransform(BTransform);
 			const FVec3 P1 = CapsuleToBoxTM.TransformPosition(A.GetX1());
 			const FVec3 P2 = CapsuleToBoxTM.TransformPosition(A.GetX2());
-			TAABB<T, 3> CapsuleAABB(P1.ComponentMin(P2), P1.ComponentMax(P2));
+			FAABB3 CapsuleAABB(P1.ComponentMin(P2), P1.ComponentMax(P2));
 			CapsuleAABB.Thicken(A.GetRadius() + CullDistance);
 			if (CapsuleAABB.Intersects(B))
 			{
-				const TVector<T, d> InitialDir = ATransform.GetRotation().Inverse() * -Constraint.GetNormal();
+				const FVec3 InitialDir = ATransform.GetRotation().Inverse() * -Constraint.GetNormal();
 				UpdateContactPoint(Constraint.Manifold, CapsuleBoxContactPoint(A, ATransform, B, BTransform, InitialDir, CullDistance));
 			}
 		}
@@ -1237,23 +1210,22 @@ namespace Chaos
 		 * For box edge with capsule vertex collisions, only the single near point/plane is used in the manifold.
 		 * 
 		 */
-		template<class T, int d>
-		void UpdateCapsuleBoxManifold(const TCapsule<T>& Capsule, const TRigidTransform<T, d>& CapsuleTM, const TAABB<T, d>& Box, const TRigidTransform<T, d>& BoxTM, const T CullDistance, const FCollisionContext& Context, FRigidBodyMultiPointContactConstraint& Constraint)
+		void UpdateCapsuleBoxManifold(const TCapsule<FReal>& Capsule, const FRigidTransform3& CapsuleTM, const FAABB3& Box, const FRigidTransform3& BoxTM, const FReal CullDistance, const FCollisionContext& Context, FRigidBodyMultiPointContactConstraint& Constraint)
 		{
 			Constraint.ResetManifoldPoints();
 
 			// Find the nearest points on the capsule and box
 			// Note: We flip the order for GJK so we get the normal in box space. This makes it easier to build the face-capsule manifold.
-			const TRigidTransform<T, d> CapsuleToBoxTM = CapsuleTM.GetRelativeTransform(BoxTM);
+			const FRigidTransform3 CapsuleToBoxTM = CapsuleTM.GetRelativeTransform(BoxTM);
 
 			// NOTE: All GJK results in box-space
 			// @todo(ccaulfield): use center-to-center direction for InitialDir
 			FVec3 InitialDir = FVec3(1, 0, 0);
-			T Penetration;
+			FReal Penetration;
 			FVec3 CapsuleClosestBoxSpace, BoxClosestBoxSpace, NormalBoxSpace;
 			{
 				SCOPE_CYCLE_COUNTER_GJK();
-				if (!ensure(GJKPenetration<true>(Box, Capsule, CapsuleToBoxTM, Penetration, BoxClosestBoxSpace, CapsuleClosestBoxSpace, NormalBoxSpace, (T)0, InitialDir, (T)0)))
+				if (!ensure(GJKPenetration<true>(Box, Capsule, CapsuleToBoxTM, Penetration, BoxClosestBoxSpace, CapsuleClosestBoxSpace, NormalBoxSpace, (FReal)0, InitialDir, (FReal)0)))
 				{
 					return;
 				}
@@ -1423,11 +1395,10 @@ namespace Chaos
 			ensure(false);
 		}
 
-		template<typename T, int d>
-		void ConstructCapsuleBoxConstraints(TGeometryParticleHandle<T, d>* Particle0, TGeometryParticleHandle<T, d>* Particle1, const FImplicitObject* Implicit0, const FImplicitObject* Implicit1, const TRigidTransform<T, d>& LocalTransform0, const TRigidTransform<T, d>& LocalTransform1, const T CullDistance, const FCollisionContext& Context, FCollisionConstraintsArray& NewConstraints)
+		void ConstructCapsuleBoxConstraints(TGeometryParticleHandle<FReal, 3>* Particle0, TGeometryParticleHandle<FReal, 3>* Particle1, const FImplicitObject* Implicit0, const FImplicitObject* Implicit1, const FRigidTransform3& LocalTransform0, const FRigidTransform3& LocalTransform1, const FReal CullDistance, const FCollisionContext& Context, FCollisionConstraintsArray& NewConstraints)
 		{
-			const TCapsule<T> * Object0 = Implicit0->template GetObject<const TCapsule<T>>();
-			const TBox<T, d> * Object1 = Implicit1->template GetObject<const TBox<T, d>>();
+			const TCapsule<FReal> * Object0 = Implicit0->template GetObject<const TCapsule<FReal>>();
+			const TBox<FReal, 3> * Object1 = Implicit1->template GetObject<const TBox<FReal, 3>>();
 			if (ensure(Object0 && Object1))
 			{
 				bool bAllowManifold = Chaos_Collision_UseManifolds;
@@ -1437,7 +1408,7 @@ namespace Chaos
 					// HACK: Disable manifolds for "horizontal" capsules. Manifolds don't work well when joints are pulling boxes down
 					// (under gravity) when the upper boxes are draped over a horizontal capsule. The box rotations about the manifold
 					// points(line) is too great and we end up with jitter.
-					TRigidTransform<T, d> WorldTransform0 = LocalTransform0 * Collisions::GetTransform(Particle0);
+					FRigidTransform3 WorldTransform0 = LocalTransform0 * Collisions::GetTransform(Particle0);
 					const FVector CapsuleAxis = Context.SpaceTransform.TransformVectorNoScale(WorldTransform0.TransformVectorNoScale(Object0->GetAxis()));
 					bAllowManifold = (FMath::Abs(CapsuleAxis.Z) > Chaos_Collision_CapsuleBoxManifoldTolerance);
 				}
@@ -1446,8 +1417,8 @@ namespace Chaos
 				{
 					FRigidBodyMultiPointContactConstraint Constraint = FRigidBodyMultiPointContactConstraint(Particle0, Implicit0, LocalTransform0, Particle1, Implicit1, LocalTransform1, EContactShapesType::CapsuleBox);
 					// @todo(ccaulfield): manifold creation should be deferrable same as non-manifold below (only Context is preventing this - Apply does not know about the context yet)
-					TRigidTransform<T, d> WorldTransform0 = LocalTransform0 * Collisions::GetTransform(Particle0);
-					TRigidTransform<T, d> WorldTransform1 = LocalTransform1 * Collisions::GetTransform(Particle1);
+					FRigidTransform3 WorldTransform0 = LocalTransform0 * Collisions::GetTransform(Particle0);
+					FRigidTransform3 WorldTransform1 = LocalTransform1 * Collisions::GetTransform(Particle1);
 					UpdateCapsuleBoxManifold(*Object0, WorldTransform0, Object1->BoundingBox(), WorldTransform1, CullDistance, Context, Constraint);
 					NewConstraints.TryAdd(CullDistance, Constraint);
 				}
@@ -1456,8 +1427,8 @@ namespace Chaos
 					FRigidBodyPointContactConstraint Constraint = FRigidBodyPointContactConstraint(Particle0, Implicit0, LocalTransform0, Particle1, Implicit1, LocalTransform1, EContactShapesType::CapsuleBox);
 					if (!bChaos_Collision_DeferNarrowPhase)
 					{
-						TRigidTransform<T, d> WorldTransform0 = LocalTransform0 * Collisions::GetTransform(Particle0);
-						TRigidTransform<T, d> WorldTransform1 = LocalTransform1 * Collisions::GetTransform(Particle1);
+						FRigidTransform3 WorldTransform0 = LocalTransform0 * Collisions::GetTransform(Particle0);
+						FRigidTransform3 WorldTransform1 = LocalTransform1 * Collisions::GetTransform(Particle1);
 						UpdateCapsuleBoxConstraint(*Object0, WorldTransform0, Object1->BoundingBox(), WorldTransform1, CullDistance, Constraint);
 						NewConstraints.TryAdd(CullDistance, Constraint);
 					}
@@ -1473,46 +1444,44 @@ namespace Chaos
 		// Capsule-HeightField
 		//
 
-		template <typename T, int d>
-		TContactPoint<T> CapsuleHeightFieldContactPoint(const TCapsule<T>& A, const TRigidTransform<T, d>& ATransform, const THeightField<T>& B, const TRigidTransform<T, d>& BTransform, const T CullDistance)
+
+		FContactPoint CapsuleHeightFieldContactPoint(const TCapsule<FReal>& A, const FRigidTransform3& ATransform, const THeightField<FReal>& B, const FRigidTransform3& BTransform, const FReal CullDistance)
 		{
 			return GJKImplicitContactPoint< TCapsule<float> >(A, ATransform, B, BTransform, CullDistance);
 		}
 
-		template <typename T, int d>
-		void UpdateCapsuleHeightFieldConstraint(const TCapsule<T>& A, const TRigidTransform<T, d>& ATransform, const THeightField<T>& B, const TRigidTransform<T, d>& BTransform, const T CullDistance, TRigidBodyPointContactConstraint<T, d>& Constraint)
+
+		void UpdateCapsuleHeightFieldConstraint(const TCapsule<FReal>& A, const FRigidTransform3& ATransform, const THeightField<FReal>& B, const FRigidTransform3& BTransform, const FReal CullDistance, FRigidBodyPointContactConstraint& Constraint)
 		{
 			UpdateContactPoint(Constraint.Manifold, CapsuleHeightFieldContactPoint(A, ATransform, B, BTransform, CullDistance));
 		}
 
-		template <typename T, int d>
-		void UpdateCapsuleHeightFieldConstraintSwept(TGeometryParticleHandle<T, d>* Particle0, const TCapsule<T>& A, const TRigidTransform<T, d>& ATransform, const THeightField<T>& B, const TRigidTransform<T, d>& BTransform, const TVector<T, d>& Dir, const T Length, const T CullDistance, TRigidBodySweptPointContactConstraint<T, d>& Constraint)
+
+		void UpdateCapsuleHeightFieldConstraintSwept(TGeometryParticleHandle<FReal, 3>* Particle0, const TCapsule<FReal>& A, const FRigidTransform3& ATransform, const THeightField<FReal>& B, const FRigidTransform3& BTransform, const FVec3& Dir, const FReal Length, const FReal CullDistance, FRigidBodySweptPointContactConstraint& Constraint)
 		{
-			T TOI = 1.0f;
+			FReal TOI = 1.0f;
 			UpdateContactPoint(Constraint.Manifold, GJKImplicitSweptContactPoint<TCapsule<float> >(A, ATransform, B, BTransform, Dir, Length, CullDistance, TOI));
 			SetSweptConstraintTOI(Particle0, TOI, Length, Dir, Constraint);
 		}
 
-		template<class T, int d>
-		void UpdateCapsuleHeightFieldManifold(TCollisionConstraintBase<T, d>&  Constraint, const TRigidTransform<T, d>& ATM, const TRigidTransform<T, d>& BTM, const T CullDistance)
+		void UpdateCapsuleHeightFieldManifold(FCollisionConstraintBase&  Constraint, const FRigidTransform3& ATM, const FRigidTransform3& BTM, const FReal CullDistance)
 		{
 			// @todo(chaos) : Stub Update Manifold
 			//   Stub function for updating the manifold prior to the Apply and ApplyPushOut
 		}
 
-		template<typename T, int d>
-		void ConstructCapsuleHeightFieldConstraints(TGeometryParticleHandle<T, d>* Particle0, TGeometryParticleHandle<T, d>* Particle1, const FImplicitObject* Implicit0, const FImplicitObject* Implicit1, const TRigidTransform<T, d>& LocalTransform0, const TRigidTransform<T, d>& LocalTransform1, const T CullDistance, FCollisionConstraintsArray& NewConstraints)
+		void ConstructCapsuleHeightFieldConstraints(TGeometryParticleHandle<FReal, 3>* Particle0, TGeometryParticleHandle<FReal, 3>* Particle1, const FImplicitObject* Implicit0, const FImplicitObject* Implicit1, const FRigidTransform3& LocalTransform0, const FRigidTransform3& LocalTransform1, const FReal CullDistance, FCollisionConstraintsArray& NewConstraints)
 		{
 
-			const TCapsule<T> * Object0 = Implicit0->template GetObject<const TCapsule<T> >();
-			const THeightField<T> * Object1 = Implicit1->template GetObject<const THeightField<T> >();
+			const TCapsule<FReal> * Object0 = Implicit0->template GetObject<const TCapsule<FReal> >();
+			const THeightField<FReal> * Object1 = Implicit1->template GetObject<const THeightField<FReal> >();
 			if (ensure(Object0 && Object1))
 			{
 				FRigidBodyPointContactConstraint Constraint = FRigidBodyPointContactConstraint(Particle0, Implicit0, LocalTransform0, Particle1, Implicit1, LocalTransform1, EContactShapesType::CapsuleHeightField);
 				if (!bChaos_Collision_DeferNarrowPhase)
 				{
-					TRigidTransform<T, d> WorldTransform0 = LocalTransform0 * Collisions::GetTransform(Particle0);
-					TRigidTransform<T, d> WorldTransform1 = LocalTransform1 * Collisions::GetTransform(Particle1);
+					FRigidTransform3 WorldTransform0 = LocalTransform0 * Collisions::GetTransform(Particle0);
+					FRigidTransform3 WorldTransform1 = LocalTransform1 * Collisions::GetTransform(Particle1);
 					UpdateCapsuleHeightFieldConstraint(*Object0, WorldTransform0, *Object1, WorldTransform1, CullDistance, Constraint);
 					NewConstraints.TryAdd(CullDistance, Constraint);
 				}
@@ -1523,16 +1492,15 @@ namespace Chaos
 			}
 		}
 
-		template<typename T, int d>
-		void ConstructCapsuleHeightFieldConstraintsSwept(TGeometryParticleHandle<T, d>* Particle0, TGeometryParticleHandle<T, d>* Particle1, const FImplicitObject* Implicit0, const FImplicitObject* Implicit1, const TRigidTransform<T, d>& LocalTransform0, const TRigidTransform<T, d>& LocalTransform1, const T CullDistance, const TVector<T,d>& Dir, FReal Length, FCollisionConstraintsArray& NewConstraints)
+		void ConstructCapsuleHeightFieldConstraintsSwept(TGeometryParticleHandle<FReal, 3>* Particle0, TGeometryParticleHandle<FReal, 3>* Particle1, const FImplicitObject* Implicit0, const FImplicitObject* Implicit1, const FRigidTransform3& LocalTransform0, const FRigidTransform3& LocalTransform1, const FReal CullDistance, const FVec3& Dir, FReal Length, FCollisionConstraintsArray& NewConstraints)
 		{
-			const TCapsule<T> * Object0 = Implicit0->template GetObject<const TCapsule<T> >();
-			const THeightField<T> * Object1 = Implicit1->template GetObject<const THeightField<T> >();
+			const TCapsule<FReal> * Object0 = Implicit0->template GetObject<const TCapsule<FReal> >();
+			const THeightField<FReal> * Object1 = Implicit1->template GetObject<const THeightField<FReal> >();
 			if (ensure(Object0 && Object1))
 			{
 				FRigidBodySweptPointContactConstraint Constraint = FRigidBodySweptPointContactConstraint(Particle0, Implicit0, LocalTransform0, Particle1, Implicit1, LocalTransform1, EContactShapesType::CapsuleHeightField);
-				TRigidTransform<T, 3> WorldTransformX0 = LocalTransform0 * FRigidTransform3(Particle0->X(), Particle0->R());
-				TRigidTransform<T, d> WorldTransform1 = LocalTransform1 * Collisions::GetTransform(Particle1);
+				FRigidTransform3 WorldTransformX0 = LocalTransform0 * FRigidTransform3(Particle0->X(), Particle0->R());
+				FRigidTransform3 WorldTransform1 = LocalTransform1 * Collisions::GetTransform(Particle1);
 				UpdateCapsuleHeightFieldConstraintSwept(Particle0, *Object0, WorldTransformX0, *Object1, WorldTransform1, Dir, Length, CullDistance, Constraint);
 				NewConstraints.TryAdd(CullDistance, Constraint);
 			}
@@ -1543,53 +1511,52 @@ namespace Chaos
 		// Capsule-TriangleMesh
 		//
 
-		template <typename TriMeshType, typename T, int d>
-		TContactPoint<T> CapsuleTriangleMeshContactPoint(const TCapsule<T>& A, const TRigidTransform<T, d>& ATransform, const TriMeshType& B, const TRigidTransform<T, d>& BTransform, const T CullDistance)
+		template <typename TriMeshType>
+		FContactPoint CapsuleTriangleMeshContactPoint(const TCapsule<FReal>& A, const FRigidTransform3& ATransform, const TriMeshType& B, const FRigidTransform3& BTransform, const FReal CullDistance)
 		{
-			return GJKImplicitContactPoint< TCapsule<T> >(A, ATransform, B, BTransform, CullDistance);
+			return GJKImplicitContactPoint< TCapsule<FReal> >(A, ATransform, B, BTransform, CullDistance);
 		}
 
-		template <typename TriMeshType, typename T, int d>
-		TContactPoint<T> CapsuleTriangleMeshSweptContactPoint(const TCapsule<T>& A, const TRigidTransform<T, d>& ATransform, const TriMeshType& B, const TRigidTransform<T, d>& BStartTransform, const TVector<T, d>& Dir, const T Length, const T CullDistance, T& TOI)
+		template <typename TriMeshType>
+		FContactPoint CapsuleTriangleMeshSweptContactPoint(const TCapsule<FReal>& A, const FRigidTransform3& ATransform, const TriMeshType& B, const FRigidTransform3& BStartTransform, const FVec3& Dir, const FReal Length, const FReal CullDistance, FReal& TOI)
 		{
 			if (const TImplicitObjectScaled<FTriangleMeshImplicitObject>* ScaledTriangleMesh = B.template GetObject<const TImplicitObjectScaled<FTriangleMeshImplicitObject>>())
 			{
-				return GJKImplicitScaledTriMeshSweptContactPoint<TCapsule<T>>(A, ATransform, *ScaledTriangleMesh, BStartTransform, Dir, Length, CullDistance, TOI);
+				return GJKImplicitScaledTriMeshSweptContactPoint<TCapsule<FReal>>(A, ATransform, *ScaledTriangleMesh, BStartTransform, Dir, Length, CullDistance, TOI);
 			}
 			else if (const FTriangleMeshImplicitObject* TriangleMesh = B.template GetObject<const FTriangleMeshImplicitObject>())
 			{
-				return GJKImplicitSweptContactPoint<TCapsule<T>>(A, ATransform, *TriangleMesh, BStartTransform, Dir, Length, CullDistance, TOI);
+				return GJKImplicitSweptContactPoint<TCapsule<FReal>>(A, ATransform, *TriangleMesh, BStartTransform, Dir, Length, CullDistance, TOI);
 			}
 
 			ensure(false);
-			return TContactPoint<T>();
+			return FContactPoint();
 		}
 
 
-		template <typename TriMeshType, typename T, int d>
-		void UpdateCapsuleTriangleMeshConstraint(const TCapsule<T>& Capsule0, const TRigidTransform<T, d>& WorldTransform0, const TriMeshType& TriangleMesh1, const TRigidTransform<T, d>& WorldTransform1, const T CullDistance, TRigidBodyPointContactConstraint<T, d>& Constraint)
+		template <typename TriMeshType>
+		void UpdateCapsuleTriangleMeshConstraint(const TCapsule<FReal>& Capsule0, const FRigidTransform3& WorldTransform0, const TriMeshType& TriangleMesh1, const FRigidTransform3& WorldTransform1, const FReal CullDistance, FRigidBodyPointContactConstraint& Constraint)
 		{
 			UpdateContactPoint(Constraint.Manifold, CapsuleTriangleMeshContactPoint(Capsule0, WorldTransform0, TriangleMesh1, WorldTransform1, CullDistance));
 		}
 
-		template <typename TriMeshType, typename T, int d>
-		void UpdateCapsuleTriangleMeshConstraintSwept(TGeometryParticleHandle<T, d>* Particle0, const TCapsule<T>& Capsule0, const TRigidTransform<T, d>& WorldTransform0, const TriMeshType& TriangleMesh1, const TRigidTransform<T, d>& WorldTransform1, const TVector<T, d>& Dir, const T Length, const T CullDistance, TRigidBodySweptPointContactConstraint<T, d>& Constraint)
+		template <typename TriMeshType>
+		void UpdateCapsuleTriangleMeshConstraintSwept(TGeometryParticleHandle<FReal, 3>* Particle0, const TCapsule<FReal>& Capsule0, const FRigidTransform3& WorldTransform0, const TriMeshType& TriangleMesh1, const FRigidTransform3& WorldTransform1, const FVec3& Dir, const FReal Length, const FReal CullDistance, FRigidBodySweptPointContactConstraint& Constraint)
 		{
-			T TOI = 1.0f;
+			FReal TOI = 1.0f;
 			UpdateContactPoint(Constraint.Manifold, CapsuleTriangleMeshSweptContactPoint(Capsule0, WorldTransform0, TriangleMesh1, WorldTransform1, Dir, Length, CullDistance, TOI));
 			SetSweptConstraintTOI(Particle0, TOI, Length, Dir, Constraint);
 		}
 
-		template <typename T, int d>
-		void UpdateCapsuleTriangleMeshManifold(TCollisionConstraintBase<T, d>&  Constraint, const TRigidTransform<T, d>& ATM, const TRigidTransform<T, d>& BTM, const T CullDistance)
+
+		void UpdateCapsuleTriangleMeshManifold(FCollisionConstraintBase&  Constraint, const FRigidTransform3& ATM, const FRigidTransform3& BTM, const FReal CullDistance)
 		{
 
 		}
 
-		template<typename T, int d>
-		void ConstructCapsuleTriangleMeshConstraints(TGeometryParticleHandle<T, d>* Particle0, TGeometryParticleHandle<T, d>* Particle1, const FImplicitObject* Implicit0, const FImplicitObject* Implicit1, const TRigidTransform<T, d>& LocalTransform0, const TRigidTransform<T, d>& LocalTransform1, const T CullDistance, FCollisionConstraintsArray& NewConstraints)
+		void ConstructCapsuleTriangleMeshConstraints(TGeometryParticleHandle<FReal, 3>* Particle0, TGeometryParticleHandle<FReal, 3>* Particle1, const FImplicitObject* Implicit0, const FImplicitObject* Implicit1, const FRigidTransform3& LocalTransform0, const FRigidTransform3& LocalTransform1, const FReal CullDistance, FCollisionConstraintsArray& NewConstraints)
 		{
-			const TCapsule<T> * Object0 = Implicit0->template GetObject<const TCapsule<T> >();
+			const TCapsule<FReal> * Object0 = Implicit0->template GetObject<const TCapsule<FReal> >();
 			if (ensure(Object0))
 			{
 				if (const TImplicitObjectScaled<FTriangleMeshImplicitObject>* ScaledTriangleMesh = Implicit1->template GetObject<const TImplicitObjectScaled<FTriangleMeshImplicitObject>>())
@@ -1597,8 +1564,8 @@ namespace Chaos
 					FRigidBodyPointContactConstraint Constraint = FRigidBodyPointContactConstraint(Particle0, Implicit0, LocalTransform0, Particle1, Implicit1, LocalTransform1, EContactShapesType::CapsuleTriMesh);
 					if (!bChaos_Collision_DeferNarrowPhase)
 					{
-						TRigidTransform<T, d> WorldTransform0 = LocalTransform0 * Collisions::GetTransform(Particle0);
-						TRigidTransform<T, d> WorldTransform1 = LocalTransform1 * Collisions::GetTransform(Particle1);
+						FRigidTransform3 WorldTransform0 = LocalTransform0 * Collisions::GetTransform(Particle0);
+						FRigidTransform3 WorldTransform1 = LocalTransform1 * Collisions::GetTransform(Particle1);
 						UpdateCapsuleTriangleMeshConstraint(*Object0, WorldTransform0, *ScaledTriangleMesh, WorldTransform1, CullDistance, Constraint);
 						NewConstraints.TryAdd(CullDistance, Constraint);
 					}
@@ -1612,8 +1579,8 @@ namespace Chaos
 					FRigidBodyPointContactConstraint Constraint = FRigidBodyPointContactConstraint(Particle0, Implicit0, LocalTransform0, Particle1, Implicit1, LocalTransform1, EContactShapesType::CapsuleTriMesh);
 					if (!bChaos_Collision_DeferNarrowPhase)
 					{
-						TRigidTransform<T, d> WorldTransform0 = LocalTransform0 * Collisions::GetTransform(Particle0);
-						TRigidTransform<T, d> WorldTransform1 = LocalTransform1 * Collisions::GetTransform(Particle1);
+						FRigidTransform3 WorldTransform0 = LocalTransform0 * Collisions::GetTransform(Particle0);
+						FRigidTransform3 WorldTransform1 = LocalTransform1 * Collisions::GetTransform(Particle1);
 						UpdateCapsuleTriangleMeshConstraint(*Object0, WorldTransform0, *TriangleMesh, WorldTransform1, CullDistance, Constraint);
 						NewConstraints.TryAdd(CullDistance, Constraint);
 					}
@@ -1629,25 +1596,24 @@ namespace Chaos
 			}
 		}
 
-		template<typename T, int d>
-		void ConstructCapsuleTriangleMeshConstraintsSwept(TGeometryParticleHandle<T, d>* Particle0, TGeometryParticleHandle<T, d>* Particle1, const FImplicitObject* Implicit0, const FImplicitObject* Implicit1, const TRigidTransform<T, d>& LocalTransform0, const TRigidTransform<T, d>& LocalTransform1, const T CullDistance, const TVector<T,d>& Dir, FReal Length, FCollisionConstraintsArray& NewConstraints)
+		void ConstructCapsuleTriangleMeshConstraintsSwept(TGeometryParticleHandle<FReal, 3>* Particle0, TGeometryParticleHandle<FReal, 3>* Particle1, const FImplicitObject* Implicit0, const FImplicitObject* Implicit1, const FRigidTransform3& LocalTransform0, const FRigidTransform3& LocalTransform1, const FReal CullDistance, const FVec3& Dir, FReal Length, FCollisionConstraintsArray& NewConstraints)
 		{
-			const TCapsule<T> * Object0 = Implicit0->template GetObject<const TCapsule<T> >();
+			const TCapsule<FReal> * Object0 = Implicit0->template GetObject<const TCapsule<FReal> >();
 			if (ensure(Object0))
 			{
 				if (const TImplicitObjectScaled<FTriangleMeshImplicitObject>* ScaledTriangleMesh = Implicit1->template GetObject<const TImplicitObjectScaled<FTriangleMeshImplicitObject>>())
 				{
 					FRigidBodySweptPointContactConstraint Constraint = FRigidBodySweptPointContactConstraint(Particle0, Implicit0, LocalTransform0, Particle1, Implicit1, LocalTransform1, EContactShapesType::CapsuleTriMesh);
-					TRigidTransform<T, 3> WorldTransformX0 = LocalTransform0 * FRigidTransform3(Particle0->X(), Particle0->R());
-					TRigidTransform<T, d> WorldTransform1 = LocalTransform1 * Collisions::GetTransform(Particle1);
+					FRigidTransform3 WorldTransformX0 = LocalTransform0 * FRigidTransform3(Particle0->X(), Particle0->R());
+					FRigidTransform3 WorldTransform1 = LocalTransform1 * Collisions::GetTransform(Particle1);
 					UpdateCapsuleTriangleMeshConstraintSwept(Particle0, *Object0, WorldTransformX0, *ScaledTriangleMesh, WorldTransform1, Dir, Length, CullDistance, Constraint);
 					NewConstraints.TryAdd(CullDistance, Constraint);
 				}
 				else if (const FTriangleMeshImplicitObject* TriangleMesh = Implicit1->template GetObject<const FTriangleMeshImplicitObject>())
 				{
 					FRigidBodySweptPointContactConstraint Constraint = FRigidBodySweptPointContactConstraint(Particle0, Implicit0, LocalTransform0, Particle1, Implicit1, LocalTransform1, EContactShapesType::CapsuleTriMesh);
-					TRigidTransform<T, 3> WorldTransformX0 = LocalTransform0 * FRigidTransform3(Particle0->X(), Particle0->R());
-					TRigidTransform<T, d> WorldTransform1 = LocalTransform1 * Collisions::GetTransform(Particle1);
+					FRigidTransform3 WorldTransformX0 = LocalTransform0 * FRigidTransform3(Particle0->X(), Particle0->R());
+					FRigidTransform3 WorldTransform1 = LocalTransform1 * Collisions::GetTransform(Particle1);
 					UpdateCapsuleTriangleMeshConstraintSwept(Particle0, *Object0, WorldTransformX0, *TriangleMesh, WorldTransform1, Dir, Length, CullDistance, Constraint);
 					NewConstraints.TryAdd(CullDistance, Constraint);
 				}
@@ -1662,10 +1628,9 @@ namespace Chaos
 		// Convex - Convex
 		//
 
-		template<class T, int d>
-		void UpdateConvexConvexConstraint(const FImplicitObject& Implicit0, const TRigidTransform<T, d>& WorldTransform0, const FImplicitObject& Implicit1, const TRigidTransform<T, d>& WorldTransform1, const T CullDistance, TCollisionConstraintBase<T, d>& ConstraintBase)
+		void UpdateConvexConvexConstraint(const FImplicitObject& Implicit0, const FRigidTransform3& WorldTransform0, const FImplicitObject& Implicit1, const FRigidTransform3& WorldTransform1, const FReal CullDistance, FCollisionConstraintBase& ConstraintBase)
 		{
-			TContactPoint<T> ContactPoint;
+			FContactPoint ContactPoint;
 
 			if (ConstraintBase.GetType() == FRigidBodyPointContactConstraint::StaticType())
 			{
@@ -1677,12 +1642,12 @@ namespace Chaos
 			}
 			else if (ConstraintBase.GetType() == FRigidBodyMultiPointContactConstraint::StaticType())
 			{
-				TRigidBodyMultiPointContactConstraint<T, d>& Constraint = *ConstraintBase.template As<TRigidBodyMultiPointContactConstraint<T, d>>();
+				FRigidBodyMultiPointContactConstraint& Constraint = *ConstraintBase.template As<FRigidBodyMultiPointContactConstraint>();
 				ContactPoint.Phi = FLT_MAX;
 
-				const TRigidTransform<T, d> AToBTM = WorldTransform0.GetRelativeTransform(WorldTransform1);
+				const FRigidTransform3 AToBTM = WorldTransform0.GetRelativeTransform(WorldTransform1);
 
-				TPlane<T, d> CollisionPlane(Constraint.GetManifoldPlanePosition(), Constraint.GetManifoldPlaneNormal());
+				TPlane<FReal, 3> CollisionPlane(Constraint.GetManifoldPlanePosition(), Constraint.GetManifoldPlaneNormal());
 
 				// re-sample the constraint based on the distance from the collision plane.
 				for (int32 Idx = 0; Idx < Constraint.NumManifoldPoints(); Idx++)
@@ -1704,20 +1669,18 @@ namespace Chaos
 			UpdateContactPoint(ConstraintBase.Manifold, ContactPoint);
 		}
 
-		template<class T, int d>
-		void UpdateConvexConvexConstraintSwept(TGeometryParticleHandle<T, d>* Particle0, const FImplicitObject& Implicit0, const TRigidTransform<T, d>& WorldTransform0, const FImplicitObject& Implicit1, const TRigidTransform<T, d>& WorldTransform1, const TVector<T, d>& Dir, const T Length, const T CullDistance, TRigidBodySweptPointContactConstraint<T, d>& Constraint)
+		void UpdateConvexConvexConstraintSwept(TGeometryParticleHandle<FReal, 3>* Particle0, const FImplicitObject& Implicit0, const FRigidTransform3& WorldTransform0, const FImplicitObject& Implicit1, const FRigidTransform3& WorldTransform1, const FVec3& Dir, const FReal Length, const FReal CullDistance, FRigidBodySweptPointContactConstraint& Constraint)
 		{
-			T TOI = 1.0f;
+			FReal TOI = 1.0f;
 			UpdateContactPoint(Constraint.Manifold, ConvexConvexContactPointSwept(Implicit0, WorldTransform0, Implicit1, WorldTransform1, Dir, Length, CullDistance, TOI));
 			SetSweptConstraintTOI(Particle0, TOI, Length, Dir, Constraint);
 		}
 
-		template<class T, int d>
-		void UpdateConvexConvexManifold(TCollisionConstraintBase<T, d>&  ConstraintBase, const TRigidTransform<T, d>& WorldTransform0, const TRigidTransform<T, d>& WorldTransform1, const T CullDistance)
+		void UpdateConvexConvexManifold(FCollisionConstraintBase&  ConstraintBase, const FRigidTransform3& WorldTransform0, const FRigidTransform3& WorldTransform1, const FReal CullDistance)
 		{
 			if (ConstraintBase.GetType() == FRigidBodyMultiPointContactConstraint::StaticType())
 			{
-				TRigidBodyMultiPointContactConstraint<T, d>* Constraint = ConstraintBase.template As<TRigidBodyMultiPointContactConstraint<T, d>>();
+				FRigidBodyMultiPointContactConstraint* Constraint = ConstraintBase.template As<FRigidBodyMultiPointContactConstraint>();
 				if (GetInnerType(ConstraintBase.Manifold.Implicit[0]->GetType()) == ImplicitObjectType::Convex)
 				{
 					UpdateSingleShotManifold(*Constraint, WorldTransform0, WorldTransform1, CullDistance);
@@ -1730,8 +1693,7 @@ namespace Chaos
 		}
 
 
-		template<class T, int d>
-		void ConstructConvexConvexConstraints(TGeometryParticleHandle<T, d>* Particle0, TGeometryParticleHandle<T, d>* Particle1, const FImplicitObject* Implicit0, const FImplicitObject* Implicit1, const TRigidTransform<T, d>& LocalTransform0, const TRigidTransform<T, d>& LocalTransform1, const T CullDistance, FCollisionConstraintsArray& NewConstraints)
+		void ConstructConvexConvexConstraints(TGeometryParticleHandle<FReal, 3>* Particle0, TGeometryParticleHandle<FReal, 3>* Particle1, const FImplicitObject* Implicit0, const FImplicitObject* Implicit1, const FRigidTransform3& LocalTransform0, const FRigidTransform3& LocalTransform1, const FReal CullDistance, FCollisionConstraintsArray& NewConstraints)
 		{
 			EImplicitObjectType Implicit0Type = Particle0->Geometry()->GetType();
 			EImplicitObjectType Implicit1Type = Particle1->Geometry()->GetType();
@@ -1741,8 +1703,8 @@ namespace Chaos
 			if (Chaos_Collision_UseManifolds && Implicit0Type != TBox<FReal, 3>::StaticType() && Implicit1Type != TBox<FReal, 3>::StaticType())
 			{
 				FRigidBodyMultiPointContactConstraint Constraint = FRigidBodyMultiPointContactConstraint(Particle0, Implicit0, LocalTransform0, Particle1, Implicit1, LocalTransform1, EContactShapesType::ConvexConvex);
-				TRigidTransform<T, d> WorldTransform0 = LocalTransform0 * Collisions::GetTransform(Particle0);
-				TRigidTransform<T, d> WorldTransform1 = LocalTransform1 * Collisions::GetTransform(Particle1);
+				FRigidTransform3 WorldTransform0 = LocalTransform0 * Collisions::GetTransform(Particle0);
+				FRigidTransform3 WorldTransform1 = LocalTransform1 * Collisions::GetTransform(Particle1);
 				UpdateConvexConvexManifold(Constraint, WorldTransform0, WorldTransform1, CullDistance);
 				if (!bChaos_Collision_DeferNarrowPhase)
 				{
@@ -1759,8 +1721,8 @@ namespace Chaos
 				FRigidBodyPointContactConstraint Constraint = FRigidBodyPointContactConstraint(Particle0, Implicit0, LocalTransform0, Particle1, Implicit1, LocalTransform1, EContactShapesType::ConvexConvex);
 				if (!bChaos_Collision_DeferNarrowPhase)
 				{
-					TRigidTransform<T, d> WorldTransform0 = LocalTransform0 * Collisions::GetTransform(Particle0);
-					TRigidTransform<T, d> WorldTransform1 = LocalTransform1 * Collisions::GetTransform(Particle1);
+					FRigidTransform3 WorldTransform0 = LocalTransform0 * Collisions::GetTransform(Particle0);
+					FRigidTransform3 WorldTransform1 = LocalTransform1 * Collisions::GetTransform(Particle1);
 					UpdateConvexConvexConstraint(*Implicit0, WorldTransform0, *Implicit1, WorldTransform1, CullDistance, Constraint);
 					NewConstraints.TryAdd(CullDistance, Constraint);
 				}
@@ -1771,12 +1733,11 @@ namespace Chaos
 			}
 		}
 
-		template<class T, int d>
-		void ConstructConvexConvexConstraintsSwept(TGeometryParticleHandle<T, d>* Particle0, TGeometryParticleHandle<T, d>* Particle1, const FImplicitObject* Implicit0, const FImplicitObject* Implicit1, const TRigidTransform<T, d>& LocalTransform0, const TRigidTransform<T, d>& LocalTransform1, const T CullDistance, const TVector<T,d>& Dir, FReal Length, FCollisionConstraintsArray& NewConstraints)
+		void ConstructConvexConvexConstraintsSwept(TGeometryParticleHandle<FReal, 3>* Particle0, TGeometryParticleHandle<FReal, 3>* Particle1, const FImplicitObject* Implicit0, const FImplicitObject* Implicit1, const FRigidTransform3& LocalTransform0, const FRigidTransform3& LocalTransform1, const FReal CullDistance, const FVec3& Dir, FReal Length, FCollisionConstraintsArray& NewConstraints)
 		{
 			FRigidBodySweptPointContactConstraint Constraint = FRigidBodySweptPointContactConstraint(Particle0, Implicit0, LocalTransform0, Particle1, Implicit1, LocalTransform1, EContactShapesType::ConvexConvex);
-			TRigidTransform<T, 3> WorldTransformX0 = LocalTransform0 * FRigidTransform3(Particle0->X(), Particle0->R());
-			TRigidTransform<T, d> WorldTransform1 = LocalTransform1 * Collisions::GetTransform(Particle1);
+			FRigidTransform3 WorldTransformX0 = LocalTransform0 * FRigidTransform3(Particle0->X(), Particle0->R());
+			FRigidTransform3 WorldTransform1 = LocalTransform1 * Collisions::GetTransform(Particle1);
 			UpdateConvexConvexConstraintSwept(Particle0, *Implicit0, WorldTransformX0, *Implicit1, WorldTransform1, Dir, Length, CullDistance, Constraint);
 			NewConstraints.TryAdd(CullDistance, Constraint);
 		}
@@ -1785,44 +1746,42 @@ namespace Chaos
 		// Convex - HeightField
 		//
 
-		template <typename T, int d>
-		TContactPoint<T> ConvexHeightFieldContactPoint(const FImplicitObject& A, const TRigidTransform<T, d>& ATransform, const THeightField<T>& B, const TRigidTransform<T, d>& BTransform, const T CullDistance)
+
+		FContactPoint ConvexHeightFieldContactPoint(const FImplicitObject& A, const FRigidTransform3& ATransform, const THeightField<FReal>& B, const FRigidTransform3& BTransform, const FReal CullDistance)
 		{
 			return GJKImplicitContactPoint< FConvex >(A, ATransform, B, BTransform, CullDistance);
 		}
 
-		template <typename T, int d>
-		void UpdateConvexHeightFieldConstraint(const FImplicitObject& A, const TRigidTransform<T, d>& ATransform, const THeightField<T>& B, const TRigidTransform<T, d>& BTransform, const T CullDistance, TRigidBodyPointContactConstraint<T, d>& Constraint)
+
+		void UpdateConvexHeightFieldConstraint(const FImplicitObject& A, const FRigidTransform3& ATransform, const THeightField<FReal>& B, const FRigidTransform3& BTransform, const FReal CullDistance, FRigidBodyPointContactConstraint& Constraint)
 		{
 			UpdateContactPoint(Constraint.Manifold, ConvexHeightFieldContactPoint(A, ATransform, B, BTransform, CullDistance));
 		}
 
-		template <typename T, int d>
-		void UpdateConvexHeightFieldConstraintSwept(TGeometryParticleHandle<T, d>* Particle0, const FImplicitObject& A, const TRigidTransform<T, d>& ATransform, const THeightField<T>& B, const TRigidTransform<T, d>& BTransform, const TVector<T, d>& Dir, const T Length, const T CullDistance, TRigidBodySweptPointContactConstraint<T, d>& Constraint)
+
+		void UpdateConvexHeightFieldConstraintSwept(TGeometryParticleHandle<FReal, 3>* Particle0, const FImplicitObject& A, const FRigidTransform3& ATransform, const THeightField<FReal>& B, const FRigidTransform3& BTransform, const FVec3& Dir, const FReal Length, const FReal CullDistance, FRigidBodySweptPointContactConstraint& Constraint)
 		{
-			T TOI = 1.0f;
+			FReal TOI = 1.0f;
 			UpdateContactPoint(Constraint.Manifold, GJKImplicitSweptContactPoint< FConvex >(A, ATransform, B, BTransform, Dir, Length, CullDistance, TOI));
 			SetSweptConstraintTOI(Particle0, TOI, Length, Dir, Constraint);
 		}
 
-		template<class T, int d>
-		void UpdateConvexHeightFieldManifold(TCollisionConstraintBase<T, d>&  Constraint, const TRigidTransform<T, d>& ATM, const TRigidTransform<T, d>& BTM, const T CullDistance)
+		void UpdateConvexHeightFieldManifold(FCollisionConstraintBase&  Constraint, const FRigidTransform3& ATM, const FRigidTransform3& BTM, const FReal CullDistance)
 		{
 			// @todo(chaos) : Stub Update Manifold
 			//   Stub function for updating the manifold prior to the Apply and ApplyPushOut
 		}
 
-		template<typename T, int d>
-		void ConstructConvexHeightFieldConstraints(TGeometryParticleHandle<T, d>* Particle0, TGeometryParticleHandle<T, d>* Particle1, const FImplicitObject* Implicit0, const FImplicitObject* Implicit1, const TRigidTransform<T, d>& LocalTransform0, const TRigidTransform<T, d>& LocalTransform1, const T CullDistance, FCollisionConstraintsArray& NewConstraints)
+		void ConstructConvexHeightFieldConstraints(TGeometryParticleHandle<FReal, 3>* Particle0, TGeometryParticleHandle<FReal, 3>* Particle1, const FImplicitObject* Implicit0, const FImplicitObject* Implicit1, const FRigidTransform3& LocalTransform0, const FRigidTransform3& LocalTransform1, const FReal CullDistance, FCollisionConstraintsArray& NewConstraints)
 		{
-			const THeightField<T> * Object1 = Implicit1->template GetObject<const THeightField<T> >();
+			const THeightField<FReal> * Object1 = Implicit1->template GetObject<const THeightField<FReal> >();
 			if (ensure(Implicit0->IsConvex() && Object1))
 			{
 				FRigidBodyPointContactConstraint Constraint = FRigidBodyPointContactConstraint(Particle0, Implicit0, LocalTransform0, Particle1, Implicit1, LocalTransform1, EContactShapesType::ConvexHeightField);
 				if (!bChaos_Collision_DeferNarrowPhase)
 				{
-					TRigidTransform<T, d> WorldTransform0 = LocalTransform0 * Collisions::GetTransform(Particle0);
-					TRigidTransform<T, d> WorldTransform1 = LocalTransform1 * Collisions::GetTransform(Particle1);
+					FRigidTransform3 WorldTransform0 = LocalTransform0 * Collisions::GetTransform(Particle0);
+					FRigidTransform3 WorldTransform1 = LocalTransform1 * Collisions::GetTransform(Particle1);
 					UpdateConvexHeightFieldConstraint(*Implicit0, WorldTransform0, *Object1, WorldTransform1, CullDistance, Constraint);
 					NewConstraints.TryAdd(CullDistance, Constraint);
 				}
@@ -1833,15 +1792,14 @@ namespace Chaos
 			}
 		}
 
-		template<typename T, int d>
-		void ConstructConvexHeightFieldConstraintsSwept(TGeometryParticleHandle<T, d>* Particle0, TGeometryParticleHandle<T, d>* Particle1, const FImplicitObject* Implicit0, const FImplicitObject* Implicit1, const TRigidTransform<T, d>& LocalTransform0, const TRigidTransform<T, d>& LocalTransform1, const T CullDistance, const TVector<T, d>& Dir, FReal Length, FCollisionConstraintsArray& NewConstraints)
+		void ConstructConvexHeightFieldConstraintsSwept(TGeometryParticleHandle<FReal, 3>* Particle0, TGeometryParticleHandle<FReal, 3>* Particle1, const FImplicitObject* Implicit0, const FImplicitObject* Implicit1, const FRigidTransform3& LocalTransform0, const FRigidTransform3& LocalTransform1, const FReal CullDistance, const FVec3& Dir, FReal Length, FCollisionConstraintsArray& NewConstraints)
 		{
-			const THeightField<T> * Object1 = Implicit1->template GetObject<const THeightField<T> >();
+			const THeightField<FReal> * Object1 = Implicit1->template GetObject<const THeightField<FReal> >();
 			if (ensure(Implicit0->IsConvex() && Object1))
 			{
 				FRigidBodySweptPointContactConstraint Constraint = FRigidBodySweptPointContactConstraint(Particle0, Implicit0, LocalTransform0, Particle1, Implicit1, LocalTransform1, EContactShapesType::ConvexHeightField);
-				TRigidTransform<T, 3> WorldTransformX0 = LocalTransform0 * FRigidTransform3(Particle0->X(), Particle0->R());
-				TRigidTransform<T, d> WorldTransform1 = LocalTransform1 * Collisions::GetTransform(Particle1);
+				FRigidTransform3 WorldTransformX0 = LocalTransform0 * FRigidTransform3(Particle0->X(), Particle0->R());
+				FRigidTransform3 WorldTransform1 = LocalTransform1 * Collisions::GetTransform(Particle1);
 				UpdateConvexHeightFieldConstraintSwept(Particle0, *Implicit0, WorldTransformX0, *Object1, WorldTransform1, Dir, Length, CullDistance, Constraint);
 				NewConstraints.TryAdd(CullDistance, Constraint);
 			}
@@ -1851,14 +1809,14 @@ namespace Chaos
 		// Convex-TriangleMesh
 		//
 
-		template <typename TriMeshType, typename T, int d>
-		TContactPoint<T> ConvexTriangleMeshContactPoint(const FImplicitObject& A, const TRigidTransform<T, d>& ATransform, const TriMeshType& B, const TRigidTransform<T, d>& BTransform, const T CullDistance)
+		template <typename TriMeshType>
+		FContactPoint ConvexTriangleMeshContactPoint(const FImplicitObject& A, const FRigidTransform3& ATransform, const TriMeshType& B, const FRigidTransform3& BTransform, const FReal CullDistance)
 		{
 			return GJKImplicitContactPoint< FConvex >(A, ATransform, B, BTransform, CullDistance);
 		}
 
-		template <typename TriMeshType, typename T, int d>
-		TContactPoint<T> ConvexTriangleMeshSweptContactPoint(const FImplicitObject& A, const TRigidTransform<T, d>& ATransform, const TriMeshType& B, const TRigidTransform<T, d>& BStartTransform, const TVector<T, d>& Dir, const T Length, const T CullDistance, T& TOI)
+		template <typename TriMeshType>
+		FContactPoint ConvexTriangleMeshSweptContactPoint(const FImplicitObject& A, const FRigidTransform3& ATransform, const TriMeshType& B, const FRigidTransform3& BStartTransform, const FVec3& Dir, const FReal Length, const FReal CullDistance, FReal& TOI)
 		{
 			if (const TImplicitObjectScaled<FTriangleMeshImplicitObject>* ScaledTriangleMesh = B.template GetObject<const TImplicitObjectScaled<FTriangleMeshImplicitObject>>())
 			{
@@ -1870,32 +1828,31 @@ namespace Chaos
 			}
 
 			ensure(false);
-			return TContactPoint<T>();
+			return FContactPoint();
 		}
 		
-		template <typename TriMeshType, typename T, int d>
-		void UpdateConvexTriangleMeshConstraint(const FImplicitObject& Convex0, const TRigidTransform<T, d>& WorldTransform0, const TriMeshType& TriangleMesh1, const TRigidTransform<T, d>& WorldTransform1, const T CullDistance, TRigidBodyPointContactConstraint<T, d>& Constraint)
+		template <typename TriMeshType>
+		void UpdateConvexTriangleMeshConstraint(const FImplicitObject& Convex0, const FRigidTransform3& WorldTransform0, const TriMeshType& TriangleMesh1, const FRigidTransform3& WorldTransform1, const FReal CullDistance, FRigidBodyPointContactConstraint& Constraint)
 		{
 			UpdateContactPoint(Constraint.Manifold, ConvexTriangleMeshContactPoint(Convex0, WorldTransform0, TriangleMesh1, WorldTransform1, CullDistance));
 		}
 
 		// Sweeps convex against trimesh
-		template <typename TriMeshType, typename T, int d>
-		void UpdateConvexTriangleMeshConstraintSwept(TGeometryParticleHandle<T,d>* Particle0, const FImplicitObject& Convex0, const TRigidTransform<T, d>& WorldTransform0, const TriMeshType& TriangleMesh1, const TRigidTransform<T, d>& WorldTransform1, const TVector<T, d>& Dir, const T Length, const T CullDistance, TRigidBodySweptPointContactConstraint<T, d>& Constraint)
+		template <typename TriMeshType>
+		void UpdateConvexTriangleMeshConstraintSwept(TGeometryParticleHandle<FReal, 3>* Particle0, const FImplicitObject& Convex0, const FRigidTransform3& WorldTransform0, const TriMeshType& TriangleMesh1, const FRigidTransform3& WorldTransform1, const FVec3& Dir, const FReal Length, const FReal CullDistance, FRigidBodySweptPointContactConstraint& Constraint)
 		{
-			T TOI = 1.0f;
+			FReal TOI = 1.0f;
 			UpdateContactPoint(Constraint.Manifold, ConvexTriangleMeshSweptContactPoint(Convex0, WorldTransform0, TriangleMesh1, WorldTransform1, Dir, Length, CullDistance, TOI));
 			SetSweptConstraintTOI(Particle0, TOI, Length, Dir, Constraint);
 		}
 
-		template <typename T, int d>
-		void UpdateConvexTriangleMeshManifold(TCollisionConstraintBase<T, d>&  Constraint, const TRigidTransform<T, d>& ATM, const TRigidTransform<T, d>& BTM, const T CullDistance)
+
+		void UpdateConvexTriangleMeshManifold(FCollisionConstraintBase&  Constraint, const FRigidTransform3& ATM, const FRigidTransform3& BTM, const FReal CullDistance)
 		{
 
 		}
 
-		template<typename T, int d>
-		void ConstructConvexTriangleMeshConstraints(TGeometryParticleHandle<T, d>* Particle0, TGeometryParticleHandle<T, d>* Particle1, const FImplicitObject* Implicit0, const FImplicitObject* Implicit1, const TRigidTransform<T, d>& LocalTransform0, const TRigidTransform<T, d>& LocalTransform1, const T CullDistance, FCollisionConstraintsArray& NewConstraints)
+		void ConstructConvexTriangleMeshConstraints(TGeometryParticleHandle<FReal, 3>* Particle0, TGeometryParticleHandle<FReal, 3>* Particle1, const FImplicitObject* Implicit0, const FImplicitObject* Implicit1, const FRigidTransform3& LocalTransform0, const FRigidTransform3& LocalTransform1, const FReal CullDistance, FCollisionConstraintsArray& NewConstraints)
 		{
 			if (ensure(Implicit0->IsConvex()))
 			{
@@ -1904,8 +1861,8 @@ namespace Chaos
 					FRigidBodyPointContactConstraint Constraint = FRigidBodyPointContactConstraint(Particle0, Implicit0, LocalTransform0, Particle1, Implicit1, LocalTransform0, EContactShapesType::ConvexTriMesh);
 					if (!bChaos_Collision_DeferNarrowPhase)
 					{
-						TRigidTransform<T, d> WorldTransform0 = LocalTransform0 * Collisions::GetTransform(Particle0);
-						TRigidTransform<T, d> WorldTransform1 = LocalTransform1 * Collisions::GetTransform(Particle1);
+						FRigidTransform3 WorldTransform0 = LocalTransform0 * Collisions::GetTransform(Particle0);
+						FRigidTransform3 WorldTransform1 = LocalTransform1 * Collisions::GetTransform(Particle1);
 						UpdateConvexTriangleMeshConstraint(*Implicit0, WorldTransform0, *ScaledTriangleMesh , WorldTransform1, CullDistance, Constraint);
 						NewConstraints.TryAdd(CullDistance, Constraint);
 					}
@@ -1919,8 +1876,8 @@ namespace Chaos
 					FRigidBodyPointContactConstraint Constraint = FRigidBodyPointContactConstraint(Particle0, Implicit0, LocalTransform0, Particle1, Implicit1, LocalTransform1, EContactShapesType::ConvexTriMesh);
 					if (!bChaos_Collision_DeferNarrowPhase)
 					{
-						TRigidTransform<T, d> WorldTransform0 = LocalTransform0 * Collisions::GetTransform(Particle0);
-						TRigidTransform<T, d> WorldTransform1 = LocalTransform1 * Collisions::GetTransform(Particle1);
+						FRigidTransform3 WorldTransform0 = LocalTransform0 * Collisions::GetTransform(Particle0);
+						FRigidTransform3 WorldTransform1 = LocalTransform1 * Collisions::GetTransform(Particle1);
 						UpdateConvexTriangleMeshConstraint(*Implicit0, WorldTransform0, *TriangleMesh, WorldTransform1, CullDistance, Constraint);
 						NewConstraints.TryAdd(CullDistance, Constraint);
 					}
@@ -1936,24 +1893,23 @@ namespace Chaos
 			}
 		}
 
-		template<typename T, int d>
-		void ConstructConvexTriangleMeshConstraintsSwept(TGeometryParticleHandle<T, d>* Particle0, TGeometryParticleHandle<T, d>* Particle1, const FImplicitObject* Implicit0, const FImplicitObject* Implicit1, const TRigidTransform<T, d>& LocalTransform0, const TRigidTransform<T, d>& LocalTransform1, const T CullDistance, const TVector<T, d>& Dir, FReal Length, FCollisionConstraintsArray& NewConstraints)
+		void ConstructConvexTriangleMeshConstraintsSwept(TGeometryParticleHandle<FReal, 3>* Particle0, TGeometryParticleHandle<FReal, 3>* Particle1, const FImplicitObject* Implicit0, const FImplicitObject* Implicit1, const FRigidTransform3& LocalTransform0, const FRigidTransform3& LocalTransform1, const FReal CullDistance, const FVec3& Dir, FReal Length, FCollisionConstraintsArray& NewConstraints)
 		{
 			if (ensure(Implicit0->IsConvex()))
 			{
 				if (const TImplicitObjectScaled<FTriangleMeshImplicitObject>* ScaledTriangleMesh = Implicit1->template GetObject<const TImplicitObjectScaled<FTriangleMeshImplicitObject>>())
 				{
 					FRigidBodySweptPointContactConstraint Constraint = FRigidBodySweptPointContactConstraint(Particle0, Implicit0, LocalTransform0, Particle1, Implicit1, LocalTransform1, EContactShapesType::ConvexTriMesh);
-					TRigidTransform<T, 3> WorldTransformX0 = LocalTransform0 * FRigidTransform3(Particle0->X(), Particle0->R());
-					TRigidTransform<T, d> WorldTransform1 = LocalTransform1 * Collisions::GetTransform(Particle1);
+					FRigidTransform3 WorldTransformX0 = LocalTransform0 * FRigidTransform3(Particle0->X(), Particle0->R());
+					FRigidTransform3 WorldTransform1 = LocalTransform1 * Collisions::GetTransform(Particle1);
 					UpdateConvexTriangleMeshConstraintSwept(Particle0, *Implicit0, WorldTransformX0, *ScaledTriangleMesh, WorldTransform1, Dir, Length, CullDistance, Constraint);
 					NewConstraints.TryAdd(CullDistance, Constraint);
 				}
 				else if (const FTriangleMeshImplicitObject* TriangleMesh = Implicit1->template GetObject<const FTriangleMeshImplicitObject>())
 				{
 					FRigidBodySweptPointContactConstraint Constraint = FRigidBodySweptPointContactConstraint(Particle0, Implicit0, LocalTransform0, Particle1, Implicit1, LocalTransform1, EContactShapesType::ConvexTriMesh);
-					TRigidTransform<T, 3> WorldTransformX0 = LocalTransform0 * FRigidTransform3(Particle0->X(), Particle0->R());
-					TRigidTransform<T, d> WorldTransform1 = LocalTransform1 * Collisions::GetTransform(Particle1);
+					FRigidTransform3 WorldTransformX0 = LocalTransform0 * FRigidTransform3(Particle0->X(), Particle0->R());
+					FRigidTransform3 WorldTransform1 = LocalTransform1 * Collisions::GetTransform(Particle1);
 					UpdateConvexTriangleMeshConstraintSwept(Particle0, *Implicit0, WorldTransformX0, *TriangleMesh, WorldTransform1, Dir, Length, CullDistance, Constraint);
 					NewConstraints.TryAdd(CullDistance, Constraint);
 				}
@@ -1970,26 +1926,26 @@ namespace Chaos
 		//
 
 		DECLARE_CYCLE_STAT(TEXT("TPBDCollisionConstraints::UpdateLevelsetLevelsetConstraint"), STAT_UpdateLevelsetLevelsetConstraint, STATGROUP_ChaosWide);
-		template<ECollisionUpdateType UpdateType, typename T, int d>
-		void UpdateLevelsetLevelsetConstraint(const T CullDistance, TRigidBodyPointContactConstraint<T, d>& Constraint)
+		template<ECollisionUpdateType UpdateType>
+		void UpdateLevelsetLevelsetConstraint(const FReal CullDistance, FRigidBodyPointContactConstraint& Constraint)
 		{
 			SCOPE_CYCLE_COUNTER(STAT_UpdateLevelsetLevelsetConstraint);
 
-			TGenericParticleHandle<T, d> Particle0 = Constraint.Particle[0];
-			TRigidTransform<T, d> ParticlesTM = TRigidTransform<T, d>(Particle0->P(), Particle0->Q());
+			TGenericParticleHandle<FReal, 3> Particle0 = Constraint.Particle[0];
+			FRigidTransform3 ParticlesTM = FRigidTransform3(Particle0->P(), Particle0->Q());
 			if (!(ensure(!FMath::IsNaN(ParticlesTM.GetTranslation().X)) && ensure(!FMath::IsNaN(ParticlesTM.GetTranslation().Y)) && ensure(!FMath::IsNaN(ParticlesTM.GetTranslation().Z))))
 			{
 				return;
 			}
 
-			TGenericParticleHandle<T, d> Particle1 = Constraint.Particle[1];
-			TRigidTransform<T, d> LevelsetTM = TRigidTransform<T, d>(Particle1->P(), Particle1->Q());
+			TGenericParticleHandle<FReal, 3> Particle1 = Constraint.Particle[1];
+			FRigidTransform3 LevelsetTM = FRigidTransform3(Particle1->P(), Particle1->Q());
 			if (!(ensure(!FMath::IsNaN(LevelsetTM.GetTranslation().X)) && ensure(!FMath::IsNaN(LevelsetTM.GetTranslation().Y)) && ensure(!FMath::IsNaN(LevelsetTM.GetTranslation().Z))))
 			{
 				return;
 			}
 
-			const TBVHParticles<T, d>* SampleParticles = nullptr;
+			const TBVHParticles<FReal, 3>* SampleParticles = nullptr;
 			SampleParticles = Particle0->CollisionParticles().Get();
 
 			if (SampleParticles)
@@ -1999,16 +1955,14 @@ namespace Chaos
 			}
 		}
 
-		template<class T, int d>
-		void UpdateLevelsetLevelsetManifold(TCollisionConstraintBase<T, d>&  Constraint, const TRigidTransform<T, d>& ATM, const TRigidTransform<T, d>& BTM, const T CullDistance)
+		void UpdateLevelsetLevelsetManifold(FCollisionConstraintBase&  Constraint, const FRigidTransform3& ATM, const FRigidTransform3& BTM, const FReal CullDistance)
 		{
 			// @todo(chaos) : Stub Update Manifold
 			//   Stub function for updating the manifold prior to the Apply and ApplyPushOut
 		}
 
 
-		template<typename T, int d>
-		void ConstructLevelsetLevelsetConstraints(TGeometryParticleHandle<T, d>* Particle0, TGeometryParticleHandle<T, d>* Particle1, const FImplicitObject* Implicit0, const FImplicitObject* Implicit1, const TRigidTransform<T, d>& LocalTransform0, const TRigidTransform<T, d>& LocalTransform1, const T CullDistance, FCollisionConstraintsArray& NewConstraints)
+		void ConstructLevelsetLevelsetConstraints(TGeometryParticleHandle<FReal, 3>* Particle0, TGeometryParticleHandle<FReal, 3>* Particle1, const FImplicitObject* Implicit0, const FImplicitObject* Implicit1, const FRigidTransform3& LocalTransform0, const FRigidTransform3& LocalTransform1, const FReal CullDistance, FCollisionConstraintsArray& NewConstraints)
 		{
 			FRigidBodyPointContactConstraint Constraint = FRigidBodyPointContactConstraint(Particle0, Implicit0, LocalTransform0, Particle1, Implicit1, LocalTransform1, EContactShapesType::LevelSetLevelSet);
 
@@ -2043,18 +1997,17 @@ namespace Chaos
 		//
 
 
-		template<typename T, int d>
-		inline void UpdateManifold(TRigidBodyMultiPointContactConstraint<T, d>& Constraint, const TRigidTransform<T, d>& ParticleTransform0, const TRigidTransform<T, d>& ParticleTransform1, const T CullDistance, const FCollisionContext& Context)
+		void UpdateManifold(FRigidBodyMultiPointContactConstraint& Constraint, const FRigidTransform3& ParticleTransform0, const FRigidTransform3& ParticleTransform1, const FReal CullDistance, const FCollisionContext& Context)
 		{
 			const FImplicitObject& Implicit0 = *Constraint.Manifold.Implicit[0];
 			const FImplicitObject& Implicit1 = *Constraint.Manifold.Implicit[1];
-			const TRigidTransform<T, d> WorldTransform0 = Constraint.ImplicitTransform[0] * ParticleTransform0;
-			const TRigidTransform<T, d> WorldTransform1 = Constraint.ImplicitTransform[1] * ParticleTransform1;
+			const FRigidTransform3 WorldTransform0 = Constraint.ImplicitTransform[0] * ParticleTransform0;
+			const FRigidTransform3 WorldTransform1 = Constraint.ImplicitTransform[1] * ParticleTransform1;
 
 			switch (Constraint.Manifold.ShapesType)
 			{
 			case EContactShapesType::CapsuleBox:
-				UpdateCapsuleBoxManifold<T, d>(*Implicit0.template GetObject<const TCapsule<T>>(), WorldTransform0, Implicit1.template GetObject<const TBox<T, d>>()->GetAABB(), WorldTransform1, CullDistance, Context, Constraint);
+				UpdateCapsuleBoxManifold(*Implicit0.template GetObject<const TCapsule<FReal>>(), WorldTransform0, Implicit1.template GetObject<const TBox<FReal, 3>>()->GetAABB(), WorldTransform1, CullDistance, Context, Constraint);
 				break;
 			case EContactShapesType::ConvexConvex:
 				UpdateConvexConvexManifold(Constraint, WorldTransform0, WorldTransform1, CullDistance);
@@ -2072,8 +2025,8 @@ namespace Chaos
 		// @todo(chaos): use a lookup table?
 		// @todo(chaos): add the missing cases below
 		// @todo(chaos): see use GetInnerObject below - we should try to uise the leaf types for all (currently only TriMesh needs this)
-		template<ECollisionUpdateType UpdateType, typename T, int d>
-		inline void UpdateConstraintFromGeometryImpl(TRigidBodyPointContactConstraint<T, d>& Constraint, const TRigidTransform<T, d>& WorldTransform0, const TRigidTransform<T, d>& WorldTransform1, const T CullDistance)
+		template<ECollisionUpdateType UpdateType>
+		inline void UpdateConstraintFromGeometryImpl(FRigidBodyPointContactConstraint& Constraint, const FRigidTransform3& WorldTransform0, const FRigidTransform3& WorldTransform1, const FReal CullDistance)
 		{
 			const FImplicitObject& Implicit0 = *Constraint.Manifold.Implicit[0];
 			const FImplicitObject& Implicit1 = *Constraint.Manifold.Implicit[1];
@@ -2081,13 +2034,13 @@ namespace Chaos
 			switch (Constraint.Manifold.ShapesType)
 			{
 			case EContactShapesType::SphereSphere:
-				UpdateSphereSphereConstraint(*Implicit0.template GetObject<TSphere<T, d>>(), WorldTransform0, *Implicit1.template GetObject<TSphere<T, d>>(), WorldTransform1, CullDistance, Constraint);
+				UpdateSphereSphereConstraint(*Implicit0.template GetObject<TSphere<FReal, 3>>(), WorldTransform0, *Implicit1.template GetObject<TSphere<FReal, 3>>(), WorldTransform1, CullDistance, Constraint);
 				break;
 			case EContactShapesType::SphereCapsule:
-				UpdateSphereCapsuleConstraint(*Implicit0.template GetObject<TSphere<T, d>>(), WorldTransform0, *Implicit1.template GetObject<TCapsule<T>>(), WorldTransform1, CullDistance, Constraint);
+				UpdateSphereCapsuleConstraint(*Implicit0.template GetObject<TSphere<FReal, 3>>(), WorldTransform0, *Implicit1.template GetObject<TCapsule<FReal>>(), WorldTransform1, CullDistance, Constraint);
 				break;
 			case EContactShapesType::SphereBox:
-				UpdateSphereBoxConstraint(*Implicit0.template GetObject<TSphere<T, d>>(), WorldTransform0, Implicit1.template GetObject<TBox<T, d>>()->GetAABB(), WorldTransform1, CullDistance, Constraint);
+				UpdateSphereBoxConstraint(*Implicit0.template GetObject<TSphere<FReal, 3>>(), WorldTransform0, Implicit1.template GetObject<TBox<FReal, 3>>()->GetAABB(), WorldTransform1, CullDistance, Constraint);
 				break;
 			case EContactShapesType::SphereConvex:
 				// MISSING CASE!!!
@@ -2096,27 +2049,27 @@ namespace Chaos
 			case EContactShapesType::SphereTriMesh:
 				if (const TImplicitObjectScaled<FTriangleMeshImplicitObject>* ScaledTriMesh = Implicit1.template GetObject<const TImplicitObjectScaled<FTriangleMeshImplicitObject> >())
 				{
-					UpdateSphereTriangleMeshConstraint(*Implicit0.template GetObject<TSphere<T, d>>(), WorldTransform0, *ScaledTriMesh, WorldTransform1, CullDistance, Constraint);
+					UpdateSphereTriangleMeshConstraint(*Implicit0.template GetObject<TSphere<FReal, 3>>(), WorldTransform0, *ScaledTriMesh, WorldTransform1, CullDistance, Constraint);
 					break;
 				}
 				else if (const FTriangleMeshImplicitObject* TriangleMeshImplicit = Implicit1.template GetObject<const FTriangleMeshImplicitObject>())
 				{
-					UpdateSphereTriangleMeshConstraint(*Implicit0.template GetObject<TSphere<T, d>>(), WorldTransform0, *TriangleMeshImplicit, WorldTransform1, CullDistance, Constraint);
+					UpdateSphereTriangleMeshConstraint(*Implicit0.template GetObject<TSphere<FReal, 3>>(), WorldTransform0, *TriangleMeshImplicit, WorldTransform1, CullDistance, Constraint);
 					break;
 				}
 				ensure(false);
 				break;
 			case EContactShapesType::SphereHeightField:
-				UpdateSphereHeightFieldConstraint(*Implicit0.template GetObject<TSphere<T, d>>(), WorldTransform0, *Implicit1.template GetObject< THeightField<T> >(), WorldTransform1, CullDistance, Constraint);
+				UpdateSphereHeightFieldConstraint(*Implicit0.template GetObject<TSphere<FReal, 3>>(), WorldTransform0, *Implicit1.template GetObject< THeightField<FReal> >(), WorldTransform1, CullDistance, Constraint);
 				break;
 			case EContactShapesType::SpherePlane:
-				UpdateSpherePlaneConstraint(*Implicit0.template GetObject<TSphere<T, d>>(), WorldTransform0, *Implicit1.template GetObject<TPlane<T, d>>(), WorldTransform1, CullDistance, Constraint);
+				UpdateSpherePlaneConstraint(*Implicit0.template GetObject<TSphere<FReal, 3>>(), WorldTransform0, *Implicit1.template GetObject<TPlane<FReal, 3>>(), WorldTransform1, CullDistance, Constraint);
 				break;
 			case EContactShapesType::CapsuleCapsule:
-				UpdateCapsuleCapsuleConstraint(*Implicit0.template GetObject<TCapsule<T>>(), WorldTransform0, *Implicit1.template GetObject<TCapsule<T>>(), WorldTransform1, CullDistance, Constraint);
+				UpdateCapsuleCapsuleConstraint(*Implicit0.template GetObject<TCapsule<FReal>>(), WorldTransform0, *Implicit1.template GetObject<TCapsule<FReal>>(), WorldTransform1, CullDistance, Constraint);
 				break;
 			case EContactShapesType::CapsuleBox:
-				UpdateCapsuleBoxConstraint(*Implicit0.template GetObject<TCapsule<T>>(), WorldTransform0, Implicit1.template GetObject<TBox<T, d>>()->GetAABB(), WorldTransform1, CullDistance, Constraint);
+				UpdateCapsuleBoxConstraint(*Implicit0.template GetObject<TCapsule<FReal>>(), WorldTransform0, Implicit1.template GetObject<TBox<FReal, 3>>()->GetAABB(), WorldTransform1, CullDistance, Constraint);
 				break;
 			case EContactShapesType::CapsuleConvex:
 				// MISSING CASE!!!
@@ -2125,21 +2078,21 @@ namespace Chaos
 			case EContactShapesType::CapsuleTriMesh:
 				if (const TImplicitObjectScaled<FTriangleMeshImplicitObject>* ScaledTriMesh = Implicit1.template GetObject<const TImplicitObjectScaled<FTriangleMeshImplicitObject> >())
 				{
-					UpdateCapsuleTriangleMeshConstraint(*Implicit0.template GetObject<TCapsule<T>>(), WorldTransform0, *ScaledTriMesh, WorldTransform1, CullDistance, Constraint);
+					UpdateCapsuleTriangleMeshConstraint(*Implicit0.template GetObject<TCapsule<FReal>>(), WorldTransform0, *ScaledTriMesh, WorldTransform1, CullDistance, Constraint);
 					break;
 				}
 				else if (const FTriangleMeshImplicitObject* TriangleMeshImplicit = Implicit1.template GetObject<const FTriangleMeshImplicitObject>())
 				{
-					UpdateCapsuleTriangleMeshConstraint(*Implicit0.template GetObject<TCapsule<T>>(), WorldTransform0, *TriangleMeshImplicit, WorldTransform1, CullDistance, Constraint);
+					UpdateCapsuleTriangleMeshConstraint(*Implicit0.template GetObject<TCapsule<FReal>>(), WorldTransform0, *TriangleMeshImplicit, WorldTransform1, CullDistance, Constraint);
 					break;
 				}
 				ensure(false);
 				break;
 			case EContactShapesType::CapsuleHeightField:
-				UpdateCapsuleHeightFieldConstraint(*Implicit0.template GetObject<TCapsule<T>>(), WorldTransform0, *Implicit1.template GetObject< THeightField<T> >(), WorldTransform1, CullDistance, Constraint);
+				UpdateCapsuleHeightFieldConstraint(*Implicit0.template GetObject<TCapsule<FReal>>(), WorldTransform0, *Implicit1.template GetObject< THeightField<FReal> >(), WorldTransform1, CullDistance, Constraint);
 				break;
 			case EContactShapesType::BoxBox:
-				UpdateBoxBoxConstraint(Implicit0.template GetObject<TBox<T, d>>()->GetAABB(), WorldTransform0, Implicit1.template GetObject<TBox<T, d>>()->GetAABB(), WorldTransform1, CullDistance, Constraint);
+				UpdateBoxBoxConstraint(Implicit0.template GetObject<TBox<FReal, 3>>()->GetAABB(), WorldTransform0, Implicit1.template GetObject<TBox<FReal, 3>>()->GetAABB(), WorldTransform1, CullDistance, Constraint);
 				break;
 			case EContactShapesType::BoxConvex:
 				// MISSING CASE!!!
@@ -2148,21 +2101,21 @@ namespace Chaos
 			case EContactShapesType::BoxTriMesh:
 				if (const TImplicitObjectScaled<FTriangleMeshImplicitObject>* ScaledTriMesh = Implicit1.template GetObject<const TImplicitObjectScaled<FTriangleMeshImplicitObject> >())
 				{
-					UpdateBoxTriangleMeshConstraint(Implicit0.template GetObject<TBox<T, d>>()->GetAABB(), WorldTransform0, *ScaledTriMesh, WorldTransform1, CullDistance, Constraint);
+					UpdateBoxTriangleMeshConstraint(Implicit0.template GetObject<TBox<FReal, 3>>()->GetAABB(), WorldTransform0, *ScaledTriMesh, WorldTransform1, CullDistance, Constraint);
 					break;
 				}
 				else if (const FTriangleMeshImplicitObject* TriangleMeshImplicit = Implicit1.template GetObject<const FTriangleMeshImplicitObject>())
 				{
-					UpdateBoxTriangleMeshConstraint(Implicit0.template GetObject<TBox<T, d>>()->GetAABB(), WorldTransform0, *TriangleMeshImplicit, WorldTransform1, CullDistance, Constraint);
+					UpdateBoxTriangleMeshConstraint(Implicit0.template GetObject<TBox<FReal, 3>>()->GetAABB(), WorldTransform0, *TriangleMeshImplicit, WorldTransform1, CullDistance, Constraint);
 					break;
 				}
 				ensure(false);
 				break;
 			case EContactShapesType::BoxHeightField:
-				UpdateBoxHeightFieldConstraint(Implicit0.template GetObject<TBox<T, d>>()->GetAABB(), WorldTransform0, *Implicit1.template GetObject< THeightField<T> >(), WorldTransform1, CullDistance, Constraint);
+				UpdateBoxHeightFieldConstraint(Implicit0.template GetObject<TBox<FReal, 3>>()->GetAABB(), WorldTransform0, *Implicit1.template GetObject< THeightField<FReal> >(), WorldTransform1, CullDistance, Constraint);
 				break;
 			case EContactShapesType::BoxPlane:
-				UpdateBoxPlaneConstraint(Implicit0.template GetObject<TBox<T, d>>()->GetAABB(), WorldTransform0, *Implicit1.template GetObject<TPlane<T, d>>(), WorldTransform1, CullDistance, Constraint);
+				UpdateBoxPlaneConstraint(Implicit0.template GetObject<TBox<FReal, 3>>()->GetAABB(), WorldTransform0, *Implicit1.template GetObject<TPlane<FReal, 3>>(), WorldTransform1, CullDistance, Constraint);
 				break;
 			case EContactShapesType::ConvexConvex:
 				UpdateConvexConvexConstraint(Implicit0, WorldTransform0, Implicit1, WorldTransform1, CullDistance, Constraint);
@@ -2181,7 +2134,7 @@ namespace Chaos
 				ensure(false);
 				break;
 			case EContactShapesType::ConvexHeightField:
-				UpdateConvexHeightFieldConstraint(Implicit0, WorldTransform0, *Implicit1.template GetObject< THeightField<T> >(), WorldTransform1, CullDistance, Constraint);
+				UpdateConvexHeightFieldConstraint(Implicit0, WorldTransform0, *Implicit1.template GetObject< THeightField<FReal> >(), WorldTransform1, CullDistance, Constraint);
 				break;
 			case EContactShapesType::LevelSetLevelSet:
 				UpdateLevelsetLevelsetConstraint<UpdateType>(CullDistance, Constraint);
@@ -2193,8 +2146,7 @@ namespace Chaos
 			}
 		}
 
-		template<typename T, int d>
-		void ConstructConstraintsImpl(TGeometryParticleHandle<T, d>* Particle0, TGeometryParticleHandle<T, d>* Particle1, const FImplicitObject* Implicit0, const FImplicitObject* Implicit1, const TRigidTransform<T, d>& LocalTransform0, const TRigidTransform<T, d>& LocalTransform1, const T CullDistance, const FCollisionContext& Context, FCollisionConstraintsArray& NewConstraints)
+		void ConstructConstraintsImpl(TGeometryParticleHandle<FReal, 3>* Particle0, TGeometryParticleHandle<FReal, 3>* Particle1, const FImplicitObject* Implicit0, const FImplicitObject* Implicit1, const FRigidTransform3& LocalTransform0, const FRigidTransform3& LocalTransform1, const FReal CullDistance, const FCollisionContext& Context, FCollisionConstraintsArray& NewConstraints)
 		{
 			// @todo(chaos): We use GetInnerType here because TriMeshes are left with their "Instanced" wrapper, unlike all other instanced implicits. Should we strip the instance on Tri Mesh too?
 			EImplicitObjectType Implicit0Type = Implicit0 ? GetInnerType(Implicit0->GetCollisionType()) : ImplicitObjectType::Unknown;
@@ -2202,8 +2154,8 @@ namespace Chaos
 			bool bIsConvex0 = Implicit0 && Implicit0->IsConvex();
 			bool bIsConvex1 = Implicit1 && Implicit1->IsConvex();
 
-			T LengthCCD = 0.0f;
-			TVector<T, d> DirCCD(0.0f);
+			FReal LengthCCD = 0.0f;
+			FVec3 DirCCD(0.0f);
 			bool bUseCCD = UseCCD(Particle0, Particle1, Implicit0, DirCCD, LengthCCD);
 
 			if (Implicit0Type == TBox<FReal, 3>::StaticType() && Implicit1Type == TBox<FReal, 3>::StaticType())
@@ -2389,18 +2341,17 @@ namespace Chaos
 
 		// Run collision detection for the specified constraint to update the nearest contact point.
 		// NOTE: Transforms are world space particle transforms
-		template<ECollisionUpdateType UpdateType, typename T, int d>
-		void UpdateConstraintFromGeometry(TRigidBodyPointContactConstraint<T, d>& Constraint, const TRigidTransform<T, d>& ParticleTransform0, const TRigidTransform<T, d>& ParticleTransform1, const T CullDistance)
+		template<ECollisionUpdateType UpdateType>
+		void UpdateConstraintFromGeometry(FRigidBodyPointContactConstraint& Constraint, const FRigidTransform3& ParticleTransform0, const FRigidTransform3& ParticleTransform1, const FReal CullDistance)
 		{
-			const TRigidTransform<T, d> WorldTransform0 = Constraint.ImplicitTransform[0] * ParticleTransform0;
-			const TRigidTransform<T, d> WorldTransform1 = Constraint.ImplicitTransform[1] * ParticleTransform1;
+			const FRigidTransform3 WorldTransform0 = Constraint.ImplicitTransform[0] * ParticleTransform0;
+			const FRigidTransform3 WorldTransform1 = Constraint.ImplicitTransform[1] * ParticleTransform1;
 			UpdateConstraintFromGeometryImpl<UpdateType>(Constraint, WorldTransform0, WorldTransform1, CullDistance);
 		}
 
 		// Select the best manifold point as the new contact point.
 		// NOTE: Transforms are world space particle transforms
-		template<class T, int d>
-		void UpdateConstraintFromManifold(TRigidBodyMultiPointContactConstraint<T, d>& Constraint, const TRigidTransform<T, d>& ParticleTransform0, const TRigidTransform<T, d>& ParticleTransform1, const T CullDistance)
+		void UpdateConstraintFromManifold(FRigidBodyMultiPointContactConstraint& Constraint, const FRigidTransform3& ParticleTransform0, const FRigidTransform3& ParticleTransform1, const FReal CullDistance)
 		{
 			const FRigidTransform3 WorldTransform0 = Constraint.ImplicitTransform[0] * ParticleTransform0;
 			const FRigidTransform3 WorldTransform1 = Constraint.ImplicitTransform[1] * ParticleTransform1;
@@ -2535,8 +2486,7 @@ namespace Chaos
 		}
 
 
-		template<typename T, int d>
-		void ConstructConstraints(TGeometryParticleHandle<T, d>* Particle0, TGeometryParticleHandle<T, d>* Particle1, const FImplicitObject* Implicit0, const FImplicitObject* Implicit1, const TRigidTransform<T, d>& LocalTransform0, const TRigidTransform<T, d>& LocalTransform1, const T CullDistance, const FCollisionContext& Context, FCollisionConstraintsArray& NewConstraints)
+		void ConstructConstraints(TGeometryParticleHandle<FReal, 3>* Particle0, TGeometryParticleHandle<FReal, 3>* Particle1, const FImplicitObject* Implicit0, const FImplicitObject* Implicit1, const FRigidTransform3& LocalTransform0, const FRigidTransform3& LocalTransform1, const FReal CullDistance, const FCollisionContext& Context, FCollisionConstraintsArray& NewConstraints)
 		{
 			EImplicitObjectType Implicit0Type = Implicit0 ? GetInnerType(Implicit0->GetType()) : ImplicitObjectType::Unknown;
 			EImplicitObjectType Implicit1Type = Implicit1 ? GetInnerType(Implicit1->GetType()) : ImplicitObjectType::Unknown;
@@ -2552,26 +2502,26 @@ namespace Chaos
 			EImplicitObjectType Implicit1OuterType = Implicit1->GetType();
 
 			// Handle transform wrapper shape
-			if ((Implicit0OuterType == TImplicitObjectTransformed<T, d>::StaticType()) && (Implicit1OuterType == TImplicitObjectTransformed<T, d>::StaticType()))
+			if ((Implicit0OuterType == TImplicitObjectTransformed<FReal, 3>::StaticType()) && (Implicit1OuterType == TImplicitObjectTransformed<FReal, 3>::StaticType()))
 			{
 				const TImplicitObjectTransformed<FReal, 3>* TransformedImplicit0 = Implicit0->template GetObject<const TImplicitObjectTransformed<FReal, 3>>();
 				const TImplicitObjectTransformed<FReal, 3>* TransformedImplicit1 = Implicit1->template GetObject<const TImplicitObjectTransformed<FReal, 3>>();
-				TRigidTransform<T, d> TransformedTransform0 = TransformedImplicit0->GetTransform() * LocalTransform0;
-				TRigidTransform<T, d> TransformedTransform1 = TransformedImplicit1->GetTransform() * LocalTransform1;
+				FRigidTransform3 TransformedTransform0 = TransformedImplicit0->GetTransform() * LocalTransform0;
+				FRigidTransform3 TransformedTransform1 = TransformedImplicit1->GetTransform() * LocalTransform1;
 				ConstructConstraints(Particle0, Particle1, TransformedImplicit0->GetTransformedObject(), TransformedImplicit1->GetTransformedObject(), TransformedTransform0, TransformedTransform1, CullDistance, Context, NewConstraints);
 				return;
 			}
-			else if (Implicit0OuterType == TImplicitObjectTransformed<T, d>::StaticType())
+			else if (Implicit0OuterType == TImplicitObjectTransformed<FReal, 3>::StaticType())
 			{
 				const TImplicitObjectTransformed<FReal, 3>* TransformedImplicit0 = Implicit0->template GetObject<const TImplicitObjectTransformed<FReal, 3>>();
-				TRigidTransform<T, d> TransformedTransform0 = TransformedImplicit0->GetTransform() * LocalTransform0;
+				FRigidTransform3 TransformedTransform0 = TransformedImplicit0->GetTransform() * LocalTransform0;
 				ConstructConstraints(Particle0, Particle1, TransformedImplicit0->GetTransformedObject(), Implicit1, TransformedTransform0, LocalTransform1, CullDistance, Context, NewConstraints);
 				return;
 			}
-			else if (Implicit1OuterType == TImplicitObjectTransformed<T, d>::StaticType())
+			else if (Implicit1OuterType == TImplicitObjectTransformed<FReal, 3>::StaticType())
 			{
 				const TImplicitObjectTransformed<FReal, 3>* TransformedImplicit1 = Implicit1->template GetObject<const TImplicitObjectTransformed<FReal, 3>>();
-				TRigidTransform<T, d> TransformedTransform1 = TransformedImplicit1->GetTransform() * LocalTransform1;
+				FRigidTransform3 TransformedTransform1 = TransformedImplicit1->GetTransform() * LocalTransform1;
 				ConstructConstraints(Particle0, Particle1, Implicit0, TransformedImplicit1->GetTransformedObject(), LocalTransform0, TransformedTransform1, CullDistance, Context, NewConstraints);
 				return;
 			}
@@ -2650,7 +2600,7 @@ namespace Chaos
 			FReal CullDistance = 0.0f;
 
 			FReal TOI = 1.0f;
-			TContactPoint<FReal> Contact = GJKImplicitSweptContactPoint<GeometryA>(A, AStartTransform, B, BTransform, Dir, Length, CullDistance, TOI);
+			FContactPoint Contact = GJKImplicitSweptContactPoint<GeometryA>(A, AStartTransform, B, BTransform, Dir, Length, CullDistance, TOI);
 			if (Contact.Phi < OutPhi)
 			{
 				OutPhi = Contact.Phi;
@@ -2664,7 +2614,7 @@ namespace Chaos
 		template <typename GeometryA>
 		bool GetPairTOIHackTriMeshImpl(const FImplicitObject& A, const FRigidTransform3& AStartTransform, const FImplicitObject* B, const FRigidTransform3& BTransform, const FVec3& Dir, const FReal Length, FReal& OutTOI, FVec3& OutNormal, FReal& OutPhi)
 		{
-			TContactPoint<FReal> Contact;
+			FContactPoint Contact;
 			FReal TOI = 0.0f;
 			FReal CullDistance = 0.0f;
 
@@ -2692,7 +2642,7 @@ namespace Chaos
 			FReal CullDistance = 0.0f;
 			FReal TOI = 0.0f;
 
-			TContactPoint<FReal> Contact = ConvexConvexContactPointSwept(A, AStartTransform, B, BTransform, Dir, Length, CullDistance, TOI);
+			FContactPoint Contact = ConvexConvexContactPointSwept(A, AStartTransform, B, BTransform, Dir, Length, CullDistance, TOI);
 			if (Contact.Phi < OutPhi)
 			{
 				OutPhi = Contact.Phi;
@@ -2880,69 +2830,11 @@ namespace Chaos
 			return GetTOIHackImpl(Particle0, Particle1, Particle0->Geometry().Get(), Particle1->Geometry().Get(), Particle0StartTransform, Collisions::GetTransform(Particle1), OutTOI, OutNormal, OutPhi);
 		}
 
-		template void UpdateBoxBoxConstraint<float, 3>(const TAABB<float, 3>& Box1, const TRigidTransform<float, 3>& Box1Transform, const TAABB<float, 3>& Box2, const TRigidTransform<float, 3>& Box2Transform, const float CullDistance, TRigidBodyPointContactConstraint<float, 3>& Constraint);
-		template void UpdateBoxBoxManifold<float, 3>(TCollisionConstraintBase<float, 3>&  Constraint, const TRigidTransform<float, 3>& ATM, const TRigidTransform<float, 3>& BTM, const float CullDistance);
-		template void ConstructBoxBoxConstraints<float, 3>(TGeometryParticleHandle<float, 3>* Particle0, TGeometryParticleHandle<float, 3>* Particle1, const FImplicitObject* Implicit0, const FImplicitObject* Implicit1, const TRigidTransform<float, 3>& Transform0, const TRigidTransform<float, 3>& Transform1, const float CullDistance, FCollisionConstraintsArray& NewConstraints);
-
-		template void UpdateBoxHeightFieldConstraint<float, 3>(const TAABB<float, 3>& A, const TRigidTransform<float, 3>& ATransform, const THeightField<float>& B, const TRigidTransform<float, 3>& BTransform, const float CullDistance, TRigidBodyPointContactConstraint<float, 3>& Constraint);
-		template void UpdateBoxHeightFieldManifold<float, 3>(TCollisionConstraintBase<float, 3>&  Constraint, const TRigidTransform<float, 3>& ATM, const TRigidTransform<float, 3>& BTM, const float CullDistance);
-		template void ConstructBoxHeightFieldConstraints<float, 3>(TGeometryParticleHandle<float, 3>* Particle0, TGeometryParticleHandle<float, 3>* Particle1, const FImplicitObject* Implicit0, const FImplicitObject* Implicit1, const TRigidTransform<float, 3>& Transform0, const TRigidTransform<float, 3>& Transform1, const float CullDistance, FCollisionConstraintsArray& NewConstraints);
-
-		template void UpdateBoxPlaneConstraint<float, 3>(const TAABB<float, 3>& Box, const TRigidTransform<float, 3>& BoxTransform, const TPlane<float, 3>& Plane, const TRigidTransform<float, 3>& PlaneTransform, const float CullDistance, TRigidBodyPointContactConstraint<float, 3>& Constraint);
-		template void UpdateBoxPlaneManifold<float, 3>(TCollisionConstraintBase<float, 3>&  Constraint, const TRigidTransform<float, 3>& ATM, const TRigidTransform<float, 3>& BTM, const float CullDistance);
-		template void ConstructBoxPlaneConstraints<float, 3>(TGeometryParticleHandle<float, 3>* Particle0, TGeometryParticleHandle<float, 3>* Particle1, const FImplicitObject* Implicit0, const FImplicitObject* Implicit1, const TRigidTransform<float, 3>& Transform0, const TRigidTransform<float, 3>& Transform1, const float CullDistance, FCollisionConstraintsArray& NewConstraints);
-
-		template void UpdateSphereSphereConstraint<float, 3>(const TSphere<float, 3>& Sphere1, const TRigidTransform<float, 3>& Sphere1Transform, const TSphere<float, 3>& Sphere2, const TRigidTransform<float, 3>& Sphere2Transform, const float CullDistance, TRigidBodyPointContactConstraint<float, 3>& Constraint);
-		template void UpdateSphereSphereManifold<float, 3>(TCollisionConstraintBase<float, 3>&  Constraint, const TRigidTransform<float, 3>& ATM, const TRigidTransform<float, 3>& BTM, const float CullDistance);
-		template void ConstructSphereSphereConstraints<float, 3>(TGeometryParticleHandle<float, 3>* Particle0, TGeometryParticleHandle<float, 3>* Particle1, const FImplicitObject* Implicit0, const FImplicitObject* Implicit1, const TRigidTransform<float, 3>& Transform0, const TRigidTransform<float, 3>& Transform1, const float CullDistance, FCollisionConstraintsArray& NewConstraints);
-
-		template void UpdateSphereHeightFieldConstraint<float, 3>(const TSphere<float, 3>& A, const TRigidTransform<float, 3>& ATransform, const THeightField<float>& B, const TRigidTransform<float, 3>& BTransform, const float CullDistance, TRigidBodyPointContactConstraint<float, 3>& Constraint);
-		template void UpdateSphereHeightFieldManifold<float, 3>(TCollisionConstraintBase<float, 3>&  Constraint, const TRigidTransform<float, 3>& ATM, const TRigidTransform<float, 3>& BTM, const float CullDistance);
-		template void ConstructSphereHeightFieldConstraints<float, 3>(TGeometryParticleHandle<float, 3>* Particle0, TGeometryParticleHandle<float, 3>* Particle1, const FImplicitObject* Implicit0, const FImplicitObject* Implicit1, const TRigidTransform<float, 3>& Transform0, const TRigidTransform<float, 3>& Transform1, const float CullDistance, FCollisionConstraintsArray& NewConstraints);
-
-		template void UpdateSpherePlaneConstraint<float, 3>(const TSphere<float, 3>& Sphere, const TRigidTransform<float, 3>& SphereTransform, const TPlane<float, 3>& Plane, const TRigidTransform<float, 3>& PlaneTransform, const float CullDistance, TRigidBodyPointContactConstraint<float, 3>& Constraint);
-		template void UpdateSpherePlaneManifold<float, 3>(TCollisionConstraintBase<float, 3>&  Constraint, const TRigidTransform<float, 3>& ATM, const TRigidTransform<float, 3>& BTM, const float CullDistance);
-		template void ConstructSpherePlaneConstraints<float, 3>(TGeometryParticleHandle<float, 3>* Particle0, TGeometryParticleHandle<float, 3>* Particle1, const FImplicitObject* Implicit0, const FImplicitObject* Implicit1, const TRigidTransform<float, 3>& Transform0, const TRigidTransform<float, 3>& Transform1, const float CullDistance, FCollisionConstraintsArray& NewConstraints);
-
-		template void UpdateSphereBoxConstraint<float, 3>(const TSphere<float, 3>& Sphere, const TRigidTransform<float, 3>& SphereTransform, const TAABB<float, 3>& Box, const TRigidTransform<float, 3>& BoxTransform, const float CullDistance, TRigidBodyPointContactConstraint<float, 3>& Constraint);
-		template void UpdateSphereBoxManifold<float, 3>(TCollisionConstraintBase<float, 3>&  Constraint, const TRigidTransform<float, 3>& ATM, const TRigidTransform<float, 3>& BTM, const float CullDistance);
-		template void ConstructSphereBoxConstraints<float, 3>(TGeometryParticleHandle<float, 3>* Particle0, TGeometryParticleHandle<float, 3>* Particle1, const FImplicitObject* Implicit0, const FImplicitObject* Implicit1, const TRigidTransform<float, 3>& Transform0, const TRigidTransform<float, 3>& Transform1, const float CullDistance, FCollisionConstraintsArray& NewConstraints);
-
-		template void UpdateSphereCapsuleConstraint<float, 3>(const TSphere<float, 3>& Sphere, const TRigidTransform<float, 3>& SphereTransform, const TCapsule<float>& Box, const TRigidTransform<float, 3>& BoxTransform, const float CullDistance, TRigidBodyPointContactConstraint<float, 3>& Constraint);
-		template void UpdateSphereCapsuleManifold<float, 3>(TCollisionConstraintBase<float, 3>&  Constraint, const TRigidTransform<float, 3>& ATM, const TRigidTransform<float, 3>& BTM, const float CullDistance);
-		template void ConstructSphereCapsuleConstraints<float, 3>(TGeometryParticleHandle<float, 3>* Particle0, TGeometryParticleHandle<float, 3>* Particle1, const FImplicitObject* Implicit0, const FImplicitObject* Implicit1, const TRigidTransform<float, 3>& Transform0, const TRigidTransform<float, 3>& Transform1, const float CullDistance, FCollisionConstraintsArray& NewConstraints);
-
-		template void UpdateCapsuleCapsuleConstraint<float, 3>(const TCapsule<float>& A, const TRigidTransform<float, 3>& ATransform, const TCapsule<float>& B, const TRigidTransform<float, 3>& BTransform, const float CullDistance, TRigidBodyPointContactConstraint<float, 3>& Constraint);
-		template void UpdateCapsuleCapsuleManifold<float, 3>(TCollisionConstraintBase<float, 3>&  Constraint, const TRigidTransform<float, 3>& ATM, const TRigidTransform<float, 3>& BTM, const float CullDistance);
-		template void ConstructCapsuleCapsuleConstraints<float, 3>(TGeometryParticleHandle<float, 3>* Particle0, TGeometryParticleHandle<float, 3>* Particle1, const FImplicitObject* Implicit0, const FImplicitObject* Implicit1, const TRigidTransform<float, 3>& Transform0, const TRigidTransform<float, 3>& Transform1, const float CullDistance, FCollisionConstraintsArray& NewConstraints);
-
-		template void UpdateCapsuleBoxConstraint<float, 3>(const TCapsule<float>& A, const TRigidTransform<float, 3>& ATransform, const TAABB<float, 3>& B, const TRigidTransform<float, 3>& BTransform, const float CullDistance, TRigidBodyPointContactConstraint<float, 3>& Constraint);
-		template void ConstructCapsuleBoxConstraints<float, 3>(TGeometryParticleHandle<float, 3>* Particle0, TGeometryParticleHandle<float, 3>* Particle1, const FImplicitObject* Implicit0, const FImplicitObject* Implicit1, const TRigidTransform<float, 3>& Transform0, const TRigidTransform<float, 3>& Transform1, const float CullDistance, const FCollisionContext& Context, FCollisionConstraintsArray& NewConstraints);
-
-		template void UpdateCapsuleHeightFieldConstraint<float, 3>(const TCapsule<float>& A, const TRigidTransform<float, 3>& ATransform, const THeightField<float>& B, const TRigidTransform<float, 3>& BTransform, const float CullDistance, TRigidBodyPointContactConstraint<float, 3>& Constraint);
-		template void UpdateCapsuleHeightFieldManifold<float, 3>(TCollisionConstraintBase<float, 3>&  Constraint, const TRigidTransform<float, 3>& ATM, const TRigidTransform<float, 3>& BTM, const float CullDistance);
-		template void ConstructCapsuleHeightFieldConstraints<float, 3>(TGeometryParticleHandle<float, 3>* Particle0, TGeometryParticleHandle<float, 3>* Particle1, const FImplicitObject* Implicit0, const FImplicitObject* Implicit1, const TRigidTransform<float, 3>& Transform0, const TRigidTransform<float, 3>& Transform1, const float CullDistance, FCollisionConstraintsArray& NewConstraints);
-
-		template void UpdateConvexConvexConstraint<float, 3>(const FImplicitObject& A, const TRigidTransform<float, 3>& ATM, const FImplicitObject& B, const TRigidTransform<float, 3>& BTM, const float CullDistance, TCollisionConstraintBase<float, 3>& Constraint);
-		template void UpdateConvexConvexManifold<float, 3>(TCollisionConstraintBase<float, 3>&  Constraint, const TRigidTransform<float, 3>& ATM, const TRigidTransform<float, 3>& BTM, const float CullDistance);
-		template void ConstructConvexConvexConstraints<float, 3>(TGeometryParticleHandle<float, 3>* Particle0, TGeometryParticleHandle<float, 3>* Particle1, const FImplicitObject* Implicit0, const FImplicitObject* Implicit1, const TRigidTransform<float, 3>& Transform0, const TRigidTransform<float, 3>& Transform1, const float CullDistance, FCollisionConstraintsArray& NewConstraints);
-
-		template void UpdateConvexHeightFieldConstraint<float, 3>(const FImplicitObject& A, const TRigidTransform<float, 3>& ATransform, const THeightField<float>& B, const TRigidTransform<float, 3>& BTransform, const float CullDistance, TRigidBodyPointContactConstraint<float, 3>& Constraint);
-		template void UpdateConvexHeightFieldManifold<float, 3>(TCollisionConstraintBase<float, 3>&  Constraint, const TRigidTransform<float, 3>& ATM, const TRigidTransform<float, 3>& BTM, const float CullDistance);
-		template void ConstructConvexHeightFieldConstraints<float, 3>(TGeometryParticleHandle<float, 3>* Particle0, TGeometryParticleHandle<float, 3>* Particle1, const FImplicitObject* Implicit0, const FImplicitObject* Implicit1, const TRigidTransform<float, 3>& Transform0, const TRigidTransform<float, 3>& Transform1, const float CullDistance, FCollisionConstraintsArray& NewConstraints);
-
-		template void UpdateLevelsetLevelsetConstraint<ECollisionUpdateType::Any, float, 3>(const float CullDistance, TRigidBodyPointContactConstraint<float, 3>& Constraint);
-		template void UpdateLevelsetLevelsetConstraint<ECollisionUpdateType::Deepest, float, 3>(const float CullDistance, TRigidBodyPointContactConstraint<float, 3>& Constraint);
-		template void UpdateLevelsetLevelsetManifold<float, 3>(TCollisionConstraintBase<float, 3>&  Constraint, const TRigidTransform<float, 3>& ATM, const TRigidTransform<float, 3>& BTM, const float CullDistance);
-		template void ConstructLevelsetLevelsetConstraints<float, 3>(TGeometryParticleHandle<float, 3>* Particle0, TGeometryParticleHandle<float, 3>* Particle1, const FImplicitObject* Implicit0, const FImplicitObject* Implicit1, const TRigidTransform<float, 3>& Transform0, const TRigidTransform<float, 3>& Transform1, const float CullDistance, FCollisionConstraintsArray& NewConstraints);
+		template void UpdateLevelsetLevelsetConstraint<ECollisionUpdateType::Any>(const float CullDistance, FRigidBodyPointContactConstraint& Constraint);
+		template void UpdateLevelsetLevelsetConstraint<ECollisionUpdateType::Deepest>(const float CullDistance, FRigidBodyPointContactConstraint& Constraint);
 
-		template void UpdateSingleShotManifold<float, 3>(TRigidBodyMultiPointContactConstraint<float, 3>&  Constraint, const TRigidTransform<float, 3>& Transform0, const TRigidTransform<float, 3>& Transform1, const float CullDistance);
-		template void UpdateIterativeManifold<float, 3>(TRigidBodyMultiPointContactConstraint<float, 3>&  Constraint, const TRigidTransform<float, 3>& Transform0, const TRigidTransform<float, 3>& Transform1, const float CullDistance);
-		template void UpdateManifold<float, 3>(TRigidBodyMultiPointContactConstraint<float, 3>& Constraint, const TRigidTransform<float, 3>& ATM, const TRigidTransform<float, 3>& BTM, const float CullDistance, const FCollisionContext& Context);
-		template void UpdateConstraintFromGeometry<ECollisionUpdateType::Any, float, 3>(TRigidBodyPointContactConstraint<float, 3>& ConstraintBase, const TRigidTransform<float, 3>& Transform0, const TRigidTransform<float, 3>& Transform1, const float CullDistance);
-		template void UpdateConstraintFromGeometry<ECollisionUpdateType::Deepest, float, 3>(TRigidBodyPointContactConstraint<float, 3>& ConstraintBase, const TRigidTransform<float, 3>& Transform0, const TRigidTransform<float, 3>& Transform1, const float CullDistance);
-		template void UpdateConstraintFromManifold(TRigidBodyMultiPointContactConstraint<float, 3>& Constraint, const TRigidTransform<float, 3>& Transform0, const TRigidTransform<float, 3>& Transform1, const float CullDistance);
-		template void ConstructConstraints<float, 3>(TGeometryParticleHandle<float, 3>* Particle0, TGeometryParticleHandle<float, 3>* Particle1, const FImplicitObject* Implicit0, const FImplicitObject* Implicit1, const TRigidTransform<float, 3>& Transform0, const TRigidTransform<float, 3>& Transform1, const float CullDistance, const FCollisionContext& Context, FCollisionConstraintsArray& NewConstraints);
+		template void UpdateConstraintFromGeometry<ECollisionUpdateType::Any>(FRigidBodyPointContactConstraint& ConstraintBase, const FRigidTransform3& Transform0, const FRigidTransform3& Transform1, const float CullDistance);
+		template void UpdateConstraintFromGeometry<ECollisionUpdateType::Deepest>(FRigidBodyPointContactConstraint& ConstraintBase, const FRigidTransform3& Transform0, const FRigidTransform3& Transform1, const float CullDistance);
 
 	} // Collisions
 
