Chaos

diff --git a/Engine/Source/Runtime/Experimental/Chaos/Private/Chaos/CollisionResolutionTypes.cpp b/Engine/Source/Runtime/Experimental/Chaos/Private/Chaos/CollisionResolutionTypes.cpp
--- a/Engine/Source/Runtime/Experimental/Chaos/Private/Chaos/CollisionResolutionTypes.cpp
+++ b/Engine/Source/Runtime/Experimental/Chaos/Private/Chaos/CollisionResolutionTypes.cpp
@@ -5,6 +5,4 @@
 
 namespace Chaos
 {
-	template class CHAOS_API TPBDCollisionConstraintHandle<float, 3>; 
-	template class CHAOS_API TCollisionContact<float, 3>;
 };
diff --git a/Engine/Source/Runtime/Experimental/Chaos/Private/Chaos/CollisionResolutionUtil.cpp b/Engine/Source/Runtime/Experimental/Chaos/Private/Chaos/CollisionResolutionUtil.cpp
--- a/Engine/Source/Runtime/Experimental/Chaos/Private/Chaos/CollisionResolutionUtil.cpp
+++ b/Engine/Source/Runtime/Experimental/Chaos/Private/Chaos/CollisionResolutionUtil.cpp
@@ -21,21 +21,19 @@ namespace Chaos
 	namespace Collisions
 	{
 
-		template<typename T, int d>
-		TRigidTransform<T, d> GetTransform(const TGeometryParticleHandle<T, d>* Particle)
+		FRigidTransform3 GetTransform(const TGeometryParticleHandle<FReal, 3>* Particle)
 		{
-			TGenericParticleHandle<T, d> Generic = const_cast<TGeometryParticleHandle<T, d>*>(Particle);
-			return TRigidTransform<T, d>(Generic->P(), Generic->Q());
+			TGenericParticleHandle<FReal, 3> Generic = const_cast<TGeometryParticleHandle<FReal, 3>*>(Particle);
+			return FRigidTransform3(Generic->P(), Generic->Q());
 		}
 
 		// @todo(ccaulfield): This is duplicated in JointConstraints - move to a utility file
-		template<typename T>
-		PMatrix<T, 3, 3> ComputeFactorMatrix3(const TVector<T, 3>& V, const PMatrix<T, 3, 3>& M, const T& Im)
+		FMatrix33 ComputeFactorMatrix3(const FVec3& V, const FMatrix33& M, const FReal Im)
 		{
 			// Rigid objects rotational contribution to the impulse.
 			// Vx*M*VxT+Im
 			ensure(Im > FLT_MIN);
-				return PMatrix<T, 3, 3>(
+				return FMatrix33(
 					-V[2] * (-V[2] * M.M[1][1] + V[1] * M.M[2][1]) + V[1] * (-V[2] * M.M[2][1] + V[1] * M.M[2][2]) + Im,
 					V[2] * (-V[2] * M.M[1][0] + V[1] * M.M[2][0]) - V[0] * (-V[2] * M.M[2][1] + V[1] * M.M[2][2]),
 					-V[1] * (-V[2] * M.M[1][0] + V[1] * M.M[2][0]) + V[0] * (-V[2] * M.M[1][1] + V[1] * M.M[2][1]),
@@ -44,48 +42,46 @@ namespace Chaos
 					-V[1] * (-V[1] * M.M[0][0] + V[0] * M.M[1][0]) + V[0] * (-V[1] * M.M[1][0] + V[0] * M.M[1][1]) + Im);
 		}
 
-		template<typename T, int d>
-		TVector<T, d> GetEnergyClampedImpulse(const TPBDRigidParticleHandle<T, d>* PBDRigid0, const TPBDRigidParticleHandle<T, d>* PBDRigid1, const TVector<T, d>& Impulse, const TVector<T, d>& VectorToPoint1, const TVector<T, d>& VectorToPoint2, const TVector<T, d>& Velocity1, const TVector<T, d>& Velocity2)
+		FVec3 GetEnergyClampedImpulse(const TPBDRigidParticleHandle<FReal, 3>* PBDRigid0, const TPBDRigidParticleHandle<FReal, 3>* PBDRigid1, const FVec3& Impulse, const FVec3& VectorToPoint1, const FVec3& VectorToPoint2, const FVec3& Velocity1, const FVec3& Velocity2)
 		{
 			const bool bIsRigidDynamic0 = PBDRigid0 && PBDRigid0->ObjectState() == EObjectStateType::Dynamic;
 			const bool bIsRigidDynamic1 = PBDRigid1 && PBDRigid1->ObjectState() == EObjectStateType::Dynamic;
 
-			TVector<T, d> Jr0, Jr1, IInvJr0, IInvJr1;
-			T ImpulseRatioNumerator0 = 0, ImpulseRatioNumerator1 = 0, ImpulseRatioDenom0 = 0, ImpulseRatioDenom1 = 0;
-			T ImpulseSize = Impulse.SizeSquared();
-			TVector<T, d> KinematicVelocity = !bIsRigidDynamic0 ? Velocity1 : !bIsRigidDynamic1 ? Velocity2 : TVector<T, d>(0);
+			FVec3 Jr0, Jr1, IInvJr0, IInvJr1;
+			FReal ImpulseRatioNumerator0 = 0, ImpulseRatioNumerator1 = 0, ImpulseRatioDenom0 = 0, ImpulseRatioDenom1 = 0;
+			FReal ImpulseSize = Impulse.SizeSquared();
+			FVec3 KinematicVelocity = !bIsRigidDynamic0 ? Velocity1 : !bIsRigidDynamic1 ? Velocity2 : FVec3(0);
 			if (bIsRigidDynamic0)
 			{
-				Jr0 = TVector<T, d>::CrossProduct(VectorToPoint1, Impulse);
+				Jr0 = FVec3::CrossProduct(VectorToPoint1, Impulse);
 				IInvJr0 = PBDRigid0->Q().RotateVector(PBDRigid0->InvI() * PBDRigid0->Q().UnrotateVector(Jr0));
-				ImpulseRatioNumerator0 = TVector<T, d>::DotProduct(Impulse, PBDRigid0->V() - KinematicVelocity) + TVector<T, d>::DotProduct(IInvJr0, PBDRigid0->W());
-				ImpulseRatioDenom0 = ImpulseSize / PBDRigid0->M() + TVector<T, d>::DotProduct(Jr0, IInvJr0);
+				ImpulseRatioNumerator0 = FVec3::DotProduct(Impulse, PBDRigid0->V() - KinematicVelocity) + FVec3::DotProduct(IInvJr0, PBDRigid0->W());
+				ImpulseRatioDenom0 = ImpulseSize / PBDRigid0->M() + FVec3::DotProduct(Jr0, IInvJr0);
 			}
 			if (bIsRigidDynamic1)
 			{
-				Jr1 = TVector<T, d>::CrossProduct(VectorToPoint2, Impulse);
+				Jr1 = FVec3::CrossProduct(VectorToPoint2, Impulse);
 				IInvJr1 = PBDRigid1->Q().RotateVector(PBDRigid1->InvI() * PBDRigid1->Q().UnrotateVector(Jr1));
-				ImpulseRatioNumerator1 = TVector<T, d>::DotProduct(Impulse, PBDRigid1->V() - KinematicVelocity) + TVector<T, d>::DotProduct(IInvJr1, PBDRigid1->W());
-				ImpulseRatioDenom1 = ImpulseSize / PBDRigid1->M() + TVector<T, d>::DotProduct(Jr1, IInvJr1);
+				ImpulseRatioNumerator1 = FVec3::DotProduct(Impulse, PBDRigid1->V() - KinematicVelocity) + FVec3::DotProduct(IInvJr1, PBDRigid1->W());
+				ImpulseRatioDenom1 = ImpulseSize / PBDRigid1->M() + FVec3::DotProduct(Jr1, IInvJr1);
 			}
-			T Numerator = -2 * (ImpulseRatioNumerator0 - ImpulseRatioNumerator1);
+			FReal Numerator = -2 * (ImpulseRatioNumerator0 - ImpulseRatioNumerator1);
 			if (Numerator < 0)
 			{
-				return TVector<T, d>(0);
+				return FVec3(0);
 			}
 			check(Numerator >= 0);
-			T Denominator = ImpulseRatioDenom0 + ImpulseRatioDenom1;
+			FReal Denominator = ImpulseRatioDenom0 + ImpulseRatioDenom1;
 			return Numerator < Denominator ? (Impulse * Numerator / Denominator) : Impulse;
 		}
 
-		template <typename T, int d>
-		bool SampleObjectHelper(const FImplicitObject& Object, const TRigidTransform<T, d>& ObjectTransform, const TRigidTransform<T, d>& SampleToObjectTransform, const TVector<T, d>& SampleParticle, T Thickness, TRigidBodyPointContactConstraint<float, 3>& Constraint)
+		bool SampleObjectHelper(const FImplicitObject& Object, const FRigidTransform3& ObjectTransform, const FRigidTransform3& SampleToObjectTransform, const FVec3& SampleParticle, FReal Thickness, FRigidBodyPointContactConstraint& Constraint)
 		{
-			TVector<T, d> LocalPoint = SampleToObjectTransform.TransformPositionNoScale(SampleParticle);
-			TVector<T, d> LocalNormal;
-			T LocalPhi = Object.PhiWithNormal(LocalPoint, LocalNormal);
+			FVec3 LocalPoint = SampleToObjectTransform.TransformPositionNoScale(SampleParticle);
+			FVec3 LocalNormal;
+			FReal LocalPhi = Object.PhiWithNormal(LocalPoint, LocalNormal);
 
-			TCollisionContact<T, d> & Contact = Constraint.Manifold;
+			FCollisionContact & Contact = Constraint.Manifold;
 			if (LocalPhi < Contact.Phi)
 			{
 				Contact.Phi = LocalPhi;
@@ -96,14 +92,13 @@ namespace Chaos
 			return false;
 		}
 
-		template <typename T, int d>
-		bool SampleObjectNoNormal(const FImplicitObject& Object, const TRigidTransform<T, d>& ObjectTransform, const TRigidTransform<T, d>& SampleToObjectTransform, const TVector<T, d>& SampleParticle, T Thickness, TRigidBodyPointContactConstraint<float, 3>& Constraint)
+		bool SampleObjectNoNormal(const FImplicitObject& Object, const FRigidTransform3& ObjectTransform, const FRigidTransform3& SampleToObjectTransform, const FVec3& SampleParticle, FReal Thickness, FRigidBodyPointContactConstraint& Constraint)
 		{
-			TVector<T, d> LocalPoint = SampleToObjectTransform.TransformPositionNoScale(SampleParticle);
-			TVector<T, d> LocalNormal;
-			T LocalPhi = Object.PhiWithNormal(LocalPoint, LocalNormal);
+			FVec3 LocalPoint = SampleToObjectTransform.TransformPositionNoScale(SampleParticle);
+			FVec3 LocalNormal;
+			FReal LocalPhi = Object.PhiWithNormal(LocalPoint, LocalNormal);
 
-			TCollisionContact<T, d> & Contact = Constraint.Manifold;
+			FCollisionContact & Contact = Constraint.Manifold;
 			if (LocalPhi < Contact.Phi)
 			{
 				Contact.Phi = LocalPhi;
@@ -112,15 +107,14 @@ namespace Chaos
 			return false;
 		}
 
-		template <typename T, int d>
-		bool SampleObjectNormalAverageHelper(const FImplicitObject& Object, const TRigidTransform<T, d>& ObjectTransform, const TRigidTransform<T, d>& SampleToObjectTransform, const TVector<T, d>& SampleParticle, T Thickness, T& TotalThickness, TRigidBodyPointContactConstraint<float, 3>& Constraint)
+		bool SampleObjectNormalAverageHelper(const FImplicitObject& Object, const FRigidTransform3& ObjectTransform, const FRigidTransform3& SampleToObjectTransform, const FVec3& SampleParticle, FReal Thickness, FReal& TotalThickness, FRigidBodyPointContactConstraint& Constraint)
 		{
-			TVector<T, d> LocalPoint = SampleToObjectTransform.TransformPositionNoScale(SampleParticle);
-			TVector<T, d> LocalNormal;
-			T LocalPhi = Object.PhiWithNormal(LocalPoint, LocalNormal);
-			T LocalThickness = LocalPhi - Thickness;
+			FVec3 LocalPoint = SampleToObjectTransform.TransformPositionNoScale(SampleParticle);
+			FVec3 LocalNormal;
+			FReal LocalPhi = Object.PhiWithNormal(LocalPoint, LocalNormal);
+			FReal LocalThickness = LocalPhi - Thickness;
 
-			TCollisionContact<T, d> & Contact = Constraint.Manifold;
+			FCollisionContact & Contact = Constraint.Manifold;
 			if (LocalThickness < -KINDA_SMALL_NUMBER)
 			{
 				Contact.Location += LocalPoint * LocalThickness;
@@ -146,13 +140,13 @@ namespace Chaos
 
 #if INTEL_ISPC
 		template<ECollisionUpdateType UpdateType>
-		void SampleObject(const FImplicitObject& Object, const TRigidTransform<float, 3>& ObjectTransform, const TBVHParticles<float, 3>& SampleParticles, const TRigidTransform<float, 3>& SampleParticlesTransform, float Thickness, TRigidBodyPointContactConstraint<float, 3>& Constraint)
+		void SampleObject(const FImplicitObject& Object, const TRigidTransform<float, 3>& ObjectTransform, const TBVHParticles<float, 3>& SampleParticles, const TRigidTransform<float, 3>& SampleParticlesTransform, float Thickness, FRigidBodyPointContactConstraint& Constraint)
 		{
 			SCOPE_CYCLE_COUNTER(STAT_SampleObject);
-			TRigidBodyPointContactConstraint<float, 3> AvgConstraint = Constraint;
+			FRigidBodyPointContactConstraint AvgConstraint = Constraint;
 
-			TCollisionContact<float, 3> & Contact = Constraint.Manifold;
-			TCollisionContact<float, 3> & AvgContact = AvgConstraint.Manifold;
+			FCollisionContact& Contact = Constraint.Manifold;
+			FCollisionContact& AvgContact = AvgConstraint.Manifold;
 
 			AvgConstraint.Particle[0] = Constraint.Particle[0];
 			AvgConstraint.Particle[1] = Constraint.Particle[1];
@@ -486,14 +480,14 @@ namespace Chaos
 			}
 		}
 #else		
-		template <ECollisionUpdateType UpdateType, typename T, int d>
-		void SampleObject(const FImplicitObject& Object, const TRigidTransform<T, d>& ObjectTransform, const TBVHParticles<T, d>& SampleParticles, const TRigidTransform<T, d>& SampleParticlesTransform, T Thickness, TRigidBodyPointContactConstraint<float, 3>& Constraint)
+		template <ECollisionUpdateType UpdateType>
+		void SampleObject(const FImplicitObject& Object, const FRigidTransform3& ObjectTransform, const TBVHParticles<T, d>& SampleParticles, const FRigidTransform3& SampleParticlesTransform, FReal Thickness, FRigidBodyPointContactConstraint& Constraint)
 		{
 			SCOPE_CYCLE_COUNTER(STAT_SampleObject);
-			TRigidBodyPointContactConstraint<float, 3> AvgConstraint = Constraint;
+			FRigidBodyPointContactConstraint AvgConstraint = Constraint;
 
-			TCollisionContact<float, 3> & Contact = Constraint.Manifold;
-			TCollisionContact<float, 3> & AvgContact = AvgConstraint.Manifold;
+			FCollisionContact & Contact = Constraint.Manifold;
+			FCollisionContact & AvgContact = AvgConstraint.Manifold;
 
 			AvgConstraint.Particle[0] = Constraint.Particle[0];
 			AvgConstraint.Particle[1] = Constraint.Particle[1];
@@ -505,7 +499,7 @@ namespace Chaos
 			int32 DeepestParticle = -1;
 			const int32 NumParticles = SampleParticles.Size();
 
-			const TRigidTransform<T, d> & SampleToObjectTM = SampleParticlesTransform.GetRelativeTransform(ObjectTransform);
+			const FRigidTransform3 & SampleToObjectTM = SampleParticlesTransform.GetRelativeTransform(ObjectTransform);
 			if (NumParticles > SampleMinParticlesForAcceleration && Object.HasBoundingBox())
 			{
 				SCOPE_CYCLE_COUNTER(STAT_UpdateLevelsetPartial);
@@ -567,9 +561,9 @@ namespace Chaos
 			{
 				if (TotalThickness < -KINDA_SMALL_NUMBER)
 				{
-					TVector<T, d> LocalPoint = AvgContact.Location / TotalThickness;
-					TVector<T, d> LocalNormal;
-					const T NewPhi = Object.PhiWithNormal(LocalPoint, LocalNormal);
+					FVec3 LocalPoint = AvgContact.Location / TotalThickness;
+					FVec3 LocalNormal;
+					const FReal NewPhi = Object.PhiWithNormal(LocalPoint, LocalNormal);
 					if (NewPhi < Contact.Phi)
 					{
 						Contact.Phi = NewPhi;
@@ -585,8 +579,8 @@ namespace Chaos
 			else if (AvgContact.Phi < Contact.Phi)
 			{
 				check(DeepestParticle >= 0);
-				TVector<T, d> LocalPoint = SampleToObjectTM.TransformPositionNoScale(SampleParticles.X(DeepestParticle));
-				TVector<T, d> LocalNormal;
+				FVec3 LocalPoint = SampleToObjectTM.TransformPositionNoScale(SampleParticles.X(DeepestParticle));
+				FVec3 LocalNormal;
 				Contact.Phi = Object.PhiWithNormal(LocalPoint, LocalNormal);
 				Contact.Location = ObjectTransform.TransformPositionNoScale(LocalPoint);
 				Contact.Normal = ObjectTransform.TransformVectorNoScale(LocalNormal);
@@ -597,18 +591,17 @@ namespace Chaos
 
 
 		DECLARE_CYCLE_STAT(TEXT("TPBDCollisionConstraints::FindRelevantShapes"), STAT_FindRelevantShapes, STATGROUP_ChaosWide);
-		template <typename T, int d>
-		TArray<Pair<const FImplicitObject*, TRigidTransform<T, d>>> FindRelevantShapes(const FImplicitObject* ParticleObj, const TRigidTransform<T, d>& ParticlesTM, const FImplicitObject& LevelsetObj, const TRigidTransform<T, d>& LevelsetTM, const T Thickness)
+		TArray<Pair<const FImplicitObject*, FRigidTransform3>> FindRelevantShapes(const FImplicitObject* ParticleObj, const FRigidTransform3& ParticlesTM, const FImplicitObject& LevelsetObj, const FRigidTransform3& LevelsetTM, const FReal Thickness)
 		{
 			SCOPE_CYCLE_COUNTER(STAT_FindRelevantShapes);
-			TArray<Pair<const FImplicitObject*, TRigidTransform<T, d>>> RelevantShapes;
+			TArray<Pair<const FImplicitObject*, FRigidTransform3>> RelevantShapes;
 			//find all levelset inner objects
 			if (ParticleObj)
 			{
 				if (ParticleObj->HasBoundingBox())
 				{
-					const TRigidTransform<T, d> ParticlesToLevelsetTM = ParticlesTM.GetRelativeTransform(LevelsetTM);
-					TAABB<T, d> ParticleBoundsInLevelset = ParticleObj->BoundingBox().TransformedAABB(ParticlesToLevelsetTM);
+					const FRigidTransform3 ParticlesToLevelsetTM = ParticlesTM.GetRelativeTransform(LevelsetTM);
+					FAABB3 ParticleBoundsInLevelset = ParticleObj->BoundingBox().TransformedAABB(ParticlesToLevelsetTM);
 					ParticleBoundsInLevelset.Thicken(Thickness);
 					{
 						LevelsetObj.FindAllIntersectingObjects(RelevantShapes, ParticleBoundsInLevelset);
@@ -616,32 +609,20 @@ namespace Chaos
 				}
 				else
 				{
-					LevelsetObj.AccumulateAllImplicitObjects(RelevantShapes, TRigidTransform<T, d>::Identity);
+					LevelsetObj.AccumulateAllImplicitObjects(RelevantShapes, FRigidTransform3::Identity);
 				}
 			}
 			else
 			{
 				//todo:compute bounds
-				LevelsetObj.AccumulateAllImplicitObjects(RelevantShapes, TRigidTransform<T, d>::Identity);
+				LevelsetObj.AccumulateAllImplicitObjects(RelevantShapes, FRigidTransform3::Identity);
 			}
 
 			return RelevantShapes;
 		}
 
-		template TRigidTransform<float, 3> GetTransform(const TGeometryParticleHandle<float, 3>* Particle);
-		template PMatrix<float, 3, 3> ComputeFactorMatrix3(const TVector<float, 3>& V, const PMatrix<float, 3, 3>& M, const float& Im);
-		template TVector<float, 3> GetEnergyClampedImpulse(const TPBDRigidParticleHandle<float, 3>* PBDRigid0, const TPBDRigidParticleHandle<float, 3>* PBDRigid1, const TVector<float, 3>& Impulse, const TVector<float, 3>& VectorToPoint1, const TVector<float, 3>& VectorToPoint2, const TVector<float, 3>& Velocity1, const TVector<float, 3>& Velocity2);
-		template bool SampleObjectHelper<float, 3>(const FImplicitObject& Object, const TRigidTransform<float, 3>& ObjectTransform, const TRigidTransform<float, 3>& SampleToObjectTransform, const TVector<float, 3>& SampleParticle, float Thickness, TRigidBodyPointContactConstraint<float, 3>& Constraint);
-		template bool SampleObjectNoNormal(const FImplicitObject& Object, const TRigidTransform<float, 3>& ObjectTransform, const TRigidTransform<float, 3>& SampleToObjectTransform, const TVector<float, 3>& SampleParticle, float Thickness, TRigidBodyPointContactConstraint<float, 3>& Constraint);
-		template bool SampleObjectNormalAverageHelper(const FImplicitObject& Object, const TRigidTransform<float, 3>& ObjectTransform, const TRigidTransform<float, 3>& SampleToObjectTransform, const TVector<float, 3>& SampleParticle, float Thickness, float& TotalThickness, TRigidBodyPointContactConstraint<float, 3>& Constraint);
-#if INTEL_ISPC
-		template void SampleObject<ECollisionUpdateType::Any>(const FImplicitObject& Object, const TRigidTransform<float, 3>& ObjectTransform, const TBVHParticles<float, 3>& SampleParticles, const TRigidTransform<float, 3>& SampleParticlesTransform, float Thickness, TRigidBodyPointContactConstraint<float, 3>& Constraint);
-		template void SampleObject<ECollisionUpdateType::Deepest>(const FImplicitObject& Object, const TRigidTransform<float, 3>& ObjectTransform, const TBVHParticles<float, 3>& SampleParticles, const TRigidTransform<float, 3>& SampleParticlesTransform, float Thickness, TRigidBodyPointContactConstraint<float, 3>& Constraint);
-#else
-		template void SampleObject<ECollisionUpdateType::Any, float, 3>(const FImplicitObject& Object, const TRigidTransform<float, 3>& ObjectTransform, const TBVHParticles<float, 3>& SampleParticles, const TRigidTransform<float, 3>& SampleParticlesTransform, float Thickness, TRigidBodyPointContactConstraint<float, 3>& Constraint);
-		template void SampleObject<ECollisionUpdateType::Deepest, float, 3>(const FImplicitObject& Object, const TRigidTransform<float, 3>& ObjectTransform, const TBVHParticles<float, 3>& SampleParticles, const TRigidTransform<float, 3>& SampleParticlesTransform, float Thickness, TRigidBodyPointContactConstraint<float, 3>& Constraint);
-#endif
-		template TArray<Pair<const FImplicitObject*, TRigidTransform<float, 3>>> FindRelevantShapes<float, 3>(const FImplicitObject* ParticleObj, const TRigidTransform<float, 3>& ParticlesTM, const FImplicitObject& LevelsetObj, const TRigidTransform<float, 3>& LevelsetTM, const float Thickness);
+		template void SampleObject<ECollisionUpdateType::Any>(const FImplicitObject& Object, const TRigidTransform<float, 3>& ObjectTransform, const TBVHParticles<float, 3>& SampleParticles, const TRigidTransform<float, 3>& SampleParticlesTransform, float Thickness, FRigidBodyPointContactConstraint& Constraint);
+		template void SampleObject<ECollisionUpdateType::Deepest>(const FImplicitObject& Object, const TRigidTransform<float, 3>& ObjectTransform, const TBVHParticles<float, 3>& SampleParticles, const TRigidTransform<float, 3>& SampleParticlesTransform, float Thickness, FRigidBodyPointContactConstraint& Constraint);
 
 	} // Collisions
 
diff --git a/Engine/Source/Runtime/Experimental/Chaos/Private/Chaos/PBDCollisionConstraintUtil.cpp b/Engine/Source/Runtime/Experimental/Chaos/Private/Chaos/PBDCollisionConstraintUtil.cpp
--- a/Engine/Source/Runtime/Experimental/Chaos/Private/Chaos/PBDCollisionConstraintUtil.cpp
+++ b/Engine/Source/Runtime/Experimental/Chaos/Private/Chaos/PBDCollisionConstraintUtil.cpp
@@ -15,7 +15,7 @@
 
 namespace Chaos
 {
-	void ComputeHashTable(const TArray<Chaos::TPBDCollisionConstraints<float, 3>::FPointContactConstraint>& ConstraintsArray,
+	void ComputeHashTable(const TArray<Chaos::FRigidBodyPointContactConstraint>& ConstraintsArray,
 						  const FBox& BoundingBox, TMultiMap<int32, int32>& HashTableMap, const float SpatialHashRadius)
 	{
 		float CellSize = 2.f * SpatialHashRadius;
diff --git a/Engine/Source/Runtime/Experimental/Chaos/Private/Chaos/PBDCollisionConstraints.cpp b/Engine/Source/Runtime/Experimental/Chaos/Private/Chaos/PBDCollisionConstraints.cpp
--- a/Engine/Source/Runtime/Experimental/Chaos/Private/Chaos/PBDCollisionConstraints.cpp
+++ b/Engine/Source/Runtime/Experimental/Chaos/Private/Chaos/PBDCollisionConstraints.cpp
@@ -75,15 +75,14 @@ namespace Chaos
 	// Collision Constraint Container
 	//
 
-	template<typename T, int d>
-	TPBDCollisionConstraints<T, d>::TPBDCollisionConstraints(
-		const TPBDRigidsSOAs<T, d>& InParticles,
+	FPBDCollisionConstraints::FPBDCollisionConstraints(
+		const TPBDRigidsSOAs<FReal, 3>& InParticles,
 		TArrayCollectionArray<bool>& Collided,
 		const TArrayCollectionArray<TSerializablePtr<FChaosPhysicsMaterial>>& InPerParticleMaterials,
 		const int32 InApplyPairIterations /*= 1*/,
 		const int32 InApplyPushOutPairIterations /*= 1*/,
-		const T CullDistance /*= (T)0*/,
-		const T ShapePadding /*= (T)0*/)
+		const FReal CullDistance /*= (FReal)0*/,
+		const FReal ShapePadding /*= (FReal)0*/)
 		: Particles(InParticles)
 		, NumActivePointConstraints(0)
 		, NumActiveSweptPointConstraints(0)
@@ -111,44 +110,39 @@ namespace Chaos
 #endif
 	}
 
-	template<typename T, int d>
-	void TPBDCollisionConstraints<T, d>::DisableHandles()
+	void FPBDCollisionConstraints::DisableHandles()
 	{
 		check(NumConstraints() == 0);
 		bHandlesEnabled = false;
 	}
 
-	template<typename T, int d>
-	void TPBDCollisionConstraints<T, d>::SetPostApplyCallback(const TRigidBodyContactConstraintsPostApplyCallback<T, d>& Callback)
+
+	void FPBDCollisionConstraints::SetPostApplyCallback(const FRigidBodyContactConstraintsPostApplyCallback& Callback)
 	{
 		PostApplyCallback = Callback;
 	}
 
-	template<typename T, int d>
-	void TPBDCollisionConstraints<T, d>::ClearPostApplyCallback()
+	void FPBDCollisionConstraints::ClearPostApplyCallback()
 	{
 		PostApplyCallback = nullptr;
 	}
 
-	template<typename T, int d>
-	void TPBDCollisionConstraints<T, d>::SetPostApplyPushOutCallback(const TRigidBodyContactConstraintsPostApplyPushOutCallback<T, d>& Callback)
+	void FPBDCollisionConstraints::SetPostApplyPushOutCallback(const FRigidBodyContactConstraintsPostApplyPushOutCallback& Callback)
 	{
 		PostApplyPushOutCallback = Callback;
 	}
 	
-	template<typename T, int d>
-	void TPBDCollisionConstraints<T, d>::ClearPostApplyPushOutCallback()
+	void FPBDCollisionConstraints::ClearPostApplyPushOutCallback()
 	{
 		PostApplyPushOutCallback = nullptr;
 	}
 
-	template<typename T, int d>
-	void TPBDCollisionConstraints<T, d>::UpdateConstraintMaterialProperties(FConstraintBase& Constraint)
+	void FPBDCollisionConstraints::UpdateConstraintMaterialProperties(FCollisionConstraintBase& Constraint)
 	{
 		TSerializablePtr<FChaosPhysicsMaterial> PhysicsMaterial0 = Constraint.Particle[0]->AuxilaryValue(MPhysicsMaterials);
 		TSerializablePtr<FChaosPhysicsMaterial> PhysicsMaterial1 = Constraint.Particle[1]->AuxilaryValue(MPhysicsMaterials);
 
-		TCollisionContact<T, d>& Contact = Constraint.Manifold;
+		FCollisionContact& Contact = Constraint.Manifold;
 		if (PhysicsMaterial0 && PhysicsMaterial1)
 		{
 			// @todo(ccaulfield): support different friction/restitution combining algorithms
@@ -183,14 +177,13 @@ namespace Chaos
 		}
 	}
 
-	template<typename T, int d>
-	void TPBDCollisionConstraints<T, d>::AddConstraint(const TRigidBodyPointContactConstraint<FReal, 3>& InConstraint)
+	void FPBDCollisionConstraints::AddConstraint(const FRigidBodyPointContactConstraint& InConstraint)
 	{
 		int32 Idx = Constraints.SinglePointConstraints.Add(InConstraint);
 
 		if (bHandlesEnabled)
 		{
-			FConstraintContainerHandle* Handle = HandleAllocator.template AllocHandle< TRigidBodyPointContactConstraint<T, d> >(this, Idx);
+			FPBDCollisionConstraintHandle* Handle = HandleAllocator.template AllocHandle<FRigidBodyPointContactConstraint>(this, Idx);
 			Handle->GetContact().Timestamp = -INT_MAX; // force point constraints to be deleted.
 
 			Constraints.SinglePointConstraints[Idx].ConstraintHandle = Handle;
@@ -207,14 +200,13 @@ namespace Chaos
 		UpdateConstraintMaterialProperties(Constraints.SinglePointConstraints[Idx]);
 	}
 
-	template<typename T, int d>
-	void TPBDCollisionConstraints<T, d>::AddConstraint(const TRigidBodySweptPointContactConstraint<FReal, 3>& InConstraint)
+	void FPBDCollisionConstraints::AddConstraint(const FRigidBodySweptPointContactConstraint& InConstraint)
 	{
 		int32 Idx = Constraints.SinglePointSweptConstraints.Add(InConstraint);
 
 		if (bHandlesEnabled)
 		{
-			FConstraintContainerHandle* Handle = HandleAllocator.template AllocHandle< TRigidBodySweptPointContactConstraint<T, d> >(this, Idx);
+			FPBDCollisionConstraintHandle* Handle = HandleAllocator.template AllocHandle<FRigidBodySweptPointContactConstraint>(this, Idx);
 			Handle->GetContact().Timestamp = -INT_MAX; // force point constraints to be deleted.
 
 			Constraints.SinglePointSweptConstraints[Idx].ConstraintHandle = Handle;
@@ -234,14 +226,13 @@ namespace Chaos
 	}
 
 
-	template<typename T, int d>
-	void TPBDCollisionConstraints<T, d>::AddConstraint(const TRigidBodyMultiPointContactConstraint<FReal, 3>& InConstraint)
+	void FPBDCollisionConstraints::AddConstraint(const FRigidBodyMultiPointContactConstraint& InConstraint)
 	{
 		int32 Idx = Constraints.MultiPointConstraints.Add(InConstraint);
 
 		if (bHandlesEnabled)
 		{
-			FConstraintContainerHandle* Handle = HandleAllocator.template AllocHandle< TRigidBodyMultiPointContactConstraint<T, d> >(this, Idx);
+			FPBDCollisionConstraintHandle* Handle = HandleAllocator.template AllocHandle<FRigidBodyMultiPointContactConstraint>(this, Idx);
 			Handle->GetContact().Timestamp = LifespanCounter;
 
 			Constraints.MultiPointConstraints[Idx].ConstraintHandle = Handle;
@@ -258,24 +249,22 @@ namespace Chaos
 		UpdateConstraintMaterialProperties(Constraints.MultiPointConstraints[Idx]);
 	}
 
-	template<typename T, int d>
-	void TPBDCollisionConstraints<T, d>::UpdatePositionBasedState(const T Dt)
+	void FPBDCollisionConstraints::UpdatePositionBasedState(const FReal Dt)
 	{
 		Reset();
 	
 		LifespanCounter++;
 	}
 
-	template<typename T, int d>
-	void TPBDCollisionConstraints<T, d>::Reset()
+	void FPBDCollisionConstraints::Reset()
 	{
 		SCOPE_CYCLE_COUNTER(STAT_Collisions_Reset);
 
 #if CHAOS_COLLISION_PERSISTENCE_ENABLED
 		check(bHandlesEnabled);	// This will need fixing for handle-free mode
-		TArray<FConstraintContainerHandle*> CopyOfHandles = Handles;
+		TArray<FPBDCollisionConstraintHandle*> CopyOfHandles = Handles;
 		int32 LifespanWindow = LifespanCounter - 1;
-		for (FConstraintContainerHandle* ContactHandle : CopyOfHandles)
+		for (FPBDCollisionConstraintHandle* ContactHandle : CopyOfHandles)
 		{
 			if (!bEnableCollisions || ContactHandle->GetContact().Timestamp< LifespanWindow)
 			{
@@ -283,7 +272,7 @@ namespace Chaos
 			}
 		}
 #else
-		for (FConstraintContainerHandle* Handle : Handles)
+		for (FPBDCollisionConstraintHandle* Handle : Handles)
 		{
 			HandleAllocator.FreeHandle(Handle);
 		}
@@ -295,14 +284,13 @@ namespace Chaos
 		bUseCCD = false;
 	}
 
-	template<typename T, int d>
-	void TPBDCollisionConstraints<T, d>::ApplyCollisionModifier(const TFunction<ECollisionModifierResult(FConstraintContainerHandle* Handle)>& CollisionModifier)
+	void FPBDCollisionConstraints::ApplyCollisionModifier(const TFunction<ECollisionModifierResult(FPBDCollisionConstraintHandle* Handle)>& CollisionModifier)
 	{
 		if (CollisionModifier)
 		{
-			TArray<FConstraintContainerHandle*> CopyOfHandles = Handles;
+			TArray<FPBDCollisionConstraintHandle*> CopyOfHandles = Handles;
 
-			for (FConstraintContainerHandle* ContactHandle : CopyOfHandles)
+			for (FPBDCollisionConstraintHandle* ContactHandle : CopyOfHandles)
 			{
 				ECollisionModifierResult Result = CollisionModifier(ContactHandle);
 				if (Result == ECollisionModifierResult::Disabled)
@@ -313,17 +301,16 @@ namespace Chaos
 		}
 	}
 
-	template<typename T, int d>
-	void TPBDCollisionConstraints<T, d>::RemoveConstraints(const TSet<TGeometryParticleHandle<T, d>*>&  InHandleSet)
+	void FPBDCollisionConstraints::RemoveConstraints(const TSet<TGeometryParticleHandle<FReal, 3>*>&  InHandleSet)
 	{
-		const TArray<TGeometryParticleHandle<T, d>*> HandleArray = InHandleSet.Array();
+		const TArray<TGeometryParticleHandle<FReal, 3>*> HandleArray = InHandleSet.Array();
 		for (auto ParticleHandle : HandleArray)
 		{
-			TArray<FConstraintContainerHandle*> CopyOfHandles = Handles;
+			TArray<FPBDCollisionConstraintHandle*> CopyOfHandles = Handles;
 
-			for (FConstraintContainerHandle* ContactHandle : CopyOfHandles)
+			for (FPBDCollisionConstraintHandle* ContactHandle : CopyOfHandles)
 			{
-				TVector<TGeometryParticleHandle<T, d>*, 2> ConstraintParticles = ContactHandle->GetConstrainedParticles();
+				TVector<TGeometryParticleHandle<FReal, 3>*, 2> ConstraintParticles = ContactHandle->GetConstrainedParticles();
 				if (ConstraintParticles[1] == ParticleHandle || ConstraintParticles[0] == ParticleHandle)
 				{
 					RemoveConstraint(ContactHandle);
@@ -332,8 +319,7 @@ namespace Chaos
 		}
 	}
 
-	template<typename T, int d>
-	void TPBDCollisionConstraints<T, d>::RemoveConstraint(FConstraintContainerHandle* Handle)
+	void FPBDCollisionConstraints::RemoveConstraint(FPBDCollisionConstraintHandle* Handle)
 	{
 		FConstraintContainerHandleKey KeyToRemove = Handle->GetKey();
 		int32 Idx = Handle->GetConstraintIndex(); // index into specific array
@@ -409,15 +395,13 @@ namespace Chaos
 	}
 
 
-	template<typename T, int d>
-	void TPBDCollisionConstraints<T, d>::UpdateConstraints(T Dt, const TSet<TGeometryParticleHandle<T, d>*>& ParticlesSet)
+	void FPBDCollisionConstraints::UpdateConstraints(FReal Dt, const TSet<TGeometryParticleHandle<FReal, 3>*>& ParticlesSet)
 	{
 		// Clustering uses update constraints to force a re-evaluation. 
 	}
 
 	// Called once per frame to update persistent constraints (reruns collision detection, or selects the best manifold point)
-	template<typename T, int d>
-	void TPBDCollisionConstraints<T, d>::UpdateConstraints(T Dt)
+	void FPBDCollisionConstraints::UpdateConstraints(FReal Dt)
 	{
 		SCOPE_CYCLE_COUNTER(STAT_Collisions_UpdatePointConstraints);
 
@@ -425,7 +409,7 @@ namespace Chaos
 
 		//PhysicsParallelFor(Handles.Num(), [&](int32 ConstraintHandleIndex)
 		//{
-		//	FConstraintContainerHandle* ConstraintHandle = Handles[ConstraintHandleIndex];
+		//	FPBDCollisionConstraintHandle* ConstraintHandle = Handles[ConstraintHandleIndex];
 		//	check(ConstraintHandle != nullptr);
 		//	Collisions::Update(MCullDistance, MShapePadding, ConstraintHandle->GetContact());
 
@@ -435,7 +419,7 @@ namespace Chaos
 		//	}
 		//}, bDisableCollisionParallelFor);
 
-		for (FPointContactConstraint& Contact : Constraints.SinglePointConstraints)
+		for (FRigidBodyPointContactConstraint& Contact : Constraints.SinglePointConstraints)
 		{
 			Collisions::Update(Contact, MCullDistance);
 			if (Contact.GetPhi() < MCullDistance)
@@ -446,8 +430,7 @@ namespace Chaos
 	}
 
 	// Called once per tick to update/regenerate persistent manifold planes and points
-	template<typename T, int d>
-	void TPBDCollisionConstraints<T, d>::UpdateManifolds(T Dt)
+	void FPBDCollisionConstraints::UpdateManifolds(FReal Dt)
 	{
 		SCOPE_CYCLE_COUNTER(STAT_Collisions_UpdateManifoldConstraints);
 
@@ -455,14 +438,14 @@ namespace Chaos
 
 		//PhysicsParallelFor(Handles.Num(), [&](int32 ConstraintHandleIndex)
 		//{
-		//	FConstraintContainerHandle* ConstraintHandle = Handles[ConstraintHandleIndex];
+		//	FPBDCollisionConstraintHandle* ConstraintHandle = Handles[ConstraintHandleIndex];
 		//	check(ConstraintHandle != nullptr);
 		//	Collisions::Update(MCullDistance, MShapePadding, ConstraintHandle->GetContact());
 		//}, bDisableCollisionParallelFor);
 
 		FCollisionContext Context;
 
-		for (FMultiPointContactConstraint& Contact : Constraints.MultiPointConstraints)
+		for (FRigidBodyMultiPointContactConstraint& Contact : Constraints.MultiPointConstraints)
 		{
 			Collisions::UpdateManifold(Contact, MCullDistance, Context);
 			if (Contact.GetPhi() < MCullDistance)
@@ -472,19 +455,18 @@ namespace Chaos
 		}
 	}
 
-	template<typename T, int d>
-	bool TPBDCollisionConstraints<T, d>::Apply(const T Dt, const int32 Iterations, const int32 NumIterations)
+	bool FPBDCollisionConstraints::Apply(const FReal Dt, const int32 Iterations, const int32 NumIterations)
 	{
 		SCOPE_CYCLE_COUNTER(STAT_Collisions_Apply);
 
 		bool bNeedsAnotherIteration = false;
 		if (MApplyPairIterations > 0)
 		{
-			const Collisions::TContactParticleParameters<T> ParticleParameters = { MCullDistance, MShapePadding, &MCollided };
-			const Collisions::TContactIterationParameters<T> IterationParameters = { Dt, Iterations, NumIterations, MApplyPairIterations, ApplyType, &bNeedsAnotherIteration };
+			const Collisions::FContactParticleParameters ParticleParameters = { MCullDistance, MShapePadding, &MCollided };
+			const Collisions::FContactIterationParameters IterationParameters = { Dt, Iterations, NumIterations, MApplyPairIterations, ApplyType, &bNeedsAnotherIteration };
 
 			NumActivePointConstraints = 0;
-			for (FPointContactConstraint& Contact : Constraints.SinglePointConstraints)
+			for (FRigidBodyPointContactConstraint& Contact : Constraints.SinglePointConstraints)
 			{
 				if (!Contact.GetDisabled())
 				{
@@ -494,7 +476,7 @@ namespace Chaos
 			}
 
 			NumActiveIterativeConstraints = 0;
-			for (FMultiPointContactConstraint& Contact : Constraints.MultiPointConstraints)
+			for (FRigidBodyMultiPointContactConstraint& Contact : Constraints.MultiPointConstraints)
 			{
 				if (!Contact.GetDisabled())
 				{
@@ -506,7 +488,7 @@ namespace Chaos
 			// Swept apply may significantly change particle position, invalidating other constraint's manifolds.
 			// We don't update manifolds on first apply iteration, so make sure we apply swept constraints last.
 			NumActiveSweptPointConstraints = 0;
-			for (FSweptPointContactConstraint& Contact : Constraints.SinglePointSweptConstraints)
+			for (FRigidBodySweptPointContactConstraint& Contact : Constraints.SinglePointSweptConstraints)
 			{
 				if (!Contact.GetDisabled())
 				{
@@ -524,19 +506,18 @@ namespace Chaos
 		return bNeedsAnotherIteration;
 	}
 
-	template<typename T, int d>
-	bool TPBDCollisionConstraints<T, d>::ApplyPushOut(const T Dt, const int32 Iterations, const int32 NumIterations)
+	bool FPBDCollisionConstraints::ApplyPushOut(const FReal Dt, const int32 Iterations, const int32 NumIterations)
 	{
 		SCOPE_CYCLE_COUNTER(STAT_Collisions_ApplyPushOut);
 
-		TSet<const TGeometryParticleHandle<T, d>*> TempStatic;
+		TSet<const TGeometryParticleHandle<FReal, 3>*> TempStatic;
 		bool bNeedsAnotherIteration = false;
 		if (MApplyPushOutPairIterations > 0)
 		{
-			const Collisions::TContactParticleParameters<T> ParticleParameters = { MCullDistance, MShapePadding, &MCollided };
-			const Collisions::TContactIterationParameters<T> IterationParameters = { Dt, Iterations, NumIterations, MApplyPushOutPairIterations, ECollisionApplyType::None, &bNeedsAnotherIteration };
+			const Collisions::FContactParticleParameters ParticleParameters = { MCullDistance, MShapePadding, &MCollided };
+			const Collisions::FContactIterationParameters IterationParameters = { Dt, Iterations, NumIterations, MApplyPushOutPairIterations, ECollisionApplyType::None, &bNeedsAnotherIteration };
 
-			for (FPointContactConstraint& Contact : Constraints.SinglePointConstraints)
+			for (FRigidBodyPointContactConstraint& Contact : Constraints.SinglePointConstraints)
 			{
 				if (!Contact.GetDisabled())
 				{
@@ -544,7 +525,7 @@ namespace Chaos
 				}
 			}
 
-			for (FSweptPointContactConstraint& Contact : Constraints.SinglePointSweptConstraints)
+			for (FRigidBodySweptPointContactConstraint& Contact : Constraints.SinglePointSweptConstraints)
 			{
 				if (!Contact.GetDisabled())
 				{
@@ -552,7 +533,7 @@ namespace Chaos
 				}
 			}
 
-			for (FMultiPointContactConstraint& Contact : Constraints.MultiPointConstraints)
+			for (FRigidBodyMultiPointContactConstraint& Contact : Constraints.MultiPointConstraints)
 			{
 				if (!Contact.GetDisabled())
 				{
@@ -570,8 +551,7 @@ namespace Chaos
 	}
 
 
-	template<typename T, int d>
-	bool TPBDCollisionConstraints<T, d>::Apply(const T Dt, const TArray<FConstraintContainerHandle*>& InConstraintHandles, const int32 Iterations, const int32 NumIterations)
+	bool FPBDCollisionConstraints::Apply(const FReal Dt, const TArray<FPBDCollisionConstraintHandle*>& InConstraintHandles, const int32 Iterations, const int32 NumIterations)
 	{
 		SCOPE_CYCLE_COUNTER(STAT_Collisions_Apply);
 
@@ -580,18 +560,22 @@ namespace Chaos
 		if (MApplyPairIterations > 0)
 		{
 			PhysicsParallelFor(InConstraintHandles.Num(), [&](int32 ConstraintHandleIndex) {
-				FConstraintContainerHandle* ConstraintHandle = InConstraintHandles[ConstraintHandleIndex];
+				FPBDCollisionConstraintHandle* ConstraintHandle = InConstraintHandles[ConstraintHandleIndex];
 				check(ConstraintHandle != nullptr);
 
-				TVector<const TGeometryParticleHandle<T, d>*, 2> ConstrainedParticles = ConstraintHandle->GetConstrainedParticles();
+				TVector<const TGeometryParticleHandle<FReal, 3>*, 2> ConstrainedParticles = ConstraintHandle->GetConstrainedParticles();
 				bool bNeedsAnotherIteration = false;
-				Collisions::TContactParticleParameters<T> ParticleParameters = { MCullDistance, MShapePadding, &MCollided };
-				Collisions::TContactIterationParameters<T> IterationParameters = { Dt, Iterations, NumIterations, MApplyPairIterations, ApplyType, &bNeedsAnotherIteration };
-				Collisions::Apply(ConstraintHandle->GetContact(), IterationParameters, ParticleParameters);
 
-				if (bNeedsAnotherIteration)
+				if (!ConstraintHandle->GetContact().GetDisabled())
 				{
-					bNeedsAnotherIterationAtomic.Store(true);
+					Collisions::FContactParticleParameters ParticleParameters = { MCullDistance, MShapePadding, &MCollided };
+					Collisions::FContactIterationParameters IterationParameters = { Dt, Iterations, NumIterations, MApplyPairIterations, ApplyType, &bNeedsAnotherIteration };
+					Collisions::Apply(ConstraintHandle->GetContact(), IterationParameters, ParticleParameters);
+
+					if (bNeedsAnotherIteration)
+					{
+						bNeedsAnotherIterationAtomic.Store(true);
+					}
 				}
 
 			}, bDisableCollisionParallelFor);
@@ -606,8 +590,7 @@ namespace Chaos
 	}
 
 
-	template<typename T, int d>
-	bool TPBDCollisionConstraints<T, d>::ApplyPushOut(const T Dt, const TArray<FConstraintContainerHandle*>& InConstraintHandles, const TSet< const TGeometryParticleHandle<T, d>*>& IsTemporarilyStatic, int32 Iteration, int32 NumIterations)
+	bool FPBDCollisionConstraints::ApplyPushOut(const FReal Dt, const TArray<FPBDCollisionConstraintHandle*>& InConstraintHandles, const TSet< const TGeometryParticleHandle<FReal, 3>*>& IsTemporarilyStatic, int32 Iteration, int32 NumIterations)
 	{
 		SCOPE_CYCLE_COUNTER(STAT_Collisions_ApplyPushOut);
 
@@ -616,12 +599,15 @@ namespace Chaos
 		{
 			PhysicsParallelFor(InConstraintHandles.Num(), [&](int32 ConstraintHandleIndex)
 			{
-				FConstraintContainerHandle* ConstraintHandle = InConstraintHandles[ConstraintHandleIndex];
+				FPBDCollisionConstraintHandle* ConstraintHandle = InConstraintHandles[ConstraintHandleIndex];
 				check(ConstraintHandle != nullptr);
 
-				Collisions::TContactParticleParameters<T> ParticleParameters = { MCullDistance, MShapePadding, &MCollided };
-				Collisions::TContactIterationParameters<T> IterationParameters = { Dt, Iteration, NumIterations, MApplyPushOutPairIterations, ECollisionApplyType::None, &bNeedsAnotherIteration };
-				Collisions::ApplyPushOut(ConstraintHandle->GetContact(), IsTemporarilyStatic, IterationParameters, ParticleParameters);
+				if (!ConstraintHandle->GetContact().GetDisabled())
+				{
+					Collisions::FContactParticleParameters ParticleParameters = { MCullDistance, MShapePadding, &MCollided };
+					Collisions::FContactIterationParameters IterationParameters = { Dt, Iteration, NumIterations, MApplyPushOutPairIterations, ECollisionApplyType::None, &bNeedsAnotherIteration };
+					Collisions::ApplyPushOut(ConstraintHandle->GetContact(), IsTemporarilyStatic, IterationParameters, ParticleParameters);
+				}
 
 			}, bDisableCollisionParallelFor);
 		}
@@ -634,8 +620,7 @@ namespace Chaos
 		return bNeedsAnotherIteration;
 	}
 
-	template<typename T, int d>
-	const typename TPBDCollisionConstraints<T, d>::FConstraintBase& TPBDCollisionConstraints<T, d>::GetConstraint(int32 Index) const
+	const FCollisionConstraintBase& FPBDCollisionConstraints::GetConstraint(int32 Index) const
 	{
 		check(Index < NumConstraints());
 		
@@ -658,5 +643,4 @@ namespace Chaos
 
 
 	template class TAccelerationStructureHandle<float, 3>;
-	template class CHAOS_API TPBDCollisionConstraints<float, 3>;
 }
diff --git a/Engine/Source/Runtime/Experimental/Chaos/Private/Chaos/PBDCollisionConstraintsContact.cpp b/Engine/Source/Runtime/Experimental/Chaos/Private/Chaos/PBDCollisionConstraintsContact.cpp
--- a/Engine/Source/Runtime/Experimental/Chaos/Private/Chaos/PBDCollisionConstraintsContact.cpp
+++ b/Engine/Source/Runtime/Experimental/Chaos/Private/Chaos/PBDCollisionConstraintsContact.cpp
@@ -62,8 +62,8 @@ namespace Chaos
 		FVec3 ApplyContact(FCollisionContact& Contact,
 			TGenericParticleHandle<FReal, 3> Particle0, 
 			TGenericParticleHandle<FReal, 3> Particle1,
-			const TContactIterationParameters<FReal> & IterationParameters,
-			const TContactParticleParameters<FReal> & ParticleParameters)
+			const FContactIterationParameters & IterationParameters,
+			const FContactParticleParameters & ParticleParameters)
 		{
 			FVec3 AccumulatedImpulse(0);
 
@@ -232,8 +232,8 @@ namespace Chaos
 		FVec3 ApplyContact2(FCollisionContact& Contact,
 			TGenericParticleHandle<FReal, 3> Particle0,
 			TGenericParticleHandle<FReal, 3> Particle1,
-			const TContactIterationParameters<FReal>& IterationParameters,
-			const TContactParticleParameters<FReal>& ParticleParameters)
+			const FContactIterationParameters& IterationParameters,
+			const FContactParticleParameters& ParticleParameters)
 		{
 			FVec3 AccumulatedImpulse(0);
 			TPBDRigidParticleHandle<FReal, 3>* PBDRigid0 = Particle0->CastToRigidParticle();
@@ -610,6 +610,18 @@ namespace Chaos
 			{
 				Update(Constraint, ParticleParameters.CullDistance);
 
+				// Permanently disable a constraint that is beyond the cull distance
+				if (Constraint.GetPhi() >= ParticleParameters.CullDistance)
+				{
+					Constraint.SetDisabled(true);
+					return;
+				}
+
+				if (Constraint.GetPhi() >= ParticleParameters.ShapePadding)
+				{
+					return;
+				}
+
 				Constraint.AccumulatedImpulse += 
 					ApplyPushOutContact(Constraint.Manifold, Particle0, Particle1, IsTemporarilyStatic, IterationParameters, ParticleParameters);
 			}
diff --git a/Engine/Source/Runtime/Experimental/Chaos/Private/Chaos/PBDConstraintRule.cpp b/Engine/Source/Runtime/Experimental/Chaos/Private/Chaos/PBDConstraintRule.cpp
--- a/Engine/Source/Runtime/Experimental/Chaos/Private/Chaos/PBDConstraintRule.cpp
+++ b/Engine/Source/Runtime/Experimental/Chaos/Private/Chaos/PBDConstraintRule.cpp
@@ -52,17 +52,17 @@ namespace Chaos
 		return Constraints.NumConstraints(); 
 	}
 
-	template class TSimpleConstraintRule<TPBDCollisionConstraints<float, 3>>;
+	template class TSimpleConstraintRule<FPBDCollisionConstraints>;
 	template class TSimpleConstraintRule<FPBDJointConstraints>;
 	template class TSimpleConstraintRule<FPBDRigidSpringConstraints>;
 
-	template class TPBDConstraintGraphRuleImpl<TPBDCollisionConstraints<float, 3>>;
+	template class TPBDConstraintGraphRuleImpl<FPBDCollisionConstraints>;
 	template class TPBDConstraintGraphRuleImpl<FPBDJointConstraints>;
 	template class TPBDConstraintGraphRuleImpl<TPBDPositionConstraints<float, 3>>;
 	template class TPBDConstraintGraphRuleImpl<TPBDRigidDynamicSpringConstraints<float, 3>>;
 	template class TPBDConstraintGraphRuleImpl<FPBDRigidSpringConstraints>;
 
-	template class TPBDConstraintColorRule<TPBDCollisionConstraints<float, 3>>;
+	template class TPBDConstraintColorRule<FPBDCollisionConstraints>;
 	template class TPBDConstraintIslandRule<FPBDJointConstraints>;
 	template class TPBDConstraintIslandRule<TPBDPositionConstraints<float, 3>>;
 	template class TPBDConstraintIslandRule<TPBDRigidDynamicSpringConstraints<float, 3>>;
diff --git a/Engine/Source/Runtime/Experimental/Chaos/Private/Chaos/PBDRigidClustering.cpp b/Engine/Source/Runtime/Experimental/Chaos/Private/Chaos/PBDRigidClustering.cpp
--- a/Engine/Source/Runtime/Experimental/Chaos/Private/Chaos/PBDRigidClustering.cpp
+++ b/Engine/Source/Runtime/Experimental/Chaos/Private/Chaos/PBDRigidClustering.cpp
@@ -77,13 +77,13 @@ namespace Chaos
 	//==========================================================================
 
 	template<class T, int d>
-	TVector<T, d> GetContactLocation(const TRigidBodyPointContactConstraint<T, d>& Contact)
+	TVector<T, d> GetContactLocation(const FRigidBodyPointContactConstraint& Contact)
 	{
 		return Contact.GetLocation();
 	}
 
 	template<class T, int d>
-	TVector<T, d> GetContactLocation(const TRigidBodyContactConstraintPGS<T, d>& Contact)
+	TVector<T, d> GetContactLocation(const FRigidBodyContactConstraintPGS& Contact)
 	{
 		// @todo(mlentine): Does the exact point matter?
 		T MinPhi = FLT_MAX;
@@ -119,7 +119,7 @@ namespace Chaos
 		const TSet<int32>& IslandsToRecollide, 
 		const TSet<TPBDRigidParticleHandle<T, d>*> AllActivatedChildren,
 		const T Dt, 
-		TPBDCollisionConstraints<T, d>& CollisionRule)
+		FPBDCollisionConstraints& CollisionRule)
 	{
 		SCOPE_CYCLE_COUNTER(STAT_RewindAndEvolve_BGF);
 		// Rewind active particles
@@ -403,9 +403,9 @@ namespace Chaos
 	// TPBDRigidClustering
 	//==========================================================================
 
-	template<class FPBDRigidsEvolution, class FPBDCollisionConstraint, class T, int d>
-	TPBDRigidClustering<FPBDRigidsEvolution, FPBDCollisionConstraint, T, d>::TPBDRigidClustering(
-		FPBDRigidsEvolution& InEvolution, 
+	template<class T_FPBDRigidsEvolution, class T_FPBDCollisionConstraint, class T, int d>
+	TPBDRigidClustering<T_FPBDRigidsEvolution, T_FPBDCollisionConstraint, T, d>::TPBDRigidClustering(
+		T_FPBDRigidsEvolution& InEvolution, 
 		TPBDRigidClusteredParticles<T, d>& InParticles)
 		: MEvolution(InEvolution)
 		, MParticles(InParticles)
@@ -415,14 +415,14 @@ namespace Chaos
 		, MClusterUnionConnectionType(FClusterCreationParameters<T>::EConnectionMethod::DelaunayTriangulation)
 	{}
 
-	template<class FPBDRigidsEvolution, class FPBDCollisionConstraint, class T, int d>
-	TPBDRigidClustering<FPBDRigidsEvolution, FPBDCollisionConstraint, T, d>::~TPBDRigidClustering()
+	template<class T_FPBDRigidsEvolution, class T_FPBDCollisionConstraint, class T, int d>
+	TPBDRigidClustering<T_FPBDRigidsEvolution, T_FPBDCollisionConstraint, T, d>::~TPBDRigidClustering()
 	{}
 
 	DECLARE_CYCLE_STAT(TEXT("TPBDRigidClustering<>::CreateClusterParticle"), STAT_CreateClusterParticle, STATGROUP_Chaos);
-	template<class FPBDRigidsEvolution, class FPBDCollisionConstraint, class T, int d>
+	template<class T_FPBDRigidsEvolution, class T_FPBDCollisionConstraint, class T, int d>
 	Chaos::TPBDRigidClusteredParticleHandle<float, 3>* 
-	TPBDRigidClustering<FPBDRigidsEvolution, FPBDCollisionConstraint, T, d>::CreateClusterParticle(
+	TPBDRigidClustering<T_FPBDRigidsEvolution, T_FPBDCollisionConstraint, T, d>::CreateClusterParticle(
 		const int32 ClusterGroupIndex,
 		TArray<Chaos::TPBDRigidParticleHandle<T,d>*>&& Children,
 		const FClusterCreationParameters<T>& Parameters,
@@ -513,9 +513,9 @@ namespace Chaos
 	FAutoConsoleVariableRef CVarUnionsHaveCollisionParticles(TEXT("p.UnionsHaveCollisionParticles"), UnionsHaveCollisionParticles, TEXT(""));
 
 	DECLARE_CYCLE_STAT(TEXT("TPBDRigidClustering<>::CreateClusterParticleFromClusterChildren"), STAT_CreateClusterParticleFromClusterChildren, STATGROUP_Chaos);
-	template<class FPBDRigidsEvolution, class FPBDCollisionConstraint, class T, int d>
+	template<class T_FPBDRigidsEvolution, class T_FPBDCollisionConstraint, class T, int d>
 	Chaos::TPBDRigidClusteredParticleHandle<float, 3>* 
-	TPBDRigidClustering<FPBDRigidsEvolution, FPBDCollisionConstraint, T, d>::CreateClusterParticleFromClusterChildren(
+	TPBDRigidClustering<T_FPBDRigidsEvolution, T_FPBDCollisionConstraint, T, d>::CreateClusterParticleFromClusterChildren(
 		TArray<TPBDRigidParticleHandle<T,d>*>&& Children, 
 		TPBDRigidClusteredParticleHandle<T,d>* Parent, 
 		const TRigidTransform<T, d>& ClusterWorldTM, 
@@ -596,8 +596,8 @@ namespace Chaos
 	FAutoConsoleVariableRef CVarMinChildrenForMultiProxy(TEXT("p.MinChildrenForMultiProxy"), MinChildrenForMultiProxy, TEXT("Min number of children needed for multi child proxy optimization"));
 
 	DECLARE_CYCLE_STAT(TEXT("TPBDRigidClustering<>::UnionClusterGroups"), STAT_UnionClusterGroups, STATGROUP_Chaos);
-	template<class FPBDRigidsEvolution, class FPBDCollisionConstraint, class T, int d>
-	void TPBDRigidClustering<FPBDRigidsEvolution, FPBDCollisionConstraint, T, d>::UnionClusterGroups()
+	template<class T_FPBDRigidsEvolution, class T_FPBDCollisionConstraint, class T, int d>
+	void TPBDRigidClustering<T_FPBDRigidsEvolution, T_FPBDCollisionConstraint, T, d>::UnionClusterGroups()
 	{
 		SCOPE_CYCLE_COUNTER(STAT_UnionClusterGroups);
 
@@ -675,8 +675,8 @@ namespace Chaos
 
 #if 0 // Not called currently
 	DECLARE_CYCLE_STAT(TEXT("TPBDRigidClustering<>::DeactivateClusterParticle"), STAT_DeactivateClusterParticle, STATGROUP_Chaos);
-	template<class FPBDRigidsEvolution, class FPBDCollisionConstraint, class T, int d>
-	TSet<uint32> TPBDRigidClustering<FPBDRigidsEvolution, FPBDCollisionConstraint, T, d>::DeactivateClusterParticle(
+	template<class T_FPBDRigidsEvolution, class T_FPBDCollisionConstraint, class T, int d>
+	TSet<uint32> TPBDRigidClustering<T_FPBDRigidsEvolution, T_FPBDCollisionConstraint, T, d>::DeactivateClusterParticle(
 		TPBDRigidClusteredParticleHandle<T,d>* ClusteredParticle)
 	{
 		SCOPE_CYCLE_COUNTER(STAT_DeactivateClusterParticle);
@@ -692,9 +692,9 @@ namespace Chaos
 #endif // 0
 
 	DECLARE_CYCLE_STAT(TEXT("TPBDRigidClustering<>::ReleaseClusterParticles(STRAIN)"), STAT_ReleaseClusterParticles_STRAIN, STATGROUP_Chaos);
-	template<class FPBDRigidsEvolution, class FPBDCollisionConstraint, class T, int d>
+	template<class T_FPBDRigidsEvolution, class T_FPBDCollisionConstraint, class T, int d>
 	TSet<TPBDRigidParticleHandle<T, d>*> 
-	TPBDRigidClustering<FPBDRigidsEvolution, FPBDCollisionConstraint, T, d>::ReleaseClusterParticles(
+	TPBDRigidClustering<T_FPBDRigidsEvolution, T_FPBDCollisionConstraint, T, d>::ReleaseClusterParticles(
 		TPBDRigidClusteredParticleHandle<T,d>* ClusteredParticle,
 		const TMap<TGeometryParticleHandle<T, d>*, float>* ExternalStrainMap)
 	{
@@ -913,8 +913,8 @@ namespace Chaos
 
 #if 0 // Not called currently
 	DECLARE_CYCLE_STAT(TEXT("TPBDRigidClustering<>::ReleaseClusterParticles(LIST)"), STAT_ReleaseClusterParticles_LIST, STATGROUP_Chaos);
-	template<class FPBDRigidsEvolution, class FPBDCollisionConstraint, class T, int d>
-	TSet<uint32> TPBDRigidClustering<FPBDRigidsEvolution, FPBDCollisionConstraint, T, d>::ReleaseClusterParticles(
+	template<class T_FPBDRigidsEvolution, class T_FPBDCollisionConstraint, class T, int d>
+	TSet<uint32> TPBDRigidClustering<T_FPBDRigidsEvolution, T_FPBDCollisionConstraint, T, d>::ReleaseClusterParticles(
 		TArray<TPBDRigidParticleHandle<T, d>*> ChildrenParticles)
 		const TArray<uint32>& ChildrenParticles)
 	{
@@ -964,11 +964,11 @@ namespace Chaos
 	DECLARE_CYCLE_STAT(TEXT("TPBDRigidClustering<>::Update Impulse from Strain"), STAT_UpdateImpulseStrain, STATGROUP_Chaos);
 	DECLARE_CYCLE_STAT(TEXT("TPBDRigidClustering<>::Update Dirty Impulses"), STAT_UpdateDirtyImpulses, STATGROUP_Chaos);
 	DECLARE_CYCLE_STAT(TEXT("TPBDRigidClustering<>::Rewind"), STAT_ClusterRewind, STATGROUP_Chaos);
-	template<class FPBDRigidsEvolution, class FPBDCollisionConstraint, class T, int d>
+	template<class T_FPBDRigidsEvolution, class T_FPBDCollisionConstraint, class T, int d>
 	void 
-	TPBDRigidClustering<FPBDRigidsEvolution, FPBDCollisionConstraint, T, d>::AdvanceClustering(
+	TPBDRigidClustering<T_FPBDRigidsEvolution, T_FPBDCollisionConstraint, T, d>::AdvanceClustering(
 		const T Dt, 
-		FPBDCollisionConstraint& CollisionRule)
+		T_FPBDCollisionConstraint& CollisionRule)
 	{
 		SCOPE_CYCLE_COUNTER(STAT_AdvanceClustering);
 		UE_LOG(LogChaos, Verbose, TEXT("START FRAME with Dt %f"), Dt);
@@ -1093,9 +1093,9 @@ namespace Chaos
 	}
 
 	DECLARE_CYCLE_STAT(TEXT("TPBDRigidClustering<>::BreakingModel()"), STAT_BreakingModel, STATGROUP_Chaos);
-	template<class FPBDRigidsEvolution, class FPBDCollisionConstraint, class T, int d>
+	template<class T_FPBDRigidsEvolution, class T_FPBDCollisionConstraint, class T, int d>
 	TMap<TPBDRigidClusteredParticleHandle<T, d>*, TSet<TPBDRigidParticleHandle<T, d>*>> 
-	TPBDRigidClustering<FPBDRigidsEvolution, FPBDCollisionConstraint, T, d>::BreakingModel(
+	TPBDRigidClustering<T_FPBDRigidsEvolution, T_FPBDCollisionConstraint, T, d>::BreakingModel(
 		TMap<TGeometryParticleHandle<T, d>*, float>* ExternalStrainMap)
 	{
 		SCOPE_CYCLE_COUNTER(STAT_BreakingModel);
@@ -1138,8 +1138,8 @@ namespace Chaos
 	}
 
 	DECLARE_CYCLE_STAT(TEXT("TPBDRigidClustering<>::PromoteStrains()"), STAT_PromoteStrains, STATGROUP_Chaos);
-	template<class FPBDRigidsEvolution, class FPBDCollisionConstraint, class T, int d>
-	T TPBDRigidClustering<FPBDRigidsEvolution, FPBDCollisionConstraint, T, d>::PromoteStrains(
+	template<class T_FPBDRigidsEvolution, class T_FPBDCollisionConstraint, class T, int d>
+	T TPBDRigidClustering<T_FPBDRigidsEvolution, T_FPBDCollisionConstraint, T, d>::PromoteStrains(
 		TPBDRigidParticleHandle<T, d>* CurrentNode)
 	{
 		SCOPE_CYCLE_COUNTER(STAT_PromoteStrains);
@@ -1161,8 +1161,8 @@ namespace Chaos
 	}
 
 	DECLARE_CYCLE_STAT(TEXT("TPBDRigidClustering<>::UpdateKinematicProperties()"), STAT_UpdateKinematicProperties, STATGROUP_Chaos);
-	template<class FPBDRigidsEvolution, class FPBDCollisionConstraint, class T, int d>
-	void TPBDRigidClustering<FPBDRigidsEvolution, FPBDCollisionConstraint, T, d>::UpdateKinematicProperties(
+	template<class T_FPBDRigidsEvolution, class T_FPBDCollisionConstraint, class T, int d>
+	void TPBDRigidClustering<T_FPBDRigidsEvolution, T_FPBDCollisionConstraint, T, d>::UpdateKinematicProperties(
 		Chaos::TPBDRigidParticleHandle<float, 3>* Parent)
 	{
 		SCOPE_CYCLE_COUNTER(STAT_UpdateKinematicProperties);
@@ -1207,8 +1207,8 @@ namespace Chaos
 	}
 
 	DECLARE_CYCLE_STAT(TEXT("TPBDRigidClustering<>::SwapBufferedData"), STAT_SwapBufferedData, STATGROUP_Chaos);
-	template<class FPBDRigidsEvolution, class FPBDCollisionConstraint, class T, int d>
-	void TPBDRigidClustering<FPBDRigidsEvolution, FPBDCollisionConstraint, T, d>::SwapBufferedData()
+	template<class T_FPBDRigidsEvolution, class T_FPBDCollisionConstraint, class T, int d>
+	void TPBDRigidClustering<T_FPBDRigidsEvolution, T_FPBDCollisionConstraint, T, d>::SwapBufferedData()
 	{
 		check(false);
 		// TODO: Ryan - this code currently uses MParticles as the only source of clustered particles.
@@ -1247,8 +1247,8 @@ namespace Chaos
 	}
 
 	DECLARE_CYCLE_STAT(TEXT("TPBDRigidClustering<>::GetActiveClusterIndex"), STAT_GetActiveClusterIndex, STATGROUP_Chaos);
-	template<class FPBDRigidsEvolution, class FPBDCollisionConstraint, class T, int d>
-	TPBDRigidParticleHandle<T, d>* TPBDRigidClustering<FPBDRigidsEvolution, FPBDCollisionConstraint, T, d>::GetActiveClusterIndex(
+	template<class T_FPBDRigidsEvolution, class T_FPBDCollisionConstraint, class T, int d>
+	TPBDRigidParticleHandle<T, d>* TPBDRigidClustering<T_FPBDRigidsEvolution, T_FPBDCollisionConstraint, T, d>::GetActiveClusterIndex(
 		TPBDRigidParticleHandle<T, d>* Child)
 	{
 		SCOPE_CYCLE_COUNTER(STAT_GetActiveClusterIndex);
@@ -1259,8 +1259,8 @@ namespace Chaos
 		return Child; 
 	}
 
-	template<class FPBDRigidsEvolution, class FPBDCollisionConstraint, class T, int d>
-	void TPBDRigidClustering<FPBDRigidsEvolution, FPBDCollisionConstraint, T, d>::IncrementPendingClusterCounter(
+	template<class T_FPBDRigidsEvolution, class T_FPBDCollisionConstraint, class T, int d>
+	void TPBDRigidClustering<T_FPBDRigidsEvolution, T_FPBDCollisionConstraint, T, d>::IncrementPendingClusterCounter(
 		uint32 ClusterGroupID)
 	{
 		if (!PendingClusterCounter.Contains(ClusterGroupID))
@@ -1270,8 +1270,8 @@ namespace Chaos
 		PendingClusterCounter[ClusterGroupID]++;
 	}
 
-	template<class FPBDRigidsEvolution, class FPBDCollisionConstraint, class T, int d>
-	void TPBDRigidClustering<FPBDRigidsEvolution, FPBDCollisionConstraint, T, d>::DecrementPendingClusterCounter(
+	template<class T_FPBDRigidsEvolution, class T_FPBDCollisionConstraint, class T, int d>
+	void TPBDRigidClustering<T_FPBDRigidsEvolution, T_FPBDCollisionConstraint, T, d>::DecrementPendingClusterCounter(
 		uint32 ClusterGroupID)
 	{
 		PendingClusterCounter[ClusterGroupID]--;
@@ -1279,8 +1279,8 @@ namespace Chaos
 	}
 
 	DECLARE_CYCLE_STAT(TEXT("TPBDRigidClustering<>::GenerateConnectionGraph"), STAT_GenerateConnectionGraph, STATGROUP_Chaos);
-	template<class FPBDRigidsEvolution, class FPBDCollisionConstraint, class T, int d>
-	void TPBDRigidClustering<FPBDRigidsEvolution, FPBDCollisionConstraint, T, d>::GenerateConnectionGraph(
+	template<class T_FPBDRigidsEvolution, class T_FPBDCollisionConstraint, class T, int d>
+	void TPBDRigidClustering<T_FPBDRigidsEvolution, T_FPBDCollisionConstraint, T, d>::GenerateConnectionGraph(
 		Chaos::TPBDRigidClusteredParticleHandle<float, 3>* Parent,
 		const FClusterCreationParameters<T>& Parameters)
 	{
@@ -1324,8 +1324,8 @@ namespace Chaos
 	}
 
 	DECLARE_CYCLE_STAT(TEXT("TPBDRigidClustering<>::UpdateMassProperties"), STAT_UpdateMassProperties, STATGROUP_Chaos);
-	template<class FPBDRigidsEvolution, class FPBDCollisionConstraint, class T, int d>
-	void TPBDRigidClustering<FPBDRigidsEvolution, FPBDCollisionConstraint, T, d>::UpdateMassProperties(
+	template<class T_FPBDRigidsEvolution, class T_FPBDCollisionConstraint, class T, int d>
+	void TPBDRigidClustering<T_FPBDRigidsEvolution, T_FPBDCollisionConstraint, T, d>::UpdateMassProperties(
 		Chaos::TPBDRigidClusteredParticleHandle<float, 3>* Parent, 
 		TSet<TPBDRigidParticleHandle<T, d>*>& Children, 
 		const TRigidTransform<T, d>* ForceMassOrientation)
@@ -1344,8 +1344,8 @@ namespace Chaos
 	DECLARE_CYCLE_STAT(TEXT("TPBDRigidClustering<>::UpdateGeometry_CopyPoints"), STAT_UpdateGeometry_CopyPoints, STATGROUP_Chaos);
 	DECLARE_CYCLE_STAT(TEXT("TPBDRigidClustering<>::UpdateGeometry_PointsBVH"), STAT_UpdateGeometry_PointsBVH, STATGROUP_Chaos);
 
-	template<class FPBDRigidsEvolution, class FPBDCollisionConstraint, class T, int d>
-	void TPBDRigidClustering<FPBDRigidsEvolution, FPBDCollisionConstraint, T, d>::UpdateGeometry(
+	template<class T_FPBDRigidsEvolution, class T_FPBDCollisionConstraint, class T, int d>
+	void TPBDRigidClustering<T_FPBDRigidsEvolution, T_FPBDCollisionConstraint, T, d>::UpdateGeometry(
 		Chaos::TPBDRigidClusteredParticleHandle<float, 3>* Parent, 
 		const TSet<TPBDRigidParticleHandle<T, d>*>& Children, 
 		TSharedPtr<Chaos::FImplicitObject, ESPMode::ThreadSafe> ProxyGeometry,
@@ -1570,16 +1570,16 @@ namespace Chaos
 	FAutoConsoleVariableRef CVarMinImpulseForStrainEval(TEXT("p.chaos.MinImpulseForStrainEval"), MinImpulseForStrainEval, TEXT("Minimum accumulated impulse before accumulating for strain eval "));
 
 	DECLARE_CYCLE_STAT(TEXT("ComputeStrainFromCollision"), STAT_ComputeStrainFromCollision, STATGROUP_Chaos);
-	template<class FPBDRigidsEvolution, class FPBDCollisionConstraint, class T, int d>
-	void TPBDRigidClustering<FPBDRigidsEvolution, FPBDCollisionConstraint, T, d>::ComputeStrainFromCollision(
-		const FPBDCollisionConstraint& CollisionRule)
+	template<class T_FPBDRigidsEvolution, class T_FPBDCollisionConstraint, class T, int d>
+	void TPBDRigidClustering<T_FPBDRigidsEvolution, T_FPBDCollisionConstraint, T, d>::ComputeStrainFromCollision(
+		const T_FPBDCollisionConstraint& CollisionRule)
 	{
 		SCOPE_CYCLE_COUNTER(STAT_ComputeStrainFromCollision);
 		FClusterMap& MParentToChildren = GetChildrenMap();
 
 		ResetCollisionImpulseArray();
 
-		for (const Chaos::TPBDCollisionConstraintHandle<T, 3> * ContactHandle : CollisionRule.GetConstConstraintHandles())
+		for (const Chaos::FPBDCollisionConstraintHandle* ContactHandle : CollisionRule.GetConstConstraintHandles())
 		{
 			if (ContactHandle->GetAccumulatedImpulse().Size() < MinImpulseForStrainEval)
 			{
@@ -1649,8 +1649,8 @@ namespace Chaos
 	}
 
 	DECLARE_CYCLE_STAT(TEXT("ResetCollisionImpulseArray"), STAT_ResetCollisionImpulseArray, STATGROUP_Chaos);
-	template<class FPBDRigidsEvolution, class FPBDCollisionConstraint, class T, int d>
-	void TPBDRigidClustering<FPBDRigidsEvolution, FPBDCollisionConstraint, T, d>::ResetCollisionImpulseArray()
+	template<class T_FPBDRigidsEvolution, class T_FPBDCollisionConstraint, class T, int d>
+	void TPBDRigidClustering<T_FPBDRigidsEvolution, T_FPBDCollisionConstraint, T, d>::ResetCollisionImpulseArray()
 	{
 		SCOPE_CYCLE_COUNTER(STAT_ResetCollisionImpulseArray);
 		if (MCollisionImpulseArrayDirty)
@@ -1662,8 +1662,8 @@ namespace Chaos
 		}
 	}
 
-	template<class FPBDRigidsEvolution, class FPBDCollisionConstraint, class T, int d>
-	void TPBDRigidClustering<FPBDRigidsEvolution, FPBDCollisionConstraint, T, d>::DisableCluster(
+	template<class T_FPBDRigidsEvolution, class T_FPBDCollisionConstraint, class T, int d>
+	void TPBDRigidClustering<T_FPBDRigidsEvolution, T_FPBDCollisionConstraint, T, d>::DisableCluster(
 		TPBDRigidClusteredParticleHandle<T,d>* ClusteredParticle)
 	{
 		// #note: we don't recursively descend to the children
@@ -1683,8 +1683,8 @@ namespace Chaos
 		MActiveRemovalIndices.Remove(ClusteredParticle);
 	}
 
-	template<class FPBDRigidsEvolution, class FPBDCollisionConstraint, class T, int d>
-	void TPBDRigidClustering<FPBDRigidsEvolution, FPBDCollisionConstraint, T, d>::DisableParticleWithBreakEvent(
+	template<class T_FPBDRigidsEvolution, class T_FPBDCollisionConstraint, class T, int d>
+	void TPBDRigidClustering<T_FPBDRigidsEvolution, T_FPBDCollisionConstraint, T, d>::DisableParticleWithBreakEvent(
 		Chaos::TPBDRigidParticleHandle<float, 3>* Particle)
 	{
 		DisableCluster(Particle->CastToClustered());
@@ -1702,8 +1702,8 @@ namespace Chaos
 	}
 
 	DECLARE_CYCLE_STAT(TEXT("TPBDRigidClustering<>::UpdateConnectivityGraphUsingPointImplicit"), STAT_UpdateConnectivityGraphUsingPointImplicit, STATGROUP_Chaos);
-	template<class FPBDRigidsEvolution, class FPBDCollisionConstraint, class T, int d>
-	void TPBDRigidClustering<FPBDRigidsEvolution, FPBDCollisionConstraint, T, d>::UpdateConnectivityGraphUsingPointImplicit(
+	template<class T_FPBDRigidsEvolution, class T_FPBDCollisionConstraint, class T, int d>
+	void TPBDRigidClustering<T_FPBDRigidsEvolution, T_FPBDCollisionConstraint, T, d>::UpdateConnectivityGraphUsingPointImplicit(
 		Chaos::TPBDRigidClusteredParticleHandle<float, 3>* Parent,
 		const FClusterCreationParameters<T>& Parameters)
 	{
@@ -1777,8 +1777,8 @@ namespace Chaos
 	}
 
 	DECLARE_CYCLE_STAT(TEXT("TPBDRigidClustering<>::FixConnectivityGraphUsingDelaunayTriangulation"), STAT_FixConnectivityGraphUsingDelaunayTriangulation, STATGROUP_Chaos);
-	template<class FPBDRigidsEvolution, class FPBDCollisionConstraint, class T, int d>
-	void TPBDRigidClustering<FPBDRigidsEvolution, FPBDCollisionConstraint, T, d>::FixConnectivityGraphUsingDelaunayTriangulation(
+	template<class T_FPBDRigidsEvolution, class T_FPBDCollisionConstraint, class T, int d>
+	void TPBDRigidClustering<T_FPBDRigidsEvolution, T_FPBDCollisionConstraint, T, d>::FixConnectivityGraphUsingDelaunayTriangulation(
 		Chaos::TPBDRigidClusteredParticleHandle<float, 3>* Parent,
 		const FClusterCreationParameters<T>& Parameters)
 	{
@@ -1933,8 +1933,8 @@ namespace Chaos
 
 	DECLARE_CYCLE_STAT(TEXT("TPBDRigidClustering<>::UpdateConnectivityGraphUsingDelaunayTriangulation"), STAT_UpdateConnectivityGraphUsingDelaunayTriangulation, STATGROUP_Chaos);
 
-	template<class FPBDRigidsEvolution, class FPBDCollisionConstraint, class T, int d>
-	void TPBDRigidClustering<FPBDRigidsEvolution, FPBDCollisionConstraint, T, d>::UpdateConnectivityGraphUsingDelaunayTriangulation(
+	template<class T_FPBDRigidsEvolution, class T_FPBDCollisionConstraint, class T, int d>
+	void TPBDRigidClustering<T_FPBDRigidsEvolution, T_FPBDCollisionConstraint, T, d>::UpdateConnectivityGraphUsingDelaunayTriangulation(
 		Chaos::TPBDRigidClusteredParticleHandle<float, 3>* Parent,
 		const FClusterCreationParameters<T>& Parameters)
 	{
@@ -1972,8 +1972,8 @@ namespace Chaos
 		}
 	}
 
-	//template<class FPBDRigidsEvolution, class FPBDCollisionConstraint, class T, int d>
-	//void TPBDRigidClustering<FPBDRigidsEvolution, FPBDCollisionConstraint, T, d>::AddUniqueConnection(uint32 Index1, uint32 Index2, T Strain)
+	//template<class T_FPBDRigidsEvolution, class T_FPBDCollisionConstraint, class T, int d>
+	//void TPBDRigidClustering<T_FPBDRigidsEvolution, T_FPBDCollisionConstraint, T, d>::AddUniqueConnection(uint32 Index1, uint32 Index2, T Strain)
 	//{
 	//	if (Index1 != Index2)
 	//	{
@@ -1994,8 +1994,8 @@ namespace Chaos
 	//	}
 	//}
 
-	template<class FPBDRigidsEvolution, class FPBDCollisionConstraint, class T, int d>
-	void TPBDRigidClustering<FPBDRigidsEvolution, FPBDCollisionConstraint, T, d>::ConnectNodes(
+	template<class T_FPBDRigidsEvolution, class T_FPBDCollisionConstraint, class T, int d>
+	void TPBDRigidClustering<T_FPBDRigidsEvolution, T_FPBDCollisionConstraint, T, d>::ConnectNodes(
 		TPBDRigidParticleHandle<T, d>* Child1,
 		TPBDRigidParticleHandle<T, d>* Child2)
 	{
@@ -2005,8 +2005,8 @@ namespace Chaos
 		ConnectNodes(ClusteredChild1, ClusteredChild2);
 	}
 
-	template<class FPBDRigidsEvolution, class FPBDCollisionConstraint, class T, int d>
-	void TPBDRigidClustering<FPBDRigidsEvolution, FPBDCollisionConstraint, T, d>::ConnectNodes(
+	template<class T_FPBDRigidsEvolution, class T_FPBDCollisionConstraint, class T, int d>
+	void TPBDRigidClustering<T_FPBDRigidsEvolution, T_FPBDCollisionConstraint, T, d>::ConnectNodes(
 		TPBDRigidClusteredParticleHandle<T, d>* ClusteredChild1,
 		TPBDRigidClusteredParticleHandle<T, d>* ClusteredChild2)
 	{
@@ -2028,16 +2028,16 @@ namespace Chaos
 		}
 	}
 
-	template<class FPBDRigidsEvolution, class FPBDCollisionConstraint, class T, int d>
-	void TPBDRigidClustering<FPBDRigidsEvolution, FPBDCollisionConstraint, T, d>::RemoveNodeConnections(
+	template<class T_FPBDRigidsEvolution, class T_FPBDCollisionConstraint, class T, int d>
+	void TPBDRigidClustering<T_FPBDRigidsEvolution, T_FPBDCollisionConstraint, T, d>::RemoveNodeConnections(
 		TPBDRigidParticleHandle<T, d>* Child)
 	{
 		RemoveNodeConnections(Child->CastToClustered());
 	}
 
 	DECLARE_CYCLE_STAT(TEXT("TPBDRigidClustering<>::RemoveNodeConnections"), STAT_RemoveNodeConnections, STATGROUP_Chaos);
-	template<class FPBDRigidsEvolution, class FPBDCollisionConstraint, class T, int d>
-	void TPBDRigidClustering<FPBDRigidsEvolution, FPBDCollisionConstraint, T, d>::RemoveNodeConnections(
+	template<class T_FPBDRigidsEvolution, class T_FPBDCollisionConstraint, class T, int d>
+	void TPBDRigidClustering<T_FPBDRigidsEvolution, T_FPBDCollisionConstraint, T, d>::RemoveNodeConnections(
 		TPBDRigidClusteredParticleHandle<T, d>* ClusteredChild)
 	{
 		SCOPE_CYCLE_COUNTER(STAT_RemoveNodeConnections);
@@ -2057,8 +2057,8 @@ namespace Chaos
 
 
 /*
-	template<class FPBDRigidsEvolution, class FPBDCollisionConstraint, class T, int d>
-	void TPBDRigidClustering<FPBDRigidsEvolution, FPBDCollisionConstraint, T, d>::ClearPendingClusterCounter(uint32 ClusterGroupID)
+	template<class T_FPBDRigidsEvolution, class T_FPBDCollisionConstraint, class T, int d>
+	void TPBDRigidClustering<T_FPBDRigidsEvolution, T_FPBDCollisionConstraint, T, d>::ClearPendingClusterCounter(uint32 ClusterGroupID)
 	{
 		PendingClusterCounter[ClusterGroupID]=0;
 		//ensure(0 <= PendingClusterCounter[ClusterGroupID]);
@@ -2067,7 +2067,7 @@ namespace Chaos
 } // namespace Chaos
 
 using namespace Chaos;
-template class CHAOS_API Chaos::TPBDRigidClustering<FPBDRigidsEvolutionGBF, TPBDCollisionConstraints<float, 3>, float, 3>;
+template class CHAOS_API Chaos::TPBDRigidClustering<FPBDRigidsEvolutionGBF, FPBDCollisionConstraints, float, 3>;
 
 
 template CHAOS_API void Chaos::UpdateClusterMassProperties<float, 3>(
diff --git a/Engine/Source/Runtime/Experimental/Chaos/Public/Chaos/AABB.h b/Engine/Source/Runtime/Experimental/Chaos/Public/Chaos/AABB.h
--- a/Engine/Source/Runtime/Experimental/Chaos/Public/Chaos/AABB.h
+++ b/Engine/Source/Runtime/Experimental/Chaos/Public/Chaos/AABB.h
@@ -607,6 +607,4 @@ namespace Chaos
 			return SamplePoints;
 		}
 	};
-
-	using FAABB3 = TAABB<FReal, 3>;
 }
\ No newline at end of file
diff --git a/Engine/Source/Runtime/Experimental/Chaos/Public/Chaos/ChaosDebugDraw.h b/Engine/Source/Runtime/Experimental/Chaos/Public/Chaos/ChaosDebugDraw.h
--- a/Engine/Source/Runtime/Experimental/Chaos/Public/Chaos/ChaosDebugDraw.h
+++ b/Engine/Source/Runtime/Experimental/Chaos/Public/Chaos/ChaosDebugDraw.h
@@ -40,9 +40,9 @@ namespace Chaos
 		CHAOS_API void DrawParticleTransforms(const FRigidTransform3& SpaceTransform, const TParticleView<TGeometryParticles<float, 3>>& ParticlesView);
 		CHAOS_API void DrawParticleTransforms(const FRigidTransform3& SpaceTransform, const TParticleView<TKinematicGeometryParticles<float, 3>>& ParticlesView);
 		CHAOS_API void DrawParticleTransforms(const FRigidTransform3& SpaceTransform, const TParticleView<TPBDRigidParticles<float, 3>>& ParticlesView);
-		CHAOS_API void DrawParticleCollisions(const FRigidTransform3& SpaceTransform, const TGeometryParticleHandle<float, 3>* Particle, const TPBDCollisionConstraints<float, 3>& Collisions);
-		CHAOS_API void DrawCollisions(const FRigidTransform3& SpaceTransform, const TPBDCollisionConstraints<float, 3>& Collisions, float ColorScale);
-		CHAOS_API void DrawCollisions(const FRigidTransform3& SpaceTransform, const TArray<TPBDCollisionConstraintHandle<float, 3>*>& ConstraintHandles, float ColorScale);
+		CHAOS_API void DrawParticleCollisions(const FRigidTransform3& SpaceTransform, const TGeometryParticleHandle<float, 3>* Particle, const FPBDCollisionConstraints& Collisions);
+		CHAOS_API void DrawCollisions(const FRigidTransform3& SpaceTransform, const FPBDCollisionConstraints& Collisions, float ColorScale);
+		CHAOS_API void DrawCollisions(const FRigidTransform3& SpaceTransform, const TArray<FPBDCollisionConstraintHandle*>& ConstraintHandles, float ColorScale);
 		CHAOS_API void DrawJointConstraints(const FRigidTransform3& SpaceTransform, const TArray<FPBDJointConstraintHandle*>& ConstraintHandles, float ColorScale, uint32 FeatureMask = (uint32)EDebugDrawJointFeature::Default);
 		CHAOS_API void DrawJointConstraints(const FRigidTransform3& SpaceTransform, const FPBDJointConstraints& Constraints, float ColorScale, uint32 FeatureMask = (uint32)EDebugDrawJointFeature::Default);
 
diff --git a/Engine/Source/Runtime/Experimental/Chaos/Public/Chaos/ChaosDebugDrawDeclares.h b/Engine/Source/Runtime/Experimental/Chaos/Public/Chaos/ChaosDebugDrawDeclares.h
--- a/Engine/Source/Runtime/Experimental/Chaos/Public/Chaos/ChaosDebugDrawDeclares.h
+++ b/Engine/Source/Runtime/Experimental/Chaos/Public/Chaos/ChaosDebugDrawDeclares.h
@@ -14,11 +14,9 @@ namespace Chaos
 {
 	class FConstraintHandle;
 
-	template<typename T, int d>
-	class TPBDCollisionConstraints;
+	class FPBDCollisionConstraints;
 
-	template<typename T, int d>
-	class TPBDCollisionConstraintHandle;
+	class FPBDCollisionConstraintHandle;
 
 	class FPBDJointConstraintHandle;
 
diff --git a/Engine/Source/Runtime/Experimental/Chaos/Public/Chaos/Collision/CollisionReceiver.h b/Engine/Source/Runtime/Experimental/Chaos/Public/Chaos/Collision/CollisionReceiver.h
--- a/Engine/Source/Runtime/Experimental/Chaos/Public/Chaos/Collision/CollisionReceiver.h
+++ b/Engine/Source/Runtime/Experimental/Chaos/Public/Chaos/Collision/CollisionReceiver.h
@@ -14,9 +14,7 @@ namespace Chaos
 	class FAsyncCollisionReceiver
 	{
 	public:
-		using FCollisionConstraints = TPBDCollisionConstraints<FReal, 3>;
-
-		FAsyncCollisionReceiver(FCollisionConstraints& InCollisionConstraints)
+		FAsyncCollisionReceiver(FPBDCollisionConstraints& InCollisionConstraints)
 			: CollisionConstraints(InCollisionConstraints)
 		{
 		}
@@ -26,16 +24,16 @@ namespace Chaos
 		 */
 		void ReceiveCollisions(const FCollisionConstraintsArray& Constraints)
 		{
-			for (const TRigidBodyPointContactConstraint<FReal, 3>& Constraint : Constraints.SinglePointConstraints)
+			for (const FRigidBodyPointContactConstraint& Constraint : Constraints.SinglePointConstraints)
 			{
 				SinglePointQueue.Enqueue(Constraint);
 			}
-			for (const TRigidBodySweptPointContactConstraint<FReal, 3>& Constraint : Constraints.SinglePointSweptConstraints)
+			for (const FRigidBodySweptPointContactConstraint& Constraint : Constraints.SinglePointSweptConstraints)
 			{
 				SingleSweptPointQueue.Enqueue(Constraint);
 			}
 
-			for (const TRigidBodyMultiPointContactConstraint<FReal, 3>& Constraint : Constraints.MultiPointConstraints)
+			for (const FRigidBodyMultiPointContactConstraint& Constraint : Constraints.MultiPointConstraints)
 			{
 				MultiPointQueue.Enqueue(Constraint);
 			}
@@ -51,7 +49,7 @@ namespace Chaos
 			//    NarrowPhase is creating temporary constraints for all interacting bodies 
 			//    even though we might already have those bodies within the constraint.
 			//    So there is a bunch of wasted computation when the CreateConstraints is called. 
-			TRigidBodyPointContactConstraint<FReal, 3> SinglePointConstraint;
+			FRigidBodyPointContactConstraint SinglePointConstraint;
 			while (SinglePointQueue.Dequeue(SinglePointConstraint))
 			{
 				if (!CollisionConstraints.Contains(&SinglePointConstraint))
@@ -59,7 +57,7 @@ namespace Chaos
 					CollisionConstraints.AddConstraint(SinglePointConstraint);
 				}
 			}
-			TRigidBodySweptPointContactConstraint<FReal, 3> SinglePointSweptConstraint;
+			FRigidBodySweptPointContactConstraint SinglePointSweptConstraint;
 			while (SingleSweptPointQueue.Dequeue(SinglePointSweptConstraint))
 			{
 				if (!CollisionConstraints.Contains(&SinglePointSweptConstraint))
@@ -68,7 +66,7 @@ namespace Chaos
 				}
 			}
 
-			TRigidBodyMultiPointContactConstraint<FReal, 3> MultiPointConstraint;
+			FRigidBodyMultiPointContactConstraint MultiPointConstraint;
 			while (MultiPointQueue.Dequeue(MultiPointConstraint))
 			{
 				if (!CollisionConstraints.Contains(&MultiPointConstraint))
@@ -80,10 +78,10 @@ namespace Chaos
 
 	private:
 		//todo(ocohen): use per thread buffer instead, need better support than ParallelFor for this
-		TQueue<TRigidBodyPointContactConstraint<FReal, 3>, EQueueMode::Mpsc> SinglePointQueue;
-		TQueue<TRigidBodySweptPointContactConstraint<FReal, 3>, EQueueMode::Mpsc> SingleSweptPointQueue;
-		TQueue<TRigidBodyMultiPointContactConstraint<FReal, 3>, EQueueMode::Mpsc> MultiPointQueue;
-		FCollisionConstraints& CollisionConstraints;
+		TQueue<FRigidBodyPointContactConstraint, EQueueMode::Mpsc> SinglePointQueue;
+		TQueue<FRigidBodySweptPointContactConstraint, EQueueMode::Mpsc> SingleSweptPointQueue;
+		TQueue<FRigidBodyMultiPointContactConstraint, EQueueMode::Mpsc> MultiPointQueue;
+		FPBDCollisionConstraints& CollisionConstraints;
 	};
 
 
@@ -94,9 +92,7 @@ namespace Chaos
 	class FSyncCollisionReceiver
 	{
 	public:
-		using FCollisionConstraints = TPBDCollisionConstraints<FReal, 3>;
-
-		FSyncCollisionReceiver(FCollisionConstraints& InCollisionConstraints)
+		FSyncCollisionReceiver(FPBDCollisionConstraints& InCollisionConstraints)
 			: CollisionConstraints(InCollisionConstraints)
 		{
 		}
@@ -121,6 +117,6 @@ namespace Chaos
 		}
 
 	private:
-		FCollisionConstraints& CollisionConstraints;
+		FPBDCollisionConstraints& CollisionConstraints;
 	};
 }
\ No newline at end of file
diff --git a/Engine/Source/Runtime/Experimental/Chaos/Public/Chaos/Collision/NarrowPhase.h b/Engine/Source/Runtime/Experimental/Chaos/Public/Chaos/Collision/NarrowPhase.h
--- a/Engine/Source/Runtime/Experimental/Chaos/Public/Chaos/Collision/NarrowPhase.h
+++ b/Engine/Source/Runtime/Experimental/Chaos/Public/Chaos/Collision/NarrowPhase.h
@@ -1,6 +1,7 @@
 // Copyright Epic Games, Inc. All Rights Reserved.
 #pragma once
 
+#include "Chaos/Core.h"
 #include "Chaos/Collision/CollisionReceiver.h"
 #include "Chaos/Collision/StatsData.h"
 #include "Chaos/CollisionResolution.h"
@@ -50,7 +51,7 @@ namespace Chaos
 				//   determine if the constraint is already defined, and then opt out of 
 				//   the creation process. 
 				//
-				Collisions::ConstructConstraints<FReal, 3>(Particle0, Particle1, Particle0->Geometry().Get(), Particle1->Geometry().Get(), FRigidTransform3(), FRigidTransform3(), CullDistance, Context, NewConstraints);
+				Collisions::ConstructConstraints(Particle0, Particle1, Particle0->Geometry().Get(), Particle1->Geometry().Get(), FRigidTransform3(), FRigidTransform3(), CullDistance, Context, NewConstraints);
 
 				CHAOS_COLLISION_STAT(if (NewConstraints.Num()) { StatData.IncrementCountNP(NewConstraints.Num()); });
 				CHAOS_COLLISION_STAT(if (!NewConstraints.Num()) { StatData.IncrementRejectedNP(); });
diff --git a/Engine/Source/Runtime/Experimental/Chaos/Public/Chaos/Collision/PBDCollisionConstraint.h b/Engine/Source/Runtime/Experimental/Chaos/Public/Chaos/Collision/PBDCollisionConstraint.h
new file mode 100644
--- /dev/null
+++ b/Engine/Source/Runtime/Experimental/Chaos/Public/Chaos/Collision/PBDCollisionConstraint.h
@@ -0,0 +1,324 @@
+// Copyright Epic Games, Inc. All Rights Reserved.
+#pragma once
+
+#include "CoreMinimal.h"
+
+#include "Chaos/Core.h"
+#include "Chaos/CollisionResolutionTypes.h"
+#include "Chaos/ParticleHandleFwd.h"
+#include "Chaos/Vector.h"
+
+namespace Chaos
+{
+	class FImplicitObject;
+	class FPBDCollisionConstraintHandle;
+
+	/*
+	*
+	*/
+	class CHAOS_API FCollisionContact
+	{
+	public:
+		FCollisionContact(const FImplicitObject* InImplicit0 = nullptr, const FImplicitObject* InImplicit1 = nullptr)
+			: bDisabled(false), Normal(0), Location(0), Phi(FLT_MAX), Friction(0), AngularFriction(0), Restitution(0), InvInertiaScale0(1.f), InvInertiaScale1(1.f), ShapesType(EContactShapesType::Unknown)
+		{
+			Implicit[0] = InImplicit0;
+			Implicit[1] = InImplicit1;
+		}
+
+		bool bDisabled;
+		FVec3 Normal;
+		FVec3 Location;
+		FReal Phi;
+
+		FReal Friction;
+		FReal AngularFriction;
+		FReal Restitution;
+		FReal InvInertiaScale0;
+		FReal InvInertiaScale1;
+		EContactShapesType ShapesType;
+
+
+		FString ToString() const
+		{
+			return FString::Printf(TEXT("Location:%s, Normal:%s, Phi:%f"), *Location.ToString(), *Normal.ToString(), Phi);
+		}
+
+		const FImplicitObject* Implicit[2]; // {Of Particle[0], Of Particle[1]}
+	};
+
+	/*
+	*
+	*/
+	class CHAOS_API FCollisionConstraintBase
+	{
+	public:
+		using FGeometryParticleHandle = TGeometryParticleHandle<FReal, 3>;
+
+		enum class FType
+		{
+			None = 0,			// default value also indicates a invalid constraint
+			SinglePoint,		// TRigidBodyPointContactConstraint
+			SinglePointSwept,	// TRigidBodySweptPointContactConstraint
+			MultiPoint			// TRigidBodyMultiPointContactConstraint
+		};
+
+		FCollisionConstraintBase(FType InType = FType::None)
+			: AccumulatedImpulse(0)
+			, Timestamp(-INT_MAX)
+			, ConstraintHandle(nullptr)
+			, Type(InType)
+		{
+			ImplicitTransform[0] = FRigidTransform3::Identity; ImplicitTransform[1] = FRigidTransform3::Identity;
+			Manifold.Implicit[0] = nullptr; Manifold.Implicit[1] = nullptr;
+			Particle[0] = nullptr; Particle[1] = nullptr;
+		}
+
+		FCollisionConstraintBase(
+			FGeometryParticleHandle* Particle0, const FImplicitObject* Implicit0, const FRigidTransform3& Transform0,
+			FGeometryParticleHandle* Particle1, const FImplicitObject* Implicit1, const FRigidTransform3& Transform1,
+			FType InType, EContactShapesType ShapesType, int32 InTimestamp = -INT_MAX)
+			: AccumulatedImpulse(0)
+			, Timestamp(InTimestamp)
+			, ConstraintHandle(nullptr)
+			, Type(InType)
+		{
+			ImplicitTransform[0] = Transform0; ImplicitTransform[1] = Transform1;
+			Manifold.Implicit[0] = Implicit0; Manifold.Implicit[1] = Implicit1;
+			Manifold.ShapesType = ShapesType;
+			Particle[0] = Particle0; Particle[1] = Particle1;
+		}
+
+		FType GetType() const { return Type; }
+
+		template<class AS_T> AS_T* As() { return static_cast<AS_T*>(this); }
+		template<class AS_T> const AS_T* As() const { return static_cast<const AS_T*>(this); }
+
+		bool ContainsManifold(const FImplicitObject* A, const FImplicitObject* B) const { return A == Manifold.Implicit[0] && B == Manifold.Implicit[1]; }
+		void SetManifold(const FImplicitObject* A, const FImplicitObject* B) { Manifold.Implicit[0] = A; Manifold.Implicit[1] = B; }
+		const FCollisionContact& GetManifold() const { return Manifold; }
+
+		//
+		// API
+		//
+
+		void ResetPhi(FReal InPhi) { SetPhi(InPhi); }
+
+		void SetPhi(FReal InPhi) { Manifold.Phi = InPhi; }
+		FReal GetPhi() const { return Manifold.Phi; }
+
+		void SetDisabled(bool bInDisabled) { Manifold.bDisabled = bInDisabled; }
+		bool GetDisabled() const { return Manifold.bDisabled; }
+
+		void SetNormal(const FVec3& InNormal) { Manifold.Normal = InNormal; }
+		FVec3 GetNormal() const { return Manifold.Normal; }
+
+		void SetLocation(const FVec3& InLocation) { Manifold.Location = InLocation; }
+		FVec3 GetLocation() const { return Manifold.Location; }
+
+		void SetInvInertiaScale0(const FReal InInvInertiaScale) { Manifold.InvInertiaScale0 = InInvInertiaScale; }
+		FReal GetInvInertiaScale0() const { return Manifold.InvInertiaScale0; }
+
+		void SetInvInertiaScale1(const FReal InInvInertiaScale) { Manifold.InvInertiaScale1 = InInvInertiaScale; }
+		FReal GetInvInertiaScale1() const { return Manifold.InvInertiaScale1; }
+
+
+		FString ToString() const;
+
+		FRigidTransform3 ImplicitTransform[2]; // { Point, Volume }
+		FGeometryParticleHandle* Particle[2]; // { Point, Volume }
+		FVec3 AccumulatedImpulse;
+		FCollisionContact Manifold;// @todo(chaos): rename
+		int32 Timestamp;
+		FPBDCollisionConstraintHandle* ConstraintHandle;
+
+	private:
+
+		FType Type;
+	};
+
+
+	/*
+	*
+	*/
+	class CHAOS_API FRigidBodyPointContactConstraint : public FCollisionConstraintBase
+	{
+	public:
+		using Base = FCollisionConstraintBase;
+		using FGeometryParticleHandle = TGeometryParticleHandle<FReal, 3>;
+
+		FRigidBodyPointContactConstraint() : Base(Base::FType::SinglePoint) {}
+		FRigidBodyPointContactConstraint(
+			FGeometryParticleHandle* Particle0, const FImplicitObject* Implicit0, const FRigidTransform3& Transform0,
+			FGeometryParticleHandle* Particle1, const FImplicitObject* Implicit1, const FRigidTransform3& Transform1,
+			EContactShapesType ShapesType)
+			: Base(Particle0, Implicit0, Transform0, Particle1, Implicit1, Transform1, Base::FType::SinglePoint, ShapesType) {}
+
+		static typename Base::FType StaticType() { return Base::FType::SinglePoint; };
+
+	protected:
+		// For use by derived types that can be used as point constraints in Update
+		FRigidBodyPointContactConstraint(typename Base::FType InType) : Base(InType) {}
+		FRigidBodyPointContactConstraint(
+			FGeometryParticleHandle* Particle0, const FImplicitObject* Implicit0, const FRigidTransform3& Transform0,
+			FGeometryParticleHandle* Particle1, const FImplicitObject* Implicit1, const FRigidTransform3& Transform1,
+			typename Base::FType InType, EContactShapesType ShapesType)
+			: Base(Particle0, Implicit0, Transform0, Particle1, Implicit1, Transform1, InType, ShapesType) {}
+	};
+
+
+	/*
+	*
+	*/
+	class CHAOS_API FRigidBodyMultiPointContactConstraint : public FRigidBodyPointContactConstraint
+	{
+	public:
+		using Base = FRigidBodyPointContactConstraint;
+		using FGeometryParticleHandle = TGeometryParticleHandle<FReal, 3>;
+
+		FRigidBodyMultiPointContactConstraint() : Base(FCollisionConstraintBase::FType::MultiPoint), PlaneOwnerIndex(INDEX_NONE), PlaneFaceIndex(INDEX_NONE), PlaneNormal(0), PlanePosition(0) {}
+		FRigidBodyMultiPointContactConstraint(
+			FGeometryParticleHandle* Particle0, const FImplicitObject* Implicit0, const FRigidTransform3& Transform0,
+			FGeometryParticleHandle* Particle1, const FImplicitObject* Implicit1, const FRigidTransform3& Transform1,
+			EContactShapesType ShapesType)
+			: Base(Particle0, Implicit0, Transform0, Particle1, Implicit1, Transform1, Base::FType::MultiPoint, ShapesType)
+			, PlaneOwnerIndex(INDEX_NONE), PlaneFaceIndex(INDEX_NONE), PlaneNormal(0), PlanePosition(0)
+		{}
+
+		static typename Base::FType StaticType() { return Base::FType::MultiPoint; };
+
+		// Get the particle that owns the plane
+		FGeometryParticleHandle* PlaneParticleHandle() const { check(PlaneOwnerIndex >= 0 && PlaneOwnerIndex < 2); return Particle[PlaneOwnerIndex]; }
+
+		// Get the particle that owns the manifold sample points
+		FGeometryParticleHandle* PointsParticleHandle() const { check(PlaneOwnerIndex >= 0 && PlaneOwnerIndex < 2); return Particle[1 - PlaneOwnerIndex]; }
+
+		int32 GetManifoldPlaneOwnerIndex() const { return PlaneOwnerIndex; }
+		int32 GetManifoldPlaneFaceIndex() const { return PlaneFaceIndex; }
+		const FVec3& GetManifoldPlaneNormal() const { return PlaneNormal; }
+		const FVec3& GetManifoldPlanePosition() const { return PlanePosition; }
+		void SetManifoldPlane(int32 OwnerIndex, int32 FaceIndex, const FVec3& Normal, const FVec3& Pos)
+		{
+			PlaneOwnerIndex = OwnerIndex;
+			PlaneFaceIndex = FaceIndex;
+			PlaneNormal = Normal;
+			PlanePosition = Pos;
+		}
+
+		int32               NumManifoldPoints() const { return Points.Num(); }
+		const FVec3& GetManifoldPoint(int32 Index) const { check(0 <= Index && Index < NumManifoldPoints()); return Points[Index]; }
+		void				SetManifoldPoint(int32 Index, const FVec3& Point) { check(0 <= Index && Index < NumManifoldPoints()); Points[Index] = Point; }
+		void                AddManifoldPoint(const FVec3& Point) { Points.Add(Point); };
+		void                ResetManifoldPoints(int32 NewSize = 0) { Points.Reset(NewSize); }
+
+	private:
+		// Manifold plane data
+		int PlaneOwnerIndex; // index of particle which owns the plane. The other owns the sample positions
+		int PlaneFaceIndex; // index of face used as the manifold plane on plane-owner body
+		FVec3 PlaneNormal; // local space contact normal on plane-owner
+		FVec3 PlanePosition; // local space surface position on plane-owner
+
+		// Manifold point data
+		static const int32 MaxPoints = 4;
+		TArray<FVec3, TInlineAllocator<MaxPoints>> Points; // @todo(chaos): use TFixedAllocator when we handle building the best manifold properly
+	};
+
+	/*
+	*
+	*/
+	class CHAOS_API FRigidBodySweptPointContactConstraint : public FRigidBodyPointContactConstraint
+	{
+	public:
+		using Base = FRigidBodyPointContactConstraint;
+		using FGeometryParticleHandle = TGeometryParticleHandle<FReal, 3>;
+
+		FRigidBodySweptPointContactConstraint() : Base(Base::FType::SinglePointSwept) {}
+		FRigidBodySweptPointContactConstraint(
+			FGeometryParticleHandle* Particle0, const FImplicitObject* Implicit0, const FRigidTransform3& Transform0,
+			FGeometryParticleHandle* Particle1, const FImplicitObject* Implicit1, const FRigidTransform3& Transform1,
+			EContactShapesType ShapesType)
+			: Base(Particle0, Implicit0, Transform0, Particle1, Implicit1, Transform1, Base::FType::SinglePointSwept, ShapesType)
+			, TimeOfImpact(0) {}
+
+		// Value in range [0,1] used to interpolate P between [X,P] that we will rollback to when solving at time of impact.
+		FReal TimeOfImpact;
+
+		static typename Base::FType StaticType() { return Base::FType::SinglePointSwept; };
+	};
+
+
+	struct CHAOS_API FCollisionConstraintsArray
+	{
+		static const int32 InlineMaxConstraints = 8;
+
+		TArray<FRigidBodyPointContactConstraint, TInlineAllocator<InlineMaxConstraints>> SinglePointConstraints;
+		TArray<FRigidBodySweptPointContactConstraint, TInlineAllocator<InlineMaxConstraints>> SinglePointSweptConstraints;
+		TArray<FRigidBodyMultiPointContactConstraint, TInlineAllocator<InlineMaxConstraints>> MultiPointConstraints;
+
+		int32 Num() const { return SinglePointConstraints.Num() + SinglePointSweptConstraints.Num() + MultiPointConstraints.Num(); }
+
+		void Reset()
+		{
+			SinglePointConstraints.Reset();
+			SinglePointSweptConstraints.Reset();
+			MultiPointConstraints.Reset();
+		}
+
+		void Empty()
+		{
+			SinglePointConstraints.Empty();
+			SinglePointSweptConstraints.Empty();
+			MultiPointConstraints.Empty();
+		}
+
+
+		FRigidBodyPointContactConstraint* Add(const FRigidBodyPointContactConstraint& C)
+		{
+			int32 ConstraintIndex = SinglePointConstraints.Add(C);
+			return &SinglePointConstraints[ConstraintIndex];
+		}
+
+		FRigidBodySweptPointContactConstraint* Add(const FRigidBodySweptPointContactConstraint& C)
+		{
+			int32 ConstraintIndex = SinglePointSweptConstraints.Add(C);
+			return &SinglePointSweptConstraints[ConstraintIndex];
+		}
+
+		FRigidBodyMultiPointContactConstraint* Add(const FRigidBodyMultiPointContactConstraint& C)
+		{
+			int32 ConstraintIndex = MultiPointConstraints.Add(C);
+			return &MultiPointConstraints[ConstraintIndex];
+		}
+
+
+		FRigidBodyPointContactConstraint* TryAdd(FReal MaxPhi, const FRigidBodyPointContactConstraint& C)
+		{
+			if (C.GetPhi() < MaxPhi)
+			{
+				int32 ConstraintIndex = SinglePointConstraints.Add(C);
+				return &SinglePointConstraints[ConstraintIndex];
+			}
+			return nullptr;
+		}
+
+		FRigidBodySweptPointContactConstraint* TryAdd(FReal MaxPhi, const FRigidBodySweptPointContactConstraint& C)
+		{
+			if (C.GetPhi() < MaxPhi)
+			{
+				int32 ConstraintIndex = SinglePointSweptConstraints.Add(C);
+				return &SinglePointSweptConstraints[ConstraintIndex];
+			}
+			return nullptr;
+		}
+
+		FRigidBodyMultiPointContactConstraint* TryAdd(FReal MaxPhi, const FRigidBodyMultiPointContactConstraint& C)
+		{
+			if (C.GetPhi() < MaxPhi)
+			{
+				int32 ConstraintIndex = MultiPointConstraints.Add(C);
+				return &MultiPointConstraints[ConstraintIndex];
+			}
+			return nullptr;
+		}
+	};
+}
diff --git a/Engine/Source/Runtime/Experimental/Chaos/Public/Chaos/Collision/PBDCollisionConstraintHandle.h b/Engine/Source/Runtime/Experimental/Chaos/Public/Chaos/Collision/PBDCollisionConstraintHandle.h
new file mode 100644
--- /dev/null
+++ b/Engine/Source/Runtime/Experimental/Chaos/Public/Chaos/Collision/PBDCollisionConstraintHandle.h
@@ -0,0 +1,96 @@
+// Copyright Epic Games, Inc. All Rights Reserved.
+#pragma once
+
+#include "CoreMinimal.h"
+
+#include "Chaos/CollisionResolutionTypes.h"
+#include "Chaos/Collision/PBDCollisionConstraint.h"
+#include "Chaos/ConstraintHandle.h"
+#include "Chaos/PBDConstraintContainer.h"
+
+namespace Chaos
+{
+	class FPBDCollisionConstraints;
+
+	class CHAOS_API FPBDCollisionConstraintHandle : public TContainerConstraintHandle<FPBDCollisionConstraints>
+	{
+	public:
+		using Base = TContainerConstraintHandle<FPBDCollisionConstraints>;
+		using FImplicitPair = TPair<const FImplicitObject*, const FImplicitObject*>;
+		using FGeometryPair = TPair<const TGeometryParticleHandle<FReal, 3>*, const TGeometryParticleHandle<FReal, 3>*>;
+		using FHandleKey = TPair<FImplicitPair, FGeometryPair>;
+
+
+		FPBDCollisionConstraintHandle()
+			: ConstraintType(FCollisionConstraintBase::FType::None)
+		{}
+
+		FPBDCollisionConstraintHandle(FPBDCollisionConstraints* InConstraintContainer, int32 InConstraintIndex, typename FCollisionConstraintBase::FType InType)
+			: TContainerConstraintHandle<FPBDCollisionConstraints>(InConstraintContainer, InConstraintIndex)
+			, ConstraintType(InType)
+		{
+		}
+
+		FHandleKey GetKey();
+
+		static FHandleKey MakeKey(const TGeometryParticleHandle<FReal, 3>* Particle0, const TGeometryParticleHandle<FReal, 3>* Particle1,
+			const FImplicitObject* Implicit0, const FImplicitObject* Implicit1)
+		{
+			return FHandleKey(FImplicitPair(Implicit0, Implicit1), FGeometryPair(Particle0, Particle1));
+		}
+
+		static FHandleKey MakeKey(const FCollisionConstraintBase* Base)
+		{
+			return FHandleKey(FImplicitPair(Base->Manifold.Implicit[0], Base->Manifold.Implicit[1]), FGeometryPair(Base->Particle[0], Base->Particle[1]));
+		}
+
+
+		const FCollisionConstraintBase& GetContact() const;
+		FCollisionConstraintBase& GetContact();
+
+		const FRigidBodyPointContactConstraint& GetPointContact() const;
+		FRigidBodyPointContactConstraint& GetPointContact();
+
+		const FRigidBodySweptPointContactConstraint& GetSweptPointContact() const;
+		FRigidBodySweptPointContactConstraint& GetSweptPointContact();
+
+		const FRigidBodyMultiPointContactConstraint& GetMultiPointContact() const;
+		FRigidBodyMultiPointContactConstraint& GetMultiPointContact();
+
+		typename FCollisionConstraintBase::FType GetType() const { return ConstraintType; }
+
+		void SetConstraintIndex(int32 IndexIn, typename FCollisionConstraintBase::FType InType)
+		{
+			ConstraintIndex = IndexIn;
+			ConstraintType = InType;
+		}
+
+		FVec3 GetContactLocation() const
+		{
+			return GetContact().GetLocation();
+		}
+
+		FVec3 GetAccumulatedImpulse() const
+		{
+			return GetContact().AccumulatedImpulse;
+		}
+
+		TVector<const TGeometryParticleHandle<FReal, 3>*, 2> GetConstrainedParticles() const
+		{
+			return { GetContact().Particle[0], GetContact().Particle[1] };
+		}
+
+		TVector<TGeometryParticleHandle<FReal, 3>*, 2> GetConstrainedParticles()
+		{
+			return { GetContact().Particle[0], GetContact().Particle[1] };
+		}
+
+
+	protected:
+		typename FCollisionConstraintBase::FType ConstraintType;
+		using Base::ConstraintIndex;
+		using Base::ConstraintContainer;
+
+
+	};
+}
\ No newline at end of file
diff --git a/Engine/Source/Runtime/Experimental/Chaos/Public/Chaos/Collision/SpatialAccelerationBroadPhase.h b/Engine/Source/Runtime/Experimental/Chaos/Public/Chaos/Collision/SpatialAccelerationBroadPhase.h
--- a/Engine/Source/Runtime/Experimental/Chaos/Public/Chaos/Collision/SpatialAccelerationBroadPhase.h
+++ b/Engine/Source/Runtime/Experimental/Chaos/Public/Chaos/Collision/SpatialAccelerationBroadPhase.h
@@ -19,28 +19,27 @@ namespace Chaos
 	/**
 	 *
 	 */
-	template <typename T, int d>
-	struct TSimOverlapVisitor
+	struct FSimOverlapVisitor
 	{
-		TArray<TAccelerationStructureHandle<T, d>>& Intersections;
-		TSimOverlapVisitor(TArray<TAccelerationStructureHandle<T, d>>& InIntersections)
+		TArray<TAccelerationStructureHandle<FReal, 3>>& Intersections;
+		FSimOverlapVisitor(TArray<TAccelerationStructureHandle<FReal, 3>>& InIntersections)
 			: Intersections(InIntersections)
 		{
 		}
 
-		bool VisitOverlap(const TSpatialVisitorData<TAccelerationStructureHandle<T, d>>& Instance)
+		bool VisitOverlap(const TSpatialVisitorData<TAccelerationStructureHandle<FReal, 3>>& Instance)
 		{
 			Intersections.Add(Instance.Payload);
 			return true;
 		}
 
-		bool VisitSweep(TSpatialVisitorData<TAccelerationStructureHandle<T, d>>, FQueryFastData& CurData)
+		bool VisitSweep(TSpatialVisitorData<TAccelerationStructureHandle<FReal, 3>>, FQueryFastData& CurData)
 		{
 			check(false);
 			return false;
 		}
 
-		bool VisitRaycast(TSpatialVisitorData<TAccelerationStructureHandle<T, d>>, FQueryFastData& CurData)
+		bool VisitRaycast(TSpatialVisitorData<TAccelerationStructureHandle<FReal, 3>>, FQueryFastData& CurData)
 		{
 			check(false);
 			return false;
@@ -154,7 +153,7 @@ namespace Chaos
 
 					CHAOS_COLLISION_STAT(StatData.RecordBoundsData(Box1));
 
-					TSimOverlapVisitor<FReal, 3> OverlapVisitor(PotentialIntersections);
+					FSimOverlapVisitor OverlapVisitor(PotentialIntersections);
 					InSpatialAcceleration.Overlap(Box1, OverlapVisitor);
 				}
 				else
@@ -246,8 +245,8 @@ namespace Chaos
 	void MoveToTOIHackImpl(FReal Dt, TTransientPBDRigidParticleHandle<FReal, 3>& Particle1, const T_SPATIALACCELERATION* SpatialAcceleration)
 	{
 		TArray<TAccelerationStructureHandle<FReal, 3>> PotentialIntersections;
-		const TAABB<FReal, 3> Box1 = ComputeWorldSpaceBoundingBox(Particle1);
-		TSimOverlapVisitor<FReal, 3> OverlapVisitor(PotentialIntersections);
+		const FAABB3 Box1 = ComputeWorldSpaceBoundingBox(Particle1);
+		FSimOverlapVisitor OverlapVisitor(PotentialIntersections);
 		SpatialAcceleration->Overlap(Box1, OverlapVisitor);
 
 		const int32 NumPotentials = PotentialIntersections.Num();
