Chaos

diff --git a/Engine/Source/Runtime/Experimental/Chaos/Public/Chaos/CollisionResolution.h b/Engine/Source/Runtime/Experimental/Chaos/Public/Chaos/CollisionResolution.h
--- a/Engine/Source/Runtime/Experimental/Chaos/Public/Chaos/CollisionResolution.h
+++ b/Engine/Source/Runtime/Experimental/Chaos/Public/Chaos/CollisionResolution.h
@@ -1,17 +1,17 @@
 // Copyright Epic Games, Inc. All Rights Reserved.
 #pragma once
 
+#include "Chaos/Core.h"
 #include "Chaos/CollisionResolutionTypes.h"
 #include "Chaos/CollisionResolutionUtil.h"
+#include "Chaos/Collision/PBDCollisionConstraint.h"
+#include "Chaos/ParticleHandleFwd.h"
 
 namespace Chaos
 {
 	template <typename T, int d>
 	class TSphere;
 
-	template <typename T, int d>
-	class TAABB;
-
 	template <typename T>
 	class TCapsule;
 
@@ -23,18 +23,13 @@ namespace Chaos
 	template <typename T>
 	class THeightField;
 
+	class FCollisionContext;
+	class FCollisionConstraintBase;
 	class FImplicitObject;
-
-	template<typename T, int d>
-	class TRigidBodyPointContactConstraint;
-
-	template <typename T, int d>
-	class TRigidTransform;
-
+	class FRigidBodyPointContactConstraint;
+	class FRigidBodyMultiPointContactConstraint;
 	class FTriangleMeshImplicitObject;
 
-	class FCollisionContext;
-
 
 	namespace Collisions
 	{
@@ -43,41 +38,32 @@ namespace Chaos
 		// Box-Box
 		//
 
-		template <typename T, int d>
-		void CHAOS_API UpdateBoxBoxConstraint(const TAABB<T, d>& Box1, const TRigidTransform<T, d>& Box1Transform, const TAABB<T, d>& Box2, const TRigidTransform<T, d>& Box2Transform, const T CullDistance, TRigidBodyPointContactConstraint<T, d>& Constraint);
+		void CHAOS_API UpdateBoxBoxConstraint(const FAABB3& Box1, const FRigidTransform3& Box1Transform, const FAABB3& Box2, const FRigidTransform3& Box2Transform, const FReal CullDistance, FRigidBodyPointContactConstraint& Constraint);
 
-		template <typename T, int d>
-		void CHAOS_API UpdateBoxBoxManifold(TCollisionConstraintBase<T, d>&  Constraint, const TRigidTransform<T, d>& ATM, const TRigidTransform<T, d>& BTM, const T CullDistance);
+		void CHAOS_API UpdateBoxBoxManifold(FCollisionConstraintBase&  Constraint, const FRigidTransform3& ATM, const FRigidTransform3& BTM, const FReal CullDistance);
 
-		template<typename T, int d>
-		void CHAOS_API ConstructBoxBoxConstraints(TGeometryParticleHandle<T, d>* Particle0, TGeometryParticleHandle<T, d>* Particle1, const FImplicitObject* Implicit0, const FImplicitObject* Implicit1, const TRigidTransform<T, d>& Transform0, const TRigidTransform<T, d>& Transform1, const T CullDistance, FCollisionConstraintsArray& NewConstraints);
+		void CHAOS_API ConstructBoxBoxConstraints(TGeometryParticleHandle<FReal, 3>* Particle0, TGeometryParticleHandle<FReal, 3>* Particle1, const FImplicitObject* Implicit0, const FImplicitObject* Implicit1, const FRigidTransform3& Transform0, const FRigidTransform3& Transform1, const FReal CullDistance, FCollisionConstraintsArray& NewConstraints);
 
 
 		//
 		// Box-HeightField
 		//
 
-		template <typename T, int d>
-		void CHAOS_API UpdateBoxHeightFieldConstraint(const TAABB<T, d>& A, const TRigidTransform<T, d>& ATransform, const THeightField<T>& B, const TRigidTransform<T, d>& BTransform, const T CullDistance, TRigidBodyPointContactConstraint<T, d>& Constraint);
+		void CHAOS_API UpdateBoxHeightFieldConstraint(const FAABB3& A, const FRigidTransform3& ATransform, const THeightField<FReal>& B, const FRigidTransform3& BTransform, const FReal CullDistance, FRigidBodyPointContactConstraint& Constraint);
 
-		template<class T, int d>
-		void CHAOS_API UpdateBoxHeightFieldManifold(TCollisionConstraintBase<T, d>&  Constraint, const TRigidTransform<T, d>& ATM, const TRigidTransform<T, d>& BTM, const T CullDistance);
+		void CHAOS_API UpdateBoxHeightFieldManifold(FCollisionConstraintBase&  Constraint, const FRigidTransform3& ATM, const FRigidTransform3& BTM, const FReal CullDistance);
 
-		template<typename T, int d>
-		void CHAOS_API ConstructBoxHeightFieldConstraints(TGeometryParticleHandle<T, d>* Particle0, TGeometryParticleHandle<T, d>* Particle1, const FImplicitObject* Implicit0, const FImplicitObject* Implicit1, const TRigidTransform<T, d>& Transform0, const TRigidTransform<T, d>& Transform1, const T CullDistance, FCollisionConstraintsArray& NewConstraints);
+		void CHAOS_API ConstructBoxHeightFieldConstraints(TGeometryParticleHandle<FReal, 3>* Particle0, TGeometryParticleHandle<FReal, 3>* Particle1, const FImplicitObject* Implicit0, const FImplicitObject* Implicit1, const FRigidTransform3& Transform0, const FRigidTransform3& Transform1, const FReal CullDistance, FCollisionConstraintsArray& NewConstraints);
 
 		//
 		// Box-Plane
 		//
 
-		template <typename T, int d>
-		void CHAOS_API UpdateBoxPlaneConstraint(const TAABB<T, d>& Box, const TRigidTransform<T, d>& BoxTransform, const TPlane<T, d>& Plane, const TRigidTransform<T, d>& PlaneTransform, const T CullDistance, TRigidBodyPointContactConstraint<T, d>& Constraint);
+		void CHAOS_API UpdateBoxPlaneConstraint(const FAABB3& Box, const FRigidTransform3& BoxTransform, const TPlane<FReal, 3>& Plane, const FRigidTransform3& PlaneTransform, const FReal CullDistance, FRigidBodyPointContactConstraint& Constraint);
 
-		template<class T, int d>
-		void CHAOS_API UpdateBoxPlaneManifold(TCollisionConstraintBase<T, d>&  Constraint, const TRigidTransform<T, d>& ATM, const TRigidTransform<T, d>& BTM, const T CullDistance);
+		void CHAOS_API UpdateBoxPlaneManifold(FCollisionConstraintBase&  Constraint, const FRigidTransform3& ATM, const FRigidTransform3& BTM, const FReal CullDistance);
 
-		template<typename T, int d>
-		void CHAOS_API ConstructBoxPlaneConstraints(TGeometryParticleHandle<T, d>* Particle0, TGeometryParticleHandle<T, d>* Particle1, const FImplicitObject* Implicit0, const FImplicitObject* Implicit1, const TRigidTransform<T, d>& Transform0, const TRigidTransform<T, d>& Transform1, const T CullDistance, FCollisionConstraintsArray& NewConstraints);
+		void CHAOS_API ConstructBoxPlaneConstraints(TGeometryParticleHandle<FReal, 3>* Particle0, TGeometryParticleHandle<FReal, 3>* Particle1, const FImplicitObject* Implicit0, const FImplicitObject* Implicit1, const FRigidTransform3& Transform0, const FRigidTransform3& Transform1, const FReal CullDistance, FCollisionConstraintsArray& NewConstraints);
 
 
 
@@ -85,176 +71,138 @@ namespace Chaos
 		// Box-TriangleMesh
 		//
 
-		template <typename TriMeshType, typename T, int d>
-		void CHAOS_API UpdateBoxTriangleMeshConstraint(const TAABB<T, d>& Box0, const TRigidTransform<T, d>& Transform0, const TriMeshType & TriangleMesh1, const TRigidTransform<T, d>& Transformw1, const T CullDistance, TRigidBodyPointContactConstraint<T, d>& Constraint);
+		template <typename TriMeshType>
+		void CHAOS_API UpdateBoxTriangleMeshConstraint(const FAABB3& Box0, const FRigidTransform3& Transform0, const TriMeshType & TriangleMesh1, const FRigidTransform3& Transformw1, const FReal CullDistance, FRigidBodyPointContactConstraint& Constraint);
 
-		template <typename T, int d>
-		void CHAOS_API UpdateBoxTriangleMeshManifold(TCollisionConstraintBase<T, d>&  Constraint, const TRigidTransform<T, d>& ATM, const TRigidTransform<T, d>& BTM, const T CullDistance);
+		void CHAOS_API UpdateBoxTriangleMeshManifold(FCollisionConstraintBase&  Constraint, const FRigidTransform3& ATM, const FRigidTransform3& BTM, const FReal CullDistance);
 
-		template<typename T, int d>
-		void CHAOS_API ConstructBoxTriangleMeshConstraints(TGeometryParticleHandle<T, d>* Particle0, TGeometryParticleHandle<T, d>* Particle1, const FImplicitObject* Implicit0, const FImplicitObject* Implicit1, const TRigidTransform<T, d>& Transform0, const TRigidTransform<T, d>& Transform1, const T CullDistance, FCollisionConstraintsArray& NewConstraints);
+		void CHAOS_API ConstructBoxTriangleMeshConstraints(TGeometryParticleHandle<FReal, 3>* Particle0, TGeometryParticleHandle<FReal, 3>* Particle1, const FImplicitObject* Implicit0, const FImplicitObject* Implicit1, const FRigidTransform3& Transform0, const FRigidTransform3& Transform1, const FReal CullDistance, FCollisionConstraintsArray& NewConstraints);
 
 
 		//
 		// Sphere-Sphere
 		//
 
-		template <typename T, int d>
-		void CHAOS_API UpdateSphereSphereConstraint(const TSphere<T, d>& Sphere1, const TRigidTransform<T, d>& Sphere1Transform, const TSphere<T, d>& Sphere2, const TRigidTransform<T, d>& Sphere2Transform, const T CullDistance, TRigidBodyPointContactConstraint<T, d>& Constraint);
+		void CHAOS_API UpdateSphereSphereConstraint(const TSphere<FReal, 3>& Sphere1, const FRigidTransform3& Sphere1Transform, const TSphere<FReal, 3>& Sphere2, const FRigidTransform3& Sphere2Transform, const FReal CullDistance, FRigidBodyPointContactConstraint& Constraint);
 
-		template<class T, int d>
-		void CHAOS_API UpdateSphereSphereManifold(TCollisionConstraintBase<T, d>&  Constraint, const TRigidTransform<T, d>& ATM, const TRigidTransform<T, d>& BTM, const T CullDistance);
+		void CHAOS_API UpdateSphereSphereManifold(FCollisionConstraintBase&  Constraint, const FRigidTransform3& ATM, const FRigidTransform3& BTM, const FReal CullDistance);
 
-		template<typename T, int d>
-		void CHAOS_API ConstructSphereSphereConstraints(TGeometryParticleHandle<T, d>* Particle0, TGeometryParticleHandle<T, d>* Particle1, const FImplicitObject* Implicit0, const FImplicitObject* Implicit1, const TRigidTransform<T, d>& Transform0, const TRigidTransform<T, d>& Transform1, const T CullDistance, FCollisionConstraintsArray& NewConstraints);
+		void CHAOS_API ConstructSphereSphereConstraints(TGeometryParticleHandle<FReal, 3>* Particle0, TGeometryParticleHandle<FReal, 3>* Particle1, const FImplicitObject* Implicit0, const FImplicitObject* Implicit1, const FRigidTransform3& Transform0, const FRigidTransform3& Transform1, const FReal CullDistance, FCollisionConstraintsArray& NewConstraints);
 
 		//
 		// Sphere-HeightField
 		//
 
-		template <typename T, int d>
-		void CHAOS_API UpdateSphereHeightFieldConstraint(const TSphere<T, d>& A, const TRigidTransform<T, d>& ATransform, const THeightField<T>& B, const TRigidTransform<T, d>& BTransform, const T CullDistance, TRigidBodyPointContactConstraint<T, d>& Constraint);
+		void CHAOS_API UpdateSphereHeightFieldConstraint(const TSphere<FReal, 3>& A, const FRigidTransform3& ATransform, const THeightField<FReal>& B, const FRigidTransform3& BTransform, const FReal CullDistance, FRigidBodyPointContactConstraint& Constraint);
 
-		template<class T, int d>
-		void CHAOS_API UpdateSphereHeightFieldManifold(TCollisionConstraintBase<T, d>&  Constraint, const TRigidTransform<T, d>& ATM, const TRigidTransform<T, d>& BTM, const T CullDistance);
+		void CHAOS_API UpdateSphereHeightFieldManifold(FCollisionConstraintBase&  Constraint, const FRigidTransform3& ATM, const FRigidTransform3& BTM, const FReal CullDistance);
 
-		template<typename T, int d>
-		void CHAOS_API ConstructSphereHeightFieldConstraints(TGeometryParticleHandle<T, d>* Particle0, TGeometryParticleHandle<T, d>* Particle1, const FImplicitObject* Implicit0, const FImplicitObject* Implicit1, const TRigidTransform<T, d>& Transform0, const TRigidTransform<T, d>& Transform1, const T CullDistance, FCollisionConstraintsArray& NewConstraints);
+		void CHAOS_API ConstructSphereHeightFieldConstraints(TGeometryParticleHandle<FReal, 3>* Particle0, TGeometryParticleHandle<FReal, 3>* Particle1, const FImplicitObject* Implicit0, const FImplicitObject* Implicit1, const FRigidTransform3& Transform0, const FRigidTransform3& Transform1, const FReal CullDistance, FCollisionConstraintsArray& NewConstraints);
 
 		//
 		// Sphere-Plane
 		//
 
-		template <typename T, int d>
-		void CHAOS_API UpdateSpherePlaneConstraint(const TSphere<T, d>& Sphere, const TRigidTransform<T, d>& SphereTransform, const TPlane<T, d>& Plane, const TRigidTransform<T, d>& PlaneTransform, const T CullDistance, TRigidBodyPointContactConstraint<T, d>& Constraint);
+		void CHAOS_API UpdateSpherePlaneConstraint(const TSphere<FReal, 3>& Sphere, const FRigidTransform3& SphereTransform, const TPlane<FReal, 3>& Plane, const FRigidTransform3& PlaneTransform, const FReal CullDistance, FRigidBodyPointContactConstraint& Constraint);
 
-		template<class T, int d>
-		void CHAOS_API UpdateSpherePlaneManifold(TCollisionConstraintBase<T, d>&  Constraint, const TRigidTransform<T, d>& ATM, const TRigidTransform<T, d>& BTM, const T CullDistance);
+		void CHAOS_API UpdateSpherePlaneManifold(FCollisionConstraintBase&  Constraint, const FRigidTransform3& ATM, const FRigidTransform3& BTM, const FReal CullDistance);
 
-		template<typename T, int d>
-		void CHAOS_API ConstructSpherePlaneConstraints(TGeometryParticleHandle<T, d>* Particle0, TGeometryParticleHandle<T, d>* Particle1, const FImplicitObject* Implicit0, const FImplicitObject* Implicit1, const TRigidTransform<T, d>& Transform0, const TRigidTransform<T, d>& Transform1, const T CullDistance, FCollisionConstraintsArray& NewConstraints);
+		void CHAOS_API ConstructSpherePlaneConstraints(TGeometryParticleHandle<FReal, 3>* Particle0, TGeometryParticleHandle<FReal, 3>* Particle1, const FImplicitObject* Implicit0, const FImplicitObject* Implicit1, const FRigidTransform3& Transform0, const FRigidTransform3& Transform1, const FReal CullDistance, FCollisionConstraintsArray& NewConstraints);
 
 		//
 		// Sphere-Box
 		//
 
-		template <typename T, int d>
-		void CHAOS_API UpdateSphereBoxConstraint(const TSphere<T, d>& Sphere, const TRigidTransform<T, d>& SphereTransform, const TAABB<T, d>& Box, const TRigidTransform<T, d>& BoxTransform, const T CullDistance, TRigidBodyPointContactConstraint<T, d>& Constraint);
+		void CHAOS_API UpdateSphereBoxConstraint(const TSphere<FReal, 3>& Sphere, const FRigidTransform3& SphereTransform, const FAABB3& Box, const FRigidTransform3& BoxTransform, const FReal CullDistance, FRigidBodyPointContactConstraint& Constraint);
 
-		template<class T, int d>
-		void CHAOS_API UpdateSphereBoxManifold(TCollisionConstraintBase<T, d>&  Constraint, const TRigidTransform<T, d>& ATM, const TRigidTransform<T, d>& BTM, const T CullDistance);
+		void CHAOS_API UpdateSphereBoxManifold(FCollisionConstraintBase&  Constraint, const FRigidTransform3& ATM, const FRigidTransform3& BTM, const FReal CullDistance);
 
-		template<typename T, int d>
-		void CHAOS_API ConstructSphereBoxConstraints(TGeometryParticleHandle<T, d>* Particle0, TGeometryParticleHandle<T, d>* Particle1, const FImplicitObject* Implicit0, const FImplicitObject* Implicit1, const TRigidTransform<T, d>& Transform0, const TRigidTransform<T, d>& Transform1, const T CullDistance, FCollisionConstraintsArray& NewConstraints);
+		void CHAOS_API ConstructSphereBoxConstraints(TGeometryParticleHandle<FReal, 3>* Particle0, TGeometryParticleHandle<FReal, 3>* Particle1, const FImplicitObject* Implicit0, const FImplicitObject* Implicit1, const FRigidTransform3& Transform0, const FRigidTransform3& Transform1, const FReal CullDistance, FCollisionConstraintsArray& NewConstraints);
 
 		//
 		// Sphere-Capsule
 		//
 
-		template <typename T, int d>
-		void CHAOS_API UpdateSphereCapsuleConstraint(const TSphere<T, d>& Sphere, const TRigidTransform<T, d>& SphereTransform, const TCapsule<T>& Box, const TRigidTransform<T, d>& BoxTransform, const T CullDistance, TRigidBodyPointContactConstraint<T, d>& Constraint);
+		void CHAOS_API UpdateSphereCapsuleConstraint(const TSphere<FReal, 3>& Sphere, const FRigidTransform3& SphereTransform, const TCapsule<FReal>& Box, const FRigidTransform3& BoxTransform, const FReal CullDistance, FRigidBodyPointContactConstraint& Constraint);
 
-		template<class T, int d>
-		void CHAOS_API UpdateSphereCapsuleManifold(TCollisionConstraintBase<T, d>&  Constraint, const TRigidTransform<T, d>& ATM, const TRigidTransform<T, d>& BTM, const T CullDistance);
+		void CHAOS_API UpdateSphereCapsuleManifold(FCollisionConstraintBase&  Constraint, const FRigidTransform3& ATM, const FRigidTransform3& BTM, const FReal CullDistance);
 
-		template<typename T, int d>
-		void CHAOS_API ConstructSphereCapsuleConstraints(TGeometryParticleHandle<T, d>* Particle0, TGeometryParticleHandle<T, d>* Particle1, const FImplicitObject* Implicit0, const FImplicitObject* Implicit1, const TRigidTransform<T, d>& Transform0, const TRigidTransform<T, d>& Transform1, const T CullDistance, FCollisionConstraintsArray& NewConstraints);
+		void CHAOS_API ConstructSphereCapsuleConstraints(TGeometryParticleHandle<FReal, 3>* Particle0, TGeometryParticleHandle<FReal, 3>* Particle1, const FImplicitObject* Implicit0, const FImplicitObject* Implicit1, const FRigidTransform3& Transform0, const FRigidTransform3& Transform1, const FReal CullDistance, FCollisionConstraintsArray& NewConstraints);
 
 		//
 		// Capsule-Capsule
 		//
 
-		template <typename T, int d>
-		void CHAOS_API UpdateCapsuleCapsuleConstraint(const TCapsule<T>& A, const TRigidTransform<T, d>& ATransform, const TCapsule<T>& B, const TRigidTransform<T, d>& BTransform, const T CullDistance, TRigidBodyPointContactConstraint<T, d>& Constraint);
+		void CHAOS_API UpdateCapsuleCapsuleConstraint(const TCapsule<FReal>& A, const FRigidTransform3& ATransform, const TCapsule<FReal>& B, const FRigidTransform3& BTransform, const FReal CullDistance, FRigidBodyPointContactConstraint& Constraint);
 
-		template<class T, int d>
-		void CHAOS_API UpdateCapsuleCapsuleManifold(TCollisionConstraintBase<T, d>&  Constraint, const TRigidTransform<T, d>& ATM, const TRigidTransform<T, d>& BTM, const T CullDistance);
+		void CHAOS_API UpdateCapsuleCapsuleManifold(FCollisionConstraintBase&  Constraint, const FRigidTransform3& ATM, const FRigidTransform3& BTM, const FReal CullDistance);
 
-		template<typename T, int d>
-		void CHAOS_API ConstructCapsuleCapsuleConstraints(TGeometryParticleHandle<T, d>* Particle0, TGeometryParticleHandle<T, d>* Particle1, const FImplicitObject* Implicit0, const FImplicitObject* Implicit1, const TRigidTransform<T, d>& Transform0, const TRigidTransform<T, d>& Transform1, const T CullDistance, FCollisionConstraintsArray& NewConstraints);
+		void CHAOS_API ConstructCapsuleCapsuleConstraints(TGeometryParticleHandle<FReal, 3>* Particle0, TGeometryParticleHandle<FReal, 3>* Particle1, const FImplicitObject* Implicit0, const FImplicitObject* Implicit1, const FRigidTransform3& Transform0, const FRigidTransform3& Transform1, const FReal CullDistance, FCollisionConstraintsArray& NewConstraints);
 
 		//
 		// Capsule-Box
 		//
 
-		template <typename T, int d>
-		void CHAOS_API UpdateCapsuleBoxConstraint(const TCapsule<T>& A, const TRigidTransform<T, d>& ATransform, const TAABB<T, d>& B, const TRigidTransform<T, d>& BTransform, const T CullDistance, TRigidBodyPointContactConstraint<T, d>& Constraint);
+		void CHAOS_API UpdateCapsuleBoxConstraint(const TCapsule<FReal>& A, const FRigidTransform3& ATransform, const FAABB3& B, const FRigidTransform3& BTransform, const FReal CullDistance, FRigidBodyPointContactConstraint& Constraint);
 
-		template<typename T, int d>
-		void CHAOS_API ConstructCapsuleBoxConstraints(TGeometryParticleHandle<T, d>* Particle0, TGeometryParticleHandle<T, d>* Particle1, const FImplicitObject* Implicit0, const FImplicitObject* Implicit1, const TRigidTransform<T, d>& Transform0, const TRigidTransform<T, d>& Transform1, const T CullDistance, const FCollisionContext& Context, FCollisionConstraintsArray& NewConstraints);
+		void CHAOS_API ConstructCapsuleBoxConstraints(TGeometryParticleHandle<FReal, 3>* Particle0, TGeometryParticleHandle<FReal, 3>* Particle1, const FImplicitObject* Implicit0, const FImplicitObject* Implicit1, const FRigidTransform3& Transform0, const FRigidTransform3& Transform1, const FReal CullDistance, const FCollisionContext& Context, FCollisionConstraintsArray& NewConstraints);
 
 		//
 		// Capsule-HeightField
 		//
 
-		template <typename T, int d>
-		void CHAOS_API UpdateCapsuleHeightFieldConstraint(const TCapsule<T>& A, const TRigidTransform<T, d>& ATransform, const THeightField<T>& B, const TRigidTransform<T, d>& BTransform, const T CullDistance, TRigidBodyPointContactConstraint<T, d>& Constraint);
+		void CHAOS_API UpdateCapsuleHeightFieldConstraint(const TCapsule<FReal>& A, const FRigidTransform3& ATransform, const THeightField<FReal>& B, const FRigidTransform3& BTransform, const FReal CullDistance, FRigidBodyPointContactConstraint& Constraint);
 
-		template<class T, int d>
-		void CHAOS_API UpdateCapsuleHeightFieldManifold(TCollisionConstraintBase<T, d>&  Constraint, const TRigidTransform<T, d>& ATM, const TRigidTransform<T, d>& BTM, const T CullDistance);
+		void CHAOS_API UpdateCapsuleHeightFieldManifold(FCollisionConstraintBase&  Constraint, const FRigidTransform3& ATM, const FRigidTransform3& BTM, const FReal CullDistance);
 
-		template<typename T, int d>
-		void ConstructCapsuleHeightFieldConstraints(TGeometryParticleHandle<T, d>* Particle0, TGeometryParticleHandle<T, d>* Particle1, const FImplicitObject* Implicit0, const FImplicitObject* Implicit1, const TRigidTransform<T, d>& Transform0, const TRigidTransform<T, d>& Transform1, const T CullDistance, FCollisionConstraintsArray& NewConstraints);
+		void ConstructCapsuleHeightFieldConstraints(TGeometryParticleHandle<FReal, 3>* Particle0, TGeometryParticleHandle<FReal, 3>* Particle1, const FImplicitObject* Implicit0, const FImplicitObject* Implicit1, const FRigidTransform3& Transform0, const FRigidTransform3& Transform1, const FReal CullDistance, FCollisionConstraintsArray& NewConstraints);
 
 		//
 		// Convex-Convex
 		//
-		template <typename T, int d>
-		void CHAOS_API UpdateConvexConvexConstraint(const FImplicitObject& A, const TRigidTransform<T, d>& ATM, const FImplicitObject& B, const TRigidTransform<T, d>& BTM, const T CullDistance, TCollisionConstraintBase<T, d>& Constraint);
+		void CHAOS_API UpdateConvexConvexConstraint(const FImplicitObject& A, const FRigidTransform3& ATM, const FImplicitObject& B, const FRigidTransform3& BTM, const FReal CullDistance, FCollisionConstraintBase& Constraint);
 
-		template<class T, int d>
-		void CHAOS_API UpdateConvexConvexManifold(TCollisionConstraintBase<T, d>&  Constraint, const TRigidTransform<T, d>& ATM, const TRigidTransform<T, d>& BTM, const T CullDistance);
+		void CHAOS_API UpdateConvexConvexManifold(FCollisionConstraintBase&  Constraint, const FRigidTransform3& ATM, const FRigidTransform3& BTM, const FReal CullDistance);
 
-		template <typename T, int d>
-		void CHAOS_API ConstructConvexConvexConstraints(TGeometryParticleHandle<T, d>* Particle0, TGeometryParticleHandle<T, d>* Particle1, const FImplicitObject* Implicit0, const FImplicitObject* Implicit1, const TRigidTransform<T, d>& Transform0, const TRigidTransform<T, d>& Transform1, const T CullDistance, FCollisionConstraintsArray& NewConstraints);
+		void CHAOS_API ConstructConvexConvexConstraints(TGeometryParticleHandle<FReal, 3>* Particle0, TGeometryParticleHandle<FReal, 3>* Particle1, const FImplicitObject* Implicit0, const FImplicitObject* Implicit1, const FRigidTransform3& Transform0, const FRigidTransform3& Transform1, const FReal CullDistance, FCollisionConstraintsArray& NewConstraints);
 
 		//
 		// Convex-HeightField
 		//
 
-		template <typename T, int d>
-		void CHAOS_API UpdateConvexHeightFieldConstraint(const FImplicitObject& A, const TRigidTransform<T, d>& ATransform, const THeightField<T>& B, const TRigidTransform<T, d>& BTransform, const T CullDistance, TRigidBodyPointContactConstraint<T, d>& Constraint);
+		void CHAOS_API UpdateConvexHeightFieldConstraint(const FImplicitObject& A, const FRigidTransform3& ATransform, const THeightField<FReal>& B, const FRigidTransform3& BTransform, const FReal CullDistance, FRigidBodyPointContactConstraint& Constraint);
 
-		template<class T, int d>
-		void CHAOS_API UpdateConvexHeightFieldManifold(TCollisionConstraintBase<T, d>&  Constraint, const TRigidTransform<T, d>& ATM, const TRigidTransform<T, d>& BTM, const T CullDistance);
+		void CHAOS_API UpdateConvexHeightFieldManifold(FCollisionConstraintBase&  Constraint, const FRigidTransform3& ATM, const FRigidTransform3& BTM, const FReal CullDistance);
 
-		template<typename T, int d>
-		void CHAOS_API ConstructConvexHeightFieldConstraints(TGeometryParticleHandle<T, d>* Particle0, TGeometryParticleHandle<T, d>* Particle1, const FImplicitObject* Implicit0, const FImplicitObject* Implicit1, const TRigidTransform<T, d>& Transform0, const TRigidTransform<T, d>& Transform1, const T CullDistance, FCollisionConstraintsArray& NewConstraints);
+		void CHAOS_API ConstructConvexHeightFieldConstraints(TGeometryParticleHandle<FReal, 3>* Particle0, TGeometryParticleHandle<FReal, 3>* Particle1, const FImplicitObject* Implicit0, const FImplicitObject* Implicit1, const FRigidTransform3& Transform0, const FRigidTransform3& Transform1, const FReal CullDistance, FCollisionConstraintsArray& NewConstraints);
 
 
 		//
 		// Levelset-Levelset
 		//
-		template<ECollisionUpdateType UpdateType, typename T = float, int d = 3>
-		void CHAOS_API UpdateLevelsetLevelsetConstraint(const T CullDistance, TRigidBodyPointContactConstraint<T, d>& Constraint);
+		template<ECollisionUpdateType UpdateType>
+		void CHAOS_API UpdateLevelsetLevelsetConstraint(const FReal CullDistance, FRigidBodyPointContactConstraint& Constraint);
 
-		template<class T, int d>
-		void CHAOS_API UpdateLevelsetLevelsetManifold(TCollisionConstraintBase<T, d>&  Constraint, const TRigidTransform<T, d>& ATM, const TRigidTransform<T, d>& BTM, const T CullDistance);
+		void CHAOS_API UpdateLevelsetLevelsetManifold(FCollisionConstraintBase&  Constraint, const FRigidTransform3& ATM, const FRigidTransform3& BTM, const FReal CullDistance);
 
-		template<typename T, int d>
-		void CHAOS_API ConstructLevelsetLevelsetConstraints(TGeometryParticleHandle<T, d>* Particle0, TGeometryParticleHandle<T, d>* Particle1, const FImplicitObject* Implicit0, const FImplicitObject* Implicit1, const TRigidTransform<T, d>& Transform0, const TRigidTransform<T, d>& Transform1, const T CullDistance, FCollisionConstraintsArray& NewConstraints);
+		void CHAOS_API ConstructLevelsetLevelsetConstraints(TGeometryParticleHandle<FReal, 3>* Particle0, TGeometryParticleHandle<FReal, 3>* Particle1, const FImplicitObject* Implicit0, const FImplicitObject* Implicit1, const FRigidTransform3& Transform0, const FRigidTransform3& Transform1, const FReal CullDistance, FCollisionConstraintsArray& NewConstraints);
 
 
 		//
 		// Constraint API
 		//
-		template <typename T, int d>
-		void CHAOS_API UpdateSingleShotManifold(TRigidBodyMultiPointContactConstraint<T, d>&  Constraint, const TRigidTransform<T, d>& Transform0, const TRigidTransform<T, d>& Transform1, const T CullDistance);
+		void CHAOS_API UpdateSingleShotManifold(FRigidBodyMultiPointContactConstraint&  Constraint, const FRigidTransform3& Transform0, const FRigidTransform3& Transform1, const FReal CullDistance);
 
-		template <typename T, int d>
-		void CHAOS_API UpdateIterativeManifold(TRigidBodyMultiPointContactConstraint<T, d>&  Constraint, const TRigidTransform<T, d>& Transform0, const TRigidTransform<T, d>& Transform1, const T CullDistance);
+		void CHAOS_API UpdateIterativeManifold(FRigidBodyMultiPointContactConstraint&  Constraint, const FRigidTransform3& Transform0, const FRigidTransform3& Transform1, const FReal CullDistance);
 
-		template<typename T, int d>
-		void CHAOS_API UpdateManifold(TRigidBodyMultiPointContactConstraint<T, d>& Constraint, const TRigidTransform<T, d>& ATM, const TRigidTransform<T, d>& BTM, const T CullDistance, const FCollisionContext& Context);
+		void CHAOS_API UpdateManifold(FRigidBodyMultiPointContactConstraint& Constraint, const FRigidTransform3& ATM, const FRigidTransform3& BTM, const FReal CullDistance, const FCollisionContext& Context);
 
-		template<class T, int d>
-		void UpdateConstraintFromManifold(TRigidBodyMultiPointContactConstraint<T, d>& Constraint, const TRigidTransform<T, d>& Transform0, const TRigidTransform<T, d>& Transform1, const T CullDistance);
+		void UpdateConstraintFromManifold(FRigidBodyMultiPointContactConstraint& Constraint, const FRigidTransform3& Transform0, const FRigidTransform3& Transform1, const FReal CullDistance);
 
-		template<ECollisionUpdateType UpdateType, typename T, int d>
-		void UpdateConstraintFromGeometry(TRigidBodyPointContactConstraint<T, d>& Constraint, const TRigidTransform<T, d>& ParticleTransform0, const TRigidTransform<T, d>& ParticleTransform1, const T CullDistance);
+		template<ECollisionUpdateType UpdateType>
+		void UpdateConstraintFromGeometry(FRigidBodyPointContactConstraint& Constraint, const FRigidTransform3& ParticleTransform0, const FRigidTransform3& ParticleTransform1, const FReal CullDistance);
 
-		template<typename T, int d>
-		void CHAOS_API ConstructConstraints(TGeometryParticleHandle<T, d>* Particle0, TGeometryParticleHandle<T, d>* Particle1, const FImplicitObject* Implicit0, const FImplicitObject* Implicit1, const TRigidTransform<T, d>& Transform0, const TRigidTransform<T, d>& Transform1, const T CullDistance, const FCollisionContext& Context, FCollisionConstraintsArray& NewConstraints);
+		void CHAOS_API ConstructConstraints(TGeometryParticleHandle<FReal, 3>* Particle0, TGeometryParticleHandle<FReal, 3>* Particle1, const FImplicitObject* Implicit0, const FImplicitObject* Implicit1, const FRigidTransform3& Transform0, const FRigidTransform3& Transform1, const FReal CullDistance, const FCollisionContext& Context, FCollisionConstraintsArray& NewConstraints);
 
 	}
 }
diff --git a/Engine/Source/Runtime/Experimental/Chaos/Public/Chaos/CollisionResolutionTypes.h b/Engine/Source/Runtime/Experimental/Chaos/Public/Chaos/CollisionResolutionTypes.h
--- a/Engine/Source/Runtime/Experimental/Chaos/Public/Chaos/CollisionResolutionTypes.h
+++ b/Engine/Source/Runtime/Experimental/Chaos/Public/Chaos/CollisionResolutionTypes.h
@@ -1,27 +1,13 @@
 // Copyright Epic Games, Inc. All Rights Reserved.
 #pragma once
-#include "Box.h"
-#include "Chaos/ExternalCollisionData.h"
+
+#include "Chaos/Core.h"
 #include "Chaos/ParticleHandleFwd.h"
-#include "Chaos/PBDConstraintContainer.h"
 #include "Chaos/Vector.h"
 
-class UPhysicalMaterial;
-
 namespace Chaos
 {
 
-	template<typename T, int d>
-	class TPBDCollisionConstraints;
-
-	template<typename  T, int d>
-	class TPBDCollisionConstraintHandle;
-
-	template<typename T, int d>
-	class TRigidBodyPointContactConstraint;
-
-
-
 	/** Specifies the type of work we should do*/
 	enum class CHAOS_API ECollisionUpdateType
 	{
@@ -66,481 +52,19 @@ namespace Chaos
 		NumShapesTypes
 	};
 
-	/*
-	*
-	*/
-	template<class T, int d>
-	class TCollisionContact
-	{
-	public:
-		TCollisionContact(const FImplicitObject* InImplicit0 = nullptr, const FImplicitObject* InImplicit1 = nullptr)
-		    : bDisabled(false), Normal(0), Location(0), Phi(FLT_MAX), Friction(0), AngularFriction(0), Restitution(0), InvInertiaScale0(1.f), InvInertiaScale1(1.f), ShapesType(EContactShapesType::Unknown)
-		{
-			Implicit[0] = InImplicit0;
-			Implicit[1] = InImplicit1;
-		}
-
-		bool bDisabled;
-		TVector<T, d> Normal;
-		TVector<T, d> Location;
-		T Phi;
-
-		T Friction;
-		T AngularFriction;
-		T Restitution;
-		T InvInertiaScale0;
-		T InvInertiaScale1;
-		EContactShapesType ShapesType;
-
-
-		FString ToString() const
-		{
-			return FString::Printf(TEXT("Location:%s, Normal:%s, Phi:%f"), *Location.ToString(), *Normal.ToString(), Phi);
-		}
-
-		const FImplicitObject* Implicit[2]; // {Of Particle[0], Of Particle[1]}
-	};
-	typedef TCollisionContact<float, 3> FCollisionContact;
-
-	/*
-	*
-	*/
-	template<class T = float, int d = 3>
-	class TCollisionConstraintBase
-	{
-	public:
-		using FGeometryParticleHandle = TGeometryParticleHandle<T, d>;
-		using FManifold = TCollisionContact<T, d>;
-
-		enum class FType
-		{
-			None = 0,     // default value also indicates a invalid constraint
-			SinglePoint,  //   TRigidBodyPointContactConstraint
-			SinglePointSwept, // TRigidBodySweptPointContactConstraint
-			MultiPoint    //   TRigidBodyMultiPointContactConstraint
-		};
-
-		TCollisionConstraintBase(FType InType = FType::None)
-			: AccumulatedImpulse(0)
-			, Timestamp(-INT_MAX)
-			, ConstraintHandle(nullptr)
-			, Type(InType)
-		{ 
-			ImplicitTransform[0] = TRigidTransform<T, d>::Identity; ImplicitTransform[1] = TRigidTransform<T,d>::Identity;
-			Manifold.Implicit[0] = nullptr; Manifold.Implicit[1] = nullptr;
-			Particle[0] = nullptr; Particle[1] = nullptr; 
-		}
-		
-		TCollisionConstraintBase(
-			FGeometryParticleHandle* Particle0, const FImplicitObject* Implicit0, const TRigidTransform<T, d>& Transform0,
-			FGeometryParticleHandle* Particle1, const FImplicitObject* Implicit1, const TRigidTransform<T, d>& Transform1,
-			FType InType, EContactShapesType ShapesType, int32 InTimestamp = -INT_MAX)
-			: AccumulatedImpulse(0)
-			, Timestamp(InTimestamp)
-			, ConstraintHandle(nullptr)
-			, Type(InType)
-		{
-			ImplicitTransform[0] = Transform0; ImplicitTransform[1] = Transform1;
-			Manifold.Implicit[0] = Implicit0; Manifold.Implicit[1] = Implicit1;
-			Manifold.ShapesType = ShapesType;
-			Particle[0] = Particle0; Particle[1] = Particle1; 
-		}
-
-		FType GetType() const { return Type; }
-
-		template<class AS_T> AS_T * As() { return static_cast<AS_T*>(this); }
-		template<class AS_T> const AS_T * As() const { return static_cast<const AS_T*>(this); }
-
-		bool ContainsManifold(const FImplicitObject* A, const FImplicitObject* B) const { return A == Manifold.Implicit[0] && B == Manifold.Implicit[1]; }
-		void SetManifold(const FImplicitObject* A, const FImplicitObject* B) { Manifold.Implicit[0] = A; Manifold.Implicit[1] = B; }
-		const FManifold& GetManifold() const { return Manifold; }
-
-		//
-		// API
-		//
-
-		void ResetPhi(T InPhi) { SetPhi(InPhi); }
-
-		void SetPhi(T InPhi) { Manifold.Phi = InPhi; }
-		T GetPhi() const { return Manifold.Phi; }
-
-		void SetDisabled(bool bInDisabled) { Manifold.bDisabled = bInDisabled; }
-		bool GetDisabled() const { return Manifold.bDisabled; }
-
-		void SetNormal(const TVector<T, d> & InNormal) { Manifold.Normal = InNormal; }
-		TVector<T, d> GetNormal() const { return Manifold.Normal; }
-
-		void SetLocation(const TVector<T, d> & InLocation) { Manifold.Location = InLocation; }
-		TVector<T, d> GetLocation() const { return Manifold.Location; }
-
-		void SetInvInertiaScale0(const T InInvInertiaScale) { Manifold.InvInertiaScale0 = InInvInertiaScale; }
-		T GetInvInertiaScale0() const { return Manifold.InvInertiaScale0; }
-
-		void SetInvInertiaScale1(const T InInvInertiaScale) { Manifold.InvInertiaScale1 = InInvInertiaScale; }
-		T GetInvInertiaScale1() const { return Manifold.InvInertiaScale1; }
-
-
-		FString ToString() const
-		{
-			return FString::Printf(TEXT("Particle:%s, Levelset:%s, AccumulatedImpulse:%s"), *Particle[0]->ToString(), *Particle[1]->ToString(), *AccumulatedImpulse.ToString());
-		}
-
-
-		TRigidTransform<T, d> ImplicitTransform[2]; // { Point, Volume }
-		FGeometryParticleHandle* Particle[2]; // { Point, Volume }
-		TVector<T, d> AccumulatedImpulse;
-		FManifold Manifold;
-		int32 Timestamp;
-		TPBDCollisionConstraintHandle<T, d>* ConstraintHandle;
-
-	private:
-
-		FType Type;
-	};
-	typedef TCollisionConstraintBase<float, 3> FCollisionConstraintBase;
-
-
-	/*
-	*
-	*/
-	template<class T, int d>
-	class TRigidBodyPointContactConstraint : public TCollisionConstraintBase<T, d>
-	{
-	public:
-		using Base = TCollisionConstraintBase<T, d>;
-		using FGeometryParticleHandle = TGeometryParticleHandle<T, d>;
-		using Base::Particle;
-
-		TRigidBodyPointContactConstraint() : Base(Base::FType::SinglePoint) {}
-		TRigidBodyPointContactConstraint(
-			FGeometryParticleHandle* Particle0, const FImplicitObject* Implicit0, const TRigidTransform<T, d>& Transform0,
-			FGeometryParticleHandle* Particle1, const FImplicitObject* Implicit1, const TRigidTransform<T, d>& Transform1,
-			EContactShapesType ShapesType)
-			: Base(Particle0, Implicit0, Transform0, Particle1, Implicit1, Transform1, Base::FType::SinglePoint, ShapesType) {}
-
-		static typename TCollisionConstraintBase<T, d>::FType StaticType() { return Base::FType::SinglePoint; };
-
-	protected:
-		// For use by derived types that can be used as point constraints in Update
-		TRigidBodyPointContactConstraint(typename Base::FType InType) : Base(InType) {}
-		TRigidBodyPointContactConstraint(
-			FGeometryParticleHandle* Particle0, const FImplicitObject* Implicit0, const TRigidTransform<T, d>& Transform0,
-			FGeometryParticleHandle* Particle1, const FImplicitObject* Implicit1, const TRigidTransform<T, d>& Transform1,
-			typename Base::FType InType, EContactShapesType ShapesType)
-			: Base(Particle0, Implicit0, Transform0, Particle1, Implicit1, Transform1, InType, ShapesType) {}
-	};
-	typedef TRigidBodyPointContactConstraint<float, 3> FRigidBodyPointContactConstraint;
-
-
-	/*
-	*
-	*/
-	template<class T, int d>
-	class TRigidBodyMultiPointContactConstraint : public TRigidBodyPointContactConstraint<T, d>
-	{
-	public:
-		using Base = TRigidBodyPointContactConstraint<T, d>;
-		using FGeometryParticleHandle = TGeometryParticleHandle<T, d>;
-		using FManifold = TCollisionContact<T, d>;
-		using Base::Particle;
-
-		TRigidBodyMultiPointContactConstraint() : Base(TCollisionConstraintBase<T, d>::FType::MultiPoint), PlaneOwnerIndex(INDEX_NONE), PlaneFaceIndex(INDEX_NONE), PlaneNormal(0), PlanePosition(0) {}
-		TRigidBodyMultiPointContactConstraint(
-			FGeometryParticleHandle* Particle0, const FImplicitObject* Implicit0, const TRigidTransform<T, d>& Transform0,
-			FGeometryParticleHandle* Particle1, const FImplicitObject* Implicit1, const TRigidTransform<T, d>& Transform1,
-			EContactShapesType ShapesType)
-			: Base(Particle0, Implicit0, Transform0, Particle1, Implicit1, Transform1, Base::FType::MultiPoint, ShapesType)
-			, PlaneOwnerIndex(INDEX_NONE), PlaneFaceIndex(INDEX_NONE), PlaneNormal(0), PlanePosition(0)
-		{}
-
-		static typename TCollisionConstraintBase<T, d>::FType StaticType() { return Base::FType::MultiPoint; };
-
-		// Get the particle that owns the plane
-		FGeometryParticleHandle* PlaneParticleHandle() const { check(PlaneOwnerIndex >= 0 && PlaneOwnerIndex < 2); return Particle[PlaneOwnerIndex]; }
-
-		// Get the particle that owns the manifold sample points
-		FGeometryParticleHandle* PointsParticleHandle() const { check(PlaneOwnerIndex >= 0 && PlaneOwnerIndex < 2); return Particle[1 - PlaneOwnerIndex]; }
-
-		int32 GetManifoldPlaneOwnerIndex() const { return PlaneOwnerIndex; }
-		int32 GetManifoldPlaneFaceIndex() const { return PlaneFaceIndex; }
-		const FVec3& GetManifoldPlaneNormal() const { return PlaneNormal; }
-		const FVec3& GetManifoldPlanePosition() const { return PlanePosition; }
-		void SetManifoldPlane(int32 OwnerIndex, int32 FaceIndex, const FVec3& Normal, const FVec3& Pos)
-		{
-			PlaneOwnerIndex = OwnerIndex;
-			PlaneFaceIndex = FaceIndex;
-			PlaneNormal = Normal;
-			PlanePosition = Pos;
-		}
-
-		int32               NumManifoldPoints() const { return Points.Num(); }
-		const FVec3&		GetManifoldPoint(int32 Index) const { check(0 <= Index && Index < NumManifoldPoints()); return Points[Index]; }
-		void				SetManifoldPoint(int32 Index, const FVec3& Point) { check(0 <= Index && Index < NumManifoldPoints()); Points[Index] = Point; }
-		void                AddManifoldPoint(const FVec3& Point) { Points.Add(Point); };
-		void                ResetManifoldPoints(int32 NewSize=0)  { Points.Reset(NewSize); }
-
-	private:
-		// Manifold plane data
-		int PlaneOwnerIndex; // index of particle which owns the plane. The other owns the sample positions
-		int PlaneFaceIndex; // index of face used as the manifold plane on plane-owner body
-		TVector<T, d> PlaneNormal; // local space contact normal on plane-owner
-		TVector<T, d> PlanePosition; // local space surface position on plane-owner
-
-		// Manifold point data
-		static const int32 MaxPoints = 4;
-		TArray<FVec3, TInlineAllocator<MaxPoints>> Points; // @todo(chaos): use TFixedAllocator when we handle building the best manifold properly
-	};
-	typedef TRigidBodyMultiPointContactConstraint<float, 3> FRigidBodyMultiPointContactConstraint;
-
-	/*
-	*
-	*/
-	template<class T, int d>
-	class TRigidBodySweptPointContactConstraint : public TRigidBodyPointContactConstraint<T,d>
-	{
-	public:
-		using Base = TRigidBodyPointContactConstraint<T, d>;
-		using FGeometryParticleHandle = TGeometryParticleHandle<T, d>;
-		using Base::Particle;
-
-		TRigidBodySweptPointContactConstraint() : Base(Base::FType::SinglePointSwept) {}
-		TRigidBodySweptPointContactConstraint(
-			FGeometryParticleHandle* Particle0, const FImplicitObject* Implicit0, const TRigidTransform<T, d>& Transform0,
-			FGeometryParticleHandle* Particle1, const FImplicitObject* Implicit1, const TRigidTransform<T, d>& Transform1,
-			EContactShapesType ShapesType)
-			: Base(Particle0, Implicit0, Transform0, Particle1, Implicit1, Transform1, Base::FType::SinglePointSwept, ShapesType)
-			, TimeOfImpact(0) {}
-
-		// Value in range [0,1] used to interpolate P between [X,P] that we will rollback to when solving at time of impact.
-		T TimeOfImpact;
-
-		static typename Base::FType StaticType() { return Base::FType::SinglePointSwept; };
-	};
-	typedef TRigidBodySweptPointContactConstraint<float, 3> FRigidBodySweptPointContactConstraint;
-
-
-
 	//
 	//
 	//
 
-	template<class T, int d>
-	struct CHAOS_API TRigidBodyContactConstraintPGS
-	{
-		TRigidBodyContactConstraintPGS() : AccumulatedImpulse(0.f) {}
-		TGeometryParticleHandle<T, d>* Particle;
-		TGeometryParticleHandle<T, d>* Levelset;
-		TArray<TVector<T, d>> Normal;
-		TArray<TVector<T, d>> Location;
-		TArray<T> Phi;
-		TVector<T, d> AccumulatedImpulse;
-	};
-
-
-	template<class T, int d>
-	class TPBDCollisionConstraintHandle : public TContainerConstraintHandle<TPBDCollisionConstraints<T, d>>
-	{
-	public:
-		using Base = TContainerConstraintHandle<TPBDCollisionConstraints<T, d>>;
-		using FConstraintContainer = TPBDCollisionConstraints<T, d>;
-		using FConstraintBase = TCollisionConstraintBase<T, d>;
-
-		using FImplicitPair = TPair<const FImplicitObject*, const FImplicitObject*>;
-		using FGeometryPair = TPair<const TGeometryParticleHandle<T, d>*, const TGeometryParticleHandle<T, d>*>;
-		using FHandleKey = TPair<FImplicitPair, FGeometryPair>;
-
-
-		TPBDCollisionConstraintHandle()
-			: ConstraintType(FConstraintBase::FType::None)
-		{}
-
-		TPBDCollisionConstraintHandle(FConstraintContainer* InConstraintContainer, int32 InConstraintIndex, typename FConstraintBase::FType InType)
-			: TContainerConstraintHandle<TPBDCollisionConstraints<T, d>>(InConstraintContainer, InConstraintIndex)
-			, ConstraintType(InType)
-		{
-		}
-
-		FHandleKey GetKey()
-		{
-			const TCollisionConstraintBase<T, d>& Contact = GetContact();
-			return FHandleKey(
-				FImplicitPair( Contact.Manifold.Implicit[0],Contact.Manifold.Implicit[1]),
-				FGeometryPair( Contact.Particle[0], Contact.Particle[1] ) );
-		}
-
-		static FHandleKey MakeKey(const TGeometryParticleHandle<T, d>* Particle0, const TGeometryParticleHandle<T, d>* Particle1, 
-			const FImplicitObject* Implicit0, const FImplicitObject* Implicit1)
-		{
-			return FHandleKey(FImplicitPair(Implicit0,Implicit1), FGeometryPair(Particle0, Particle1));
-		}
-
-		static FHandleKey MakeKey(const TCollisionConstraintBase<T,d> * Base)
-		{
-			return FHandleKey(FImplicitPair(Base->Manifold.Implicit[0], Base->Manifold.Implicit[1]), FGeometryPair(Base->Particle[0], Base->Particle[1]));
-		}
-
-
-		const TCollisionConstraintBase<T, d>& GetContact() const 
-		{ 
-			if (GetType() == FConstraintBase::FType::SinglePoint)
-			{
-				return ConstraintContainer->Constraints.SinglePointConstraints[ConstraintIndex];
-			}
-			else if (GetType() == FConstraintBase::FType::SinglePointSwept)
-			{
-				return ConstraintContainer->Constraints.SinglePointSweptConstraints[ConstraintIndex];
-			}
-
-			check(GetType() == FConstraintBase::FType::MultiPoint);
-			return ConstraintContainer->Constraints.MultiPointConstraints[ConstraintIndex];
-		}
-
-		TCollisionConstraintBase<T, d>& GetContact()
-		{
-			if (GetType() == FConstraintBase::FType::SinglePoint)
-			{
-				return ConstraintContainer->Constraints.SinglePointConstraints[ConstraintIndex];
-			}
-			else if (GetType() == FConstraintBase::FType::SinglePointSwept)
-			{
-				return ConstraintContainer->Constraints.SinglePointSweptConstraints[ConstraintIndex];
-			}
-
-			check(GetType() == FConstraintBase::FType::MultiPoint);
-			return ConstraintContainer->Constraints.MultiPointConstraints[ConstraintIndex];
-		}
-
-
-
-		const TRigidBodyPointContactConstraint<T, d>& GetPointContact() const { check(GetType() == FConstraintBase::FType::SinglePoint); return ConstraintContainer->Constraints.SinglePointConstraints[ConstraintIndex]; }
-		TRigidBodyPointContactConstraint<T, d>& GetPointContact() { check(GetType() == FConstraintBase::FType::SinglePoint); return ConstraintContainer->Constraints.SinglePointConstraints[ConstraintIndex]; }
-
-		const TRigidBodySweptPointContactConstraint<T, d>& GetSweptPointContact() const { check(GetType() == FConstraintBase::FType::SinglePointSwept); return ConstraintContainer->Constraints.SinglePointSweptConstraints[ConstraintIndex]; }
-		TRigidBodySweptPointContactConstraint<T, d>& GetSweptPointContact() { check(GetType() == FConstraintBase::FType::SinglePointSwept); return ConstraintContainer->Constraints.SinglePointSweptConstraints[ConstraintIndex]; }
-
-		const TRigidBodyMultiPointContactConstraint<T, d>& GetMultiPointContact() const { check(GetType() == FConstraintBase::FType::MultiPoint); return ConstraintContainer->Constraints.MultiPointConstraints[ConstraintIndex]; }
-		TRigidBodyMultiPointContactConstraint<T, d>& GetMultiPointContact() { check(GetType() == FConstraintBase::FType::MultiPoint); return ConstraintContainer->Constraints.MultiPointConstraints[ConstraintIndex]; }
-
-		typename FConstraintBase::FType GetType() const { return ConstraintType; }
-
-		void SetConstraintIndex(int32 IndexIn, typename FConstraintBase::FType InType)
-		{
-			ConstraintIndex = IndexIn;
-			ConstraintType = InType;
-		}
-
-		TVector<T, d> GetContactLocation() const
-		{
-			return GetContact().GetLocation();
-		}
-
-		TVector<T, d> GetAccumulatedImpulse() const
-		{
-			return GetContact().AccumulatedImpulse;
-		}
-
-		TVector<const TGeometryParticleHandle<T, d>*, 2> GetConstrainedParticles() const
-		{
-			return { GetContact().Particle[0], GetContact().Particle[1] };
-		}
-
-		TVector<TGeometryParticleHandle<T, d>*, 2> GetConstrainedParticles()
-		{
-			return { GetContact().Particle[0], GetContact().Particle[1] };
-		}
-
-
-	protected:
-		typename FConstraintBase::FType ConstraintType;
-		using Base::ConstraintIndex;
-		using Base::ConstraintContainer;
-
-
-	};
-	typedef TPBDCollisionConstraintHandle<float, 3> FPBDCollisionConstraintHandle;
-
-
-	struct FCollisionConstraintsArray
+	struct CHAOS_API FRigidBodyContactConstraintPGS
 	{
-		static const int32 InlineMaxConstraints = 8;
-
-		TArray<FRigidBodyPointContactConstraint, TInlineAllocator<InlineMaxConstraints>> SinglePointConstraints;
-		TArray<FRigidBodySweptPointContactConstraint, TInlineAllocator<InlineMaxConstraints>> SinglePointSweptConstraints;
-		TArray<FRigidBodyMultiPointContactConstraint, TInlineAllocator<InlineMaxConstraints>> MultiPointConstraints;
-
-		int32 Num() const { return SinglePointConstraints.Num() + SinglePointSweptConstraints.Num() + MultiPointConstraints.Num(); }
-
-		void Reset()
-		{
-			SinglePointConstraints.Reset();
-			SinglePointSweptConstraints.Reset();
-			MultiPointConstraints.Reset();
-		}
-
-		void Empty()
-		{
-			SinglePointConstraints.Empty();
-			SinglePointSweptConstraints.Empty();
-			MultiPointConstraints.Empty();
-		}
-
-
-		TRigidBodyPointContactConstraint<FReal, 3>* Add(const TRigidBodyPointContactConstraint<FReal, 3>& C)
-		{
-			int32 ConstraintIndex = SinglePointConstraints.Add(C);
-			return &SinglePointConstraints[ConstraintIndex];
-		}
-
-		TRigidBodySweptPointContactConstraint<FReal, 3>* Add(const TRigidBodySweptPointContactConstraint<FReal, 3>& C)
-		{
-			int32 ConstraintIndex = SinglePointSweptConstraints.Add(C);
-			return &SinglePointSweptConstraints[ConstraintIndex];
-		}
-
-		TRigidBodyMultiPointContactConstraint<FReal, 3>* Add(const TRigidBodyMultiPointContactConstraint<FReal, 3>& C)
-		{
-			int32 ConstraintIndex = MultiPointConstraints.Add(C);
-			return &MultiPointConstraints[ConstraintIndex];
-		}
-
-
-		TRigidBodyPointContactConstraint<FReal, 3>* TryAdd(FReal MaxPhi, const TRigidBodyPointContactConstraint<FReal, 3>& C)
-		{
-			if (C.GetPhi() < MaxPhi)
-			{
-				int32 ConstraintIndex = SinglePointConstraints.Add(C);
-				return &SinglePointConstraints[ConstraintIndex];
-			}
-			return nullptr;
-		}
-
-		TRigidBodySweptPointContactConstraint<FReal, 3>* TryAdd(FReal MaxPhi, const TRigidBodySweptPointContactConstraint<FReal, 3>& C)
-		{
-			if (C.GetPhi() < MaxPhi)
-			{
-				int32 ConstraintIndex = SinglePointSweptConstraints.Add(C);
-				return &SinglePointSweptConstraints[ConstraintIndex];
-			}
-			return nullptr;
-		}
-
-		TRigidBodyMultiPointContactConstraint<FReal, 3>* TryAdd(FReal MaxPhi, const TRigidBodyMultiPointContactConstraint<FReal, 3>& C)
-		{
-			if (C.GetPhi() < MaxPhi)
-			{
-				int32 ConstraintIndex = MultiPointConstraints.Add(C);
-				return &MultiPointConstraints[ConstraintIndex];
-			}
-			return nullptr;
-		}
+		FRigidBodyContactConstraintPGS() : AccumulatedImpulse(0.f) {}
+		TGeometryParticleHandle<FReal, 3>* Particle;
+		TGeometryParticleHandle<FReal, 3>* Levelset;
+		TArray<FVec3> Normal;
+		TArray<FVec3> Location;
+		TArray<FReal> Phi;
+		FVec3 AccumulatedImpulse;
 	};
 
-
-#if PLATFORM_MAC || PLATFORM_LINUX
-	extern template class CHAOS_API TCollisionContact<float, 3>;
-	extern template class CHAOS_API TPBDCollisionConstraintHandle<float, 3>;
-#else
-	extern template class TCollisionContact<float, 3>;
-	extern template class TPBDCollisionConstraintHandle<float, 3>;
-#endif
 }
diff --git a/Engine/Source/Runtime/Experimental/Chaos/Public/Chaos/CollisionResolutionUtil.h b/Engine/Source/Runtime/Experimental/Chaos/Public/Chaos/CollisionResolutionUtil.h
--- a/Engine/Source/Runtime/Experimental/Chaos/Public/Chaos/CollisionResolutionUtil.h
+++ b/Engine/Source/Runtime/Experimental/Chaos/Public/Chaos/CollisionResolutionUtil.h
@@ -1,18 +1,17 @@
 // Copyright Epic Games, Inc. All Rights Reserved.
-
 #pragma once
+
+#include "Chaos/Core.h"
 #include "Chaos/CollisionResolutionTypes.h"
 #include "Chaos/Pair.h"
 
 namespace Chaos
 {
-	class FImplicitObject;
+	template<class T, int d>
+	class TBVHParticles;
 
-	template<typename T, int d>
-	class TRigidBodyPointContactConstraint;
-
-	template <typename T, int d>
-	class TRigidTransform;
+	class FImplicitObject;
+	class FRigidBodyPointContactConstraint;
 
 	namespace Collisions
 	{
@@ -20,43 +19,30 @@ namespace Chaos
 		// Utility
 		//
 
-		template<typename T = float, int d = 3>
-		TRigidTransform<T, d> 
-		GetTransform(const TGeometryParticleHandle<T, d>* Particle);
+		FRigidTransform3 
+		GetTransform(const TGeometryParticleHandle<FReal, 3>* Particle);
 
-		template<typename T>
-		PMatrix<T, 3, 3> 
-		ComputeFactorMatrix3(const TVector<T, 3>& V, const PMatrix<T, 3, 3>& M, const T& Im);
+		FMatrix33 
+		ComputeFactorMatrix3(const FVec3& V, const FMatrix33& M, const FReal Im);
 
-		template<typename T, int d>
-		TVector<T, d> 
-		GetEnergyClampedImpulse(const TPBDRigidParticleHandle<T, d>* PBDRigid0, const TPBDRigidParticleHandle<T, d>* PBDRigid1, const TVector<T, d>& Impulse, const TVector<T, d>& VectorToPoint1, const TVector<T, d>& VectorToPoint2, const TVector<T, d>& Velocity1, const TVector<T, d>& Velocity2);
+		FVec3
+		GetEnergyClampedImpulse(const TPBDRigidParticleHandle<FReal, 3>* PBDRigid0, const TPBDRigidParticleHandle<FReal, 3>* PBDRigid1, const FVec3& Impulse, const FVec3& VectorToPoint1, const FVec3& VectorToPoint2, const FVec3& Velocity1, const FVec3& Velocity2);
 
-		template <typename T, int d>
 		bool 
-		SampleObjectHelper(const FImplicitObject& Object, const TRigidTransform<T, d>& ObjectTransform, const TRigidTransform<T, d>& SampleToObjectTransform, const TVector<T, d>& SampleParticle, T Thickness, TRigidBodyPointContactConstraint<float, 3>& Constraint);
+		SampleObjectHelper(const FImplicitObject& Object, const FRigidTransform3& ObjectTransform, const FRigidTransform3& SampleToObjectTransform, const FVec3& SampleParticle, FReal Thickness, FRigidBodyPointContactConstraint& Constraint);
 
-		template <typename T, int d>
 		bool 
-		SampleObjectNoNormal(const FImplicitObject& Object, const TRigidTransform<T, d>& ObjectTransform, const TRigidTransform<T, d>& SampleToObjectTransform, const TVector<T, d>& SampleParticle, T Thickness, TRigidBodyPointContactConstraint<float, 3>& Constraint);
+		SampleObjectNoNormal(const FImplicitObject& Object, const FRigidTransform3& ObjectTransform, const FRigidTransform3& SampleToObjectTransform, const FVec3& SampleParticle, FReal Thickness, FRigidBodyPointContactConstraint& Constraint);
 
-		template <typename T, int d>
 		bool 
-		SampleObjectNormalAverageHelper(const FImplicitObject& Object, const TRigidTransform<T, d>& ObjectTransform, const TRigidTransform<T, d>& SampleToObjectTransform, const TVector<T, d>& SampleParticle, T Thickness, T& TotalThickness, TRigidBodyPointContactConstraint<float, 3>& Constraint);
-
-#if INTEL_ISPC
-		template<ECollisionUpdateType UpdateType>
-		void 
-		SampleObject(const FImplicitObject& Object, const TRigidTransform<float, 3>& ObjectTransform, const TBVHParticles<float, 3>& SampleParticles, const TRigidTransform<float, 3>& SampleParticlesTransform, float Thickness, TRigidBodyPointContactConstraint<float, 3>& Constraint);
-#else
-		template <ECollisionUpdateType UpdateType, typename T, int d>
+		SampleObjectNormalAverageHelper(const FImplicitObject& Object, const FRigidTransform3& ObjectTransform, const FRigidTransform3& SampleToObjectTransform, const FVec3& SampleParticle, FReal Thickness, FReal& TotalThickness, FRigidBodyPointContactConstraint& Constraint);
+
+		template <ECollisionUpdateType UpdateType>
 		void
-		SampleObject(const FImplicitObject& Object, const TRigidTransform<T, d>& ObjectTransform, const TBVHParticles<T, d>& SampleParticles, const TRigidTransform<T, d>& SampleParticlesTransform, T Thickness, TRigidBodyPointContactConstraint<float, 3>& Constraint);
-#endif
+		SampleObject(const FImplicitObject& Object, const FRigidTransform3& ObjectTransform, const TBVHParticles<FReal, 3>& SampleParticles, const FRigidTransform3& SampleParticlesTransform, FReal Thickness, FRigidBodyPointContactConstraint& Constraint);
 
-		template <typename T, int d>
-		TArray<Pair<const FImplicitObject*, TRigidTransform<T, d>>> 
-		FindRelevantShapes(const FImplicitObject* ParticleObj, const TRigidTransform<T, d>& ParticlesTM, const FImplicitObject& LevelsetObj, const TRigidTransform<T, d>& LevelsetTM, const T Thickness);
+		TArray<Pair<const FImplicitObject*, FRigidTransform3>> 
+		FindRelevantShapes(const FImplicitObject* ParticleObj, const FRigidTransform3& ParticlesTM, const FImplicitObject& LevelsetObj, const FRigidTransform3& LevelsetTM, const FReal Thickness);
 
 	}// Collisions
 
diff --git a/Engine/Source/Runtime/Experimental/Chaos/Public/Chaos/Evolution/PBDMinEvolution.h b/Engine/Source/Runtime/Experimental/Chaos/Public/Chaos/Evolution/PBDMinEvolution.h
--- a/Engine/Source/Runtime/Experimental/Chaos/Public/Chaos/Evolution/PBDMinEvolution.h
+++ b/Engine/Source/Runtime/Experimental/Chaos/Public/Chaos/Evolution/PBDMinEvolution.h
@@ -9,11 +9,10 @@ namespace Chaos
 {
 	class FNarrowPhase;
 	class FParticlePairBroadPhase;
+	class FPBDCollisionConstraints;
 	class FSimpleConstraintRule;
 	class FSyncCollisionReceiver;
 
-	template <typename T, int d>
-	class TPBDCollisionConstraints;
 
 	template<typename T_BROADPHASE, typename T_NARROWPHASE, typename T_RECEIVER, typename T_CONTAINER>
 	class TCollisionDetector;
@@ -34,8 +33,7 @@ namespace Chaos
 	public:
 		// @todo(ccaulfield): make it so that CollisionDetection is plugged in with a constraint rule...
 
-		using FCollisionConstraints = TPBDCollisionConstraints<FReal, 3>;
-		using FCollisionDetector = TCollisionDetector<FParticlePairBroadPhase, FNarrowPhase, FSyncCollisionReceiver, FCollisionConstraints>;
+		using FCollisionDetector = TCollisionDetector<FParticlePairBroadPhase, FNarrowPhase, FSyncCollisionReceiver, FPBDCollisionConstraints>;
 		using FEvolutionCallback = TFunction<void()>;
 		using FRigidParticleSOAs = TPBDRigidsSOAs<FReal, 3>;
 
diff --git a/Engine/Source/Runtime/Experimental/Chaos/Public/Chaos/Framework/PhysicsProxy.h b/Engine/Source/Runtime/Experimental/Chaos/Public/Chaos/Framework/PhysicsProxy.h
--- a/Engine/Source/Runtime/Experimental/Chaos/Public/Chaos/Framework/PhysicsProxy.h
+++ b/Engine/Source/Runtime/Experimental/Chaos/Public/Chaos/Framework/PhysicsProxy.h
@@ -40,7 +40,7 @@ public:
 	using FParticleType = Concrete;
 
 	using FParticlesType = Chaos::TPBDRigidParticles<float, 3>;
-	using FCollisionConstraintsType = Chaos::TPBDCollisionConstraints<float, 3>;
+	using FCollisionConstraintsType = Chaos::FPBDCollisionConstraints;
 	using FIntArray = Chaos::TArrayCollectionArray<int32>;
 
 	TPhysicsProxy()
diff --git a/Engine/Source/Runtime/Experimental/Chaos/Public/Chaos/PBDCollisionConstraints.h b/Engine/Source/Runtime/Experimental/Chaos/Public/Chaos/PBDCollisionConstraints.h
--- a/Engine/Source/Runtime/Experimental/Chaos/Public/Chaos/PBDCollisionConstraints.h
+++ b/Engine/Source/Runtime/Experimental/Chaos/Public/Chaos/PBDCollisionConstraints.h
@@ -1,10 +1,9 @@
 // Copyright Epic Games, Inc. All Rights Reserved.
 #pragma once
 
-#include "Chaos/CollisionResolutionTypes.h"
 #include "Chaos/CollisionResolutionTypes.h"
 #include "Chaos/Collision/CollisionApplyType.h"
-#include "Chaos/ConstraintHandle.h"
+#include "Chaos/Collision/PBDCollisionConstraintHandle.h"
 #include "Chaos/PBDConstraintContainer.h"
 #include "Framework/BufferedData.h"
 
@@ -19,76 +18,44 @@
 
 namespace Chaos
 {
-template<typename T, int d>
-class TPBDCollisionConstraints;
-
-template <typename T, int d>
-class TPBDCollisionConstraintHandle;
-
-template <typename T, int d>
-class TCollisionConstraintBase;
-
-template <typename T, int d>
-class TRigidBodyPointContactConstraint;
-
-template <typename T, int d>
-class TRigidBodyMultiPointContactConstraint;
-
-template <typename T, int d>
-class TRigidTransform;
-
+class FCollisionConstraintBase;
 class FImplicitObject;
-
-template <typename T, int d>
-class TBVHParticles;
-
-template <typename T, int d>
-class TAABB;
-
-template<class T>
-class TChaosPhysicsMaterial;
+class FPBDCollisionConstraints;
+class FRigidBodyPointContactConstraint;
+class FRigidBodyMultiPointContactConstraint;
 
 template <typename T, int d>
 class TPBDRigidsSOAs;
 
-template<typename T, int d>
-using TRigidBodyContactConstraintsPostComputeCallback = TFunction<void()>;
-
-template<typename T, int d>
-using TRigidBodyContactConstraintsPostApplyCallback = TFunction<void(const T Dt, const TArray<TPBDCollisionConstraintHandle<T, d>*>&)>;
-
-template<typename T, int d>
-using TRigidBodyContactConstraintsPostApplyPushOutCallback = TFunction<void(const T Dt, const TArray<TPBDCollisionConstraintHandle<T, d>*>&, bool)>;
-
-template <typename T, int d>
-using TCollisionModifierCallback = TFunction<ECollisionModifierResult(TPBDCollisionConstraintHandle<T, d>*)>;
+using FRigidBodyContactConstraintsPostComputeCallback = TFunction<void()>;
+using FRigidBodyContactConstraintsPostApplyCallback = TFunction<void(const FReal Dt, const TArray<FPBDCollisionConstraintHandle*>&)>;
+using FRigidBodyContactConstraintsPostApplyPushOutCallback = TFunction<void(const FReal Dt, const TArray<FPBDCollisionConstraintHandle*>&, bool)>;
+using FCollisionModifierCallback = TFunction<ECollisionModifierResult(FPBDCollisionConstraintHandle*)>;
 
 /**
  * A container and solver for collision constraints.
  */
-template<typename T, int d>
-class CHAOS_API TPBDCollisionConstraints : public FPBDConstraintContainer
+class CHAOS_API FPBDCollisionConstraints : public FPBDConstraintContainer
 {
 public:
-	friend class TPBDCollisionConstraintHandle<T, d>;
+	friend class FPBDCollisionConstraintHandle;
 
 	using Base = FPBDConstraintContainer;
-	using FHandles = TArray<TPBDCollisionConstraintHandle<T, d>*>;
-	using FConstraintContainerHandle = TPBDCollisionConstraintHandle<T, d>;
-	using FConstraintBase = TCollisionConstraintBase<T, d>;
- 	using FPointContactConstraint = TRigidBodyPointContactConstraint<T, d>;
- 	using FSweptPointContactConstraint = TRigidBodySweptPointContactConstraint<T, d>;
-	using FMultiPointContactConstraint = TRigidBodyMultiPointContactConstraint<T, d>;
-	using FConstraintHandleAllocator = TConstraintHandleAllocator<TPBDCollisionConstraints<T, d>>;
-	using FConstraintContainerHandleKey = typename TPBDCollisionConstraintHandle<T, d>::FHandleKey;
-	using FCollisionModifier = TCollisionModifierCallback<T, d>;
-
-	TPBDCollisionConstraints(const TPBDRigidsSOAs<T,d>& InParticles, 
+	using FHandles = TArray<FPBDCollisionConstraintHandle*>;
+	using FConstraintHandleAllocator = TConstraintHandleAllocator<FPBDCollisionConstraints>;
+	using FConstraintContainerHandleKey = typename FPBDCollisionConstraintHandle::FHandleKey;
+
+	// For use by dependent types
+	using FPointContactConstraint = FRigidBodyPointContactConstraint;
+	using FConstraintContainerHandle = FPBDCollisionConstraintHandle;
+
+
+	FPBDCollisionConstraints(const TPBDRigidsSOAs<FReal, 3>& InParticles, 
 		TArrayCollectionArray<bool>& Collided, 
 		const TArrayCollectionArray<TSerializablePtr<FChaosPhysicsMaterial>>& PerParticleMaterials, 
-		const int32 ApplyPairIterations = 1, const int32 ApplyPushOutPairIterations = 1, const T CullDistance = (T)0, const T ShapePadding = (T)0);
+		const int32 ApplyPairIterations = 1, const int32 ApplyPushOutPairIterations = 1, const FReal CullDistance = (FReal)0, const FReal ShapePadding = (FReal)0);
 
-	virtual ~TPBDCollisionConstraints() {}
+	virtual ~FPBDCollisionConstraints() {}
 
 	/**
 	 * Whether this container provides constraint handles (simple solvers do not need them)
@@ -117,9 +84,9 @@ public:
 	*  This function will make a deep copy of the constraint and 
 	*  then delete the InConstraint. 
 	*/
-	void AddConstraint(const TRigidBodyPointContactConstraint<FReal, 3>& InConstraint);
-	void AddConstraint(const TRigidBodySweptPointContactConstraint<FReal, 3>& InConstraint);
-	void AddConstraint(const TRigidBodyMultiPointContactConstraint<FReal, 3>& InConstraint);
+	void AddConstraint(const FRigidBodyPointContactConstraint& InConstraint);
+	void AddConstraint(const FRigidBodySweptPointContactConstraint& InConstraint);
+	void AddConstraint(const FRigidBodyMultiPointContactConstraint& InConstraint);
 
 	/**
 	*  Reset the constraint frame. 
@@ -131,25 +98,25 @@ public:
 	 * You would probably call this in the PostComputeCallback. Prefer this to calling RemoveConstraints in a loop,
 	 * so you don't have to worry about constraint iterator/indices changing.
 	 */
-	void ApplyCollisionModifier(const FCollisionModifier& CollisionModifier);
+	void ApplyCollisionModifier(const FCollisionModifierCallback& CollisionModifier);
 
 
 	/**
 	* Remove the constraints associated with the ParticleHandle.
 	*/
-	void RemoveConstraints(const TSet<TGeometryParticleHandle<T, d>*>&  ParticleHandle);
+	void RemoveConstraints(const TSet<TGeometryParticleHandle<FReal, 3>*>&  ParticleHandle);
 
 
 	/**
 	* Remove the constraint, update the handle, and any maps. 
 	*/
-	void RemoveConstraint(FConstraintContainerHandle*);
+	void RemoveConstraint(FPBDCollisionConstraintHandle* ConstraintHandle);
 
 
 	/**
 	 * Update all constraint values within the set
 	 */
-	void UpdateConstraints(T Dt, const TSet<TGeometryParticleHandle<T, d>*>& AddedParticles);
+	void UpdateConstraints(FReal Dt, const TSet<TGeometryParticleHandle<FReal, 3>*>& AddedParticles);
 
 	/**
 	 * Update all constraint values
@@ -158,12 +125,12 @@ public:
 	 /**
 	 * Update all constraint values
 	 */
-	void UpdateConstraints(T Dt);
+	void UpdateConstraints(FReal Dt);
 
 	/**
 	* Update all contact manifolds
 	*/
-	void UpdateManifolds(T Dt);
+	void UpdateManifolds(FReal Dt);
 
 
 	//
@@ -177,32 +144,32 @@ public:
 	/**
 	 * Generate all contact constraints.
 	 */
-	void UpdatePositionBasedState(const T Dt);
+	void UpdatePositionBasedState(const FReal Dt);
 
 	//
 	// Simple Rule API
 	//
 
-	bool Apply(const T Dt, const int32 It, const int32 NumIts);
-	bool ApplyPushOut(const T Dt, const int32 It, const int32 NumIts);
+	bool Apply(const FReal Dt, const int32 It, const int32 NumIts);
+	bool ApplyPushOut(const FReal Dt, const int32 It, const int32 NumIts);
 
 	//
 	// Island Rule API
 	//
 	// @todo(ccaulfield): this runs wide. The serial/parallel decision should be in the ConstraintRule
 
-	bool Apply(const T Dt, const TArray<FConstraintContainerHandle*>& InConstraintHandles, const int32 It, const int32 NumIts);
-	bool ApplyPushOut(const T Dt, const TArray<FConstraintContainerHandle*>& InConstraintHandles, 
-		const TSet<const TGeometryParticleHandle<T,d>*>& IsTemporarilyStatic, int32 Iteration, int32 NumIterations);
+	bool Apply(const FReal Dt, const TArray<FPBDCollisionConstraintHandle*>& InConstraintHandles, const int32 It, const int32 NumIts);
+	bool ApplyPushOut(const FReal Dt, const TArray<FPBDCollisionConstraintHandle*>& InConstraintHandles, 
+		const TSet<const TGeometryParticleHandle<FReal, 3>*>& IsTemporarilyStatic, int32 Iteration, int32 NumIterations);
 
 
 	/**
 	 *  Callbacks
 	 */
-	void SetPostApplyCallback(const TRigidBodyContactConstraintsPostApplyCallback<T, d>& Callback);
+	void SetPostApplyCallback(const FRigidBodyContactConstraintsPostApplyCallback& Callback);
 	void ClearPostApplyCallback();
 
-	void SetPostApplyPushOutCallback(const TRigidBodyContactConstraintsPostApplyPushOutCallback<T, d>& Callback);
+	void SetPostApplyPushOutCallback(const FRigidBodyContactConstraintsPostApplyPushOutCallback& Callback);
 	void ClearPostApplyPushOutCallback();
 
 
@@ -210,27 +177,27 @@ public:
 	// Member Access
 	//
 
-	const TArray<FConstraintContainerHandle*>& GetAllConstraintHandles() const 
+	const TArray<FPBDCollisionConstraintHandle*>& GetAllConstraintHandles() const 
 	{ 
 		return Handles; 
 	}
 
-	bool Contains(const FConstraintBase* Base) const 
+	bool Contains(const FCollisionConstraintBase* Base) const
 	{
 #if CHAOS_COLLISION_PERSISTENCE_ENABLED
-		return Manifolds.Contains(FConstraintContainerHandle::MakeKey(Base));
+		return Manifolds.Contains(FPBDCollisionConstraintHandle::MakeKey(Base));
 #else
 		return false;
 #endif
 	}
 
 	// @todo(chaos): remove
-	//void SetThickness(T InThickness)
+	//void SetThickness(FReal InThickness)
 	//{
 	//	MCullDistance = InThickness;
 	//}
 
-	void SetCullDistance(T InCullDistance)
+	void SetCullDistance(FReal InCullDistance)
 	{
 		MCullDistance = InCullDistance;
 	}
@@ -240,7 +207,7 @@ public:
 		return MCullDistance;
 	}
 
-	void SetShapePadding(T InShapePadding)
+	void SetShapePadding(FReal InShapePadding)
 	{
 		MShapePadding = InShapePadding;
 	}
@@ -285,7 +252,7 @@ public:
 		return Handles;
 	}
 
-	const FConstraintBase& GetConstraint(int32 Index) const;
+	const FCollisionConstraintBase& GetConstraint(int32 Index) const;
 
 	FCollisionConstraintsArray& GetConstraintsArray() { return Constraints; }
 
@@ -293,11 +260,11 @@ protected:
 	using Base::GetConstraintIndex;
 	using Base::SetConstraintIndex;
 
-	void UpdateConstraintMaterialProperties(FConstraintBase& Contact);
+	void UpdateConstraintMaterialProperties(FCollisionConstraintBase& Contact);
 
 private:
 
-	const TPBDRigidsSOAs<T,d>& Particles;
+	const TPBDRigidsSOAs<FReal, 3>& Particles;
 
 	FCollisionConstraintsArray Constraints;
 	int32 NumActivePointConstraints;
@@ -305,18 +272,18 @@ private:
 	int32 NumActiveIterativeConstraints;
 
 #if CHAOS_COLLISION_PERSISTENCE_ENABLED
-	TMap< FConstraintContainerHandleKey, FConstraintContainerHandle* > Manifolds;
+	TMap< FConstraintContainerHandleKey, FPBDCollisionConstraintHandle* > Manifolds;
 #endif
-	TArray<FConstraintContainerHandle*> Handles;
+	TArray<FPBDCollisionConstraintHandle*> Handles;
 	FConstraintHandleAllocator HandleAllocator;
 
 	TArrayCollectionArray<bool>& MCollided;
 	const TArrayCollectionArray<TSerializablePtr<FChaosPhysicsMaterial>>& MPhysicsMaterials;
 	int32 MApplyPairIterations;
 	int32 MApplyPushOutPairIterations;
-	T MCullDistance;
-	T MShapePadding;
-	T MAngularFriction;
+	FReal MCullDistance;
+	FReal MShapePadding;
+	FReal MAngularFriction;
 	bool bUseCCD;
 	bool bEnableCollisions;
 	bool bHandlesEnabled;
@@ -324,7 +291,7 @@ private:
 
 	int32 LifespanCounter;
 
-	TRigidBodyContactConstraintsPostApplyCallback<T, d> PostApplyCallback;
-	TRigidBodyContactConstraintsPostApplyPushOutCallback<T, d> PostApplyPushOutCallback;
+	FRigidBodyContactConstraintsPostApplyCallback PostApplyCallback;
+	FRigidBodyContactConstraintsPostApplyPushOutCallback PostApplyPushOutCallback;
 };
 }
diff --git a/Engine/Source/Runtime/Experimental/Chaos/Public/Chaos/PBDCollisionConstraintsContact.h b/Engine/Source/Runtime/Experimental/Chaos/Public/Chaos/PBDCollisionConstraintsContact.h
--- a/Engine/Source/Runtime/Experimental/Chaos/Public/Chaos/PBDCollisionConstraintsContact.h
+++ b/Engine/Source/Runtime/Experimental/Chaos/Public/Chaos/PBDCollisionConstraintsContact.h
@@ -12,24 +12,20 @@ namespace Chaos
 
 	namespace Collisions
 	{
-		template<class T = float>
-		struct TContactParticleParameters {
-			T CullDistance;
-			T ShapePadding;
+		struct FContactParticleParameters {
+			FReal CullDistance;
+			FReal ShapePadding;
 			TArrayCollectionArray<bool>* Collided;
 		};
-		using FContactParticleParameters = TContactParticleParameters<FReal>;
 
-		template<class T = float>
-		struct TContactIterationParameters {
-			const T Dt;
+		struct FContactIterationParameters {
+			const FReal Dt;
 			const int32 Iteration;
 			const int32 NumIterations;
 			const int32 NumPairIterations;
 			const ECollisionApplyType ApplyType;	// @todo(chaos): a better way to customize the collision solver
 			bool* NeedsAnotherIteration;
 		};
-		using FContactIterationParameters = TContactIterationParameters<FReal>;
 
 		// Regenerate (one-shot or incremental) the manifold plane and points
 		extern void UpdateManifold(FRigidBodyMultiPointContactConstraint& Constraint, const FReal CullDistance, const FCollisionContext& Context);
diff --git a/Engine/Source/Runtime/Experimental/Chaos/Public/Chaos/PBDCollisionConstraintsUtil.h b/Engine/Source/Runtime/Experimental/Chaos/Public/Chaos/PBDCollisionConstraintsUtil.h
--- a/Engine/Source/Runtime/Experimental/Chaos/Public/Chaos/PBDCollisionConstraintsUtil.h
+++ b/Engine/Source/Runtime/Experimental/Chaos/Public/Chaos/PBDCollisionConstraintsUtil.h
@@ -7,7 +7,7 @@
 
 namespace Chaos
 {
-	void CHAOS_API ComputeHashTable(const TArray<Chaos::TPBDCollisionConstraints<float, 3>::FPointContactConstraint>& ConstraintsArray,
+	void CHAOS_API ComputeHashTable(const TArray<Chaos::FRigidBodyPointContactConstraint>& ConstraintsArray,
 									const FBox& BoundingBox, TMultiMap<int32, int32>& HashTableMap, const float SpatialHashRadius);
 
 	void CHAOS_API ComputeHashTable(const TArray<TCollisionData<float, 3>>& CollisionsArray,
diff --git a/Engine/Source/Runtime/Experimental/Chaos/Public/Chaos/PBDConstraintRule.h b/Engine/Source/Runtime/Experimental/Chaos/Public/Chaos/PBDConstraintRule.h
--- a/Engine/Source/Runtime/Experimental/Chaos/Public/Chaos/PBDConstraintRule.h
+++ b/Engine/Source/Runtime/Experimental/Chaos/Public/Chaos/PBDConstraintRule.h
@@ -416,5 +416,5 @@ namespace Chaos
 
 // Only way to make this compile at the moment due to visibility attribute issues. TODO: Change this once a fix for this problem is applied.
 #if PLATFORM_MAC || PLATFORM_LINUX
-extern template class CHAOS_API Chaos::TPBDConstraintColorRule<Chaos::TPBDCollisionConstraints<float, 3>>;
+extern template class CHAOS_API Chaos::TPBDConstraintColorRule<Chaos::FPBDCollisionConstraints>;
 #endif
diff --git a/Engine/Source/Runtime/Experimental/Chaos/Public/Chaos/PBDRigidClustering.h b/Engine/Source/Runtime/Experimental/Chaos/Public/Chaos/PBDRigidClustering.h
--- a/Engine/Source/Runtime/Experimental/Chaos/Public/Chaos/PBDRigidClustering.h
+++ b/Engine/Source/Runtime/Experimental/Chaos/Public/Chaos/PBDRigidClustering.h
@@ -2,6 +2,7 @@
 #pragma once
 
 #include "Chaos/PBDRigidClusteredParticles.h"
+#include "Chaos/PBDCollisionConstraints.h"
 #include "Chaos/Transform.h"
 #include "Chaos/CollisionResolutionTypes.h"
 #include "Chaos/ExternalCollisionData.h"
@@ -31,17 +32,17 @@ public:
 /* 
 * PDBRigidClustering
 */
-template<class FPBDRigidEvolution, class FPBDCollisionConstraint, class T, int d>
+template<class T_FPBDRigidEvolution, class T_FPBDCollisionConstraint, class T, int d>
 class CHAOS_API TPBDRigidClustering
 {
-	typedef typename FPBDCollisionConstraint::FPointContactConstraint FPointContactConstraint;
+	typedef typename T_FPBDCollisionConstraint::FPointContactConstraint FPointContactConstraint;
 public:
 	/** Parent to children */
 	typedef TMap<TPBDRigidParticleHandle<T, d>*, TArray<TPBDRigidParticleHandle<T, d>*> > FClusterMap;
 
-	using FCollisionConstraintHandle = TPBDCollisionConstraintHandle<T, d>;
+	using FCollisionConstraintHandle = FPBDCollisionConstraintHandle;
 
-	TPBDRigidClustering(FPBDRigidEvolution& InEvolution, TPBDRigidClusteredParticles<T, d>& InParticles);
+	TPBDRigidClustering(T_FPBDRigidEvolution& InEvolution, TPBDRigidClusteredParticles<T, d>& InParticles);
 	~TPBDRigidClustering();
 
 	//
@@ -123,7 +124,7 @@ public:
 	*   ... Release bodies based collision impulses.
 	*   ... Updating properties as necessary.
 	*/
-	void AdvanceClustering(const T dt, FPBDCollisionConstraint& CollisionRule);
+	void AdvanceClustering(const T dt, T_FPBDCollisionConstraint& CollisionRule);
 
 	/**
 	*  BreakingModel
@@ -288,7 +289,7 @@ public:
 		TSharedPtr<Chaos::FImplicitObject, ESPMode::ThreadSafe> ProxyGeometry,
 		const FClusterCreationParameters<T>& Parameters);
 
-	void ComputeStrainFromCollision(const FPBDCollisionConstraint& CollisionRule);
+	void ComputeStrainFromCollision(const T_FPBDCollisionConstraint& CollisionRule);
 	void ResetCollisionImpulseArray();
 	void DisableCluster(TPBDRigidClusteredParticleHandle<T, d>* ClusteredParticle);
 	void DisableParticleWithBreakEvent(Chaos::TPBDRigidParticleHandle<float, 3>* Particle);
@@ -318,7 +319,7 @@ public:
 
 private:
 
-	FPBDRigidEvolution& MEvolution;
+	T_FPBDRigidEvolution& MEvolution;
 	TPBDRigidClusteredParticles<T, d>& MParticles;
 	TSet<Chaos::TPBDRigidClusteredParticleHandle<float, 3>*> TopLevelClusterParents;
 	TSet<Chaos::TPBDRigidParticleHandle<float, 3>*> MActiveRemovalIndices;
diff --git a/Engine/Source/Runtime/Experimental/Chaos/Public/Chaos/PBDRigidsEvolutionGBF.h b/Engine/Source/Runtime/Experimental/Chaos/Public/Chaos/PBDRigidsEvolutionGBF.h
--- a/Engine/Source/Runtime/Experimental/Chaos/Public/Chaos/PBDRigidsEvolutionGBF.h
+++ b/Engine/Source/Runtime/Experimental/Chaos/Public/Chaos/PBDRigidsEvolutionGBF.h
@@ -38,7 +38,7 @@ namespace Chaos
 		using Base = FPBDRigidsEvolutionBase;
 
 		using FGravityForces = TPerParticleGravity<FReal, 3>;
-		using FCollisionConstraints = TPBDCollisionConstraints<FReal, 3>;
+		using FCollisionConstraints = FPBDCollisionConstraints;
 		using FCollisionConstraintRule = TPBDConstraintColorRule<FCollisionConstraints>;
 		using FCollisionDetector = TCollisionDetector<FSpatialAccelerationBroadPhase, FNarrowPhase, FAsyncCollisionReceiver, FCollisionConstraints>;
 		using FExternalForces = TPerParticleExternalForces<FReal, 3>;
@@ -61,7 +61,7 @@ namespace Chaos
 			PostDetectCollisionsCallback = Cb;
 		}
 
-		void SetCollisionModifierCallback(const TCollisionModifierCallback<FReal, 3>& Cb)
+		void SetCollisionModifierCallback(const FCollisionModifierCallback& Cb)
 		{
 			CollisionModifierCallback = Cb;
 		}
@@ -182,7 +182,7 @@ namespace Chaos
 		CHAOS_API void Serialize(FChaosArchive& Ar);
 
 	protected:
-		TPBDRigidClustering<FPBDRigidsEvolutionGBF, TPBDCollisionConstraints<FReal, 3>, FReal, 3> Clustering;
+		TPBDRigidClustering<FPBDRigidsEvolutionGBF, FPBDCollisionConstraints, FReal, 3> Clustering;
 
 		FGravityForces GravityForces;
 		FCollisionConstraints CollisionConstraints;
@@ -192,7 +192,7 @@ namespace Chaos
 
 		FPBDRigidsEvolutionCallback PostIntegrateCallback;
 		FPBDRigidsEvolutionCallback PostDetectCollisionsCallback;
-		TCollisionModifierCallback<FReal, 3> CollisionModifierCallback;
+		FCollisionModifierCallback CollisionModifierCallback;
 		FPBDRigidsEvolutionCallback PreApplyCallback;
 		FPBDRigidsEvolutionIslandCallback PostApplyCallback;
 		FPBDRigidsEvolutionIslandCallback PostApplyPushOutCallback;
diff --git a/Engine/Source/Runtime/Experimental/ChaosCore/Public/Chaos/Core.h b/Engine/Source/Runtime/Experimental/ChaosCore/Public/Chaos/Core.h
--- a/Engine/Source/Runtime/Experimental/ChaosCore/Public/Chaos/Core.h
+++ b/Engine/Source/Runtime/Experimental/ChaosCore/Public/Chaos/Core.h
@@ -10,6 +10,8 @@
 
 namespace Chaos
 {
+	template<class T, int d> class TAABB;
+
 	/**
 	 * Common data types for the Chaos physics engine. Unless a specific
 	 * precision of type is required most code should use these existing types
@@ -21,7 +23,7 @@ namespace Chaos
 	using FMatrix33 = PMatrix<FReal, 3, 3>;
 	using FRigidTransform3 = TRigidTransform<FReal, 3>;
 	using FRotation3 = TRotation<FReal, 3>;
-
+	using FAABB3 = TAABB<FReal, 3>;
 
 	template <typename T>
 	using TVec3 = TVector<T, 3>;
diff --git a/Engine/Source/Runtime/Experimental/ChaosSolvers/Private/EventDefaults.cpp b/Engine/Source/Runtime/Experimental/ChaosSolvers/Private/EventDefaults.cpp
--- a/Engine/Source/Runtime/Experimental/ChaosSolvers/Private/EventDefaults.cpp
+++ b/Engine/Source/Runtime/Experimental/ChaosSolvers/Private/EventDefaults.cpp
@@ -47,7 +47,7 @@ namespace Chaos
 
 			const FPBDRigidsSolver::FPBDRigidsEvolution* Evolution = Solver->GetEvolution();
 
-			const FPhysicsSolver::FPBDCollisionConstraints& CollisionRule = Evolution->GetCollisionConstraints();
+			const FPBDCollisionConstraints& CollisionRule = Evolution->GetCollisionConstraints();
 
 
 			const TPBDRigidParticles<float, 3>& Particles = Evolution->GetParticles().GetDynamicParticles();
@@ -55,20 +55,20 @@ namespace Chaos
 #if TODO_REIMPLEMENT_RIGID_CLUSTERING
 			const Chaos::FPBDRigidsSolver::FClusteringType::FClusterMap& ParentToChildrenMap = Evolution->GetRigidClustering().GetChildrenMap();
 #endif
-			const Chaos::TPBDRigidClustering<FPhysicsSolver::FPBDRigidsEvolution, FPhysicsSolver::FPBDCollisionConstraints, float, 3>::FClusterMap& ParentToChildrenMap = Evolution->GetRigidClustering().GetChildrenMap();
+			const Chaos::TPBDRigidClustering<FPhysicsSolver::FPBDRigidsEvolution, FPBDCollisionConstraints, float, 3>::FClusterMap& ParentToChildrenMap = Evolution->GetRigidClustering().GetChildrenMap();
 
 			if(CollisionRule.NumConstraints() > 0)
 			{
 				// Get the number of valid constraints (AccumulatedImpulse != 0.f and Phi < 0.f) from AllConstraintsArray
-				TArray<const Chaos::TPBDCollisionConstraintHandle<float, 3> *> ValidCollisionHandles;
+				TArray<const Chaos::FPBDCollisionConstraintHandle*> ValidCollisionHandles;
 				ValidCollisionHandles.SetNumUninitialized(CollisionRule.NumConstraints());
 				int32 NumValidCollisions = 0;
 
-				for (const Chaos::TPBDCollisionConstraintHandle<float, 3> * ContactHandle : CollisionRule.GetConstConstraintHandles())
+				for (const Chaos::FPBDCollisionConstraintHandle * ContactHandle : CollisionRule.GetConstConstraintHandles())
 				{
-					if (ContactHandle->GetType() == TPBDCollisionConstraintHandle<float, 3>::FConstraintBase::FType::SinglePoint)
+					if (ContactHandle->GetType() == FCollisionConstraintBase::FType::SinglePoint)
 					{
-						const TRigidBodyPointContactConstraint<float, 3>& Constraint = ContactHandle->GetPointContact();
+						const FRigidBodyPointContactConstraint& Constraint = ContactHandle->GetPointContact();
 
 						// Since Clustered GCs can be unioned the particleIndex representing the union 
 						// is not associated with a PhysicsProxy
@@ -100,9 +100,9 @@ namespace Chaos
 							}
 						}
 					}
-					else if (ContactHandle->GetType() == TPBDCollisionConstraintHandle<float, 3>::FConstraintBase::FType::MultiPoint)
+					else if (ContactHandle->GetType() == FCollisionConstraintBase::FType::MultiPoint)
 					{
-						const TRigidBodyMultiPointContactConstraint<float, 3>& Constraint = ContactHandle->GetMultiPointContact();
+						const FRigidBodyMultiPointContactConstraint& Constraint = ContactHandle->GetMultiPointContact();
 
 						// Since Clustered GCs can be unioned the particleIndex representing the union 
 						// is not associated with a PhysicsProxy
@@ -141,9 +141,9 @@ namespace Chaos
 				{
 					for (int32 IdxCollision = 0; IdxCollision < ValidCollisionHandles.Num(); ++IdxCollision)
 					{
-						if (ValidCollisionHandles[IdxCollision]->GetType() == TPBDCollisionConstraintHandle<float, 3>::FConstraintBase::FType::SinglePoint)
+						if (ValidCollisionHandles[IdxCollision]->GetType() == FCollisionConstraintBase::FType::SinglePoint)
 						{
-							Chaos::TPBDCollisionConstraints<float, 3>::FPointContactConstraint const& Constraint = ValidCollisionHandles[IdxCollision]->GetPointContact();
+							Chaos::FRigidBodyPointContactConstraint const& Constraint = ValidCollisionHandles[IdxCollision]->GetPointContact();
 
 							TGeometryParticleHandle<float, 3>* Particle0 = Constraint.Particle[0];
 							TGeometryParticleHandle<float, 3>* Particle1 = Constraint.Particle[1];
@@ -215,9 +215,9 @@ namespace Chaos
 							}
 						}
 
-						else if (ValidCollisionHandles[IdxCollision]->GetType() == TPBDCollisionConstraintHandle<float, 3>::FConstraintBase::FType::MultiPoint)
+						else if (ValidCollisionHandles[IdxCollision]->GetType() == FCollisionConstraintBase::FType::MultiPoint)
 						{
-							Chaos::TPBDCollisionConstraints<float, 3>::FMultiPointContactConstraint const& Constraint = ValidCollisionHandles[IdxCollision]->GetMultiPointContact();
+							Chaos::FRigidBodyMultiPointContactConstraint const& Constraint = ValidCollisionHandles[IdxCollision]->GetMultiPointContact();
 
 							TGeometryParticleHandle<float, 3>* Particle0 = Constraint.Particle[0];
 							TGeometryParticleHandle<float, 3>* Particle1 = Constraint.Particle[1];
diff --git a/Engine/Source/Runtime/Experimental/ChaosSolvers/Public/PBDRigidsSolver.h b/Engine/Source/Runtime/Experimental/ChaosSolvers/Public/PBDRigidsSolver.h
--- a/Engine/Source/Runtime/Experimental/ChaosSolvers/Public/PBDRigidsSolver.h
+++ b/Engine/Source/Runtime/Experimental/ChaosSolvers/Public/PBDRigidsSolver.h
@@ -96,9 +96,7 @@ namespace Chaos
 		typedef Chaos::TGeometryParticle<float, 3> FParticle;
 		typedef Chaos::TGeometryParticleHandle<float, 3> FHandle;
 		typedef Chaos::FPBDRigidsEvolutionGBF FPBDRigidsEvolution;
-		typedef Chaos::TPBDCollisionConstraints<float, 3> FPBDCollisionConstraints;
 
-		typedef FPBDCollisionConstraints FCollisionConstraints;
 		typedef TPBDRigidDynamicSpringConstraints<float, 3> FRigidDynamicSpringConstraints;
 		typedef TPBDPositionConstraints<float, 3> FPositionConstraints;
 
