Fix CIS - minor HeightField.cpp change. Not sure why it fixes it!

diff --git a/Engine/Source/Runtime/Experimental/Chaos/Private/Chaos/HeightField.cpp b/Engine/Source/Runtime/Experimental/Chaos/Private/Chaos/HeightField.cpp
--- a/Engine/Source/Runtime/Experimental/Chaos/Private/Chaos/HeightField.cpp
+++ b/Engine/Source/Runtime/Experimental/Chaos/Private/Chaos/HeightField.cpp
@@ -1700,7 +1700,7 @@ namespace Chaos
 		FlatGrid = TUniformGrid<T, 2>(MinCorner, MaxCorner, Cells);
 	}
 
+	template class THeightField<float>;
 }
 
-template class Chaos::THeightField<float>;
 
diff --git a/Engine/Source/Runtime/Experimental/Chaos/Public/Chaos/HeightField.h b/Engine/Source/Runtime/Experimental/Chaos/Public/Chaos/HeightField.h
--- a/Engine/Source/Runtime/Experimental/Chaos/Public/Chaos/HeightField.h
+++ b/Engine/Source/Runtime/Experimental/Chaos/Public/Chaos/HeightField.h
@@ -557,5 +557,8 @@ namespace Chaos
 		bool GJKContactPointImp(const GeomType& QueryGeom, const TRigidTransform<T, 3>& QueryTM, const T Thickness, TVector<T, 3>& ContactLocation, TVector<T, 3>& ContactNormal, T& ContactPhi) const;
 
 	};
-}
 
+#if PLATFORM_MAC || PLATFORM_LINUX
+	extern template class CHAOS_API THeightField<float>;
+#endif
+}
