Fix CIS - CollisionResolutionUtil.cpp (non ISPC build)

diff --git a/Engine/Source/Runtime/Experimental/Chaos/Private/Chaos/CollisionResolutionUtil.cpp b/Engine/Source/Runtime/Experimental/Chaos/Private/Chaos/CollisionResolutionUtil.cpp
--- a/Engine/Source/Runtime/Experimental/Chaos/Private/Chaos/CollisionResolutionUtil.cpp
+++ b/Engine/Source/Runtime/Experimental/Chaos/Private/Chaos/CollisionResolutionUtil.cpp
@@ -481,7 +481,7 @@ namespace Chaos
 		}
 #else		
 		template <ECollisionUpdateType UpdateType>
-		void SampleObject(const FImplicitObject& Object, const FRigidTransform3& ObjectTransform, const TBVHParticles<T, d>& SampleParticles, const FRigidTransform3& SampleParticlesTransform, FReal Thickness, FRigidBodyPointContactConstraint& Constraint)
+		void SampleObject(const FImplicitObject& Object, const FRigidTransform3& ObjectTransform, const TBVHParticles<FReal, 3>& SampleParticles, const FRigidTransform3& SampleParticlesTransform, FReal Thickness, FRigidBodyPointContactConstraint& Constraint)
 		{
 			SCOPE_CYCLE_COUNTER(STAT_SampleObject);
 			FRigidBodyPointContactConstraint AvgConstraint = Constraint;
@@ -503,7 +503,7 @@ namespace Chaos
 			if (NumParticles > SampleMinParticlesForAcceleration && Object.HasBoundingBox())
 			{
 				SCOPE_CYCLE_COUNTER(STAT_UpdateLevelsetPartial);
-				TAABB<T, d> ImplicitBox = Object.BoundingBox().TransformedAABB(ObjectTransform.GetRelativeTransform(SampleParticlesTransform));
+				FAABB3 ImplicitBox = Object.BoundingBox().TransformedAABB(ObjectTransform.GetRelativeTransform(SampleParticlesTransform));
 				ImplicitBox.Thicken(Thickness);
 				TArray<int32> PotentialParticles;
 				{
