Fix CIS - another attempt at linker error fixes. extern template fun.

diff --git a/Engine/Source/Runtime/Experimental/Chaos/Public/Chaos/HeightField.h b/Engine/Source/Runtime/Experimental/Chaos/Public/Chaos/HeightField.h
--- a/Engine/Source/Runtime/Experimental/Chaos/Public/Chaos/HeightField.h
+++ b/Engine/Source/Runtime/Experimental/Chaos/Public/Chaos/HeightField.h
@@ -24,28 +24,28 @@ namespace Chaos
 namespace Chaos
 {
 	template<typename T>
-	class CHAOS_API THeightField final : public FImplicitObject
+	class THeightField final : public FImplicitObject
 	{
 	public:
 		using FImplicitObject::GetTypeName;
 
-		THeightField(TArray<T>&& Height, TArray<uint8>&& InMaterialIndices, int32 InNumRows, int32 InNumCols, const TVector<T,3>& InScale);
-		THeightField(TArrayView<const uint16> InHeights, TArrayView<uint8> InMaterialIndices, int32 InNumRows, int32 InNumCols, const TVector<T, 3>& InScale);
-		THeightField(const THeightField& Other) = delete;
-		THeightField(THeightField&& Other) = default;
+		CHAOS_API THeightField(TArray<T>&& Height, TArray<uint8>&& InMaterialIndices, int32 InNumRows, int32 InNumCols, const TVector<T,3>& InScale);
+		CHAOS_API THeightField(TArrayView<const uint16> InHeights, TArrayView<uint8> InMaterialIndices, int32 InNumRows, int32 InNumCols, const TVector<T, 3>& InScale);
+		CHAOS_API THeightField(const THeightField& Other) = delete;
+		CHAOS_API THeightField(THeightField&& Other) = default;
 
 		virtual ~THeightField() {}
 
 		/** Support for editing a subsection of the heightfield */
-		void EditHeights(TArrayView<T> InHeights, int32 InBeginRow, int32 InBeginCol, int32 InNumRows, int32 InNumCols);
-		void EditHeights(TArrayView<const uint16> InHeights, int32 InBeginRow, int32 InBeginCol, int32 InNumRows, int32 InNumCols);
-		T GetHeight(int32 InIndex) const;
-		T GetHeight(int32 InX, int32 InY) const;
-		uint8 GetMaterialIndex(int32 InIndex) const;
-		uint8 GetMaterialIndex(int32 InX, int32 InY) const;
-		bool IsHole(int32 InIndex) const;
-		bool IsHole(int32 InCellX, int32 InCellY) const;
-		T GetHeightAt(const TVector<T, 2>& InGridLocationLocal) const;
+		CHAOS_API void EditHeights(TArrayView<T> InHeights, int32 InBeginRow, int32 InBeginCol, int32 InNumRows, int32 InNumCols);
+		CHAOS_API void EditHeights(TArrayView<const uint16> InHeights, int32 InBeginRow, int32 InBeginCol, int32 InNumRows, int32 InNumCols);
+		CHAOS_API T GetHeight(int32 InIndex) const;
+		CHAOS_API T GetHeight(int32 InX, int32 InY) const;
+		CHAOS_API uint8 GetMaterialIndex(int32 InIndex) const;
+		CHAOS_API uint8 GetMaterialIndex(int32 InX, int32 InY) const;
+		CHAOS_API bool IsHole(int32 InIndex) const;
+		CHAOS_API bool IsHole(int32 InCellX, int32 InCellY) const;
+		CHAOS_API T GetHeightAt(const TVector<T, 2>& InGridLocationLocal) const;
 
 		int32 GetNumRows() const { return GeomData.NumRows; }
 		int32 GetNumCols() const { return GeomData.NumCols; }
@@ -56,39 +56,39 @@ namespace Chaos
 			return (T)0;
 		}
 
-		virtual bool Raycast(const TVector<T, 3>& StartPoint, const TVector<T, 3>& Dir, const T Length, const T Thickness, T& OutTime, TVector<T,3>& OutPosition, TVector<T,3>& OutNormal, int32& OutFaceIndex) const override;
-		virtual bool Overlap(const TVector<T, 3>& Point, const T Thickness) const override;
+		CHAOS_API virtual bool Raycast(const TVector<T, 3>& StartPoint, const TVector<T, 3>& Dir, const T Length, const T Thickness, T& OutTime, TVector<T,3>& OutPosition, TVector<T,3>& OutNormal, int32& OutFaceIndex) const override;
+		CHAOS_API virtual bool Overlap(const TVector<T, 3>& Point, const T Thickness) const override;
 		
-		bool OverlapGeom(const TSphere<T, 3>& QueryGeom, const TRigidTransform<T, 3>& QueryTM, const T Thickness, FMTDInfo* OutMTD = nullptr) const;
-		bool OverlapGeom(const TBox<T, 3>& QueryGeom, const TRigidTransform<T, 3>& QueryTM, const T Thickness, FMTDInfo* OutMTD = nullptr) const;
-		bool OverlapGeom(const TCapsule<T>& QueryGeom, const TRigidTransform<T, 3>& QueryTM, const T Thickness, FMTDInfo* OutMTD = nullptr) const;
-		bool OverlapGeom(const FConvex& QueryGeom, const TRigidTransform<T, 3>& QueryTM, const T Thickness, FMTDInfo* OutMTD = nullptr) const;
-		bool OverlapGeom(const TImplicitObjectScaled<TSphere<T, 3>>& QueryGeom, const TRigidTransform<T, 3>& QueryTM, const T Thickness, FMTDInfo* OutMTD = nullptr) const;
-		bool OverlapGeom(const TImplicitObjectScaled<TBox<T, 3>>& QueryGeom, const TRigidTransform<T, 3>& QueryTM, const T Thickness, FMTDInfo* OutMTD = nullptr) const;
-		bool OverlapGeom(const TImplicitObjectScaled<TCapsule<T>>& QueryGeom, const TRigidTransform<T, 3>& QueryTM, const T Thickness, FMTDInfo* OutMTD = nullptr) const;
-		bool OverlapGeom(const TImplicitObjectScaled<FConvex>& QueryGeom, const TRigidTransform<T, 3>& QueryTM, const T Thickness, FMTDInfo* OutMTD = nullptr) const;
+		CHAOS_API bool OverlapGeom(const TSphere<T, 3>& QueryGeom, const TRigidTransform<T, 3>& QueryTM, const T Thickness, FMTDInfo* OutMTD = nullptr) const;
+		CHAOS_API bool OverlapGeom(const TBox<T, 3>& QueryGeom, const TRigidTransform<T, 3>& QueryTM, const T Thickness, FMTDInfo* OutMTD = nullptr) const;
+		CHAOS_API bool OverlapGeom(const TCapsule<T>& QueryGeom, const TRigidTransform<T, 3>& QueryTM, const T Thickness, FMTDInfo* OutMTD = nullptr) const;
+		CHAOS_API bool OverlapGeom(const FConvex& QueryGeom, const TRigidTransform<T, 3>& QueryTM, const T Thickness, FMTDInfo* OutMTD = nullptr) const;
+		CHAOS_API bool OverlapGeom(const TImplicitObjectScaled<TSphere<T, 3>>& QueryGeom, const TRigidTransform<T, 3>& QueryTM, const T Thickness, FMTDInfo* OutMTD = nullptr) const;
+		CHAOS_API bool OverlapGeom(const TImplicitObjectScaled<TBox<T, 3>>& QueryGeom, const TRigidTransform<T, 3>& QueryTM, const T Thickness, FMTDInfo* OutMTD = nullptr) const;
+		CHAOS_API bool OverlapGeom(const TImplicitObjectScaled<TCapsule<T>>& QueryGeom, const TRigidTransform<T, 3>& QueryTM, const T Thickness, FMTDInfo* OutMTD = nullptr) const;
+		CHAOS_API bool OverlapGeom(const TImplicitObjectScaled<FConvex>& QueryGeom, const TRigidTransform<T, 3>& QueryTM, const T Thickness, FMTDInfo* OutMTD = nullptr) const;
 
-		bool SweepGeom(const TSphere<T, 3>& QueryGeom, const TRigidTransform<T, 3>& StartTM, const TVector<T, 3>& Dir, const T Length, T& OutTime, TVector<T, 3>& OutPosition, TVector<T, 3>& OutNormal, int32& OutFaceIndex, const T Thickness = 0, bool bComputeMTD = false) const;
-		bool SweepGeom(const TBox<T, 3>& QueryGeom, const TRigidTransform<T, 3>& StartTM, const TVector<T, 3>& Dir, const T Length, T& OutTime, TVector<T, 3>& OutPosition, TVector<T, 3>& OutNormal, int32& OutFaceIndex, const T Thickness = 0, bool bComputeMTD = false) const;
-		bool SweepGeom(const TCapsule<T>& QueryGeom, const TRigidTransform<T, 3>& StartTM, const TVector<T, 3>& Dir, const T Length, T& OutTime, TVector<T, 3>& OutPosition, TVector<T, 3>& OutNormal, int32& OutFaceIndex, const T Thickness = 0, bool bComputeMTD = false) const;
-		bool SweepGeom(const FConvex& QueryGeom, const TRigidTransform<T, 3>& StartTM, const TVector<T, 3>& Dir, const T Length, T& OutTime, TVector<T, 3>& OutPosition, TVector<T, 3>& OutNormal, int32& OutFaceIndex, const T Thickness = 0, bool bComputeMTD = false) const;
-		bool SweepGeom(const TImplicitObjectScaled<TSphere<T, 3>>& QueryGeom, const TRigidTransform<T, 3>& StartTM, const TVector<T, 3>& Dir, const T Length, T& OutTime, TVector<T, 3>& OutPosition, TVector<T, 3>& OutNormal, int32& OutFaceIndex, const T Thickness = 0, bool bComputeMTD = false) const;
-		bool SweepGeom(const TImplicitObjectScaled<TBox<T, 3>>& QueryGeom, const TRigidTransform<T, 3>& StartTM, const TVector<T, 3>& Dir, const T Length, T& OutTime, TVector<T, 3>& OutPosition, TVector<T, 3>& OutNormal, int32& OutFaceIndex, const T Thickness = 0, bool bComputeMTD = false) const;
-		bool SweepGeom(const TImplicitObjectScaled<TCapsule<T>>& QueryGeom, const TRigidTransform<T, 3>& StartTM, const TVector<T, 3>& Dir, const T Length, T& OutTime, TVector<T, 3>& OutPosition, TVector<T, 3>& OutNormal, int32& OutFaceIndex, const T Thickness = 0, bool bComputeMTD = false) const;
-		bool SweepGeom(const TImplicitObjectScaled<FConvex>& QueryGeom, const TRigidTransform<T, 3>& StartTM, const TVector<T, 3>& Dir, const T Length, T& OutTime, TVector<T, 3>& OutPosition, TVector<T, 3>& OutNormal, int32& OutFaceIndex, const T Thickness = 0, bool bComputeMTD = false) const;
+		CHAOS_API bool SweepGeom(const TSphere<T, 3>& QueryGeom, const TRigidTransform<T, 3>& StartTM, const TVector<T, 3>& Dir, const T Length, T& OutTime, TVector<T, 3>& OutPosition, TVector<T, 3>& OutNormal, int32& OutFaceIndex, const T Thickness = 0, bool bComputeMTD = false) const;
+		CHAOS_API bool SweepGeom(const TBox<T, 3>& QueryGeom, const TRigidTransform<T, 3>& StartTM, const TVector<T, 3>& Dir, const T Length, T& OutTime, TVector<T, 3>& OutPosition, TVector<T, 3>& OutNormal, int32& OutFaceIndex, const T Thickness = 0, bool bComputeMTD = false) const;
+		CHAOS_API bool SweepGeom(const TCapsule<T>& QueryGeom, const TRigidTransform<T, 3>& StartTM, const TVector<T, 3>& Dir, const T Length, T& OutTime, TVector<T, 3>& OutPosition, TVector<T, 3>& OutNormal, int32& OutFaceIndex, const T Thickness = 0, bool bComputeMTD = false) const;
+		CHAOS_API bool SweepGeom(const FConvex& QueryGeom, const TRigidTransform<T, 3>& StartTM, const TVector<T, 3>& Dir, const T Length, T& OutTime, TVector<T, 3>& OutPosition, TVector<T, 3>& OutNormal, int32& OutFaceIndex, const T Thickness = 0, bool bComputeMTD = false) const;
+		CHAOS_API bool SweepGeom(const TImplicitObjectScaled<TSphere<T, 3>>& QueryGeom, const TRigidTransform<T, 3>& StartTM, const TVector<T, 3>& Dir, const T Length, T& OutTime, TVector<T, 3>& OutPosition, TVector<T, 3>& OutNormal, int32& OutFaceIndex, const T Thickness = 0, bool bComputeMTD = false) const;
+		CHAOS_API bool SweepGeom(const TImplicitObjectScaled<TBox<T, 3>>& QueryGeom, const TRigidTransform<T, 3>& StartTM, const TVector<T, 3>& Dir, const T Length, T& OutTime, TVector<T, 3>& OutPosition, TVector<T, 3>& OutNormal, int32& OutFaceIndex, const T Thickness = 0, bool bComputeMTD = false) const;
+		CHAOS_API bool SweepGeom(const TImplicitObjectScaled<TCapsule<T>>& QueryGeom, const TRigidTransform<T, 3>& StartTM, const TVector<T, 3>& Dir, const T Length, T& OutTime, TVector<T, 3>& OutPosition, TVector<T, 3>& OutNormal, int32& OutFaceIndex, const T Thickness = 0, bool bComputeMTD = false) const;
+		CHAOS_API bool SweepGeom(const TImplicitObjectScaled<FConvex>& QueryGeom, const TRigidTransform<T, 3>& StartTM, const TVector<T, 3>& Dir, const T Length, T& OutTime, TVector<T, 3>& OutPosition, TVector<T, 3>& OutNormal, int32& OutFaceIndex, const T Thickness = 0, bool bComputeMTD = false) const;
 
-		bool GJKContactPoint(const TBox<T, 3>& QueryGeom, const TRigidTransform<T, 3>& QueryTM, const T Thickness, TVector<T, 3>& ContactLocation, TVector<T, 3>& ContactNormal, T& ContactPhi) const;
-		bool GJKContactPoint(const TSphere<T, 3>& QueryGeom, const TRigidTransform<T, 3>& QueryTM, const T Thickness, TVector<T, 3>& ContactLocation, TVector<T, 3>& ContactNormal, T& ContactPhi) const;
-		bool GJKContactPoint(const TCapsule<T>& QueryGeom, const TRigidTransform<T, 3>& QueryTM, const T Thickness, TVector<T, 3>& ContactLocation, TVector<T, 3>& ContactNormal, T& ContactPhi) const;
-		bool GJKContactPoint(const FConvex& QueryGeom, const TRigidTransform<T, 3>& QueryTM, const T Thickness, TVector<T, 3>& ContactLocation, TVector<T, 3>& ContactNormal, T& ContactPhi) const;
-		bool GJKContactPoint(const TImplicitObjectScaled<TBox<T, 3>>& QueryGeom, const TRigidTransform<T, 3>& QueryTM, const T Thickness, TVector<T, 3>& ContactLocation, TVector<T, 3>& ContactNormal, T& ContactPhi) const;
-		bool GJKContactPoint(const TImplicitObjectScaled<TSphere<T, 3>>& QueryGeom, const TRigidTransform<T, 3>& QueryTM, const T Thickness, TVector<T, 3>& ContactLocation, TVector<T, 3>& ContactNormal, T& ContactPhi) const;
-		bool GJKContactPoint(const TImplicitObjectScaled<TCapsule<T>>& QueryGeom, const TRigidTransform<T, 3>& QueryTM, const T Thickness, TVector<T, 3>& ContactLocation, TVector<T, 3>& ContactNormal, T& ContactPhi) const;
-		bool GJKContactPoint(const TImplicitObjectScaled<FConvex>& QueryGeom, const TRigidTransform<T, 3>& QueryTM, const T Thickness, TVector<T, 3>& ContactLocation, TVector<T, 3>& ContactNormal, T& ContactPhi) const;
+		CHAOS_API bool GJKContactPoint(const TBox<T, 3>& QueryGeom, const TRigidTransform<T, 3>& QueryTM, const T Thickness, TVector<T, 3>& ContactLocation, TVector<T, 3>& ContactNormal, T& ContactPhi) const;
+		CHAOS_API bool GJKContactPoint(const TSphere<T, 3>& QueryGeom, const TRigidTransform<T, 3>& QueryTM, const T Thickness, TVector<T, 3>& ContactLocation, TVector<T, 3>& ContactNormal, T& ContactPhi) const;
+		CHAOS_API bool GJKContactPoint(const TCapsule<T>& QueryGeom, const TRigidTransform<T, 3>& QueryTM, const T Thickness, TVector<T, 3>& ContactLocation, TVector<T, 3>& ContactNormal, T& ContactPhi) const;
+		CHAOS_API bool GJKContactPoint(const FConvex& QueryGeom, const TRigidTransform<T, 3>& QueryTM, const T Thickness, TVector<T, 3>& ContactLocation, TVector<T, 3>& ContactNormal, T& ContactPhi) const;
+		CHAOS_API bool GJKContactPoint(const TImplicitObjectScaled<TBox<T, 3>>& QueryGeom, const TRigidTransform<T, 3>& QueryTM, const T Thickness, TVector<T, 3>& ContactLocation, TVector<T, 3>& ContactNormal, T& ContactPhi) const;
+		CHAOS_API bool GJKContactPoint(const TImplicitObjectScaled<TSphere<T, 3>>& QueryGeom, const TRigidTransform<T, 3>& QueryTM, const T Thickness, TVector<T, 3>& ContactLocation, TVector<T, 3>& ContactNormal, T& ContactPhi) const;
+		CHAOS_API bool GJKContactPoint(const TImplicitObjectScaled<TCapsule<T>>& QueryGeom, const TRigidTransform<T, 3>& QueryTM, const T Thickness, TVector<T, 3>& ContactLocation, TVector<T, 3>& ContactNormal, T& ContactPhi) const;
+		CHAOS_API bool GJKContactPoint(const TImplicitObjectScaled<FConvex>& QueryGeom, const TRigidTransform<T, 3>& QueryTM, const T Thickness, TVector<T, 3>& ContactLocation, TVector<T, 3>& ContactNormal, T& ContactPhi) const;
 
 
-		virtual int32 FindMostOpposingFace(const TVector<T, 3>& Position, const TVector<T, 3>& UnitDir, int32 HintFaceIndex, T SearchDist) const override;
-		virtual TVector<T, 3> FindGeometryOpposingNormal(const TVector<T, 3>& DenormDir, int32 FaceIndex, const TVector<T, 3>& OriginalNormal) const override;
+		CHAOS_API virtual int32 FindMostOpposingFace(const TVector<T, 3>& Position, const TVector<T, 3>& UnitDir, int32 HintFaceIndex, T SearchDist) const override;
+		CHAOS_API virtual TVector<T, 3> FindGeometryOpposingNormal(const TVector<T, 3>& DenormDir, int32 FaceIndex, const TVector<T, 3>& OriginalNormal) const override;
 
 
 
@@ -517,19 +517,19 @@ namespace Chaos
 		};
 
 		// Helpers for accessing bounds
-		bool GetCellBounds2D(const TVector<int32, 2> InCoord, FBounds2D& OutBounds, const TVector<T, 2>& InInflate = {0}) const;
-		bool GetCellBounds3D(const TVector<int32, 2> InCoord, TVector<T, 3>& OutMin, TVector<T, 3>& OutMax, const TVector<T, 3>& InInflate = TVector<T, 3>(0)) const;
-		bool GetCellBounds2DScaled(const TVector<int32, 2> InCoord, FBounds2D& OutBounds, const TVector<T, 2>& InInflate = {0}) const;
-		bool GetCellBounds3DScaled(const TVector<int32, 2> InCoord, TVector<T, 3>& OutMin, TVector<T, 3>& OutMax, const TVector<T, 3>& InInflate = TVector<T, 3>(0)) const;
-		bool CalcCellBounds3D(const TVector<int32, 2> InCoord, TVector<T, 3>& OutMin, TVector<T, 3>& OutMax, const TVector<T, 3>& InInflate = TVector<T, 3>(0)) const;
+		CHAOS_API bool GetCellBounds2D(const TVector<int32, 2> InCoord, FBounds2D& OutBounds, const TVector<T, 2>& InInflate = {0}) const;
+		CHAOS_API bool GetCellBounds3D(const TVector<int32, 2> InCoord, TVector<T, 3>& OutMin, TVector<T, 3>& OutMax, const TVector<T, 3>& InInflate = TVector<T, 3>(0)) const;
+		CHAOS_API bool GetCellBounds2DScaled(const TVector<int32, 2> InCoord, FBounds2D& OutBounds, const TVector<T, 2>& InInflate = {0}) const;
+		CHAOS_API bool GetCellBounds3DScaled(const TVector<int32, 2> InCoord, TVector<T, 3>& OutMin, TVector<T, 3>& OutMax, const TVector<T, 3>& InInflate = TVector<T, 3>(0)) const;
+		CHAOS_API bool CalcCellBounds3D(const TVector<int32, 2> InCoord, TVector<T, 3>& OutMin, TVector<T, 3>& OutMax, const TVector<T, 3>& InInflate = TVector<T, 3>(0)) const;
 
 		// Query functions - sweep, ray, overlap
 		template<typename SQVisitor>
 		bool GridSweep(const TVector<T, 3>& StartPoint, const TVector<T, 3>& Dir, const T Length, const TVector<T, 2> InHalfExtents, SQVisitor& Visitor) const;
-		bool GridCast(const TVector<T, 3>& StartPoint, const TVector<T, 3>& Dir, const T Length, THeightfieldRaycastVisitor<T>& Visitor) const;
-		bool GetGridIntersections(FBounds2D InFlatBounds, TArray<TVector<int32, 2>>& OutInterssctions) const;
+		CHAOS_API bool GridCast(const TVector<T, 3>& StartPoint, const TVector<T, 3>& Dir, const T Length, THeightfieldRaycastVisitor<T>& Visitor) const;
+		CHAOS_API bool GetGridIntersections(FBounds2D InFlatBounds, TArray<TVector<int32, 2>>& OutInterssctions) const;
 		
-		FBounds2D GetFlatBounds() const;
+		CHAOS_API FBounds2D GetFlatBounds() const;
 
 		// Grid for queries, faster than bounding volumes for heightfields
 		TUniformGrid<T, 2> FlatGrid;
@@ -540,8 +540,8 @@ namespace Chaos
 		// Cached when bounds are requested. Mutable to allow GetBounds to be logical const
 		mutable TAABB<T, 3> CachedBounds;
 
-		void CalcBounds();
-		void BuildQueryData();
+		CHAOS_API void CalcBounds();
+		CHAOS_API void BuildQueryData();
 		
 		// Needed for serialization
 		THeightField() : FImplicitObject(EImplicitObject::HasBoundingBox, ImplicitObjectType::HeightField) {}
@@ -560,5 +560,7 @@ namespace Chaos
 
 #if PLATFORM_MAC || PLATFORM_LINUX
 	extern template class CHAOS_API THeightField<float>;
+#else
+	extern template class THeightField<float>;
 #endif
 }
