Fix non-unity build

diff --git a/Engine/Source/Runtime/Experimental/ChaosSolvers/Private/PhysicsProxy/SolverEventFilters.cpp b/Engine/Source/Runtime/Experimental/ChaosSolvers/Private/PhysicsProxy/SolverEventFilters.cpp
--- a/Engine/Source/Runtime/Experimental/ChaosSolvers/Private/PhysicsProxy/SolverEventFilters.cpp
+++ b/Engine/Source/Runtime/Experimental/ChaosSolvers/Private/PhysicsProxy/SolverEventFilters.cpp
@@ -2,6 +2,7 @@
 
 #include "SolverEventFilters.h"
 #include "Chaos/CollisionResolutionTypes.h"
+#include "Chaos/ExternalCollisionData.h"
 
 namespace Chaos
 {
