Refactor rewind data to make it easier to track desync

diff --git a/Engine/Source/Runtime/Experimental/Chaos/Public/Chaos/ISpatialAcceleration.h b/Engine/Source/Runtime/Experimental/Chaos/Public/Chaos/ISpatialAcceleration.h
--- a/Engine/Source/Runtime/Experimental/Chaos/Public/Chaos/ISpatialAcceleration.h
+++ b/Engine/Source/Runtime/Experimental/Chaos/Public/Chaos/ISpatialAcceleration.h
@@ -411,6 +411,11 @@ public:
 		return Entries[GetUniqueIdx(Key).Idx].Value;
 	}
 
+	const TValue& FindChecked(const TKey& Key) const
+	{
+		return Entries[GetUniqueIdx(Key).Idx].Value;
+	}
+
 	TValue& FindOrAdd(const TKey& Key)
 	{
 		if(TValue* Elem = Find(Key))
diff --git a/Engine/Source/Runtime/Experimental/ChaosSolvers/Public/RewindData.h b/Engine/Source/Runtime/Experimental/ChaosSolvers/Public/RewindData.h
--- a/Engine/Source/Runtime/Experimental/ChaosSolvers/Public/RewindData.h
+++ b/Engine/Source/Runtime/Experimental/ChaosSolvers/Public/RewindData.h
@@ -113,7 +113,7 @@ public:
 		return Dynamics.IsSet() ? Dynamics.Read().F : Particle.CastToRigidParticle()->F();
 	}
 
-	void SyncSimWritablePropsFromSim(FDirtyPropertiesManager& Manager,int32 Idx,const TPBDRigidParticleHandle<FReal,3>& Rigid)
+	bool SyncSimWritablePropsFromSim(FDirtyPropertiesManager& Manager,int32 Idx,const TPBDRigidParticleHandle<FReal,3>& Rigid)
 	{
 		FParticleDirtyFlags Flags;
 		Flags.MarkDirty(EParticleFlags::XR);
@@ -132,9 +132,11 @@ public:
 			Data.V = Rigid.PreV();
 			Data.W = Rigid.PreW();
 		});
+
+		return true;	//todo: check if values are actually different
 	}
 
-	void SyncDirtyDynamics(FDirtyPropertiesManager& DestManager,int32 DataIdx,const FParticleDirtyData& Dirty,const FDirtyPropertiesManager& SrcManager)
+	bool SyncDirtyDynamics(FDirtyPropertiesManager& DestManager,int32 DataIdx,const FParticleDirtyData& Dirty,const FDirtyPropertiesManager& SrcManager)
 	{
 		FParticleDirtyData DirtyFlags;
 		DirtyFlags.SetFlags(Dirty.GetFlags());
@@ -143,6 +145,8 @@ public:
 		{
 			Data = Dirty.GetDynamics(SrcManager,DataIdx);
 		});
+
+		return true;	//todo: check if values are actually different
 	}
 
 	void SyncToParticle(TGeometryParticle<FReal,3>& Particle) const
@@ -193,7 +197,7 @@ public:
 		}
 	}
 
-	void SyncPrevFrame(FDirtyPropertiesManager& Manager,int32 Idx,const FDirtyProxy& Dirty)
+	bool SyncPrevFrame(FDirtyPropertiesManager& Manager,int32 Idx,const FDirtyProxy& Dirty)
 	{
 		//syncs the data before it was made dirty
 		//for sim-writable props this is only possible if those props are immutable from the sim side (sleeping, not simulated, etc...)
@@ -241,6 +245,9 @@ public:
 				Data.AngularEtherDrag = Rigid->AngularEtherDrag();
 			}
 		});
+
+		//todo: check whether the data actually changed
+		return true;
 	}
 
 	void SyncIfDirty(FDirtyPropertiesManager& Manager,int32 Idx,const TGeometryParticle<FReal,3>& InParticle, const FGeometryParticleStateBase& RewindState)
@@ -387,6 +394,7 @@ public:
 	FRewindData(int32 NumFrames)
 	: Managers(NumFrames+1)	//give 1 extra for saving at head
 	, CurFrame(0)
+	, CurWave(0)
 	, FramesSaved(0)
 	, DataIdxOffset(0)
 	, bNeedsSave(false)
@@ -423,18 +431,19 @@ public:
 		int32 DataIdx = 0;
 		for(FDirtyParticleInfo& DirtyParticleInfo : AllDirtyParticles)
 		{
+			DirtyParticleInfo.bDesync = false;	//after rewind particle is pristine
+
 			if(bNeedsSave)
 			{
 				//GetStateAtFrameImp returns a pointer from the TArray that holds state data
 				//But it's possible that we'll need to save state from head, which would grow that TArray
 				//So preallocate just in case
-				FParticleRewindInfo& RewindInfo = ParticleToRewindInfo.FindChecked(DirtyParticleInfo.CachedUniqueIdx);
-				FGeometryParticleStateBase& LatestState = RewindInfo.AddFrame(CurFrame);
+				FGeometryParticleStateBase& LatestState = DirtyParticleInfo.AddFrame(CurFrame, CurWave);
 			
 				if(const FGeometryParticleStateBase* RewindState = GetStateAtFrameImp(*DirtyParticleInfo.Particle, Frame))
 				{
 					LatestState.SyncIfDirty(*DestManager,DataIdx++,*DirtyParticleInfo.Particle,*RewindState);
-					CoalesceBack(RewindInfo.GetFrames(), CurFrame);
+					CoalesceBack(DirtyParticleInfo.Frames, CurFrame);
 
 					RewindState->SyncToParticle(*DirtyParticleInfo.Particle);
 				}
@@ -457,17 +466,17 @@ public:
 
 	void RemoveParticle(const FUniqueIdx UniqueIdx)
 	{
-		if(const FParticleRewindInfo* Info = ParticleToRewindInfo.Find(UniqueIdx))
+		if(const int32* IdxPtr = ParticleToAllDirtyIdx.Find(UniqueIdx))
 		{
-			const int32 Idx = Info->AllDirtyParticlesIdx;
+			const int32 Idx = *IdxPtr;
 			AllDirtyParticles.RemoveAtSwap(Idx);
-			if(Idx< AllDirtyParticles.Num())
+			if(Idx < AllDirtyParticles.Num())
 			{
 				//update particle in new position
-				ParticleToRewindInfo.FindChecked(AllDirtyParticles[Idx].CachedUniqueIdx).AllDirtyParticlesIdx = Idx;
+				ParticleToAllDirtyIdx.FindChecked(AllDirtyParticles[Idx].CachedUniqueIdx) = Idx;
 			}
 
-			ParticleToRewindInfo.RemoveChecked(UniqueIdx);
+			ParticleToAllDirtyIdx.RemoveChecked(UniqueIdx);
 		}
 	}
 
@@ -484,6 +493,7 @@ public:
 
 	void AdvanceFrame()
 	{
+		QUICK_SCOPE_CYCLE_COUNTER(RewindDataAdvance);
 		++CurFrame;
 		FramesSaved = FMath::Min(FramesSaved+1,static_cast<int32>(Managers.Capacity()));
 		
@@ -543,22 +553,30 @@ public:
 		auto ProcessProxy = [this,&SrcManager, DataIdx, Dirty, &DestManager](const auto Proxy)
 		{
 			const auto PTParticle = Proxy->GetHandle();
-			FParticleRewindInfo& Info = FindOrAddParticle(Proxy->GetParticle(),PTParticle->UniqueIdx());
-			AllDirtyParticles[Info.AllDirtyParticlesIdx].LastDirtyFrame = CurFrame;
+			FDirtyParticleInfo& Info = FindOrAddParticle(*Proxy->GetParticle(),PTParticle->UniqueIdx());
+			Info.LastDirtyFrame = CurFrame;
+
+			bool bChanged = false;
 
 			//Most properties are always a frame behind
 			if(Proxy->IsInitialized())	//Frame delay so proxy must be initialized
 			{
-				FGeometryParticleStateBase& LatestState = Info.AddFrame(CurFrame-1);
-				LatestState.SyncPrevFrame(DestManager,DataIdx,Dirty);
-				CoalesceBack(Info.GetFrames(), CurFrame-1);
+				FGeometryParticleStateBase& LatestState = Info.AddFrame(CurFrame-1, CurWave);
+				bChanged = LatestState.SyncPrevFrame(DestManager,DataIdx,Dirty);
+				CoalesceBack(Info.Frames, CurFrame-1);
 			}
 
 			//If dynamics are dirty we must record them immediately because the sim resets them to 0
 			if(Dirty.ParticleData.IsDirty(EParticleFlags::Dynamics))
 			{
-				FGeometryParticleStateBase& LatestState = Info.AddFrame(CurFrame);
-				LatestState.SyncDirtyDynamics(DestManager,DataIdx,Dirty.ParticleData,SrcManager);
+				FGeometryParticleStateBase& LatestState = Info.AddFrame(CurFrame, CurWave);
+				bChanged = LatestState.SyncDirtyDynamics(DestManager,DataIdx,Dirty.ParticleData,SrcManager);
+			}
+
+			if(bChanged)
+			{
+				//value actually changed, so we must desync
+				Info.bDesync = true;
 			}
 		};
 
@@ -594,18 +612,22 @@ public:
 		//todo: is this check needed? why do we pass sleeping rigids into this function?
 		if(SimWritablePropsMayChange(Rigid))
 		{
-			FParticleRewindInfo& Info = FindOrAddParticle(Rigid.GTGeometryParticle(), Rigid.UniqueIdx());
-			AllDirtyParticles[Info.AllDirtyParticlesIdx].LastDirtyFrame = CurFrame-1;
+			FDirtyParticleInfo& Info = FindOrAddParticle(*Rigid.GTGeometryParticle(), Rigid.UniqueIdx());
+			Info.LastDirtyFrame = CurFrame-1;
 
 			//User called PrepareManagerForFrame (or PrepareFrameForPTDirty) for the previous frame, so use it
 			FDirtyPropertiesManager& DestManager = *Managers[CurFrame-1].Manager;
 
 			//sim-writable properties changed at head, so we must write down what they were
-			FGeometryParticleStateBase& LatestState = Info.AddFrame(CurFrame-1);
-			LatestState.SyncSimWritablePropsFromSim(DestManager,DataIdx,Rigid);
+			FGeometryParticleStateBase& LatestState = Info.AddFrame(CurFrame-1, CurWave);
+			bool bChanged = LatestState.SyncSimWritablePropsFromSim(DestManager,DataIdx,Rigid);
+			if(bChanged)
+			{
+				Info.bDesync = true;
+			}
 
 			//update any previous frames that were pointing at head
-			CoalesceBack(Info.GetFrames(), CurFrame-1);
+			CoalesceBack(Info.Frames, CurFrame-1);
 		}
 	}
 
@@ -620,31 +642,32 @@ private:
 
 		}
 
-		FGeometryParticleStateBase* GetState(int32 Frame)
+		FGeometryParticleStateBase* GetState(int32 Frame, int8 Wave)
 		{
-			return (bSet && Frame == RecordedFrame) ? &State : nullptr;
+			return (bSet && Frame == RecordedFrame && Wave == RewindWave) ? &State : nullptr;
 		}
 
-		const FGeometryParticleStateBase* GetState(int32 Frame) const
+		const FGeometryParticleStateBase* GetState(int32 Frame, int8 Wave) const
 		{
-			return (bSet && Frame == RecordedFrame) ? &State : nullptr;
+			return (bSet && Frame == RecordedFrame && Wave == RewindWave) ? &State : nullptr;
 		}
 
-		FGeometryParticleStateBase& GetStateChecked(int32 Frame)
+		FGeometryParticleStateBase& GetStateChecked(int32 Frame, int8 Wave)
 		{
-			check(bSet && Frame == RecordedFrame);
+			check(bSet && Frame == RecordedFrame && Wave == RewindWave);
 			return State;
 		}
 
-		const FGeometryParticleStateBase& GetStateChecked(int32 Frame) const
+		const FGeometryParticleStateBase& GetStateChecked(int32 Frame, int8 Wave) const
 		{
-			check(bSet && Frame == RecordedFrame);
+			check(bSet && Frame == RecordedFrame && Wave == RewindWave);
 			return State;
 		}
 
-		FGeometryParticleStateBase& NewState(int32 Frame)
+		FGeometryParticleStateBase& NewState(int32 Frame, int8 Wave)
 		{
 			RecordedFrame = Frame;
+			RewindWave = Wave;
 			bSet = true;
 			State = FGeometryParticleStateBase();
 			return State;
@@ -653,59 +676,17 @@ private:
 	private:
 		FGeometryParticleStateBase State;
 		int32 RecordedFrame;
+		int8 RewindWave;
 		bool bSet;
 	};
 
-	struct FParticleRewindInfo
-	{
-		int32 AllDirtyParticlesIdx;
-
-		FGeometryParticleStateBase& AddFrame(int32 FrameIdx)
-		{
-			FFrameInfo& Info = GetFrames()[FrameIdx];
-			if(FGeometryParticleStateBase* State = Info.GetState(FrameIdx))
-			{
-				return *State;
-			}
-
-			return Info.NewState(FrameIdx);
-		}
-
-		const TCircularBuffer<FFrameInfo>& GetFrames() const
-		{
-			return Frames.GetValue();
-
-		}
-
-		TCircularBuffer<FFrameInfo>& GetFrames()
-		{
-			return Frames.GetValue();
-
-		}
-
-		bool IsInitialized() const
-		{
-			return Frames.IsSet();
-		}
-
-		void Initialize(int32 NumFrames)
-		{
-			ensure(IsInitialized() == false);
-			Frames = TCircularBuffer<FFrameInfo>(NumFrames);
-		}
-
-	private:
-		TOptional<TCircularBuffer<FFrameInfo>> Frames;
-
-	};
-
 	void CoalesceBack(TCircularBuffer<FFrameInfo>& Frames, int32 LatestIdx)
 	{
-		const FGeometryParticleStateBase& LatestState = Frames[LatestIdx].GetStateChecked(LatestIdx);
+		const FGeometryParticleStateBase& LatestState = Frames[LatestIdx].GetStateChecked(LatestIdx, CurWave);
 		const int32 EarliestFrame = LatestIdx - FramesSaved;
 		for(int32 FrameIdx = LatestIdx - 1; FrameIdx >= EarliestFrame; --FrameIdx)
 		{
-			if(FGeometryParticleStateBase* State = Frames[FrameIdx].GetState(FrameIdx))
+			if(FGeometryParticleStateBase* State = Frames[FrameIdx].GetState(FrameIdx, CurWave))
 			{
 				if(State->CoalesceState(LatestState) == false)
 				{
@@ -726,24 +707,12 @@ private:
 		int32 FrameCreatedFor;
 	};
 
-	FParticleRewindInfo& FindOrAddParticle(TGeometryParticle<FReal,3>* UnsafeGTParticle, FUniqueIdx UniqueIdx)
-	{
-		FParticleRewindInfo& Info = ParticleToRewindInfo.FindOrAdd(UniqueIdx);
-		if(!Info.IsInitialized())
-		{
-			Info.Initialize(Managers.Capacity());
-			Info.AllDirtyParticlesIdx = AllDirtyParticles.Add(FDirtyParticleInfo{UnsafeGTParticle,UniqueIdx,CurFrame});
-		}
-		
-		return Info;
-	}
-
 	const FGeometryParticleStateBase* GetStateAtFrameImp(const TGeometryParticle<FReal,3>& Particle,int32 Frame) const
 	{
-		if(const FParticleRewindInfo* Info = ParticleToRewindInfo.Find(Particle.UniqueIdx()))
+		if(const FDirtyParticleInfo* Info = FindParticle(Particle.UniqueIdx()))
 		{
-			const TCircularBuffer<FFrameInfo>& Frames = Info->GetFrames();
-			if(const FGeometryParticleStateBase* FrameState = Frames[Frame].GetState(Frame))
+			const TCircularBuffer<FFrameInfo>& Frames = Info->Frames;
+			if(const FGeometryParticleStateBase* FrameState = Frames[Frame].GetState(Frame, CurWave))
 			{
 				return FrameState;
 			}
@@ -753,7 +722,7 @@ private:
 
 			for(int32 FrameIdx = Frame + 1; FrameIdx <= CurFrame; ++FrameIdx)
 			{
-				if(const FGeometryParticleStateBase* FrameState = Frames[FrameIdx].GetState(FrameIdx))
+				if(const FGeometryParticleStateBase* FrameState = Frames[FrameIdx].GetState(FrameIdx, CurWave))
 				{
 					return FrameState;
 				}
@@ -766,15 +735,84 @@ private:
 
 	struct FDirtyParticleInfo
 	{
+		TCircularBuffer<FFrameInfo> Frames;
 		TGeometryParticle<FReal,3>* Particle;
 		FUniqueIdx CachedUniqueIdx;	//Needed when manipulating on physics thread and Particle data cannot be read
 		int32 LastDirtyFrame;	//Track how recently this was made dirty
+		bool bDesync;
+
+		FDirtyParticleInfo(TGeometryParticle<FReal, 3>& UnsafeGTParticle, const FUniqueIdx UniqueIdx, const int32 CurFrame, const int32 NumFrames)
+		: Frames(NumFrames)
+		, Particle(&UnsafeGTParticle)
+		, CachedUniqueIdx(UniqueIdx)
+		, LastDirtyFrame(CurFrame)
+		, bDesync(true)
+		{
+
+		}
+
+		FGeometryParticleStateBase& AddFrame(int32 FrameIdx, int8 Wave)
+		{
+			FFrameInfo& Info = Frames[FrameIdx];
+			if(FGeometryParticleStateBase* State = Info.GetState(FrameIdx, Wave))
+			{
+				return *State;
+			}
+
+			return Info.NewState(FrameIdx, Wave);
+		}
 	};
 
-	TArrayAsMap<FUniqueIdx,FParticleRewindInfo> ParticleToRewindInfo;
+	const FDirtyParticleInfo& FindParticleChecked(const FUniqueIdx UniqueIdx) const
+	{
+		const int32 Idx = ParticleToAllDirtyIdx.FindChecked(UniqueIdx);
+		return AllDirtyParticles[Idx];
+	}
+
+	FDirtyParticleInfo& FindParticleChecked(const FUniqueIdx UniqueIdx)
+	{
+		const int32 Idx = ParticleToAllDirtyIdx.FindChecked(UniqueIdx);
+		return AllDirtyParticles[Idx];
+	}
+
+	const FDirtyParticleInfo* FindParticle(const FUniqueIdx UniqueIdx) const
+	{
+		if(const int32* Idx = ParticleToAllDirtyIdx.Find(UniqueIdx))
+		{
+			return &AllDirtyParticles[*Idx];
+		}
+
+		return nullptr;
+	}
+
+	FDirtyParticleInfo* FindParticle(const FUniqueIdx UniqueIdx)
+	{
+		if(const int32* Idx = ParticleToAllDirtyIdx.Find(UniqueIdx))
+		{
+			return &AllDirtyParticles[*Idx];
+		}
+
+		return nullptr;
+	}
+
+	FDirtyParticleInfo& FindOrAddParticle(TGeometryParticle<FReal,3>& UnsafeGTParticle,const FUniqueIdx UniqueIdx)
+	{
+		if(FDirtyParticleInfo* Info = FindParticle(UniqueIdx))
+		{
+			return *Info;
+		}
+
+		const int32 DirtyIdx = AllDirtyParticles.Add(FDirtyParticleInfo(UnsafeGTParticle,UniqueIdx,CurFrame,Managers.Capacity()));
+		ParticleToAllDirtyIdx.Add(UniqueIdx,DirtyIdx);
+
+		return AllDirtyParticles[DirtyIdx];
+	}
+
+	TArrayAsMap<FUniqueIdx,int32> ParticleToAllDirtyIdx;
 	TCircularBuffer<FFrameManagerInfo> Managers;
 	TArray<FDirtyParticleInfo> AllDirtyParticles;
 	int32 CurFrame;
+	int8 CurWave;
 	int32 FramesSaved;
 	int32 DataIdxOffset;
 	bool bNeedsSave;	//Indicates that some data is pointing at head and requires saving before a rewind
