Chaos
- narrow phase policy support
- solvers can specify update deferral and manifold use
- hooked RBAN and Main solvers up with cvars to control them
- cleaned up collision detector class (removed templates etc)

diff --git a/Engine/Source/Runtime/Engine/Private/PhysicsEngine/ImmediatePhysics/ImmediatePhysicsChaos/ImmediatePhysicsSimulation_Chaos.cpp b/Engine/Source/Runtime/Engine/Private/PhysicsEngine/ImmediatePhysics/ImmediatePhysicsChaos/ImmediatePhysicsSimulation_Chaos.cpp
--- a/Engine/Source/Runtime/Engine/Private/PhysicsEngine/ImmediatePhysics/ImmediatePhysicsChaos/ImmediatePhysicsSimulation_Chaos.cpp
+++ b/Engine/Source/Runtime/Engine/Private/PhysicsEngine/ImmediatePhysics/ImmediatePhysicsChaos/ImmediatePhysicsSimulation_Chaos.cpp
@@ -47,12 +47,16 @@ int32 ChaosImmediate_Collision_PushOutPairIterations = -1;
 int32 ChaosImmediate_Collision_Priority = 1;
 float ChaosImmediate_Collision_CullDistance = 1.0f;
 float ChaosImmediate_Collision_ShapePadding = 0;
+int32 ChaosImmediate_Collision_DeferNarrowPhase = 1;
+int32 ChaosImmediate_Collision_UseManifolds = 0;
 FAutoConsoleVariableRef CVarChaosImmPhysCollisionDisable(TEXT("p.Chaos.ImmPhys.Collision.Enabled"), ChaosImmediate_Collision_Enabled, TEXT("Enable/Disable collisions in Immediate Physics."));
 FAutoConsoleVariableRef CVarChaosImmPhysCollisionPairIterations(TEXT("p.Chaos.ImmPhys.Collision.PairIterations"), ChaosImmediate_Collision_PairIterations, TEXT("Override collision pair iterations (if >= 0)"));
 FAutoConsoleVariableRef CVarChaosImmPhysCollisionPushOutPairIterations(TEXT("p.Chaos.ImmPhys.Collision.PushOutPairIterations"), ChaosImmediate_Collision_PushOutPairIterations, TEXT("Override collision push-out pair iterations (if >= 0)"));
 FAutoConsoleVariableRef CVarChaosImmPhysCollisionPriority(TEXT("p.Chaos.ImmPhys.Collision.Priority"), ChaosImmediate_Collision_Priority, TEXT("Set the Collision constraint sort order (Joints have priority 0)"));
 FAutoConsoleVariableRef CVarChaosImmPhysCollisionCullDistance(TEXT("p.Chaos.ImmPhys.Collision.CullDistance"), ChaosImmediate_Collision_CullDistance, TEXT("CullDistance"));
 FAutoConsoleVariableRef CVarChaosImmPhysCollisionShapePadding(TEXT("p.Chaos.ImmPhys.Collision.ShapePadding"), ChaosImmediate_Collision_ShapePadding, TEXT("ShapePadding"));
+FAutoConsoleVariableRef CVarChaosImmPhysCollisionDeferNarrowPhase(TEXT("p.Chaos.ImmPhys.Collision.DeferNarrowPhase"), ChaosImmediate_Collision_DeferNarrowPhase, TEXT("Create contacts for all broadphase pairs, perform NarrowPhase later."));
+FAutoConsoleVariableRef CVarChaosImmPhysCollisionUseManifolds(TEXT("p.Chaos.ImmPhys.Collision.UseManifolds"), ChaosImmediate_Collision_UseManifolds, TEXT("Enable/Disable use of manifoldes in collision."));
 
 int32 ChaosImmediate_Joint_PairIterations = -1;
 int32 ChaosImmediate_Joint_PushOutPairIterations = -1;
@@ -156,7 +160,8 @@ namespace ImmediatePhysics_Chaos
 		, Joints()
 		, Collisions(Particles, CollidedParticles, ParticleMaterials, 0, 0, ChaosImmediate_Collision_CullDistance, ChaosImmediate_Collision_ShapePadding)
 		, BroadPhase(ActivePotentiallyCollidingPairs, ChaosImmediate_Collision_CullDistance)
-		, CollisionDetector(BroadPhase, Collisions)
+		, NarrowPhase()
+		, CollisionDetector(BroadPhase, NarrowPhase, Collisions)
 		, JointsRule(0, Joints)
 		, CollisionsRule(1, Collisions)
 		, Evolution(Particles, CollisionDetector, ChaosImmediate_Evolution_BoundsExtension)
@@ -178,7 +183,9 @@ namespace ImmediatePhysics_Chaos
 		// RBAN collision customization
 		Collisions.DisableHandles();
 		Collisions.SetApplyType(ECollisionApplyType::Position);
-		CollisionDetector.GetContext().bFilteringEnabled = false;
+		NarrowPhase.GetContext().bFilteringEnabled = false;
+		NarrowPhase.GetContext().bDeferUpdate = true;
+		NarrowPhase.GetContext().bAllowManifolds = false;
 
 #if CHAOS_DEBUG_DRAW
 		Evolution.SetPostIntegrateCallback(
@@ -449,7 +456,7 @@ namespace ImmediatePhysics_Chaos
 	void FSimulation::SetSimulationSpaceTransform(const FTransform& Transform)
 	{ 
 		SimulationSpaceTransform = Transform;
-		CollisionDetector.GetContext().SpaceTransform = Transform;
+		NarrowPhase.GetContext().SpaceTransform = Transform;	// @todo(chaos): remove when manifolds are fixed or removed
 	}
 
 
@@ -519,6 +526,9 @@ namespace ImmediatePhysics_Chaos
 			BroadPhase.SetCullDustance(ChaosImmediate_Collision_CullDistance);
 			Evolution.SetBoundsExtension(ChaosImmediate_Evolution_BoundsExtension);
 
+			NarrowPhase.GetContext().bDeferUpdate = (ChaosImmediate_Collision_DeferNarrowPhase != 0);
+			NarrowPhase.GetContext().bAllowManifolds = (ChaosImmediate_Collision_UseManifolds != 0);
+
 			if (ChaosImmediate_UsePositionSolver)
 			{
 				Collisions.SetApplyType(ECollisionApplyType::Position);
diff --git a/Engine/Source/Runtime/Engine/Public/Physics/ImmediatePhysics/ImmediatePhysicsChaos/ImmediatePhysicsSimulation_Chaos.h b/Engine/Source/Runtime/Engine/Public/Physics/ImmediatePhysics/ImmediatePhysicsChaos/ImmediatePhysicsSimulation_Chaos.h
--- a/Engine/Source/Runtime/Engine/Public/Physics/ImmediatePhysics/ImmediatePhysicsChaos/ImmediatePhysicsSimulation_Chaos.h
+++ b/Engine/Source/Runtime/Engine/Public/Physics/ImmediatePhysics/ImmediatePhysicsChaos/ImmediatePhysicsSimulation_Chaos.h
@@ -5,11 +5,10 @@
 #include "Physics/ImmediatePhysics/ImmediatePhysicsChaos/ImmediatePhysicsCore_Chaos.h"
 
 #include "Chaos/ChaosDebugDrawDeclares.h"
-#include "Chaos/Evolution/PBDMinEvolution.h"
-#include "Chaos/Collision/CollisionDetector.h"
-#include "Chaos/Collision/CollisionReceiver.h"
 #include "Chaos/Collision/NarrowPhase.h"
 #include "Chaos/Collision/ParticlePairBroadPhase.h"
+#include "Chaos/Collision/ParticlePairCollisionDetector.h"
+#include "Chaos/Evolution/PBDMinEvolution.h"
 #include "Chaos/PBDCollisionConstraints.h"
 #include "Chaos/PBDConstraintRule.h"
 #include "Chaos/PBDJointConstraints.h"
@@ -81,7 +80,7 @@ namespace ImmediatePhysics_Chaos
 		void DebugDrawConstraints(const int32 MinDebugLevel, const int32 MaxDebugLevel, const float ColorScale);
 
 		using FCollisionConstraints = Chaos::FPBDCollisionConstraints;
-		using FCollisionDetector = Chaos::TCollisionDetector<Chaos::FParticlePairBroadPhase, Chaos::FNarrowPhase, Chaos::FSyncCollisionReceiver, FCollisionConstraints>;
+		using FCollisionDetector = Chaos::FParticlePairCollisionDetector;
 		using FRigidParticleSOAs = Chaos::TPBDRigidsSOAs<FReal, 3>;
 		using FParticleHandle = Chaos::TGeometryParticleHandle<FReal, Dimensions>;
 		using FParticlePair = Chaos::TVector<Chaos::TGeometryParticleHandle<Chaos::FReal, 3>*, 2>;
@@ -96,6 +95,7 @@ namespace ImmediatePhysics_Chaos
 		Chaos::FPBDJointConstraints Joints;
 		FCollisionConstraints Collisions;
 		Chaos::FParticlePairBroadPhase BroadPhase;
+		Chaos::FNarrowPhase NarrowPhase;
 		FCollisionDetector CollisionDetector;
 		Chaos::TSimpleConstraintRule<Chaos::FPBDJointConstraints> JointsRule;
 		Chaos::TSimpleConstraintRule<FCollisionConstraints> CollisionsRule;
diff --git a/Engine/Source/Runtime/Experimental/Chaos/Private/Chaos/Collision/CollisionDetector.cpp b/Engine/Source/Runtime/Experimental/Chaos/Private/Chaos/Collision/CollisionDetector.cpp
--- a/Engine/Source/Runtime/Experimental/Chaos/Private/Chaos/Collision/CollisionDetector.cpp
+++ b/Engine/Source/Runtime/Experimental/Chaos/Private/Chaos/Collision/CollisionDetector.cpp
@@ -11,7 +11,4 @@ namespace Chaos
 #if CHAOS_ENABLE_STAT_NARROWPHASE
 	DEFINE_STAT(STAT_Collisions_NarrowPhase);
 #endif
-
-	template class TCollisionDetector<FParticlePairBroadPhase, FNarrowPhase, FSyncCollisionReceiver, FPBDCollisionConstraints>;
-	template class TCollisionDetector<FSpatialAccelerationBroadPhase, FNarrowPhase, FAsyncCollisionReceiver, FPBDCollisionConstraints>;
 }
\ No newline at end of file
diff --git a/Engine/Source/Runtime/Experimental/Chaos/Private/Chaos/CollisionResolution.cpp b/Engine/Source/Runtime/Experimental/Chaos/Private/Chaos/CollisionResolution.cpp
--- a/Engine/Source/Runtime/Experimental/Chaos/Private/Chaos/CollisionResolution.cpp
+++ b/Engine/Source/Runtime/Experimental/Chaos/Private/Chaos/CollisionResolution.cpp
@@ -43,11 +43,8 @@ FAutoConsoleVariableRef CVarCCDAllowedDepthBoundsScale(TEXT("p.Chaos.CCD.Allowed
 #define PHI_RESAMPLE_THRESHOLD 0.001
 
 
-bool Chaos_Collision_UseManifolds = false;
-FAutoConsoleVariableRef CVarChaosCollisionUseManifolds(TEXT("p.Chaos.Collision.UseManifolds"), Chaos_Collision_UseManifolds, TEXT("Enable/Disable use of manifoldes in collision."));
-
-bool Chaos_Collision_UseManifolds_Test = false;
-FAutoConsoleVariableRef CVarChaosCollisionUseManifoldsTest(TEXT("p.Chaos.Collision.UseManifoldsTest"), Chaos_Collision_UseManifolds_Test, TEXT("Enable/Disable use of manifoldes in collision."));
+bool bChaos_Collision_UseManifolds_Test = false;
+FAutoConsoleVariableRef CVarChaosCollisionUseManifoldsTest(TEXT("p.Chaos.Collision.UseManifoldsTest"), bChaos_Collision_UseManifolds_Test, TEXT("Enable/Disable use of manifoldes in collision."));
 
 float Chaos_Collision_ManifoldFaceAngle = 5.0f;
 float Chaos_Collision_ManifoldFaceEpsilon = FMath::Sin(FMath::DegreesToRadians(Chaos_Collision_ManifoldFaceAngle));
@@ -59,13 +56,11 @@ float Chaos_Collision_CapsuleBoxManifoldTolerance = FMath::Sin(FMath::DegreesToR
 FConsoleVariableDelegate Chaos_Collision_CapsuleBoxManifoldDelegate = FConsoleVariableDelegate::CreateLambda([](IConsoleVariable* CVar) { Chaos_Collision_CapsuleBoxManifoldTolerance = FMath::Sin(FMath::DegreesToRadians(Chaos_Collision_CapsuleBoxManifoldAngle)); });
 FAutoConsoleVariableRef CVarChaosCollisionBoxCapsuleManifoldAngle(TEXT("p.Chaos.Collision.CapsuleBoxManifoldAngle"), Chaos_Collision_CapsuleBoxManifoldAngle, TEXT("If a capsule is more than this angle from vertical, do not use a manifold"), Chaos_Collision_CapsuleBoxManifoldDelegate);
 
-bool bChaos_Collision_DeferNarrowPhase = true;
-FAutoConsoleVariableRef CVarChaosCollisionDeferNarrowPhase(TEXT("p.Chaos.Collision.DeferNarrowPhase"), bChaos_Collision_DeferNarrowPhase, TEXT("Create contacts for all broadphase pairs, perform NarrowPhase later."));
-
 namespace Chaos
 {
 	namespace Collisions
 	{
+		// Data returned by the low-level collision functions
 		struct FContactPoint
 		{
 			FVec3 Normal;
@@ -77,8 +72,26 @@ namespace Chaos
 			, Phi(TNumericLimits<FReal>::Max()) {}
 		};
 
-		// Determines if body should use CCD. If using CCD, computes Dir and Length of sweep.
+		// Traits to control how contacts are generated
+		template<bool B_IMMEDIATEUPDATE, bool B_ALLOWMANIFOLD>
+		struct TConstructCollisionTraits
+		{
+			// If true, ConstructConstraints also initializes the constraint Phi and Normal based on current state, and
+			// contacts beyond CullDistance are culled. If false, a constraint is created for every shape pair.
+			// NOTE: Contact Phi and Normal are also calculated at the beginning of each iteration as well. The reason you may
+			// want bImmediateUpdate=true is to reduce the number of contacts that make it to the constraint graph (because
+			// we have culled those with a separation greater than CullDistance).
+			// However, early culling can cause collisions to be missed, e.g., if a joint moves a body into a space where
+			// a culled collision should now be active. This can (does) lead to jitter in some cases.
+			static const bool bImmediateUpdate = B_IMMEDIATEUPDATE;
+
+			// If true, use contact manifolds where supported. A manifold defines simplified collision between a pair of shapes:
+			// a plane on one shape and some points on the other. This makes updating the contact each iteration very fast, but
+			// the manifold may not be a good representation if either body moves or rotates too far.
+			static const bool bAllowManifold = B_ALLOWMANIFOLD;
+		};
 
+		// Determines if body should use CCD. If using CCD, computes Dir and Length of sweep.
 		bool UseCCD(const TGeometryParticleHandle<FReal, 3>* SweptParticle, const TGeometryParticleHandle<FReal, 3>* OtherParticle, const FImplicitObject* Implicit, FVec3& Dir, FReal& Length)
 		{
 			if (OtherParticle->ObjectState() != EObjectStateType::Static)
@@ -556,6 +569,7 @@ namespace Chaos
 			//   Stub function for updating the manifold prior to the Apply and ApplyPushOut
 		}
 
+		template<typename T_TRAITS>
 		void ConstructBoxBoxConstraints(TGeometryParticleHandle<FReal, 3>* Particle0, TGeometryParticleHandle<FReal, 3>* Particle1, const FImplicitObject* Implicit0, const FImplicitObject* Implicit1, const FRigidTransform3& LocalTransform0, const FRigidTransform3& LocalTransform1, const FReal CullDistance, FCollisionConstraintsArray& NewConstraints)
 		{
 			const TBox<FReal, 3> * Object0 = Implicit0->template GetObject<const TBox<FReal, 3> >();
@@ -563,7 +577,7 @@ namespace Chaos
 			if (ensure(Object0 && Object1))
 			{
 				FRigidBodyPointContactConstraint Constraint = FRigidBodyPointContactConstraint(Particle0, Implicit0, LocalTransform0, Particle1, Implicit1, LocalTransform1, EContactShapesType::BoxBox);
-				if (!bChaos_Collision_DeferNarrowPhase)
+				if (T_TRAITS::bImmediateUpdate)
 				{
 					FRigidTransform3 WorldTransform0 = LocalTransform0 * Collisions::GetTransform(Particle0);
 					FRigidTransform3 WorldTransform1 = LocalTransform1 * Collisions::GetTransform(Particle1);
@@ -599,6 +613,7 @@ namespace Chaos
 			//   Stub function for updating the manifold prior to the Apply and ApplyPushOut
 		}
 
+		template<typename T_TRAITS>
 		void ConstructBoxHeightFieldConstraints(TGeometryParticleHandle<FReal, 3>* Particle0, TGeometryParticleHandle<FReal, 3>* Particle1, const FImplicitObject* Implicit0, const FImplicitObject* Implicit1, const FRigidTransform3& LocalTransform0, const FRigidTransform3& LocalTransform1, const FReal CullDistance, FCollisionConstraintsArray& NewConstraints)
 		{
 
@@ -607,7 +622,7 @@ namespace Chaos
 			if (ensure(Object0 && Object1))
 			{
 				FRigidBodyPointContactConstraint Constraint = FRigidBodyPointContactConstraint(Particle0, Implicit0, LocalTransform0, Particle1, Implicit1, LocalTransform1, EContactShapesType::BoxHeightField);
-				if (!bChaos_Collision_DeferNarrowPhase)
+				if (T_TRAITS::bImmediateUpdate)
 				{
 					FRigidTransform3 WorldTransform0 = LocalTransform0 * Collisions::GetTransform(Particle0);
 					FRigidTransform3 WorldTransform1 = LocalTransform1 * Collisions::GetTransform(Particle1);
@@ -691,6 +706,7 @@ namespace Chaos
 		}
 
 
+		template<typename T_TRAITS>
 		void ConstructBoxPlaneConstraints(TGeometryParticleHandle<FReal, 3>* Particle0, TGeometryParticleHandle<FReal, 3>* Particle1, const FImplicitObject* Implicit0, const FImplicitObject* Implicit1, const FRigidTransform3& LocalTransform0, const FRigidTransform3& LocalTransform1, const FReal CullDistance, FCollisionConstraintsArray& NewConstraints)
 		{
 
@@ -699,7 +715,7 @@ namespace Chaos
 			if (ensure(Object0 && Object1))
 			{
 				FRigidBodyPointContactConstraint Constraint = FRigidBodyPointContactConstraint(Particle0, Implicit0, LocalTransform0, Particle1, Implicit1, LocalTransform1, EContactShapesType::BoxPlane);
-				if (!bChaos_Collision_DeferNarrowPhase)
+				if (T_TRAITS::bImmediateUpdate)
 				{
 					FRigidTransform3 WorldTransform0 = LocalTransform0 * Collisions::GetTransform(Particle0);
 					FRigidTransform3 WorldTransform1 = LocalTransform1 * Collisions::GetTransform(Particle1);
@@ -736,6 +752,7 @@ namespace Chaos
 
 		}
 
+		template<typename T_TRAITS>
 		void ConstructBoxTriangleMeshConstraints(TGeometryParticleHandle<FReal, 3>* Particle0, TGeometryParticleHandle<FReal, 3>* Particle1, const FImplicitObject* Implicit0, const FImplicitObject* Implicit1, const FRigidTransform3& LocalTransform0, const FRigidTransform3& LocalTransform1, const FReal CullDistance, FCollisionConstraintsArray& NewConstraints)		{
 			const TBox<FReal, 3> * Object0 = Implicit0->template GetObject<const TBox<FReal, 3> >();
 			if (ensure(Object0))
@@ -743,7 +760,7 @@ namespace Chaos
 				if (const TImplicitObjectScaled<FTriangleMeshImplicitObject>* ScaledTriangleMesh = Implicit1->template GetObject<const TImplicitObjectScaled<FTriangleMeshImplicitObject>>())
 				{
 					FRigidBodyPointContactConstraint Constraint = FRigidBodyPointContactConstraint(Particle0, Implicit0, LocalTransform0, Particle1, Implicit1, LocalTransform1, EContactShapesType::BoxTriMesh);
-					if (!bChaos_Collision_DeferNarrowPhase)
+					if (T_TRAITS::bImmediateUpdate)
 					{
 						FRigidTransform3 WorldTransform0 = LocalTransform0 * Collisions::GetTransform(Particle0);
 						FRigidTransform3 WorldTransform1 = LocalTransform1 * Collisions::GetTransform(Particle1);
@@ -758,7 +775,7 @@ namespace Chaos
 				else if(const FTriangleMeshImplicitObject* TriangleMesh = Implicit1->template GetObject<const FTriangleMeshImplicitObject>())
 				{
 					FRigidBodyPointContactConstraint Constraint = FRigidBodyPointContactConstraint(Particle0, Implicit0, LocalTransform0, Particle1, Implicit1, LocalTransform1, EContactShapesType::BoxTriMesh);
-					if (!bChaos_Collision_DeferNarrowPhase)
+					if (T_TRAITS::bImmediateUpdate)
 					{
 						FRigidTransform3 WorldTransform0 = LocalTransform0 * Collisions::GetTransform(Particle0);
 						FRigidTransform3 WorldTransform1 = LocalTransform1 * Collisions::GetTransform(Particle1);
@@ -810,6 +827,7 @@ namespace Chaos
 			//   Stub function for updating the manifold prior to the Apply and ApplyPushOut
 		}
 
+		template<typename T_TRAITS>
 		void ConstructSphereSphereConstraints(TGeometryParticleHandle<FReal, 3>* Particle0, TGeometryParticleHandle<FReal, 3>* Particle1, const FImplicitObject* Implicit0, const FImplicitObject* Implicit1, const FRigidTransform3& LocalTransform0, const FRigidTransform3& LocalTransform1, const FReal CullDistance, FCollisionConstraintsArray& NewConstraints)
 		{
 			const TSphere<FReal, 3>* Object0 = Implicit0->template GetObject<const TSphere<FReal, 3> >();
@@ -817,7 +835,7 @@ namespace Chaos
 			if (ensure(Object0 && Object1))
 			{
 				FRigidBodyPointContactConstraint Constraint = FRigidBodyPointContactConstraint(Particle0, Implicit0, LocalTransform0, Particle1, Implicit1, LocalTransform1, EContactShapesType::SphereSphere);
-				if (!bChaos_Collision_DeferNarrowPhase)
+				if (T_TRAITS::bImmediateUpdate)
 				{
 					FRigidTransform3 WorldTransform0 = LocalTransform0 * Collisions::GetTransform(Particle0);
 					FRigidTransform3 WorldTransform1 = LocalTransform1 * Collisions::GetTransform(Particle1);
@@ -852,6 +870,7 @@ namespace Chaos
 			//   Stub function for updating the manifold prior to the Apply and ApplyPushOut
 		}
 
+		template<typename T_TRAITS>
 		void ConstructSphereHeightFieldConstraints(TGeometryParticleHandle<FReal, 3>* Particle0, TGeometryParticleHandle<FReal, 3>* Particle1, const FImplicitObject* Implicit0, const FImplicitObject* Implicit1, const FRigidTransform3& LocalTransform0, const FRigidTransform3& LocalTransform1, const FReal CullDistance, FCollisionConstraintsArray& NewConstraints)
 		{
 
@@ -860,7 +879,7 @@ namespace Chaos
 			if (ensure(Object0 && Object1))
 			{
 				FRigidBodyPointContactConstraint Constraint = FRigidBodyPointContactConstraint(Particle0, Implicit0, LocalTransform0, Particle1, Implicit1, LocalTransform1, EContactShapesType::SphereHeightField);
-				if (!bChaos_Collision_DeferNarrowPhase)
+				if (T_TRAITS::bImmediateUpdate)
 				{
 					FRigidTransform3 WorldTransform0 = LocalTransform0 * Collisions::GetTransform(Particle0);
 					FRigidTransform3 WorldTransform1 = LocalTransform1 * Collisions::GetTransform(Particle1);
@@ -903,6 +922,7 @@ namespace Chaos
 			//   Stub function for updating the manifold prior to the Apply and ApplyPushOut
 		}
 
+		template<typename T_TRAITS>
 		void ConstructSpherePlaneConstraints(TGeometryParticleHandle<FReal, 3>* Particle0, TGeometryParticleHandle<FReal, 3>* Particle1, const FImplicitObject* Implicit0, const FImplicitObject* Implicit1, const FRigidTransform3& LocalTransform0, const FRigidTransform3& LocalTransform1, const FReal CullDistance, FCollisionConstraintsArray& NewConstraints)
 		{
 
@@ -911,7 +931,7 @@ namespace Chaos
 			if (ensure(Object0 && Object1))
 			{
 				FRigidBodyPointContactConstraint Constraint = FRigidBodyPointContactConstraint(Particle0, Implicit0, LocalTransform0, Particle1, Implicit1, LocalTransform1, EContactShapesType::SpherePlane);
-				if (!bChaos_Collision_DeferNarrowPhase)
+				if (T_TRAITS::bImmediateUpdate)
 				{
 					FRigidTransform3 WorldTransform0 = LocalTransform0 * Collisions::GetTransform(Particle0);
 					FRigidTransform3 WorldTransform1 = LocalTransform1 * Collisions::GetTransform(Particle1);
@@ -959,6 +979,7 @@ namespace Chaos
 			//   Stub function for updating the manifold prior to the Apply and ApplyPushOut
 		}
 
+		template<typename T_TRAITS>
 		void ConstructSphereBoxConstraints(TGeometryParticleHandle<FReal, 3>* Particle0, TGeometryParticleHandle<FReal, 3>* Particle1, const FImplicitObject* Implicit0, const FImplicitObject* Implicit1, const FRigidTransform3& LocalTransform0, const FRigidTransform3& LocalTransform1, const FReal CullDistance, FCollisionConstraintsArray& NewConstraints)
 		{
 
@@ -967,7 +988,7 @@ namespace Chaos
 			if (ensure(Object0 && Object1))
 			{
 				FRigidBodyPointContactConstraint Constraint = FRigidBodyPointContactConstraint(Particle0, Implicit0, LocalTransform0, Particle1, Implicit1, LocalTransform1, EContactShapesType::SphereBox);
-				if (!bChaos_Collision_DeferNarrowPhase)
+				if (T_TRAITS::bImmediateUpdate)
 				{
 					FRigidTransform3 WorldTransform0 = LocalTransform0 * Collisions::GetTransform(Particle0);
 					FRigidTransform3 WorldTransform1 = LocalTransform1 * Collisions::GetTransform(Particle1);
@@ -1021,6 +1042,7 @@ namespace Chaos
 			//   Stub function for updating the manifold prior to the Apply and ApplyPushOut
 		}
 
+		template<typename T_TRAITS>
 		void ConstructSphereCapsuleConstraints(TGeometryParticleHandle<FReal, 3>* Particle0, TGeometryParticleHandle<FReal, 3>* Particle1, const FImplicitObject* Implicit0, const FImplicitObject* Implicit1, const FRigidTransform3& LocalTransform0, const FRigidTransform3& LocalTransform1, const FReal CullDistance, FCollisionConstraintsArray& NewConstraints)
 		{
 			const TSphere<FReal, 3>* Object0 = Implicit0->template GetObject<const TSphere<FReal, 3> >();
@@ -1028,7 +1050,7 @@ namespace Chaos
 			if (ensure(Object0 && Object1))
 			{
 				FRigidBodyPointContactConstraint Constraint = FRigidBodyPointContactConstraint(Particle0, Implicit0, LocalTransform0, Particle1, Implicit1, LocalTransform1, EContactShapesType::SphereCapsule);
-				if (!bChaos_Collision_DeferNarrowPhase)
+				if (T_TRAITS::bImmediateUpdate)
 				{
 					FRigidTransform3 WorldTransform0 = LocalTransform0 * Collisions::GetTransform(Particle0);
 					FRigidTransform3 WorldTransform1 = LocalTransform1 * Collisions::GetTransform(Particle1);
@@ -1065,6 +1087,7 @@ namespace Chaos
 
 		}
 
+		template<typename T_TRAITS>
 		void ConstructSphereTriangleMeshConstraints(TGeometryParticleHandle<FReal, 3>* Particle0, TGeometryParticleHandle<FReal, 3>* Particle1, const FImplicitObject* Implicit0, const FImplicitObject* Implicit1, const FRigidTransform3& LocalTransform0, const FRigidTransform3& LocalTransform1, const FReal CullDistance, FCollisionConstraintsArray& NewConstraints)
 		{
 			const TSphere<FReal, 3>* Object0 = Implicit0->template GetObject<const TSphere<FReal, 3> >();
@@ -1073,7 +1096,7 @@ namespace Chaos
 				if (const TImplicitObjectScaled<FTriangleMeshImplicitObject>* ScaledTriangleMesh = Implicit1->template GetObject<const TImplicitObjectScaled<FTriangleMeshImplicitObject>>())
 				{
 					FRigidBodyPointContactConstraint Constraint = FRigidBodyPointContactConstraint(Particle0, Implicit0, LocalTransform0, Particle1, Implicit1, LocalTransform1, EContactShapesType::SphereTriMesh);
-					if (!bChaos_Collision_DeferNarrowPhase)
+					if (T_TRAITS::bImmediateUpdate)
 					{
 						FRigidTransform3 WorldTransform0 = LocalTransform0 * Collisions::GetTransform(Particle0);
 						FRigidTransform3 WorldTransform1 = LocalTransform1 * Collisions::GetTransform(Particle1);
@@ -1088,7 +1111,7 @@ namespace Chaos
 				else if (const FTriangleMeshImplicitObject* TriangleMesh = Implicit1->template GetObject<const FTriangleMeshImplicitObject>())
 				{
 					FRigidBodyPointContactConstraint Constraint = FRigidBodyPointContactConstraint(Particle0, Implicit0, LocalTransform0, Particle1, Implicit1, LocalTransform1, EContactShapesType::SphereTriMesh);
-					if (!bChaos_Collision_DeferNarrowPhase)
+					if (T_TRAITS::bImmediateUpdate)
 					{
 						FRigidTransform3 WorldTransform0 = LocalTransform0 * Collisions::GetTransform(Particle0);
 						FRigidTransform3 WorldTransform1 = LocalTransform1 * Collisions::GetTransform(Particle1);
@@ -1151,6 +1174,7 @@ namespace Chaos
 			//   Stub function for updating the manifold prior to the Apply and ApplyPushOut
 		}
 
+		template<typename T_TRAITS>
 		void ConstructCapsuleCapsuleConstraints(TGeometryParticleHandle<FReal, 3>* Particle0, TGeometryParticleHandle<FReal, 3>* Particle1, const FImplicitObject* Implicit0, const FImplicitObject* Implicit1, const FRigidTransform3& LocalTransform0, const FRigidTransform3& LocalTransform1, const FReal CullDistance, FCollisionConstraintsArray& NewConstraints)
 		{
 			const TCapsule<FReal> * Object0 = Implicit0->template GetObject<const TCapsule<FReal> >();
@@ -1158,7 +1182,7 @@ namespace Chaos
 			if (ensure(Object0 && Object1))
 			{
 				FRigidBodyPointContactConstraint Constraint = FRigidBodyPointContactConstraint(Particle0, Implicit0, LocalTransform0, Particle1, Implicit1, LocalTransform1, EContactShapesType::CapsuleCapsule);
-				if (!bChaos_Collision_DeferNarrowPhase)
+				if (T_TRAITS::bImmediateUpdate)
 				{
 					FRigidTransform3 WorldTransform0 = LocalTransform0 * Collisions::GetTransform(Particle0);
 					FRigidTransform3 WorldTransform1 = LocalTransform1 * Collisions::GetTransform(Particle1);
@@ -1395,47 +1419,50 @@ namespace Chaos
 			ensure(false);
 		}
 
+		template<typename T_TRAITS>
 		void ConstructCapsuleBoxConstraints(TGeometryParticleHandle<FReal, 3>* Particle0, TGeometryParticleHandle<FReal, 3>* Particle1, const FImplicitObject* Implicit0, const FImplicitObject* Implicit1, const FRigidTransform3& LocalTransform0, const FRigidTransform3& LocalTransform1, const FReal CullDistance, const FCollisionContext& Context, FCollisionConstraintsArray& NewConstraints)
 		{
 			const TCapsule<FReal> * Object0 = Implicit0->template GetObject<const TCapsule<FReal>>();
 			const TBox<FReal, 3> * Object1 = Implicit1->template GetObject<const TBox<FReal, 3>>();
 			if (ensure(Object0 && Object1))
 			{
-				bool bAllowManifold = Chaos_Collision_UseManifolds;
-				if (bAllowManifold && (Chaos_Collision_CapsuleBoxManifoldTolerance > KINDA_SMALL_NUMBER))
+				if (T_TRAITS::bAllowManifold)
 				{
-					// @todo(ccaulfield): weak sauce - fix capsule-box manifolds.
-					// HACK: Disable manifolds for "horizontal" capsules. Manifolds don't work well when joints are pulling boxes down
-					// (under gravity) when the upper boxes are draped over a horizontal capsule. The box rotations about the manifold
-					// points(line) is too great and we end up with jitter.
-					FRigidTransform3 WorldTransform0 = LocalTransform0 * Collisions::GetTransform(Particle0);
-					const FVector CapsuleAxis = Context.SpaceTransform.TransformVectorNoScale(WorldTransform0.TransformVectorNoScale(Object0->GetAxis()));
-					bAllowManifold = (FMath::Abs(CapsuleAxis.Z) > Chaos_Collision_CapsuleBoxManifoldTolerance);
+					bool bAllowManifold = true;
+					if (Chaos_Collision_CapsuleBoxManifoldTolerance > KINDA_SMALL_NUMBER)
+					{
+						// @todo(ccaulfield): weak sauce - fix capsule-box manifolds.
+						// HACK: Disable manifolds for "horizontal" capsules. Manifolds don't work well when joints are pulling boxes down
+						// (under gravity) when the upper boxes are draped over a horizontal capsule. The box rotations about the manifold
+						// points(line) is too great and we end up with jitter.
+						FRigidTransform3 WorldTransform0 = LocalTransform0 * Collisions::GetTransform(Particle0);
+						const FVector CapsuleAxis = Context.SpaceTransform.TransformVectorNoScale(WorldTransform0.TransformVectorNoScale(Object0->GetAxis()));
+						bAllowManifold = (FMath::Abs(CapsuleAxis.Z) > Chaos_Collision_CapsuleBoxManifoldTolerance);
+					}
+
+					if (bAllowManifold)
+					{
+						FRigidBodyMultiPointContactConstraint Constraint = FRigidBodyMultiPointContactConstraint(Particle0, Implicit0, LocalTransform0, Particle1, Implicit1, LocalTransform1, EContactShapesType::CapsuleBox);
+						// @todo(ccaulfield): manifold creation should be deferrable same as non-manifold below (only Context is preventing this - Apply does not know about the context yet)
+						FRigidTransform3 WorldTransform0 = LocalTransform0 * Collisions::GetTransform(Particle0);
+						FRigidTransform3 WorldTransform1 = LocalTransform1 * Collisions::GetTransform(Particle1);
+						UpdateCapsuleBoxManifold(*Object0, WorldTransform0, Object1->BoundingBox(), WorldTransform1, CullDistance, Context, Constraint);
+						NewConstraints.TryAdd(CullDistance, Constraint);
+						return;
+					}
 				}
 
-				if (bAllowManifold)
+				FRigidBodyPointContactConstraint Constraint = FRigidBodyPointContactConstraint(Particle0, Implicit0, LocalTransform0, Particle1, Implicit1, LocalTransform1, EContactShapesType::CapsuleBox);
+				if (T_TRAITS::bImmediateUpdate)
 				{
-					FRigidBodyMultiPointContactConstraint Constraint = FRigidBodyMultiPointContactConstraint(Particle0, Implicit0, LocalTransform0, Particle1, Implicit1, LocalTransform1, EContactShapesType::CapsuleBox);
-					// @todo(ccaulfield): manifold creation should be deferrable same as non-manifold below (only Context is preventing this - Apply does not know about the context yet)
 					FRigidTransform3 WorldTransform0 = LocalTransform0 * Collisions::GetTransform(Particle0);
 					FRigidTransform3 WorldTransform1 = LocalTransform1 * Collisions::GetTransform(Particle1);
-					UpdateCapsuleBoxManifold(*Object0, WorldTransform0, Object1->BoundingBox(), WorldTransform1, CullDistance, Context, Constraint);
+					UpdateCapsuleBoxConstraint(*Object0, WorldTransform0, Object1->BoundingBox(), WorldTransform1, CullDistance, Constraint);
 					NewConstraints.TryAdd(CullDistance, Constraint);
 				}
 				else
 				{
-					FRigidBodyPointContactConstraint Constraint = FRigidBodyPointContactConstraint(Particle0, Implicit0, LocalTransform0, Particle1, Implicit1, LocalTransform1, EContactShapesType::CapsuleBox);
-					if (!bChaos_Collision_DeferNarrowPhase)
-					{
-						FRigidTransform3 WorldTransform0 = LocalTransform0 * Collisions::GetTransform(Particle0);
-						FRigidTransform3 WorldTransform1 = LocalTransform1 * Collisions::GetTransform(Particle1);
-						UpdateCapsuleBoxConstraint(*Object0, WorldTransform0, Object1->BoundingBox(), WorldTransform1, CullDistance, Constraint);
-						NewConstraints.TryAdd(CullDistance, Constraint);
-					}
-					else
-					{
-						NewConstraints.Add(Constraint);
-					}
+					NewConstraints.Add(Constraint);
 				}
 			}
 		}
@@ -1470,6 +1497,7 @@ namespace Chaos
 			//   Stub function for updating the manifold prior to the Apply and ApplyPushOut
 		}
 
+		template<typename T_TRAITS>
 		void ConstructCapsuleHeightFieldConstraints(TGeometryParticleHandle<FReal, 3>* Particle0, TGeometryParticleHandle<FReal, 3>* Particle1, const FImplicitObject* Implicit0, const FImplicitObject* Implicit1, const FRigidTransform3& LocalTransform0, const FRigidTransform3& LocalTransform1, const FReal CullDistance, FCollisionConstraintsArray& NewConstraints)
 		{
 
@@ -1478,7 +1506,7 @@ namespace Chaos
 			if (ensure(Object0 && Object1))
 			{
 				FRigidBodyPointContactConstraint Constraint = FRigidBodyPointContactConstraint(Particle0, Implicit0, LocalTransform0, Particle1, Implicit1, LocalTransform1, EContactShapesType::CapsuleHeightField);
-				if (!bChaos_Collision_DeferNarrowPhase)
+				if (T_TRAITS::bImmediateUpdate)
 				{
 					FRigidTransform3 WorldTransform0 = LocalTransform0 * Collisions::GetTransform(Particle0);
 					FRigidTransform3 WorldTransform1 = LocalTransform1 * Collisions::GetTransform(Particle1);
@@ -1554,6 +1582,7 @@ namespace Chaos
 
 		}
 
+		template<typename T_TRAITS>
 		void ConstructCapsuleTriangleMeshConstraints(TGeometryParticleHandle<FReal, 3>* Particle0, TGeometryParticleHandle<FReal, 3>* Particle1, const FImplicitObject* Implicit0, const FImplicitObject* Implicit1, const FRigidTransform3& LocalTransform0, const FRigidTransform3& LocalTransform1, const FReal CullDistance, FCollisionConstraintsArray& NewConstraints)
 		{
 			const TCapsule<FReal> * Object0 = Implicit0->template GetObject<const TCapsule<FReal> >();
@@ -1562,7 +1591,7 @@ namespace Chaos
 				if (const TImplicitObjectScaled<FTriangleMeshImplicitObject>* ScaledTriangleMesh = Implicit1->template GetObject<const TImplicitObjectScaled<FTriangleMeshImplicitObject>>())
 				{
 					FRigidBodyPointContactConstraint Constraint = FRigidBodyPointContactConstraint(Particle0, Implicit0, LocalTransform0, Particle1, Implicit1, LocalTransform1, EContactShapesType::CapsuleTriMesh);
-					if (!bChaos_Collision_DeferNarrowPhase)
+					if (T_TRAITS::bImmediateUpdate)
 					{
 						FRigidTransform3 WorldTransform0 = LocalTransform0 * Collisions::GetTransform(Particle0);
 						FRigidTransform3 WorldTransform1 = LocalTransform1 * Collisions::GetTransform(Particle1);
@@ -1577,7 +1606,7 @@ namespace Chaos
 				else if (const FTriangleMeshImplicitObject* TriangleMesh = Implicit1->template GetObject<const FTriangleMeshImplicitObject>())
 				{
 					FRigidBodyPointContactConstraint Constraint = FRigidBodyPointContactConstraint(Particle0, Implicit0, LocalTransform0, Particle1, Implicit1, LocalTransform1, EContactShapesType::CapsuleTriMesh);
-					if (!bChaos_Collision_DeferNarrowPhase)
+					if (T_TRAITS::bImmediateUpdate)
 					{
 						FRigidTransform3 WorldTransform0 = LocalTransform0 * Collisions::GetTransform(Particle0);
 						FRigidTransform3 WorldTransform1 = LocalTransform1 * Collisions::GetTransform(Particle1);
@@ -1693,43 +1722,46 @@ namespace Chaos
 		}
 
 
+		template<typename T_TRAITS>
 		void ConstructConvexConvexConstraints(TGeometryParticleHandle<FReal, 3>* Particle0, TGeometryParticleHandle<FReal, 3>* Particle1, const FImplicitObject* Implicit0, const FImplicitObject* Implicit1, const FRigidTransform3& LocalTransform0, const FRigidTransform3& LocalTransform1, const FReal CullDistance, FCollisionConstraintsArray& NewConstraints)
 		{
 			EImplicitObjectType Implicit0Type = Particle0->Geometry()->GetType();
 			EImplicitObjectType Implicit1Type = Particle1->Geometry()->GetType();
 
-			// Note: This TBox check is a temporary workaround to avoid jitter in cases of Box vs Convex; investigation ongoing
-			// We need to improve iterative manifolds for this case
-			if (Chaos_Collision_UseManifolds && Implicit0Type != TBox<FReal, 3>::StaticType() && Implicit1Type != TBox<FReal, 3>::StaticType())
+			if (T_TRAITS::bAllowManifold)
 			{
-				FRigidBodyMultiPointContactConstraint Constraint = FRigidBodyMultiPointContactConstraint(Particle0, Implicit0, LocalTransform0, Particle1, Implicit1, LocalTransform1, EContactShapesType::ConvexConvex);
-				FRigidTransform3 WorldTransform0 = LocalTransform0 * Collisions::GetTransform(Particle0);
-				FRigidTransform3 WorldTransform1 = LocalTransform1 * Collisions::GetTransform(Particle1);
-				UpdateConvexConvexManifold(Constraint, WorldTransform0, WorldTransform1, CullDistance);
-				if (!bChaos_Collision_DeferNarrowPhase)
-				{
-					UpdateConvexConvexConstraint(*Implicit0, WorldTransform0, *Implicit1, WorldTransform1, CullDistance, Constraint);
-					NewConstraints.TryAdd(CullDistance, Constraint);
-				}
-				else
+				// Note: This TBox check is a temporary workaround to avoid jitter in cases of Box vs Convex; investigation ongoing
+				// We need to improve iterative manifolds for this case
+				if (Implicit0Type != TBox<FReal, 3>::StaticType() && Implicit1Type != TBox<FReal, 3>::StaticType())
 				{
-					NewConstraints.Add(Constraint);
+					FRigidBodyMultiPointContactConstraint Constraint = FRigidBodyMultiPointContactConstraint(Particle0, Implicit0, LocalTransform0, Particle1, Implicit1, LocalTransform1, EContactShapesType::ConvexConvex);
+					FRigidTransform3 WorldTransform0 = LocalTransform0 * Collisions::GetTransform(Particle0);
+					FRigidTransform3 WorldTransform1 = LocalTransform1 * Collisions::GetTransform(Particle1);
+					UpdateConvexConvexManifold(Constraint, WorldTransform0, WorldTransform1, CullDistance);
+					if (T_TRAITS::bImmediateUpdate)
+					{
+						UpdateConvexConvexConstraint(*Implicit0, WorldTransform0, *Implicit1, WorldTransform1, CullDistance, Constraint);
+						NewConstraints.TryAdd(CullDistance, Constraint);
+					}
+					else
+					{
+						NewConstraints.Add(Constraint);
+					}
+					return;
 				}
 			}
+
+			FRigidBodyPointContactConstraint Constraint = FRigidBodyPointContactConstraint(Particle0, Implicit0, LocalTransform0, Particle1, Implicit1, LocalTransform1, EContactShapesType::ConvexConvex);
+			if (T_TRAITS::bImmediateUpdate)
+			{
+				FRigidTransform3 WorldTransform0 = LocalTransform0 * Collisions::GetTransform(Particle0);
+				FRigidTransform3 WorldTransform1 = LocalTransform1 * Collisions::GetTransform(Particle1);
+				UpdateConvexConvexConstraint(*Implicit0, WorldTransform0, *Implicit1, WorldTransform1, CullDistance, Constraint);
+				NewConstraints.TryAdd(CullDistance, Constraint);
+			}
 			else
 			{
-				FRigidBodyPointContactConstraint Constraint = FRigidBodyPointContactConstraint(Particle0, Implicit0, LocalTransform0, Particle1, Implicit1, LocalTransform1, EContactShapesType::ConvexConvex);
-				if (!bChaos_Collision_DeferNarrowPhase)
-				{
-					FRigidTransform3 WorldTransform0 = LocalTransform0 * Collisions::GetTransform(Particle0);
-					FRigidTransform3 WorldTransform1 = LocalTransform1 * Collisions::GetTransform(Particle1);
-					UpdateConvexConvexConstraint(*Implicit0, WorldTransform0, *Implicit1, WorldTransform1, CullDistance, Constraint);
-					NewConstraints.TryAdd(CullDistance, Constraint);
-				}
-				else
-				{
-					NewConstraints.Add(Constraint);
-				}
+				NewConstraints.Add(Constraint);
 			}
 		}
 
@@ -1772,13 +1804,14 @@ namespace Chaos
 			//   Stub function for updating the manifold prior to the Apply and ApplyPushOut
 		}
 
+		template<typename T_TRAITS>
 		void ConstructConvexHeightFieldConstraints(TGeometryParticleHandle<FReal, 3>* Particle0, TGeometryParticleHandle<FReal, 3>* Particle1, const FImplicitObject* Implicit0, const FImplicitObject* Implicit1, const FRigidTransform3& LocalTransform0, const FRigidTransform3& LocalTransform1, const FReal CullDistance, FCollisionConstraintsArray& NewConstraints)
 		{
 			const THeightField<FReal> * Object1 = Implicit1->template GetObject<const THeightField<FReal> >();
 			if (ensure(Implicit0->IsConvex() && Object1))
 			{
 				FRigidBodyPointContactConstraint Constraint = FRigidBodyPointContactConstraint(Particle0, Implicit0, LocalTransform0, Particle1, Implicit1, LocalTransform1, EContactShapesType::ConvexHeightField);
-				if (!bChaos_Collision_DeferNarrowPhase)
+				if (T_TRAITS::bImmediateUpdate)
 				{
 					FRigidTransform3 WorldTransform0 = LocalTransform0 * Collisions::GetTransform(Particle0);
 					FRigidTransform3 WorldTransform1 = LocalTransform1 * Collisions::GetTransform(Particle1);
@@ -1852,6 +1885,7 @@ namespace Chaos
 
 		}
 
+		template<typename T_TRAITS>
 		void ConstructConvexTriangleMeshConstraints(TGeometryParticleHandle<FReal, 3>* Particle0, TGeometryParticleHandle<FReal, 3>* Particle1, const FImplicitObject* Implicit0, const FImplicitObject* Implicit1, const FRigidTransform3& LocalTransform0, const FRigidTransform3& LocalTransform1, const FReal CullDistance, FCollisionConstraintsArray& NewConstraints)
 		{
 			if (ensure(Implicit0->IsConvex()))
@@ -1859,7 +1893,7 @@ namespace Chaos
 				if (const TImplicitObjectScaled<FTriangleMeshImplicitObject>* ScaledTriangleMesh = Implicit1->template GetObject<const TImplicitObjectScaled<FTriangleMeshImplicitObject>>())
 				{
 					FRigidBodyPointContactConstraint Constraint = FRigidBodyPointContactConstraint(Particle0, Implicit0, LocalTransform0, Particle1, Implicit1, LocalTransform0, EContactShapesType::ConvexTriMesh);
-					if (!bChaos_Collision_DeferNarrowPhase)
+					if (T_TRAITS::bImmediateUpdate)
 					{
 						FRigidTransform3 WorldTransform0 = LocalTransform0 * Collisions::GetTransform(Particle0);
 						FRigidTransform3 WorldTransform1 = LocalTransform1 * Collisions::GetTransform(Particle1);
@@ -1874,7 +1908,7 @@ namespace Chaos
 				else if (const FTriangleMeshImplicitObject* TriangleMesh = Implicit1->template GetObject<const FTriangleMeshImplicitObject>())
 				{
 					FRigidBodyPointContactConstraint Constraint = FRigidBodyPointContactConstraint(Particle0, Implicit0, LocalTransform0, Particle1, Implicit1, LocalTransform1, EContactShapesType::ConvexTriMesh);
-					if (!bChaos_Collision_DeferNarrowPhase)
+					if (T_TRAITS::bImmediateUpdate)
 					{
 						FRigidTransform3 WorldTransform0 = LocalTransform0 * Collisions::GetTransform(Particle0);
 						FRigidTransform3 WorldTransform1 = LocalTransform1 * Collisions::GetTransform(Particle1);
@@ -1962,6 +1996,7 @@ namespace Chaos
 		}
 
 
+		template<typename T_TRAITS>
 		void ConstructLevelsetLevelsetConstraints(TGeometryParticleHandle<FReal, 3>* Particle0, TGeometryParticleHandle<FReal, 3>* Particle1, const FImplicitObject* Implicit0, const FImplicitObject* Implicit1, const FRigidTransform3& LocalTransform0, const FRigidTransform3& LocalTransform1, const FReal CullDistance, FCollisionConstraintsArray& NewConstraints)
 		{
 			FRigidBodyPointContactConstraint Constraint = FRigidBodyPointContactConstraint(Particle0, Implicit0, LocalTransform0, Particle1, Implicit1, LocalTransform1, EContactShapesType::LevelSetLevelSet);
@@ -1980,7 +2015,7 @@ namespace Chaos
 				Constraint.SetManifold(Implicit0, Implicit1);
 			}
 
-			if (!bChaos_Collision_DeferNarrowPhase)
+			if (T_TRAITS::bImmediateUpdate)
 			{
 				UpdateLevelsetLevelsetConstraint<ECollisionUpdateType::Deepest>(CullDistance, Constraint);
 				NewConstraints.TryAdd(CullDistance, Constraint);
@@ -2146,6 +2181,7 @@ namespace Chaos
 			}
 		}
 
+		template<typename T_TRAITS>
 		void ConstructConstraintsImpl(TGeometryParticleHandle<FReal, 3>* Particle0, TGeometryParticleHandle<FReal, 3>* Particle1, const FImplicitObject* Implicit0, const FImplicitObject* Implicit1, const FRigidTransform3& LocalTransform0, const FRigidTransform3& LocalTransform1, const FReal CullDistance, const FCollisionContext& Context, FCollisionConstraintsArray& NewConstraints)
 		{
 			// @todo(chaos): We use GetInnerType here because TriMeshes are left with their "Instanced" wrapper, unlike all other instanced implicits. Should we strip the instance on Tri Mesh too?
@@ -2160,71 +2196,71 @@ namespace Chaos
 
 			if (Implicit0Type == TBox<FReal, 3>::StaticType() && Implicit1Type == TBox<FReal, 3>::StaticType())
 			{
-				ConstructBoxBoxConstraints(Particle0, Particle1, Implicit0, Implicit1, LocalTransform0, LocalTransform1, CullDistance, NewConstraints);
+				ConstructBoxBoxConstraints<T_TRAITS>(Particle0, Particle1, Implicit0, Implicit1, LocalTransform0, LocalTransform1, CullDistance, NewConstraints);
 			}
 			else if (Implicit0Type == TBox<FReal, 3>::StaticType() && Implicit1Type == THeightField<FReal>::StaticType())
 			{
-				ConstructBoxHeightFieldConstraints(Particle0, Particle1, Implicit0, Implicit1, LocalTransform0, LocalTransform1, CullDistance, NewConstraints);
+				ConstructBoxHeightFieldConstraints<T_TRAITS>(Particle0, Particle1, Implicit0, Implicit1, LocalTransform0, LocalTransform1, CullDistance, NewConstraints);
 			}
 			else if (Implicit0Type == THeightField<FReal>::StaticType() && Implicit1Type == TBox<FReal, 3>::StaticType())
 			{
-				ConstructBoxHeightFieldConstraints(Particle1, Particle0, Implicit1, Implicit0, LocalTransform1, LocalTransform0, CullDistance, NewConstraints);
+				ConstructBoxHeightFieldConstraints<T_TRAITS>(Particle1, Particle0, Implicit1, Implicit0, LocalTransform1, LocalTransform0, CullDistance, NewConstraints);
 			}
 			else if (Implicit0Type == TBox<FReal, 3>::StaticType() && Implicit1Type == TPlane<FReal, 3>::StaticType())
 			{
-				ConstructBoxPlaneConstraints(Particle0, Particle1, Implicit0, Implicit1, LocalTransform0, LocalTransform1, CullDistance, NewConstraints);
+				ConstructBoxPlaneConstraints<T_TRAITS>(Particle0, Particle1, Implicit0, Implicit1, LocalTransform0, LocalTransform1, CullDistance, NewConstraints);
 			}
 			else if (Implicit0Type == TPlane<FReal, 3>::StaticType() && Implicit1Type == TBox<FReal, 3>::StaticType())
 			{
-				ConstructBoxPlaneConstraints(Particle1, Particle0, Implicit1, Implicit0, LocalTransform1, LocalTransform0, CullDistance, NewConstraints);
+				ConstructBoxPlaneConstraints<T_TRAITS>(Particle1, Particle0, Implicit1, Implicit0, LocalTransform1, LocalTransform0, CullDistance, NewConstraints);
 			}
 			else if (Implicit0Type == TSphere<FReal, 3>::StaticType() && Implicit1Type == TSphere<FReal, 3>::StaticType())
 			{
-				ConstructSphereSphereConstraints(Particle0, Particle1, Implicit0, Implicit1, LocalTransform0, LocalTransform1, CullDistance, NewConstraints);
+				ConstructSphereSphereConstraints<T_TRAITS>(Particle0, Particle1, Implicit0, Implicit1, LocalTransform0, LocalTransform1, CullDistance, NewConstraints);
 			}
 			else if (Implicit0Type == TSphere<FReal, 3>::StaticType() && Implicit1Type == THeightField<FReal>::StaticType())
 			{
-				ConstructSphereHeightFieldConstraints(Particle0, Particle1, Implicit0, Implicit1, LocalTransform0, LocalTransform1, CullDistance, NewConstraints);
+				ConstructSphereHeightFieldConstraints<T_TRAITS>(Particle0, Particle1, Implicit0, Implicit1, LocalTransform0, LocalTransform1, CullDistance, NewConstraints);
 			}
 			else if (Implicit0Type == THeightField<FReal>::StaticType() && Implicit1Type == TSphere<FReal, 3>::StaticType())
 			{
-				ConstructSphereHeightFieldConstraints(Particle1, Particle0, Implicit1, Implicit0, LocalTransform1, LocalTransform0, CullDistance, NewConstraints);
+				ConstructSphereHeightFieldConstraints<T_TRAITS>(Particle1, Particle0, Implicit1, Implicit0, LocalTransform1, LocalTransform0, CullDistance, NewConstraints);
 			}
 			else if (Implicit0Type == TSphere<FReal, 3>::StaticType() && Implicit1Type == TPlane<FReal, 3>::StaticType())
 			{
-				ConstructSpherePlaneConstraints(Particle0, Particle1, Implicit0, Implicit1, LocalTransform0, LocalTransform1, CullDistance, NewConstraints);
+				ConstructSpherePlaneConstraints<T_TRAITS>(Particle0, Particle1, Implicit0, Implicit1, LocalTransform0, LocalTransform1, CullDistance, NewConstraints);
 			}
 			else if (Implicit0Type == TPlane<FReal, 3>::StaticType() && Implicit1Type == TSphere<FReal, 3>::StaticType())
 			{
-				ConstructSpherePlaneConstraints(Particle1, Particle0, Implicit1, Implicit0, LocalTransform1, LocalTransform0, CullDistance, NewConstraints);
+				ConstructSpherePlaneConstraints<T_TRAITS>(Particle1, Particle0, Implicit1, Implicit0, LocalTransform1, LocalTransform0, CullDistance, NewConstraints);
 			}
 			else if (Implicit0Type == TSphere<FReal, 3>::StaticType() && Implicit1Type == TBox<FReal, 3>::StaticType())
 			{
-				ConstructSphereBoxConstraints(Particle0, Particle1, Implicit0, Implicit1, LocalTransform0, LocalTransform1, CullDistance, NewConstraints);
+				ConstructSphereBoxConstraints<T_TRAITS>(Particle0, Particle1, Implicit0, Implicit1, LocalTransform0, LocalTransform1, CullDistance, NewConstraints);
 			}
 			else if (Implicit0Type == TBox<FReal, 3>::StaticType() && Implicit1Type == TSphere<FReal, 3>::StaticType())
 			{
-				ConstructSphereBoxConstraints(Particle1, Particle0, Implicit1, Implicit0, LocalTransform1, LocalTransform0, CullDistance, NewConstraints);
+				ConstructSphereBoxConstraints<T_TRAITS>(Particle1, Particle0, Implicit1, Implicit0, LocalTransform1, LocalTransform0, CullDistance, NewConstraints);
 			}
 			else if (Implicit0Type == TSphere<FReal, 3>::StaticType() && Implicit1Type == TCapsule<FReal>::StaticType())
 			{
-				ConstructSphereCapsuleConstraints(Particle0, Particle1, Implicit0, Implicit1, LocalTransform0, LocalTransform1, CullDistance, NewConstraints);
+				ConstructSphereCapsuleConstraints<T_TRAITS>(Particle0, Particle1, Implicit0, Implicit1, LocalTransform0, LocalTransform1, CullDistance, NewConstraints);
 			}
 			else if (Implicit0Type == TCapsule<FReal>::StaticType() && Implicit1Type == TSphere<FReal, 3>::StaticType())
 			{
-				ConstructSphereCapsuleConstraints(Particle1, Particle0, Implicit1, Implicit0, LocalTransform1, LocalTransform0, CullDistance, NewConstraints);
+				ConstructSphereCapsuleConstraints<T_TRAITS>(Particle1, Particle0, Implicit1, Implicit0, LocalTransform1, LocalTransform0, CullDistance, NewConstraints);
 			}
 			else if (Implicit0Type == TCapsule<FReal>::StaticType() && Implicit1Type == TCapsule<FReal>::StaticType())
 			{
-				ConstructCapsuleCapsuleConstraints(Particle0, Particle1, Implicit0, Implicit1, LocalTransform0, LocalTransform1, CullDistance, NewConstraints);
+				ConstructCapsuleCapsuleConstraints<T_TRAITS>(Particle0, Particle1, Implicit0, Implicit1, LocalTransform0, LocalTransform1, CullDistance, NewConstraints);
 			}
 			else if (Implicit0Type == TCapsule<FReal>::StaticType() && Implicit1Type == TBox<FReal, 3>::StaticType())
 			{
-				ConstructCapsuleBoxConstraints(Particle0, Particle1, Implicit0, Implicit1, LocalTransform0, LocalTransform1, CullDistance, Context, NewConstraints);
+				ConstructCapsuleBoxConstraints<T_TRAITS>(Particle0, Particle1, Implicit0, Implicit1, LocalTransform0, LocalTransform1, CullDistance, Context, NewConstraints);
 			}
 			else if (Implicit0Type == TBox<FReal, 3>::StaticType() && Implicit1Type == TCapsule<FReal>::StaticType())
 			{
-				ConstructCapsuleBoxConstraints(Particle1, Particle0, Implicit1, Implicit0, LocalTransform1, LocalTransform0, CullDistance, Context, NewConstraints);
+				ConstructCapsuleBoxConstraints<T_TRAITS>(Particle1, Particle0, Implicit1, Implicit0, LocalTransform1, LocalTransform0, CullDistance, Context, NewConstraints);
 			}
 			else if (Implicit0Type == TCapsule<FReal>::StaticType() && Implicit1Type == THeightField<FReal>::StaticType())
 			{
@@ -2234,7 +2270,7 @@ namespace Chaos
 					return;
 				}
 
-				ConstructCapsuleHeightFieldConstraints(Particle0, Particle1, Implicit0, Implicit1, LocalTransform0, LocalTransform1, CullDistance, NewConstraints);
+				ConstructCapsuleHeightFieldConstraints<T_TRAITS>(Particle0, Particle1, Implicit0, Implicit1, LocalTransform0, LocalTransform1, CullDistance, NewConstraints);
 			}
 			else if (Implicit0Type == THeightField<FReal>::StaticType() && Implicit1Type == TCapsule<FReal>::StaticType())
 			{
@@ -2244,23 +2280,23 @@ namespace Chaos
 					return;
 				}
 
-				ConstructCapsuleHeightFieldConstraints(Particle1, Particle0, Implicit1, Implicit0, LocalTransform1, LocalTransform0, CullDistance, NewConstraints);
+				ConstructCapsuleHeightFieldConstraints<T_TRAITS>(Particle1, Particle0, Implicit1, Implicit0, LocalTransform1, LocalTransform0, CullDistance, NewConstraints);
 			}
 			else if (Implicit0Type == TBox<FReal, 3>::StaticType() && Implicit1Type == FTriangleMeshImplicitObject::StaticType())
 			{
-				ConstructBoxTriangleMeshConstraints(Particle0, Particle1, Implicit0, Implicit1, LocalTransform0, LocalTransform1, CullDistance, NewConstraints);
+				ConstructBoxTriangleMeshConstraints<T_TRAITS>(Particle0, Particle1, Implicit0, Implicit1, LocalTransform0, LocalTransform1, CullDistance, NewConstraints);
 			}
 			else if (Implicit0Type == FTriangleMeshImplicitObject::StaticType() && Implicit1Type == TBox<FReal, 3>::StaticType())
 			{
-				ConstructBoxTriangleMeshConstraints(Particle1, Particle0, Implicit1, Implicit0, LocalTransform1, LocalTransform0, CullDistance, NewConstraints);
+				ConstructBoxTriangleMeshConstraints<T_TRAITS>(Particle1, Particle0, Implicit1, Implicit0, LocalTransform1, LocalTransform0, CullDistance, NewConstraints);
 			}
 			else if (Implicit0Type == TSphere<FReal, 3>::StaticType() && Implicit1Type == FTriangleMeshImplicitObject::StaticType())
 			{
-				ConstructSphereTriangleMeshConstraints(Particle0, Particle1, Implicit0, Implicit1, LocalTransform0, LocalTransform1, CullDistance, NewConstraints);
+				ConstructSphereTriangleMeshConstraints<T_TRAITS>(Particle0, Particle1, Implicit0, Implicit1, LocalTransform0, LocalTransform1, CullDistance, NewConstraints);
 			}
 			else if (Implicit0Type == FTriangleMeshImplicitObject::StaticType() && Implicit1Type == TSphere<FReal, 3>::StaticType())
 			{
-				ConstructSphereTriangleMeshConstraints(Particle1, Particle0, Implicit1, Implicit0, LocalTransform1, LocalTransform0, CullDistance, NewConstraints);
+				ConstructSphereTriangleMeshConstraints<T_TRAITS>(Particle1, Particle0, Implicit1, Implicit0, LocalTransform1, LocalTransform0, CullDistance, NewConstraints);
 			}
 			else if (Implicit0Type == TCapsule<FReal>::StaticType() && Implicit1Type == FTriangleMeshImplicitObject::StaticType())
 			{
@@ -2270,7 +2306,7 @@ namespace Chaos
 					return;
 				}
 
-				ConstructCapsuleTriangleMeshConstraints(Particle0, Particle1, Implicit0, Implicit1, LocalTransform0, LocalTransform1, CullDistance, NewConstraints);
+				ConstructCapsuleTriangleMeshConstraints<T_TRAITS>(Particle0, Particle1, Implicit0, Implicit1, LocalTransform0, LocalTransform1, CullDistance, NewConstraints);
 			}
 			else if (Implicit0Type == FTriangleMeshImplicitObject::StaticType() && Implicit1Type == TCapsule<FReal>::StaticType())
 			{
@@ -2280,7 +2316,7 @@ namespace Chaos
 					return;
 				}
 
-				ConstructCapsuleTriangleMeshConstraints(Particle1, Particle0, Implicit1, Implicit0, LocalTransform1, LocalTransform0, CullDistance, NewConstraints);
+				ConstructCapsuleTriangleMeshConstraints<T_TRAITS>(Particle1, Particle0, Implicit1, Implicit0, LocalTransform1, LocalTransform0, CullDistance, NewConstraints);
 			}
 			else if (bIsConvex0 && Implicit1Type == THeightField<FReal>::StaticType())
 			{
@@ -2290,7 +2326,7 @@ namespace Chaos
 					return;
 				}
 
-				ConstructConvexHeightFieldConstraints(Particle0, Particle1, Implicit0, Implicit1, LocalTransform0, LocalTransform1, CullDistance, NewConstraints);
+				ConstructConvexHeightFieldConstraints<T_TRAITS>(Particle0, Particle1, Implicit0, Implicit1, LocalTransform0, LocalTransform1, CullDistance, NewConstraints);
 			}
 			else if (Implicit0Type == THeightField<FReal>::StaticType() && bIsConvex1)
 			{
@@ -2300,7 +2336,7 @@ namespace Chaos
 					return;
 				}
 
-				ConstructConvexHeightFieldConstraints(Particle1, Particle0, Implicit1, Implicit0, LocalTransform1, LocalTransform0, CullDistance, NewConstraints);
+				ConstructConvexHeightFieldConstraints<T_TRAITS>(Particle1, Particle0, Implicit1, Implicit0, LocalTransform1, LocalTransform0, CullDistance, NewConstraints);
 			}
 			else if (bIsConvex0 && Implicit1Type == FTriangleMeshImplicitObject::StaticType())
 			{
@@ -2310,7 +2346,7 @@ namespace Chaos
 					return;
 				}
 
-				ConstructConvexTriangleMeshConstraints(Particle0, Particle1, Implicit0, Implicit1, LocalTransform0, LocalTransform1, CullDistance, NewConstraints);
+				ConstructConvexTriangleMeshConstraints<T_TRAITS>(Particle0, Particle1, Implicit0, Implicit1, LocalTransform0, LocalTransform1, CullDistance, NewConstraints);
 			}
 			else if (Implicit0Type == FTriangleMeshImplicitObject::StaticType() && bIsConvex1)
 			{
@@ -2320,7 +2356,7 @@ namespace Chaos
 					return;
 				}
 
-				ConstructConvexTriangleMeshConstraints(Particle1, Particle0, Implicit1, Implicit0, LocalTransform1, LocalTransform0, CullDistance, NewConstraints);
+				ConstructConvexTriangleMeshConstraints<T_TRAITS>(Particle1, Particle0, Implicit1, Implicit0, LocalTransform1, LocalTransform0, CullDistance, NewConstraints);
 			}
 			else if (bIsConvex0 && bIsConvex1)
 			{
@@ -2330,11 +2366,11 @@ namespace Chaos
 					return;
 				}
 
-				ConstructConvexConvexConstraints(Particle0, Particle1, Implicit0, Implicit1, LocalTransform0, LocalTransform1, CullDistance, NewConstraints);
+				ConstructConvexConvexConstraints<T_TRAITS>(Particle0, Particle1, Implicit0, Implicit1, LocalTransform0, LocalTransform1, CullDistance, NewConstraints);
 			}
 			else
 			{
-				ConstructLevelsetLevelsetConstraints(Particle0, Particle1, Implicit0, Implicit1, LocalTransform0, LocalTransform1, CullDistance, NewConstraints);
+				ConstructLevelsetLevelsetConstraints<T_TRAITS>(Particle0, Particle1, Implicit0, Implicit1, LocalTransform0, LocalTransform1, CullDistance, NewConstraints);
 			}
 		}
 
@@ -2358,7 +2394,7 @@ namespace Chaos
 			const int32 NumPoints = Constraint.NumManifoldPoints();
 
 			// Fall back to full collision detection if we have no manifold (or for testing)
-			if ((NumPoints == 0) || Chaos_Collision_UseManifolds_Test)
+			if ((NumPoints == 0) || bChaos_Collision_UseManifolds_Test)
 			{
 				UpdateConstraintFromGeometryImpl<ECollisionUpdateType::Deepest>(Constraint, WorldTransform0, WorldTransform1, CullDistance);
 				return;
@@ -2486,6 +2522,7 @@ namespace Chaos
 		}
 
 
+		template<typename T_TRAITS>
 		void ConstructConstraints(TGeometryParticleHandle<FReal, 3>* Particle0, TGeometryParticleHandle<FReal, 3>* Particle1, const FImplicitObject* Implicit0, const FImplicitObject* Implicit1, const FRigidTransform3& LocalTransform0, const FRigidTransform3& LocalTransform1, const FReal CullDistance, const FCollisionContext& Context, FCollisionConstraintsArray& NewConstraints)
 		{
 			EImplicitObjectType Implicit0Type = Implicit0 ? GetInnerType(Implicit0->GetType()) : ImplicitObjectType::Unknown;
@@ -2494,7 +2531,7 @@ namespace Chaos
 
 			if (!Implicit0 || !Implicit1)
 			{
-				ConstructLevelsetLevelsetConstraints(Particle0, Particle1, Implicit0, Implicit1, LocalTransform0, LocalTransform1, CullDistance, NewConstraints);
+				ConstructLevelsetLevelsetConstraints<T_TRAITS>(Particle0, Particle1, Implicit0, Implicit1, LocalTransform0, LocalTransform1, CullDistance, NewConstraints);
 				return;
 			}
 
@@ -2508,21 +2545,21 @@ namespace Chaos
 				const TImplicitObjectTransformed<FReal, 3>* TransformedImplicit1 = Implicit1->template GetObject<const TImplicitObjectTransformed<FReal, 3>>();
 				FRigidTransform3 TransformedTransform0 = TransformedImplicit0->GetTransform() * LocalTransform0;
 				FRigidTransform3 TransformedTransform1 = TransformedImplicit1->GetTransform() * LocalTransform1;
-				ConstructConstraints(Particle0, Particle1, TransformedImplicit0->GetTransformedObject(), TransformedImplicit1->GetTransformedObject(), TransformedTransform0, TransformedTransform1, CullDistance, Context, NewConstraints);
+				ConstructConstraints<T_TRAITS>(Particle0, Particle1, TransformedImplicit0->GetTransformedObject(), TransformedImplicit1->GetTransformedObject(), TransformedTransform0, TransformedTransform1, CullDistance, Context, NewConstraints);
 				return;
 			}
 			else if (Implicit0OuterType == TImplicitObjectTransformed<FReal, 3>::StaticType())
 			{
 				const TImplicitObjectTransformed<FReal, 3>* TransformedImplicit0 = Implicit0->template GetObject<const TImplicitObjectTransformed<FReal, 3>>();
 				FRigidTransform3 TransformedTransform0 = TransformedImplicit0->GetTransform() * LocalTransform0;
-				ConstructConstraints(Particle0, Particle1, TransformedImplicit0->GetTransformedObject(), Implicit1, TransformedTransform0, LocalTransform1, CullDistance, Context, NewConstraints);
+				ConstructConstraints<T_TRAITS>(Particle0, Particle1, TransformedImplicit0->GetTransformedObject(), Implicit1, TransformedTransform0, LocalTransform1, CullDistance, Context, NewConstraints);
 				return;
 			}
 			else if (Implicit1OuterType == TImplicitObjectTransformed<FReal, 3>::StaticType())
 			{
 				const TImplicitObjectTransformed<FReal, 3>* TransformedImplicit1 = Implicit1->template GetObject<const TImplicitObjectTransformed<FReal, 3>>();
 				FRigidTransform3 TransformedTransform1 = TransformedImplicit1->GetTransform() * LocalTransform1;
-				ConstructConstraints(Particle0, Particle1, Implicit0, TransformedImplicit1->GetTransformedObject(), LocalTransform0, TransformedTransform1, CullDistance, Context, NewConstraints);
+				ConstructConstraints<T_TRAITS>(Particle0, Particle1, Implicit0, TransformedImplicit1->GetTransformedObject(), LocalTransform0, TransformedTransform1, CullDistance, Context, NewConstraints);
 				return;
 			}
 			// Handle Instanced shapes
@@ -2541,17 +2578,17 @@ namespace Chaos
 				}
 				if (InnerImplicit0 && InnerImplicit1)
 				{
-					ConstructConstraints(Particle0, Particle1, InnerImplicit0, InnerImplicit1, LocalTransform0, LocalTransform1, CullDistance, Context, NewConstraints);
+					ConstructConstraints<T_TRAITS>(Particle0, Particle1, InnerImplicit0, InnerImplicit1, LocalTransform0, LocalTransform1, CullDistance, Context, NewConstraints);
 					return;
 				}
 				else if (InnerImplicit0 && !InnerImplicit1)
 				{
-					ConstructConstraints(Particle0, Particle1, InnerImplicit0, Implicit1, LocalTransform0, LocalTransform1, CullDistance, Context, NewConstraints);
+					ConstructConstraints<T_TRAITS>(Particle0, Particle1, InnerImplicit0, Implicit1, LocalTransform0, LocalTransform1, CullDistance, Context, NewConstraints);
 					return;
 				}
 				else if (!InnerImplicit0 && InnerImplicit1)
 				{
-					ConstructConstraints(Particle0, Particle1, Implicit0, InnerImplicit1, LocalTransform0, LocalTransform1, CullDistance, Context, NewConstraints);
+					ConstructConstraints<T_TRAITS>(Particle0, Particle1, Implicit0, InnerImplicit1, LocalTransform0, LocalTransform1, CullDistance, Context, NewConstraints);
 					return;
 				}
 
@@ -2563,7 +2600,7 @@ namespace Chaos
 				const FImplicitObjectUnion* Union0 = Implicit0->template GetObject<FImplicitObjectUnion>();
 				for (const auto& Child0 : Union0->GetObjects())
 				{
-					ConstructConstraints(Particle0, Particle1, Child0.Get(), Implicit1, LocalTransform0, LocalTransform1, CullDistance, Context, NewConstraints);
+					ConstructConstraints<T_TRAITS>(Particle0, Particle1, Child0.Get(), Implicit1, LocalTransform0, LocalTransform1, CullDistance, Context, NewConstraints);
 				}
 				return;
 			}
@@ -2572,7 +2609,7 @@ namespace Chaos
 				const FImplicitObjectUnion* Union1 = Implicit1->template GetObject<FImplicitObjectUnion>();
 				for (const auto& Child1 : Union1->GetObjects())
 				{
-					ConstructConstraints(Particle0, Particle1, Implicit0, Child1.Get(), LocalTransform0, LocalTransform1, CullDistance, Context, NewConstraints);
+					ConstructConstraints<T_TRAITS>(Particle0, Particle1, Implicit0, Child1.Get(), LocalTransform0, LocalTransform1, CullDistance, Context, NewConstraints);
 				}
 				return;
 			}
@@ -2585,15 +2622,33 @@ namespace Chaos
 
 			// If we get here, we have a pair of concrete shapes (i.e., no wrappers or containers)
 			// Create a constraint for the shape pair
-			FReal ModifiedCullDistance = CullDistance;
-			if (bChaos_Collision_DeferNarrowPhase)
+			ConstructConstraintsImpl<T_TRAITS>(Particle0, Particle1, Implicit0, Implicit1, LocalTransform0, LocalTransform1, CullDistance, Context, NewConstraints);
+		}
+
+		void ConstructConstraints(TGeometryParticleHandle<FReal, 3>* Particle0, TGeometryParticleHandle<FReal, 3>* Particle1, const FImplicitObject* Implicit0, const FImplicitObject* Implicit1, const FRigidTransform3& LocalTransform0, const FRigidTransform3& LocalTransform1, const FReal CullDistance, const FCollisionContext& Context, FCollisionConstraintsArray& NewConstraints)
+		{
+			if (Context.bDeferUpdate && !Context.bAllowManifolds)
+			{
+				using TTraits = TConstructCollisionTraits<false, false>;
+				ConstructConstraints<TTraits>(Particle0, Particle1, Implicit0, Implicit1, LocalTransform0, LocalTransform1, CullDistance, Context, NewConstraints);
+			}
+			else if (Context.bDeferUpdate && Context.bAllowManifolds)
+			{
+				using TTraits = TConstructCollisionTraits<false, true>;
+				ConstructConstraints<TTraits>(Particle0, Particle1, Implicit0, Implicit1, LocalTransform0, LocalTransform1, CullDistance, Context, NewConstraints);
+			}
+			else if (!Context.bDeferUpdate && Context.bAllowManifolds)
+			{
+				using TTraits = TConstructCollisionTraits<true, true>;
+				ConstructConstraints<TTraits>(Particle0, Particle1, Implicit0, Implicit1, LocalTransform0, LocalTransform1, CullDistance, Context, NewConstraints);
+			}
+			else if (!Context.bDeferUpdate && !Context.bAllowManifolds)
 			{
-				ModifiedCullDistance = FLT_MAX;
+				using TTraits = TConstructCollisionTraits<true, false>;
+				ConstructConstraints<TTraits>(Particle0, Particle1, Implicit0, Implicit1, LocalTransform0, LocalTransform1, CullDistance, Context, NewConstraints);
 			}
-			ConstructConstraintsImpl(Particle0, Particle1, Implicit0, Implicit1, LocalTransform0, LocalTransform1, ModifiedCullDistance, Context, NewConstraints);
 		}
 
-
 		template <typename GeometryA, typename GeometryB>
 		bool GetPairTOIHackImpl(const FImplicitObject& A, const FRigidTransform3& AStartTransform, const GeometryB& B, const FRigidTransform3& BTransform, const FVec3& Dir, const FReal Length, FReal& OutTOI, FVec3& OutNormal, FReal& OutPhi)
 		{
diff --git a/Engine/Source/Runtime/Experimental/Chaos/Private/Chaos/Evolution/PBDMinEvolution.cpp b/Engine/Source/Runtime/Experimental/Chaos/Private/Chaos/Evolution/PBDMinEvolution.cpp
--- a/Engine/Source/Runtime/Experimental/Chaos/Private/Chaos/Evolution/PBDMinEvolution.cpp
+++ b/Engine/Source/Runtime/Experimental/Chaos/Private/Chaos/Evolution/PBDMinEvolution.cpp
@@ -1,9 +1,8 @@
 // Copyright Epic Games, Inc. All Rights Reserved.
 
 #include "Chaos/Evolution/PBDMinEvolution.h"
-#include "Chaos/Collision/CollisionDetector.h"
 #include "Chaos/Collision/NarrowPhase.h"
-#include "Chaos/Collision/ParticlePairBroadPhase.h"
+#include "Chaos/Collision/ParticlePairCollisionDetector.h"
 #include "Chaos/PBDCollisionConstraints.h"
 #include "Chaos/PBDConstraintRule.h"
 #include "Chaos/PBDRigidsSOAs.h"
diff --git a/Engine/Source/Runtime/Experimental/Chaos/Private/Chaos/PBDRigidsEvolutionGBF.cpp b/Engine/Source/Runtime/Experimental/Chaos/Private/Chaos/PBDRigidsEvolutionGBF.cpp
--- a/Engine/Source/Runtime/Experimental/Chaos/Private/Chaos/PBDRigidsEvolutionGBF.cpp
+++ b/Engine/Source/Runtime/Experimental/Chaos/Private/Chaos/PBDRigidsEvolutionGBF.cpp
@@ -288,7 +288,7 @@ void FPBDRigidsEvolutionGBF::AdvanceOneTimeStep(const FReal Dt, const FReal Step
 
 		CollisionStats::FStatData StatData(bPendingHierarchyDump);
 
-		CollisionDetector.DetectCollisions(Dt, StatData);
+		CollisionDetector.DetectCollisionsWithStats(Dt, StatData);
 
 		CHAOS_COLLISION_STAT(StatData.Print());
 	}
@@ -432,7 +432,8 @@ FPBDRigidsEvolutionGBF::FPBDRigidsEvolutionGBF(TPBDRigidsSOAs<FReal, 3>& InParti
 	, CollisionConstraints(InParticles, Collided, PhysicsMaterials, DefaultNumPairIterations, DefaultNumPushOutPairIterations)
 	, CollisionRule(CollisionConstraints)
 	, BroadPhase(InParticles, BoundsThickness, BoundsThicknessVelocityMultiplier)
-	, CollisionDetector(BroadPhase, CollisionConstraints)
+	, NarrowPhase()
+	, CollisionDetector(BroadPhase, NarrowPhase, CollisionConstraints)
 	, PostIntegrateCallback(nullptr)
 	, PreApplyCallback(nullptr)
 	, PostApplyCallback(nullptr)
@@ -462,6 +463,10 @@ FPBDRigidsEvolutionGBF::FPBDRigidsEvolutionGBF(TPBDRigidsSOAs<FReal, 3>& InParti
 	});
 
 	AddConstraintRule(&CollisionRule);
+
+	NarrowPhase.GetContext().bFilteringEnabled = true;
+	NarrowPhase.GetContext().bDeferUpdate = true;
+	NarrowPhase.GetContext().bAllowManifolds = false;
 }
 
 void FPBDRigidsEvolutionGBF::Serialize(FChaosArchive& Ar)
diff --git a/Engine/Source/Runtime/Experimental/Chaos/Public/Chaos/ChaosPerfTest.h b/Engine/Source/Runtime/Experimental/Chaos/Public/Chaos/ChaosPerfTest.h
--- a/Engine/Source/Runtime/Experimental/Chaos/Public/Chaos/ChaosPerfTest.h
+++ b/Engine/Source/Runtime/Experimental/Chaos/Public/Chaos/ChaosPerfTest.h
@@ -16,7 +16,7 @@ enum class EChaosPerfUnits
 	Num
 };
 
-class FChaosScopedDurationTimeLogger
+class CHAOS_API FChaosScopedDurationTimeLogger
 {
 public:
 	explicit FChaosScopedDurationTimeLogger(const TCHAR* InLabel)
diff --git a/Engine/Source/Runtime/Experimental/Chaos/Public/Chaos/Collision/CollisionContext.h b/Engine/Source/Runtime/Experimental/Chaos/Public/Chaos/Collision/CollisionContext.h
--- a/Engine/Source/Runtime/Experimental/Chaos/Public/Chaos/Collision/CollisionContext.h
+++ b/Engine/Source/Runtime/Experimental/Chaos/Public/Chaos/Collision/CollisionContext.h
@@ -14,13 +14,23 @@ namespace Chaos
 		FCollisionContext()
 			: SpaceTransform(FVec3(0), FRotation3::FromIdentity())
 			, bFilteringEnabled(true)
+			, bDeferUpdate(true)
+			, bAllowManifolds(false)
 		{
 		}
 
 		// The simulation space to world-space transform (e.g., Rigid Body nodes are often simulated relative to a component or bone)
 		FRigidTransform3 SpaceTransform;
 
-		// Whether to check the shape query flags (e.g., Rigid Body nodes have already performed filtering prior to collision detection)
+		// Whether to check the shape query flags (e.g., Rigid Body nodes have already performed filtering prior to collision detection) [default: true]
 		bool bFilteringEnabled;
+
+		// Whether to defer constraint phi/normal calculation to the Apply step [default: true]. If true, constraints are speculatively created for each shape
+		// pair passed to the narrow phase. This prevents premature culling of constraints, but it can lead to more items in the constraint graph
+		// which could be undesirable in some cases (destruction?).
+		bool bDeferUpdate;
+
+		// Whether to use manifolds wheer supported [default: false]
+		bool bAllowManifolds;
 	};
 }
\ No newline at end of file
diff --git a/Engine/Source/Runtime/Experimental/Chaos/Public/Chaos/Collision/CollisionDetector.h b/Engine/Source/Runtime/Experimental/Chaos/Public/Chaos/Collision/CollisionDetector.h
--- a/Engine/Source/Runtime/Experimental/Chaos/Public/Chaos/Collision/CollisionDetector.h
+++ b/Engine/Source/Runtime/Experimental/Chaos/Public/Chaos/Collision/CollisionDetector.h
@@ -2,68 +2,40 @@
 #pragma once
 
 #include "Chaos/Collision/CollisionContext.h"
-#include "Chaos/Collision/CollisionReceiver.h"
-#include "Chaos/Collision/NarrowPhase.h"
-#include "Chaos/Collision/ParticlePairBroadPhase.h"
-#include "Chaos/Collision/SpatialAccelerationBroadPhase.h"
-#include "Chaos/ParticleHandle.h"
-#include "Chaos/PBDCollisionConstraints.h"
-#include "Chaos/PBDRigidsSOAs.h"
+#include "Chaos/Collision/StatsData.h"
 
 namespace Chaos
 {
+	class FPBDCollisionConstraints;
+	class FNarrowPhase;
+
 	DECLARE_CYCLE_STAT_EXTERN(TEXT("Collisions::Detect"), STAT_Collisions_Detect, STATGROUP_ChaosCollision, CHAOS_API);
 
-	template<typename T_BROADPHASE, typename T_NARROWPHASE, typename T_RECEIVER, typename T_CONTAINER>
-	class CHAOS_API TCollisionDetector
+	class CHAOS_API FCollisionDetector
 	{
 	public:
-		using FBroadPhase = T_BROADPHASE;
-		using FReceiver = T_RECEIVER;
-		using FNarrowPhase = T_NARROWPHASE;
-		using FContainer = T_CONTAINER;
-
-		TCollisionDetector(FBroadPhase& InBroadPhase, FContainer& InCollisionContainer)
-			: BroadPhase(InBroadPhase)
+		FCollisionDetector(FNarrowPhase& InNarrowPhase, FPBDCollisionConstraints& InCollisionContainer)
+			: NarrowPhase(InNarrowPhase)
 			, CollisionContainer(InCollisionContainer)
 		{
 		}
 
-		FBroadPhase& GetBroadPhase() { return BroadPhase; }
-		FContainer& GetCollisionContainer() { return CollisionContainer; }
-		FCollisionContext& GetContext() { return Context; }
+		virtual ~FCollisionDetector() {}
+
+		FPBDCollisionConstraints& GetCollisionContainer() { return CollisionContainer; }
+		FNarrowPhase& GetNarrowPhase() { return NarrowPhase; }
+
+		virtual void DetectCollisionsWithStats(const FReal Dt, CollisionStats::FStatData& StatData) = 0;
 
 		void DetectCollisions(const FReal Dt)
 		{
 			CollisionStats::FStatData StatData(false);
-			DetectCollisions(Dt, StatData);
-		}
-
-		void DetectCollisions(const FReal Dt, CollisionStats::FStatData& StatData)
-		{
-			SCOPE_CYCLE_COUNTER(STAT_Collisions_Detect);
-			CHAOS_SCOPED_TIMER(DetectCollisions);
-
-			if (!GetCollisionContainer().GetCollisionsEnabled())
-			{
-				return;
-			}
-
-			CollisionContainer.UpdateManifolds(Dt);
-			CollisionContainer.UpdateConstraints(Dt);
-
-			// Collision detection pipeline: BroadPhase -> NarrowPhase -> Receiver -> Container
-			// Receivers and NarrowPhase are assumed to be stateless atm. If we change that, they need to
-			// be passed into the constructor with the BroadPhase and Container.
-			FReceiver Receiver(CollisionContainer);
-			FNarrowPhase NarrowPhase(Context);
-			BroadPhase.ProduceOverlaps(Dt, NarrowPhase, Receiver, StatData);
-			Receiver.ProcessCollisions();
+			DetectCollisionsWithStats(Dt, StatData);
 		}
 
-	private:
-		FBroadPhase& BroadPhase;
-		FContainer& CollisionContainer;
+	protected:
+		FNarrowPhase& NarrowPhase;
+		FPBDCollisionConstraints& CollisionContainer;
 		FCollisionContext Context;
 	};
 
diff --git a/Engine/Source/Runtime/Experimental/Chaos/Public/Chaos/Collision/NarrowPhase.h b/Engine/Source/Runtime/Experimental/Chaos/Public/Chaos/Collision/NarrowPhase.h
--- a/Engine/Source/Runtime/Experimental/Chaos/Public/Chaos/Collision/NarrowPhase.h
+++ b/Engine/Source/Runtime/Experimental/Chaos/Public/Chaos/Collision/NarrowPhase.h
@@ -2,7 +2,7 @@
 #pragma once
 
 #include "Chaos/Core.h"
-#include "Chaos/Collision/CollisionReceiver.h"
+#include "Chaos/Collision/CollisionContext.h"
 #include "Chaos/Collision/StatsData.h"
 #include "Chaos/CollisionResolution.h"
 #include "Chaos/ParticleHandle.h"
@@ -30,11 +30,12 @@ namespace Chaos
 	class CHAOS_API FNarrowPhase
 	{
 	public:
-		FNarrowPhase(const FCollisionContext& InContext)
-			: Context(InContext)
+		FNarrowPhase()
 		{
 		}
 
+		FCollisionContext& GetContext() { return Context; }
+
 		// @todo(chaos): COLLISION Transient Handle version
 		/**
 		 * /param CullDistance The contact separation at which we ignore the constraint
@@ -59,6 +60,6 @@ namespace Chaos
 		}
 
 	private:
-		const FCollisionContext& Context;
+		FCollisionContext Context;
 	};
 }
diff --git a/Engine/Source/Runtime/Experimental/Chaos/Public/Chaos/Collision/ParticlePairBroadPhase.h b/Engine/Source/Runtime/Experimental/Chaos/Public/Chaos/Collision/ParticlePairBroadPhase.h
--- a/Engine/Source/Runtime/Experimental/Chaos/Public/Chaos/Collision/ParticlePairBroadPhase.h
+++ b/Engine/Source/Runtime/Experimental/Chaos/Public/Chaos/Collision/ParticlePairBroadPhase.h
@@ -2,7 +2,6 @@
 #pragma once
 
 #include "Chaos/BoundingVolumeUtilities.h"
-#include "Chaos/Collision/CollisionReceiver.h"
 #include "Chaos/Collision/NarrowPhase.h"
 #include "Chaos/ParticleHandle.h"
 #include "Chaos/PBDCollisionConstraints.h"
@@ -43,14 +42,12 @@ namespace Chaos
 		 *
 		 */
 		void ProduceOverlaps(FReal Dt,
+			FCollisionConstraintsArray& ConstraintsArray,
 			FNarrowPhase& NarrowPhase,
-			FSyncCollisionReceiver& Receiver,
 			CollisionStats::FStatData& StatData)
 		{
 			SCOPE_CYCLE_COUNTER(STAT_Collisions_BroadPhase);
 
-			FCollisionConstraintsArray& ConstraintsArray = Receiver.GetConstraintsArray();
-
 			int32 NumPairs = ParticlePairs.Num();
 			for (int32 PairIndex = 0; PairIndex < NumPairs; ++PairIndex)
 			{
diff --git a/Engine/Source/Runtime/Experimental/Chaos/Public/Chaos/Collision/ParticlePairCollisionDetector.h b/Engine/Source/Runtime/Experimental/Chaos/Public/Chaos/Collision/ParticlePairCollisionDetector.h
new file mode 100644
--- /dev/null
+++ b/Engine/Source/Runtime/Experimental/Chaos/Public/Chaos/Collision/ParticlePairCollisionDetector.h
@@ -0,0 +1,42 @@
+// Copyright Epic Games, Inc. All Rights Reserved.
+#pragma once
+
+#include "Chaos/Collision/CollisionDetector.h"
+#include "Chaos/Collision/NarrowPhase.h"
+#include "Chaos/Collision/ParticlePairBroadPhase.h"
+
+namespace Chaos
+{
+	class CHAOS_API FParticlePairCollisionDetector : public FCollisionDetector
+	{
+	public:
+		FParticlePairCollisionDetector(FParticlePairBroadPhase& InBroadPhase, FNarrowPhase& InNarrowPhase, FPBDCollisionConstraints& InCollisionContainer)
+			: FCollisionDetector(InNarrowPhase, InCollisionContainer)
+			, BroadPhase(InBroadPhase)
+		{
+		}
+
+		FParticlePairBroadPhase& GetBroadPhase() { return BroadPhase; }
+
+		virtual void DetectCollisionsWithStats(const FReal Dt, CollisionStats::FStatData& StatData) override
+		{
+			SCOPE_CYCLE_COUNTER(STAT_Collisions_Detect);
+			CHAOS_SCOPED_TIMER(DetectCollisions);
+
+			if (!GetCollisionContainer().GetCollisionsEnabled())
+			{
+				return;
+			}
+
+			CollisionContainer.UpdateManifolds(Dt);
+			CollisionContainer.UpdateConstraints(Dt);
+
+			// Collision detection pipeline: BroadPhase -> NarrowPhase -> Container
+			BroadPhase.ProduceOverlaps(Dt, CollisionContainer.GetConstraintsArray(), NarrowPhase, StatData);
+		}
+
+	private:
+		FParticlePairBroadPhase& BroadPhase;
+	};
+
+}
\ No newline at end of file
diff --git a/Engine/Source/Runtime/Experimental/Chaos/Public/Chaos/Collision/SpatialAccelerationBroadPhase.h b/Engine/Source/Runtime/Experimental/Chaos/Public/Chaos/Collision/SpatialAccelerationBroadPhase.h
--- a/Engine/Source/Runtime/Experimental/Chaos/Public/Chaos/Collision/SpatialAccelerationBroadPhase.h
+++ b/Engine/Source/Runtime/Experimental/Chaos/Public/Chaos/Collision/SpatialAccelerationBroadPhase.h
@@ -2,6 +2,7 @@
 #pragma once
 
 #include "Chaos/Collision/BroadPhase.h"
+#include "Chaos/Collision/CollisionReceiver.h"
 #include "Chaos/Collision/StatsData.h"
 #include "Chaos/Collision/NarrowPhase.h"
 #include "Chaos/ISpatialAccelerationCollection.h"
@@ -11,8 +12,6 @@
 
 namespace Chaos
 {
-	class FAsyncCollisionReceiver;
-
 	template <typename TPayloadType, typename T, int d>
 	class ISpatialAcceleration;
 
diff --git a/Engine/Source/Runtime/Experimental/Chaos/Public/Chaos/Collision/SpatialAccelerationCollisionDetector.h b/Engine/Source/Runtime/Experimental/Chaos/Public/Chaos/Collision/SpatialAccelerationCollisionDetector.h
new file mode 100644
--- /dev/null
+++ b/Engine/Source/Runtime/Experimental/Chaos/Public/Chaos/Collision/SpatialAccelerationCollisionDetector.h
@@ -0,0 +1,45 @@
+// Copyright Epic Games, Inc. All Rights Reserved.
+#pragma once
+
+#include "Chaos/Collision/CollisionDetector.h"
+#include "Chaos/Collision/CollisionContext.h"
+#include "Chaos/Collision/CollisionReceiver.h"
+#include "Chaos/Collision/NarrowPhase.h"
+#include "Chaos/Collision/SpatialAccelerationBroadPhase.h"
+
+namespace Chaos
+{
+	class CHAOS_API FSpatialAccelerationCollisionDetector : public FCollisionDetector
+	{
+	public:
+		FSpatialAccelerationCollisionDetector(FSpatialAccelerationBroadPhase& InBroadPhase, FNarrowPhase& InNarrowPhase, FPBDCollisionConstraints& InCollisionContainer)
+			: FCollisionDetector(InNarrowPhase, InCollisionContainer)
+			, BroadPhase(InBroadPhase)
+		{
+		}
+
+		FSpatialAccelerationBroadPhase& GetBroadPhase() { return BroadPhase; }
+
+		virtual void DetectCollisionsWithStats(const FReal Dt, CollisionStats::FStatData& StatData) override
+		{
+			SCOPE_CYCLE_COUNTER(STAT_Collisions_Detect);
+			CHAOS_SCOPED_TIMER(DetectCollisions);
+
+			if (!GetCollisionContainer().GetCollisionsEnabled())
+			{
+				return;
+			}
+
+			CollisionContainer.UpdateManifolds(Dt);
+			CollisionContainer.UpdateConstraints(Dt);
+
+			// Collision detection pipeline: BroadPhase -[parallel]-> NarrowPhase -[parallel]-> Receiver -[serial]-> Container
+			FAsyncCollisionReceiver Receiver(CollisionContainer);
+			BroadPhase.ProduceOverlaps(Dt, NarrowPhase, Receiver, StatData);
+			Receiver.ProcessCollisions();
+		}
+
+	private:
+		FSpatialAccelerationBroadPhase& BroadPhase;
+	};
+}
diff --git a/Engine/Source/Runtime/Experimental/Chaos/Public/Chaos/CollisionResolution.h b/Engine/Source/Runtime/Experimental/Chaos/Public/Chaos/CollisionResolution.h
--- a/Engine/Source/Runtime/Experimental/Chaos/Public/Chaos/CollisionResolution.h
+++ b/Engine/Source/Runtime/Experimental/Chaos/Public/Chaos/CollisionResolution.h
@@ -34,175 +34,27 @@ namespace Chaos
 	namespace Collisions
 	{
 
-		//
-		// Box-Box
-		//
-
-		void CHAOS_API UpdateBoxBoxConstraint(const FAABB3& Box1, const FRigidTransform3& Box1Transform, const FAABB3& Box2, const FRigidTransform3& Box2Transform, const FReal CullDistance, FRigidBodyPointContactConstraint& Constraint);
-
-		void CHAOS_API UpdateBoxBoxManifold(FCollisionConstraintBase&  Constraint, const FRigidTransform3& ATM, const FRigidTransform3& BTM, const FReal CullDistance);
-
-		void CHAOS_API ConstructBoxBoxConstraints(TGeometryParticleHandle<FReal, 3>* Particle0, TGeometryParticleHandle<FReal, 3>* Particle1, const FImplicitObject* Implicit0, const FImplicitObject* Implicit1, const FRigidTransform3& Transform0, const FRigidTransform3& Transform1, const FReal CullDistance, FCollisionConstraintsArray& NewConstraints);
-
-
-		//
-		// Box-HeightField
-		//
-
-		void CHAOS_API UpdateBoxHeightFieldConstraint(const FAABB3& A, const FRigidTransform3& ATransform, const THeightField<FReal>& B, const FRigidTransform3& BTransform, const FReal CullDistance, FRigidBodyPointContactConstraint& Constraint);
-
-		void CHAOS_API UpdateBoxHeightFieldManifold(FCollisionConstraintBase&  Constraint, const FRigidTransform3& ATM, const FRigidTransform3& BTM, const FReal CullDistance);
-
-		void CHAOS_API ConstructBoxHeightFieldConstraints(TGeometryParticleHandle<FReal, 3>* Particle0, TGeometryParticleHandle<FReal, 3>* Particle1, const FImplicitObject* Implicit0, const FImplicitObject* Implicit1, const FRigidTransform3& Transform0, const FRigidTransform3& Transform1, const FReal CullDistance, FCollisionConstraintsArray& NewConstraints);
-
-		//
-		// Box-Plane
-		//
-
-		void CHAOS_API UpdateBoxPlaneConstraint(const FAABB3& Box, const FRigidTransform3& BoxTransform, const TPlane<FReal, 3>& Plane, const FRigidTransform3& PlaneTransform, const FReal CullDistance, FRigidBodyPointContactConstraint& Constraint);
-
-		void CHAOS_API UpdateBoxPlaneManifold(FCollisionConstraintBase&  Constraint, const FRigidTransform3& ATM, const FRigidTransform3& BTM, const FReal CullDistance);
-
-		void CHAOS_API ConstructBoxPlaneConstraints(TGeometryParticleHandle<FReal, 3>* Particle0, TGeometryParticleHandle<FReal, 3>* Particle1, const FImplicitObject* Implicit0, const FImplicitObject* Implicit1, const FRigidTransform3& Transform0, const FRigidTransform3& Transform1, const FReal CullDistance, FCollisionConstraintsArray& NewConstraints);
-
-
-
-		//
-		// Box-TriangleMesh
-		//
-
-		template <typename TriMeshType>
-		void CHAOS_API UpdateBoxTriangleMeshConstraint(const FAABB3& Box0, const FRigidTransform3& Transform0, const TriMeshType & TriangleMesh1, const FRigidTransform3& Transformw1, const FReal CullDistance, FRigidBodyPointContactConstraint& Constraint);
-
-		void CHAOS_API UpdateBoxTriangleMeshManifold(FCollisionConstraintBase&  Constraint, const FRigidTransform3& ATM, const FRigidTransform3& BTM, const FReal CullDistance);
-
-		void CHAOS_API ConstructBoxTriangleMeshConstraints(TGeometryParticleHandle<FReal, 3>* Particle0, TGeometryParticleHandle<FReal, 3>* Particle1, const FImplicitObject* Implicit0, const FImplicitObject* Implicit1, const FRigidTransform3& Transform0, const FRigidTransform3& Transform1, const FReal CullDistance, FCollisionConstraintsArray& NewConstraints);
-
-
-		//
-		// Sphere-Sphere
-		//
-
-		void CHAOS_API UpdateSphereSphereConstraint(const TSphere<FReal, 3>& Sphere1, const FRigidTransform3& Sphere1Transform, const TSphere<FReal, 3>& Sphere2, const FRigidTransform3& Sphere2Transform, const FReal CullDistance, FRigidBodyPointContactConstraint& Constraint);
-
-		void CHAOS_API UpdateSphereSphereManifold(FCollisionConstraintBase&  Constraint, const FRigidTransform3& ATM, const FRigidTransform3& BTM, const FReal CullDistance);
-
-		void CHAOS_API ConstructSphereSphereConstraints(TGeometryParticleHandle<FReal, 3>* Particle0, TGeometryParticleHandle<FReal, 3>* Particle1, const FImplicitObject* Implicit0, const FImplicitObject* Implicit1, const FRigidTransform3& Transform0, const FRigidTransform3& Transform1, const FReal CullDistance, FCollisionConstraintsArray& NewConstraints);
-
-		//
-		// Sphere-HeightField
-		//
-
-		void CHAOS_API UpdateSphereHeightFieldConstraint(const TSphere<FReal, 3>& A, const FRigidTransform3& ATransform, const THeightField<FReal>& B, const FRigidTransform3& BTransform, const FReal CullDistance, FRigidBodyPointContactConstraint& Constraint);
-
-		void CHAOS_API UpdateSphereHeightFieldManifold(FCollisionConstraintBase&  Constraint, const FRigidTransform3& ATM, const FRigidTransform3& BTM, const FReal CullDistance);
-
-		void CHAOS_API ConstructSphereHeightFieldConstraints(TGeometryParticleHandle<FReal, 3>* Particle0, TGeometryParticleHandle<FReal, 3>* Particle1, const FImplicitObject* Implicit0, const FImplicitObject* Implicit1, const FRigidTransform3& Transform0, const FRigidTransform3& Transform1, const FReal CullDistance, FCollisionConstraintsArray& NewConstraints);
-
-		//
-		// Sphere-Plane
-		//
-
-		void CHAOS_API UpdateSpherePlaneConstraint(const TSphere<FReal, 3>& Sphere, const FRigidTransform3& SphereTransform, const TPlane<FReal, 3>& Plane, const FRigidTransform3& PlaneTransform, const FReal CullDistance, FRigidBodyPointContactConstraint& Constraint);
-
-		void CHAOS_API UpdateSpherePlaneManifold(FCollisionConstraintBase&  Constraint, const FRigidTransform3& ATM, const FRigidTransform3& BTM, const FReal CullDistance);
-
-		void CHAOS_API ConstructSpherePlaneConstraints(TGeometryParticleHandle<FReal, 3>* Particle0, TGeometryParticleHandle<FReal, 3>* Particle1, const FImplicitObject* Implicit0, const FImplicitObject* Implicit1, const FRigidTransform3& Transform0, const FRigidTransform3& Transform1, const FReal CullDistance, FCollisionConstraintsArray& NewConstraints);
-
-		//
-		// Sphere-Box
-		//
-
-		void CHAOS_API UpdateSphereBoxConstraint(const TSphere<FReal, 3>& Sphere, const FRigidTransform3& SphereTransform, const FAABB3& Box, const FRigidTransform3& BoxTransform, const FReal CullDistance, FRigidBodyPointContactConstraint& Constraint);
-
-		void CHAOS_API UpdateSphereBoxManifold(FCollisionConstraintBase&  Constraint, const FRigidTransform3& ATM, const FRigidTransform3& BTM, const FReal CullDistance);
-
-		void CHAOS_API ConstructSphereBoxConstraints(TGeometryParticleHandle<FReal, 3>* Particle0, TGeometryParticleHandle<FReal, 3>* Particle1, const FImplicitObject* Implicit0, const FImplicitObject* Implicit1, const FRigidTransform3& Transform0, const FRigidTransform3& Transform1, const FReal CullDistance, FCollisionConstraintsArray& NewConstraints);
-
-		//
-		// Sphere-Capsule
-		//
-
-		void CHAOS_API UpdateSphereCapsuleConstraint(const TSphere<FReal, 3>& Sphere, const FRigidTransform3& SphereTransform, const TCapsule<FReal>& Box, const FRigidTransform3& BoxTransform, const FReal CullDistance, FRigidBodyPointContactConstraint& Constraint);
-
-		void CHAOS_API UpdateSphereCapsuleManifold(FCollisionConstraintBase&  Constraint, const FRigidTransform3& ATM, const FRigidTransform3& BTM, const FReal CullDistance);
-
-		void CHAOS_API ConstructSphereCapsuleConstraints(TGeometryParticleHandle<FReal, 3>* Particle0, TGeometryParticleHandle<FReal, 3>* Particle1, const FImplicitObject* Implicit0, const FImplicitObject* Implicit1, const FRigidTransform3& Transform0, const FRigidTransform3& Transform1, const FReal CullDistance, FCollisionConstraintsArray& NewConstraints);
-
-		//
-		// Capsule-Capsule
-		//
-
-		void CHAOS_API UpdateCapsuleCapsuleConstraint(const TCapsule<FReal>& A, const FRigidTransform3& ATransform, const TCapsule<FReal>& B, const FRigidTransform3& BTransform, const FReal CullDistance, FRigidBodyPointContactConstraint& Constraint);
-
-		void CHAOS_API UpdateCapsuleCapsuleManifold(FCollisionConstraintBase&  Constraint, const FRigidTransform3& ATM, const FRigidTransform3& BTM, const FReal CullDistance);
-
-		void CHAOS_API ConstructCapsuleCapsuleConstraints(TGeometryParticleHandle<FReal, 3>* Particle0, TGeometryParticleHandle<FReal, 3>* Particle1, const FImplicitObject* Implicit0, const FImplicitObject* Implicit1, const FRigidTransform3& Transform0, const FRigidTransform3& Transform1, const FReal CullDistance, FCollisionConstraintsArray& NewConstraints);
-
-		//
-		// Capsule-Box
-		//
-
-		void CHAOS_API UpdateCapsuleBoxConstraint(const TCapsule<FReal>& A, const FRigidTransform3& ATransform, const FAABB3& B, const FRigidTransform3& BTransform, const FReal CullDistance, FRigidBodyPointContactConstraint& Constraint);
-
-		void CHAOS_API ConstructCapsuleBoxConstraints(TGeometryParticleHandle<FReal, 3>* Particle0, TGeometryParticleHandle<FReal, 3>* Particle1, const FImplicitObject* Implicit0, const FImplicitObject* Implicit1, const FRigidTransform3& Transform0, const FRigidTransform3& Transform1, const FReal CullDistance, const FCollisionContext& Context, FCollisionConstraintsArray& NewConstraints);
-
-		//
-		// Capsule-HeightField
-		//
-
-		void CHAOS_API UpdateCapsuleHeightFieldConstraint(const TCapsule<FReal>& A, const FRigidTransform3& ATransform, const THeightField<FReal>& B, const FRigidTransform3& BTransform, const FReal CullDistance, FRigidBodyPointContactConstraint& Constraint);
-
-		void CHAOS_API UpdateCapsuleHeightFieldManifold(FCollisionConstraintBase&  Constraint, const FRigidTransform3& ATM, const FRigidTransform3& BTM, const FReal CullDistance);
-
-		void ConstructCapsuleHeightFieldConstraints(TGeometryParticleHandle<FReal, 3>* Particle0, TGeometryParticleHandle<FReal, 3>* Particle1, const FImplicitObject* Implicit0, const FImplicitObject* Implicit1, const FRigidTransform3& Transform0, const FRigidTransform3& Transform1, const FReal CullDistance, FCollisionConstraintsArray& NewConstraints);
-
-		//
-		// Convex-Convex
-		//
-		void CHAOS_API UpdateConvexConvexConstraint(const FImplicitObject& A, const FRigidTransform3& ATM, const FImplicitObject& B, const FRigidTransform3& BTM, const FReal CullDistance, FCollisionConstraintBase& Constraint);
-
-		void CHAOS_API UpdateConvexConvexManifold(FCollisionConstraintBase&  Constraint, const FRigidTransform3& ATM, const FRigidTransform3& BTM, const FReal CullDistance);
-
-		void CHAOS_API ConstructConvexConvexConstraints(TGeometryParticleHandle<FReal, 3>* Particle0, TGeometryParticleHandle<FReal, 3>* Particle1, const FImplicitObject* Implicit0, const FImplicitObject* Implicit1, const FRigidTransform3& Transform0, const FRigidTransform3& Transform1, const FReal CullDistance, FCollisionConstraintsArray& NewConstraints);
-
-		//
-		// Convex-HeightField
-		//
-
-		void CHAOS_API UpdateConvexHeightFieldConstraint(const FImplicitObject& A, const FRigidTransform3& ATransform, const THeightField<FReal>& B, const FRigidTransform3& BTransform, const FReal CullDistance, FRigidBodyPointContactConstraint& Constraint);
-
-		void CHAOS_API UpdateConvexHeightFieldManifold(FCollisionConstraintBase&  Constraint, const FRigidTransform3& ATM, const FRigidTransform3& BTM, const FReal CullDistance);
-
-		void CHAOS_API ConstructConvexHeightFieldConstraints(TGeometryParticleHandle<FReal, 3>* Particle0, TGeometryParticleHandle<FReal, 3>* Particle1, const FImplicitObject* Implicit0, const FImplicitObject* Implicit1, const FRigidTransform3& Transform0, const FRigidTransform3& Transform1, const FReal CullDistance, FCollisionConstraintsArray& NewConstraints);
-
-
-		//
-		// Levelset-Levelset
-		//
-		template<ECollisionUpdateType UpdateType>
-		void CHAOS_API UpdateLevelsetLevelsetConstraint(const FReal CullDistance, FRigidBodyPointContactConstraint& Constraint);
-
-		void CHAOS_API UpdateLevelsetLevelsetManifold(FCollisionConstraintBase&  Constraint, const FRigidTransform3& ATM, const FRigidTransform3& BTM, const FReal CullDistance);
-
-		void CHAOS_API ConstructLevelsetLevelsetConstraints(TGeometryParticleHandle<FReal, 3>* Particle0, TGeometryParticleHandle<FReal, 3>* Particle1, const FImplicitObject* Implicit0, const FImplicitObject* Implicit1, const FRigidTransform3& Transform0, const FRigidTransform3& Transform1, const FReal CullDistance, FCollisionConstraintsArray& NewConstraints);
-
-
 		//
 		// Constraint API
 		//
-		void CHAOS_API UpdateSingleShotManifold(FRigidBodyMultiPointContactConstraint&  Constraint, const FRigidTransform3& Transform0, const FRigidTransform3& Transform1, const FReal CullDistance);
-
-		void CHAOS_API UpdateIterativeManifold(FRigidBodyMultiPointContactConstraint&  Constraint, const FRigidTransform3& Transform0, const FRigidTransform3& Transform1, const FReal CullDistance);
 
+		// Build a contact manifold for the shape pair. The manifold is a plane attached to one shape, and points attached to the other.
 		void CHAOS_API UpdateManifold(FRigidBodyMultiPointContactConstraint& Constraint, const FRigidTransform3& ATM, const FRigidTransform3& BTM, const FReal CullDistance, const FCollisionContext& Context);
 
-		void UpdateConstraintFromManifold(FRigidBodyMultiPointContactConstraint& Constraint, const FRigidTransform3& Transform0, const FRigidTransform3& Transform1, const FReal CullDistance);
+		// Update the constraint using the pre-built manifold, finding the manifold point that is most deeply penetrating the manifold plane.
+		void CHAOS_API UpdateConstraintFromManifold(FRigidBodyMultiPointContactConstraint& Constraint, const FRigidTransform3& Transform0, const FRigidTransform3& Transform1, const FReal CullDistance);
 
+		// Update the constraint by re-running collision detection on the shape pair.
 		template<ECollisionUpdateType UpdateType>
-		void UpdateConstraintFromGeometry(FRigidBodyPointContactConstraint& Constraint, const FRigidTransform3& ParticleTransform0, const FRigidTransform3& ParticleTransform1, const FReal CullDistance);
+		void CHAOS_API UpdateConstraintFromGeometry(FRigidBodyPointContactConstraint& Constraint, const FRigidTransform3& ParticleTransform0, const FRigidTransform3& ParticleTransform1, const FReal CullDistance);
 
+		// Create constraints for the particle pair. This could create multiple constraints: one for each potentially colliding shape pair in multi-shape particles.
 		void CHAOS_API ConstructConstraints(TGeometryParticleHandle<FReal, 3>* Particle0, TGeometryParticleHandle<FReal, 3>* Particle1, const FImplicitObject* Implicit0, const FImplicitObject* Implicit1, const FRigidTransform3& Transform0, const FRigidTransform3& Transform1, const FReal CullDistance, const FCollisionContext& Context, FCollisionConstraintsArray& NewConstraints);
 
+
+		// @todo(chaos): this is only called in tests - should it really be exposed?
+		template<ECollisionUpdateType UpdateType>
+		void CHAOS_API UpdateLevelsetLevelsetConstraint(const FReal CullDistance, FRigidBodyPointContactConstraint& Constraint);
+
 	}
 }
diff --git a/Engine/Source/Runtime/Experimental/Chaos/Public/Chaos/Evolution/PBDMinEvolution.h b/Engine/Source/Runtime/Experimental/Chaos/Public/Chaos/Evolution/PBDMinEvolution.h
--- a/Engine/Source/Runtime/Experimental/Chaos/Public/Chaos/Evolution/PBDMinEvolution.h
+++ b/Engine/Source/Runtime/Experimental/Chaos/Public/Chaos/Evolution/PBDMinEvolution.h
@@ -7,15 +7,9 @@
 
 namespace Chaos
 {
-	class FNarrowPhase;
-	class FParticlePairBroadPhase;
+	class FParticlePairCollisionDetector;
 	class FPBDCollisionConstraints;
 	class FSimpleConstraintRule;
-	class FSyncCollisionReceiver;
-
-
-	template<typename T_BROADPHASE, typename T_NARROWPHASE, typename T_RECEIVER, typename T_CONTAINER>
-	class TCollisionDetector;
 
 	template <typename T, int d>
 	class TPBDRigidsSOAs;
@@ -33,7 +27,7 @@ namespace Chaos
 	public:
 		// @todo(ccaulfield): make it so that CollisionDetection is plugged in with a constraint rule...
 
-		using FCollisionDetector = TCollisionDetector<FParticlePairBroadPhase, FNarrowPhase, FSyncCollisionReceiver, FPBDCollisionConstraints>;
+		using FCollisionDetector = FParticlePairCollisionDetector;
 		using FEvolutionCallback = TFunction<void()>;
 		using FRigidParticleSOAs = TPBDRigidsSOAs<FReal, 3>;
 
diff --git a/Engine/Source/Runtime/Experimental/Chaos/Public/Chaos/PBDRigidsEvolutionGBF.h b/Engine/Source/Runtime/Experimental/Chaos/Public/Chaos/PBDRigidsEvolutionGBF.h
--- a/Engine/Source/Runtime/Experimental/Chaos/Public/Chaos/PBDRigidsEvolutionGBF.h
+++ b/Engine/Source/Runtime/Experimental/Chaos/Public/Chaos/PBDRigidsEvolutionGBF.h
@@ -2,19 +2,18 @@
 #pragma once
 
 #include "Chaos/ChaosPerfTest.h"
-#include "Chaos/Collision/CollisionDetector.h"
-#include "Chaos/Collision/CollisionReceiver.h"
 #include "Chaos/Collision/NarrowPhase.h"
 #include "Chaos/Collision/SpatialAccelerationBroadPhase.h"
+#include "Chaos/Collision/SpatialAccelerationCollisionDetector.h"
 #include "Chaos/PBDCollisionConstraints.h"
 #include "Chaos/PBDRigidsEvolution.h"
+#include "Chaos/PerParticleAddImpulses.h"
 #include "Chaos/PerParticleEtherDrag.h"
 #include "Chaos/PerParticleEulerStepVelocity.h"
 #include "Chaos/PerParticleExternalForces.h"
 #include "Chaos/PerParticleGravity.h"
 #include "Chaos/PerParticleInitForce.h"
 #include "Chaos/PerParticlePBDEulerStep.h"
-#include "Chaos/PerParticleAddImpulses.h"
 
 namespace Chaos
 {
@@ -40,7 +39,7 @@ namespace Chaos
 		using FGravityForces = TPerParticleGravity<FReal, 3>;
 		using FCollisionConstraints = FPBDCollisionConstraints;
 		using FCollisionConstraintRule = TPBDConstraintColorRule<FCollisionConstraints>;
-		using FCollisionDetector = TCollisionDetector<FSpatialAccelerationBroadPhase, FNarrowPhase, FAsyncCollisionReceiver, FCollisionConstraints>;
+		using FCollisionDetector = FSpatialAccelerationCollisionDetector;
 		using FExternalForces = TPerParticleExternalForces<FReal, 3>;
 
 		static constexpr int32 DefaultNumIterations = 1;
@@ -188,7 +187,8 @@ namespace Chaos
 		FCollisionConstraints CollisionConstraints;
 		FCollisionConstraintRule CollisionRule;
 		FSpatialAccelerationBroadPhase BroadPhase;
-		FCollisionDetector CollisionDetector;
+		FNarrowPhase NarrowPhase;
+		FSpatialAccelerationCollisionDetector CollisionDetector;
 
 		FPBDRigidsEvolutionCallback PostIntegrateCallback;
 		FPBDRigidsEvolutionCallback PostDetectCollisionsCallback;
diff --git a/Engine/Source/Runtime/Experimental/ChaosSolvers/Private/PBDRigidsSolver.cpp b/Engine/Source/Runtime/Experimental/ChaosSolvers/Private/PBDRigidsSolver.cpp
--- a/Engine/Source/Runtime/Experimental/ChaosSolvers/Private/PBDRigidsSolver.cpp
+++ b/Engine/Source/Runtime/Experimental/ChaosSolvers/Private/PBDRigidsSolver.cpp
@@ -46,6 +46,12 @@ FAutoConsoleVariableRef CVarChaosSolverCollisionDefaultPushoutIterations(TEXT("p
 int32 ChaosSolverCleanupCommandsOnDestruction = 1;
 FAutoConsoleVariableRef CVarChaosSolverCleanupCommandsOnDestruction(TEXT("p.Chaos.Solver.CleanupCommandsOnDestruction"), ChaosSolverCleanupCommandsOnDestruction, TEXT("Whether or not to run internal command queue cleanup on solver destruction (0 = no cleanup, >0 = cleanup all commands)"));
 
+int32 ChaosSolverCollisionDeferNarrowPhase = 1;
+FAutoConsoleVariableRef CVarChaosSolverCollisionDeferNarrowPhase(TEXT("p.Chaos.Solver.Collision.DeferNarrowPhase"), ChaosSolverCollisionDeferNarrowPhase, TEXT("Create contacts for all broadphase pairs, perform NarrowPhase later."));
+
+int32 ChaosSolverCollisionUseManifolds = 0;
+FAutoConsoleVariableRef CVarChaosSolverCollisionUseManifolds(TEXT("p.Chaos.Solver.Collision.UseManifolds"), ChaosSolverCollisionUseManifolds, TEXT("Enable/Disable use of manifoldes in collision."));
+
 
 
 namespace Chaos
@@ -555,13 +561,15 @@ namespace Chaos
 
 	void FPBDRigidsSolver::AdvanceSolverBy(float DeltaTime)
 	{
+		MEvolution->GetCollisionDetector().GetNarrowPhase().GetContext().bDeferUpdate = (ChaosSolverCollisionDeferNarrowPhase != 0);
+		MEvolution->GetCollisionDetector().GetNarrowPhase().GetContext().bAllowManifolds = (ChaosSolverCollisionUseManifolds != 0);
+
 		UE_LOG(LogPBDRigidsSolver, Verbose, TEXT("PBDRigidsSolver::Tick(%3.5f)"), DeltaTime);
 		if (bEnabled)
 		{
 			MLastDt = DeltaTime;
 			AdvanceOneTimeStepTask(this, DeltaTime).DoWork();
 		}
-
 	}
 
 	void FPBDRigidsSolver::SyncEvents_GameThread()
