Fix Static Analysis.

diff --git a/Engine/Source/Runtime/Experimental/GeometryCollectionSimulationCore/Private/GeometryCollection/GeometryCollectionCollisionStructureManager.cpp b/Engine/Source/Runtime/Experimental/GeometryCollectionSimulationCore/Private/GeometryCollection/GeometryCollectionCollisionStructureManager.cpp
--- a/Engine/Source/Runtime/Experimental/GeometryCollectionSimulationCore/Private/GeometryCollection/GeometryCollectionCollisionStructureManager.cpp
+++ b/Engine/Source/Runtime/Experimental/GeometryCollectionSimulationCore/Private/GeometryCollection/GeometryCollectionCollisionStructureManager.cpp
@@ -173,11 +173,14 @@ void FCollisionStructureManager::UpdateImplicitFlags(
 	FImplicit* Implicit, 
 	const ECollisionTypeEnum CollisionType)
 {
-	Implicit->SetCollsionType(CollisionType==ECollisionTypeEnum::Chaos_Surface_Volumetric?Chaos::ImplicitObjectType::LevelSet:Implicit->GetType());
-	if (Implicit && (CollisionType == ECollisionTypeEnum::Chaos_Surface_Volumetric) && Implicit->GetType() == Chaos::ImplicitObjectType::LevelSet)
+	if (ensure(Implicit))
 	{
-		Implicit->SetDoCollide(false);
-		Implicit->SetConvex(false);
+		Implicit->SetCollsionType(CollisionType == ECollisionTypeEnum::Chaos_Surface_Volumetric ? Chaos::ImplicitObjectType::LevelSet : Implicit->GetType());
+		if (Implicit && (CollisionType == ECollisionTypeEnum::Chaos_Surface_Volumetric) && Implicit->GetType() == Chaos::ImplicitObjectType::LevelSet)
+		{
+			Implicit->SetDoCollide(false);
+			Implicit->SetConvex(false);
+		}
 	}
 }
 
