reworked parameters for building connectivity graph so values are passed through.

diff --git a/Engine/Source/Runtime/Experimental/Chaos/Public/Chaos/ClusterCreationParameters.h b/Engine/Source/Runtime/Experimental/Chaos/Public/Chaos/ClusterCreationParameters.h
--- a/Engine/Source/Runtime/Experimental/Chaos/Public/Chaos/ClusterCreationParameters.h
+++ b/Engine/Source/Runtime/Experimental/Chaos/Public/Chaos/ClusterCreationParameters.h
@@ -14,7 +14,7 @@ namespace Chaos
 	template <typename T>
 	struct CHAOS_API FClusterCreationParameters
 	{
-		enum EConnectionMethod { None = 0, PointImplicit, DelaunayTriangulation, MinimalSpanningSubsetDelaunayTriangulation, PointImplicitAugmentedWithMinimalDelaunay };
+		enum EConnectionMethod { PointImplicit = 0, DelaunayTriangulation, MinimalSpanningSubsetDelaunayTriangulation, PointImplicitAugmentedWithMinimalDelaunay,  None };
 
 
 		FClusterCreationParameters(
@@ -22,8 +22,7 @@ namespace Chaos
 			, int32 MaxNumConnectionsIn = 100
 			, bool bCleanCollisionParticlesIn = true
 			, bool bCopyCollisionParticlesIn = true
-			, bool bGenerateConnectionGraphIn = true
-			, EConnectionMethod ConnectionMethodIn = EConnectionMethod::PointImplicitAugmentedWithMinimalDelaunay
+			, bool bGenerateConnectionGraphIn = true, EConnectionMethod ConnectionMethodIn = EConnectionMethod::MinimalSpanningSubsetDelaunayTriangulation
 			, TBVHParticles<float, 3>* CollisionParticlesIn = nullptr
 			, Chaos::TPBDRigidClusteredParticleHandle<float,3>* ClusterParticleHandleIn = nullptr
 		)
diff --git a/Engine/Source/Runtime/Experimental/ChaosSolvers/Private/PhysicsProxy/GeometryCollectionPhysicsProxy.cpp b/Engine/Source/Runtime/Experimental/ChaosSolvers/Private/PhysicsProxy/GeometryCollectionPhysicsProxy.cpp
--- a/Engine/Source/Runtime/Experimental/ChaosSolvers/Private/PhysicsProxy/GeometryCollectionPhysicsProxy.cpp
+++ b/Engine/Source/Runtime/Experimental/ChaosSolvers/Private/PhysicsProxy/GeometryCollectionPhysicsProxy.cpp
@@ -922,7 +922,12 @@ void FGeometryCollectionPhysicsProxy::InitializeBodiesPT(Chaos::FPBDRigidsSolver
 				{
 					if (SolverClusterHandles[TransformGroupIndex])
 					{
-						GetSolver()->GetEvolution()->GetRigidClustering().GenerateConnectionGraph(SolverClusterHandles[TransformGroupIndex]);
+						Chaos::FClusterCreationParameters<float> ClusterParams;
+						// #todo: should other parameters be set here?  Previously, there was no parameters being sent, and it is unclear
+						// where some of these parameters are defined (ie: CoillisionThicknessPercent)
+						ClusterParams.ConnectionMethod = Parameters.ClusterConnectionMethod;
+						
+						GetSolver()->GetEvolution()->GetRigidClustering().GenerateConnectionGraph(SolverClusterHandles[TransformGroupIndex], ClusterParams);
 					}
 				}
 			}
