Detect desync in the common case from GT

diff --git a/Engine/Source/Programs/HeadlessChaos/Private/HeadlessChaosTestRewind.cpp b/Engine/Source/Programs/HeadlessChaos/Private/HeadlessChaosTestRewind.cpp
--- a/Engine/Source/Programs/HeadlessChaos/Private/HeadlessChaosTestRewind.cpp
+++ b/Engine/Source/Programs/HeadlessChaos/Private/HeadlessChaosTestRewind.cpp
@@ -85,7 +85,7 @@ namespace ChaosTest {
 		//check state at every step except latest
 		for(int Step = 0; Step < 10; ++Step)
 		{
-			const auto ParticleState = RewindData->GetStateAtFrame(*Particle,Step);
+			const auto ParticleState = RewindData->GetPastStateAtFrame(*Particle,Step);
 			EXPECT_EQ(ParticleState.X()[2],100 - Step);
 
 			if(Step < 3)
@@ -146,7 +146,7 @@ namespace ChaosTest {
 		//check state at every step except latest
 		for(int Step = 0; Step < 10; ++Step)
 		{
-			const auto ParticleState = RewindData->GetStateAtFrame(*Particle,Step);
+			const auto ParticleState = RewindData->GetPastStateAtFrame(*Particle,Step);
 			EXPECT_EQ(ParticleState.F()[2],Step+1);
 		}
 
@@ -197,7 +197,7 @@ namespace ChaosTest {
 		//check state at every step except latest
 		for(int Step = 0; Step < 10; ++Step)
 		{
-			const auto ParticleState = RewindData->GetStateAtFrame(*Particle,Step);
+			const auto ParticleState = RewindData->GetPastStateAtFrame(*Particle,Step);
 
 			if(Step == 3)
 			{
@@ -267,7 +267,7 @@ namespace ChaosTest {
 		//check state at every step except latest
 		for(int Step = 0; Step < 10; ++Step)
 		{
-			const auto ParticleState = RewindData->GetStateAtFrame(*Particle,Step);
+			const auto ParticleState = RewindData->GetPastStateAtFrame(*Particle,Step);
 
 			
 			if(Step < 3)
@@ -337,7 +337,7 @@ namespace ChaosTest {
 
 		for(int Step = 0; Step < 9; ++Step)
 		{
-			const auto ParticleState = RewindData->GetStateAtFrame(*Particle,Step);
+			const auto ParticleState = RewindData->GetPastStateAtFrame(*Particle,Step);
 			
 			EXPECT_EQ(ParticleState.X()[2],X[Step][2]);
 			EXPECT_EQ(ParticleState.V()[2],V[Step][2]);
@@ -396,7 +396,9 @@ namespace ChaosTest {
 		//make sure recorded data is still valid even at head
 		for(int Step = 0; Step < 11; ++Step)
 		{
-			const FGeometryParticleState State = RewindData->GetStateAtFrame(*Particle,Step);
+			FGeometryParticleState State(*Particle);
+			const EFutureQueryResult Status = RewindData->GetFutureStateAtFrame(State,Step);
+			EXPECT_EQ(Status,EFutureQueryResult::Ok);
 			EXPECT_EQ(State.X()[2],X[Step][2]);
 			EXPECT_EQ(State.V()[2],V[Step][2]);
 		}
@@ -453,7 +455,7 @@ namespace ChaosTest {
 		FRewindData* RewindData = Solver->GetRewindData();
 
 		{
-			const FGeometryParticleState State = RewindData->GetStateAtFrame(*Particle,5);
+			const FGeometryParticleState State = RewindData->GetPastStateAtFrame(*Particle,5);
 			EXPECT_EQ(State.X(),X[5]);
 		}
 
@@ -462,7 +464,7 @@ namespace ChaosTest {
 
 		// State should be the same as being at head because we removed it from solver
 		{
-			const FGeometryParticleState State = RewindData->GetStateAtFrame(*Particle,5);
+			const FGeometryParticleState State = RewindData->GetPastStateAtFrame(*Particle,5);
 			EXPECT_EQ(Particle->X(), State.X());
 		}
 
@@ -619,16 +621,22 @@ namespace ChaosTest {
 		Kinematic->SetX(FVec3(2,2,2));
 
 		TArray<FVec3> X;
+		const int32 LastStep = 12;
 
-		for(int Step = 0; Step < 10; ++Step)
+		for(int Step = 0; Step <= LastStep; ++Step)
 		{
 			X.Add(Particle->X());
 
-			if(Step == 9)
+			if(Step == 8)
 			{
 				Kinematic->SetX(FVec3(50,50,50));
 			}
 
+			if(Step == 10)
+			{
+				Kinematic->SetX(FVec3(60,60,60));
+			}
+
 			TickSolverHelper(Module,Solver);
 		}
 
@@ -639,23 +647,55 @@ namespace ChaosTest {
 
 		//Move particle and rerun
 		Particle->SetX(FVec3(0,0,100));
-		for(int Step = RewindStep; Step < 10; ++Step)
+		for(int Step = RewindStep; Step <= LastStep; ++Step)
 		{
+			if(Step == 8)
+			{
+				Kinematic->SetX(FVec3(50));
+			}
+
 			X[Step] = Particle->X();
 			TickSolverHelper(Module,Solver);
+
+			//see that particle has desynced
+			if(Step < LastStep)
+			{
+				//If we're still in the past make sure future has been marked as desync
+				FGeometryParticleState State(*Particle);
+				EXPECT_EQ(EFutureQueryResult::Desync,RewindData->GetFutureStateAtFrame(State,Step));
+
+				FGeometryParticleState KinState(*Kinematic);
+				const EFutureQueryResult KinFutureStatus = RewindData->GetFutureStateAtFrame(KinState,Step);
+				if(Step < 10)
+				{
+					EXPECT_EQ(KinFutureStatus,EFutureQueryResult::Ok);
+				}
+				else
+				{
+					//todo: desync from lack of input
+					//EXPECT_EQ(KinFutureStatus,EFutureQueryResult::Desync);
+				}
+			}
 		}
 
-		EXPECT_EQ(Kinematic->X()[2],2);	//Rewound kinematic and never updated it so back to original value
+		EXPECT_EQ(Kinematic->X()[2],50);	//Rewound kinematic and only did one update, so use that first update
 
 		//Make sure we recorded the new data
-		for(int Step = RewindStep; Step < 10; ++Step)
+		for(int Step = RewindStep; Step <= LastStep; ++Step)
 		{
-			const FGeometryParticleState State = RewindData->GetStateAtFrame(*Particle,Step);
+			const FGeometryParticleState State = RewindData->GetPastStateAtFrame(*Particle,Step);
 			EXPECT_EQ(State.X()[2],X[Step][2]);
 
-			const FGeometryParticleState KinState = RewindData->GetStateAtFrame(*Kinematic,Step);
-			//TODO: this is a form of desync, need support for that still
-			//EXPECT_EQ(KinState.X()[2],2);	//even though not dirty, data is properly stored as not moving
+			const FGeometryParticleState KinState = RewindData->GetPastStateAtFrame(*Kinematic,Step);
+			if(Step < 8)
+			{
+				EXPECT_EQ(KinState.X()[2],2);
+			}
+			else
+			{
+				//todo: desync from lack of input
+				//EXPECT_EQ(KinState.X()[2],50);	//in resim we didn't do second move, so recorded data must be updated
+			}
 		}
 
 
diff --git a/Engine/Source/Runtime/Experimental/Chaos/Public/Chaos/ParticleDirtyFlags.h b/Engine/Source/Runtime/Experimental/Chaos/Public/Chaos/ParticleDirtyFlags.h
--- a/Engine/Source/Runtime/Experimental/Chaos/Public/Chaos/ParticleDirtyFlags.h
+++ b/Engine/Source/Runtime/Experimental/Chaos/Public/Chaos/ParticleDirtyFlags.h
@@ -26,6 +26,11 @@ struct FParticlePositionRotation
 	{
 		Ar << X << R;
 	}
+
+	bool operator==(const FParticlePositionRotation& Other) const
+	{
+		return X == Other.X && R == Other.R;
+	}
 };
 
 inline FChaosArchive& operator<<(FChaosArchive& Ar,FParticlePositionRotation& Data)
@@ -41,6 +46,11 @@ struct FParticleVelocities
 		Ar << V << W;
 	}
 
+	bool operator==(const FParticleVelocities& Other) const
+	{
+		return V == Other.V && W == Other.W;
+	}
+
 	FVec3 V;
 	FVec3 W;
 };
@@ -65,6 +75,14 @@ struct FParticleDynamics
 		Ar << LinearImpulse;
 		Ar << AngularImpulse;	
 	}
+
+	bool operator==(const FParticleDynamics& Other) const
+	{
+		return F == Other.F
+			&& Torque == Other.Torque
+			&& LinearImpulse == Other.LinearImpulse
+			&& AngularImpulse == Other.AngularImpulse;
+	}
 };
 
 inline FChaosArchive& operator<<(FChaosArchive& Ar, FParticleDynamics& Data)
@@ -91,6 +109,16 @@ struct FParticleMassProps
 		Ar << M;
 		Ar << InvM;
 	}
+
+	bool operator==(const FParticleMassProps& Other) const
+	{
+		return CenterOfMass == Other.CenterOfMass
+			&& RotationOfMass == Other.RotationOfMass
+			&& I == Other.I
+			&& InvI == Other.InvI
+			&& M == Other.M
+			&& InvM == Other.InvM;
+	}
 };
 
 inline FChaosArchive& operator<<(FChaosArchive& Ar,FParticleMassProps& Data)
@@ -120,6 +148,15 @@ struct FParticleMisc
 		Ar << GravityEnabled;
 		bGravityEnabled = GravityEnabled;
 	}
+
+	bool operator==(const FParticleMisc& Other) const
+	{
+		return CollisionGroup == Other.CollisionGroup
+			&& ObjectState == Other.ObjectState
+			&& SpatialIdx == Other.SpatialIdx
+			&& bDisabled == Other.bDisabled
+			&& bGravityEnabled == Other.bGravityEnabled;
+	}
 };
 
 inline FChaosArchive& operator<<(FChaosArchive& Ar,FParticleMisc& Data)
@@ -150,6 +187,19 @@ struct FParticleNonFrequentData
 	{
 		Ar << Geometry;
 	}
+
+	bool operator==(const FParticleNonFrequentData& Other) const
+	{
+		return Geometry == Other.Geometry
+			&& UserData == Other.UserData
+			&& UniqueIdx == Other.UniqueIdx
+			&& LinearEtherDrag == Other.LinearEtherDrag
+			&& AngularEtherDrag == Other.AngularEtherDrag
+#if CHAOS_CHECKED
+			&& DebugName == Other.DebugName
+#endif
+			;
+	}
 };
 
 inline FChaosArchive& operator<<(FChaosArchive& Ar,FParticleNonFrequentData& Data)
diff --git a/Engine/Source/Runtime/Experimental/ChaosSolvers/Private/PBDRigidsSolver.cpp b/Engine/Source/Runtime/Experimental/ChaosSolvers/Private/PBDRigidsSolver.cpp
--- a/Engine/Source/Runtime/Experimental/ChaosSolvers/Private/PBDRigidsSolver.cpp
+++ b/Engine/Source/Runtime/Experimental/ChaosSolvers/Private/PBDRigidsSolver.cpp
@@ -804,7 +804,15 @@ namespace Chaos
 
 				if(RewindData)
 				{
-					RewindData->PushGTDirtyData(*Manager,DataIdx,Dirty);
+					//may want to remove branch by templatizing lambda
+					if(RewindData->IsResim())
+					{
+						RewindData->PushGTDirtyData<true>(*Manager,DataIdx,Dirty);
+					}
+					else
+					{
+						RewindData->PushGTDirtyData<false>(*Manager,DataIdx,Dirty);
+					}
 				}
 
 				Proxy->PushToPhysicsState(*Manager, DataIdx, Dirty, ShapeDirtyData);
@@ -1056,7 +1064,15 @@ namespace Chaos
 				{
 					if(ActiveObject.GetParticleType() == EParticleType::Rigid)
 					{
-						MRewindData->PushPTDirtyData(*static_cast<const FRigidParticlePhysicsProxy*>(Proxy)->GetHandle(), DataIdx++);
+						//may want to remove branch using templates outside loop
+						if(MRewindData->IsResim())
+						{
+							MRewindData->PushPTDirtyData<true>(*static_cast<const FRigidParticlePhysicsProxy*>(Proxy)->GetHandle(),DataIdx++);
+						}
+						else
+						{
+							MRewindData->PushPTDirtyData<false>(*static_cast<const FRigidParticlePhysicsProxy*>(Proxy)->GetHandle(),DataIdx++);
+						}
 					}
 				}
 			}
diff --git a/Engine/Source/Runtime/Experimental/ChaosSolvers/Public/RewindData.h b/Engine/Source/Runtime/Experimental/ChaosSolvers/Public/RewindData.h
--- a/Engine/Source/Runtime/Experimental/ChaosSolvers/Public/RewindData.h
+++ b/Engine/Source/Runtime/Experimental/ChaosSolvers/Public/RewindData.h
@@ -48,7 +48,31 @@ public:
 	{
 		Manager = &InManager;
 		Idx = InIdx;
-		SetFunc(Manager->GetParticlePool<T,PropName>().GetElement(Idx));
+		T& NewVal = Manager->GetParticlePool<T,PropName>().GetElement(Idx);
+		SetFunc(NewVal);
+	}
+
+	bool IsDesynced(const FDirtyPropertiesManager& SrcManager, const int32 DataIdxIn, const FDirtyProxy& Dirty) const
+	{
+		const FParticleDirtyFlags Flags = Dirty.ParticleData.GetFlags();
+		if(Flags.IsDirty(ParticlePropToFlag(PropName)))
+		{
+			if(Manager)
+			{
+				//both dirty, if value differs we're desynced
+				return !(SrcManager.GetParticlePool<T,PropName>().GetElement(DataIdxIn) == Manager->GetParticlePool<T,PropName>().GetElement(Idx));
+			}
+			else
+			{
+				//Should have an entry but don't, so desynced
+				return true;
+			}
+		}
+		else
+		{
+			//if they're clean and we have no state then we're desynced
+			return Manager != nullptr;
+		}
 	}
 
 	template <typename LambdaSet>
@@ -113,7 +137,7 @@ public:
 		return Dynamics.IsSet() ? Dynamics.Read().F : Particle.CastToRigidParticle()->F();
 	}
 
-	bool SyncSimWritablePropsFromSim(FDirtyPropertiesManager& Manager,int32 Idx,const TPBDRigidParticleHandle<FReal,3>& Rigid)
+	void SyncSimWritablePropsFromSim(FDirtyPropertiesManager& Manager,int32 Idx,const TPBDRigidParticleHandle<FReal,3>& Rigid)
 	{
 		FParticleDirtyFlags Flags;
 		Flags.MarkDirty(EParticleFlags::XR);
@@ -132,11 +156,9 @@ public:
 			Data.V = Rigid.PreV();
 			Data.W = Rigid.PreW();
 		});
-
-		return true;	//todo: check if values are actually different
 	}
 
-	bool SyncDirtyDynamics(FDirtyPropertiesManager& DestManager,int32 DataIdx,const FParticleDirtyData& Dirty,const FDirtyPropertiesManager& SrcManager)
+	void SyncDirtyDynamics(FDirtyPropertiesManager& DestManager,int32 DataIdx,const FParticleDirtyData& Dirty,const FDirtyPropertiesManager& SrcManager)
 	{
 		FParticleDirtyData DirtyFlags;
 		DirtyFlags.SetFlags(Dirty.GetFlags());
@@ -145,8 +167,6 @@ public:
 		{
 			Data = Dirty.GetDynamics(SrcManager,DataIdx);
 		});
-
-		return true;	//todo: check if values are actually different
 	}
 
 	void SyncToParticle(TGeometryParticle<FReal,3>& Particle) const
@@ -197,7 +217,7 @@ public:
 		}
 	}
 
-	bool SyncPrevFrame(FDirtyPropertiesManager& Manager,int32 Idx,const FDirtyProxy& Dirty)
+	void SyncPrevFrame(FDirtyPropertiesManager& Manager,int32 Idx,const FDirtyProxy& Dirty)
 	{
 		//syncs the data before it was made dirty
 		//for sim-writable props this is only possible if those props are immutable from the sim side (sleeping, not simulated, etc...)
@@ -245,9 +265,6 @@ public:
 				Data.AngularEtherDrag = Rigid->AngularEtherDrag();
 			}
 		});
-
-		//todo: check whether the data actually changed
-		return true;
 	}
 
 	void SyncIfDirty(FDirtyPropertiesManager& Manager,int32 Idx,const TGeometryParticle<FReal,3>& InParticle, const FGeometryParticleStateBase& RewindState)
@@ -324,6 +341,19 @@ public:
 
 		return bCoalesced;
 	}
+
+	bool IsDesynced(const FDirtyPropertiesManager& SrcManager, const int32 DataIdxIn, const FDirtyProxy& Dirty) const
+	{
+		//Don't use Particle, should use dirty data which is thread safe
+		//We most likely will do this on GT, but maybe we could do other work while running the sim tick in parallel?
+		bool Desynced = false;
+		Desynced = ParticlePositionRotation.IsDesynced(SrcManager,DataIdxIn,Dirty);
+		Desynced = Desynced || NonFrequentData.IsDesynced(SrcManager,DataIdxIn,Dirty);
+		Desynced = Desynced || Velocities.IsDesynced(SrcManager,DataIdxIn,Dirty);
+		Desynced = Desynced || Dynamics.IsDesynced(SrcManager,DataIdxIn,Dirty);
+		return Desynced;
+	}
+
 private:
 
 	TParticleStateProperty<FParticlePositionRotation,EParticleProperty::XR> ParticlePositionRotation;
@@ -339,7 +369,7 @@ private:
 		PARTICLE_PROPERTY(MassProps,FParticleMassProps)*/
 };
 
-class FGeometryParticleState : private FGeometryParticleStateBase
+class FGeometryParticleState
 {
 public:
 
@@ -348,44 +378,67 @@ public:
 	{
 	}
 
-	FGeometryParticleState(const FGeometryParticleStateBase& State, const TGeometryParticle<FReal,3>& InParticle)
-	: FGeometryParticleStateBase(State)
-	, Particle(InParticle)
+	FGeometryParticleState(const FGeometryParticleStateBase& InState, const TGeometryParticle<FReal,3>& InParticle)
+	: Particle(InParticle)
+	, State(InState)
 	{
 	}
 
 	const FVec3& X() const
 	{
-		return FGeometryParticleStateBase::X(Particle);
+		return State.X(Particle);
 	}
 
 	const FRotation3& R() const
 	{
-		return FGeometryParticleStateBase::R(Particle);
+		return State.R(Particle);
 	}
 
 	const FVec3& V() const
 	{
-		return FGeometryParticleStateBase::V(Particle);
+		return State.V(Particle);
 	}
 
 	const FVec3& W() const
 	{
-		return FGeometryParticleStateBase::W(Particle);
+		return State.W(Particle);
 	}
 
 	TSerializablePtr<FImplicitObject> Geometry() const
 	{
-		return FGeometryParticleStateBase::Geometry(Particle);
+		return State.Geometry(Particle);
 	}
 
 	const FVec3& F() const
 	{
-		return FGeometryParticleStateBase::F(Particle);
+		return State.F(Particle);
+	}
+
+	const TGeometryParticle<FReal,3>& GetParticle() const
+	{
+		return Particle;
+	}
+
+	void SetState(const FGeometryParticleStateBase& InState)
+	{
+		State = InState;
+	}
+
+	bool IsDesynced(const FDirtyPropertiesManager& SrcManager, const int32 DataIdxIn, const FDirtyProxy& Dirty) const
+	{
+		return State.IsDesynced(SrcManager,DataIdxIn,Dirty);
 	}
 
 private:
 	const TGeometryParticle<FReal,3>& Particle;
+	FGeometryParticleStateBase State;
+};
+
+enum class EFutureQueryResult
+{
+	Ok,	//There is reliable data for this particle
+	Untracked, //The particle is untracked. This could mean it's new, or that it was unchanged in prior simulations
+	Desync //The particle's state has diverged from the previous recordings
 };
 
 class FRewindData
@@ -394,6 +447,7 @@ public:
 	FRewindData(int32 NumFrames)
 	: Managers(NumFrames+1)	//give 1 extra for saving at head
 	, CurFrame(0)
+	, LatestFrame(0)
 	, CurWave(0)
 	, FramesSaved(0)
 	, DataIdxOffset(0)
@@ -440,7 +494,7 @@ public:
 				//So preallocate just in case
 				FGeometryParticleStateBase& LatestState = DirtyParticleInfo.AddFrame(CurFrame, CurWave);
 			
-				if(const FGeometryParticleStateBase* RewindState = GetStateAtFrameImp(*DirtyParticleInfo.Particle, Frame))
+				if(const FGeometryParticleStateBase* RewindState = GetStateAtFrameImp(DirtyParticleInfo, Frame))
 				{
 					LatestState.SyncIfDirty(*DestManager,DataIdx++,*DirtyParticleInfo.Particle,*RewindState);
 					CoalesceBack(DirtyParticleInfo.Frames, CurFrame);
@@ -450,7 +504,7 @@ public:
 			}
 			else
 			{
-				if(const FGeometryParticleStateBase* RewindState = GetStateAtFrameImp(*DirtyParticleInfo.Particle,Frame))
+				if(const FGeometryParticleStateBase* RewindState = GetStateAtFrameImp(DirtyParticleInfo,Frame))
 				{
 					RewindState->SyncToParticle(*DirtyParticleInfo.Particle);
 				}
@@ -460,6 +514,7 @@ public:
 		CurFrame = Frame;
 		bNeedsSave = false;
 		FramesSaved = 0; //can't rewind before this point. This simplifies saving the state at head
+		//++CurWave;
 
 		return true;
 	}
@@ -480,21 +535,50 @@ public:
 		}
 	}
 
-	FGeometryParticleState GetStateAtFrame(const TGeometryParticle<FReal,3>& Particle,int32 Frame) const
+	/* Query the state of particles from the past. Once a rewind happens state captured must be queried using GetFutureStateAtFrame */
+	FGeometryParticleState GetPastStateAtFrame(const TGeometryParticle<FReal,3>& Particle,int32 Frame) const
 	{
-		if(const FGeometryParticleStateBase* State = GetStateAtFrameImp(Particle, Frame))
+		ensure(!IsResim());
+		if(const FDirtyParticleInfo* Info = FindParticle(Particle.UniqueIdx()))
 		{
-			return FGeometryParticleState(*State,Particle);
+			if(const FGeometryParticleStateBase* State = GetStateAtFrameImp(*Info,Frame))
+			{
+				return FGeometryParticleState(*State,Particle);
+			}
 		}
 
 		//If no data, or past capture, just use head
 		return FGeometryParticleState(Particle);
 	}
 
+	/* Query the state of particles in the future. This operation can fail for particles that are desynced or that we have not been tracking */
+	EFutureQueryResult GetFutureStateAtFrame(FGeometryParticleState& OutState,int32 Frame) const
+	{
+		ensure(IsResim());
+		const TGeometryParticle<FReal,3>& Particle = OutState.GetParticle();
+
+		if(const FDirtyParticleInfo* Info = FindParticle(Particle.UniqueIdx()))
+		{
+			if(Info->bDesync)
+			{
+				return EFutureQueryResult::Desync;
+			}
+
+			if(const FGeometryParticleStateBase* State = GetStateAtFrameImp(*Info,Frame))
+			{
+				OutState.SetState(*State);
+				return EFutureQueryResult::Ok;
+			}
+		}
+
+		return EFutureQueryResult::Untracked;
+	}
+
 	void AdvanceFrame()
 	{
 		QUICK_SCOPE_CYCLE_COUNTER(RewindDataAdvance);
 		++CurFrame;
+		LatestFrame = FMath::Max(LatestFrame,CurFrame);
 		FramesSaved = FMath::Min(FramesSaved+1,static_cast<int32>(Managers.Capacity()));
 		
 		const int32 EarliestFrame = CurFrame - 1 - FramesSaved;
@@ -508,6 +592,11 @@ public:
 		}
 	}
 
+	bool IsResim() const
+	{
+		return CurFrame < LatestFrame;
+	}
+
 	//Number of particles that we're currently storing history for
 	int32 GetNumDirtyParticles() const { return AllDirtyParticles.Num(); }
 
@@ -537,9 +626,10 @@ public:
 		Info.Manager->SetNumParticles(DataIdxOffset + NumActiveParticles);
 	}
 
-	void PushGTDirtyData(const FDirtyPropertiesManager& SrcManager, const int32 DataIdxIn, const FDirtyProxy& Dirty)
+	template <bool bResim>
+	void PushGTDirtyData(const FDirtyPropertiesManager& SrcManager, const int32 SrcDataIdx, const FDirtyProxy& Dirty)
 	{
-		const int32 DataIdx = DataIdxIn + DataIdxOffset;
+		const int32 DestDataIdx = SrcDataIdx + DataIdxOffset;
 		//This records changes enqueued by GT.
 		//Most new particles do not change, so to avoid useless writes we wait until the next frame's dirty flag
 		//This is possible because most properties are const on the physics thread
@@ -550,33 +640,45 @@ public:
 		FDirtyPropertiesManager& DestManager = *Managers[CurFrame].Manager;
 		bNeedsSave = true;
 		
-		auto ProcessProxy = [this,&SrcManager, DataIdx, Dirty, &DestManager](const auto Proxy)
+		auto ProcessProxy = [this,&SrcManager, DestDataIdx, SrcDataIdx, Dirty, &DestManager](const auto Proxy)
 		{
 			const auto PTParticle = Proxy->GetHandle();
 			FDirtyParticleInfo& Info = FindOrAddParticle(*Proxy->GetParticle(),PTParticle->UniqueIdx());
 			Info.LastDirtyFrame = CurFrame;
 
-			bool bChanged = false;
+			//check if particle has desynced
+			if(bResim)
+			{
+				FGeometryParticleState FutureState(*Proxy->GetParticle());
+				if(GetFutureStateAtFrame(FutureState,CurFrame) == EFutureQueryResult::Ok)
+				{
+					Info.bDesync = FutureState.IsDesynced(SrcManager, SrcDataIdx, Dirty);
+				}
+				else
+				{
+					//Most likely already desynced, but if it's a previously untracked object mark it as desynced
+					Info.bDesync = true;
+				}
+			}
 
 			//Most properties are always a frame behind
 			if(Proxy->IsInitialized())	//Frame delay so proxy must be initialized
 			{
-				FGeometryParticleStateBase& LatestState = Info.AddFrame(CurFrame-1, CurWave);
-				bChanged = LatestState.SyncPrevFrame(DestManager,DataIdx,Dirty);
-				CoalesceBack(Info.Frames, CurFrame-1);
+				//If we're in a resim and this is the first frame or the resim, no need to save prev frame
+				//In fact, since we have a circular buffer the prev state could end up overwriting head which we need for fast forward
+				if(!bResim || FramesSaved > 0)
+				{
+					FGeometryParticleStateBase& LatestState = Info.AddFrame(CurFrame-1,CurWave);
+					LatestState.SyncPrevFrame(DestManager,DestDataIdx,Dirty);
+					CoalesceBack(Info.Frames,CurFrame-1);
+				}
 			}
 
 			//If dynamics are dirty we must record them immediately because the sim resets them to 0
 			if(Dirty.ParticleData.IsDirty(EParticleFlags::Dynamics))
 			{
 				FGeometryParticleStateBase& LatestState = Info.AddFrame(CurFrame, CurWave);
-				bChanged = LatestState.SyncDirtyDynamics(DestManager,DataIdx,Dirty.ParticleData,SrcManager);
-			}
-
-			if(bChanged)
-			{
-				//value actually changed, so we must desync
-				Info.bDesync = true;
+				LatestState.SyncDirtyDynamics(DestManager,DestDataIdx,Dirty.ParticleData,SrcManager);
 			}
 		};
 
@@ -605,9 +707,10 @@ public:
 		}
 	}
 
-	void PushPTDirtyData(const TPBDRigidParticleHandle<FReal,3>& Rigid, const int32 DataIdxIn)
+	template <bool bResim>
+	void PushPTDirtyData(const TPBDRigidParticleHandle<FReal,3>& Rigid, const int32 SrcDataIdx)
 	{
-		const int32 DataIdx = DataIdxIn + DataIdxOffset;
+		const int32 DestDataIdx = SrcDataIdx + DataIdxOffset;
 
 		//todo: is this check needed? why do we pass sleeping rigids into this function?
 		if(SimWritablePropsMayChange(Rigid))
@@ -620,11 +723,7 @@ public:
 
 			//sim-writable properties changed at head, so we must write down what they were
 			FGeometryParticleStateBase& LatestState = Info.AddFrame(CurFrame-1, CurWave);
-			bool bChanged = LatestState.SyncSimWritablePropsFromSim(DestManager,DataIdx,Rigid);
-			if(bChanged)
-			{
-				Info.bDesync = true;
-			}
+			LatestState.SyncSimWritablePropsFromSim(DestManager,DestDataIdx,Rigid);
 
 			//update any previous frames that were pointing at head
 			CoalesceBack(Info.Frames, CurFrame-1);
@@ -707,32 +806,6 @@ private:
 		int32 FrameCreatedFor;
 	};
 
-	const FGeometryParticleStateBase* GetStateAtFrameImp(const TGeometryParticle<FReal,3>& Particle,int32 Frame) const
-	{
-		if(const FDirtyParticleInfo* Info = FindParticle(Particle.UniqueIdx()))
-		{
-			const TCircularBuffer<FFrameInfo>& Frames = Info->Frames;
-			if(const FGeometryParticleStateBase* FrameState = Frames[Frame].GetState(Frame, CurWave))
-			{
-				return FrameState;
-			}
-
-			//If frame is between two captures, use later capture. We always store the last data before a change
-			//We can never use an earlier capture because the fact that we captured at all implies _something_ is different from proceeding frames
-
-			for(int32 FrameIdx = Frame + 1; FrameIdx <= CurFrame; ++FrameIdx)
-			{
-				if(const FGeometryParticleStateBase* FrameState = Frames[FrameIdx].GetState(FrameIdx, CurWave))
-				{
-					return FrameState;
-				}
-			}
-		}
-
-		//If no data, or past capture, just use head
-		return nullptr;
-	}
-
 	struct FDirtyParticleInfo
 	{
 		TCircularBuffer<FFrameInfo> Frames;
@@ -763,6 +836,29 @@ private:
 		}
 	};
 
+	const FGeometryParticleStateBase* GetStateAtFrameImp(const FDirtyParticleInfo& Info,int32 Frame) const
+	{
+		const TCircularBuffer<FFrameInfo>& Frames = Info.Frames;
+		if(const FGeometryParticleStateBase* FrameState = Frames[Frame].GetState(Frame,CurWave))
+		{
+			return FrameState;
+		}
+
+		//If frame is between two captures, use later capture. We always store the last data before a change
+		//We can never use an earlier capture because the fact that we captured at all implies _something_ is different from proceeding frames
+
+		for(int32 FrameIdx = Frame + 1; FrameIdx <= LatestFrame; ++FrameIdx)
+		{
+			if(const FGeometryParticleStateBase* FrameState = Frames[FrameIdx].GetState(FrameIdx,CurWave))
+			{
+				return FrameState;
+			}
+		}
+
+		//If no data, or past capture, just use head
+		return nullptr;
+	}
+
 	const FDirtyParticleInfo& FindParticleChecked(const FUniqueIdx UniqueIdx) const
 	{
 		const int32 Idx = ParticleToAllDirtyIdx.FindChecked(UniqueIdx);
@@ -812,6 +908,7 @@ private:
 	TCircularBuffer<FFrameManagerInfo> Managers;
 	TArray<FDirtyParticleInfo> AllDirtyParticles;
 	int32 CurFrame;
+	int32 LatestFrame;
 	int8 CurWave;
 	int32 FramesSaved;
 	int32 DataIdxOffset;
