Port TPBDRigidClustering::DeactivateClusterParticle() and test RigidBodies_ClusterTest_DeactivateClusterParticle().

diff --git a/Engine/Source/Runtime/Experimental/Chaos/Private/Chaos/PBDRigidClustering.cpp b/Engine/Source/Runtime/Experimental/Chaos/Private/Chaos/PBDRigidClustering.cpp
--- a/Engine/Source/Runtime/Experimental/Chaos/Private/Chaos/PBDRigidClustering.cpp
+++ b/Engine/Source/Runtime/Experimental/Chaos/Private/Chaos/PBDRigidClustering.cpp
@@ -673,15 +673,14 @@ namespace Chaos
 	}
 #endif
 
-#if 0 // Not called currently
 	DECLARE_CYCLE_STAT(TEXT("TPBDRigidClustering<>::DeactivateClusterParticle"), STAT_DeactivateClusterParticle, STATGROUP_Chaos);
-	template<class T_FPBDRigidsEvolution, class T_FPBDCollisionConstraint, class T, int d>
-	TSet<uint32> TPBDRigidClustering<T_FPBDRigidsEvolution, T_FPBDCollisionConstraint, T, d>::DeactivateClusterParticle(
+	template<class FPBDRigidsEvolution, class FPBDCollisionConstraint, class T, int d>
+	TSet<TPBDRigidParticleHandle<T, d>*> TPBDRigidClustering<FPBDRigidsEvolution, FPBDCollisionConstraint, T, d>::DeactivateClusterParticle(
 		TPBDRigidClusteredParticleHandle<T,d>* ClusteredParticle)
 	{
 		SCOPE_CYCLE_COUNTER(STAT_DeactivateClusterParticle);
 
-		TSet<uint32> ActivatedChildren;
+		TSet<TPBDRigidParticleHandle<T, d>*> ActivatedChildren;
 		check(!ClusteredParticle->Disabled());
 		if (MChildren.Contains(ClusteredParticle))
 		{
@@ -689,7 +688,6 @@ namespace Chaos
 		}
 		return ActivatedChildren;
 	}
-#endif // 0
 
 	DECLARE_CYCLE_STAT(TEXT("TPBDRigidClustering<>::ReleaseClusterParticles(STRAIN)"), STAT_ReleaseClusterParticles_STRAIN, STATGROUP_Chaos);
 	template<class T_FPBDRigidsEvolution, class T_FPBDCollisionConstraint, class T, int d>
@@ -911,54 +909,49 @@ namespace Chaos
 		return ActivatedChildren;
 	}
 
-#if 0 // Not called currently
 	DECLARE_CYCLE_STAT(TEXT("TPBDRigidClustering<>::ReleaseClusterParticles(LIST)"), STAT_ReleaseClusterParticles_LIST, STATGROUP_Chaos);
-	template<class T_FPBDRigidsEvolution, class T_FPBDCollisionConstraint, class T, int d>
-	TSet<uint32> TPBDRigidClustering<T_FPBDRigidsEvolution, T_FPBDCollisionConstraint, T, d>::ReleaseClusterParticles(
+	template<class FPBDRigidsEvolution, class FPBDCollisionConstraint, class T, int d>
+	TSet<TPBDRigidParticleHandle<T, d>*> TPBDRigidClustering<FPBDRigidsEvolution, FPBDCollisionConstraint, T, d>::ReleaseClusterParticles(
 		TArray<TPBDRigidParticleHandle<T, d>*> ChildrenParticles)
-		const TArray<uint32>& ChildrenParticles)
 	{
 		SCOPE_CYCLE_COUNTER(STAT_ReleaseClusterParticles_LIST);
-		TSet<uint32> ActivatedBodies;
-		bool bFound = false;
+		TSet<TPBDRigidParticleHandle<T, d>*> ActivatedBodies;
 		if (ChildrenParticles.Num())
 		{
-			uint32 ClusterIdx = 0;
+			TPBDRigidParticleHandle<float, 3>* ClusterHandle = nullptr;
 			//todo(ocohen): refactor incoming, for now just assume these all belong to same cluster and hack strain array
-			TArray<float> FakeStrain;
-			FakeStrain.Init(0, MParticles.Size()); //this part especially sucks
+			
+			TMap<TGeometryParticleHandle<T, d>*, float> FakeStrain;
 
 			bool bPreDoGenerateData = DoGenerateBreakingData;
 			DoGenerateBreakingData = false;
 
-			for (uint32 ChildIdx : ChildrenParticles)
+			for (TPBDRigidParticleHandle<T, d>* ChildHandle : ChildrenParticles)
 			{
-				if (MParticles.Disabled(ChildIdx) && MParticles.ClusterIds(ChildIdx).Id != INDEX_NONE)
+				if (TPBDRigidClusteredParticleHandle<T, d>* ClusteredChildHandle = ChildHandle->CastToClustered())
 				{
-					if (ensure(!bFound || MParticles.ClusterIds(ChildIdx).Id == ClusterIdx))
+					if (ClusteredChildHandle->Disabled() && ClusteredChildHandle->ClusterIds().Id != nullptr)
 					{
-						bFound = true;
-						FakeStrain[ChildIdx] = FLT_MAX;
-						ClusterIdx = MParticles.ClusterIds(ChildIdx).Id;
-					}
-					else
-					{
-						break; //shouldn't be here
+						if (ensure(!ClusterHandle || ClusteredChildHandle->ClusterIds().Id == ClusterHandle))
+						{
+							FakeStrain.Add(ClusteredChildHandle, TNumericLimits<float>::Max());
+							ClusterHandle = ClusteredChildHandle->ClusterIds().Id;
+						}
+						else
+						{
+							break; //shouldn't be here
+						}
 					}
 				}
 			}
-
-			if (bFound)
+			if (ClusterHandle)
 			{
-				ActivatedBodies.Append(ReleaseClusterParticles(MParticles.Handle(ClusterIdx), FakeStrain));
+				ActivatedBodies = ReleaseClusterParticles(ClusterHandle->CastToClustered(), &FakeStrain);
 			}
-
 			DoGenerateBreakingData = bPreDoGenerateData;
 		}
-
 		return ActivatedBodies;
 	}
-#endif // 0
 
 	DECLARE_CYCLE_STAT(TEXT("TPBDRigidClustering<>::AdvanceClustering"), STAT_AdvanceClustering, STATGROUP_Chaos);
 	DECLARE_CYCLE_STAT(TEXT("TPBDRigidClustering<>::Update Impulse from Strain"), STAT_UpdateImpulseStrain, STATGROUP_Chaos);
diff --git a/Engine/Source/Runtime/Experimental/Chaos/Public/Chaos/PBDRigidClustering.h b/Engine/Source/Runtime/Experimental/Chaos/Public/Chaos/PBDRigidClustering.h
--- a/Engine/Source/Runtime/Experimental/Chaos/Public/Chaos/PBDRigidClustering.h
+++ b/Engine/Source/Runtime/Experimental/Chaos/Public/Chaos/PBDRigidClustering.h
@@ -88,13 +88,11 @@ public:
 	// Releasing
 	//
 
-#if 0 // Not called currently
 	/*
 	*  DeactivateClusterParticle
 	*    Release all the particles within the cluster particle
 	*/
-	TSet<uint32> DeactivateClusterParticle(TPBDRigidClusteredParticleHandle<T,d>* ClusteredParticle);
-#endif // 0
+	TSet<TPBDRigidParticleHandle<T, d>*> DeactivateClusterParticle(TPBDRigidClusteredParticleHandle<T,d>* ClusteredParticle);
 
 	/*
 	*  ReleaseClusterParticles (BasedOnStrain)
@@ -103,16 +101,17 @@ public:
 	*    that have a strain value less than this valid will be released from the 
 	*    cluster.
 	*/
-	TSet<TPBDRigidParticleHandle<T, d>*> ReleaseClusterParticles(TPBDRigidClusteredParticleHandle<T, d>* ClusteredParticle, const TMap<TGeometryParticleHandle<T, d>*, float>* ExternalStrainMap = nullptr);
+	TSet<TPBDRigidParticleHandle<T, d>*> ReleaseClusterParticles(
+		TPBDRigidClusteredParticleHandle<T, d>* ClusteredParticle, 
+		const TMap<TGeometryParticleHandle<T, d>*, float>* ExternalStrainMap = nullptr);
 
-#if 0 // Not called currently
 	/*
 	*  ReleaseClusterParticles
 	*    Release all rigid body IDs passed,
 	*/
-	TSet<uint32> ReleaseClusterParticles(
+	TSet<TPBDRigidParticleHandle<T, d>*> ReleaseClusterParticles(
 		TArray<TPBDRigidParticleHandle<T, d>*> ChildrenParticles);
-#endif // 0
+
 	//
 	// Operational 
 	//
