Chaos : Union Cluster Groups
- Union cluster groups

diff --git a/Engine/Source/Runtime/Experimental/Chaos/Private/Chaos/PBDRigidClustering.cpp b/Engine/Source/Runtime/Experimental/Chaos/Private/Chaos/PBDRigidClustering.cpp
--- a/Engine/Source/Runtime/Experimental/Chaos/Private/Chaos/PBDRigidClustering.cpp
+++ b/Engine/Source/Runtime/Experimental/Chaos/Private/Chaos/PBDRigidClustering.cpp
@@ -434,7 +434,6 @@ namespace Chaos
 		Chaos::TPBDRigidClusteredParticleHandle<float, 3>* NewParticle = Parameters.ClusterParticleHandle;
 		if (!NewParticle)
 		{
-			//NewParticle = MEvolution.CreateClusteredParticles(1)[0]; // calls Evolution.DirtyParticle()
 			NewParticle = MEvolution.CreateClusteredParticles(1)[0]; // calls Evolution.DirtyParticle()
 		}
 
@@ -506,6 +505,7 @@ namespace Chaos
 		UpdateGeometry(NewParticle, ChildrenSet, ProxyGeometry, Parameters);
 		GenerateConnectionGraph(NewParticle, Parameters);
 		NewParticle->SetSleeping(bClusterIsAsleep);
+		if (ClusterGroupIndex) AddToClusterUnion(ClusterGroupIndex, NewParticle);
 		return NewParticle;
 	}
 
@@ -588,7 +588,6 @@ namespace Chaos
 		return NewParticle;
 	}
 
-#if 0
 	int32 UseMultiChildProxy = 1;
 	FAutoConsoleVariableRef CVarUseMultiChildProxy(TEXT("p.UseMultiChildProxy"), UseMultiChildProxy, TEXT("Whether to merge multiple children into a single collision proxy when one is available"));
 
@@ -601,77 +600,23 @@ namespace Chaos
 	{
 		SCOPE_CYCLE_COUNTER(STAT_UnionClusterGroups);
 
-		TMap<int32, TArray<uint32>> GroupMapping;
-		for (uint32 i = 0; i < MParticles.Size(); i++) // just loop active clusters here.
-		{
-			uint32 ParticleIndex = i;
-			int32 GroupIndex = MParticles.ClusterGroupIndex(ParticleIndex);
-			if (GroupIndex > 0)
-			{
-				if (!GroupMapping.Contains(GroupIndex))
-				{
-					GroupMapping.Add(GroupIndex, TArray<uint32>());
-				}
-				GroupMapping[GroupIndex].Add(ParticleIndex);
-			}
-		}
 
-		for (TTuple<int32, TArray<uint32>>& Group : GroupMapping)
+		for (TTuple<int32, TArray<TPBDRigidClusteredParticleHandle<T, 3>* >>& Group : ClusterUnionMap)
 		{
-			if (PendingClusterCounter.Contains(Group.Key) && PendingClusterCounter[Group.Key] == 0)
-			{
-				TArray<uint32> ClusterChildren;
-				for (uint32& OriginalRootIdx : Group.Value)
-				{
-					TUniquePtr<TMultiChildProxyData<T, d>> ProxyData;
-					if (MChildren.Contains(MParticles.Handle(OriginalRootIdx)))
-					{
-						if (UseMultiChildProxy && !MParticles.DynamicGeometry(OriginalRootIdx) && MChildren[MParticles.Handle(OriginalRootIdx)].Num() > MinChildrenForMultiProxy) //Don't support dynamic geometry
-						{
-							if (ensure(MChildren[MParticles.Handle(OriginalRootIdx)].Num()))
-							{
-								ProxyData = MakeUnique<TMultiChildProxyData<T, d>>();
-								ProxyData->KeyChild = (MChildren[MParticles.Handle(OriginalRootIdx)])[0];
-								ProxyData->RelativeToKeyChild = TRigidTransform<T, d>(MParticles.X(OriginalRootIdx), MParticles.R(OriginalRootIdx)); //store world space of original root. Need to break it up and then compute relative to world space of key child
-							}
-						}
+			uint32 ClusterGroupID = Group.Key;
+			TArray<TPBDRigidClusteredParticleHandle<T, 3>* > Handles = Group.Value;
 
-						const TArray<uint32> OriginalRootChildren = DeactivateClusterParticle(MParticles.Handle(OriginalRootIdx)).Array();
-						ClusterChildren.Append(OriginalRootChildren);
 
-						if (ProxyData)
-						{
-							//now that we have world space updated for key child, compute relative transform for original root
-							const TRigidTransform<T, d> OriginalRootWorldTM = ProxyData->RelativeToKeyChild;
-							ProxyData->RelativeToKeyChild = 
-								OriginalRootWorldTM.GetRelativeTransform(
-									TRigidTransform<T, d>(MParticles.X(ProxyData->KeyChild), MParticles.R(ProxyData->KeyChild)));
-							MParticles.MultiChildProxyData(OriginalRootIdx) = MoveTemp(ProxyData);
-
-							for (uint32 Child : OriginalRootChildren)
-							{
-								//remember original proxy of child cluster
-								MParticles.MultiChildProxyId(Child).Id = OriginalRootIdx;
-							}
-						}
-					}
-					else
-					{
-						ClusterChildren.Add(OriginalRootIdx);
-					}
-				}
-
-				FClusterCreationParameters<T> Parameters(0.3, 100, false, !!UnionsHaveCollisionParticles);
-				Parameters.ConnectionMethod = MClusterUnionConnectionType;
-				int32 NewIndex = CreateClusterParticle(-Group.Key, MoveTemp(ClusterChildren), TSerializablePtr<FImplicitObject>(), nullptr, Parameters);
-				MParticles.InternalCluster(NewIndex) = true;
-				MEvolution.SetPhysicsMaterial(MParticles.Handle(NewIndex), MEvolution.GetPhysicsMaterial(MParticles.Handle(Group.Value[0])));
-
-				PendingClusterCounter.Remove(Group.Key);
+			TArray<TPBDRigidParticleHandle<T, 3>*> ClusterBodies;
+			for (TPBDRigidClusteredParticleHandle<T, 3>* ActiveCluster : Handles)
+			{
+				ClusterBodies.Append(ReleaseClusterParticles(ActiveCluster, nullptr, true).Array());
 			}
+			CreateClusterParticle(-Group.Key, MoveTemp(ClusterBodies));
 		}
+		ClusterUnionMap.Empty();
 	}
-#endif
+
 
 	DECLARE_CYCLE_STAT(TEXT("TPBDRigidClustering<>::DeactivateClusterParticle"), STAT_DeactivateClusterParticle, STATGROUP_Chaos);
 	template<class FPBDRigidsEvolution, class FPBDCollisionConstraint, class T, int d>
@@ -694,7 +639,8 @@ namespace Chaos
 	TSet<TPBDRigidParticleHandle<T, d>*> 
 	TPBDRigidClustering<T_FPBDRigidsEvolution, T_FPBDCollisionConstraint, T, d>::ReleaseClusterParticles(
 		TPBDRigidClusteredParticleHandle<T,d>* ClusteredParticle,
-		const TMap<TGeometryParticleHandle<T, d>*, float>* ExternalStrainMap)
+		const TMap<TGeometryParticleHandle<T, d>*, float>* ExternalStrainMap,
+		bool bForceRelease)
 	{
 		SCOPE_CYCLE_COUNTER(STAT_ReleaseClusterParticles_STRAIN);
 
@@ -761,7 +707,8 @@ namespace Chaos
 			if (!Child)
 				continue;
 			if ((ExternalStrainMap && (ExternalStrainMap->Find(Child))) ||
-			   (!ExternalStrainMap && Child->CollisionImpulses() >= Child->Strain()))
+			   (!ExternalStrainMap && Child->CollisionImpulses() >= Child->Strain()) ||
+				bForceRelease)
 			{
 				// The piece that hits just breaks off - we may want more control 
 				// by looking at the edges of this piece which would give us cleaner 
@@ -1252,25 +1199,6 @@ namespace Chaos
 		return Child; 
 	}
 
-	template<class T_FPBDRigidsEvolution, class T_FPBDCollisionConstraint, class T, int d>
-	void TPBDRigidClustering<T_FPBDRigidsEvolution, T_FPBDCollisionConstraint, T, d>::IncrementPendingClusterCounter(
-		uint32 ClusterGroupID)
-	{
-		if (!PendingClusterCounter.Contains(ClusterGroupID))
-		{
-			PendingClusterCounter.Add(ClusterGroupID, 0);
-		}
-		PendingClusterCounter[ClusterGroupID]++;
-	}
-
-	template<class T_FPBDRigidsEvolution, class T_FPBDCollisionConstraint, class T, int d>
-	void TPBDRigidClustering<T_FPBDRigidsEvolution, T_FPBDCollisionConstraint, T, d>::DecrementPendingClusterCounter(
-		uint32 ClusterGroupID)
-	{
-		PendingClusterCounter[ClusterGroupID]--;
-		ensure(0 <= PendingClusterCounter[ClusterGroupID]);
-	}
-
 	DECLARE_CYCLE_STAT(TEXT("TPBDRigidClustering<>::GenerateConnectionGraph"), STAT_GenerateConnectionGraph, STATGROUP_Chaos);
 	template<class T_FPBDRigidsEvolution, class T_FPBDCollisionConstraint, class T, int d>
 	void TPBDRigidClustering<T_FPBDRigidsEvolution, T_FPBDCollisionConstraint, T, d>::GenerateConnectionGraph(
@@ -2049,14 +1977,6 @@ namespace Chaos
 	}
 
 
-/*
-	template<class T_FPBDRigidsEvolution, class T_FPBDCollisionConstraint, class T, int d>
-	void TPBDRigidClustering<T_FPBDRigidsEvolution, T_FPBDCollisionConstraint, T, d>::ClearPendingClusterCounter(uint32 ClusterGroupID)
-	{
-		PendingClusterCounter[ClusterGroupID]=0;
-		//ensure(0 <= PendingClusterCounter[ClusterGroupID]);
-	}
-*/
 } // namespace Chaos
 
 using namespace Chaos;
diff --git a/Engine/Source/Runtime/Experimental/Chaos/Private/Chaos/PBDRigidsEvolutionGBF.cpp b/Engine/Source/Runtime/Experimental/Chaos/Private/Chaos/PBDRigidsEvolutionGBF.cpp
--- a/Engine/Source/Runtime/Experimental/Chaos/Private/Chaos/PBDRigidsEvolutionGBF.cpp
+++ b/Engine/Source/Runtime/Experimental/Chaos/Private/Chaos/PBDRigidsEvolutionGBF.cpp
@@ -257,6 +257,10 @@ void FPBDRigidsEvolutionGBF::AdvanceOneTimeStep(const FReal Dt, const FReal Step
 	}
 #endif
 
+	{
+		Clustering.UnionClusterGroups();
+	}
+
 	{
 		SCOPE_CYCLE_COUNTER(STAT_Evolution_Integrate);
 		Integrate(Particles.GetActiveParticlesView(), Dt);
diff --git a/Engine/Source/Runtime/Experimental/Chaos/Public/Chaos/PBDRigidClustering.h b/Engine/Source/Runtime/Experimental/Chaos/Public/Chaos/PBDRigidClustering.h
--- a/Engine/Source/Runtime/Experimental/Chaos/Public/Chaos/PBDRigidClustering.h
+++ b/Engine/Source/Runtime/Experimental/Chaos/Public/Chaos/PBDRigidClustering.h
@@ -62,7 +62,7 @@ public:
 	Chaos::TPBDRigidClusteredParticleHandle<float, 3>* CreateClusterParticle(
 		const int32 ClusterGroupIndex, 
 		TArray<Chaos::TPBDRigidParticleHandle<T,d>*>&& Children, 
-		const FClusterCreationParameters<T>& Parameters/* = FClusterCreationParameters<T>()*/,
+		const FClusterCreationParameters<T>& Parameters = FClusterCreationParameters<T>(),
 		TSharedPtr<Chaos::FImplicitObject, ESPMode::ThreadSafe> ProxyGeometry = nullptr,
 		const TRigidTransform<T, d>* ForceMassOrientation = nullptr);
 
@@ -76,14 +76,13 @@ public:
 		const TRigidTransform<T, d>& ClusterWorldTM, 
 		const FClusterCreationParameters<T>& Parameters/* = FClusterCreationParameters<T>()*/);
 
-#if 0 // Not called currently
 	/**
 	 *  UnionClusterGroups
 	 *    Clusters that share a group index should be unioned into a single cluster prior to simulation.
 	 *    The GroupIndex should be set on creation, and never touched by the client again.
 	 */
 	void UnionClusterGroups();
-#endif // 0
+
 	//
 	// Releasing
 	//
@@ -103,7 +102,8 @@ public:
 	*/
 	TSet<TPBDRigidParticleHandle<T, d>*> ReleaseClusterParticles(
 		TPBDRigidClusteredParticleHandle<T, d>* ClusteredParticle, 
-		const TMap<TGeometryParticleHandle<T, d>*, float>* ExternalStrainMap = nullptr);
+		const TMap<TGeometryParticleHandle<T, d>*, float>* ExternalStrainMap = nullptr,
+		bool bForceRelease = false);
 
 	/*
 	*  ReleaseClusterParticles
@@ -223,18 +223,11 @@ public:
 	/** If multi child proxy is used, this is the data needed */
 	const TArrayCollectionArray<TUniquePtr<TMultiChildProxyData<T, d>>>& GetMultiChildProxyDataArray() const { return MParticles.MultiChildProxyDataArray(); }
 
-	/*
-	*  GetClusterGroupIndexArray
-	*  Cluster counters are used the defer the initialization of grouped clusters until
-	*  all bodies are initialized within the simulation. During construction the counter
-	*  is incremented, and when ActivateBodies on the solver triggers activation across
-	*  multiple sets of bodies, the counter is decremented, when its back to Zero the
-	*  the union cluster is allowed to initialize. Once a cluster group ID is used up
-	*  is can not be reused.
-	*/
-	void IncrementPendingClusterCounter(uint32 ClusterGroupID);
-	void DecrementPendingClusterCounter(uint32 ClusterGroupID);
-	int32 NumberOfPendingClusters() const { return PendingClusterCounter.Num(); }
+	void AddToClusterUnion(int32 ClusterID, TPBDRigidClusteredParticleHandle<T, 3>* Handle) {
+		if (ClusterID <= 0) return;
+		if (!ClusterUnionMap.Contains(ClusterID)) ClusterUnionMap.Add(ClusterID, TArray<TPBDRigidClusteredParticleHandle<T, 3>*>());
+		ClusterUnionMap[ClusterID].Add(Handle);
+	}
 
 	const TArray<TBreakingData<float, 3>>& GetAllClusterBreakings() const { return MAllClusterBreakings; }
 	void SetGenerateClusterBreaking(bool DoGenerate) { DoGenerateBreakingData = DoGenerate; }
@@ -323,11 +316,12 @@ private:
 	TSet<Chaos::TPBDRigidClusteredParticleHandle<float, 3>*> TopLevelClusterParents;
 	TSet<Chaos::TPBDRigidParticleHandle<float, 3>*> MActiveRemovalIndices;
 
+
 	// Cluster data
 	mutable FRWLock ResourceLock;
 	TClusterBuffer<T, d> BufferResource;
 	FClusterMap MChildren;
-	TMap<int32, int32> PendingClusterCounter;
+	TMap<int32, TArray<TPBDRigidClusteredParticleHandle<T, 3>*> > ClusterUnionMap;
 
 
 	// Collision Impulses
diff --git a/Engine/Source/Runtime/Experimental/ChaosSolvers/Private/PBDRigidsSolver.cpp b/Engine/Source/Runtime/Experimental/ChaosSolvers/Private/PBDRigidsSolver.cpp
--- a/Engine/Source/Runtime/Experimental/ChaosSolvers/Private/PBDRigidsSolver.cpp
+++ b/Engine/Source/Runtime/Experimental/ChaosSolvers/Private/PBDRigidsSolver.cpp
@@ -875,8 +875,7 @@ namespace Chaos
 		TParticleView<TPBDRigidParticles<float, 3>>& ActiveParticles = GetParticles().GetActiveParticlesView();
 		for (Chaos::TPBDRigidParticleHandleImp<float, 3, false>& ActiveObject : ActiveParticles)
 		{
-			IPhysicsProxyBase* Proxy = GetProxy(ActiveObject.Handle());
-			if (ensure(Proxy))
+			if (IPhysicsProxyBase* Proxy = GetProxy(ActiveObject.Handle()))
 			{
 				switch (ActiveObject.GetParticleType())
 				{
@@ -916,8 +915,7 @@ namespace Chaos
 		TParticleView<TPBDRigidParticles<float, 3>>& ActiveParticles = GetParticles().GetActiveParticlesView();
 		for (Chaos::TPBDRigidParticleHandleImp<float, 3, false>& ActiveObject : ActiveParticles)
 		{
-			IPhysicsProxyBase* Proxy = GetProxy(ActiveObject.Handle());
-			if (ensure(Proxy))
+			if (IPhysicsProxyBase* Proxy = GetProxy(ActiveObject.Handle()))
 			{
 				switch (ActiveObject.GetParticleType())
 				{
@@ -962,8 +960,7 @@ namespace Chaos
 		TParticleView<TPBDRigidParticles<float, 3>>& ActiveParticles = GetParticles().GetActiveParticlesView();
 		for (Chaos::TPBDRigidParticleHandleImp<float, 3, false>& ActiveObject : ActiveParticles)
 		{
-			IPhysicsProxyBase* Proxy = GetProxy(ActiveObject.Handle());
-			if (ensure(Proxy))
+			if (IPhysicsProxyBase* Proxy = GetProxy(ActiveObject.Handle()))
 			{
 				switch (ActiveObject.GetParticleType())
 				{
