Fixed disable particles getting pushed into spatial acceleration structures and fixed cluster particles being considered global objects and skipping collision resoluton while slowing down any legitimate query.

diff --git a/Engine/Source/Runtime/Experimental/Chaos/Public/Chaos/PBDRigidsEvolution.h b/Engine/Source/Runtime/Experimental/Chaos/Public/Chaos/PBDRigidsEvolution.h
--- a/Engine/Source/Runtime/Experimental/Chaos/Public/Chaos/PBDRigidsEvolution.h
+++ b/Engine/Source/Runtime/Experimental/Chaos/Public/Chaos/PBDRigidsEvolution.h
@@ -339,6 +339,13 @@ class FPBDRigidsEvolutionBase
 	template <bool bPersistent>
 	FORCEINLINE_DEBUGGABLE void DirtyParticle(TGeometryParticleHandleImp<FReal, 3, bPersistent>& Particle)
 	{
+		const TPBDRigidParticleHandleImp<FReal, 3, bPersistent>* AsRigid = Particle.CastToRigidParticle();
+		if(AsRigid && AsRigid->Disabled())
+		{
+			// Disabled particles take no immediate part in sim or query so shouldn't be added to the acceleration
+			return;
+		}
+
 		//TODO: distinguish between new particles and dirty particles
 		const FUniqueIdx UniqueIdx = Particle.UniqueIdx();
 		FPendingSpatialData& SpatialData = InternalAccelerationQueue.FindOrAdd(UniqueIdx);
diff --git a/Engine/Source/Runtime/Experimental/ChaosSolvers/Private/PhysicsProxy/GeometryCollectionPhysicsProxy.cpp b/Engine/Source/Runtime/Experimental/ChaosSolvers/Private/PhysicsProxy/GeometryCollectionPhysicsProxy.cpp
--- a/Engine/Source/Runtime/Experimental/ChaosSolvers/Private/PhysicsProxy/GeometryCollectionPhysicsProxy.cpp
+++ b/Engine/Source/Runtime/Experimental/ChaosSolvers/Private/PhysicsProxy/GeometryCollectionPhysicsProxy.cpp
@@ -922,7 +922,7 @@ void FGeometryCollectionPhysicsProxy::InitializeBodiesPT(Chaos::FPBDRigidsSolver
 				{
 					if (SolverClusterHandles[TransformGroupIndex])
 					{
-						Chaos::FClusterCreationParameters<float> ClusterParams;
+						Chaos::FClusterCreationParameters<Chaos::FReal> ClusterParams;
 						// #todo: should other parameters be set here?  Previously, there was no parameters being sent, and it is unclear
 						// where some of these parameters are defined (ie: CoillisionThicknessPercent)
 						ClusterParams.ConnectionMethod = Parameters.ClusterConnectionMethod;
@@ -1116,6 +1116,18 @@ FGeometryCollectionPhysicsProxy::BuildClusters(
 		MinCollisionGroup = FMath::Min(Child->CollisionGroup(), MinCollisionGroup);
 	}
 	Parent->SetCollisionGroup(MinCollisionGroup);
+
+	// Populate bounds as we didn't pass a shared implicit to PopulateSimulatedParticle this will have been skipped, now that we have the full cluster we can build it
+	if(Parent->Geometry() && Parent->Geometry()->HasBoundingBox())
+	{
+		Parent->SetHasBounds(true);
+		Parent->SetLocalBounds(Parent->Geometry()->BoundingBox());
+		const Chaos::TAABB<float, 3>& LocalBounds = Parent->LocalBounds();
+		const Chaos::TRigidTransform<float, 3> Xf(Parent->X(), Parent->R());
+		const Chaos::TAABB<float, 3> TransformedBBox = LocalBounds.TransformedAABB(Xf);
+		Parent->SetWorldSpaceInflatedBounds(TransformedBBox);
+	}
+
 	return Parent;
 }
 
