Disable UnionClusterGroups for geometry collections as larger groups are failing to generate cluster groups and end up deadlocking.

diff --git a/Engine/Source/Runtime/Experimental/Chaos/Private/Chaos/PBDRigidsEvolutionGBF.cpp b/Engine/Source/Runtime/Experimental/Chaos/Private/Chaos/PBDRigidsEvolutionGBF.cpp
--- a/Engine/Source/Runtime/Experimental/Chaos/Private/Chaos/PBDRigidsEvolutionGBF.cpp
+++ b/Engine/Source/Runtime/Experimental/Chaos/Private/Chaos/PBDRigidsEvolutionGBF.cpp
@@ -258,7 +258,8 @@ void FPBDRigidsEvolutionGBF::AdvanceOneTimeStep(const FReal Dt, const FReal Step
 #endif
 
 	{
-		Clustering.UnionClusterGroups();
+		// #TODO re-enable when larger groups of collections can successfully resolve groups
+		//Clustering.UnionClusterGroups();
 	}
 
 	{
