Chaos - hook up per-shape materials

diff --git a/Engine/Source/Runtime/Experimental/Chaos/Private/Chaos/PBDCollisionConstraints.cpp b/Engine/Source/Runtime/Experimental/Chaos/Private/Chaos/PBDCollisionConstraints.cpp
--- a/Engine/Source/Runtime/Experimental/Chaos/Private/Chaos/PBDCollisionConstraints.cpp
+++ b/Engine/Source/Runtime/Experimental/Chaos/Private/Chaos/PBDCollisionConstraints.cpp
@@ -137,10 +137,36 @@ namespace Chaos
 		PostApplyPushOutCallback = nullptr;
 	}
 
+	const FChaosPhysicsMaterial* GetPhysicsMaterial(const TGeometryParticleHandle<FReal, 3>* Particle, const FImplicitObject* Geom, const TArrayCollectionArray<TSerializablePtr<FChaosPhysicsMaterial>>& PhysicsMaterials)
+	{
+		// Use the per-particle material if it exists
+		const FChaosPhysicsMaterial* PhysicsMaterial = Particle->AuxilaryValue(PhysicsMaterials).Get();
+
+		// If no particle material, see if the shape has one
+		if (PhysicsMaterial == nullptr)
+		{
+			// @todo(chaos): handle materials for meshes etc
+		
+			for (const TUniquePtr<FPerShapeData>& ShapeData : Particle->ShapesArray())
+			{
+				if (Geom == ShapeData->GetGeometry().Get())
+				{
+					if (ShapeData->GetMaterials().Num() > 0)
+					{
+						PhysicsMaterial = ShapeData->GetMaterials()[0].Get();
+					}
+					break;
+				}
+			}
+		}
+
+		return PhysicsMaterial;
+	}
+
 	void FPBDCollisionConstraints::UpdateConstraintMaterialProperties(FCollisionConstraintBase& Constraint)
 	{
-		TSerializablePtr<FChaosPhysicsMaterial> PhysicsMaterial0 = Constraint.Particle[0]->AuxilaryValue(MPhysicsMaterials);
-		TSerializablePtr<FChaosPhysicsMaterial> PhysicsMaterial1 = Constraint.Particle[1]->AuxilaryValue(MPhysicsMaterials);
+		const FChaosPhysicsMaterial* PhysicsMaterial0 = GetPhysicsMaterial(Constraint.Particle[0], Constraint.Manifold.Implicit[0], MPhysicsMaterials);
+		const FChaosPhysicsMaterial* PhysicsMaterial1 = GetPhysicsMaterial(Constraint.Particle[1], Constraint.Manifold.Implicit[1], MPhysicsMaterials);
 
 		FCollisionContact& Contact = Constraint.Manifold;
 		if (PhysicsMaterial0 && PhysicsMaterial1)
