Added desync detection when input is missing entirely

diff --git a/Engine/Source/Programs/HeadlessChaos/Private/HeadlessChaosTestRewind.cpp b/Engine/Source/Programs/HeadlessChaos/Private/HeadlessChaosTestRewind.cpp
--- a/Engine/Source/Programs/HeadlessChaos/Private/HeadlessChaosTestRewind.cpp
+++ b/Engine/Source/Programs/HeadlessChaos/Private/HeadlessChaosTestRewind.cpp
@@ -672,8 +672,7 @@ namespace ChaosTest {
 				}
 				else
 				{
-					//todo: desync from lack of input
-					//EXPECT_EQ(KinFutureStatus,EFutureQueryResult::Desync);
+					EXPECT_EQ(KinFutureStatus,EFutureQueryResult::Desync);
 				}
 			}
 		}
@@ -693,8 +692,7 @@ namespace ChaosTest {
 			}
 			else
 			{
-				//todo: desync from lack of input
-				//EXPECT_EQ(KinState.X()[2],50);	//in resim we didn't do second move, so recorded data must be updated
+				EXPECT_EQ(KinState.X()[2],50);	//in resim we didn't do second move, so recorded data must be updated
 			}
 		}
 
diff --git a/Engine/Source/Runtime/Experimental/ChaosSolvers/Public/RewindData.h b/Engine/Source/Runtime/Experimental/ChaosSolvers/Public/RewindData.h
--- a/Engine/Source/Runtime/Experimental/ChaosSolvers/Public/RewindData.h
+++ b/Engine/Source/Runtime/Experimental/ChaosSolvers/Public/RewindData.h
@@ -448,7 +448,7 @@ public:
 	: Managers(NumFrames+1)	//give 1 extra for saving at head
 	, CurFrame(0)
 	, LatestFrame(0)
-	, CurWave(0)
+	, CurWave(1)
 	, FramesSaved(0)
 	, DataIdxOffset(0)
 	, bNeedsSave(false)
@@ -492,7 +492,7 @@ public:
 				//GetStateAtFrameImp returns a pointer from the TArray that holds state data
 				//But it's possible that we'll need to save state from head, which would grow that TArray
 				//So preallocate just in case
-				FGeometryParticleStateBase& LatestState = DirtyParticleInfo.AddFrame(CurFrame, CurWave);
+				FGeometryParticleStateBase& LatestState = DirtyParticleInfo.AddFrame(CurFrame);
 			
 				if(const FGeometryParticleStateBase* RewindState = GetStateAtFrameImp(DirtyParticleInfo, Frame))
 				{
@@ -514,7 +514,12 @@ public:
 		CurFrame = Frame;
 		bNeedsSave = false;
 		FramesSaved = 0; //can't rewind before this point. This simplifies saving the state at head
-		//++CurWave;
+		++CurWave;
+		if(CurWave == 0)
+		{
+			//0 indicates nothing written so skip it
+			CurWave = 1;
+		}
 
 		return true;
 	}
@@ -585,10 +590,20 @@ public:
 		//remove any old dirty particles
 		for(int32 DirtyIdx = AllDirtyParticles.Num() - 1; DirtyIdx >= 0; --DirtyIdx)
 		{
-			if(AllDirtyParticles[DirtyIdx].LastDirtyFrame < EarliestFrame)
+			FDirtyParticleInfo& Info = AllDirtyParticles[DirtyIdx];
+			if(Info.LastDirtyFrame < EarliestFrame)
 			{
 				RemoveParticle(AllDirtyParticles[DirtyIdx].CachedUniqueIdx);
 			}
+			else if(IsResim())
+			{
+				//During a resim it's possible the user will not dirty a particle that was previously dirty.
+				//If this happens we need to mark the particle as desynced
+				if(!Info.bDesync && Info.GTDirtyOnFrame[CurFrame-1].MissingWrite(CurFrame-1, CurWave))
+				{
+					Info.Desync(CurFrame-1, LatestFrame);
+				}
+			}
 		}
 	}
 
@@ -645,6 +660,7 @@ public:
 			const auto PTParticle = Proxy->GetHandle();
 			FDirtyParticleInfo& Info = FindOrAddParticle(*Proxy->GetParticle(),PTParticle->UniqueIdx());
 			Info.LastDirtyFrame = CurFrame;
+			Info.GTDirtyOnFrame[CurFrame].SetWave(CurFrame,CurWave);
 
 			//check if particle has desynced
 			if(bResim)
@@ -652,12 +668,14 @@ public:
 				FGeometryParticleState FutureState(*Proxy->GetParticle());
 				if(GetFutureStateAtFrame(FutureState,CurFrame) == EFutureQueryResult::Ok)
 				{
-					Info.bDesync = FutureState.IsDesynced(SrcManager, SrcDataIdx, Dirty);
+					if(FutureState.IsDesynced(SrcManager, SrcDataIdx, Dirty))
+					{
+						Info.Desync(CurFrame-1, LatestFrame);
+					}
 				}
-				else
+				else if(!Info.bDesync)
 				{
-					//Most likely already desynced, but if it's a previously untracked object mark it as desynced
-					Info.bDesync = true;
+					Info.Desync(CurFrame-1, LatestFrame);
 				}
 			}
 
@@ -668,7 +686,7 @@ public:
 				//In fact, since we have a circular buffer the prev state could end up overwriting head which we need for fast forward
 				if(!bResim || FramesSaved > 0)
 				{
-					FGeometryParticleStateBase& LatestState = Info.AddFrame(CurFrame-1,CurWave);
+					FGeometryParticleStateBase& LatestState = Info.AddFrame(CurFrame-1);
 					LatestState.SyncPrevFrame(DestManager,DestDataIdx,Dirty);
 					CoalesceBack(Info.Frames,CurFrame-1);
 				}
@@ -677,7 +695,7 @@ public:
 			//If dynamics are dirty we must record them immediately because the sim resets them to 0
 			if(Dirty.ParticleData.IsDirty(EParticleFlags::Dynamics))
 			{
-				FGeometryParticleStateBase& LatestState = Info.AddFrame(CurFrame, CurWave);
+				FGeometryParticleStateBase& LatestState = Info.AddFrame(CurFrame);
 				LatestState.SyncDirtyDynamics(DestManager,DestDataIdx,Dirty.ParticleData,SrcManager);
 			}
 		};
@@ -722,7 +740,7 @@ public:
 			FDirtyPropertiesManager& DestManager = *Managers[CurFrame-1].Manager;
 
 			//sim-writable properties changed at head, so we must write down what they were
-			FGeometryParticleStateBase& LatestState = Info.AddFrame(CurFrame-1, CurWave);
+			FGeometryParticleStateBase& LatestState = Info.AddFrame(CurFrame-1);
 			LatestState.SyncSimWritablePropsFromSim(DestManager,DestDataIdx,Rigid);
 
 			//update any previous frames that were pointing at head
@@ -732,6 +750,24 @@ public:
 
 private:
 
+	struct FDirtyFrameInfo
+	{
+		int32 Frame;	//needed to protect against stale entries in circular buffer
+		uint8 Wave;
+
+		void SetWave(int32 InFrame, uint8 InWave)
+		{
+			Frame = InFrame;
+			Wave = InWave;
+		}
+
+		bool MissingWrite(int32 InFrame, uint8 InWave) const
+		{
+			//If this is not a stale entry and it was written to, but not during this latest sim
+			return (Wave != 0 && Frame == InFrame) && Wave != InWave;
+		}
+	};
+
 	class FFrameInfo
 	{
 	public:
@@ -741,51 +777,54 @@ private:
 
 		}
 
-		FGeometryParticleStateBase* GetState(int32 Frame, int8 Wave)
+		FGeometryParticleStateBase* GetState(int32 Frame)
 		{
-			return (bSet && Frame == RecordedFrame && Wave == RewindWave) ? &State : nullptr;
+			return (bSet && Frame == RecordedFrame) ? &State : nullptr;
 		}
 
-		const FGeometryParticleStateBase* GetState(int32 Frame, int8 Wave) const
+		const FGeometryParticleStateBase* GetState(int32 Frame) const
 		{
-			return (bSet && Frame == RecordedFrame && Wave == RewindWave) ? &State : nullptr;
+			return (bSet && Frame == RecordedFrame) ? &State : nullptr;
 		}
 
-		FGeometryParticleStateBase& GetStateChecked(int32 Frame, int8 Wave)
+		FGeometryParticleStateBase& GetStateChecked(int32 Frame)
 		{
-			check(bSet && Frame == RecordedFrame && Wave == RewindWave);
+			check(bSet && Frame == RecordedFrame);
 			return State;
 		}
 
-		const FGeometryParticleStateBase& GetStateChecked(int32 Frame, int8 Wave) const
+		const FGeometryParticleStateBase& GetStateChecked(int32 Frame) const
 		{
-			check(bSet && Frame == RecordedFrame && Wave == RewindWave);
+			check(bSet && Frame == RecordedFrame);
 			return State;
 		}
 
-		FGeometryParticleStateBase& NewState(int32 Frame, int8 Wave)
+		FGeometryParticleStateBase& NewState(int32 Frame)
 		{
 			RecordedFrame = Frame;
-			RewindWave = Wave;
 			bSet = true;
 			State = FGeometryParticleStateBase();
 			return State;
 		}
 
+		void ClearState()
+		{
+			bSet = false;
+		}
+
 	private:
 		FGeometryParticleStateBase State;
 		int32 RecordedFrame;
-		int8 RewindWave;
 		bool bSet;
 	};
 
 	void CoalesceBack(TCircularBuffer<FFrameInfo>& Frames, int32 LatestIdx)
 	{
-		const FGeometryParticleStateBase& LatestState = Frames[LatestIdx].GetStateChecked(LatestIdx, CurWave);
+		const FGeometryParticleStateBase& LatestState = Frames[LatestIdx].GetStateChecked(LatestIdx);
 		const int32 EarliestFrame = LatestIdx - FramesSaved;
 		for(int32 FrameIdx = LatestIdx - 1; FrameIdx >= EarliestFrame; --FrameIdx)
 		{
-			if(FGeometryParticleStateBase* State = Frames[FrameIdx].GetState(FrameIdx, CurWave))
+			if(FGeometryParticleStateBase* State = Frames[FrameIdx].GetState(FrameIdx))
 			{
 				if(State->CoalesceState(LatestState) == false)
 				{
@@ -809,13 +848,15 @@ private:
 	struct FDirtyParticleInfo
 	{
 		TCircularBuffer<FFrameInfo> Frames;
+		TCircularBuffer<FDirtyFrameInfo> GTDirtyOnFrame;
 		TGeometryParticle<FReal,3>* Particle;
 		FUniqueIdx CachedUniqueIdx;	//Needed when manipulating on physics thread and Particle data cannot be read
 		int32 LastDirtyFrame;	//Track how recently this was made dirty
 		bool bDesync;
 
-		FDirtyParticleInfo(TGeometryParticle<FReal, 3>& UnsafeGTParticle, const FUniqueIdx UniqueIdx, const int32 CurFrame, const int32 NumFrames)
+		FDirtyParticleInfo(TGeometryParticle<FReal,3>& UnsafeGTParticle,const FUniqueIdx UniqueIdx,const int32 CurFrame,const int32 NumFrames)
 		: Frames(NumFrames)
+		, GTDirtyOnFrame(NumFrames)
 		, Particle(&UnsafeGTParticle)
 		, CachedUniqueIdx(UniqueIdx)
 		, LastDirtyFrame(CurFrame)
@@ -824,22 +865,31 @@ private:
 
 		}
 
-		FGeometryParticleStateBase& AddFrame(int32 FrameIdx, int8 Wave)
+		FGeometryParticleStateBase& AddFrame(int32 FrameIdx)
 		{
 			FFrameInfo& Info = Frames[FrameIdx];
-			if(FGeometryParticleStateBase* State = Info.GetState(FrameIdx, Wave))
+			if(FGeometryParticleStateBase* State = Info.GetState(FrameIdx))
 			{
 				return *State;
 			}
 
-			return Info.NewState(FrameIdx, Wave);
+			return Info.NewState(FrameIdx);
+		}
+
+		void Desync(int32 StartDesync, int32 LastFrame)
+		{
+			bDesync = true;
+			for(int32 Frame = StartDesync; Frame <= LastFrame; ++Frame)
+			{
+				Frames[Frame].ClearState();
+			}
 		}
 	};
 
 	const FGeometryParticleStateBase* GetStateAtFrameImp(const FDirtyParticleInfo& Info,int32 Frame) const
 	{
 		const TCircularBuffer<FFrameInfo>& Frames = Info.Frames;
-		if(const FGeometryParticleStateBase* FrameState = Frames[Frame].GetState(Frame,CurWave))
+		if(const FGeometryParticleStateBase* FrameState = Frames[Frame].GetState(Frame))
 		{
 			return FrameState;
 		}
@@ -849,7 +899,7 @@ private:
 
 		for(int32 FrameIdx = Frame + 1; FrameIdx <= LatestFrame; ++FrameIdx)
 		{
-			if(const FGeometryParticleStateBase* FrameState = Frames[FrameIdx].GetState(FrameIdx,CurWave))
+			if(const FGeometryParticleStateBase* FrameState = Frames[FrameIdx].GetState(FrameIdx))
 			{
 				return FrameState;
 			}
@@ -909,7 +959,7 @@ private:
 	TArray<FDirtyParticleInfo> AllDirtyParticles;
 	int32 CurFrame;
 	int32 LatestFrame;
-	int8 CurWave;
+	uint8 CurWave;
 	int32 FramesSaved;
 	int32 DataIdxOffset;
 	bool bNeedsSave;	//Indicates that some data is pointing at head and requires saving before a rewind
