Chaos - disabled deferred collision update for main physics scene (collision callbacks may rely on valid contact position and normal)

diff --git a/Engine/Source/Runtime/Experimental/ChaosSolvers/Private/PBDRigidsSolver.cpp b/Engine/Source/Runtime/Experimental/ChaosSolvers/Private/PBDRigidsSolver.cpp
--- a/Engine/Source/Runtime/Experimental/ChaosSolvers/Private/PBDRigidsSolver.cpp
+++ b/Engine/Source/Runtime/Experimental/ChaosSolvers/Private/PBDRigidsSolver.cpp
@@ -46,7 +46,7 @@ FAutoConsoleVariableRef CVarChaosSolverCollisionDefaultPushoutIterations(TEXT("p
 int32 ChaosSolverCleanupCommandsOnDestruction = 1;
 FAutoConsoleVariableRef CVarChaosSolverCleanupCommandsOnDestruction(TEXT("p.Chaos.Solver.CleanupCommandsOnDestruction"), ChaosSolverCleanupCommandsOnDestruction, TEXT("Whether or not to run internal command queue cleanup on solver destruction (0 = no cleanup, >0 = cleanup all commands)"));
 
-int32 ChaosSolverCollisionDeferNarrowPhase = 1;
+int32 ChaosSolverCollisionDeferNarrowPhase = 0;
 FAutoConsoleVariableRef CVarChaosSolverCollisionDeferNarrowPhase(TEXT("p.Chaos.Solver.Collision.DeferNarrowPhase"), ChaosSolverCollisionDeferNarrowPhase, TEXT("Create contacts for all broadphase pairs, perform NarrowPhase later."));
 
 int32 ChaosSolverCollisionUseManifolds = 0;
