ToolsFramework: add accessors to get active tool names
#rnx

diff --git a/Engine/Source/Runtime/Experimental/InteractiveToolsFramework/Private/InteractiveToolManager.cpp b/Engine/Source/Runtime/Experimental/InteractiveToolsFramework/Private/InteractiveToolManager.cpp
--- a/Engine/Source/Runtime/Experimental/InteractiveToolsFramework/Private/InteractiveToolManager.cpp
+++ b/Engine/Source/Runtime/Experimental/InteractiveToolsFramework/Private/InteractiveToolManager.cpp
@@ -239,6 +239,15 @@ UInteractiveToolBuilder* UInteractiveToolManager::GetActiveToolBuilder(EToolSide
 	return (Side == EToolSide::Left) ? ActiveLeftBuilder : ActiveRightBuilder;
 }
 
+FString UInteractiveToolManager::GetActiveToolName(EToolSide Side)
+{
+	if (GetActiveTool(Side) == nullptr)
+	{
+		return FString();
+	}
+	return (Side == EToolSide::Left) ? ActiveLeftToolName : ActiveLeftToolName;
+}
+
 
 bool UInteractiveToolManager::CanAcceptActiveTool(EToolSide Side)
 {
diff --git a/Engine/Source/Runtime/Experimental/InteractiveToolsFramework/Public/InteractiveToolManager.h b/Engine/Source/Runtime/Experimental/InteractiveToolsFramework/Public/InteractiveToolManager.h
--- a/Engine/Source/Runtime/Experimental/InteractiveToolsFramework/Public/InteractiveToolManager.h
+++ b/Engine/Source/Runtime/Experimental/InteractiveToolsFramework/Public/InteractiveToolManager.h
@@ -143,6 +143,12 @@ public:
 	 */
 	virtual UInteractiveToolBuilder* GetActiveToolBuilder(EToolSide Side);
 
+	/**
+	 * Get name of registered ToolBuilder that created active tool for given side, or empty string if no tool is active
+	 * @param Side which Side is being requested
+	 * @return name of tool, or empty string if no tool is active
+	 */
+	virtual FString GetActiveToolName(EToolSide Side);
 
 	/**
 	 * Check if an active Tool on the given Side can be Accepted in its current state
