Fixed crash in field evaluation when more than one per-collection command was logged against a geometry collection

diff --git a/Engine/Source/Runtime/Experimental/ChaosSolvers/Private/PhysicsProxy/GeometryCollectionPhysicsProxy.cpp b/Engine/Source/Runtime/Experimental/ChaosSolvers/Private/PhysicsProxy/GeometryCollectionPhysicsProxy.cpp
--- a/Engine/Source/Runtime/Experimental/ChaosSolvers/Private/PhysicsProxy/GeometryCollectionPhysicsProxy.cpp
+++ b/Engine/Source/Runtime/Experimental/ChaosSolvers/Private/PhysicsProxy/GeometryCollectionPhysicsProxy.cpp
@@ -1140,8 +1140,11 @@ void FGeometryCollectionPhysicsProxy::GetRelevantHandles(
 	EFieldResolutionType ResolutionType, 
 	bool bForce = true)
 {
-	if (bForce)
+	if(bForce)
 	{
+		Samples.Reset();
+		SampleIndices.Reset();
+
 		// only the local handles
 		TArray<FClusterHandle*>& ParticleHandles = GetSolverParticleHandles();
 		Handles.SetNumUninitialized(ParticleHandles.Num());
