Chaos - fix joint projection (point constraints only). Disabled by iteration counts atm

diff --git a/Engine/Source/Runtime/Engine/Private/PhysicsEngine/ImmediatePhysics/ImmediatePhysicsChaos/ImmediatePhysicsJointHandle_Chaos.cpp b/Engine/Source/Runtime/Engine/Private/PhysicsEngine/ImmediatePhysics/ImmediatePhysicsChaos/ImmediatePhysicsJointHandle_Chaos.cpp
--- a/Engine/Source/Runtime/Engine/Private/PhysicsEngine/ImmediatePhysics/ImmediatePhysicsChaos/ImmediatePhysicsJointHandle_Chaos.cpp
+++ b/Engine/Source/Runtime/Engine/Private/PhysicsEngine/ImmediatePhysics/ImmediatePhysicsChaos/ImmediatePhysicsJointHandle_Chaos.cpp
@@ -36,7 +36,7 @@ namespace ImmediatePhysics_Chaos
 	float ChaosImmediate_SoftAngularDampingScale = 1000;
 
 	float ChaosImmediate_JointMinLinearProjection = 0.0f;
-	float ChaosImmediate_JointMaxLinearProjection = 0.8f;
+	float ChaosImmediate_JointMaxLinearProjection = 1.0f;
 	float ChaosImmediate_JointMinAngularProjection = 0.0f;
 	float ChaosImmediate_JointMaxAngularProjection = 0.1f;
 
diff --git a/Engine/Source/Runtime/Experimental/Chaos/Private/Chaos/Joint/PBDJointSolverGaussSeidel.cpp b/Engine/Source/Runtime/Experimental/Chaos/Private/Chaos/Joint/PBDJointSolverGaussSeidel.cpp
--- a/Engine/Source/Runtime/Experimental/Chaos/Private/Chaos/Joint/PBDJointSolverGaussSeidel.cpp
+++ b/Engine/Source/Runtime/Experimental/Chaos/Private/Chaos/Joint/PBDJointSolverGaussSeidel.cpp
@@ -230,7 +230,7 @@ namespace Chaos
 
 		//NumActive += ApplyRotationProjection(Dt, SolverSettings, JointSettings);
 
-		NumActive += ApplyPositionProjection(Dt, SolverSettings, JointSettings);
+		NumActive += ApplyPositionProjections(Dt, SolverSettings, JointSettings);
 
 		return NumActive;
 	}
@@ -396,7 +396,7 @@ namespace Chaos
 	}
 
 
-	int32 FJointSolverGaussSeidel::ApplyRotationProjection(
+	int32 FJointSolverGaussSeidel::ApplyRotationProjections(
 		const FReal Dt,
 		const FPBDJointSolverSettings& SolverSettings,
 		const FPBDJointSettings& JointSettings)
@@ -587,6 +587,39 @@ namespace Chaos
 	}
 
 
+	int32 FJointSolverGaussSeidel::ApplyPositionProjections(
+		const FReal Dt,
+		const FPBDJointSolverSettings& SolverSettings,
+		const FPBDJointSettings& JointSettings)
+	{
+		int32 NumActive = 0;
+
+		const TVector<EJointMotionType, 3>& LinearMotion = JointSettings.LinearMotionTypes;
+		const TVector<bool, 3> bLinearLocked =
+		{
+			(LinearMotion[0] == EJointMotionType::Locked),
+			(LinearMotion[1] == EJointMotionType::Locked),
+			(LinearMotion[2] == EJointMotionType::Locked),
+		};
+		const TVector<bool, 3> bLinearLimted =
+		{
+			(LinearMotion[0] == EJointMotionType::Limited),
+			(LinearMotion[1] == EJointMotionType::Limited),
+			(LinearMotion[2] == EJointMotionType::Limited),
+		};
+
+		if (bLinearLocked[0] && bLinearLocked[1] && bLinearLocked[2])
+		{
+			// Hard point constraint (most common case)
+			NumActive += ApplyPointPositionProjection(Dt, SolverSettings, JointSettings);
+		}
+
+		// @todo(ccaulfield): support projection for other linear constraint configs
+
+		return NumActive;
+	}
+
+
 	//
 	//
 	//////////////////////////////////////////////////////////////////////////
@@ -1640,62 +1673,72 @@ namespace Chaos
 	}
 
 
-	int32 FJointSolverGaussSeidel::ApplyPositionProjection(
+	int32 FJointSolverGaussSeidel::ApplyPointPositionProjection(
 		const FReal Dt,
 		const FPBDJointSolverSettings& SolverSettings,
 		const FPBDJointSettings& JointSettings)
 	{
 		FReal LinearProjection = FPBDJointUtilities::GetLinearProjection(SolverSettings, JointSettings);
-		if ((LinearProjection > 0) && !JointSettings.bSoftLinearLimitsEnabled)
+		if (LinearProjection > 0)
 		{
-			// Apply a position correction with the parent body set to infinite mass, then correct the velocity.
-			FVec3 CX = FPBDJointUtilities::GetLimitedPositionError(JointSettings, Rs[0], Xs[1] - Xs[0]);
-			const FReal CXLen = CX.Size();
-			if (CXLen > KINDA_SMALL_NUMBER + PositionTolerance)
+			const FVec3 CX = Xs[1] - Xs[0];
+			bool bApplyProjection = (CX.Size() > KINDA_SMALL_NUMBER); // + PositionTolerance);
+			if (bApplyProjection && bChaos_Joint_ISPC_Enabled)
 			{
-				const FReal ParentMassScale = FMath::Max(0.0f, 1.0f - LinearProjection);
-				if (bChaos_Joint_ISPC_Enabled)
-				{
 #if INTEL_ISPC
-					const FReal Stiffness = FPBDJointUtilities::GetLinearStiffness(SolverSettings, JointSettings);
-					ispc::ApplyPositionProjection((ispc::FJointSolverGaussSeidel*)this, (ispc::FVector&)CX, CXLen, ParentMassScale, Stiffness);
+				ispc::ApplyPointPositionProjection((ispc::FJointSolverGaussSeidel*)this, (ispc::FVector&)CX, LinearProjection);
 #endif
+			}
+			else if (bApplyProjection)
+			{
+				// Apply a position correction with the parent body set to infinite mass
+				{
+					const FMatrix33 J = Utilities::ComputeJointFactorMatrix(Xs[1] - Ps[1], InvIs[1], InvMs[1]);
+					const FMatrix33 IJ = J.Inverse();
+					const FVec3 DX = Utilities::Multiply(IJ, CX);
+					const FVec3 DP1 = -InvMs[1] * DX;
+					const FVec3 DR1 = Utilities::Multiply(InvIs[1], FVec3::CrossProduct(Xs[1] - Ps[1], -DX));
+					ApplyDelta(1, LinearProjection, DP1, DR1);
 				}
-				else
+
+				// Apply a velocity correction to set the relative velocity at the joint to zero
+				// NOTE: Velocity correction must come after position correction and derived state must be updated accordingly.
+				// NOTE: No parent mass scale here - velocity correction is shared between bodies normally
+				// (if we don't do this, velocity tends to accumulate down chains leading to whip-like behaviour)
 				{
-					const FVec3 CXDir = CX / CXLen;
+					// Calculate joint-space mass
+					FMatrix33 J = FMatrix33(0, 0, 0);
+					if (InvMs[0] > 0)
+					{
+						J = J + Utilities::ComputeJointFactorMatrix(Xs[0] - Ps[0], InvIs[0], InvMs[0]);
+					}
+					if (InvMs[1] > 0)
+					{
+						J = J + Utilities::ComputeJointFactorMatrix(Xs[1] - Ps[1], InvIs[1], InvMs[1]);
+					}
+					const FMatrix33 IJ = J.Inverse();
+
+					// Velocity correction
 					const FVec3 V0 = Vs[0] + FVec3::CrossProduct(Ws[0], Xs[0] - Ps[0]);
 					const FVec3 V1 = Vs[1] + FVec3::CrossProduct(Ws[1], Xs[1] - Ps[1]);
-					FVec3 CV = FVec3::DotProduct(V1 - V0, CXDir) * CXDir;
-
-					const FReal IM0 = ParentMassScale * InvMs[0];
-					const FReal IM1 = InvMs[1];
-					const FVec3 IIL0 = ParentMassScale * InvILs[0];
-					const FVec3& IIL1 = InvILs[1];
-					FMatrix33 J0 = (IM0 > 0) ? Utilities::ComputeJointFactorMatrix(Xs[0] - Ps[0], InvIs[0], IM0) : FMatrix33(0, 0, 0);
-					FMatrix33 J1 = Utilities::ComputeJointFactorMatrix(Xs[1] - Ps[1], InvIs[1], IM1);
-					FMatrix33 IJ = (J0 + J1).Inverse();
-
-					const FVec3 DX = Utilities::Multiply(IJ, CX);
+					const FVec3 CV = V1 - V0;
 					const FVec3 DV = Utilities::Multiply(IJ, CV);
 
-					const FVec3 DP0 = IM0 * DX;
-					const FVec3 DP1 = -IM1 * DX;
-					const FVec3 DR0 = Utilities::Multiply(InvIs[0], FVec3::CrossProduct(Xs[0] - Ps[0], DX));
-					const FVec3 DR1 = Utilities::Multiply(InvIs[1], FVec3::CrossProduct(Xs[1] - Ps[1], -DX));
-
-					const FVec3 DV0 = IM0 * DV;
-					const FVec3 DV1 = -IM1 * DV;
-					const FVec3 DW0 = Utilities::Multiply(InvIs[0], FVec3::CrossProduct(Xs[0] - Ps[0], DV));
-					const FVec3 DW1 = Utilities::Multiply(InvIs[1], FVec3::CrossProduct(Xs[1] - Ps[1], -DV));
-
-					const FReal Stiffness = FPBDJointUtilities::GetLinearStiffness(SolverSettings, JointSettings);
-					ApplyPositionDelta(Stiffness, DP0, DP1);
-					ApplyRotationDelta(Stiffness, DR0, DR1);
-					ApplyVelocityDelta(Stiffness, DV0, DW0, DV1, DW1);
+					if (InvMs[0] > 0)
+					{
+						const FVec3 DV0 = InvMs[0] * DV;
+						const FVec3 DW0 = Utilities::Multiply(InvIs[0], FVec3::CrossProduct(Xs[0] - Ps[0], DV));
+						ApplyVelocityDelta(0, LinearProjection, DV0, DW0);
+					}
+					if (InvMs[1] > 0)
+					{
+						const FVec3 DV1 = -InvMs[1] * DV;
+						const FVec3 DW1 = Utilities::Multiply(InvIs[1], FVec3::CrossProduct(Xs[1] - Ps[1], -DV));
+						ApplyVelocityDelta(1, LinearProjection, DV1, DW1);
+					}
 				}
-				return 1;
 			}
+			return bApplyProjection ? 1 : 0;
 		}
 		return 0;
 	}
diff --git a/Engine/Source/Runtime/Experimental/Chaos/Private/Chaos/Joint/PBDJointSolverGaussSeidel.ispc b/Engine/Source/Runtime/Experimental/Chaos/Private/Chaos/Joint/PBDJointSolverGaussSeidel.ispc
--- a/Engine/Source/Runtime/Experimental/Chaos/Private/Chaos/Joint/PBDJointSolverGaussSeidel.ispc
+++ b/Engine/Source/Runtime/Experimental/Chaos/Private/Chaos/Joint/PBDJointSolverGaussSeidel.ispc
@@ -222,6 +222,18 @@ static inline void ApplyVelocityDelta(
 	M->Ws[1] = M->Ws[1] + Stiffness * DW1;
 }
 
+
+static inline void ApplyVelocityDelta(
+	uniform FJointSolverGaussSeidel * uniform M,
+	const uniform int Index,
+	const uniform float Stiffness,
+	const uniform FVector& DV,
+	const uniform FVector& DW)
+{
+	M->Vs[Index] = M->Vs[Index] + Stiffness * DV;
+	M->Ws[Index] = M->Ws[Index] + Stiffness * DW;
+}
+
 export void ApplyPositionConstraintSoft(
 	uniform FJointSolverGaussSeidel * uniform M,
 	const uniform float Dt,
@@ -392,46 +404,57 @@ export void ApplyPointPositionConstraintDD(
 	ApplyRotationDelta(M, LinearStiffness, DR0, DR1);
 }
 
-export void ApplyPositionProjection(
+export void ApplyPointPositionProjection(
 	uniform FJointSolverGaussSeidel * uniform M,
 	const uniform FVector &CX,
-	const uniform float CXLen,
-	const uniform float ParentMassScale,
-	const uniform float Stiffness)
-{
-	const uniform FVector CXDir = CX / CXLen;
-	const uniform FVector V0 = M->Vs[0] + VectorCross(M->Ws[0], M->Xs[0] - M->Ps[0]);
-	const uniform FVector V1 = M->Vs[1] + VectorCross(M->Ws[1], M->Xs[1] - M->Ps[1]);
-	const uniform FVector CV = VectorDot(V1 - V0, CXDir) * CXDir;
-
-	const uniform float IM0 = ParentMassScale * M->InvMs[0];
-	const uniform float IM1 = M->InvMs[1];
-	const uniform FVector IIL0 = ParentMassScale * M->InvILs[0];
-	const uniform FVector& IIL1 = M->InvILs[1];
-	uniform FMatrix33 J0 = Matrix33Zero;
-	if(IM0 > 0)
+	const uniform float LinearProjection)
+{
+	// Apply a position correction with the parent body set to infinite mass
 	{
-		J0 = ComputeJointFactorMatrix(M->Xs[0] - M->Ps[0], M->InvIs[0], IM0);
+		const uniform FMatrix33 J = ComputeJointFactorMatrix(M->Xs[1] - M->Ps[1], M->InvIs[1], M->InvMs[1]);
+		const uniform FMatrix33 IJ = MatrixInverse(J);
+		const uniform FVector DX = VectorNegate(Multiply(IJ, CX));
+		const uniform FVector DP1 = M->InvMs[1] * DX;
+		const uniform FVector DR1 = Multiply(M->InvIs[1], VectorCross(M->Xs[1] - M->Ps[1], DX));
+		ApplyDelta(M, 1, LinearProjection, DP1, DR1);
 	}
-	const uniform FMatrix33 J1 = ComputeJointFactorMatrix(M->Xs[1] - M->Ps[1], M->InvIs[1], IM1);
-	const uniform FMatrix33 IJ = MatrixInverse(AddAB(J0, J1));
 
-	const uniform FVector DX = Multiply(IJ, CX);
-	const uniform FVector DV = Multiply(IJ, CV);
+	// Apply a velocity correction to set the relative velocity at the joint to zero
+	// NOTE: No parent mass scale here - velocity correction is shared between bodies normally
+	// (if we don't do this, velocity tends to accumulate down chains leading to whip-like behaviour)
+	{
+		// Calculate joint-space mass
+		uniform FMatrix33 J = SetMatrix33(0, 0, 0);
+		if (M->InvMs[0] > 0)
+		{
+			J = AddAB(J, ComputeJointFactorMatrix(M->Xs[0] - M->Ps[0], M->InvIs[0], M->InvMs[0]));
+		}
+		if (M->InvMs[1] > 0)
+		{
+			J = AddAB(J, ComputeJointFactorMatrix(M->Xs[1] - M->Ps[1], M->InvIs[1], M->InvMs[1]));
+		}
+		const uniform FMatrix33 IJ = MatrixInverse(J);
 
-	const uniform FVector DP0 = IM0 * DX;
-	const uniform FVector DP1 = -IM1 * DX;
-	const uniform FVector DR0 = Multiply(M->InvIs[0], VectorCross(M->Xs[0] - M->Ps[0], DX));
-	const uniform FVector DR1 = Multiply(M->InvIs[1], VectorCross(M->Xs[1] - M->Ps[1], VectorNegate(DX)));
+		// Velocity correction
+		const uniform FVector V0 = M->Vs[0] + VectorCross(M->Ws[0], M->Xs[0] - M->Ps[0]);
+		const uniform FVector V1 = M->Vs[1] + VectorCross(M->Ws[1], M->Xs[1] - M->Ps[1]);
+		const uniform FVector CV = V1 - V0;
+		const uniform FVector DV = Multiply(IJ, CV);
 
-	const uniform FVector DV0 = IM0 * DV;
-	const uniform FVector DV1 = -IM1 * DV;
-	const uniform FVector DW0 = Multiply(M->InvIs[0], VectorCross(M->Xs[0] - M->Ps[0], DV));
-	const uniform FVector DW1 = Multiply(M->InvIs[1], VectorCross(M->Xs[1] - M->Ps[1], VectorNegate(DV)));
+		if (M->InvMs[0] > 0)
+		{
+			const uniform FVector DV0 = M->InvMs[0] * DV;
+			const uniform FVector DW0 = Multiply(M->InvIs[0], VectorCross(M->Xs[0] - M->Ps[0], DV));
+			ApplyVelocityDelta(M, 0, LinearProjection, DV0, DW0);
+		}
+		if (M->InvMs[1] > 0)
+		{
+			const uniform FVector DV1 = M->InvMs[1] * VectorNegate(DV);
+			const uniform FVector DW1 = Multiply(M->InvIs[1], VectorCross(M->Xs[1] - M->Ps[1], VectorNegate(DV)));
+			ApplyVelocityDelta(M, 1, LinearProjection, DV1, DW1);
+		}
+	}
 
-	ApplyPositionDelta(M, Stiffness, DP0, DP1);
-	ApplyRotationDelta(M, Stiffness, DR0, DR1);
-	ApplyVelocityDelta(M, Stiffness, DV0, DW0, DV1, DW1);
 }
 
 //
diff --git a/Engine/Source/Runtime/Experimental/Chaos/Public/Chaos/Joint/PBDJointSolverGaussSeidel.h b/Engine/Source/Runtime/Experimental/Chaos/Public/Chaos/Joint/PBDJointSolverGaussSeidel.h
--- a/Engine/Source/Runtime/Experimental/Chaos/Public/Chaos/Joint/PBDJointSolverGaussSeidel.h
+++ b/Engine/Source/Runtime/Experimental/Chaos/Public/Chaos/Joint/PBDJointSolverGaussSeidel.h
@@ -130,7 +130,12 @@ namespace Chaos
 			const FPBDJointSolverSettings& SolverSettings,
 			const FPBDJointSettings& JointSettings);
 
-		int32 ApplyRotationProjection(
+		int32 ApplyPositionProjections(
+			const FReal Dt,
+			const FPBDJointSolverSettings& SolverSettings,
+			const FPBDJointSettings& JointSettings);
+
+		int32 ApplyRotationProjections(
 			const FReal Dt,
 			const FPBDJointSolverSettings& SolverSettings,
 			const FPBDJointSettings& JointSettings);
@@ -352,7 +357,7 @@ namespace Chaos
 			const FPBDJointSolverSettings& SolverSettings,
 			const FPBDJointSettings& JointSettings);
 
-		int32 ApplyPositionProjection(
+		int32 ApplyPointPositionProjection(
 			const FReal Dt,
 			const FPBDJointSolverSettings& SolverSettings,
 			const FPBDJointSettings& JointSettings);
