Fix compile error from engine merge with CHAOS_CHECKED on

diff --git a/Engine/Source/Runtime/Experimental/ChaosSolvers/Private/PBDRigidsSolver.cpp b/Engine/Source/Runtime/Experimental/ChaosSolvers/Private/PBDRigidsSolver.cpp
--- a/Engine/Source/Runtime/Experimental/ChaosSolvers/Private/PBDRigidsSolver.cpp
+++ b/Engine/Source/Runtime/Experimental/ChaosSolvers/Private/PBDRigidsSolver.cpp
@@ -151,10 +151,6 @@ namespace Chaos
 				// If time remains, then log why we have lost energy over the timestep.
 				if (TimeRemaining > 0.f)
 				{
-					auto& GeomCollectionParticles = MSolver->GetEvolution()->GetParticles().GetGeometryCollectionParticles();
-					Obj->FieldForcesUpdateCallback(MSolver, GeomCollectionParticles, Forces, Torques, MSolver->GetSolverTime());
-					auto& ClusteredParticles = MSolver->GetEvolution()->GetParticles().GetClusteredParticles();
-					Obj->FieldForcesUpdateCallback(MSolver, ClusteredParticles, Forces, Torques, MSolver->GetSolverTime());
 					if (StepsRemaining == 0)
 					{
 						UE_LOG(LogPBDRigidsSolver, Warning, TEXT("AdvanceOneTimeStepTask::DoWork() - Energy lost over %fs due to too many substeps over large timestep"), TimeRemaining);
