Add a test for cluster breaking - RigidBodies_ClusterTest_BreakClusterParticle().  Fix strain propagation.

diff --git a/Engine/Source/Runtime/Experimental/Chaos/Private/Chaos/PBDRigidClustering.cpp b/Engine/Source/Runtime/Experimental/Chaos/Private/Chaos/PBDRigidClustering.cpp
--- a/Engine/Source/Runtime/Experimental/Chaos/Private/Chaos/PBDRigidClustering.cpp
+++ b/Engine/Source/Runtime/Experimental/Chaos/Private/Chaos/PBDRigidClustering.cpp
@@ -1083,21 +1083,24 @@ namespace Chaos
 		TPBDRigidParticleHandle<T, d>* CurrentNode)
 	{
 		SCOPE_CYCLE_COUNTER(STAT_PromoteStrains);
-
-		T Result = 0;
-		if (MChildren.Contains(CurrentNode))
+		if (TPBDRigidClusteredParticleHandle<T, d>* ClusteredCurrentNode = CurrentNode->CastToClustered())
 		{
-			for (TPBDRigidParticleHandle<T, d>* Child : MChildren[CurrentNode])
+			T ChildrenStrains = 0;
+			if (MChildren.Contains(CurrentNode))
 			{
-				Result += PromoteStrains(Child);
+				for (TPBDRigidParticleHandle<T, d>* Child : MChildren[CurrentNode])
+				{
+					ChildrenStrains += PromoteStrains(Child);
+				}
 			}
+			else
+			{
+				return ClusteredCurrentNode->Strains();
+			}
+			ClusteredCurrentNode->SetStrains(ClusteredCurrentNode->Strains() + ChildrenStrains);
+			return ClusteredCurrentNode->Strains();
 		}
-		else
-		{
-			return CurrentNode->CastToClustered()->Strains();
-		}
-		CurrentNode->CastToClustered()->Strains() += Result;
-		return Result;
+		return (T)0.0;
 	}
 
 	DECLARE_CYCLE_STAT(TEXT("TPBDRigidClustering<>::UpdateKinematicProperties()"), STAT_UpdateKinematicProperties, STATGROUP_Chaos);
