fix copypasta error caught by static analysis
#rnx

diff --git a/Engine/Source/Runtime/Experimental/InteractiveToolsFramework/Private/InteractiveToolManager.cpp b/Engine/Source/Runtime/Experimental/InteractiveToolsFramework/Private/InteractiveToolManager.cpp
--- a/Engine/Source/Runtime/Experimental/InteractiveToolsFramework/Private/InteractiveToolManager.cpp
+++ b/Engine/Source/Runtime/Experimental/InteractiveToolsFramework/Private/InteractiveToolManager.cpp
@@ -245,7 +245,7 @@ FString UInteractiveToolManager::GetActiveToolName(EToolSide Side)
 	{
 		return FString();
 	}
-	return (Side == EToolSide::Left) ? ActiveLeftToolName : ActiveLeftToolName;
+	return (Side == EToolSide::Left) ? ActiveLeftToolName : ActiveRightToolName;
 }
 
 
