Desync when simulated particle changes from GT

diff --git a/Engine/Source/Programs/HeadlessChaos/Private/HeadlessChaosTestRewind.cpp b/Engine/Source/Programs/HeadlessChaos/Private/HeadlessChaosTestRewind.cpp
--- a/Engine/Source/Programs/HeadlessChaos/Private/HeadlessChaosTestRewind.cpp
+++ b/Engine/Source/Programs/HeadlessChaos/Private/HeadlessChaosTestRewind.cpp
@@ -698,6 +698,81 @@ namespace ChaosTest {
 
 
 
+		// Throw out the proxy
+		Solver->UnregisterObject(Particle.Get());
+
+		Module->DestroySolver(Solver);
+	}
+
+	GTEST_TEST(RewindTest,ResimDesyncAfterMissingTeleport)
+	{
+		auto Sphere = TSharedPtr<FImplicitObject,ESPMode::ThreadSafe>(new TSphere<float,3>(TVector<float,3>(0),10));
+
+		FChaosSolversModule* Module = FChaosSolversModule::GetModule();
+		Module->ChangeThreadingMode(EChaosThreadingMode::SingleThread);
+
+		// Make a solver
+		FPhysicsSolver* Solver = Module->CreateSolver(nullptr,ESolverFlags::Standalone);
+		Solver->SetEnabled(true);
+
+		Solver->EnableRewindCapture(7);
+
+		// Make particles
+		auto Particle = TPBDRigidParticle<float,3>::CreateParticle();
+
+		Particle->SetGeometry(Sphere);
+		Solver->RegisterObject(Particle.Get());
+		Particle->SetGravityEnabled(true);
+
+		const int LastStep = 11;
+		TArray<FVec3> X;
+
+		for(int Step = 0; Step <= LastStep; ++Step)
+		{
+			if(Step == 7)
+			{
+				Particle->SetX(FVec3(0,0,5));
+			}
+
+			if(Step == 9)
+			{
+				Particle->SetX(FVec3(0,0,1));
+			}
+			X.Add(Particle->X());
+			TickSolverHelper(Module,Solver);
+		}
+		X.Add(Particle->X());
+
+		const int RewindStep = 5;
+
+		FRewindData* RewindData = Solver->GetRewindData();
+		EXPECT_TRUE(RewindData->RewindToFrame(RewindStep));
+
+		for(int Step = RewindStep; Step <= LastStep; ++Step)
+		{
+			FGeometryParticleState FutureState(*Particle);
+			EXPECT_EQ(RewindData->GetFutureStateAtFrame(FutureState,Step+1), Step < 10 ? EFutureQueryResult::Ok : EFutureQueryResult::Desync);
+			if(Step < 10)
+			{
+				EXPECT_EQ(X[Step+1][2],FutureState.X()[2]);
+			}
+
+			if(Step == 7)
+			{
+				Particle->SetX(FVec3(0,0,5));
+			}
+
+			//skip step 9 SetX to trigger a desync
+
+			TickSolverHelper(Module,Solver);
+
+			//can't compare future with end of frame because we overwrite the result
+			if(Step != 6 && Step != 8 && Step < 9)
+			{
+				EXPECT_EQ(Particle->X()[2],FutureState.X()[2]);
+			}
+		}
+
 		// Throw out the proxy
 		Solver->UnregisterObject(Particle.Get());
 
diff --git a/Engine/Source/Runtime/Experimental/ChaosSolvers/Public/RewindData.h b/Engine/Source/Runtime/Experimental/ChaosSolvers/Public/RewindData.h
--- a/Engine/Source/Runtime/Experimental/ChaosSolvers/Public/RewindData.h
+++ b/Engine/Source/Runtime/Experimental/ChaosSolvers/Public/RewindData.h
@@ -10,6 +10,13 @@
 namespace Chaos
 {
 
+enum EDesyncResult
+{
+	InSync, //both have entries and are identical, or both have no entries
+	Desync, //both have entries but they are different
+	NeedInfo //one of the entries is missing. Need more context to determine whether desynced
+};
+
 template <typename T,EParticleProperty PropName>
 class TParticleStateProperty
 {
@@ -52,29 +59,6 @@ public:
 		SetFunc(NewVal);
 	}
 
-	bool IsDesynced(const FDirtyPropertiesManager& SrcManager, const int32 DataIdxIn, const FDirtyProxy& Dirty) const
-	{
-		const FParticleDirtyFlags Flags = Dirty.ParticleData.GetFlags();
-		if(Flags.IsDirty(ParticlePropToFlag(PropName)))
-		{
-			if(Manager)
-			{
-				//both dirty, if value differs we're desynced
-				return !(SrcManager.GetParticlePool<T,PropName>().GetElement(DataIdxIn) == Manager->GetParticlePool<T,PropName>().GetElement(Idx));
-			}
-			else
-			{
-				//Should have an entry but don't, so desynced
-				return true;
-			}
-		}
-		else
-		{
-			//if they're clean and we have no state then we're desynced
-			return Manager != nullptr;
-		}
-	}
-
 	template <typename LambdaSet>
 	void SyncRemoteData(FDirtyPropertiesManager& InManager,int32 InIdx, const FParticleDirtyData& DirtyData, const LambdaSet& SetFunc)
 	{
@@ -89,6 +73,24 @@ public:
 		return Manager != nullptr;
 	}
 
+	EDesyncResult DesyncQuery(const FDirtyPropertiesManager& SrcManager, const int32 DataIdxIn, const FParticleDirtyFlags Flags) const
+	{
+		if(Flags.IsDirty(ParticlePropToFlag(PropName)))
+		{
+			if(Manager)
+			{
+				//both dirty, if value differs we're desynced
+				return (SrcManager.GetParticlePool<T,PropName>().GetElement(DataIdxIn) == Manager->GetParticlePool<T,PropName>().GetElement(Idx)) ? EDesyncResult::InSync : EDesyncResult::Desync;
+			}
+		}
+		else if(Manager == nullptr)
+		{
+			return EDesyncResult::InSync;	//Both have no entry so in sync
+		}
+
+		return EDesyncResult::NeedInfo;	//entry is missing in one or the other but not both, so we don't know if it's desynced
+	}
+
 private:
 	FDirtyPropertiesManager* Manager;
 	int32 Idx;
@@ -107,32 +109,38 @@ inline bool SimWritablePropsMayChange(const TGeometryParticleHandle<FReal,3>& Ha
 class FGeometryParticleStateBase
 {
 public:
-	const FVec3& X(const TGeometryParticle<FReal,3>& Particle) const
+	template <typename TParticle>
+	const FVec3& X(const TParticle& Particle) const
 	{
 		return ParticlePositionRotation.IsSet() ? ParticlePositionRotation.Read().X : Particle.X();
 	}
 
-	const FRotation3& R(const TGeometryParticle<FReal,3>& Particle) const
+	template <typename TParticle>
+	const FRotation3& R(const TParticle& Particle) const
 	{
 		return ParticlePositionRotation.IsSet() ? ParticlePositionRotation.Read().R : Particle.R();
 	}
 
-	const FVec3& V(const TGeometryParticle<FReal,3>& Particle) const
+	template <typename TParticle>
+	const FVec3& V(const TParticle& Particle) const
 	{
 		return Velocities.IsSet() ? Velocities.Read().V : Particle.CastToKinematicParticle()->V();
 	}
 
-	const FVec3& W(const TGeometryParticle<FReal,3>& Particle) const
+	template <typename TParticle>
+	const FVec3& W(const TParticle& Particle) const
 	{
 		return Velocities.IsSet() ? Velocities.Read().W : Particle.CastToKinematicParticle()->W();
 	}
 
-	TSerializablePtr<FImplicitObject> Geometry(const TGeometryParticle<FReal,3>& Particle) const
+	template <typename TParticle>
+	TSerializablePtr<FImplicitObject> Geometry(const TParticle& Particle) const
 	{
 		return NonFrequentData.IsSet() ? MakeSerializable(NonFrequentData.Read().Geometry) : Particle.Geometry();
 	}
 
-	const FVec3& F(const TGeometryParticle<FReal,3>& Particle) const
+	template <typename TParticle>
+	const FVec3& F(const TParticle& Particle) const
 	{
 		return Dynamics.IsSet() ? Dynamics.Read().F : Particle.CastToRigidParticle()->F();
 	}
@@ -342,16 +350,37 @@ public:
 		return bCoalesced;
 	}
 
-	bool IsDesynced(const FDirtyPropertiesManager& SrcManager, const int32 DataIdxIn, const FDirtyProxy& Dirty) const
+	bool IsDesynced(const FDirtyPropertiesManager& SrcManager, const int32 DataIdxIn, const TGeometryParticleHandle<FReal,3>& Handle, const FParticleDirtyFlags Flags) const
 	{
-		//Don't use Particle, should use dirty data which is thread safe
-		//We most likely will do this on GT, but maybe we could do other work while running the sim tick in parallel?
 		bool Desynced = false;
-		Desynced = ParticlePositionRotation.IsDesynced(SrcManager,DataIdxIn,Dirty);
+		{
+			EDesyncResult Status = ParticlePositionRotation.DesyncQuery(SrcManager, DataIdxIn, Flags);
+			if(Status == EDesyncResult::Desync)
+			{
+				return true;
+			}
+			else if(Status == EDesyncResult::NeedInfo)
+			{
+				if(Handle.X() != X(Handle))
+				{
+					return true;
+				}
+
+				if(Handle.R() != R(Handle))
+				{
+					return true;
+				}
+			}
+		}
+
+		//TODO: test other properties, should probably find a better way to do this to avoid getting into the individual variables
+		
+		/*Desynced = ParticlePositionRotation.IsDesynced(SrcManager,DataIdxIn,Dirty);
 		Desynced = Desynced || NonFrequentData.IsDesynced(SrcManager,DataIdxIn,Dirty);
 		Desynced = Desynced || Velocities.IsDesynced(SrcManager,DataIdxIn,Dirty);
 		Desynced = Desynced || Dynamics.IsDesynced(SrcManager,DataIdxIn,Dirty);
-		return Desynced;
+		return Desynced;*/
+		return false;
 	}
 
 private:
@@ -424,9 +453,9 @@ public:
 		State = InState;
 	}
 
-	bool IsDesynced(const FDirtyPropertiesManager& SrcManager, const int32 DataIdxIn, const FDirtyProxy& Dirty) const
+	bool IsDesynced(const FDirtyPropertiesManager& SrcManager, const int32 DataIdxIn, const TGeometryParticleHandle<FReal,3>& Handle, const FParticleDirtyFlags Flags) const
 	{
-		return State.IsDesynced(SrcManager,DataIdxIn,Dirty);
+		return State.IsDesynced(SrcManager,DataIdxIn,Handle,Flags);
 	}
 
 private:
@@ -668,7 +697,7 @@ public:
 				FGeometryParticleState FutureState(*Proxy->GetParticle());
 				if(GetFutureStateAtFrame(FutureState,CurFrame) == EFutureQueryResult::Ok)
 				{
-					if(FutureState.IsDesynced(SrcManager, SrcDataIdx, Dirty))
+					if(FutureState.IsDesynced(SrcManager, SrcDataIdx, *PTParticle, Dirty.ParticleData.GetFlags()))
 					{
 						Info.Desync(CurFrame-1, LatestFrame);
 					}
