Making it easier to marshal data on PT side

diff --git a/Engine/Source/Runtime/Experimental/Chaos/Public/Chaos/ParticleDirtyFlags.h b/Engine/Source/Runtime/Experimental/Chaos/Public/Chaos/ParticleDirtyFlags.h
--- a/Engine/Source/Runtime/Experimental/Chaos/Public/Chaos/ParticleDirtyFlags.h
+++ b/Engine/Source/Runtime/Experimental/Chaos/Public/Chaos/ParticleDirtyFlags.h
@@ -19,18 +19,47 @@ namespace Chaos
 
 struct FParticlePositionRotation
 {
-	FVec3 X;
-	FRotation3 R;
-
 	void Serialize(FChaosArchive& Ar)
 	{
-		Ar << X << R;
+		Ar << MX << MR;
+	}
+
+	template <typename TOther>
+	void CopyFrom(const TOther& Other)
+	{
+		MX = Other.X();
+		MR = Other.R();
+	}
+
+	template <typename TOther>
+	bool IsEqual(const TOther& Other) const
+	{
+		return MX == Other.X() && MR == Other.R();
 	}
 
 	bool operator==(const FParticlePositionRotation& Other) const
 	{
-		return X == Other.X && R == Other.R;
+		return IsEqual(Other);
 	}
+
+	//API for templated code
+	const FVec3& X() const { return MX; }
+	const FRotation3& R() const { return MR; }
+
+	void SetX(const FVec3& InX)
+	{
+		MX = InX;
+	}
+
+	void SetR(const FRotation3& InR)
+	{
+		MR = InR;
+	}
+
+private:
+	FVec3 MX;
+	FRotation3 MR;
+
 };
 
 inline FChaosArchive& operator<<(FChaosArchive& Ar,FParticlePositionRotation& Data)
diff --git a/Engine/Source/Runtime/Experimental/Chaos/Public/Chaos/ParticleHandle.h b/Engine/Source/Runtime/Experimental/Chaos/Public/Chaos/ParticleHandle.h
--- a/Engine/Source/Runtime/Experimental/Chaos/Public/Chaos/ParticleHandle.h
+++ b/Engine/Source/Runtime/Experimental/Chaos/Public/Chaos/ParticleHandle.h
@@ -422,6 +422,13 @@ public:
 	TRotation<T, d>& R() { return GeometryParticles->R(ParticleIdx); }
 	void SetR(const TRotation<T, d>& InR) { GeometryParticles->R(ParticleIdx) = InR; }
 
+	void SetXR(const FParticlePositionRotation& XR)
+	{
+		SetX(XR.X());
+		SetR(XR.R());
+	}
+
+
 	TSerializablePtr<FImplicitObject> Geometry() const { return GeometryParticles->Geometry(ParticleIdx); }
 	void SetGeometry(TSerializablePtr<FImplicitObject> InGeometry) { GeometryParticles->SetGeometry(ParticleIdx, InGeometry); }
 
@@ -563,6 +570,7 @@ public:
 	using TGeometryParticleHandleImp<T, d, bPersistent>::CastToRigidParticle;
 	using TTransientHandle = TTransientKinematicGeometryParticleHandle<T, d>;
 	using TSOAType = TKinematicGeometryParticles<T, d>;
+	using TGeometryParticleHandleImp<T, d, bPersistent>::SetXR;
 	
 protected:
 	friend class TGeometryParticleHandleImp<T, d, bPersistent>;
@@ -656,6 +664,13 @@ public:
 		return Serializable;
 	}
 
+	void SetXR(const FParticlePositionRotation& InXR)
+	{
+		TKinematicGeometryParticleHandleImp<T, d, bPersistent>::SetXR(InXR);
+		SetP(InXR.X());
+		SetQ(InXR.R());
+	}
+
 	operator TPBDRigidParticleHandleImp<T, d, false>& () { return reinterpret_cast<TPBDRigidParticleHandleImp<T, d, false>&>(*this); }
 
 	const TUniquePtr<TBVHParticles<T, d>>& CollisionParticles() const { return PBDRigidParticles->CollisionParticles(ParticleIdx); }
@@ -1399,7 +1414,7 @@ public:
 
 	static TGeometryParticle<T, d>* SerializationFactory(FChaosArchive& Ar, TGeometryParticle<T, d>* Serializable);
 
-	const TVector<T, d>& X() const { return MXR.Read().X; }
+	const TVector<T, d>& X() const { return MXR.Read().X(); }
 	void SetX(const TVector<T, d>& InX, bool bInvalidate = true);
 
 	FUniqueIdx UniqueIdx() const { return MNonFrequentData.Read().UniqueIdx; }
@@ -1408,9 +1423,14 @@ public:
 		MNonFrequentData.Modify(bInvalidate,MDirtyFlags,Proxy,[UniqueIdx](auto& Data){ Data.UniqueIdx = UniqueIdx;});
 	}
 
-	const TRotation<T, d>& R() const { return MXR.Read().R; }
+	const TRotation<T, d>& R() const { return MXR.Read().R(); }
 	void SetR(const TRotation<T, d>& InR, bool bInvalidate = true);
 
+	void SetXR(const FParticlePositionRotation& InXR, bool bInvalidate = true)
+	{
+		MXR.Write(InXR,bInvalidate,MDirtyFlags,Proxy);
+	}
+	
 	//todo: geometry should not be owned by particle
 	void SetGeometry(TUniquePtr<FImplicitObject>&& UniqueGeometry)
 	{
@@ -2329,7 +2349,7 @@ void TGeometryParticle<T, d>::SetX(const TVector<T, d>& InX, bool bInvalidate)
 			Dyn->SetObjectState(EObjectStateType::Dynamic, true);
 		}
 	}
-	MXR.Modify(bInvalidate, MDirtyFlags, Proxy, [&InX](auto& Data) { Data.X = InX; });
+	MXR.Modify(bInvalidate, MDirtyFlags, Proxy, [&InX](auto& Data) { Data.SetX(InX); });
 }
 
 template <typename T, int d>
@@ -2343,7 +2363,7 @@ void TGeometryParticle<T, d>::SetR(const TRotation<T, d>& InR, bool bInvalidate)
 			Dyn->SetObjectState(EObjectStateType::Dynamic, true);
 		}
 	}
-	MXR.Modify(bInvalidate, MDirtyFlags, Proxy, [&InR](auto& Data) { Data.R = InR; });
+	MXR.Modify(bInvalidate, MDirtyFlags, Proxy, [&InR](auto& Data) { Data.SetR(InR); });
 }
 
 template <typename T, int d>
diff --git a/Engine/Source/Runtime/Experimental/ChaosSolvers/Private/PhysicsProxy/SingleParticlePhysicsProxy.cpp b/Engine/Source/Runtime/Experimental/ChaosSolvers/Private/PhysicsProxy/SingleParticlePhysicsProxy.cpp
--- a/Engine/Source/Runtime/Experimental/ChaosSolvers/Private/PhysicsProxy/SingleParticlePhysicsProxy.cpp
+++ b/Engine/Source/Runtime/Experimental/ChaosSolvers/Private/PhysicsProxy/SingleParticlePhysicsProxy.cpp
@@ -36,8 +36,8 @@ FSingleParticlePhysicsProxy<PARTICLE_TYPE>::GetInitialState() const
 	return InitialState; 
 }
 
-template <Chaos::EParticleType ParticleType>
-void PushToPhysicsStateImp(const Chaos::FDirtyPropertiesManager& Manager, Chaos::TGeometryParticleHandle<float,3>* Handle, int32 DataIdx, const Chaos::FDirtyProxy& Dirty, Chaos::FShapeDirtyData* ShapesData, Chaos::FPhysicsSolver* Solver, const bool bInitialized)
+template <Chaos::EParticleType ParticleType, typename TParticleHandle>
+void PushToPhysicsStateImp(const Chaos::FDirtyPropertiesManager& Manager, TParticleHandle* Handle, int32 DataIdx, const Chaos::FDirtyProxy& Dirty, Chaos::FShapeDirtyData* ShapesData, Chaos::FPhysicsSolver* Solver, const bool bInitialized)
 {
 	using namespace Chaos;
 	constexpr bool bHasKinematicData = ParticleType != EParticleType::Static;
@@ -52,13 +52,7 @@ void PushToPhysicsStateImp(const Chaos::FDirtyPropertiesManager& Manager, Chaos:
 
 		if(NewXR)
 		{
-			Handle->SetX(NewXR->X);
-			Handle->SetR(NewXR->R);
-			if(bHasDynamicData)
-			{
-				RigidHandle->SetP(NewXR->X);
-				RigidHandle->SetQ(NewXR->R);
-			}
+			Handle->SetXR(*NewXR);
 		}
 
 		if(NewNonFrequentData)
@@ -240,7 +234,7 @@ void FSingleParticlePhysicsProxy<Chaos::TGeometryParticle<float, 3>>::ClearEvent
 template< >
 void FSingleParticlePhysicsProxy<Chaos::TKinematicGeometryParticle<float, 3>>::PushToPhysicsState(const Chaos::FDirtyPropertiesManager& Manager, int32 DataIdx, const Chaos::FDirtyProxy& Dirty, Chaos::FShapeDirtyData* ShapesData)
 {
-	PushToPhysicsStateImp<Chaos::EParticleType::Kinematic>(Manager,Handle,DataIdx,Dirty,ShapesData,GetSolver(),bInitialized);
+	PushToPhysicsStateImp<Chaos::EParticleType::Kinematic>(Manager,Handle->CastToKinematicParticle(),DataIdx,Dirty,ShapesData,GetSolver(),bInitialized);
 #if 0
 	// move the copied game thread data into the handle
 	if(auto* RigidHandle = static_cast<Chaos::TKinematicGeometryParticleHandle<float,3>*>(Handle))
@@ -412,7 +406,7 @@ void FSingleParticlePhysicsProxy<Chaos::TKinematicGeometryParticle<float, 3>>::C
 template< >
 void FSingleParticlePhysicsProxy<Chaos::TPBDRigidParticle<float, 3>>::PushToPhysicsState(const Chaos::FDirtyPropertiesManager& Manager, int32 DataIdx, const Chaos::FDirtyProxy& Dirty, Chaos::FShapeDirtyData* ShapesData)
 {
-	PushToPhysicsStateImp<Chaos::EParticleType::Rigid>(Manager,Handle,DataIdx,Dirty,ShapesData,GetSolver(),bInitialized);
+	PushToPhysicsStateImp<Chaos::EParticleType::Rigid>(Manager,Handle->CastToRigidParticle(),DataIdx,Dirty,ShapesData,GetSolver(),bInitialized);
 #if 0
 	// move the copied game thread data into the handle
 	if (auto* RigidHandle = static_cast<Chaos::TPBDRigidParticleHandle<float, 3>*>(Handle))
diff --git a/Engine/Source/Runtime/Experimental/ChaosSolvers/Public/RewindData.h b/Engine/Source/Runtime/Experimental/ChaosSolvers/Public/RewindData.h
--- a/Engine/Source/Runtime/Experimental/ChaosSolvers/Public/RewindData.h
+++ b/Engine/Source/Runtime/Experimental/ChaosSolvers/Public/RewindData.h
@@ -73,27 +73,48 @@ public:
 		return Manager != nullptr;
 	}
 
-	EDesyncResult DesyncQuery(const FDirtyPropertiesManager& SrcManager, const int32 DataIdxIn, const FParticleDirtyFlags Flags) const
+	template <typename TParticleHandle>
+	bool IsInSync(const FDirtyPropertiesManager& SrcManager, const int32 DataIdxIn, const FParticleDirtyFlags Flags, const TParticleHandle& Handle) const
 	{
-		if(Flags.IsDirty(ParticlePropToFlag(PropName)))
+		const T* RecordedEntry = Manager ? &GetValue(*Manager,Idx) : nullptr;
+		const T* NewEntry = Flags.IsDirty(ParticlePropToFlag(PropName)) ? &GetValue(SrcManager,DataIdxIn) : nullptr;
+
+		if(NewEntry)
 		{
-			if(Manager)
+			if(RecordedEntry)
+			{
+				//We have an entry from current run and previous run, so check that they are equal
+				return NewEntry->IsEqual(*RecordedEntry);
+			}
+			else
 			{
-				//both dirty, if value differs we're desynced
-				return (SrcManager.GetParticlePool<T,PropName>().GetElement(DataIdxIn) == Manager->GetParticlePool<T,PropName>().GetElement(Idx)) ? EDesyncResult::InSync : EDesyncResult::Desync;
+				//Previous run had no entry. If the current PT data matches the new data, then this is a harmless idnetical write and we are still in sync
+				return NewEntry->IsEqual(Handle);
 			}
 		}
-		else if(Manager == nullptr)
+		else
 		{
-			return EDesyncResult::InSync;	//Both have no entry so in sync
+			if(RecordedEntry)
+			{
+				//We have an entry from previous run, but not anymore. It's possible this will get written out by PT and hasn't yet, so check if the values are the same
+				return RecordedEntry->IsEqual(Handle);
+			}
+			else
+			{
+				//Both current run and recorded run have no entry, so both pointed at head and saw no change
+				return true;
+			}
 		}
-
-		return EDesyncResult::NeedInfo;	//entry is missing in one or the other but not both, so we don't know if it's desynced
 	}
 
 private:
 	FDirtyPropertiesManager* Manager;
 	int32 Idx;
+
+	static const T& GetValue(const FDirtyPropertiesManager& InManager, const int32 InIdx)
+	{
+		return InManager.GetParticlePool<T,PropName>().GetElement(InIdx);
+	}
 };
 
 inline bool SimWritablePropsMayChange(const TGeometryParticleHandle<FReal,3>& Handle)
@@ -112,15 +133,15 @@ public:
 	template <typename TParticle>
 	const FVec3& X(const TParticle& Particle) const
 	{
-		return ParticlePositionRotation.IsSet() ? ParticlePositionRotation.Read().X : Particle.X();
+		return ParticlePositionRotation.IsSet() ? ParticlePositionRotation.Read().X() : Particle.X();
 	}
 
 	template <typename TParticle>
 	const FRotation3& R(const TParticle& Particle) const
 	{
-		return ParticlePositionRotation.IsSet() ? ParticlePositionRotation.Read().R : Particle.R();
+		return ParticlePositionRotation.IsSet() ? ParticlePositionRotation.Read().R() : Particle.R();
 	}
-
+	
 	template <typename TParticle>
 	const FVec3& V(const TParticle& Particle) const
 	{
@@ -155,8 +176,7 @@ public:
 
 		ParticlePositionRotation.SyncRemoteData(Manager,Idx,Dirty,[&Rigid](auto& Data)
 		{
-			Data.X = Rigid.X();
-			Data.R = Rigid.R();
+			Data.CopyFrom(Rigid);
 		});
 
 		Velocities.SyncRemoteData(Manager,Idx,Dirty,[&Rigid](auto& Data)
@@ -183,8 +203,7 @@ public:
 
 		ParticlePositionRotation.SyncToParticle([&Particle](const auto& Data)
 		{
-			Particle.SetX(Data.X);
-			Particle.SetR(Data.R);
+			Particle.SetXR(Data);
 		});
 
 		if(auto Kinematic = Particle.CastToKinematicParticle())
@@ -241,8 +260,7 @@ public:
 		{
 			ParticlePositionRotation.SyncRemoteData(Manager,Idx,Dirty.ParticleData,[Handle](FParticlePositionRotation& Data)
 			{
-				Data.X = Handle->X();
-				Data.R = Handle->R();
+				Data.CopyFrom(*Handle);
 			});
 
 			if(auto Kinematic = Handle->CastToKinematicParticle())
@@ -284,8 +302,7 @@ public:
 		{
 			ParticlePositionRotation.SyncRemoteDataForced(Manager,Idx,[Particle](FParticlePositionRotation& Data)
 			{
-				Data.X = Particle->X();
-				Data.R = Particle->R();
+				Data.CopyFrom(*Particle);
 			});
 		}
 		
@@ -354,23 +371,10 @@ public:
 	{
 		bool Desynced = false;
 		{
-			EDesyncResult Status = ParticlePositionRotation.DesyncQuery(SrcManager, DataIdxIn, Flags);
-			if(Status == EDesyncResult::Desync)
+			if(!ParticlePositionRotation.IsInSync(SrcManager,DataIdxIn,Flags,Handle))
 			{
 				return true;
 			}
-			else if(Status == EDesyncResult::NeedInfo)
-			{
-				if(Handle.X() != X(Handle))
-				{
-					return true;
-				}
-
-				if(Handle.R() != R(Handle))
-				{
-					return true;
-				}
-			}
 		}
 
 		//TODO: test other properties, should probably find a better way to do this to avoid getting into the individual variables
