More structs changed to make PT data marshal easier

diff --git a/Engine/Source/Runtime/Experimental/Chaos/Public/Chaos/ParticleDirtyFlags.h b/Engine/Source/Runtime/Experimental/Chaos/Public/Chaos/ParticleDirtyFlags.h
--- a/Engine/Source/Runtime/Experimental/Chaos/Public/Chaos/ParticleDirtyFlags.h
+++ b/Engine/Source/Runtime/Experimental/Chaos/Public/Chaos/ParticleDirtyFlags.h
@@ -17,8 +17,9 @@ class FName;
 namespace Chaos
 {
 
-struct FParticlePositionRotation
+class FParticlePositionRotation
 {
+public:
 	void Serialize(FChaosArchive& Ar)
 	{
 		Ar << MX << MR;
@@ -42,20 +43,12 @@ struct FParticlePositionRotation
 		return IsEqual(Other);
 	}
 
-	//API for templated code
 	const FVec3& X() const { return MX; }
-	const FRotation3& R() const { return MR; }
-
-	void SetX(const FVec3& InX)
-	{
-		MX = InX;
-	}
-
-	void SetR(const FRotation3& InR)
-	{
-		MR = InR;
-	}
+	void SetX(const FVec3& InX){ MX = InX; }
 
+	const FRotation3& R() const { return MR; }
+	void SetR(const FRotation3& InR){ MR = InR; }
+	
 private:
 	FVec3 MX;
 	FRotation3 MR;
@@ -68,20 +61,41 @@ inline FChaosArchive& operator<<(FChaosArchive& Ar,FParticlePositionRotation& Da
 	return Ar;
 }
 
-struct FParticleVelocities
+class FParticleVelocities
 {
+public:
 	void Serialize(FChaosArchive& Ar)
 	{
-		Ar << V << W;
+		Ar << MV << MW;
+	}
+
+	template <typename TOther>
+	void CopyFrom(const TOther& Other)
+	{
+		MV = Other.V();
+		MW = Other.W();
+	}
+
+	template <typename TOther>
+	bool IsEqual(const TOther& Other) const
+	{
+		return MV == Other.V() && MW == Other.W();
 	}
 
 	bool operator==(const FParticleVelocities& Other) const
 	{
-		return V == Other.V && W == Other.W;
+		return IsEqual(Other);
 	}
 
-	FVec3 V;
-	FVec3 W;
+	const FVec3& V() const { return MV; }
+	void SetV(const FVec3& V) { MV = V; }
+
+	const FVec3& W() const { return MW; }
+	void SetW(const FVec3& W){ MW = W; }
+
+private:
+	FVec3 MV;
+	FVec3 MW;
 };
 
 inline FChaosArchive& operator<<(FChaosArchive& Ar,FParticleVelocities& Data)
@@ -90,28 +104,58 @@ inline FChaosArchive& operator<<(FChaosArchive& Ar,FParticleVelocities& Data)
 	return Ar;
 }
 
-struct FParticleDynamics
+class FParticleDynamics
 {
-	FVec3 F;
-	FVec3 Torque;
-	FVec3 LinearImpulse;
-	FVec3 AngularImpulse;
-
+public:
 	void Serialize(FChaosArchive& Ar)
 	{
-		Ar << F;
-		Ar << Torque;
-		Ar << LinearImpulse;
-		Ar << AngularImpulse;	
+		Ar << MF;
+		Ar << MTorque;
+		Ar << MLinearImpulse;
+		Ar << MAngularImpulse;	
+	}
+
+	template <typename TOther>
+	void CopyFrom(const TOther& Other)
+	{
+		MF = Other.F();
+		MTorque = Other.Torque();
+		MLinearImpulse = Other.LinearImpulse();
+		MAngularImpulse = Other.AngularImpulse();
+	}
+
+	template <typename TOther>
+	bool IsEqual(const TOther& Other) const
+	{
+		return F() == Other.F()
+			&& Torque() == Other.Torque()
+			&& LinearImpulse() == Other.LinearImpulse()
+			&& AngularImpulse() == Other.AngularImpulse();
 	}
 
 	bool operator==(const FParticleDynamics& Other) const
 	{
-		return F == Other.F
-			&& Torque == Other.Torque
-			&& LinearImpulse == Other.LinearImpulse
-			&& AngularImpulse == Other.AngularImpulse;
+		return IsEqual(Other);
 	}
+
+	const FVec3& F() const { return MF; }
+	void SetF(const FVec3& F){ MF = F; }
+
+	const FVec3& Torque() const { return MTorque; }
+	void SetTorque(const FVec3& Torque){ MTorque = Torque; }
+
+	const FVec3& LinearImpulse() const { return MF; }
+	void SetLinearImpulse(const FVec3& LinearImpulse){ MF = LinearImpulse; }
+
+	const FVec3& AngularImpulse() const { return MF; }
+	void SetAngularImpulse(const FVec3& AngularImpulse){ MF = AngularImpulse; }
+
+private:
+	FVec3 MF;
+	FVec3 MTorque;
+	FVec3 MLinearImpulse;
+	FVec3 MAngularImpulse;
+
 };
 
 inline FChaosArchive& operator<<(FChaosArchive& Ar, FParticleDynamics& Data)
diff --git a/Engine/Source/Runtime/Experimental/Chaos/Public/Chaos/ParticleHandle.h b/Engine/Source/Runtime/Experimental/Chaos/Public/Chaos/ParticleHandle.h
--- a/Engine/Source/Runtime/Experimental/Chaos/Public/Chaos/ParticleHandle.h
+++ b/Engine/Source/Runtime/Experimental/Chaos/Public/Chaos/ParticleHandle.h
@@ -607,6 +607,12 @@ public:
 	TVector<T, d>& W() { return KinematicGeometryParticles->W(ParticleIdx); }
 	void SetW(const TVector<T, d>& InW) { KinematicGeometryParticles->W(ParticleIdx) = InW; }
 
+	void SetVelocities(const FParticleVelocities& Velocities)
+	{
+		SetV(Velocities.V());
+		SetW(Velocities.W());
+	}
+
 	const TKinematicTarget<T, d>& KinematicTarget() const { return KinematicGeometryParticles->KinematicTarget(ParticleIdx); }
 	TKinematicTarget<T, d>& KinematicTarget() { return KinematicGeometryParticles->KinematicTarget(ParticleIdx); }
 
@@ -723,6 +729,14 @@ public:
 	TVector<T, d>& AngularImpulse() { return PBDRigidParticles->AngularImpulse(ParticleIdx); }
 	void SetAngularImpulse(const TVector<T, d>& InAngularImpulse) { PBDRigidParticles->AngularImpulse(ParticleIdx) = InAngularImpulse; }
 
+	void SetDynamics(const FParticleDynamics& Dynamics)
+	{
+		SetF(Dynamics.F());
+		SetTorque(Dynamics.Torque());
+		SetLinearImpulse(Dynamics.LinearImpulse());
+		SetAngularImpulse(Dynamics.AngularImpulse());
+	}
+
 	const PMatrix<T, d, d>& I() const { return PBDRigidParticles->I(ParticleIdx); }
 	PMatrix<T, d, d>& I() { return PBDRigidParticles->I(ParticleIdx); }
 	void SetI(const PMatrix<T, d, d>& InI) { PBDRigidParticles->I(ParticleIdx) = InI; }
@@ -1851,12 +1865,17 @@ public:
 		Ar << MVelocities;
 	}
 
-	const TVector<T, d>& V() const { return MVelocities.Read().V; }
+	const TVector<T, d>& V() const { return MVelocities.Read().V(); }
 	void SetV(const TVector<T, d>& InV, bool bInvalidate = true);
 
-	const TVector<T, d>& W() const { return MVelocities.Read().W; }
+	const TVector<T, d>& W() const { return MVelocities.Read().W(); }
 	void SetW(const TVector<T, d>& InW, bool bInvalidate = true);
 
+	void SetVelocities(const FParticleVelocities& InVelocities,bool bInvalidate = true)
+	{
+		MVelocities.Write(InVelocities,bInvalidate,MDirtyFlags,Proxy);
+	}
+
 	const TVector<T, d>& CenterOfMass() const { return MMassProps.Read().CenterOfMass; }
 	void SetCenterOfMass(const TVector<T, d>& InCenterOfMass, bool bInvalidate = true)
 	{
@@ -2006,44 +2025,49 @@ public:
 		this->MInitialized = InInitialized;
 	}
 
-	const TVector<T, d>& F() const { return MDynamics.Read().F; }
+	const TVector<T, d>& F() const { return MDynamics.Read().F(); }
 	void SetF(const TVector<T, d>& InF, bool bInvalidate = true)
 	{
 		if (bInvalidate)
 		{
 			SetObjectState(EObjectStateType::Dynamic, true);
 		}
-		MDynamics.Modify(bInvalidate,MDirtyFlags,Proxy,[&InF](auto& Data){ Data.F = InF;});
+		MDynamics.Modify(bInvalidate,MDirtyFlags,Proxy,[&InF](auto& Data){ Data.SetF(InF);});
 	}
 
-	const TVector<T, d>& Torque() const { return MDynamics.Read().Torque; }
+	const TVector<T, d>& Torque() const { return MDynamics.Read().Torque(); }
 	void SetTorque(const TVector<T, d>& InTorque, bool bInvalidate=true)
 	{
 		if (bInvalidate)
 		{
 			SetObjectState(EObjectStateType::Dynamic, true);
 		}
-		MDynamics.Modify(bInvalidate,MDirtyFlags,Proxy,[&InTorque](auto& Data){ Data.Torque = InTorque;});
+		MDynamics.Modify(bInvalidate,MDirtyFlags,Proxy,[&InTorque](auto& Data){ Data.SetTorque(InTorque);});
 	}
 
-	const TVector<T, d>& LinearImpulse() const { return MDynamics.Read().LinearImpulse; }
+	const TVector<T, d>& LinearImpulse() const { return MDynamics.Read().LinearImpulse(); }
 	void SetLinearImpulse(const TVector<T, d>& InLinearImpulse, bool bInvalidate = true)
 	{
 		if (bInvalidate)
 		{
 			SetObjectState(EObjectStateType::Dynamic, true);
 		}
-		MDynamics.Modify(bInvalidate,MDirtyFlags,Proxy,[&InLinearImpulse](auto& Data){ Data.LinearImpulse = InLinearImpulse;});
+		MDynamics.Modify(bInvalidate,MDirtyFlags,Proxy,[&InLinearImpulse](auto& Data){ Data.SetLinearImpulse(InLinearImpulse);});
 	}
 
-	const TVector<T, d>& AngularImpulse() const { return MDynamics.Read().AngularImpulse; }
+	const TVector<T, d>& AngularImpulse() const { return MDynamics.Read().AngularImpulse(); }
 	void SetAngularImpulse(const TVector<T, d>& InAngularImpulse, bool bInvalidate = true)
 	{
 		if (bInvalidate)
 		{
 			SetObjectState(EObjectStateType::Dynamic, true);
 		}
-		MDynamics.Modify(bInvalidate,MDirtyFlags,Proxy,[&InAngularImpulse](auto& Data){ Data.AngularImpulse = InAngularImpulse;});
+		MDynamics.Modify(bInvalidate,MDirtyFlags,Proxy,[&InAngularImpulse](auto& Data){ Data.SetAngularImpulse(InAngularImpulse);});
+	}
+
+	void SetDynamics(const FParticleDynamics& InDynamics,bool bInvalidate = true)
+	{
+		MDynamics.Write(InDynamics,bInvalidate,MDirtyFlags,Proxy);
 	}
 
 	const PMatrix<T, d, d>& I() const { return MMassProps.Read().I; }
@@ -2391,7 +2415,7 @@ void TKinematicGeometryParticle<T, d>::SetV(const TVector<T, d>& InV, bool bInva
 			Dyn->SetObjectState(EObjectStateType::Dynamic, true);
 		}
 	}
-	MVelocities.Modify(bInvalidate, MDirtyFlags, Proxy, [&InV](auto& Data) { Data.V = InV; });
+	MVelocities.Modify(bInvalidate, MDirtyFlags, Proxy, [&InV](auto& Data) { Data.SetV(InV); });
 }
 
 template <typename T, int d>
@@ -2405,7 +2429,7 @@ void TKinematicGeometryParticle<T, d>::SetW(const TVector<T, d>& InW, bool bInva
 			Dyn->SetObjectState(EObjectStateType::Dynamic, true);
 		}
 	}
-	MVelocities.Modify(bInvalidate, MDirtyFlags, Proxy, [&InW](auto& Data) { Data.W = InW; });
+	MVelocities.Modify(bInvalidate, MDirtyFlags, Proxy, [&InW](auto& Data) { Data.SetW(InW); });
 }
 
 template <typename T, int d>
diff --git a/Engine/Source/Runtime/Experimental/ChaosSolvers/Private/PhysicsProxy/SingleParticlePhysicsProxy.cpp b/Engine/Source/Runtime/Experimental/ChaosSolvers/Private/PhysicsProxy/SingleParticlePhysicsProxy.cpp
--- a/Engine/Source/Runtime/Experimental/ChaosSolvers/Private/PhysicsProxy/SingleParticlePhysicsProxy.cpp
+++ b/Engine/Source/Runtime/Experimental/ChaosSolvers/Private/PhysicsProxy/SingleParticlePhysicsProxy.cpp
@@ -75,8 +75,7 @@ void PushToPhysicsStateImp(const Chaos::FDirtyPropertiesManager& Manager, TParti
 		auto NewVelocities = bHasKinematicData ? ParticleData.FindVelocities(Manager, DataIdx) : nullptr;
 		if(NewVelocities)
 		{
-			KinematicHandle->SetV(NewVelocities->V);
-			KinematicHandle->SetW(NewVelocities->W);
+			KinematicHandle->SetVelocities(*NewVelocities);
 		}
 
 		if(NewXR || NewNonFrequentData || NewVelocities)
@@ -129,10 +128,7 @@ void PushToPhysicsStateImp(const Chaos::FDirtyPropertiesManager& Manager, TParti
 		{
 			if(auto NewData = ParticleData.FindDynamics(Manager, DataIdx))
 			{
-				RigidHandle->SetF(NewData->F);
-				RigidHandle->SetTorque(NewData->Torque);
-				RigidHandle->SetLinearImpulse(NewData->LinearImpulse);
-				RigidHandle->SetAngularImpulse(NewData->AngularImpulse);
+				RigidHandle->SetDynamics(*NewData);
 			}
 		}
 
diff --git a/Engine/Source/Runtime/Experimental/ChaosSolvers/Public/RewindData.h b/Engine/Source/Runtime/Experimental/ChaosSolvers/Public/RewindData.h
--- a/Engine/Source/Runtime/Experimental/ChaosSolvers/Public/RewindData.h
+++ b/Engine/Source/Runtime/Experimental/ChaosSolvers/Public/RewindData.h
@@ -145,13 +145,13 @@ public:
 	template <typename TParticle>
 	const FVec3& V(const TParticle& Particle) const
 	{
-		return Velocities.IsSet() ? Velocities.Read().V : Particle.CastToKinematicParticle()->V();
+		return Velocities.IsSet() ? Velocities.Read().V() : Particle.CastToKinematicParticle()->V();
 	}
 
 	template <typename TParticle>
 	const FVec3& W(const TParticle& Particle) const
 	{
-		return Velocities.IsSet() ? Velocities.Read().W : Particle.CastToKinematicParticle()->W();
+		return Velocities.IsSet() ? Velocities.Read().W() : Particle.CastToKinematicParticle()->W();
 	}
 
 	template <typename TParticle>
@@ -163,7 +163,7 @@ public:
 	template <typename TParticle>
 	const FVec3& F(const TParticle& Particle) const
 	{
-		return Dynamics.IsSet() ? Dynamics.Read().F : Particle.CastToRigidParticle()->F();
+		return Dynamics.IsSet() ? Dynamics.Read().F() : Particle.CastToRigidParticle()->F();
 	}
 
 	void SyncSimWritablePropsFromSim(FDirtyPropertiesManager& Manager,int32 Idx,const TPBDRigidParticleHandle<FReal,3>& Rigid)
@@ -181,8 +181,8 @@ public:
 
 		Velocities.SyncRemoteData(Manager,Idx,Dirty,[&Rigid](auto& Data)
 		{
-			Data.V = Rigid.PreV();
-			Data.W = Rigid.PreW();
+			Data.SetV(Rigid.PreV());
+			Data.SetW(Rigid.PreW());
 		});
 	}
 
@@ -210,8 +210,7 @@ public:
 		{
 			Velocities.SyncToParticle([Kinematic](const auto& Data)
 			{
-				Kinematic->SetV(Data.V);
-				Kinematic->SetW(Data.W);
+				Kinematic->SetVelocities(Data);
 			});
 		}
 
@@ -236,10 +235,7 @@ public:
 		{
 			Dynamics.SyncToParticle([Rigid](const auto& Data)
 			{
-				Rigid->SetF(Data.F);
-				Rigid->SetTorque(Data.Torque);
-				Rigid->SetLinearImpulse(Data.LinearImpulse);
-				Rigid->SetAngularImpulse(Data.AngularImpulse);
+				Rigid->SetDynamics(Data);
 			});
 		}
 	}
@@ -267,8 +263,7 @@ public:
 			{
 				Velocities.SyncRemoteData(Manager,Idx,Dirty.ParticleData,[Kinematic](auto& Data)
 				{
-					Data.V = Kinematic->V();
-					Data.W = Kinematic->W();
+					Data.CopyFrom(*Kinematic);
 				});
 			}
 		}
@@ -312,8 +307,7 @@ public:
 			{
 				Velocities.SyncRemoteDataForced(Manager,Idx,[Kinematic](auto& Data)
 				{
-					Data.V = Kinematic->V();
-					Data.W = Kinematic->W();
+					Data.CopyFrom(*Kinematic);
 				});
 			}
 		}
