Fix Clang wanring (HeightField.h)

diff --git a/Engine/Source/Runtime/Experimental/Chaos/Public/Chaos/HeightField.h b/Engine/Source/Runtime/Experimental/Chaos/Public/Chaos/HeightField.h
--- a/Engine/Source/Runtime/Experimental/Chaos/Public/Chaos/HeightField.h
+++ b/Engine/Source/Runtime/Experimental/Chaos/Public/Chaos/HeightField.h
@@ -31,8 +31,8 @@ namespace Chaos
 
 		CHAOS_API THeightField(TArray<T>&& Height, TArray<uint8>&& InMaterialIndices, int32 InNumRows, int32 InNumCols, const TVector<T,3>& InScale);
 		CHAOS_API THeightField(TArrayView<const uint16> InHeights, TArrayView<uint8> InMaterialIndices, int32 InNumRows, int32 InNumCols, const TVector<T, 3>& InScale);
-		CHAOS_API THeightField(const THeightField& Other) = delete;
-		CHAOS_API THeightField(THeightField&& Other) = default;
+		THeightField(const THeightField& Other) = delete;
+		THeightField(THeightField&& Other) = default;
 
 		virtual ~THeightField() {}
 
