Chaos solver module will be initialized to multithread with the TaskGraph when -PostForkThreading is set on the master server.

diff --git a/Engine/Source/Runtime/Experimental/ChaosSolvers/Private/ChaosSolversModule.cpp b/Engine/Source/Runtime/Experimental/ChaosSolvers/Private/ChaosSolversModule.cpp
--- a/Engine/Source/Runtime/Experimental/ChaosSolvers/Private/ChaosSolversModule.cpp
+++ b/Engine/Source/Runtime/Experimental/ChaosSolvers/Private/ChaosSolversModule.cpp
@@ -969,7 +969,7 @@ void FChaosSolversModule::ChangeBufferMode(Chaos::EMultiBufferMode BufferMode)
 
 Chaos::EThreadingMode FChaosSolversModule::GetDesiredThreadingMode() const
 {
-	const bool bForceSingleThread = !FApp::ShouldUseThreadingForPerformance();
+	const bool bForceSingleThread = !(FApp::ShouldUseThreadingForPerformance() || FForkProcessHelper::SupportsMultithreadingPostFork());
 
 	// If the platform isn't using threads for perf - force Chaos to
 	// run single threaded no matter the selected mode.
