Implement friction/restitution combine mode

diff --git a/Engine/Source/Runtime/Engine/Private/PhysicsEngine/Experimental/PhysInterface_Chaos.cpp b/Engine/Source/Runtime/Engine/Private/PhysicsEngine/Experimental/PhysInterface_Chaos.cpp
--- a/Engine/Source/Runtime/Engine/Private/PhysicsEngine/Experimental/PhysInterface_Chaos.cpp
+++ b/Engine/Source/Runtime/Engine/Private/PhysicsEngine/Experimental/PhysInterface_Chaos.cpp
@@ -311,12 +311,30 @@ void FPhysInterface_Chaos::ReleaseMaterial(FPhysicsMaterialHandle& InHandle)
 	Chaos::FPhysicalMaterialManager::Get().Destroy(InHandle);
 }
 
+Chaos::FChaosPhysicsMaterial::ECombineMode UToCCombineMode(EFrictionCombineMode::Type Mode)
+{
+	using namespace Chaos;
+	switch(Mode)
+	{
+	case EFrictionCombineMode::Average: return FChaosPhysicsMaterial::ECombineMode::Avg;
+	case EFrictionCombineMode::Min: return FChaosPhysicsMaterial::ECombineMode::Min;
+	case EFrictionCombineMode::Multiply: return FChaosPhysicsMaterial::ECombineMode::Multiply;
+	case EFrictionCombineMode::Max: return FChaosPhysicsMaterial::ECombineMode::Max;
+	default: ensure(false);
+	}
+
+	return FChaosPhysicsMaterial::ECombineMode::Avg;
+}
+
+
 void FPhysInterface_Chaos::UpdateMaterial(FPhysicsMaterialHandle& InHandle, UPhysicalMaterial* InMaterial)
 {
 	if(Chaos::FChaosPhysicsMaterial* Material = InHandle.Get())
 	{
 		Material->Friction = InMaterial->Friction;
+		Material->FrictionCombineMode = UToCCombineMode(InMaterial->FrictionCombineMode);
 		Material->Restitution = InMaterial->Restitution;
+		Material->RestitutionCombineMode = UToCCombineMode(InMaterial->RestitutionCombineMode);
 		Material->SleepingLinearThreshold = InMaterial->SleepLinearVelocityThreshold;
 		Material->SleepingAngularThreshold = InMaterial->SleepAngularVelocityThreshold;
 		Material->SleepCounterThreshold = InMaterial->SleepCounterThreshold;
diff --git a/Engine/Source/Runtime/Experimental/Chaos/Private/Chaos/PBDCollisionConstraints.cpp b/Engine/Source/Runtime/Experimental/Chaos/Private/Chaos/PBDCollisionConstraints.cpp
--- a/Engine/Source/Runtime/Experimental/Chaos/Private/Chaos/PBDCollisionConstraints.cpp
+++ b/Engine/Source/Runtime/Experimental/Chaos/Private/Chaos/PBDCollisionConstraints.cpp
@@ -171,9 +171,11 @@ namespace Chaos
 		FCollisionContact& Contact = Constraint.Manifold;
 		if (PhysicsMaterial0 && PhysicsMaterial1)
 		{
-			// @todo(ccaulfield): support different friction/restitution combining algorithms
-			Contact.Restitution = FMath::Min(PhysicsMaterial0->Restitution, PhysicsMaterial1->Restitution);
-			Contact.Friction = FMath::Max(PhysicsMaterial0->Friction, PhysicsMaterial1->Friction);
+			const FChaosPhysicsMaterial::ECombineMode RestitutionCombineMode = FChaosPhysicsMaterial::ChooseCombineMode(PhysicsMaterial0->RestitutionCombineMode,PhysicsMaterial1->RestitutionCombineMode);
+			Contact.Restitution = FChaosPhysicsMaterial::CombineHelper(PhysicsMaterial0->Restitution, PhysicsMaterial1->Restitution, RestitutionCombineMode);
+
+			const FChaosPhysicsMaterial::ECombineMode FrictionCombineMode = FChaosPhysicsMaterial::ChooseCombineMode(PhysicsMaterial0->FrictionCombineMode,PhysicsMaterial1->FrictionCombineMode);
+			Contact.Friction = FChaosPhysicsMaterial::CombineHelper(PhysicsMaterial0->Friction,PhysicsMaterial1->Friction, FrictionCombineMode);
 		}
 		else if (PhysicsMaterial0)
 		{
diff --git a/Engine/Source/Runtime/Experimental/Chaos/Public/Chaos/Defines.h b/Engine/Source/Runtime/Experimental/Chaos/Public/Chaos/Defines.h
--- a/Engine/Source/Runtime/Experimental/Chaos/Public/Chaos/Defines.h
+++ b/Engine/Source/Runtime/Experimental/Chaos/Public/Chaos/Defines.h
@@ -25,15 +25,30 @@ namespace Chaos
 	{
 	public:
 
-		Chaos::FReal Friction;
-		Chaos::FReal Restitution;
-		Chaos::FReal SleepingLinearThreshold;
-		Chaos::FReal SleepingAngularThreshold;
-		Chaos::FReal DisabledLinearThreshold;
-		Chaos::FReal DisabledAngularThreshold;
+		/** The different combine modes to determine the friction / restitution of two actors.
+		    If the two materials have two different combine modes, we use the largest one (Min beats Avg, Max beats Multiply, etc...)
+			*/
+		enum class ECombineMode : uint8
+		{
+			Avg,
+			Min,
+			Multiply,
+			Max
+		};
+
+		FReal Friction;
+		FReal Restitution;
+		FReal SleepingLinearThreshold;
+		FReal SleepingAngularThreshold;
+		FReal DisabledLinearThreshold;
+		FReal DisabledAngularThreshold;
 		int32 SleepCounterThreshold;
 		void* UserData;
 
+		ECombineMode FrictionCombineMode;
+		ECombineMode RestitutionCombineMode;
+
+
 		FChaosPhysicsMaterial()
 			: Friction(0.5)
 			, Restitution(0.1)
@@ -43,7 +58,29 @@ namespace Chaos
 			, DisabledAngularThreshold(0)
 			, SleepCounterThreshold(0)
 			, UserData(nullptr)
+			, FrictionCombineMode(ECombineMode::Avg)
+			, RestitutionCombineMode(ECombineMode::Avg)
+		{
+		}
+
+		static FReal CombineHelper(FReal A, FReal B, ECombineMode Mode)
+		{
+			switch(Mode)
+			{
+			case ECombineMode::Avg: return (A+B)*(FReal)0.5;
+			case ECombineMode::Min: return FMath::Min(A,B);
+			case ECombineMode::Max: return FMath::Max(A,B);
+			case ECombineMode::Multiply: return A*B;
+			default: ensure(false);
+			}
+
+			return 0;
+		}
+
+		static ECombineMode ChooseCombineMode(ECombineMode A, ECombineMode B)
 		{
+			const uint8 MaxVal = FMath::Max(static_cast<uint8>(A),static_cast<uint8>(B));
+			return static_cast<ECombineMode>(MaxVal);
 		}
 
 		static constexpr bool IsSerializablePtr = true;
