Move rotation of mass and center of mass from Kinematic to Rigid. Makes it easier to marshal data

diff --git a/Engine/Source/Runtime/Engine/Private/PhysicsEngine/Experimental/PhysInterface_Chaos.cpp b/Engine/Source/Runtime/Engine/Private/PhysicsEngine/Experimental/PhysInterface_Chaos.cpp
--- a/Engine/Source/Runtime/Engine/Private/PhysicsEngine/Experimental/PhysInterface_Chaos.cpp
+++ b/Engine/Source/Runtime/Engine/Private/PhysicsEngine/Experimental/PhysInterface_Chaos.cpp
@@ -724,7 +724,8 @@ FVector FPhysInterface_Chaos::GetWorldVelocityAtPoint_AssumesLocked(const FPhysi
 		Chaos::TKinematicGeometryParticle<float, 3>* Kinematic = InActorReference->CastToKinematicParticle();
 		if (ensure(Kinematic))
 		{
-			const Chaos::FVec3 COM = Chaos::FParticleUtilitiesGT::GetCoMWorldPosition(Kinematic);
+			const Chaos::TPBDRigidParticle<float,3>* Rigid = Kinematic->CastToRigidParticle();
+			const Chaos::FVec3 COM = Rigid ? Chaos::FParticleUtilitiesGT::GetCoMWorldPosition(Rigid) : Chaos::FParticleUtilitiesGT::GetActorWorldTransform(Rigid).GetTranslation();
 			const Chaos::FVec3 Diff = InPoint - COM;
 			return Kinematic->V() - Chaos::FVec3::CrossProduct(Diff, Kinematic->W());
 		}
@@ -736,9 +737,9 @@ FTransform FPhysInterface_Chaos::GetComTransform_AssumesLocked(const FPhysicsAct
 {
 	if (ensure(FPhysicsInterface::IsValid(InActorReference)))
 	{
-		if (const Chaos::TKinematicGeometryParticle<float, 3>* Kinematic = InActorReference->CastToKinematicParticle())
+		if (const auto* Rigid = InActorReference->CastToRigidParticle())
 		{
-			return Chaos::FParticleUtilitiesGT::GetCoMWorldTransform(Kinematic);
+			return Chaos::FParticleUtilitiesGT::GetCoMWorldTransform(Rigid);
 		}
 	}
 	return FTransform();
@@ -748,9 +749,9 @@ FTransform FPhysInterface_Chaos::GetComTransformLocal_AssumesLocked(const FPhysi
 {
 	if (ensure(FPhysicsInterface::IsValid(InActorReference)))
 	{
-		if (Chaos::TKinematicGeometryParticle<float, 3>* Kinematic = InActorReference->CastToKinematicParticle())
+		if (auto* Rigid = InActorReference->CastToRigidParticle())
 		{
-			return FTransform(Kinematic->RotationOfMass(), Kinematic->CenterOfMass());
+			return FTransform(Rigid->RotationOfMass(), Rigid->CenterOfMass());
 		}
 	}
 	return FTransform();
@@ -930,10 +931,10 @@ void FPhysInterface_Chaos::SetComLocalPose_AssumesLocked(const FPhysicsActorHand
 {
     //@todo(mlentine): What is InComLocalPose? If the center of an object is not the local pose then many things break including the three vector represtnation of inertia.
 
-	if (Chaos::TKinematicGeometryParticle<float, 3>* KinematicParticle = InHandle->CastToKinematicParticle())
+	if (auto Rigid = InHandle->CastToRigidParticle())
 	{
-		KinematicParticle->SetCenterOfMass(InComLocalPose.GetLocation());
-		KinematicParticle->SetRotationOfMass(InComLocalPose.GetRotation());
+		Rigid->SetCenterOfMass(InComLocalPose.GetLocation());
+		Rigid->SetRotationOfMass(InComLocalPose.GetRotation());
 	}
 }
 
diff --git a/Engine/Source/Runtime/Engine/Private/PhysicsEngine/ImmediatePhysics/ImmediatePhysicsChaos/ImmediatePhysicsActorHandle_Chaos.cpp b/Engine/Source/Runtime/Engine/Private/PhysicsEngine/ImmediatePhysics/ImmediatePhysicsChaos/ImmediatePhysicsActorHandle_Chaos.cpp
--- a/Engine/Source/Runtime/Engine/Private/PhysicsEngine/ImmediatePhysics/ImmediatePhysicsChaos/ImmediatePhysicsActorHandle_Chaos.cpp
+++ b/Engine/Source/Runtime/Engine/Private/PhysicsEngine/ImmediatePhysics/ImmediatePhysicsChaos/ImmediatePhysicsActorHandle_Chaos.cpp
@@ -317,8 +317,6 @@ namespace ImmediatePhysics_Chaos
 				{
 					Kinematic->SetV(FVector::ZeroVector);
 					Kinematic->SetW(FVector::ZeroVector);
-					Kinematic->SetCenterOfMass(CoMTransform.GetTranslation());
-					Kinematic->SetRotationOfMass(CoMTransform.GetRotation());
 				}
 
 				auto* Dynamic = ParticleHandle->CastToRigidParticle();
@@ -328,6 +326,8 @@ namespace ImmediatePhysics_Chaos
 					FVector InertiaInv = (Mass > 0.0f) ? Inertia.Reciprocal() : FVector::ZeroVector;
 					Dynamic->SetM(Mass);
 					Dynamic->SetInvM(MassInv);
+					Dynamic->SetCenterOfMass(CoMTransform.GetTranslation());
+					Dynamic->SetRotationOfMass(CoMTransform.GetRotation());
 					Dynamic->SetI({ Inertia.X, Inertia.Y, Inertia.Z });
 					Dynamic->SetInvI({ InertiaInv.X, InertiaInv.Y, InertiaInv.Z });
 					if (BodyInstance != nullptr)
diff --git a/Engine/Source/Runtime/Experimental/Chaos/Public/Chaos/KinematicGeometryParticles.h b/Engine/Source/Runtime/Experimental/Chaos/Public/Chaos/KinematicGeometryParticles.h
--- a/Engine/Source/Runtime/Experimental/Chaos/Public/Chaos/KinematicGeometryParticles.h
+++ b/Engine/Source/Runtime/Experimental/Chaos/Public/Chaos/KinematicGeometryParticles.h
@@ -70,8 +70,6 @@ class TKinematicGeometryParticlesImp : public TGeometryParticlesImp<T, d, SimTyp
 		TArrayCollection::AddArray(&MV);
 		TArrayCollection::AddArray(&MW);
 		TArrayCollection::AddArray(&KinematicTargets);
-		TArrayCollection::AddArray(&MCenterOfMass);
-		TArrayCollection::AddArray(&MRotationOfMass);
 	}
 	TKinematicGeometryParticlesImp(const TKinematicGeometryParticlesImp<T, d, SimType>& Other) = delete;
 	CHAOS_API TKinematicGeometryParticlesImp(TKinematicGeometryParticlesImp<T, d, SimType>&& Other)
@@ -81,8 +79,6 @@ class TKinematicGeometryParticlesImp : public TGeometryParticlesImp<T, d, SimTyp
 		TArrayCollection::AddArray(&MV);
 		TArrayCollection::AddArray(&MW);
 		TArrayCollection::AddArray(&KinematicTargets);
-		TArrayCollection::AddArray(&MCenterOfMass);
-		TArrayCollection::AddArray(&MRotationOfMass);
 	}
 	CHAOS_API virtual ~TKinematicGeometryParticlesImp();
 
@@ -95,12 +91,6 @@ class TKinematicGeometryParticlesImp : public TGeometryParticlesImp<T, d, SimTyp
 	const TKinematicTarget<T, d>& KinematicTarget(const int32 Index) const { return KinematicTargets[Index]; }
 	TKinematicTarget<T, d>& KinematicTarget(const int32 Index) { return KinematicTargets[Index]; }
 
-	const TVector<T, d>& CenterOfMass(const int32 Index) const { return MCenterOfMass[Index]; }
-	TVector<T, d>& CenterOfMass(const int32 Index) { return MCenterOfMass[Index]; }
-
-	const TRotation<T, d>& RotationOfMass(const int32 Index) const { return MRotationOfMass[Index]; }
-	TRotation<T, d>& RotationOfMass(const int32 Index) { return MRotationOfMass[Index]; }
-
 	FString ToString(int32 index) const
 	{
 		FString BaseString = TGeometryParticlesImp<T, d, SimType>::ToString(index);
@@ -117,25 +107,18 @@ class TKinematicGeometryParticlesImp : public TGeometryParticlesImp<T, d, SimTyp
 	{
 		TGeometryParticlesImp<T, d, SimType>::Serialize(Ar);
 		Ar << MV << MW;
-		
+
 		Ar.UsingCustomVersion(FExternalPhysicsCustomObjectVersion::GUID);
-		if (Ar.CustomVer(FExternalPhysicsCustomObjectVersion::GUID) >= FExternalPhysicsCustomObjectVersion::KinematicTargets)
+		if(Ar.CustomVer(FExternalPhysicsCustomObjectVersion::GUID) >= FExternalPhysicsCustomObjectVersion::KinematicTargets)
 		{
 			Ar << KinematicTargets;
 		}
-		if (Ar.CustomVer(FExternalPhysicsCustomObjectVersion::GUID) >= FExternalPhysicsCustomObjectVersion::KinematicCentersOfMass)
-		{
-			Ar << MCenterOfMass;
-			Ar << MRotationOfMass;
-		}
 	}
 
   private:
 	TArrayCollectionArray<TVector<T, d>> MV;
 	TArrayCollectionArray<TVector<T, d>> MW;
 	TArrayCollectionArray<TKinematicTarget<T, d>> KinematicTargets;
-	TArrayCollectionArray<TVector<T, d>> MCenterOfMass;
-	TArrayCollectionArray<TRotation<T, d>> MRotationOfMass;
 };
 
 template <typename T, int d, EGeometryParticlesSimType SimType>
diff --git a/Engine/Source/Runtime/Experimental/Chaos/Public/Chaos/ParticleHandle.h b/Engine/Source/Runtime/Experimental/Chaos/Public/Chaos/ParticleHandle.h
--- a/Engine/Source/Runtime/Experimental/Chaos/Public/Chaos/ParticleHandle.h
+++ b/Engine/Source/Runtime/Experimental/Chaos/Public/Chaos/ParticleHandle.h
@@ -63,8 +63,6 @@ void KinematicGeometryParticleDefaultConstruct(FConcrete& Concrete, const TKinem
 {
 	Concrete.SetV(TVector<T, d>(0));
 	Concrete.SetW(TVector<T, d>(0));
-	Concrete.SetCenterOfMass(TVector<T, d>(0));
-	Concrete.SetRotationOfMass(TRotation<T, d>(FQuat(EForceInit::ForceInit)));
 }
 
 template <typename T, int d, bool bPersistent>
@@ -83,6 +81,8 @@ void PBDRigidParticleHandleImpDefaultConstruct(TPBDRigidParticleHandleImp<T, d,
 	Concrete.SetAngularImpulse(TVector<T, d>(0));
 	Concrete.SetM(1);
 	Concrete.SetInvM(1);
+	Concrete.SetCenterOfMass(TVector<T,d>(0));
+	Concrete.SetRotationOfMass(TRotation<T,d>::FromIdentity());
 	Concrete.SetI(PMatrix<T, d, d>(1, 1, 1));
 	Concrete.SetInvI(PMatrix<T, d, d>(1, 1, 1));
 	Concrete.SetLinearEtherDrag(0.f);
@@ -102,6 +102,8 @@ void PBDRigidParticleDefaultConstruct(TPBDRigidParticle<T,d>& Concrete, const TP
 	Concrete.SetAngularImpulse(TVector<T, d>(0));
 	Concrete.SetM(1);
 	Concrete.SetInvM(1);
+	Concrete.SetCenterOfMass(TVector<T,d>(0));
+	Concrete.SetRotationOfMass(TRotation<T,d>(FQuat(EForceInit::ForceInit)));
 	Concrete.SetI(PMatrix<T, d, d>(1, 1, 1));
 	Concrete.SetInvI(PMatrix<T, d, d>(1, 1, 1));
 	Concrete.SetLinearEtherDrag(0.f);
@@ -616,12 +618,6 @@ public:
 	const TKinematicTarget<T, d>& KinematicTarget() const { return KinematicGeometryParticles->KinematicTarget(ParticleIdx); }
 	TKinematicTarget<T, d>& KinematicTarget() { return KinematicGeometryParticles->KinematicTarget(ParticleIdx); }
 
-	const TVector<T, d>& CenterOfMass() const { return KinematicGeometryParticles->CenterOfMass(ParticleIdx); }
-	void SetCenterOfMass(const TVector<T, d>& InCenterOfMass) { KinematicGeometryParticles->CenterOfMass(ParticleIdx) = InCenterOfMass; }
-	
-	const TRotation<T, d>& RotationOfMass() const { return KinematicGeometryParticles->RotationOfMass(ParticleIdx); }
-	void SetRotationOfMass(const TRotation<T, d>& InRotationOfMass) { KinematicGeometryParticles->RotationOfMass(ParticleIdx) = InRotationOfMass; }
-
 	//Really only useful when using a transient handle
 	const TKinematicGeometryParticleHandleImp<T, d, true>* Handle() const { return KinematicGeometryParticles->Handle(ParticleIdx); }
 	TKinematicGeometryParticleHandleImp<T, d, true>* Handle() { return KinematicGeometryParticles->Handle(ParticleIdx); }
@@ -753,6 +749,13 @@ public:
 	T& InvM() { return PBDRigidParticles->InvM(ParticleIdx); }
 	void SetInvM(const T& InInvM) { PBDRigidParticles->InvM(ParticleIdx) = InInvM; }
 
+	const TVector<T,d>& CenterOfMass() const { return PBDRigidParticles->CenterOfMass(ParticleIdx); }
+	void SetCenterOfMass(const TVector<T,d>& InCenterOfMass) { PBDRigidParticles->CenterOfMass(ParticleIdx) = InCenterOfMass; }
+
+	const TRotation<T,d>& RotationOfMass() const { return PBDRigidParticles->RotationOfMass(ParticleIdx); }
+	void SetRotationOfMass(const TRotation<T,d>& InRotationOfMass) { PBDRigidParticles->RotationOfMass(ParticleIdx) = InRotationOfMass; }
+
+
 	T LinearEtherDrag() const { return PBDRigidParticles->LinearEtherDrag(ParticleIdx); }
 	T& LinearEtherDrag() { return PBDRigidParticles->LinearEtherDrag(ParticleIdx); }
 	void SetLinearEtherDrag(const T& InLinearEtherDrag) { PBDRigidParticles->LinearEtherDrag(ParticleIdx) = InLinearEtherDrag; }
@@ -1189,9 +1192,9 @@ public:
 
 	TVector<T, d> CenterOfMass() const
 	{
-		if (auto KinematicHandle = Handle->CastToKinematicParticle())
+		if (auto RigidHandle = Handle->CastToRigidParticle())
 		{
-			return KinematicHandle->CenterOfMass();
+			return RigidHandle->CenterOfMass();
 		}
 
 		return TVector<T, d>(0);
@@ -1199,9 +1202,9 @@ public:
 
 	TRotation<T, d> RotationOfMass() const
 	{
-		if (auto KinematicHandle = Handle->CastToKinematicParticle())
+		if (auto RigidHandle = Handle->CastToRigidParticle())
 		{
-			return KinematicHandle->RotationOfMass();
+			return RigidHandle->RotationOfMass();
 		}
 
 		return TRotation<T, d>::FromIdentity();
@@ -1876,33 +1879,16 @@ public:
 		MVelocities.Write(InVelocities,bInvalidate,MDirtyFlags,Proxy);
 	}
 
-	const TVector<T, d>& CenterOfMass() const { return MMassProps.Read().CenterOfMass; }
-	void SetCenterOfMass(const TVector<T, d>& InCenterOfMass, bool bInvalidate = true)
-	{
-		MMassProps.Modify(bInvalidate,MDirtyFlags,Proxy,[&InCenterOfMass](auto& Data){ Data.CenterOfMass = InCenterOfMass;});
-	}
-	
-	const TRotation<T, d>& RotationOfMass() const { return MMassProps.Read().RotationOfMass; }
-	void SetRotationOfMass(const TRotation<T, d>& InRotationOfMass, bool bInvalidate = true)
-	{
-		MMassProps.Modify(bInvalidate,MDirtyFlags,Proxy,[&InRotationOfMass](auto& Data){ Data.RotationOfMass = InRotationOfMass;});
-	}
-
 	EObjectStateType ObjectState() const;
 
 private:
 	TParticleProperty<FParticleVelocities, EParticleProperty::Velocities> MVelocities;
 
-protected:
-	//todo: move this into pbdrigid
-	TParticleProperty<FParticleMassProps,EParticleProperty::MassProps> MMassProps;
-
 protected:
 	virtual void SyncRemoteDataImp(FDirtyPropertiesManager& Manager, int32 DataIdx, const FParticleDirtyData& RemoteData) const
 	{
 		Base::SyncRemoteDataImp(Manager, DataIdx, RemoteData);
 		MVelocities.SyncRemote(Manager, DataIdx, RemoteData);
-		MMassProps.SyncRemote(Manager, DataIdx, RemoteData);
 	}
 };
 
@@ -1917,16 +1903,12 @@ public:
 	TKinematicGeometryParticleData(EParticleType InType = EParticleType::Kinematic)
 		: Base(InType)
 		, MV(TVector<T, d>(0))
-		, MW(TVector<T, d>(0))
-		, MCenterOfMass(TVector<T, d>(0))
-		, MRotationOfMass(TRotation<T, d>(FQuat(EForceInit::ForceInit))) {}
+		, MW(TVector<T, d>(0)){}
 
 	TKinematicGeometryParticleData(const TKinematicGeometryParticle<T, d>& InParticle)
 		: Base(InParticle)
 		, MV(InParticle.V())
 		, MW(InParticle.W()) 
-		, MCenterOfMass(InParticle.CenterOfMass()) 
-		, MRotationOfMass(InParticle.RotationOfMass()) 
 	{
 		Type = EParticleType::Kinematic;
 	}
@@ -1937,22 +1919,16 @@ public:
 		Type = EParticleType::Kinematic;
 		MV = TVector<T, d>(0);
 		MW = TVector<T, d>(0);
-		MCenterOfMass = TVector<T, d>(0);
-		MRotationOfMass = TRotation<T, d>(FQuat(EForceInit::ForceInit));
 	}
 
 	void Init(const TKinematicGeometryParticle<T, d>& InParticle) {
 			Base::Init(InParticle);
 			MV = InParticle.V();
 			MW = InParticle.W();
-			MCenterOfMass = InParticle.CenterOfMass();
-			MRotationOfMass = InParticle.RotationOfMass();
 			Type = EParticleType::Kinematic;
 	}
 	TVector<T, d> MV;
 	TVector<T, d> MW;
-	TVector<T, d> MCenterOfMass;
-	TRotation<T, d> MRotationOfMass;
 };
 
 
@@ -1968,7 +1944,6 @@ public:
 
 	using Base::MDirtyFlags;
 	using Base::MMiscData;
-	using Base::MMassProps;
 
 
 protected:
@@ -2094,6 +2069,19 @@ public:
 		MMassProps.Modify(true,MDirtyFlags,Proxy,[InInvM](auto& Data){ Data.InvM = InInvM;});
 	}
 
+	const TVector<T,d>& CenterOfMass() const { return MMassProps.Read().CenterOfMass; }
+	void SetCenterOfMass(const TVector<T,d>& InCenterOfMass,bool bInvalidate = true)
+	{
+		MMassProps.Modify(bInvalidate,MDirtyFlags,Proxy,[&InCenterOfMass](auto& Data){ Data.CenterOfMass = InCenterOfMass;});
+	}
+
+	const TRotation<T,d>& RotationOfMass() const { return MMassProps.Read().RotationOfMass; }
+	void SetRotationOfMass(const TRotation<T,d>& InRotationOfMass,bool bInvalidate = true)
+	{
+		MMassProps.Modify(bInvalidate,MDirtyFlags,Proxy,[&InRotationOfMass](auto& Data){ Data.RotationOfMass = InRotationOfMass;});
+	}
+
+
 	T LinearEtherDrag() const { return this->MNonFrequentData.Read().LinearEtherDrag; }
 	void SetLinearEtherDrag(const T& InLinearEtherDrag)
 	{
@@ -2150,6 +2138,7 @@ public:
 	bool HasAwakeEvent() { return MAwakeEvent; }
 
 private:
+	TParticleProperty<FParticleMassProps,EParticleProperty::MassProps> MMassProps;
 	TParticleProperty<FParticleDynamics, EParticleProperty::Dynamics> MDynamics;
 	TUniquePtr<TBVHParticles<T, d>> MCollisionParticles;
 	int32 MIsland;
@@ -2161,6 +2150,7 @@ protected:
 	virtual void SyncRemoteDataImp(FDirtyPropertiesManager& Manager, int32 DataIdx, const FParticleDirtyData& RemoteData) const
 	{
 		Base::SyncRemoteDataImp(Manager,DataIdx,RemoteData);
+		MMassProps.SyncRemote(Manager,DataIdx,RemoteData);
 		MDynamics.SyncRemote(Manager,DataIdx,RemoteData);
 	}
 };
@@ -2205,6 +2195,8 @@ public:
 		, MCollisionParticles(nullptr)
 		, MM(T(0))
 		, MInvM(T(0))
+		, MCenterOfMass(TVector<T, d>(0))
+		, MRotationOfMass(TRotation<T, d>(FQuat(EForceInit::ForceInit)))
 		, MLinearEtherDrag(T(0))
 		, MAngularEtherDrag(T(0))
 		, MIsland(INDEX_NONE)
@@ -2227,6 +2219,8 @@ public:
 		, MCollisionParticles(nullptr)
 		, MM(InParticle.M())
 		, MInvM(InParticle.InvM())
+		, MCenterOfMass(InParticle.CenterOfMass())
+		, MRotationOfMass(InParticle.RotationOfMass())
 		, MLinearEtherDrag(InParticle.LinearEtherDrag())
 		, MAngularEtherDrag(InParticle.AngularEtherDrag())
 		, MIsland(InParticle.Island())
@@ -2250,6 +2244,8 @@ public:
 	const TBVHParticles<T, d> * MCollisionParticles;
 	T MM;
 	T MInvM;
+	TVector<T,d> MCenterOfMass;
+	TRotation<T,d> MRotationOfMass;
 	T MLinearEtherDrag;
 	T MAngularEtherDrag;
 	int32 MIsland;
@@ -2272,6 +2268,8 @@ public:
 		MCollisionParticles = nullptr;
 		MM = T(0);
 		MInvM = T(0);
+		MCenterOfMass = TVector<T,d>(0);
+		MRotationOfMass = TRotation<T,d>::FromIdentity();
 		MLinearEtherDrag = T(0);
 		MAngularEtherDrag = T(0);
 		MIsland = INDEX_NONE;
diff --git a/Engine/Source/Runtime/Experimental/Chaos/Public/Chaos/RigidParticles.h b/Engine/Source/Runtime/Experimental/Chaos/Public/Chaos/RigidParticles.h
--- a/Engine/Source/Runtime/Experimental/Chaos/Public/Chaos/RigidParticles.h
+++ b/Engine/Source/Runtime/Experimental/Chaos/Public/Chaos/RigidParticles.h
@@ -68,6 +68,8 @@ class TRigidParticles : public TKinematicGeometryParticles<T, d>
 		TArrayCollection::AddArray(&MInvI);
 		TArrayCollection::AddArray(&MM);
 		TArrayCollection::AddArray(&MInvM);
+		TArrayCollection::AddArray(&MCenterOfMass);
+		TArrayCollection::AddArray(&MRotationOfMass);
 		TArrayCollection::AddArray(&MLinearEtherDrag);
 		TArrayCollection::AddArray(&MAngularEtherDrag);
 		TArrayCollection::AddArray(&MCollisionParticles);
@@ -79,7 +81,19 @@ class TRigidParticles : public TKinematicGeometryParticles<T, d>
 	}
 	TRigidParticles(const TRigidParticles<T, d>& Other) = delete;
 	CHAOS_API TRigidParticles(TRigidParticles<T, d>&& Other)
-	    : TKinematicGeometryParticles<T, d>(MoveTemp(Other)), MF(MoveTemp(Other.MF)), MT(MoveTemp(Other.MT)), MLinearImpulse(MoveTemp(Other.MLinearImpulse)), MAngularImpulse(MoveTemp(Other.MAngularImpulse)), MI(MoveTemp(Other.MI)), MInvI(MoveTemp(Other.MInvI)), MM(MoveTemp(Other.MM)), MInvM(MoveTemp(Other.MInvM)), MCollisionParticles(MoveTemp(Other.MCollisionParticles)), MCollisionGroup(MoveTemp(Other.MCollisionGroup)), MObjectState(MoveTemp(Other.MObjectState))
+	    : TKinematicGeometryParticles<T, d>(MoveTemp(Other))
+		, MF(MoveTemp(Other.MF))
+		, MT(MoveTemp(Other.MT))
+		, MLinearImpulse(MoveTemp(Other.MLinearImpulse))
+		, MAngularImpulse(MoveTemp(Other.MAngularImpulse))
+		, MI(MoveTemp(Other.MI)), MInvI(MoveTemp(Other.MInvI))
+		, MM(MoveTemp(Other.MM))
+		, MInvM(MoveTemp(Other.MInvM))
+		, MCenterOfMass(MoveTemp(Other.MCenterOfMass))
+		, MRotationOfMass(MoveTemp(Other.MRotationOfMass))
+		, MCollisionParticles(MoveTemp(Other.MCollisionParticles))
+		, MCollisionGroup(MoveTemp(Other.MCollisionGroup))
+		, MObjectState(MoveTemp(Other.MObjectState))
 	{
 		TArrayCollection::AddArray(&MF);
 		TArrayCollection::AddArray(&MT);
@@ -89,6 +103,8 @@ class TRigidParticles : public TKinematicGeometryParticles<T, d>
 		TArrayCollection::AddArray(&MInvI);
 		TArrayCollection::AddArray(&MM);
 		TArrayCollection::AddArray(&MInvM);
+		TArrayCollection::AddArray(&MCenterOfMass);
+		TArrayCollection::AddArray(&MRotationOfMass);
 		TArrayCollection::AddArray(&MLinearEtherDrag);
 		TArrayCollection::AddArray(&MAngularEtherDrag);
 		TArrayCollection::AddArray(&MCollisionParticles);
@@ -126,6 +142,13 @@ class TRigidParticles : public TKinematicGeometryParticles<T, d>
 	FORCEINLINE const T InvM(const int32 Index) const { return MInvM[Index]; }
 	FORCEINLINE T& InvM(const int32 Index) { return MInvM[Index]; }
 
+	const TVector<T,d>& CenterOfMass(const int32 Index) const { return MCenterOfMass[Index]; }
+	TVector<T,d>& CenterOfMass(const int32 Index) { return MCenterOfMass[Index]; }
+
+	const TRotation<T,d>& RotationOfMass(const int32 Index) const { return MRotationOfMass[Index]; }
+	TRotation<T,d>& RotationOfMass(const int32 Index) { return MRotationOfMass[Index]; }
+
+
 	FORCEINLINE const T& LinearEtherDrag(const int32 index) const { return MLinearEtherDrag[index]; }
 	FORCEINLINE T& LinearEtherDrag(const int32 index) { return MLinearEtherDrag[index]; }
 
@@ -188,12 +211,23 @@ class TRigidParticles : public TKinematicGeometryParticles<T, d>
 	FORCEINLINE FString ToString(int32 index) const
 	{
 		FString BaseString = TKinematicGeometryParticles<T, d>::ToString(index);
-		return FString::Printf(TEXT("%s, MF:%s, MT:%s, MLinearImpulse:%s, MAngularImpulse:%s, MI:%s, MInvI:%s, MM:%f, MInvM:%f, MCollisionParticles(num):%d, MCollisionGroup:%d, MDisabled:%d, MSleepring:%d, MIsland:%d"), *BaseString, *F(index).ToString(), *Torque(index).ToString(), *LinearImpulse(index).ToString(), *AngularImpulse(index).ToString(), *I(index).ToString(), *InvI(index).ToString(), M(index), InvM(index), CollisionParticlesSize(index), CollisionGroup(index), Disabled(index), Sleeping(index), Island(index));
+		return FString::Printf(TEXT("%s, MF:%s, MT:%s, MLinearImpulse:%s, MAngularImpulse:%s, MI:%s, MInvI:%s, MM:%f, MInvM:%f, MCenterOfMass:%s, MRotationOfMass:%s, MCollisionParticles(num):%d, MCollisionGroup:%d, MDisabled:%d, MSleepring:%d, MIsland:%d"),
+			*BaseString, *F(index).ToString(), *Torque(index).ToString(), *LinearImpulse(index).ToString(), *AngularImpulse(index).ToString(),
+			*I(index).ToString(), *InvI(index).ToString(), M(index), InvM(index), *CenterOfMass(index).ToString(), *RotationOfMass(index).ToString(), CollisionParticlesSize(index),
+			CollisionGroup(index), Disabled(index), Sleeping(index), Island(index));
 	}
 
 	CHAOS_API virtual void Serialize(FChaosArchive& Ar) override
 	{
 		TKinematicGeometryParticles<T,d>::Serialize(Ar);
+		
+		Ar.UsingCustomVersion(FExternalPhysicsCustomObjectVersion::GUID);
+		if(Ar.CustomVer(FExternalPhysicsCustomObjectVersion::GUID) >= FExternalPhysicsCustomObjectVersion::KinematicCentersOfMass)
+		{
+			Ar << MCenterOfMass;
+			Ar << MRotationOfMass;
+		}
+
 		Ar << MF << MT << MLinearImpulse << MAngularImpulse << MI << MInvI << MM << MInvM;
 
 		Ar.UsingCustomVersion(FExternalPhysicsCustomObjectVersion::GUID);
@@ -214,6 +248,8 @@ class TRigidParticles : public TKinematicGeometryParticles<T, d>
 	TArrayCollectionArray<PMatrix<T, d, d>> MInvI;
 	TArrayCollectionArray<T> MM;
 	TArrayCollectionArray<T> MInvM;
+	TArrayCollectionArray<TVector<T,d>> MCenterOfMass;
+	TArrayCollectionArray<TRotation<T,d>> MRotationOfMass;
 	TArrayCollectionArray<T> MLinearEtherDrag;
 	TArrayCollectionArray<T> MAngularEtherDrag;
 	TArrayCollectionArray<TUniquePtr<TBVHParticles<T, d>>> MCollisionParticles;
diff --git a/Engine/Source/Runtime/Experimental/ChaosSolvers/Private/PhysicsProxy/SingleParticlePhysicsProxy.cpp b/Engine/Source/Runtime/Experimental/ChaosSolvers/Private/PhysicsProxy/SingleParticlePhysicsProxy.cpp
--- a/Engine/Source/Runtime/Experimental/ChaosSolvers/Private/PhysicsProxy/SingleParticlePhysicsProxy.cpp
+++ b/Engine/Source/Runtime/Experimental/ChaosSolvers/Private/PhysicsProxy/SingleParticlePhysicsProxy.cpp
@@ -107,20 +107,16 @@ void PushToPhysicsStateImp(const Chaos::FDirtyPropertiesManager& Manager, TParti
 			}
 		}
 
-		if(auto NewData = ParticleData.FindMassProps(Manager, DataIdx))
+		if(bHasDynamicData)
 		{
-			if(bHasKinematicData)
-			{
-				KinematicHandle->SetCenterOfMass(NewData->CenterOfMass);
-				KinematicHandle->SetRotationOfMass(NewData->RotationOfMass);
-			}
-
-			if(bHasDynamicData)
+			if(auto NewData = ParticleData.FindMassProps(Manager,DataIdx))
 			{
 				RigidHandle->SetM(NewData->M);
 				RigidHandle->SetInvM(NewData->InvM);
 				RigidHandle->SetI(NewData->I);
 				RigidHandle->SetInvI(NewData->InvI);
+				RigidHandle->SetCenterOfMass(NewData->CenterOfMass);
+				RigidHandle->SetRotationOfMass(NewData->RotationOfMass);
 			}
 		}
 
