Fix issues with queries against geometry with very large scales.

diff --git a/Engine/Source/Runtime/Experimental/Chaos/Public/Chaos/ImplicitObjectScaled.h b/Engine/Source/Runtime/Experimental/Chaos/Public/Chaos/ImplicitObjectScaled.h
--- a/Engine/Source/Runtime/Experimental/Chaos/Public/Chaos/ImplicitObjectScaled.h
+++ b/Engine/Source/Runtime/Experimental/Chaos/Public/Chaos/ImplicitObjectScaled.h
@@ -343,7 +343,7 @@ public:
 				if (NewTime < Length && NewTime != 0) // Normal/Position output may be uninitialized with TOI 0.
 				{
 					OutPosition = MScale * UnscaledPosition;
-					OutNormal = (MInvScale * UnscaledNormal).GetSafeNormal();
+					OutNormal = (MInvScale * UnscaledNormal).GetSafeNormal(TNumericLimits<T>::Min());
 					OutTime = NewTime;
 					return true;
 				}
@@ -447,7 +447,7 @@ public:
 		// Compute final normal
 		const TVector<T, d> LocalNormal = MObject->FindGeometryOpposingNormal(LocalDenormDir, HintFaceIndex, LocalOriginalNormal);
 		TVector<T, d> Normal = LocalNormal * MInvScale;
-		if (CHAOS_ENSURE(Normal.SafeNormalize()) == 0)
+		if (CHAOS_ENSURE(Normal.SafeNormalize(TNumericLimits<T>::Min())) == 0)
 		{
 			Normal = TVector<T,3>(0,0,1);
 		}
@@ -528,10 +528,10 @@ public:
 	const TVector<T, d>& GetInvScale() const { return MInvScale; }
 	void SetScale(const TVector<T, d>& Scale)
 	{
-		constexpr T MinMagnitude = 1e-4;
+		constexpr T MinMagnitude = 1e-6;
 		for (int Axis = 0; Axis < d; ++Axis)
 		{
-			if (FMath::Abs(Scale[Axis]) < MinMagnitude)
+			if (!CHAOS_ENSURE(FMath::Abs(Scale[Axis]) >= MinMagnitude))
 			{
 				MScale[Axis] = MinMagnitude;
 			}
diff --git a/Engine/Source/Runtime/Experimental/ChaosCore/Public/Chaos/Vector.h b/Engine/Source/Runtime/Experimental/ChaosCore/Public/Chaos/Vector.h
--- a/Engine/Source/Runtime/Experimental/ChaosCore/Public/Chaos/Vector.h
+++ b/Engine/Source/Runtime/Experimental/ChaosCore/Public/Chaos/Vector.h
@@ -504,10 +504,10 @@ namespace Chaos
 					return MakePair(max.Z, 2);
 			}
 		}
-		float SafeNormalize()
+		float SafeNormalize(float Epsilon = 1e-4)
 		{
 			float Size = SizeSquared();
-			if (Size < (float)1e-4)
+			if (Size < Epsilon)
 			{
 				*this = AxisVector(0);
 				return 0.f;
