Easier marshaling of mass props

diff --git a/Engine/Source/Runtime/Experimental/Chaos/Public/Chaos/ParticleDirtyFlags.h b/Engine/Source/Runtime/Experimental/Chaos/Public/Chaos/ParticleDirtyFlags.h
--- a/Engine/Source/Runtime/Experimental/Chaos/Public/Chaos/ParticleDirtyFlags.h
+++ b/Engine/Source/Runtime/Experimental/Chaos/Public/Chaos/ParticleDirtyFlags.h
@@ -164,34 +164,73 @@ inline FChaosArchive& operator<<(FChaosArchive& Ar, FParticleDynamics& Data)
 	return Ar;
 }
 
-struct FParticleMassProps
+class FParticleMassProps
 {
-	FVec3 CenterOfMass;
-	FRotation3 RotationOfMass;
-	FMatrix33 I;
-	FMatrix33 InvI;
-	FReal M;
-	FReal InvM;
-
+public:
 	void Serialize(FChaosArchive& Ar)
 	{
-		Ar << CenterOfMass;
-		Ar << RotationOfMass;
-		Ar << I;
-		Ar << InvI;
-		Ar << M;
-		Ar << InvM;
+		Ar << MCenterOfMass;
+		Ar << MRotationOfMass;
+		Ar << MI;
+		Ar << MInvI;
+		Ar << MM;
+		Ar << MInvM;
+	}
+
+	template <typename TOther>
+	void CopyFrom(const TOther& Other)
+	{
+		MCenterOfMass = Other.CenterOfMass();
+		MRotationOfMass = Other.RotationOfMass();
+		MI = Other.I();
+		MInvI = Other.InvI();
+		MM = Other.M();
+		MInvM = Other.InvM();
+	}
+
+	template <typename TOther>
+	bool IsEqual(const TOther& Other) const
+	{
+		return CenterOfMass() == Other.CenterOfMass()
+			&& RotationOfMass() == Other.RotationOfMass()
+			&& I() == Other.I()
+			&& InvI() == Other.InvI()
+			&& M() == Other.M()
+			&& InvM() == Other.InvM();
 	}
 
 	bool operator==(const FParticleMassProps& Other) const
 	{
-		return CenterOfMass == Other.CenterOfMass
-			&& RotationOfMass == Other.RotationOfMass
-			&& I == Other.I
-			&& InvI == Other.InvI
-			&& M == Other.M
-			&& InvM == Other.InvM;
+		return IsEqual(Other);
 	}
+
+	const FVec3& CenterOfMass() const { return MCenterOfMass; }
+	void SetCenterOfMass(const FVec3& InCenterOfMass){ MCenterOfMass = InCenterOfMass; }
+
+	const FRotation3& RotationOfMass() const { return MRotationOfMass; }
+	void SetRotationOfMass(const FRotation3& InRotationOfMass){ MRotationOfMass = InRotationOfMass; }
+
+	const FMatrix33& I() const { return MI; }
+	void SetI(const FMatrix33& InI){ MI = InI; }
+
+	const FMatrix33& InvI() const { return MInvI; }
+	void SetInvI(const FMatrix33& InInvI){ MInvI = InInvI; }
+
+	FReal M() const { return MM; }
+	void SetM(FReal InM){ MM = InM; }
+
+	FReal InvM() const { return MInvM; }
+	void SetInvM(FReal InInvM){ MInvM = InInvM; }
+
+private:
+	FVec3 MCenterOfMass;
+	FRotation3 MRotationOfMass;
+	FMatrix33 MI;
+	FMatrix33 MInvI;
+	FReal MM;
+	FReal MInvM;
+
+
 };
 
 inline FChaosArchive& operator<<(FChaosArchive& Ar,FParticleMassProps& Data)
diff --git a/Engine/Source/Runtime/Experimental/Chaos/Public/Chaos/ParticleHandle.h b/Engine/Source/Runtime/Experimental/Chaos/Public/Chaos/ParticleHandle.h
--- a/Engine/Source/Runtime/Experimental/Chaos/Public/Chaos/ParticleHandle.h
+++ b/Engine/Source/Runtime/Experimental/Chaos/Public/Chaos/ParticleHandle.h
@@ -733,6 +733,16 @@ public:
 		SetAngularImpulse(Dynamics.AngularImpulse());
 	}
 
+	void SetMassProps(const FParticleMassProps& Props)
+	{
+		SetCenterOfMass(Props.CenterOfMass());
+		SetRotationOfMass(Props.RotationOfMass());
+		SetI(Props.I());
+		SetInvI(Props.InvI());
+		SetM(Props.M());
+		SetInvM(Props.InvM());
+	}
+
 	const PMatrix<T, d, d>& I() const { return PBDRigidParticles->I(ParticleIdx); }
 	PMatrix<T, d, d>& I() { return PBDRigidParticles->I(ParticleIdx); }
 	void SetI(const PMatrix<T, d, d>& InI) { PBDRigidParticles->I(ParticleIdx) = InI; }
@@ -2045,42 +2055,46 @@ public:
 		MDynamics.Write(InDynamics,bInvalidate,MDirtyFlags,Proxy);
 	}
 
-	const PMatrix<T, d, d>& I() const { return MMassProps.Read().I; }
+	const PMatrix<T, d, d>& I() const { return MMassProps.Read().I(); }
 	void SetI(const PMatrix<T, d, d>& InI)
 	{
-		MMassProps.Modify(true,MDirtyFlags,Proxy,[&InI](auto& Data){ Data.I = InI;});
+		MMassProps.Modify(true,MDirtyFlags,Proxy,[&InI](auto& Data){ Data.SetI(InI);});
 	}
 
-	const PMatrix<T, d, d>& InvI() const { return MMassProps.Read().InvI; }
+	const PMatrix<T, d, d>& InvI() const { return MMassProps.Read().InvI(); }
 	void SetInvI(const PMatrix<T, d, d>& InInvI)
 	{
-		MMassProps.Modify(true,MDirtyFlags,Proxy,[&InInvI](auto& Data){ Data.InvI = InInvI;});
+		MMassProps.Modify(true,MDirtyFlags,Proxy,[&InInvI](auto& Data){ Data.SetInvI(InInvI);});
 	}
 
-	T M() const { return MMassProps.Read().M; }
+	T M() const { return MMassProps.Read().M(); }
 	void SetM(const T& InM)
 	{
-		MMassProps.Modify(true,MDirtyFlags,Proxy,[InM](auto& Data){ Data.M = InM;});
+		MMassProps.Modify(true,MDirtyFlags,Proxy,[InM](auto& Data){ Data.SetM(InM);});
 	}
 
-	T InvM() const { return MMassProps.Read().InvM; }
+	T InvM() const { return MMassProps.Read().InvM(); }
 	void SetInvM(const T& InInvM)
 	{
-		MMassProps.Modify(true,MDirtyFlags,Proxy,[InInvM](auto& Data){ Data.InvM = InInvM;});
+		MMassProps.Modify(true,MDirtyFlags,Proxy,[InInvM](auto& Data){ Data.SetInvM(InInvM);});
 	}
-
-	const TVector<T,d>& CenterOfMass() const { return MMassProps.Read().CenterOfMass; }
+	
+	const TVector<T,d>& CenterOfMass() const { return MMassProps.Read().CenterOfMass(); }
 	void SetCenterOfMass(const TVector<T,d>& InCenterOfMass,bool bInvalidate = true)
 	{
-		MMassProps.Modify(bInvalidate,MDirtyFlags,Proxy,[&InCenterOfMass](auto& Data){ Data.CenterOfMass = InCenterOfMass;});
+		MMassProps.Modify(bInvalidate,MDirtyFlags,Proxy,[&InCenterOfMass](auto& Data){ Data.SetCenterOfMass(InCenterOfMass);});
 	}
 
-	const TRotation<T,d>& RotationOfMass() const { return MMassProps.Read().RotationOfMass; }
+	const TRotation<T,d>& RotationOfMass() const { return MMassProps.Read().RotationOfMass(); }
 	void SetRotationOfMass(const TRotation<T,d>& InRotationOfMass,bool bInvalidate = true)
 	{
-		MMassProps.Modify(bInvalidate,MDirtyFlags,Proxy,[&InRotationOfMass](auto& Data){ Data.RotationOfMass = InRotationOfMass;});
+		MMassProps.Modify(bInvalidate,MDirtyFlags,Proxy,[&InRotationOfMass](auto& Data){ Data.SetRotationOfMass(InRotationOfMass);});
 	}
 
+	void SetMassProps(const FParticleMassProps& InProps)
+	{
+		MMassProps.Write(InProps,true,MDirtyFlags,Proxy);
+	}
 
 	T LinearEtherDrag() const { return this->MNonFrequentData.Read().LinearEtherDrag; }
 	void SetLinearEtherDrag(const T& InLinearEtherDrag)
diff --git a/Engine/Source/Runtime/Experimental/ChaosSolvers/Private/PhysicsProxy/SingleParticlePhysicsProxy.cpp b/Engine/Source/Runtime/Experimental/ChaosSolvers/Private/PhysicsProxy/SingleParticlePhysicsProxy.cpp
--- a/Engine/Source/Runtime/Experimental/ChaosSolvers/Private/PhysicsProxy/SingleParticlePhysicsProxy.cpp
+++ b/Engine/Source/Runtime/Experimental/ChaosSolvers/Private/PhysicsProxy/SingleParticlePhysicsProxy.cpp
@@ -111,17 +111,9 @@ void PushToPhysicsStateImp(const Chaos::FDirtyPropertiesManager& Manager, TParti
 		{
 			if(auto NewData = ParticleData.FindMassProps(Manager,DataIdx))
 			{
-				RigidHandle->SetM(NewData->M);
-				RigidHandle->SetInvM(NewData->InvM);
-				RigidHandle->SetI(NewData->I);
-				RigidHandle->SetInvI(NewData->InvI);
-				RigidHandle->SetCenterOfMass(NewData->CenterOfMass);
-				RigidHandle->SetRotationOfMass(NewData->RotationOfMass);
+				RigidHandle->SetMassProps(*NewData);
 			}
-		}
 
-		if(bHasDynamicData)
-		{
 			if(auto NewData = ParticleData.FindDynamics(Manager, DataIdx))
 			{
 				RigidHandle->SetDynamics(*NewData);
