Chaos : Destruction
- Union cluster groups
... Initilization
... Particle to Proxy updated to be a set.

diff --git a/Engine/Source/Runtime/Experimental/Chaos/Private/Chaos/PBDRigidClustering.cpp b/Engine/Source/Runtime/Experimental/Chaos/Private/Chaos/PBDRigidClustering.cpp
--- a/Engine/Source/Runtime/Experimental/Chaos/Private/Chaos/PBDRigidClustering.cpp
+++ b/Engine/Source/Runtime/Experimental/Chaos/Private/Chaos/PBDRigidClustering.cpp
@@ -21,6 +21,7 @@
 #include "Chaos/PerParticleEulerStepVelocity.h"
 #include "Chaos/PerParticleEtherDrag.h"
 #include "Chaos/PerParticlePBDEulerStep.h"
+#include "CoreMinimal.h"
 
 namespace Chaos
 {
@@ -599,22 +600,36 @@ namespace Chaos
 	void TPBDRigidClustering<T_FPBDRigidsEvolution, T_FPBDCollisionConstraint, T, d>::UnionClusterGroups()
 	{
 		SCOPE_CYCLE_COUNTER(STAT_UnionClusterGroups);
-
-
-		for (TTuple<int32, TArray<TPBDRigidClusteredParticleHandle<T, 3>* >>& Group : ClusterUnionMap)
+		if (ClusterUnionMap.Num())
 		{
-			uint32 ClusterGroupID = Group.Key;
-			TArray<TPBDRigidClusteredParticleHandle<T, 3>* > Handles = Group.Value;
+			TMap < TPBDRigidParticleHandle<T, 3>*, TPBDRigidParticleHandle<T, 3>*> ClusterParents;
+			TMap < int32, TArray< TPBDRigidParticleHandle<T, 3>*>> NewClusterGroups;
+			for (TTuple<int32, TArray<TPBDRigidClusteredParticleHandle<T, 3>* >>& Group : ClusterUnionMap)
+			{
+				int32 ClusterGroupID = Group.Key;
+				TArray<TPBDRigidClusteredParticleHandle<T, 3>* > Handles = Group.Value;
+
+				if (!NewClusterGroups.Contains(ClusterGroupID))
+					NewClusterGroups.Add(ClusterGroupID, TArray < TPBDRigidParticleHandle<T, 3>*>());
 
+				TArray<TPBDRigidParticleHandle<T, 3>*> ClusterBodies;
+				for (TPBDRigidClusteredParticleHandle<T, 3>* ActiveCluster : Handles)
+				{
+					TSet<TPBDRigidParticleHandle<T, 3>*> Children = ReleaseClusterParticles(ActiveCluster, nullptr, true);
+					NewClusterGroups[ClusterGroupID].Append(Children.Array());
+					for (auto& Child : Children) ClusterParents.Add(Child, ActiveCluster);
+				}
+			}
 
-			TArray<TPBDRigidParticleHandle<T, 3>*> ClusterBodies;
-			for (TPBDRigidClusteredParticleHandle<T, 3>* ActiveCluster : Handles)
+			for (TTuple<int32, TArray<TPBDRigidParticleHandle<T, 3>* >>& Group : NewClusterGroups)
 			{
-				ClusterBodies.Append(ReleaseClusterParticles(ActiveCluster, nullptr, true).Array());
+				int32 ClusterGroupID = Group.Key;
+				TArray< TPBDRigidParticleHandle<T, 3> *> ActiveCluster = Group.Value;
+				TPBDRigidParticleHandle<T, 3>* NewCluster = CreateClusterParticle(-FMath::Abs(ClusterGroupID), MoveTemp(Group.Value));
+				for (auto& Constituent : ActiveCluster) MEvolution.DoInternalParticleInitilization( ClusterParents[Constituent], NewCluster);
 			}
-			CreateClusterParticle(-Group.Key, MoveTemp(ClusterBodies));
+			ClusterUnionMap.Empty();
 		}
-		ClusterUnionMap.Empty();
 	}
 
 
diff --git a/Engine/Source/Runtime/Experimental/Chaos/Private/Chaos/PBDRigidsEvolutionGBF.cpp b/Engine/Source/Runtime/Experimental/Chaos/Private/Chaos/PBDRigidsEvolutionGBF.cpp
--- a/Engine/Source/Runtime/Experimental/Chaos/Private/Chaos/PBDRigidsEvolutionGBF.cpp
+++ b/Engine/Source/Runtime/Experimental/Chaos/Private/Chaos/PBDRigidsEvolutionGBF.cpp
@@ -258,8 +258,7 @@ void FPBDRigidsEvolutionGBF::AdvanceOneTimeStep(const FReal Dt, const FReal Step
 #endif
 
 	{
-		// #TODO re-enable when larger groups of collections can successfully resolve groups
-		//Clustering.UnionClusterGroups();
+		Clustering.UnionClusterGroups();
 	}
 
 	{
@@ -469,6 +468,7 @@ FPBDRigidsEvolutionGBF::FPBDRigidsEvolutionGBF(TPBDRigidsSOAs<FReal, 3>& InParti
 
 	AddConstraintRule(&CollisionRule);
 
+	SetInternalParticleInitilizationFunction([](const TGeometryParticleHandle<float, 3>*, const TGeometryParticleHandle<float, 3>*) {});
 	NarrowPhase.GetContext().bFilteringEnabled = true;
 	NarrowPhase.GetContext().bDeferUpdate = true;
 	NarrowPhase.GetContext().bAllowManifolds = false;
diff --git a/Engine/Source/Runtime/Experimental/Chaos/Public/Chaos/PBDRigidsEvolutionGBF.h b/Engine/Source/Runtime/Experimental/Chaos/Public/Chaos/PBDRigidsEvolutionGBF.h
--- a/Engine/Source/Runtime/Experimental/Chaos/Public/Chaos/PBDRigidsEvolutionGBF.h
+++ b/Engine/Source/Runtime/Experimental/Chaos/Public/Chaos/PBDRigidsEvolutionGBF.h
@@ -31,6 +31,10 @@ namespace Chaos
 
 	using FPBDRigidsEvolutionIslandCallback = TFunction<void(int32 Island)>;
 
+	using FPBDRigidsEvolutionInternalHandleCallback = TFunction<void(
+		const TGeometryParticleHandle<float, 3> * OldParticle,
+		const TGeometryParticleHandle<float, 3> * NewParticle)>;
+	
 	class FPBDRigidsEvolutionGBF : public FPBDRigidsEvolutionBase
 	{
 	public:
@@ -80,6 +84,17 @@ namespace Chaos
 			PostApplyPushOutCallback = Cb;
 		}
 
+		void SetInternalParticleInitilizationFunction(const FPBDRigidsEvolutionInternalHandleCallback& Cb)
+		{ 
+			InternalParticleInitilization = Cb;
+		}
+
+		void DoInternalParticleInitilization(const TGeometryParticleHandle<float, 3>* OldParticle, const TGeometryParticleHandle<float, 3>* NewParticle) 
+		{ 
+			if (InternalParticleInitilization) InternalParticleInitilization(OldParticle, NewParticle); 
+		}
+
+
 		CHAOS_API void Advance(const FReal Dt, const FReal MaxStepDt, const int32 MaxSteps);
 		CHAOS_API void AdvanceOneTimeStep(const FReal dt, const FReal StepFraction = (FReal)1.0);
 
@@ -196,5 +211,6 @@ namespace Chaos
 		FPBDRigidsEvolutionCallback PreApplyCallback;
 		FPBDRigidsEvolutionIslandCallback PostApplyCallback;
 		FPBDRigidsEvolutionIslandCallback PostApplyPushOutCallback;
+		FPBDRigidsEvolutionInternalHandleCallback InternalParticleInitilization;
 	};
 }
diff --git a/Engine/Source/Runtime/Experimental/ChaosSolvers/Private/EventDefaults.cpp b/Engine/Source/Runtime/Experimental/ChaosSolvers/Private/EventDefaults.cpp
--- a/Engine/Source/Runtime/Experimental/ChaosSolvers/Private/EventDefaults.cpp
+++ b/Engine/Source/Runtime/Experimental/ChaosSolvers/Private/EventDefaults.cpp
@@ -72,29 +72,33 @@ namespace Chaos
 
 						// Since Clustered GCs can be unioned the particleIndex representing the union 
 						// is not associated with a PhysicsProxy
-						IPhysicsProxyBase* Proxy = Solver->GetProxy(Constraint.Particle[0]->Handle());
-						if (Proxy != nullptr)
-
+						if (const TSet<IPhysicsProxyBase*>* Proxies = Solver->GetProxies(Constraint.Particle[0]->Handle()))
 						{
-							if (ensure(!Constraint.AccumulatedImpulse.ContainsNaN() && FMath::IsFinite(Constraint.GetPhi())))
+							for (IPhysicsProxyBase* Proxy : *Proxies)
 							{
-								TGeometryParticleHandle<float, 3>* Particle0 = Constraint.Particle[0];
-								TGeometryParticleHandle<float, 3>* Particle1 = Constraint.Particle[1];
-								TKinematicGeometryParticleHandle<float, 3>* Body0 = Particle0->CastToKinematicParticle();
-
-								// presently when a rigidbody or kinematic hits static geometry then Body1 is null
-								TKinematicGeometryParticleHandle<float, 3>* Body1 = Particle1->CastToKinematicParticle();
-
-								if (!Constraint.AccumulatedImpulse.IsZero() && Body0)
+								if (Proxy != nullptr)
 								{
-									if (ensure(!Constraint.GetLocation().ContainsNaN() &&
-										!Constraint.GetNormal().ContainsNaN()) &&
-										!Body0->V().ContainsNaN() &&
-										!Body0->W().ContainsNaN() &&
-										(Body1 == nullptr || ((!Body1->V().ContainsNaN()) && !Body1->W().ContainsNaN())))
+									if (ensure(!Constraint.AccumulatedImpulse.ContainsNaN() && FMath::IsFinite(Constraint.GetPhi())))
 									{
-										ValidCollisionHandles[NumValidCollisions] = ContactHandle;
-										NumValidCollisions++;
+										TGeometryParticleHandle<float, 3>* Particle0 = Constraint.Particle[0];
+										TGeometryParticleHandle<float, 3>* Particle1 = Constraint.Particle[1];
+										TKinematicGeometryParticleHandle<float, 3>* Body0 = Particle0->CastToKinematicParticle();
+
+										// presently when a rigidbody or kinematic hits static geometry then Body1 is null
+										TKinematicGeometryParticleHandle<float, 3>* Body1 = Particle1->CastToKinematicParticle();
+
+										if (!Constraint.AccumulatedImpulse.IsZero() && Body0)
+										{
+											if (ensure(!Constraint.GetLocation().ContainsNaN() &&
+												!Constraint.GetNormal().ContainsNaN()) &&
+												!Body0->V().ContainsNaN() &&
+												!Body0->W().ContainsNaN() &&
+												(Body1 == nullptr || ((!Body1->V().ContainsNaN()) && !Body1->W().ContainsNaN())))
+											{
+												ValidCollisionHandles[NumValidCollisions] = ContactHandle;
+												NumValidCollisions++;
+											}
+										}
 									}
 								}
 							}
@@ -106,28 +110,33 @@ namespace Chaos
 
 						// Since Clustered GCs can be unioned the particleIndex representing the union 
 						// is not associated with a PhysicsProxy
-						IPhysicsProxyBase* Proxy = Solver->GetProxy(Constraint.Particle[0]->Handle());
-						if (Proxy != nullptr)
+						if (const TSet<IPhysicsProxyBase*>* Proxies = Solver->GetProxies(Constraint.Particle[0]->Handle()))
 						{
-							if (ensure(!Constraint.AccumulatedImpulse.ContainsNaN() && FMath::IsFinite(Constraint.GetPhi())))
+							for (IPhysicsProxyBase* Proxy : *Proxies)
 							{
-								TGeometryParticleHandle<float, 3>* Particle0 = Constraint.Particle[0];
-								TGeometryParticleHandle<float, 3>* Particle1 = Constraint.Particle[1];
-								TKinematicGeometryParticleHandle<float, 3>* Body0 = Particle0->CastToKinematicParticle();
-
-								// presently when a rigidbody or kinematic hits static geometry then Body1 is null
-								TKinematicGeometryParticleHandle<float, 3>* Body1 = Particle1->CastToKinematicParticle();
-
-								if (!Constraint.AccumulatedImpulse.IsZero() && Body0)
+								if (Proxy != nullptr)
 								{
-									if (ensure(!Constraint.GetLocation().ContainsNaN() &&
-										!Constraint.GetNormal().ContainsNaN()) &&
-										!Body0->V().ContainsNaN() &&
-										!Body0->W().ContainsNaN() &&
-										(Body1 == nullptr || ((!Body1->V().ContainsNaN()) && !Body1->W().ContainsNaN())))
+									if (ensure(!Constraint.AccumulatedImpulse.ContainsNaN() && FMath::IsFinite(Constraint.GetPhi())))
 									{
-										ValidCollisionHandles[NumValidCollisions] = ContactHandle;
-										NumValidCollisions++;
+										TGeometryParticleHandle<float, 3>* Particle0 = Constraint.Particle[0];
+										TGeometryParticleHandle<float, 3>* Particle1 = Constraint.Particle[1];
+										TKinematicGeometryParticleHandle<float, 3>* Body0 = Particle0->CastToKinematicParticle();
+
+										// presently when a rigidbody or kinematic hits static geometry then Body1 is null
+										TKinematicGeometryParticleHandle<float, 3>* Body1 = Particle1->CastToKinematicParticle();
+
+										if (!Constraint.AccumulatedImpulse.IsZero() && Body0)
+										{
+											if (ensure(!Constraint.GetLocation().ContainsNaN() &&
+												!Constraint.GetNormal().ContainsNaN()) &&
+												!Body0->V().ContainsNaN() &&
+												!Body0->W().ContainsNaN() &&
+												(Body1 == nullptr || ((!Body1->V().ContainsNaN()) && !Body1->W().ContainsNaN())))
+											{
+												ValidCollisionHandles[NumValidCollisions] = ContactHandle;
+												NumValidCollisions++;
+											}
+										}
 									}
 								}
 							}
@@ -153,8 +162,10 @@ namespace Chaos
 							Data.AccumulatedImpulse = Constraint.AccumulatedImpulse;
 							Data.Normal = Constraint.GetNormal();
 							Data.PenetrationDepth = Constraint.GetPhi();
-							Data.ParticleProxy = Solver->GetProxy(Particle0->Handle());
-							Data.LevelsetProxy = Solver->GetProxy(Particle1->Handle());
+							Data.ParticleProxy = Solver->GetProxies(Particle0->Handle()) && Solver->GetProxies(Particle0->Handle())->Array().Num() ? 
+								Solver->GetProxies(Particle0->Handle())->Array().operator[](0) : nullptr; // @todo(chaos) : Iterate all proxies
+							Data.LevelsetProxy = Solver->GetProxies(Particle1->Handle()) && Solver->GetProxies(Particle0->Handle())->Array().Num() ? 
+								Solver->GetProxies(Particle1->Handle())->Array().operator[](0) : nullptr; // @todo(chaos) : Iterate all proxies
 
 							// todo: do we need these anymore now we are storing the particles you can access all of this stuff from there
 							// do we still need these now we have pointers to particles returned?
@@ -227,8 +238,10 @@ namespace Chaos
 							Data.AccumulatedImpulse = Constraint.AccumulatedImpulse;
 							Data.Normal = Constraint.GetNormal();
 							Data.PenetrationDepth = Constraint.GetPhi();
-							Data.ParticleProxy = Solver->GetProxy(Particle0->Handle());
-							Data.LevelsetProxy = Solver->GetProxy(Particle1->Handle());
+							Data.ParticleProxy = Solver->GetProxies(Particle0->Handle()) && Solver->GetProxies(Particle0->Handle())->Array().Num() ?
+								Solver->GetProxies(Particle0->Handle())->Array().operator[](0) : nullptr; // @todo(chaos) : Iterate all proxies
+							Data.LevelsetProxy = Solver->GetProxies(Particle1->Handle()) && Solver->GetProxies(Particle0->Handle())->Array().Num() ?
+								Solver->GetProxies(Particle1->Handle())->Array().operator[](0) : nullptr; // @todo(chaos) : Iterate all proxies
 
 							// todo: do we need these anymore now we are storing the particles you can access all of this stuff from there
 							// do we still need these now we have pointers to particles returned?
@@ -444,8 +457,16 @@ namespace Chaos
 #endif
 
 								// Add to AllTrailingsIndicesByPhysicsProxy
-								IPhysicsProxyBase* PhysicsProxy = Solver->GetProxy(ActiveParticle.Handle());
-								AllTrailingsIndicesByPhysicsProxy.FindOrAdd(PhysicsProxy).Add(NewIdx);
+								if (const TSet<IPhysicsProxyBase*>* Proxies = Solver->GetProxies(ActiveParticle.Handle()))
+								{
+									for (IPhysicsProxyBase* Proxy : *Proxies)
+									{
+										if (Proxy != nullptr)
+										{
+											AllTrailingsIndicesByPhysicsProxy.FindOrAdd(Proxy).Add(NewIdx);
+										}
+									}
+								}
 							}
 						}
 					}
@@ -477,17 +498,21 @@ namespace Chaos
 			{
 				if(SleepData.Particle)
 				{
-					IPhysicsProxyBase* PhysicsProxy = NonConstSolver->GetProxy(SleepData.Particle);
-
-					TGeometryParticle<float, 3>* Particle = SleepData.Particle->GTGeometryParticle();
-					if(Particle != nullptr && PhysicsProxy != nullptr)
+					if (const TSet<IPhysicsProxyBase*>* Proxies = Solver->GetProxies(SleepData.Particle))
 					{
-						int32 NewIdx = EventSleepDataArray.Add(TSleepingData<float, 3>());
-						TSleepingData<float, 3>& SleepingDataArrayItem = EventSleepDataArray[NewIdx];
-						SleepingDataArrayItem.Particle = Particle;
-						SleepingDataArrayItem.Sleeping = SleepData.Sleeping;
+						for (IPhysicsProxyBase* Proxy : *Proxies)
+						{
+							TGeometryParticle<float, 3>* Particle = SleepData.Particle->GTGeometryParticle();
+							if (Particle != nullptr && Proxy != nullptr)
+							{
+								int32 NewIdx = EventSleepDataArray.Add(TSleepingData<float, 3>());
+								TSleepingData<float, 3>& SleepingDataArrayItem = EventSleepDataArray[NewIdx];
+								SleepingDataArrayItem.Particle = Particle;
+								SleepingDataArrayItem.Sleeping = SleepData.Sleeping;
 
-						AllSleepIndicesByPhysicsProxy.FindOrAdd(PhysicsProxy).Add(NewIdx);
+								AllSleepIndicesByPhysicsProxy.FindOrAdd(Proxy).Add(NewIdx);
+							}
+						}
 					}
 				}
 			}
diff --git a/Engine/Source/Runtime/Experimental/ChaosSolvers/Private/PBDRigidActiveParticlesBuffer.cpp b/Engine/Source/Runtime/Experimental/ChaosSolvers/Private/PBDRigidActiveParticlesBuffer.cpp
--- a/Engine/Source/Runtime/Experimental/ChaosSolvers/Private/PBDRigidActiveParticlesBuffer.cpp
+++ b/Engine/Source/Runtime/Experimental/ChaosSolvers/Private/PBDRigidActiveParticlesBuffer.cpp
@@ -44,9 +44,15 @@ namespace Chaos
 				{
 					ActiveGameThreadParticles.Add(Handle);
 				}
-				else if(IPhysicsProxyBase** Proxy = Solver->MParticleToProxy.Find(ActiveParticle.Handle()))
+				else if (const TSet<IPhysicsProxyBase*> * Proxies = Solver->GetProxies(ActiveParticle.Handle()))
 				{
-					PhysicsParticleProxies.Add(*Proxy);
+					for (IPhysicsProxyBase* Proxy : *Proxies)
+					{
+						if (Proxy != nullptr)
+						{
+							PhysicsParticleProxies.Add(Proxy);
+						}
+					}
 				}
 			}
 		}
diff --git a/Engine/Source/Runtime/Experimental/ChaosSolvers/Private/PBDRigidsSolver.cpp b/Engine/Source/Runtime/Experimental/ChaosSolvers/Private/PBDRigidsSolver.cpp
--- a/Engine/Source/Runtime/Experimental/ChaosSolvers/Private/PBDRigidsSolver.cpp
+++ b/Engine/Source/Runtime/Experimental/ChaosSolvers/Private/PBDRigidsSolver.cpp
@@ -218,6 +218,17 @@ namespace Chaos
 	{
 		UE_LOG(LogPBDRigidsSolver, Verbose, TEXT("PBDRigidsSolver::PBDRigidsSolver()"));
 		Reset();
+
+		MEvolution->SetInternalParticleInitilizationFunction(
+			[this](const Chaos::TGeometryParticleHandle<float, 3>* OldParticle, const Chaos::TGeometryParticleHandle<float, 3>* NewParticle) {
+				if (const TSet<IPhysicsProxyBase*>* Proxies = GetProxies(OldParticle))
+				{
+					for (IPhysicsProxyBase* Proxy : *Proxies)
+					{
+						this->AddParticleToProxy(NewParticle, Proxy);
+					}
+				}
+			});
 	}
 
 	FPBDRigidsSolver::~FPBDRigidsSolver()
@@ -817,7 +828,7 @@ namespace Chaos
 				{
 					auto Handle = Proxy->GetHandle();
 					Handle->GTGeometryParticle() = Proxy->GetParticle();
-					Solver->MParticleToProxy.Add(Handle,Proxy);
+					Solver->AddParticleToProxy(Handle,Proxy);
 					Solver->GetEvolution()->CreateParticle(Handle);
 					Proxy->SetInitialized(true);
 				}
@@ -871,27 +882,33 @@ namespace Chaos
 		TParticleView<TPBDRigidParticles<float, 3>>& ActiveParticles = GetParticles().GetActiveParticlesView();
 		for (Chaos::TPBDRigidParticleHandleImp<float, 3, false>& ActiveObject : ActiveParticles)
 		{
-			if (IPhysicsProxyBase* Proxy = GetProxy(ActiveObject.Handle()))
+			if( const TSet<IPhysicsProxyBase*>* Proxies = GetProxies(ActiveObject.Handle()))
 			{
-				switch (ActiveObject.GetParticleType())
+				for (IPhysicsProxyBase* Proxy : *Proxies)
 				{
-				case Chaos::EParticleType::Rigid:
-					((FRigidParticlePhysicsProxy*)(Proxy))->BufferPhysicsResults();
-					break;
-				case Chaos::EParticleType::Kinematic:
-					((FKinematicGeometryParticlePhysicsProxy*)(Proxy))->BufferPhysicsResults();
-					break;
-				case Chaos::EParticleType::Static:
-					((FGeometryParticlePhysicsProxy*)(Proxy))->BufferPhysicsResults();
-					break;
-				case Chaos::EParticleType::GeometryCollection:
-					ActiveGC.AddUnique((FGeometryCollectionPhysicsProxy*)(Proxy));
-					break;
-				case Chaos::EParticleType::Clustered:
-					ActiveGC.AddUnique((FGeometryCollectionPhysicsProxy*)(Proxy));
-					break;
-				default:
-					check(false);
+					if (Proxy != nullptr)
+					{
+						switch (ActiveObject.GetParticleType())
+						{
+						case Chaos::EParticleType::Rigid:
+							((FRigidParticlePhysicsProxy*)(Proxy))->BufferPhysicsResults();
+							break;
+						case Chaos::EParticleType::Kinematic:
+							((FKinematicGeometryParticlePhysicsProxy*)(Proxy))->BufferPhysicsResults();
+							break;
+						case Chaos::EParticleType::Static:
+							((FGeometryParticlePhysicsProxy*)(Proxy))->BufferPhysicsResults();
+							break;
+						case Chaos::EParticleType::GeometryCollection:
+							ActiveGC.AddUnique((FGeometryCollectionPhysicsProxy*)(Proxy));
+							break;
+						case Chaos::EParticleType::Clustered:
+							ActiveGC.AddUnique((FGeometryCollectionPhysicsProxy*)(Proxy));
+							break;
+						default:
+							check(false);
+						}
+					}
 				}
 			}
 		}
@@ -911,27 +928,33 @@ namespace Chaos
 		TParticleView<TPBDRigidParticles<float, 3>>& ActiveParticles = GetParticles().GetActiveParticlesView();
 		for (Chaos::TPBDRigidParticleHandleImp<float, 3, false>& ActiveObject : ActiveParticles)
 		{
-			if (IPhysicsProxyBase* Proxy = GetProxy(ActiveObject.Handle()))
+			if (const TSet<IPhysicsProxyBase*>* Proxies = GetProxies(ActiveObject.Handle()))
 			{
-				switch (ActiveObject.GetParticleType())
+				for (IPhysicsProxyBase* Proxy : *Proxies)
 				{
-				case Chaos::EParticleType::Rigid:
-					((FRigidParticlePhysicsProxy*)(Proxy))->FlipBuffer();
-					break;
-				case Chaos::EParticleType::Kinematic:
-					((FKinematicGeometryParticlePhysicsProxy*)(Proxy))->FlipBuffer();
-					break;
-				case Chaos::EParticleType::Static:
-					((FGeometryParticlePhysicsProxy*)(Proxy))->FlipBuffer();
-					break;
-				case Chaos::EParticleType::GeometryCollection:
-					ActiveGC.AddUnique((FGeometryCollectionPhysicsProxy*)(Proxy));
-					break;
-				case Chaos::EParticleType::Clustered:
-					ActiveGC.AddUnique((FGeometryCollectionPhysicsProxy*)(Proxy));
-					break;
-				default:
-					check(false);
+					if (Proxy != nullptr)
+					{
+						switch (ActiveObject.GetParticleType())
+						{
+						case Chaos::EParticleType::Rigid:
+							((FRigidParticlePhysicsProxy*)(Proxy))->FlipBuffer();
+							break;
+						case Chaos::EParticleType::Kinematic:
+							((FKinematicGeometryParticlePhysicsProxy*)(Proxy))->FlipBuffer();
+							break;
+						case Chaos::EParticleType::Static:
+							((FGeometryParticlePhysicsProxy*)(Proxy))->FlipBuffer();
+							break;
+						case Chaos::EParticleType::GeometryCollection:
+							ActiveGC.AddUnique((FGeometryCollectionPhysicsProxy*)(Proxy));
+							break;
+						case Chaos::EParticleType::Clustered:
+							ActiveGC.AddUnique((FGeometryCollectionPhysicsProxy*)(Proxy));
+							break;
+						default:
+							check(false);
+						}
+					}
 				}
 			}
 		}
@@ -956,27 +979,33 @@ namespace Chaos
 		TParticleView<TPBDRigidParticles<float, 3>>& ActiveParticles = GetParticles().GetActiveParticlesView();
 		for (Chaos::TPBDRigidParticleHandleImp<float, 3, false>& ActiveObject : ActiveParticles)
 		{
-			if (IPhysicsProxyBase* Proxy = GetProxy(ActiveObject.Handle()))
+			if (const TSet<IPhysicsProxyBase*>* Proxies = GetProxies(ActiveObject.Handle()))
 			{
-				switch (ActiveObject.GetParticleType())
+				for (IPhysicsProxyBase* Proxy : *Proxies)
 				{
-				case Chaos::EParticleType::Rigid:
-					((FRigidParticlePhysicsProxy*)(Proxy))->PullFromPhysicsState();
-					break;
-				case Chaos::EParticleType::Kinematic:
-					((FKinematicGeometryParticlePhysicsProxy*)(Proxy))->PullFromPhysicsState();
-					break;
-				case Chaos::EParticleType::Static:
-					((FGeometryParticlePhysicsProxy*)(Proxy))->PullFromPhysicsState();
-					break;
-				case Chaos::EParticleType::GeometryCollection:
-					ActiveGC.AddUnique((FGeometryCollectionPhysicsProxy*)(Proxy));
-					break;
-				case Chaos::EParticleType::Clustered:
-					ActiveGC.AddUnique((FGeometryCollectionPhysicsProxy*)(Proxy));
-					break;
-				default:
-					check(false);
+					if (Proxy != nullptr)
+					{
+						switch (ActiveObject.GetParticleType())
+						{
+						case Chaos::EParticleType::Rigid:
+							((FRigidParticlePhysicsProxy*)(Proxy))->PullFromPhysicsState();
+							break;
+						case Chaos::EParticleType::Kinematic:
+							((FKinematicGeometryParticlePhysicsProxy*)(Proxy))->PullFromPhysicsState();
+							break;
+						case Chaos::EParticleType::Static:
+							((FGeometryParticlePhysicsProxy*)(Proxy))->PullFromPhysicsState();
+							break;
+						case Chaos::EParticleType::GeometryCollection:
+							ActiveGC.AddUnique((FGeometryCollectionPhysicsProxy*)(Proxy));
+							break;
+						case Chaos::EParticleType::Clustered:
+							ActiveGC.AddUnique((FGeometryCollectionPhysicsProxy*)(Proxy));
+							break;
+						default:
+							check(false);
+						}
+					}
 				}
 			}
 		}
@@ -1053,18 +1082,22 @@ namespace Chaos
 			int32 DataIdx = 0;
 			for(TPBDRigidParticleHandleImp<float,3,false>& ActiveObject : ActiveParticles)
 			{
-				const IPhysicsProxyBase* Proxy = GetProxy(ActiveObject.Handle());
+				if (const TSet<IPhysicsProxyBase*>* Proxies = GetProxies(ActiveObject.Handle()))
 				{
-					if(ActiveObject.GetParticleType() == EParticleType::Rigid)
+					for (IPhysicsProxyBase* Proxy : *Proxies)
 					{
-						//may want to remove branch using templates outside loop
-						if(MRewindData->IsResim())
-						{
-							MRewindData->PushPTDirtyData<true>(*static_cast<const FRigidParticlePhysicsProxy*>(Proxy)->GetHandle(),DataIdx++);
-						}
-						else
+
+						if (ActiveObject.GetParticleType() == EParticleType::Rigid)
 						{
-							MRewindData->PushPTDirtyData<false>(*static_cast<const FRigidParticlePhysicsProxy*>(Proxy)->GetHandle(),DataIdx++);
+							//may want to remove branch using templates outside loop
+							if (MRewindData->IsResim())
+							{
+								MRewindData->PushPTDirtyData<true>(*static_cast<const FRigidParticlePhysicsProxy*>(Proxy)->GetHandle(), DataIdx++);
+							}
+							else
+							{
+								MRewindData->PushPTDirtyData<false>(*static_cast<const FRigidParticlePhysicsProxy*>(Proxy)->GetHandle(), DataIdx++);
+							}
 						}
 					}
 				}
diff --git a/Engine/Source/Runtime/Experimental/ChaosSolvers/Private/PhysicsProxy/FieldSystemPhysicsProxy.cpp b/Engine/Source/Runtime/Experimental/ChaosSolvers/Private/PhysicsProxy/FieldSystemPhysicsProxy.cpp
--- a/Engine/Source/Runtime/Experimental/ChaosSolvers/Private/PhysicsProxy/FieldSystemPhysicsProxy.cpp
+++ b/Engine/Source/Runtime/Experimental/ChaosSolvers/Private/PhysicsProxy/FieldSystemPhysicsProxy.cpp
@@ -386,12 +386,13 @@ void FFieldSystemPhysicsProxy::FieldParameterUpdateCallback(
 
 								for(TPBDRigidParticleHandle<FReal, 3>* Handle : Activated)
 								{
-									if(!CurrentSolver->GetProxy(Handle))
+									if(!CurrentSolver->GetProxies(Handle))
 									{
-										IPhysicsProxyBase* ParentProxy = CurrentSolver->GetProxy(Iter.Key);
-										if(ensure(ParentProxy))
+										const TSet<IPhysicsProxyBase*> * ParentProxies = CurrentSolver->GetProxies(Iter.Key);
+										if(ensure(ParentProxies))
 										{
-											CurrentSolver->AddParticleToProxy(Handle, ParentProxy);
+											for( IPhysicsProxyBase* ParentProxy : *ParentProxies)
+												CurrentSolver->AddParticleToProxy(Handle, ParentProxy);
 										}
 									}
 								}
diff --git a/Engine/Source/Runtime/Experimental/ChaosSolvers/Public/PBDRigidsSolver.h b/Engine/Source/Runtime/Experimental/ChaosSolvers/Public/PBDRigidsSolver.h
--- a/Engine/Source/Runtime/Experimental/ChaosSolvers/Public/PBDRigidsSolver.h
+++ b/Engine/Source/Runtime/Experimental/ChaosSolvers/Public/PBDRigidsSolver.h
@@ -284,7 +284,9 @@ namespace Chaos
 
 		void AddParticleToProxy(const Chaos::TGeometryParticleHandle<float, 3>* Particle, IPhysicsProxyBase* Proxy)
 		{
-			MParticleToProxy.Add(Particle, Proxy); 
+			if (!MParticleToProxy.Find(Particle))
+				MParticleToProxy.Add(Particle, TSet<IPhysicsProxyBase*>());
+			MParticleToProxy[Particle].Add(Proxy); 
 		}
 		
 		void RemoveParticleToProxy(const Chaos::TGeometryParticleHandle<float, 3>* Particle)
@@ -292,10 +294,10 @@ namespace Chaos
 			MParticleToProxy.Remove(Particle);
 		}
 		
-		IPhysicsProxyBase* GetProxy(const Chaos::TGeometryParticleHandle<float, 3>* Handle) const
+		const TSet<IPhysicsProxyBase*> * GetProxies(const Chaos::TGeometryParticleHandle<float, 3>* Handle) const
 		{
-			IPhysicsProxyBase* const* PhysicsProxyPtr = MParticleToProxy.Find(Handle);
-			return PhysicsProxyPtr ? *PhysicsProxyPtr : nullptr;
+			const TSet<IPhysicsProxyBase*>* PhysicsProxyPtr = MParticleToProxy.Find(Handle);
+			return PhysicsProxyPtr ? PhysicsProxyPtr : nullptr;
 		}
 
 		/**/
@@ -379,7 +381,7 @@ namespace Chaos
 		TUniquePtr<FEventManager> MEventManager;
 		TUniquePtr<FSolverEventFilters> MSolverEventFilters;
 		TUniquePtr<FActiveParticlesBuffer> MActiveParticlesBuffer;
-		TMap<const Chaos::TGeometryParticleHandle<float, 3>*, IPhysicsProxyBase*> MParticleToProxy;
+		TMap<const Chaos::TGeometryParticleHandle<float, 3>*, TSet<IPhysicsProxyBase*> > MParticleToProxy;
 		TUniquePtr<FRewindData> MRewindData;
 
 		//
