Fix UGS/CIS

diff --git a/Engine/Source/Runtime/Experimental/ChaosSolvers/Public/PBDRigidsSolver.h b/Engine/Source/Runtime/Experimental/ChaosSolvers/Public/PBDRigidsSolver.h
--- a/Engine/Source/Runtime/Experimental/ChaosSolvers/Public/PBDRigidsSolver.h
+++ b/Engine/Source/Runtime/Experimental/ChaosSolvers/Public/PBDRigidsSolver.h
@@ -347,19 +347,37 @@ namespace Chaos
 		template<typename ParticleType>
 		void FlipBuffer(Chaos::TGeometryParticleHandle<float, 3>* Handle)
 		{
-			((ParticleType*)(GetProxy(Handle)))->FlipBuffer();
+			if (const TSet<IPhysicsProxyBase*>* Proxies = Solver->GetProxies(Handle))
+			{
+				for (IPhysicsProxyBase* Proxy : *Proxies)
+				{
+					((ParticleType*)(Proxy))->FlipBuffer();
+				}
+			}
 		}
 
 		template<typename ParticleType>
 		void PullFromPhysicsState(Chaos::TGeometryParticleHandle<float, 3>* Handle)
 		{
-			((ParticleType*)(GetProxy(Handle)))->PullFromPhysicsState();
+			if (const TSet<IPhysicsProxyBase*>* Proxies = Solver->GetProxies(Handle))
+			{
+				for (IPhysicsProxyBase* Proxy : *Proxies)
+				{
+					((ParticleType*)(Proxy))->PullFromPhysicsState();
+				}
+			}
 		}
 
 		template<typename ParticleType>
 		void BufferPhysicsResults(Chaos::TGeometryParticleHandle<float, 3>* Handle)
 		{
-			((ParticleType*)(GetProxy(Handle)))->BufferPhysicsResults();
+			if (const TSet<IPhysicsProxyBase*>* Proxies = Solver->GetProxies(Handle))
+			{
+				for (IPhysicsProxyBase* Proxy : *Proxies)
+				{
+					((ParticleType*)(Proxy))->BufferPhysicsResults();
+				}
+			}
 		}
 
 		//
