Chaos : Fix
- Removed reference to solver.

diff --git a/Engine/Source/Runtime/Experimental/ChaosSolvers/Public/PBDRigidsSolver.h b/Engine/Source/Runtime/Experimental/ChaosSolvers/Public/PBDRigidsSolver.h
--- a/Engine/Source/Runtime/Experimental/ChaosSolvers/Public/PBDRigidsSolver.h
+++ b/Engine/Source/Runtime/Experimental/ChaosSolvers/Public/PBDRigidsSolver.h
@@ -347,7 +347,7 @@ namespace Chaos
 		template<typename ParticleType>
 		void FlipBuffer(Chaos::TGeometryParticleHandle<float, 3>* Handle)
 		{
-			if (const TSet<IPhysicsProxyBase*>* Proxies = Solver->GetProxies(Handle))
+			if (const TSet<IPhysicsProxyBase*>* Proxies = GetProxies(Handle))
 			{
 				for (IPhysicsProxyBase* Proxy : *Proxies)
 				{
@@ -359,7 +359,7 @@ namespace Chaos
 		template<typename ParticleType>
 		void PullFromPhysicsState(Chaos::TGeometryParticleHandle<float, 3>* Handle)
 		{
-			if (const TSet<IPhysicsProxyBase*>* Proxies = Solver->GetProxies(Handle))
+			if (const TSet<IPhysicsProxyBase*>* Proxies = GetProxies(Handle))
 			{
 				for (IPhysicsProxyBase* Proxy : *Proxies)
 				{
@@ -371,7 +371,7 @@ namespace Chaos
 		template<typename ParticleType>
 		void BufferPhysicsResults(Chaos::TGeometryParticleHandle<float, 3>* Handle)
 		{
-			if (const TSet<IPhysicsProxyBase*>* Proxies = Solver->GetProxies(Handle))
+			if (const TSet<IPhysicsProxyBase*>* Proxies = GetProxies(Handle))
 			{
 				for (IPhysicsProxyBase* Proxy : *Proxies)
 				{
