Fix Geometry Collection cluster body implicit sphere radius.  Fall back on sphere for cluster body implicit, if level set rasterization fails.  Don't override particle handle UniqueId.Idx for tracking transform group index, use a map instead.

diff --git a/Engine/Source/Runtime/Experimental/ChaosSolvers/Private/PhysicsProxy/GeometryCollectionPhysicsProxy.cpp b/Engine/Source/Runtime/Experimental/ChaosSolvers/Private/PhysicsProxy/GeometryCollectionPhysicsProxy.cpp
--- a/Engine/Source/Runtime/Experimental/ChaosSolvers/Private/PhysicsProxy/GeometryCollectionPhysicsProxy.cpp
+++ b/Engine/Source/Runtime/Experimental/ChaosSolvers/Private/PhysicsProxy/GeometryCollectionPhysicsProxy.cpp
@@ -1878,7 +1878,6 @@ void FGeometryCollectionPhysicsProxy::InitializeSharedCollisionStructures(
 	TArray<TUniquePtr<TTriangleMesh<float>>> TriangleMeshesArray;	//use to union trimeshes in cluster case
 	TriangleMeshesArray.AddDefaulted(NumTransforms);
 
-
 	TParticles<float, 3> MassSpaceParticles;
 	MassSpaceParticles.AddParticles(Vertex.Num());
 	for (int32 Idx = 0; Idx < Vertex.Num(); ++Idx)
@@ -1892,7 +1891,6 @@ void FGeometryCollectionPhysicsProxy::InitializeSharedCollisionStructures(
 	TArray<bool> InertiaComputationNeeded;
 	InertiaComputationNeeded.Init(false, NumGeometries);
 
-
 	float TotalVolume = 0.f;
 	// The geometry group has a set of transform indices that maps a geometry index
 	// to a transform index, but only in the case where there is a 1-to-1 mapping 
@@ -2099,8 +2097,12 @@ void FGeometryCollectionPhysicsProxy::InitializeSharedCollisionStructures(
 						SizeSpecificData.CollisionType,
 						*TriMesh,
 						SizeSpecificData.CollisionParticlesFraction);
-				CollectionSimplicials[TransformGroupIndex] = TUniquePtr<FSimplicial>(Simplicial);
-				ensureMsgf(CollectionSimplicials[TransformGroupIndex], TEXT("No simplicial representation."));
+				CollectionSimplicials[TransformGroupIndex] = TUniquePtr<FSimplicial>(Simplicial); // CollectionSimplicials is in the TransformGroup
+				//ensureMsgf(CollectionSimplicials[TransformGroupIndex], TEXT("No simplicial representation."));
+				if (!CollectionSimplicials[TransformGroupIndex]->Size())
+				{
+					ensureMsgf(false, TEXT("Simplicial is empty."));
+				}
 
 				if (SizeSpecificData.ImplicitType == EImplicitTypeEnum::Chaos_Implicit_LevelSet)
 				{
@@ -2177,15 +2179,22 @@ void FGeometryCollectionPhysicsProxy::InitializeSharedCollisionStructures(
 			CollectionSpaceParticles->X(Idx) = Chaos::TVector<float, 3>(-TNumericLimits<float>::Max());
 		}
 
+		//
+		// TODO: We generate particles & handles for leaf nodes so that we can use some 
+		// runtime clustering functions.  That's adding a lot of work and dependencies
+		// just so we can make an API happy.  We should refactor the common routines
+		// to have a handle agnostic implementation.
+		//
+
+		TMap<const TGeometryParticleHandle<float, 3>*, int32> HandleToTransformIdx;
 		TArray<TUniquePtr<TPBDRigidClusteredParticleHandle<float, 3>>> Handles;
 		Handles.Reserve(NumTransforms);
 		for (int32 Idx = 0; Idx < NumTransforms; Idx++)
 		{
 			Handles.Add(TPBDRigidClusteredParticleHandle<float, 3>::CreateParticleHandle(
 				MakeSerializable(CollectionSpaceParticles), Idx, Idx));
-			Handles[Idx]->SetUniqueIdx(FUniqueIdx(Idx));
+			HandleToTransformIdx.Add(Handles[Handles.Num() - 1].Get(), Idx);
 		}
-
 		for (int32 GeometryIdx = 0; GeometryIdx < NumGeometries; ++GeometryIdx)
 		{
 			const int32 TransformGroupIndex = TransformIndex[GeometryIdx];
@@ -2195,7 +2204,7 @@ void FGeometryCollectionPhysicsProxy::InitializeSharedCollisionStructures(
 					Handles[TransformGroupIndex].Get(),
 					SharedParams, 
 					CollectionSimplicials[TransformGroupIndex].Get(),
-					CollectionImplicits[TransformGroupIndex], 
+					CollectionImplicits[TransformGroupIndex],
 					FCollisionFilterData(),		// SimFilter
 					FCollisionFilterData(),		// QueryFilter
 					CollectionMass[TransformGroupIndex],
@@ -2258,7 +2267,7 @@ void FGeometryCollectionPhysicsProxy::InitializeSharedCollisionStructures(
 					int32 NumChildIndices = 0;
 					for (TPBDRigidParticleHandle<float, 3>* Child : ChildrenIndices)
 					{
-						const int32 ChildTransformIdx = Child->UniqueIdx().Idx;
+						const int32 ChildTransformIdx = HandleToTransformIdx[Child];
 						if (Chaos::TTriangleMesh<float>* ChildMesh = TriangleMeshesArray[ChildTransformIdx].Get())
 						{
 							BiggestNumElements = FMath::Max(BiggestNumElements, ChildMesh->GetNumElements());
@@ -2274,7 +2283,7 @@ void FGeometryCollectionPhysicsProxy::InitializeSharedCollisionStructures(
 					VertsAdded.Reserve(BiggestNumElements);
 					for (TPBDRigidParticleHandle<float, 3>* Child : ChildrenIndices)
 					{
-						const int32 ChildTransformIdx = Child->UniqueIdx().Idx;
+						const int32 ChildTransformIdx = HandleToTransformIdx[Child];
 						if (Chaos::TTriangleMesh<float>* ChildMesh = TriangleMeshesArray[ChildTransformIdx].Get())
 						{
 							const TArray<TVector<int32, 3>>& ChildIndices = ChildMesh->GetSurfaceElements();
@@ -2306,7 +2315,7 @@ void FGeometryCollectionPhysicsProxy::InitializeSharedCollisionStructures(
 
 				if (SizeSpecificData.ImplicitType == EImplicitTypeEnum::Chaos_Implicit_LevelSet)
 				{
-					const TVector<float, 3> Scale = 2 * InstanceBoundingBox.GetExtent() / MaxChildBounds;
+					const TVector<float, 3> Scale = 2 * InstanceBoundingBox.GetExtent() / MaxChildBounds; // FBox's extents are 1/2 (Max - Min)
 					const float ScaleMax = Scale.GetAbsMax();
 					const float ScaleMin = Scale.GetAbsMin();
 
@@ -2328,8 +2337,17 @@ void FGeometryCollectionPhysicsProxy::InitializeSharedCollisionStructures(
 							MaxResolution,
 							SizeSpecificData.CollisionObjectReductionPercentage,
 							SizeSpecificData.CollisionType));
+					// Fall back on sphere if level set rasterization failed.
+					if (!CollectionImplicits[ClusterTransformIdx])
+					{
+						CollectionImplicits[ClusterTransformIdx] = FGeometryDynamicCollection::FSharedImplicit(
+							FCollisionStructureManager::NewImplicitSphere(
+								InstanceBoundingBox.GetExtent().GetAbsMin(), // FBox's extents are 1/2 (Max - Min)
+								SizeSpecificData.CollisionObjectReductionPercentage,
+								SizeSpecificData.CollisionType));
+					}
 
-					CollectionSimplicials[ClusterTransformIdx] = TUniquePtr< FSimplicial >(
+					CollectionSimplicials[ClusterTransformIdx] = TUniquePtr<FSimplicial>(
 						FCollisionStructureManager::NewSimplicial(MassSpaceParticles, *UnionMesh, CollectionImplicits[ClusterTransformIdx].Get(),
 						SharedParams.MaximumCollisionParticleCount));
 				}
@@ -2341,7 +2359,7 @@ void FGeometryCollectionPhysicsProxy::InitializeSharedCollisionStructures(
 							SizeSpecificData.CollisionObjectReductionPercentage,
 							SizeSpecificData.CollisionType));
 
-					CollectionSimplicials[ClusterTransformIdx] = TUniquePtr< FSimplicial >(
+					CollectionSimplicials[ClusterTransformIdx] = TUniquePtr<FSimplicial>(
 						FCollisionStructureManager::NewSimplicial(MassSpaceParticles, *UnionMesh, CollectionImplicits[ClusterTransformIdx].Get(),
 						SharedParams.MaximumCollisionParticleCount));
 				}
@@ -2349,11 +2367,11 @@ void FGeometryCollectionPhysicsProxy::InitializeSharedCollisionStructures(
 				{
  					CollectionImplicits[ClusterTransformIdx] = FGeometryDynamicCollection::FSharedImplicit(
 						FCollisionStructureManager::NewImplicitSphere(
-							InstanceBoundingBox.GetExtent().GetAbsMin() / 2,
+							InstanceBoundingBox.GetExtent().GetAbsMin(), // FBox's extents are 1/2 (Max - Min)
 							SizeSpecificData.CollisionObjectReductionPercentage,
 							SizeSpecificData.CollisionType));
 
-					CollectionSimplicials[ClusterTransformIdx] = TUniquePtr< FSimplicial >(
+					CollectionSimplicials[ClusterTransformIdx] = TUniquePtr<FSimplicial>(
 						FCollisionStructureManager::NewSimplicial(MassSpaceParticles, *UnionMesh, CollectionImplicits[ClusterTransformIdx].Get(),
 						SharedParams.MaximumCollisionParticleCount));
 				}
