Chaos: Crash bug fix, bad cast to wrong type

diff --git a/Engine/Source/Runtime/Experimental/ChaosSolvers/Private/PBDRigidsSolver.cpp b/Engine/Source/Runtime/Experimental/ChaosSolvers/Private/PBDRigidsSolver.cpp
--- a/Engine/Source/Runtime/Experimental/ChaosSolvers/Private/PBDRigidsSolver.cpp
+++ b/Engine/Source/Runtime/Experimental/ChaosSolvers/Private/PBDRigidsSolver.cpp
@@ -75,6 +75,8 @@ namespace Chaos
 			LLM_SCOPE(ELLMTag::Chaos);
 			UE_LOG(LogPBDRigidsSolver, Verbose, TEXT("AdvanceOneTimeStepTask::DoWork()"));
 
+			MSolver->GetEvolution()->GetRigidClustering().ResetAllClusterBreakings();
+
 			{
 				SCOPE_CYCLE_COUNTER(STAT_UpdateParams);
 				Chaos::TPBDPositionConstraints<float, 3> PositionTarget; // Dummy for now
@@ -374,7 +376,7 @@ namespace Chaos
 			UE_LOG(LogPBDRigidsSolver, Verbose, TEXT("FPBDRigidsSolver::UnregisterObject() ~ Dequeue"));
 
 				// Generally need to remove stale events for particles that no longer exist
-				Solver->GetEventManager()->ClearEvents<FCollisionEventData>([InProxy]
+				Solver->GetEventManager()->ClearEvents<FCollisionEventData>(EEventType::Collision, [InProxy]
 				(FCollisionEventData& EventDataInOut)
 				{
 					Chaos::FCollisionDataArray const& CollisionData = EventDataInOut.CollisionData.AllCollisionsArray;
diff --git a/Engine/Source/Runtime/Experimental/ChaosSolvers/Public/EventManager.h b/Engine/Source/Runtime/Experimental/ChaosSolvers/Public/EventManager.h
--- a/Engine/Source/Runtime/Experimental/ChaosSolvers/Public/EventManager.h
+++ b/Engine/Source/Runtime/Experimental/ChaosSolvers/Public/EventManager.h
@@ -280,16 +280,11 @@ namespace Chaos
 		 * Modify the producer side of the event buffer
 		 */
 		template<typename PayloadType>
-		void ClearEvents(TFunction<void(PayloadType & EventData)> InFunction)
+		void ClearEvents(const FEventID& EventID, TFunction<void(PayloadType & EventData)> InFunction)
 		{
 			ContainerLock.ReadLock();
-			for (FEventContainerPtr EventContainer : EventContainers)
-			{
-				if (EventContainer)
-				{
-					((TEventContainer<PayloadType>*)(EventContainer))->DestroyStaleEvents(InFunction);
-				}
-			}
+
+			((TEventContainer<PayloadType>*)(EventContainers[EventID]))->DestroyStaleEvents(InFunction);
 			ContainerLock.ReadUnlock();
 		}
 
@@ -346,4 +341,4 @@ namespace Chaos
 
 	};
 
-}
\ No newline at end of file
+}
