Fix CIS

diff --git a/Engine/Source/Runtime/Experimental/Chaos/Public/Chaos/ExternalCollisionData.h b/Engine/Source/Runtime/Experimental/Chaos/Public/Chaos/ExternalCollisionData.h
--- a/Engine/Source/Runtime/Experimental/Chaos/Public/Chaos/ExternalCollisionData.h
+++ b/Engine/Source/Runtime/Experimental/Chaos/Public/Chaos/ExternalCollisionData.h
@@ -34,7 +34,9 @@ namespace Chaos
 		    , LevelsetProxy(nullptr)
 		{}
 
-		TCollisionData(TVector<T, d> InLocation, TVector<T, d> InAccumulatedImpulse, TVector<T, d> InNormal, TVector<T, d> InVelocity1, TVector<T, d> InVelocity2, TVector<T, d> InAngularVelocity1, TVector<T, d> InAngularVelocity2, T InMass1, T InMass2, T InPenetrationDepth, TGeometryParticle<T, d>* InParticle, TGeometryParticle<T, d>* InLevelset, IPhysicsProxyBase* InParticleProxy, IPhysicsProxyBase* InLevelsetProxy)
+		TCollisionData(TVector<T, d> InLocation, TVector<T, d> InAccumulatedImpulse, TVector<T, d> InNormal, TVector<T, d> InVelocity1, TVector<T, d> InVelocity2
+			, TVector<T, d> InAngularVelocity1, TVector<T, d> InAngularVelocity2, T InMass1, T InMass2, T InPenetrationDepth, TGeometryParticle<T, d>* InParticle
+			, TGeometryParticle<T, d>* InLevelset, IPhysicsProxyBase* InParticleProxy, IPhysicsProxyBase* InLevelsetProxy)
 			: Location(InLocation)
 			, AccumulatedImpulse(InAccumulatedImpulse)
 			, Normal(InNormal)
@@ -93,7 +95,10 @@ namespace Chaos
 		{}
 
 		TCollisionDataExt(
-		    TVector<T, d> InLocation, TVector<T, d> InAccumulatedImpulse, TVector<T, d> InNormal, TVector<T, d> InVelocity1, TVector<T, d> InVelocity2, TVector<T, d> InAngularVelocity1, TVector<T, d> InAngularVelocity2, T InMass1, T InMass2, TGeometryParticle<T, d>* InParticle, TGeometryParticle<T, d>* InLevelset, IPhysicsProxyBase* ParticleProxy, IPhysicsProxyBase* LevelSetProxy, float InBoundingboxVolume, float InBoundingboxExtentMin, float InBoundingboxExtentMax, int32 InSurfaceType)
+		    TVector<T, d> InLocation, TVector<T, d> InAccumulatedImpulse, TVector<T, d> InNormal, TVector<T, d> InVelocity1, TVector<T, d> InVelocity2
+			, TVector<T, d> InAngularVelocity1, TVector<T, d> InAngularVelocity2, T InMass1, T InMass2, TGeometryParticle<T, d>* InParticle
+			, TGeometryParticle<T, d>* InLevelset, IPhysicsProxyBase* InParticleProxy, IPhysicsProxyBase* InLevelsetProxy
+			, float InBoundingboxVolume, float InBoundingboxExtentMin, float InBoundingboxExtentMax, int32 InSurfaceType)
 			: Location(InLocation)
 			, AccumulatedImpulse(InAccumulatedImpulse)
 			, Normal(InNormal)
@@ -197,8 +202,8 @@ namespace Chaos
 			, TVector<T, d> InVelocity
 			, TVector<T, d> InAngularVelocity
 			, T InMass
-						, TGeometryParticleHandle<T, d>* Particle
-						, IPhysicsProxyBase* ParticleProxy
+			, TGeometryParticleHandle<T, d>* InParticle
+			, IPhysicsProxyBase* InParticleProxy
 			, float InBoundingboxVolume
 			, float InBoundingboxExtentMin
 			, float InBoundingboxExtentMax
@@ -267,13 +272,14 @@ namespace Chaos
 			, BoundingBox(TAABB<T, d>(TVector<T, d>((T)0.0), TVector<T, d>((T)0.0)))
 		{}
 
-		TTrailingData(TVector<T, d> InLocation, TVector<T, d> InVelocity, TVector<T, d> InAngularVelocity, T InMass, TGeometryParticleHandle<T, d>* InParticle
-		    , float InBoundingboxVolume, float InBoundingboxExtentMin, float InBoundingboxExtentMax, int32 InSurfaceType, Chaos::TAABB<T, d>& InBoundingBox)
+		TTrailingData(TVector<T, d> InLocation, TVector<T, d> InVelocity, TVector<T, d> InAngularVelocity, T InMass
+			, TGeometryParticleHandle<T, d>* InParticle, IPhysicsProxyBase* InParticleProxy, Chaos::TAABB<T, d>& InBoundingBox)
 			: Location(InLocation)
 			, Velocity(InVelocity)
 			, AngularVelocity(InAngularVelocity)
 			, Mass(InMass)
 			, Particle(InParticle)
+		    , ParticleProxy(InParticleProxy)
 			, BoundingBox(InBoundingBox)
 		{}
 
@@ -318,8 +324,8 @@ namespace Chaos
 			, TVector<T, d> InVelocity
 			, TVector<T, d> InAngularVelocity
 			, T InMass
-				, TGeometryParticleHandle<T, d>* InParticle
-				, IPhysicsProxyBase* ParticleProxy
+			, TGeometryParticleHandle<T, d>* InParticle
+			, IPhysicsProxyBase* InParticleProxy
 			, float InBoundingboxVolume
 			, float InBoundingboxExtentMin
 			, float InBoundingboxExtentMax
