Final cleanup for making it easier to marshal data to PT

Michael.Lentine

diff --git a/Engine/Source/Runtime/Experimental/Chaos/Private/Chaos/ParticleHandle.cpp b/Engine/Source/Runtime/Experimental/Chaos/Private/Chaos/ParticleHandle.cpp
--- a/Engine/Source/Runtime/Experimental/Chaos/Private/Chaos/ParticleHandle.cpp
+++ b/Engine/Source/Runtime/Experimental/Chaos/Private/Chaos/ParticleHandle.cpp
@@ -26,7 +26,7 @@ namespace Chaos
 			}
 		}
 
-		auto& Geometry = MNonFrequentData.Read().Geometry;
+		auto& Geometry = MNonFrequentData.Read().Geometry();
 		if (Geometry)
 		{
 			int32 CurrentShapeIndex = INDEX_NONE;
diff --git a/Engine/Source/Runtime/Experimental/Chaos/Public/Chaos/ParticleDirtyFlags.h b/Engine/Source/Runtime/Experimental/Chaos/Public/Chaos/ParticleDirtyFlags.h
--- a/Engine/Source/Runtime/Experimental/Chaos/Public/Chaos/ParticleDirtyFlags.h
+++ b/Engine/Source/Runtime/Experimental/Chaos/Public/Chaos/ParticleDirtyFlags.h
@@ -239,79 +239,143 @@ inline FChaosArchive& operator<<(FChaosArchive& Ar,FParticleMassProps& Data)
 	return Ar;
 }
 
-struct FParticleMisc
+class FParticleDynamicMisc
 {
-	int32 CollisionGroup;
-	EObjectStateType ObjectState;
-	FSpatialAccelerationIdx SpatialIdx;
-
-	uint8 bDisabled : 1;
-	uint8 bGravityEnabled : 1;
-
+public:
 	void Serialize(FChaosArchive& Ar)
 	{
-		Ar << CollisionGroup;
-		Ar << ObjectState;
-		bool Disabled = bDisabled;
-		Ar << Disabled;
-		bDisabled = Disabled;
+		Ar << MLinearEtherDrag;
+		Ar << MAngularEtherDrag;
+		Ar << MObjectState;
+		Ar << MGravityEnabled;
+	}
 
-		bool GravityEnabled = bGravityEnabled;
-		Ar << GravityEnabled;
-		bGravityEnabled = GravityEnabled;
+	template <typename TOther>
+	void CopyFrom(const TOther& Other)
+	{
+		SetLinearEtherDrag(Other.LinearEtherDrag());
+		SetAngularEtherDrag(Other.AngularEtherDrag());
+		SetObjectState(Other.ObjectState());
+		SetGravityEnabled(Other.GravityEnabled());
+		SetCollisionGroup(Other.CollisionGroup());
 	}
 
-	bool operator==(const FParticleMisc& Other) const
+	template <typename TOther>
+	bool IsEqual(const TOther& Other) const
+	{
+		return ObjectState() == Other.ObjectState()
+			&& LinearEtherDrag() == Other.LinearEtherDrag()
+			&& AngularEtherDrag() == Other.AngularEtherDrag()
+			&& GravityEnabled() == Other.GravityEnabled()
+			&& CollisionGroup() == Other.CollisionGroup();
+	}
+
+	bool operator==(const FParticleDynamicMisc& Other) const
 	{
-		return CollisionGroup == Other.CollisionGroup
-			&& ObjectState == Other.ObjectState
-			&& SpatialIdx == Other.SpatialIdx
-			&& bDisabled == Other.bDisabled
-			&& bGravityEnabled == Other.bGravityEnabled;
+		return IsEqual(Other);
 	}
+
+	FReal LinearEtherDrag() const { return MLinearEtherDrag; }
+	void SetLinearEtherDrag(FReal InLinearEtherDrag) { MLinearEtherDrag = InLinearEtherDrag; }
+
+	FReal AngularEtherDrag() const { return MAngularEtherDrag; }
+	void SetAngularEtherDrag(FReal InAngularEtherDrag) { MAngularEtherDrag = InAngularEtherDrag; }
+
+	EObjectStateType ObjectState() const { return MObjectState; }
+	void SetObjectState(EObjectStateType InState){ MObjectState = InState; }
+
+	bool GravityEnabled() const { return MGravityEnabled; }
+	void SetGravityEnabled(bool InGravity){ MGravityEnabled = InGravity; }
+
+	int32 CollisionGroup() const { return MCollisionGroup; }
+	void SetCollisionGroup(int32 InGroup){ MCollisionGroup = InGroup; }
+
+
+private:
+	FReal MLinearEtherDrag;
+	FReal MAngularEtherDrag;
+	int32 MCollisionGroup;
+
+	EObjectStateType MObjectState;
+	bool MGravityEnabled;
 };
 
-inline FChaosArchive& operator<<(FChaosArchive& Ar,FParticleMisc& Data)
+inline FChaosArchive& operator<<(FChaosArchive& Ar,FParticleDynamicMisc& Data)
 {
 	Data.Serialize(Ar);
 	return Ar;
 }
 
-struct FParticleNonFrequentData
+class FParticleNonFrequentData
 {
+public:
 	FParticleNonFrequentData()
-		: UserData(nullptr)
+	: MUserData(nullptr)
 	{
 
 	}
 
-	TSharedPtr<FImplicitObject,ESPMode::ThreadSafe> Geometry;
-	void* UserData;
-	FUniqueIdx UniqueIdx;
-	FReal LinearEtherDrag;
-	FReal AngularEtherDrag;
+	void Serialize(FChaosArchive& Ar)
+	{
+		Ar << MGeometry;
+	}
 
+	template <typename TOther>
+	void CopyFrom(const TOther& Other)
+	{
+		SetGeometry(Other.SharedGeometryLowLevel());
+		SetUserData(Other.UserData());
+		SetUniqueIdx(Other.UniqueIdx());
+		SetSpatialIdx(Other.SpatialIdx());
 #if CHAOS_CHECKED
-	FName DebugName;
+		SetDebugName(Other.DebugName());
 #endif
+	}
 
-	void Serialize(FChaosArchive& Ar)
+	template <typename TOther>
+	bool IsEqual(const TOther& Other) const
 	{
-		Ar << Geometry;
+		return Geometry() == Other.SharedGeometryLowLevel()
+			&& UserData() == Other.UserData()
+			&& UniqueIdx() == Other.UniqueIdx()
+			&& SpatialIdx() == Other.SpatialIdx()
+#if CHAOS_CHECKED
+			&& DebugName() == Other.DebugName()
+#endif
+			;
 	}
 
 	bool operator==(const FParticleNonFrequentData& Other) const
 	{
-		return Geometry == Other.Geometry
-			&& UserData == Other.UserData
-			&& UniqueIdx == Other.UniqueIdx
-			&& LinearEtherDrag == Other.LinearEtherDrag
-			&& AngularEtherDrag == Other.AngularEtherDrag
+		return IsEqual(Other);
+	}
+
+	const TSharedPtr<FImplicitObject,ESPMode::ThreadSafe>& Geometry() const { return MGeometry;}
+	const TSharedPtr<FImplicitObject,ESPMode::ThreadSafe>& SharedGeometryLowLevel() const { return MGeometry;}
+	void SetGeometry(const TSharedPtr<FImplicitObject,ESPMode::ThreadSafe>& InGeometry) { MGeometry = InGeometry;}
+
+	void* UserData() const { return MUserData; }
+	void SetUserData(void* InData){ MUserData = InData;}
+
+	const FUniqueIdx& UniqueIdx() const { return MUniqueIdx; }
+	void SetUniqueIdx(FUniqueIdx InIdx){ MUniqueIdx = InIdx; }
+
+	FSpatialAccelerationIdx SpatialIdx() const { return MSpatialIdx; }
+	void SetSpatialIdx(FSpatialAccelerationIdx InIdx){ MSpatialIdx = InIdx; }
+
 #if CHAOS_CHECKED
-			&& DebugName == Other.DebugName
+	FName DebugName() const { return MDebugName; }
+	void SetDebugName(FName InName) { MDebugName = InName; }
+#endif
+private:
+	TSharedPtr<FImplicitObject,ESPMode::ThreadSafe> MGeometry;
+	void* MUserData;
+	FUniqueIdx MUniqueIdx;
+	FSpatialAccelerationIdx MSpatialIdx;
+
+#if CHAOS_CHECKED
+	FName MDebugName;
 #endif
-			;
-	}
 };
 
 inline FChaosArchive& operator<<(FChaosArchive& Ar,FParticleNonFrequentData& Data)
diff --git a/Engine/Source/Runtime/Experimental/Chaos/Public/Chaos/ParticleHandle.h b/Engine/Source/Runtime/Experimental/Chaos/Public/Chaos/ParticleHandle.h
--- a/Engine/Source/Runtime/Experimental/Chaos/Public/Chaos/ParticleHandle.h
+++ b/Engine/Source/Runtime/Experimental/Chaos/Public/Chaos/ParticleHandle.h
@@ -95,7 +95,7 @@ void PBDRigidParticleDefaultConstruct(TPBDRigidParticle<T,d>& Concrete, const TP
 {
 	//don't bother calling parent since the call gets made by the corresponding hierarchy in FConcrete
 	Concrete.SetCollisionGroup(0);
-	Concrete.SetDisabled(Params.bDisabled);
+	//Concrete.SetDisabled(Params.bDisabled);
 	Concrete.SetF(TVector<T, d>(0));
 	Concrete.SetTorque(TVector<T, d>(0));
 	Concrete.SetLinearImpulse(TVector<T, d>(0));
@@ -430,6 +430,17 @@ public:
 		SetR(XR.R());
 	}
 
+	void SetNonFrequentData(const FParticleNonFrequentData& InData)
+	{
+		SetSharedGeometry(InData.Geometry());
+		SetUserData(InData.UserData());
+		SetUniqueIdx(InData.UniqueIdx());
+		SetSpatialIdx(InData.SpatialIdx());
+
+#if CHAOS_CHECKED
+		SetDebugName(InData.DebugName());
+#endif
+	}
 
 	TSerializablePtr<FImplicitObject> Geometry() const { return GeometryParticles->Geometry(ParticleIdx); }
 	void SetGeometry(TSerializablePtr<FImplicitObject> InGeometry) { GeometryParticles->SetGeometry(ParticleIdx, InGeometry); }
@@ -743,6 +754,13 @@ public:
 		SetInvM(Props.InvM());
 	}
 
+	void SetDynamicMisc(const FParticleDynamicMisc& DynamicMisc)
+	{
+		SetLinearEtherDrag(DynamicMisc.LinearEtherDrag());
+		SetAngularEtherDrag(DynamicMisc.AngularEtherDrag());
+		SetCollisionGroup(DynamicMisc.CollisionGroup());
+	}
+
 	const PMatrix<T, d, d>& I() const { return PBDRigidParticles->I(ParticleIdx); }
 	PMatrix<T, d, d>& I() { return PBDRigidParticles->I(ParticleIdx); }
 	void SetI(const PMatrix<T, d, d>& InI) { PBDRigidParticles->I(ParticleIdx) = InI; }
@@ -1435,7 +1453,7 @@ public:
 
 	virtual bool IsParticleValid() const
 	{
-		auto& Geometry = MNonFrequentData.Read().Geometry;
+		auto& Geometry = MNonFrequentData.Read().Geometry();
 		return Geometry && Geometry->IsValidGeometry();	//todo: if we want support for sample particles without geometry we need to adjust this
 	}
 
@@ -1444,10 +1462,10 @@ public:
 	const TVector<T, d>& X() const { return MXR.Read().X(); }
 	void SetX(const TVector<T, d>& InX, bool bInvalidate = true);
 
-	FUniqueIdx UniqueIdx() const { return MNonFrequentData.Read().UniqueIdx; }
+	FUniqueIdx UniqueIdx() const { return MNonFrequentData.Read().UniqueIdx(); }
 	void SetUniqueIdx(const FUniqueIdx UniqueIdx, bool bInvalidate = true)
 	{
-		MNonFrequentData.Modify(bInvalidate,MDirtyFlags,Proxy,[UniqueIdx](auto& Data){ Data.UniqueIdx = UniqueIdx;});
+		MNonFrequentData.Modify(bInvalidate,MDirtyFlags,Proxy,[UniqueIdx](auto& Data){ Data.SetUniqueIdx(UniqueIdx);});
 	}
 
 	const TRotation<T, d>& R() const { return MXR.Read().R(); }
@@ -1473,7 +1491,7 @@ public:
 	//       We should replace this with a method for supporting SetGeometry(RawGeometry).
 	void SetGeometry(TSharedPtr<FImplicitObject, ESPMode::ThreadSafe> SharedGeometry)
 	{
-		MNonFrequentData.Modify(true,MDirtyFlags,Proxy,[&SharedGeometry](auto& Data){ Data.Geometry = SharedGeometry;});
+		MNonFrequentData.Modify(true,MDirtyFlags,Proxy,[&SharedGeometry](auto& Data){ Data.SetGeometry(SharedGeometry);});
 		UpdateShapesArray();
 	}
 
@@ -1484,12 +1502,12 @@ public:
 		check(false);
 	}
 
-	const TSharedPtr<FImplicitObject,ESPMode::ThreadSafe>& SharedGeometryLowLevel() const { return MNonFrequentData.Read().Geometry; }
+	const TSharedPtr<FImplicitObject,ESPMode::ThreadSafe>& SharedGeometryLowLevel() const { return MNonFrequentData.Read().Geometry(); }
 
-	void* UserData() const { return MNonFrequentData.Read().UserData; }
+	void* UserData() const { return MNonFrequentData.Read().UserData(); }
 	void SetUserData(void* InUserData)
 	{
-		MNonFrequentData.Modify(true,MDirtyFlags,Proxy,[InUserData](auto& Data){ Data.UserData = InUserData;});
+		MNonFrequentData.Modify(true,MDirtyFlags,Proxy,[InUserData](auto& Data){ Data.SetUserData(InUserData);});
 	}
 
 	void UpdateShapeBounds()
@@ -1499,7 +1517,7 @@ public:
 
 	void UpdateShapeBounds(const FTransform& Transform)
 	{
-		if (MNonFrequentData.Read().Geometry->HasBoundingBox())
+		if (MNonFrequentData.Read().Geometry()->HasBoundingBox())
 		{
 			for (auto& Shape : MShapesArray)
 			{
@@ -1536,10 +1554,10 @@ public:
 	}
 
 #if CHAOS_CHECKED
-	const FName& DebugName() const { return MNonFrequentData.Read().DebugName; }
+	const FName& DebugName() const { return MNonFrequentData.Read().DebugName(); }
 	void SetDebugName(const FName& InDebugName)
 	{
-		MNonFrequentData.Modify(true,MDirtyFlags,Proxy,[&InDebugName](auto& Data){ Data.DebugName = InDebugName;});
+		MNonFrequentData.Modify(true,MDirtyFlags,Proxy,[&InDebugName](auto& Data){ Data.SetDebugName(InDebugName);});
 	}
 #endif
 
@@ -1551,7 +1569,7 @@ public:
 		MapImplicitShapes();
 	}
 
-	TSerializablePtr<FImplicitObject> Geometry() const { return MakeSerializable(MNonFrequentData.Read().Geometry); }
+	TSerializablePtr<FImplicitObject> Geometry() const { return MakeSerializable(MNonFrequentData.Read().Geometry()); }
 
 	const FShapesArray& ShapesArray() const { return MShapesArray; }
 
@@ -1570,10 +1588,15 @@ public:
 	const TPBDRigidParticle<T, d>* CastToRigidParticle() const;
 	TPBDRigidParticle<T, d>* CastToRigidParticle();
 
-	FSpatialAccelerationIdx SpatialIdx() const { return MMiscData.Read().SpatialIdx; }
+	FSpatialAccelerationIdx SpatialIdx() const { return MNonFrequentData.Read().SpatialIdx(); }
 	void SetSpatialIdx(FSpatialAccelerationIdx Idx)
 	{
-		MMiscData.Modify(true,MDirtyFlags,Proxy,[Idx](auto& Data){ Data.SpatialIdx = Idx;});
+		MNonFrequentData.Modify(true,MDirtyFlags,Proxy,[Idx](auto& Data){ Data.SetSpatialIdx(Idx);});
+	}
+
+	void SetNonFrequentData(const FParticleNonFrequentData& InData)
+	{
+		MNonFrequentData.Write(InData,true,MDirtyFlags,Proxy);
 	}
 
 	bool IsDirty() const
@@ -1670,11 +1693,7 @@ protected:
 private:
 
 	TParticleProperty<FParticlePositionRotation, EParticleProperty::XR> MXR;
-
-protected:
 	TParticleProperty<FParticleNonFrequentData,EParticleProperty::NonFrequentData> MNonFrequentData;
-	TParticleProperty<FParticleMisc,EParticleProperty::Misc> MMiscData;
-private:
 
 	FShapesArray MShapesArray;
 	TMap<const FImplicitObject*, int32> ImplicitShapeMap;
@@ -1690,7 +1709,7 @@ public:
 	// when the geometry is being copied from GT to PT.
 	TSharedPtr<FImplicitObject, ESPMode::ThreadSafe> GeometrySharedLowLevel() const
 	{
-		return MNonFrequentData.Read().Geometry;
+		return MNonFrequentData.Read().Geometry();
 	}
 private:
 
@@ -1703,7 +1722,7 @@ protected:
 
 	void UpdateShapesArray()
 	{
-		UpdateShapesArrayFromGeometry(MShapesArray, MakeSerializable(MNonFrequentData.Read().Geometry), FRigidTransform3(X(), R()), Proxy);
+		UpdateShapesArrayFromGeometry(MShapesArray, MakeSerializable(MNonFrequentData.Read().Geometry()), FRigidTransform3(X(), R()), Proxy);
 		MapImplicitShapes();
 	}
 
@@ -1711,7 +1730,6 @@ protected:
 	{
 		MXR.SyncRemote(Manager, DataIdx, RemoteData);
 		MNonFrequentData.SyncRemote(Manager, DataIdx, RemoteData);
-		MMiscData.SyncRemote(Manager, DataIdx, RemoteData);
 	}
 
 	void MapImplicitShapes();
@@ -1854,7 +1872,6 @@ public:
 	using TGeometryParticle<T, d>::CastToRigidParticle;
 	using Base = TGeometryParticle<T,d>;
 	using Base::MDirtyFlags;
-	using Base::MMiscData;
 
 protected:
 	using Base::Proxy;
@@ -1953,8 +1970,6 @@ public:
 	using Base = TKinematicGeometryParticle<T,d>;
 
 	using Base::MDirtyFlags;
-	using Base::MMiscData;
-
 
 protected:
 	using Base::Proxy;
@@ -1985,24 +2000,25 @@ public:
 
 	const TUniquePtr<TBVHParticles<T, d>>& CollisionParticles() const { return MCollisionParticles; }
 
-	int32 CollisionGroup() const { return MMiscData.Read().CollisionGroup; }
+	int32 CollisionGroup() const { return MMiscData.Read().CollisionGroup(); }
 	void SetCollisionGroup(const int32 InCollisionGroup)
 	{
-		MMiscData.Modify(true,MDirtyFlags,Proxy,[InCollisionGroup](auto& Data){ Data.CollisionGroup = InCollisionGroup;});
+		MMiscData.Modify(true,MDirtyFlags,Proxy,[InCollisionGroup](auto& Data){ Data.SetCollisionGroup(InCollisionGroup);});
 	}
 
+	/*
 	bool Disabled() const { return MMiscData.Read().bDisabled; }
 	void SetDisabled(const bool InDisabled)
 	{
 		MMiscData.Modify(true,MDirtyFlags,Proxy,[InDisabled](auto& Data){ Data.bDisabled = InDisabled;});
-	}
+	}*/
 
-	bool IsGravityEnabled() const { return MMiscData.Read().bGravityEnabled; }
+	bool IsGravityEnabled() const { return MMiscData.Read().GravityEnabled(); }
 	void SetGravityEnabled(const bool InGravityEnabled)
 	{
-		MMiscData.Modify(true,MDirtyFlags,Proxy,[InGravityEnabled](auto& Data){ Data.bGravityEnabled = InGravityEnabled;});
+		MMiscData.Modify(true,MDirtyFlags,Proxy,[InGravityEnabled](auto& Data){ Data.SetGravityEnabled (InGravityEnabled);});
 	}
-
+	
 	//todo: remove this
 	bool IsInitialized() const { return MInitialized; }
 	void SetInitialized(const bool InInitialized)
@@ -2096,16 +2112,21 @@ public:
 		MMassProps.Write(InProps,true,MDirtyFlags,Proxy);
 	}
 
-	T LinearEtherDrag() const { return this->MNonFrequentData.Read().LinearEtherDrag; }
+	void SetDynamicMisc(const FParticleDynamicMisc& DynamicMisc)
+	{
+		MMiscData.Write(DynamicMisc,true,MDirtyFlags,Proxy);
+	}
+
+	T LinearEtherDrag() const { return MMiscData.Read().LinearEtherDrag(); }
 	void SetLinearEtherDrag(const T& InLinearEtherDrag)
 	{
-		this->MNonFrequentData.Modify(true,MDirtyFlags,Proxy,[&InLinearEtherDrag](auto& Data){ Data.LinearEtherDrag = InLinearEtherDrag;});
+		MMiscData.Modify(true,MDirtyFlags,Proxy,[&InLinearEtherDrag](auto& Data){ Data.SetLinearEtherDrag(InLinearEtherDrag);});
 	}
 
-	T AngularEtherDrag() const { return this->MNonFrequentData.Read().AngularEtherDrag; }
+	T AngularEtherDrag() const { return MMiscData.Read().AngularEtherDrag(); }
 	void SetAngularEtherDrag(const T& InAngularEtherDrag)
 	{
-		this->MNonFrequentData.Modify(true,MDirtyFlags,Proxy,[&InAngularEtherDrag](auto& Data){ Data.AngularEtherDrag = InAngularEtherDrag;});
+		MMiscData.Modify(true,MDirtyFlags,Proxy,[&InAngularEtherDrag](auto& Data){ Data.SetAngularEtherDrag(InAngularEtherDrag);});
 	}
 
 	int32 Island() const { return MIsland; }
@@ -2122,7 +2143,7 @@ public:
 		this->MToBeRemovedOnFracture = bToBeRemovedOnFracture;
 	}
 
-	EObjectStateType ObjectState() const { return MMiscData.Read().ObjectState; }
+	EObjectStateType ObjectState() const { return MMiscData.Read().ObjectState(); }
 	void SetObjectState(const EObjectStateType InState, bool bAllowEvents = false)
 	{
 		if (bAllowEvents && ObjectState() != EObjectStateType::Dynamic && InState == EObjectStateType::Dynamic)
@@ -2144,7 +2165,7 @@ public:
 			MDirtyFlags.MarkClean(ParticlePropToFlag(EParticleProperty::Dynamics));
 		}
 
-		MMiscData.Modify(true,MDirtyFlags,Proxy,[&InState](auto& Data){ Data.ObjectState = InState;});
+		MMiscData.Modify(true,MDirtyFlags,Proxy,[&InState](auto& Data){ Data.SetObjectState(InState);});
 
 	}
 
@@ -2154,6 +2175,8 @@ public:
 private:
 	TParticleProperty<FParticleMassProps,EParticleProperty::MassProps> MMassProps;
 	TParticleProperty<FParticleDynamics, EParticleProperty::Dynamics> MDynamics;
+	TParticleProperty<FParticleDynamicMisc,EParticleProperty::DynamicMisc> MMiscData;
+
 	TUniquePtr<TBVHParticles<T, d>> MCollisionParticles;
 	int32 MIsland;
 	bool MToBeRemovedOnFracture;
@@ -2166,6 +2189,7 @@ protected:
 		Base::SyncRemoteDataImp(Manager,DataIdx,RemoteData);
 		MMassProps.SyncRemote(Manager,DataIdx,RemoteData);
 		MDynamics.SyncRemote(Manager,DataIdx,RemoteData);
+		MMiscData.SyncRemote(Manager,DataIdx,RemoteData);
 	}
 };
 
@@ -2216,7 +2240,7 @@ public:
 		, MIsland(INDEX_NONE)
 		, MCollisionGroup(0)
 		, MObjectState(EObjectStateType::Uninitialized)
-		, MDisabled(false)
+		//, MDisabled(false)
 		, MToBeRemovedOnFracture(false)
 		, MGravityEnabled(false)
 		, MInitialized(false)
@@ -2240,7 +2264,7 @@ public:
 		, MIsland(InParticle.Island())
 		, MCollisionGroup(InParticle.CollisionGroup())
 		, MObjectState(InParticle.ObjectState())
-		, MDisabled(InParticle.Disabled())
+		//, MDisabled(InParticle.Disabled())
 		, MToBeRemovedOnFracture(InParticle.ToBeRemovedOnFracture())
 		, MGravityEnabled(InParticle.IsGravityEnabled())
 		, MInitialized(InParticle.IsInitialized())
@@ -2265,7 +2289,7 @@ public:
 	int32 MIsland;
 	int32 MCollisionGroup;
 	EObjectStateType MObjectState;
-	bool MDisabled;
+	//bool MDisabled;
 	bool MToBeRemovedOnFracture;
 	bool MGravityEnabled;
 	bool MInitialized;
@@ -2289,7 +2313,7 @@ public:
 		MIsland = INDEX_NONE;
 		MCollisionGroup = 0;
 		MObjectState = EObjectStateType::Uninitialized;
-		MDisabled = false;
+		//MDisabled = false;
 		MToBeRemovedOnFracture = false;
 		MGravityEnabled = false;
 		MInitialized = false;
@@ -2310,7 +2334,7 @@ public:
 			MIsland = InParticle.Island();
 			MCollisionGroup = InParticle.CollisionGroup();
 			MObjectState = InParticle.ObjectState();
-			MDisabled = InParticle.Disabled();
+			//MDisabled = InParticle.Disabled();
 			MToBeRemovedOnFracture = InParticle.ToBeRemovedOnFracture();
 			MGravityEnabled = InParticle.IsGravityEnabled();
 			MInitialized = InParticle.IsInitialized();
diff --git a/Engine/Source/Runtime/Experimental/Chaos/Public/Chaos/ParticleProperties.inl b/Engine/Source/Runtime/Experimental/Chaos/Public/Chaos/ParticleProperties.inl
--- a/Engine/Source/Runtime/Experimental/Chaos/Public/Chaos/ParticleProperties.inl
+++ b/Engine/Source/Runtime/Experimental/Chaos/Public/Chaos/ParticleProperties.inl
@@ -10,6 +10,6 @@
 PARTICLE_PROPERTY(XR, FParticlePositionRotation)
 PARTICLE_PROPERTY(Velocities, FParticleVelocities)
 PARTICLE_PROPERTY(Dynamics, FParticleDynamics)
-PARTICLE_PROPERTY(Misc, FParticleMisc)
+PARTICLE_PROPERTY(DynamicMisc, FParticleDynamicMisc)
 PARTICLE_PROPERTY(NonFrequentData, FParticleNonFrequentData)
 PARTICLE_PROPERTY(MassProps, FParticleMassProps)
\ No newline at end of file
diff --git a/Engine/Source/Runtime/Experimental/Chaos/Public/Chaos/PropertiesTypes.inl b/Engine/Source/Runtime/Experimental/Chaos/Public/Chaos/PropertiesTypes.inl
--- a/Engine/Source/Runtime/Experimental/Chaos/Public/Chaos/PropertiesTypes.inl
+++ b/Engine/Source/Runtime/Experimental/Chaos/Public/Chaos/PropertiesTypes.inl
@@ -3,7 +3,7 @@
 PROPERTY_TYPE(FParticlePositionRotation_type, FParticlePositionRotation)
 PROPERTY_TYPE(FParticleVelocities_type, FParticleVelocities)
 PROPERTY_TYPE(FParticleDynamics_type, FParticleDynamics)
-PROPERTY_TYPE(FParticleMisc_type, FParticleMisc)
+PROPERTY_TYPE(FParticleDynamicMisc_type, FParticleDynamicMisc)
 PROPERTY_TYPE(FParticleNonFrequentData_type, FParticleNonFrequentData)
 PROPERTY_TYPE(FCollisionData_type, FCollisionData)
 PROPERTY_TYPE(FMaterialData_type, FMaterialData)
diff --git a/Engine/Source/Runtime/Experimental/ChaosSolvers/Private/PBDRigidsSolver.cpp b/Engine/Source/Runtime/Experimental/ChaosSolvers/Private/PBDRigidsSolver.cpp
--- a/Engine/Source/Runtime/Experimental/ChaosSolvers/Private/PBDRigidsSolver.cpp
+++ b/Engine/Source/Runtime/Experimental/ChaosSolvers/Private/PBDRigidsSolver.cpp
@@ -807,7 +807,7 @@ namespace Chaos
 				if(!bIsSingleThreaded && bIsNew)
 				{
 					const auto* NonFrequentData = Dirty.ParticleData.FindNonFrequentData(*Manager,DataIdx);
-					const FUniqueIdx* UniqueIdx = NonFrequentData ? &NonFrequentData->UniqueIdx : nullptr;
+					const FUniqueIdx* UniqueIdx = NonFrequentData ? &NonFrequentData->UniqueIdx() : nullptr;
 					Proxy->SetHandle(CreateHandleFunc(UniqueIdx));
 				}
 
diff --git a/Engine/Source/Runtime/Experimental/ChaosSolvers/Private/PhysicsProxy/SingleParticlePhysicsProxy.cpp b/Engine/Source/Runtime/Experimental/ChaosSolvers/Private/PhysicsProxy/SingleParticlePhysicsProxy.cpp
--- a/Engine/Source/Runtime/Experimental/ChaosSolvers/Private/PhysicsProxy/SingleParticlePhysicsProxy.cpp
+++ b/Engine/Source/Runtime/Experimental/ChaosSolvers/Private/PhysicsProxy/SingleParticlePhysicsProxy.cpp
@@ -57,19 +57,7 @@ void PushToPhysicsStateImp(const Chaos::FDirtyPropertiesManager& Manager, TParti
 
 		if(NewNonFrequentData)
 		{
-			Handle->SetSharedGeometry(NewNonFrequentData->Geometry);
-			Handle->SetUniqueIdx(NewNonFrequentData->UniqueIdx);
-			Handle->SetUserData(NewNonFrequentData->UserData);
-
-			if(bHasDynamicData)
-			{
-				RigidHandle->SetLinearEtherDrag(NewNonFrequentData->LinearEtherDrag);
-				RigidHandle->SetAngularEtherDrag(NewNonFrequentData->AngularEtherDrag);
-			}
-			
-#if CHAOS_CHECKED
-			Handle->SetDebugName(NewNonFrequentData->DebugName);
-#endif
+			Handle->SetNonFrequentData(*NewNonFrequentData);
 		}
 
 		auto NewVelocities = bHasKinematicData ? ParticleData.FindVelocities(Manager, DataIdx) : nullptr;
@@ -96,17 +84,6 @@ void PushToPhysicsStateImp(const Chaos::FDirtyPropertiesManager& Manager, TParti
 			Solver->GetEvolution()->DirtyParticle(*Handle);
 		}
 
-		if(auto NewData = ParticleData.FindMisc(Manager, DataIdx))
-		{
-			Handle->SetSpatialIdx(NewData->SpatialIdx);
-
-			if(bHasDynamicData)
-			{
-				Solver->GetEvolution()->SetParticleObjectState(RigidHandle,NewData->ObjectState);
-				Solver->GetEvolution()->GetGravityForces().SetEnabled(*RigidHandle,NewData->bGravityEnabled);
-			}
-		}
-
 		if(bHasDynamicData)
 		{
 			if(auto NewData = ParticleData.FindMassProps(Manager,DataIdx))
@@ -118,6 +95,14 @@ void PushToPhysicsStateImp(const Chaos::FDirtyPropertiesManager& Manager, TParti
 			{
 				RigidHandle->SetDynamics(*NewData);
 			}
+
+			if(auto NewData = ParticleData.FindDynamicMisc(Manager,DataIdx))
+			{
+				Solver->GetEvolution()->SetParticleObjectState(RigidHandle,NewData->ObjectState());
+				Solver->GetEvolution()->GetGravityForces().SetEnabled(*RigidHandle,NewData->GravityEnabled());
+
+				RigidHandle->SetDynamicMisc(*NewData);
+			}
 		}
 
 		//shape properties
@@ -573,7 +558,7 @@ void FSingleParticlePhysicsProxy<Chaos::TPBDRigidParticle<float, 3>>::PullFromPh
 		Particle->UpdateShapeBounds();
 		//if (!Particle->IsDirty(Chaos::EParticleFlags::ObjectState))
 		//question: is it ok to call this when it was one of the other properties that changed?
-		if (!Particle->IsDirty(Chaos::EParticleFlags::Misc))
+		if (!Particle->IsDirty(Chaos::EParticleFlags::DynamicMisc))
 		{
 			Particle->SetObjectState(Buffer->MObjectState, true);
 		}
diff --git a/Engine/Source/Runtime/Experimental/ChaosSolvers/Public/RewindData.h b/Engine/Source/Runtime/Experimental/ChaosSolvers/Public/RewindData.h
--- a/Engine/Source/Runtime/Experimental/ChaosSolvers/Public/RewindData.h
+++ b/Engine/Source/Runtime/Experimental/ChaosSolvers/Public/RewindData.h
@@ -157,7 +157,7 @@ public:
 	template <typename TParticle>
 	TSerializablePtr<FImplicitObject> Geometry(const TParticle& Particle) const
 	{
-		return NonFrequentData.IsSet() ? MakeSerializable(NonFrequentData.Read().Geometry) : Particle.Geometry();
+		return NonFrequentData.IsSet() ? MakeSerializable(NonFrequentData.Read().Geometry()) : Particle.Geometry();
 	}
 
 	template <typename TParticle>
@@ -217,18 +217,7 @@ public:
 
 		NonFrequentData.SyncToParticle([&Particle](const auto& Data)
 		{
-			Particle.SetGeometry(Data.Geometry);
-			Particle.SetUserData(Data.UserData);
-			ensure(Data.UniqueIdx == Particle.UniqueIdx());	//this should never change
-
-#if CHAOS_CHECKED
-			Particle.SetDebugName(Data.DebugName);
-#endif
-			if(auto Rigid = Particle.CastToRigidParticle())
-			{
-				Rigid->SetLinearEtherDrag(Data.LinearEtherDrag);
-				Rigid->SetAngularEtherDrag(Data.AngularEtherDrag);
-			}
+			Particle.SetNonFrequentData(Data);
 		});
 
 		if(auto Rigid = Particle.CastToRigidParticle())
@@ -270,21 +259,7 @@ public:
 
 		NonFrequentData.SyncRemoteData(Manager,Idx,Dirty.ParticleData,[Handle](FParticleNonFrequentData& Data)
 		{
-			Data.Geometry = Handle->SharedGeometryLowLevel();
-			Data.UserData = Handle->UserData();
-
-			//note: this data is keyed based on unique idx so it's not really possible to change this
-			//but we save it anyway since it's part of a big struct
-			Data.UniqueIdx = Handle->UniqueIdx();
-#if CHAOS_CHECKED
-			Data.DebugName = Handle->DebugName();
-#endif
-
-			if(auto Rigid = Handle->CastToRigidParticle())
-			{
-				Data.LinearEtherDrag = Rigid->LinearEtherDrag();
-				Data.AngularEtherDrag = Rigid->AngularEtherDrag();
-			}
+			Data.CopyFrom(*Handle);
 		});
 	}
 
@@ -316,21 +291,7 @@ public:
 		{
 			NonFrequentData.SyncRemoteDataForced(Manager,Idx,[Particle](FParticleNonFrequentData& Data)
 			{
-				Data.Geometry = Particle->SharedGeometryLowLevel();
-				Data.UserData = Particle->UserData();
-
-				//note: this data is keyed based on unique idx so it's not really possible to change this
-				//but we save it anyway since it's part of a big struct
-				Data.UniqueIdx = Particle->UniqueIdx();
-	#if CHAOS_CHECKED
-				Data.DebugName = Particle->DebugName();
-	#endif
-
-				if(auto Rigid = Particle->CastToRigidParticle())
-				{
-					Data.LinearEtherDrag = Rigid->LinearEtherDrag();
-					Data.AngularEtherDrag = Rigid->AngularEtherDrag();
-				}
+				Data.CopyFrom(*Particle);
 			});
 		}
 	}
