Fixed generation of cluster implicits due to errors in mass space recursive calculation for particle locations.
Fixed initial cluster transforms being incorrect in SQ

diff --git a/Engine/Source/Runtime/Experimental/ChaosSolvers/Private/PhysicsProxy/GeometryCollectionPhysicsProxy.cpp b/Engine/Source/Runtime/Experimental/ChaosSolvers/Private/PhysicsProxy/GeometryCollectionPhysicsProxy.cpp
--- a/Engine/Source/Runtime/Experimental/ChaosSolvers/Private/PhysicsProxy/GeometryCollectionPhysicsProxy.cpp
+++ b/Engine/Source/Runtime/Experimental/ChaosSolvers/Private/PhysicsProxy/GeometryCollectionPhysicsProxy.cpp
@@ -878,6 +878,7 @@ void FGeometryCollectionPhysicsProxy::InitializeBodiesPT(Chaos::FPBDRigidsSolver
 						RigidChildrenTransformGroupIndex.Add(ChildIndex);
 					}
 				}
+
 				if (RigidChildren.Num())
 				{
 					if (ReportTooManyChildrenNum >= 0 && RigidChildren.Num() > ReportTooManyChildrenNum)
@@ -899,12 +900,20 @@ void FGeometryCollectionPhysicsProxy::InitializeBodiesPT(Chaos::FPBDRigidsSolver
 					SolverClusterID[TransformGroupIndex] = Handle->ClusterIds().Id;
 
 					// Hook the handle up with the GT particle
-					Handle->GTGeometryParticle() = GTParticles[TransformGroupIndex].Get();
+					Chaos::TGeometryParticle<Chaos::FReal, 3>* GTParticle = GTParticles[TransformGroupIndex].Get();
+					Handle->GTGeometryParticle() = GTParticle;
+					
+					// Cluster transform has been recalculated based on children - copy to the GT particle (not threadsafe - just testing)
+					GTParticle->SetX(Handle->X());
+					GTParticle->SetR(Handle->R());
+					GTParticle->UpdateShapeBounds();
 
 					SolverClusterHandles[TransformGroupIndex] = Handle;
 					SolverParticleHandles[TransformGroupIndex] = Handle;
 					HandleToTransformGroupIndex.Add(Handle, TransformGroupIndex);
 					RigidsSolver->AddParticleToProxy(Handle, this);
+
+					RigidsSolver->GetEvolution()->DirtyParticle(*Handle);
 				}
 			}
 
@@ -1079,14 +1088,17 @@ FGeometryCollectionPhysicsProxy::BuildClusters(
 	const int32 Level = FMath::Clamp(CalculateHierarchyLevel(DynamicCollection, CollectionClusterIndex), 0, INT_MAX);
 	const float DefaultDamage = NumThresholds > 0 ? Parameters.DamageThreshold[NumThresholds - 1] : 0.f;
 	float Damage = Level < NumThresholds ? Parameters.DamageThreshold[Level] : DefaultDamage;
-	if(Level >= Parameters.MaxClusterLevel) 
+
+	if(Level >= Parameters.MaxClusterLevel)
+	{
 		Damage = FLT_MAX;
+	}
+
 	Parent->SetStrains(Damage);
 
 	GetSolver()->GetEvolution()->SetPhysicsMaterial(Parent, Parameters.PhysicalMaterial);
 
-	const FTransform ParentTransform = 
-		GeometryCollectionAlgo::GlobalMatrix(DynamicCollection.Transform, DynamicCollection.Parent, CollectionClusterIndex);
+	const FTransform ParentTransform = GeometryCollectionAlgo::GlobalMatrix(DynamicCollection.Transform, DynamicCollection.Parent, CollectionClusterIndex);
 
 	int32 MinCollisionGroup = INT_MAX;
 	for(int32 Idx=0; Idx < ChildHandles.Num(); Idx++)
@@ -1126,6 +1138,8 @@ FGeometryCollectionPhysicsProxy::BuildClusters(
 		const Chaos::TRigidTransform<float, 3> Xf(Parent->X(), Parent->R());
 		const Chaos::TAABB<float, 3> TransformedBBox = LocalBounds.TransformedAABB(Xf);
 		Parent->SetWorldSpaceInflatedBounds(TransformedBBox);
+
+		GetSolver()->GetEvolution()->DirtyParticle(*Parent);
 	}
 
 	return Parent;
@@ -2289,8 +2303,10 @@ void FGeometryCollectionPhysicsProxy::InitializeSharedCollisionStructures(
 							const TArray<TVector<int32, 3>>& ChildIndices = ChildMesh->GetSurfaceElements();
 							UnionMeshIndices.Append(ChildIndices);
 
-							const FTransform ChildMassToClusterMass =
-								CollectionSpaceTransforms[ChildTransformIdx].GetRelativeTransform(ClusterMassToCollection);
+							// To move a particle from mass-space in the child to mass-space in the cluster parent, calculate
+							// the relative transform between the mass-space origin for both the parent and child before
+							// transforming the mass space particles into the parent mass-space.
+							const FTransform ChildMassToClusterMass = (CollectionSpaceTransforms[ChildTransformIdx] * CollectionMassToLocal[ChildTransformIdx]).GetRelativeTransform(CollectionSpaceTransforms[ClusterTransformIdx] * CollectionMassToLocal[ClusterTransformIdx]);
 
 							ChildMesh->GetVertexSet(VertsAdded);
 							for (const int32 VertIdx : VertsAdded)
diff --git a/Engine/Source/Runtime/Experimental/GeometryCollectionEngine/Private/GeometryCollection/DerivedDataGeometryCollectionCooker.cpp b/Engine/Source/Runtime/Experimental/GeometryCollectionEngine/Private/GeometryCollection/DerivedDataGeometryCollectionCooker.cpp
--- a/Engine/Source/Runtime/Experimental/GeometryCollectionEngine/Private/GeometryCollection/DerivedDataGeometryCollectionCooker.cpp
+++ b/Engine/Source/Runtime/Experimental/GeometryCollectionEngine/Private/GeometryCollection/DerivedDataGeometryCollectionCooker.cpp
@@ -59,7 +59,7 @@ const TCHAR* FDerivedDataGeometryCollectionCooker::GetVersionString() const
 		return OverrideVersion;	//force load old ddc if found. Not recommended
 	}
 
-	return TEXT("655D3B4051AE4A1E99135F2417195C8C");
+	return TEXT("6C8317B9BC9C4DA1A38AEA126F8F24D0");
 }
 
 FString FDerivedDataGeometryCollectionCooker::GetPluginSpecificCacheKeySuffix() const
