Chaos - RBAN fixed timestep mode (Disabled by default)

diff --git a/Engine/Source/Editor/PhysicsAssetEditor/Private/PhysicsAssetEditorAnimInstanceProxy.cpp b/Engine/Source/Editor/PhysicsAssetEditor/Private/PhysicsAssetEditorAnimInstanceProxy.cpp
--- a/Engine/Source/Editor/PhysicsAssetEditor/Private/PhysicsAssetEditorAnimInstanceProxy.cpp
+++ b/Engine/Source/Editor/PhysicsAssetEditor/Private/PhysicsAssetEditorAnimInstanceProxy.cpp
@@ -79,6 +79,7 @@ bool FPhysicsAssetEditorAnimInstanceProxy::Evaluate_WithRoot(FPoseContext& Outpu
 	if (Simulation != nullptr)
 	{
 		Simulation->SetSolverIterations(
+			SolverIterations.FixedTimeStep,
 			SolverIterations.SolverIterations,
 			SolverIterations.JointIterations,
 			SolverIterations.CollisionIterations,
diff --git a/Engine/Source/Runtime/AnimGraphRuntime/Private/BoneControllers/AnimNode_RigidBody.cpp b/Engine/Source/Runtime/AnimGraphRuntime/Private/BoneControllers/AnimNode_RigidBody.cpp
--- a/Engine/Source/Runtime/AnimGraphRuntime/Private/BoneControllers/AnimNode_RigidBody.cpp
+++ b/Engine/Source/Runtime/AnimGraphRuntime/Private/BoneControllers/AnimNode_RigidBody.cpp
@@ -501,6 +501,7 @@ void FAnimNode_RigidBody::EvaluateSkeletalControl_AnyThread(FComponentSpacePoseC
 			}
 #else
 			PhysicsSimulation->SetSolverIterations(
+				SolverIterations.FixedTimeStep,
 				SolverIterations.SolverIterations,
 				SolverIterations.JointIterations,
 				SolverIterations.CollisionIterations,
@@ -509,7 +510,8 @@ void FAnimNode_RigidBody::EvaluateSkeletalControl_AnyThread(FComponentSpacePoseC
 				SolverIterations.CollisionPushOutIterations
 			);
 			PhysicsSimulation->SetSolverIterations(
-				OverrideSolverIterations.SolverIterations, 
+				SolverIterations.FixedTimeStep,
+				OverrideSolverIterations.SolverIterations,
 				OverrideSolverIterations.JointIterations, 
 				OverrideSolverIterations.CollisionIterations, 
 				OverrideSolverIterations.SolverPushOutIterations, 
@@ -875,6 +877,7 @@ void FAnimNode_RigidBody::InitPhysics(const UAnimInstance* InAnimInstance)
 #if WITH_CHAOS
 		SolverIterations = UsePhysicsAsset->SolverIterations;
 		PhysicsSimulation->SetSolverIterations(
+			SolverIterations.FixedTimeStep,
 			SolverIterations.SolverIterations,
 			SolverIterations.JointIterations,
 			SolverIterations.CollisionIterations,
diff --git a/Engine/Source/Runtime/AnimGraphRuntime/Private/BoneControllers/AnimNode_RigidBody_Chaos.cpp b/Engine/Source/Runtime/AnimGraphRuntime/Private/BoneControllers/AnimNode_RigidBody_Chaos.cpp
--- a/Engine/Source/Runtime/AnimGraphRuntime/Private/BoneControllers/AnimNode_RigidBody_Chaos.cpp
+++ b/Engine/Source/Runtime/AnimGraphRuntime/Private/BoneControllers/AnimNode_RigidBody_Chaos.cpp
@@ -470,6 +470,7 @@ void FAnimNode_RigidBody_Chaos::EvaluateSkeletalControl_AnyThread(FComponentSpac
 			const FVector SimSpaceGravity = WorldVectorToSpaceNoScaleChaos(SimulationSpace, WorldSpaceGravity, CompWorldSpaceTM, BaseBoneTM);
 
 			PhysicsSimulation->SetSolverIterations(
+				SolverIterations.FixedTimeStep,
 				SolverIterations.SolverIterations,
 				SolverIterations.JointIterations,
 				SolverIterations.CollisionIterations,
@@ -478,6 +479,7 @@ void FAnimNode_RigidBody_Chaos::EvaluateSkeletalControl_AnyThread(FComponentSpac
 				SolverIterations.CollisionPushOutIterations
 			);
 			PhysicsSimulation->SetSolverIterations(
+				SolverIterations.FixedTimeStep,
 				OverrideSolverIterations.SolverIterations,
 				OverrideSolverIterations.JointIterations,
 				OverrideSolverIterations.CollisionIterations,
@@ -819,6 +821,7 @@ void FAnimNode_RigidBody_Chaos::InitPhysics(const UAnimInstance* InAnimInstance)
 
 		SolverIterations = UsePhysicsAsset->SolverIterations;
 		PhysicsSimulation->SetSolverIterations(
+			SolverIterations.FixedTimeStep,
 			SolverIterations.SolverIterations,
 			SolverIterations.JointIterations,
 			SolverIterations.CollisionIterations,
diff --git a/Engine/Source/Runtime/Engine/Classes/PhysicsEngine/PhysicsAsset.h b/Engine/Source/Runtime/Engine/Classes/PhysicsEngine/PhysicsAsset.h
--- a/Engine/Source/Runtime/Engine/Classes/PhysicsEngine/PhysicsAsset.h
+++ b/Engine/Source/Runtime/Engine/Classes/PhysicsEngine/PhysicsAsset.h
@@ -25,6 +25,14 @@ struct ENGINE_API FSolverIterations
 
 		FSolverIterations();
 
+	/**
+	 * [Chaos Only]
+	 * The recommended fixed timestep for the solver if supported (e.g., in RigidBody Anim Node). 0 to run with variable timestep.
+	 * NOTE: If this value is non-zero and less than the current frame time, physics will step multiple times.
+	 */
+	UPROPERTY(EditAnywhere, BlueprintReadWrite, Category = SolverSettings, meta = (ClampMin = 0))
+		float FixedTimeStep;
+
 	/**
 	 * [Chaos Only]
 	 * The recommended number of solver iterations. Increase this if collision and joints are fighting, or joint chains are stretching.
diff --git a/Engine/Source/Runtime/Engine/Private/PhysicsEngine/ImmediatePhysics/ImmediatePhysicsChaos/ImmediatePhysicsActorHandle_Chaos.cpp b/Engine/Source/Runtime/Engine/Private/PhysicsEngine/ImmediatePhysics/ImmediatePhysicsChaos/ImmediatePhysicsActorHandle_Chaos.cpp
--- a/Engine/Source/Runtime/Engine/Private/PhysicsEngine/ImmediatePhysics/ImmediatePhysicsChaos/ImmediatePhysicsActorHandle_Chaos.cpp
+++ b/Engine/Source/Runtime/Engine/Private/PhysicsEngine/ImmediatePhysics/ImmediatePhysicsChaos/ImmediatePhysicsActorHandle_Chaos.cpp
@@ -275,8 +275,16 @@ namespace ImmediatePhysics_Chaos
 	// Actor Handle
 	//
 
-	FActorHandle::FActorHandle(Chaos::TPBDRigidsSOAs<FReal, 3>& InParticles, EActorType ActorType, FBodyInstance* BodyInstance, const FTransform& Transform)
+	FActorHandle::FActorHandle(
+		Chaos::TPBDRigidsSOAs<FReal, 3>& InParticles, 
+		Chaos::TArrayCollectionArray<Chaos::FVec3>& InParticlePrevXs, 
+		Chaos::TArrayCollectionArray<Chaos::FRotation3>& InParticlePrevRs, 
+		EActorType ActorType, 
+		FBodyInstance* BodyInstance, 
+		const FTransform& Transform)
 		: Particles(InParticles)
+		, ParticlePrevXs(InParticlePrevXs)
+		, ParticlePrevRs(InParticlePrevRs)
 		, ParticleHandle(nullptr)
 	{
 		using namespace Chaos;
@@ -386,6 +394,8 @@ namespace ImmediatePhysics_Chaos
 		{
 			Dynamic->X() = Dynamic->P();
 			Dynamic->R() = Dynamic->Q();
+			Dynamic->AuxilaryValue(ParticlePrevXs) = Dynamic->P();
+			Dynamic->AuxilaryValue(ParticlePrevRs) = Dynamic->Q();
 		}
 	}
 
diff --git a/Engine/Source/Runtime/Engine/Private/PhysicsEngine/ImmediatePhysics/ImmediatePhysicsChaos/ImmediatePhysicsSimulation_Chaos.cpp b/Engine/Source/Runtime/Engine/Private/PhysicsEngine/ImmediatePhysics/ImmediatePhysicsChaos/ImmediatePhysicsSimulation_Chaos.cpp
--- a/Engine/Source/Runtime/Engine/Private/PhysicsEngine/ImmediatePhysics/ImmediatePhysicsChaos/ImmediatePhysicsSimulation_Chaos.cpp
+++ b/Engine/Source/Runtime/Engine/Private/PhysicsEngine/ImmediatePhysics/ImmediatePhysicsChaos/ImmediatePhysicsSimulation_Chaos.cpp
@@ -41,6 +41,11 @@ FAutoConsoleVariableRef CVarChaosImmPhysIterations(TEXT("p.Chaos.ImmPhys.Iterati
 FAutoConsoleVariableRef CVarChaosImmPhysPushOutIterations(TEXT("p.Chaos.ImmPhys.PushOutIterations"), ChaosImmediate_Evolution_PushOutIterations, TEXT("Override number of solver push-out loops (if >= 0)"));
 FAutoConsoleVariableRef CVarChaosImmPhysBoundsExtension(TEXT("p.Chaos.ImmPhys.BoundsExtension"), ChaosImmediate_Evolution_BoundsExtension, TEXT("Bounds are grown by this fraction of their size (should be >= 0.0)"));
 
+float ChaosImmediate_Evolution_FixedStepTime = -1.0f;
+float ChaosImmediate_Evolution_FixedStepTolerance = 0.05f;
+FAutoConsoleVariableRef CVarChaosImmPhysFixedStepTime(TEXT("p.Chaos.ImmPhys.FixedStepTime"), ChaosImmediate_Evolution_FixedStepTime, TEXT("Override fixed step time mode: fixed step time (if positive); variable time mode (if zero); asset defined (if negative)"));
+FAutoConsoleVariableRef CVarChaosImmPhysFixedStepTolerance(TEXT("p.Chaos.ImmPhys.FixedStepTolerance"), ChaosImmediate_Evolution_FixedStepTolerance, TEXT("Tiem remainder required to add a new step (fraction of FixedStepTime)"));
+
 int32 ChaosImmediate_Collision_Enabled = 1;
 int32 ChaosImmediate_Collision_PairIterations = -1;
 int32 ChaosImmediate_Collision_PushOutPairIterations = -1;
@@ -156,6 +161,8 @@ namespace ImmediatePhysics_Chaos
 		, CollidedParticles()
 		, ParticleMaterials()
 		, PerParticleMaterials()
+		, ParticlePrevXs()
+		, ParticlePrevRs()
 		, Particles()
 		, Joints()
 		, Collisions(Particles, CollidedParticles, ParticleMaterials, 0, 0, ChaosImmediate_Collision_CullDistance, ChaosImmediate_Collision_ShapePadding)
@@ -164,7 +171,7 @@ namespace ImmediatePhysics_Chaos
 		, CollisionDetector(BroadPhase, NarrowPhase, Collisions)
 		, JointsRule(0, Joints)
 		, CollisionsRule(1, Collisions)
-		, Evolution(Particles, CollisionDetector, ChaosImmediate_Evolution_BoundsExtension)
+		, Evolution(Particles, ParticlePrevXs, ParticlePrevRs, CollisionDetector, ChaosImmediate_Evolution_BoundsExtension)
 		, NumActiveDynamicActorHandles(0)
 		, SimulationSpaceTransform(FTransform::Identity)
 		, RollingAverageStepTime(ChaosImmediate_Evolution_InitialStepTime)
@@ -176,6 +183,8 @@ namespace ImmediatePhysics_Chaos
 
 		Particles.GetParticleHandles().AddArray(&ParticleMaterials);
 		Particles.GetParticleHandles().AddArray(&PerParticleMaterials);
+		Particles.GetParticleHandles().AddArray(&ParticlePrevXs);
+		Particles.GetParticleHandles().AddArray(&ParticlePrevRs);
 
 		Evolution.AddConstraintRule(&CollisionsRule);
 		Evolution.AddConstraintRule(&JointsRule);
@@ -291,7 +300,7 @@ namespace ImmediatePhysics_Chaos
 
 		using namespace Chaos;
 
-		FActorHandle* ActorHandle = new FActorHandle(Particles, ActorType, BodyInstance, Transform);
+		FActorHandle* ActorHandle = new FActorHandle(Particles, ParticlePrevXs, ParticlePrevRs, ActorType, BodyInstance, Transform);
 		int ActorIndex = ActorHandles.Add(ActorHandle);
 
 		TUniquePtr<FChaosPhysicsMaterial> Material = MakeUnique<FChaosPhysicsMaterial>();
@@ -308,6 +317,8 @@ namespace ImmediatePhysics_Chaos
 		ParticleMaterials.Add(MakeSerializable(Material));
 		PerParticleMaterials.Add(MoveTemp(Material));
 		CollidedParticles.Add(false);
+		ParticlePrevXs.Add(ActorHandle->GetParticle()->X());
+		ParticlePrevRs.Add(ActorHandle->GetParticle()->R());
 
 		bActorsDirty = true;
 
@@ -325,6 +336,8 @@ namespace ImmediatePhysics_Chaos
 		ParticleMaterials.RemoveAt(Index, 1);
 		PerParticleMaterials.RemoveAt(Index, 1);
 		CollidedParticles.RemoveAt(Index, 1);
+		ParticlePrevXs.RemoveAt(Index, 1);
+		ParticlePrevRs.RemoveAt(Index, 1);
 
 		bActorsDirty = true;
 	}
@@ -485,15 +498,6 @@ namespace ImmediatePhysics_Chaos
 
 		// TEMP: overrides
 		{
-			if (ChaosImmediate_Evolution_StepTime > 0)
-			{
-				StepTime = ChaosImmediate_Evolution_StepTime;
-			}
-			if (ChaosImmediate_Evolution_NumSteps > 0)
-			{
-				NumSteps = ChaosImmediate_Evolution_NumSteps;
-			}
-
 			FPBDJointSolverSettings JointsSettings = Joints.GetSettings();
 			JointsSettings.SwingTwistAngleTolerance = ChaosImmediate_Joint_SwingTwistAngleTolerance;
 			JointsSettings.PositionTolerance = ChaosImmediate_Joint_PositionTolerance;
@@ -538,7 +542,17 @@ namespace ImmediatePhysics_Chaos
 				Collisions.SetApplyType(ECollisionApplyType::Velocity);
 			}
 
+			if (ChaosImmediate_Evolution_StepTime > 0)
+			{
+				StepTime = ChaosImmediate_Evolution_StepTime;
+			}
+			if (ChaosImmediate_Evolution_NumSteps > 0)
+			{
+				NumSteps = ChaosImmediate_Evolution_NumSteps;
+			}
+
 			SetSolverIterations(
+				ChaosImmediate_Evolution_FixedStepTime,
 				ChaosImmediate_Evolution_Iterations,
 				ChaosImmediate_Joint_PairIterations,
 				ChaosImmediate_Collision_PairIterations,
@@ -546,29 +560,50 @@ namespace ImmediatePhysics_Chaos
 				ChaosImmediate_Joint_PushOutPairIterations,
 				ChaosImmediate_Collision_PushOutPairIterations);
 		}
-		UE_LOG(LogChaosJoint, Verbose, TEXT("Simulate Dt = %f Steps %d x %f"), DeltaTime, NumSteps, StepTime);
 
 		DebugDrawKinematicParticles(2, 2, FColor(128, 0, 0));
 		DebugDrawDynamicParticles(2, 2, FColor(192, 192, 0));
 		DebugDrawConstraints(2, 2, 0.7f);
 
+		// Fixed timestep mode
+		// @todo(ccaulfield): integrate this better with the non-fixed step calculate above
+		FReal RewindTime = 0.0f;
+		if (FixedStepTime > 0)
+		{
+			StepTime = FixedStepTime;
+			NumSteps = FMath::FloorToInt(DeltaTime / StepTime);
+			FReal RemainderTime = DeltaTime - NumSteps * StepTime;
+			if (RemainderTime > ChaosImmediate_Evolution_FixedStepTolerance* StepTime)
+			{
+				++NumSteps;
+				RewindTime = StepTime - RemainderTime;
+			}
+			NumSteps = FMath::Max(1, NumSteps);
+		}
+
+		// Handle new or deleted particles
 		if (bActorsDirty)
 		{
 			UpdateActivePotentiallyCollidingPairs();
 			bActorsDirty = false;
 		}
 
+		UE_LOG(LogChaosJoint, Verbose, TEXT("Simulate Dt = %f Steps %d x %f (Rewind %f)"), DeltaTime, NumSteps, StepTime, RewindTime);
 		Evolution.SetGravity(InGravity);
-		
-		Evolution.Advance(StepTime, NumSteps);
+		Evolution.Advance(StepTime, NumSteps, RewindTime);
 
 		DebugDrawKinematicParticles(1, 4, FColor(128, 0, 0));
 		DebugDrawDynamicParticles(1, 3, FColor(255, 255, 0));
 		DebugDrawConstraints(1, 2, 1.0f);
 	}
 
-	void FSimulation::SetSolverIterations(const int32 SolverIts, const int32 JointIts, const int32 CollisionIts, const int32 SolverPushOutIts, const int32 JointPushOutIts, const int32 CollisionPushOutIts)
+	void FSimulation::SetSolverIterations(const FReal InFixedDt, const int32 SolverIts, const int32 JointIts, const int32 CollisionIts, const int32 SolverPushOutIts, const int32 JointPushOutIts, const int32 CollisionPushOutIts)
 	{
+		if (InFixedDt >= 0.0f)
+		{
+			FixedStepTime = InFixedDt;
+		}
+
 		if (SolverIts >= 0)
 		{
 			Evolution.SetNumIterations(SolverIts);
diff --git a/Engine/Source/Runtime/Engine/Private/PhysicsEngine/PhysicsAsset.cpp b/Engine/Source/Runtime/Engine/Private/PhysicsEngine/PhysicsAsset.cpp
--- a/Engine/Source/Runtime/Engine/Private/PhysicsEngine/PhysicsAsset.cpp
+++ b/Engine/Source/Runtime/Engine/Private/PhysicsEngine/PhysicsAsset.cpp
@@ -22,11 +22,12 @@
 #define LOCTEXT_NAMESPACE "PhysicsAsset"
 
 FSolverIterations::FSolverIterations()
-	: SolverIterations(3)
+	: FixedTimeStep(0)
+	, SolverIterations(3)
 	, JointIterations(2)
 	, CollisionIterations(2)
-	, SolverPushOutIterations(0)
-	, JointPushOutIterations(0)
+	, SolverPushOutIterations(1)
+	, JointPushOutIterations(1)
 	, CollisionPushOutIterations(0)
 {
 }
diff --git a/Engine/Source/Runtime/Engine/Public/Physics/ImmediatePhysics/ImmediatePhysicsChaos/ImmediatePhysicsActorHandle_Chaos.h b/Engine/Source/Runtime/Engine/Public/Physics/ImmediatePhysics/ImmediatePhysicsChaos/ImmediatePhysicsActorHandle_Chaos.h
--- a/Engine/Source/Runtime/Engine/Public/Physics/ImmediatePhysics/ImmediatePhysicsChaos/ImmediatePhysicsActorHandle_Chaos.h
+++ b/Engine/Source/Runtime/Engine/Public/Physics/ImmediatePhysics/ImmediatePhysicsChaos/ImmediatePhysicsActorHandle_Chaos.h
@@ -126,13 +126,22 @@ namespace ImmediatePhysics_Chaos
 		friend struct FSimulation;
 		friend struct FJointHandle;
 
-		FActorHandle(Chaos::TPBDRigidsSOAs<FReal, 3>& InParticles, EActorType ActorType, FBodyInstance* BodyInstance, const FTransform& Transform);
+		FActorHandle(
+			Chaos::TPBDRigidsSOAs<FReal, 3>& InParticles,
+			Chaos::TArrayCollectionArray<Chaos::FVec3>& InParticlePrevXs,
+			Chaos::TArrayCollectionArray<Chaos::FRotation3>& InParticlePrevRs,
+			EActorType ActorType,
+			FBodyInstance* BodyInstance,
+			const FTransform& Transform);
+
 
 		Chaos::TGenericParticleHandle<FReal, Dimensions> Handle() const;
 
 		FName Name;
 		Chaos::TPBDRigidsSOAs<FReal, 3>& Particles;
 		Chaos::TGeometryParticleHandle<FReal, Dimensions>* ParticleHandle;
+		Chaos::TArrayCollectionArray<Chaos::FVec3>& ParticlePrevXs;
+		Chaos::TArrayCollectionArray<Chaos::FRotation3>& ParticlePrevRs;
 		TUniquePtr<Chaos::FImplicitObject> Geometry;
 		TArray<TUniquePtr<Chaos::FPerShapeData>> Shapes;
 		int32 Level;
diff --git a/Engine/Source/Runtime/Engine/Public/Physics/ImmediatePhysics/ImmediatePhysicsChaos/ImmediatePhysicsSimulation_Chaos.h b/Engine/Source/Runtime/Engine/Public/Physics/ImmediatePhysics/ImmediatePhysicsChaos/ImmediatePhysicsSimulation_Chaos.h
--- a/Engine/Source/Runtime/Engine/Public/Physics/ImmediatePhysics/ImmediatePhysicsChaos/ImmediatePhysicsSimulation_Chaos.h
+++ b/Engine/Source/Runtime/Engine/Public/Physics/ImmediatePhysics/ImmediatePhysicsChaos/ImmediatePhysicsSimulation_Chaos.h
@@ -69,7 +69,7 @@ namespace ImmediatePhysics_Chaos
 		void SetSimulationSpaceTransform(const FTransform& Transform);
 
 		/** Set new iteration counts. A negative value with leave that iteration count unchanged */
-		void SetSolverIterations(const int32 SolverIts, const int32 JointIts, const int32 CollisionIts, const int32 SolverPushOutIts, const int32 JointPushOutIts, const int32 CollisionPushOutIts);
+		void SetSolverIterations(const FReal FixedDt, const int32 SolverIts, const int32 JointIts, const int32 CollisionIts, const int32 SolverPushOutIts, const int32 JointPushOutIts, const int32 CollisionPushOutIts);
 
 	private:
 		void UpdateActivePotentiallyCollidingPairs();
@@ -90,6 +90,8 @@ namespace ImmediatePhysics_Chaos
 		Chaos::TArrayCollectionArray<bool> CollidedParticles;
 		Chaos::TArrayCollectionArray<Chaos::TSerializablePtr<Chaos::FChaosPhysicsMaterial>> ParticleMaterials;
 		Chaos::TArrayCollectionArray<TUniquePtr<Chaos::FChaosPhysicsMaterial>> PerParticleMaterials;
+		Chaos::TArrayCollectionArray<Chaos::FVec3> ParticlePrevXs;
+		Chaos::TArrayCollectionArray<Chaos::FRotation3> ParticlePrevRs;
 
 		FRigidParticleSOAs Particles;
 		Chaos::FPBDJointConstraints Joints;
@@ -118,6 +120,7 @@ namespace ImmediatePhysics_Chaos
 
 		FTransform SimulationSpaceTransform;
 
+		FReal FixedStepTime;
 		FReal RollingAverageStepTime;
 		int32 NumRollingAverageStepTimes;
 		int32 MaxNumRollingAverageStepTimes;
diff --git a/Engine/Source/Runtime/Experimental/Chaos/Private/Chaos/Evolution/PBDMinEvolution.cpp b/Engine/Source/Runtime/Experimental/Chaos/Private/Chaos/Evolution/PBDMinEvolution.cpp
--- a/Engine/Source/Runtime/Experimental/Chaos/Private/Chaos/Evolution/PBDMinEvolution.cpp
+++ b/Engine/Source/Runtime/Experimental/Chaos/Private/Chaos/Evolution/PBDMinEvolution.cpp
@@ -20,20 +20,37 @@
 
 namespace Chaos
 {
-	DECLARE_CYCLE_STAT(TEXT("MinEvolution::AdvanceOneTimeStep"), STAT_MinEvolution_AdvanceOneTimeStep, STATGROUP_Chaos);
-	DECLARE_CYCLE_STAT(TEXT("MinEvolution::Integrate"), STAT_MinEvolution_Integrate, STATGROUP_Chaos);
-	DECLARE_CYCLE_STAT(TEXT("MinEvolution::KinematicTargets"), STAT_MinEvolution_KinematicTargets, STATGROUP_Chaos);
-	DECLARE_CYCLE_STAT(TEXT("MinEvolution::PrepareConstraints"), STAT_MinEvolution_PrepareConstraints, STATGROUP_Chaos);
-	DECLARE_CYCLE_STAT(TEXT("MinEvolution::UnprepareConstraints"), STAT_MinEvolution_UnprepareConstraints, STATGROUP_Chaos);
-	DECLARE_CYCLE_STAT(TEXT("MinEvolution::ApplyConstraints"), STAT_MinEvolution_ApplyConstraints, STATGROUP_Chaos);
-	DECLARE_CYCLE_STAT(TEXT("MinEvolution::UpdateVelocities"), STAT_MinEvolution_UpdateVelocites, STATGROUP_Chaos);
-	DECLARE_CYCLE_STAT(TEXT("MinEvolution::ApplyPushOut"), STAT_MinEvolution_ApplyPushOut, STATGROUP_Chaos);
-	DECLARE_CYCLE_STAT(TEXT("MinEvolution::DetectCollisions"), STAT_MinEvolution_DetectCollisions, STATGROUP_Chaos);
-	DECLARE_CYCLE_STAT(TEXT("MinEvolution::UpdatePositions"), STAT_MinEvolution_UpdatePositions, STATGROUP_Chaos);
-
-	FPBDMinEvolution::FPBDMinEvolution(FRigidParticleSOAs& InParticles, FCollisionDetector& InCollisionDetector, const FReal InBoundsExtension)
+#if UE_BUILD_SHIPPING || UE_BUILD_TEST
+	CHAOS_API DECLARE_LOG_CATEGORY_EXTERN(LogChaosMinEvolution, Log, Warning);
+#else
+	CHAOS_API DECLARE_LOG_CATEGORY_EXTERN(LogChaosMinEvolution, Log, All);
+#endif
+	DEFINE_LOG_CATEGORY(LogChaosMinEvolution);
+
+	DECLARE_CYCLE_STAT(TEXT("MinEvolution::Advance"), STAT_MinEvolution_Advance, STATGROUP_ChaosMinEvolution);
+	DECLARE_CYCLE_STAT(TEXT("MinEvolution::PrepareTick"), STAT_MinEvolution_PrepareTick, STATGROUP_ChaosMinEvolution);
+	DECLARE_CYCLE_STAT(TEXT("MinEvolution::UnprepareTick"), STAT_MinEvolution_UnprepareTick, STATGROUP_ChaosMinEvolution);
+	DECLARE_CYCLE_STAT(TEXT("MinEvolution::Rewind"), STAT_MinEvolution_Rewind, STATGROUP_ChaosMinEvolution);
+	DECLARE_CYCLE_STAT(TEXT("MinEvolution::AdvanceOneTimeStep"), STAT_MinEvolution_AdvanceOneTimeStep, STATGROUP_ChaosMinEvolution);
+	DECLARE_CYCLE_STAT(TEXT("MinEvolution::Integrate"), STAT_MinEvolution_Integrate, STATGROUP_ChaosMinEvolution);
+	DECLARE_CYCLE_STAT(TEXT("MinEvolution::KinematicTargets"), STAT_MinEvolution_KinematicTargets, STATGROUP_ChaosMinEvolution);
+	DECLARE_CYCLE_STAT(TEXT("MinEvolution::PrepareIteration"), STAT_MinEvolution_PrepareIteration, STATGROUP_ChaosMinEvolution);
+	DECLARE_CYCLE_STAT(TEXT("MinEvolution::UnprepareIteration"), STAT_MinEvolution_UnprepareIteration, STATGROUP_ChaosMinEvolution);
+	DECLARE_CYCLE_STAT(TEXT("MinEvolution::ApplyConstraints"), STAT_MinEvolution_ApplyConstraints, STATGROUP_ChaosMinEvolution);
+	DECLARE_CYCLE_STAT(TEXT("MinEvolution::UpdateVelocities"), STAT_MinEvolution_UpdateVelocites, STATGROUP_ChaosMinEvolution);
+	DECLARE_CYCLE_STAT(TEXT("MinEvolution::ApplyPushOut"), STAT_MinEvolution_ApplyPushOut, STATGROUP_ChaosMinEvolution);
+	DECLARE_CYCLE_STAT(TEXT("MinEvolution::DetectCollisions"), STAT_MinEvolution_DetectCollisions, STATGROUP_ChaosMinEvolution);
+	DECLARE_CYCLE_STAT(TEXT("MinEvolution::UpdatePositions"), STAT_MinEvolution_UpdatePositions, STATGROUP_ChaosMinEvolution);
+
+
+	bool bChaos_MinEvolution_RewindLerp = true;
+	FAutoConsoleVariableRef CVarChaosMinEvolutionRewindLerp(TEXT("p.Chaos.MinEvolution.RewindLerp"), bChaos_MinEvolution_RewindLerp, TEXT("If rewinding (fixed dt mode) use Backwards-Lerp as opposed to Backwards Velocity"));
+
+	FPBDMinEvolution::FPBDMinEvolution(FRigidParticleSOAs& InParticles, TArrayCollectionArray<FVec3>& InPrevX, TArrayCollectionArray<FRotation3>& InPrevR, FCollisionDetector& InCollisionDetector, const FReal InBoundsExtension)
 		: Particles(InParticles)
 		, CollisionDetector(InCollisionDetector)
+		, ParticlePrevXs(InPrevX)
+		, ParticlePrevRs(InPrevR)
 		, NumApplyIterations(0)
 		, NumApplyPushOutIterations(0)
 		, BoundsExtension(InBoundsExtension)
@@ -46,8 +63,17 @@ namespace Chaos
 		ConstraintRules.Add(Rule);
 	}
 
-	void FPBDMinEvolution::Advance(const FReal StepDt, const int32 NumSteps)
+	void FPBDMinEvolution::Advance(const FReal StepDt, const int32 NumSteps, const FReal RewindDt)
 	{
+		SCOPE_CYCLE_COUNTER(STAT_MinEvolution_Advance);
+
+		PrepareTick();
+
+		if (RewindDt > SMALL_NUMBER)
+		{
+			Rewind(StepDt, RewindDt);
+		}
+
 		for (TTransientPBDRigidParticleHandle<FReal, 3>& Particle : Particles.GetActiveParticlesView())
 		{
 			if (Particle.ObjectState() == EObjectStateType::Dynamic)
@@ -62,7 +88,7 @@ namespace Chaos
 			// E.g., for 4 steps this will be: 1/4, 1/3, 1/2, 1
 			const float StepFraction = (FReal)1 / (FReal)(NumSteps - Step);
 
-			UE_LOG(LogChaos, Verbose, TEXT("Advance dt = %f [%d/%d]"), StepDt, Step + 1, NumSteps);
+			UE_LOG(LogChaosMinEvolution, Verbose, TEXT("Advance dt = %f [%d/%d]"), StepDt, Step + 1, NumSteps);
 
 			AdvanceOneTimeStep(StepDt, StepFraction);
 		}
@@ -75,6 +101,8 @@ namespace Chaos
 				Particle.Torque() = FVec3(0);
 			}
 		}
+
+		UnprepareTick();
 	}
 
 	void FPBDMinEvolution::AdvanceOneTimeStep(const FReal Dt, const FReal StepFraction)
@@ -99,7 +127,7 @@ namespace Chaos
 
 		if (Dt > 0)
 		{
-			PrepareConstraints(Dt);
+			PrepareIteration(Dt);
 
 			ApplyConstraints(Dt);
 
@@ -117,12 +145,75 @@ namespace Chaos
 				PostApplyPushOutCallback();
 			}
 
-			UnprepareConstraints(Dt);
+			UnprepareIteration(Dt);
 
 			UpdatePositions(Dt);
 		}
 	}
 
+	// A opportunity for systems to allocate buffers for the duration of the tick, if they have enough info to do so
+	void FPBDMinEvolution::PrepareTick()
+	{
+		SCOPE_CYCLE_COUNTER(STAT_MinEvolution_PrepareTick);
+
+		for (FSimpleConstraintRule* ConstraintRule : ConstraintRules)
+		{
+			ConstraintRule->PrepareTick();
+		}
+	}
+
+	void FPBDMinEvolution::UnprepareTick()
+	{
+		SCOPE_CYCLE_COUNTER(STAT_MinEvolution_UnprepareTick);
+
+		for (FSimpleConstraintRule* ConstraintRule : ConstraintRules)
+		{
+			ConstraintRule->UnprepareTick();
+		}
+	}
+
+	// Update X/R as if we started the next tick 'RewindDt' seconds ago.
+	void FPBDMinEvolution::Rewind(FReal Dt, FReal RewindDt)
+	{
+		SCOPE_CYCLE_COUNTER(STAT_MinEvolution_Rewind);
+
+		if (bChaos_MinEvolution_RewindLerp)
+		{
+			const FReal T = (Dt - RewindDt) / Dt;
+			UE_LOG(LogChaosMinEvolution, Verbose, TEXT("Rewind dt = %f; rt = %f; T = %f"), Dt, RewindDt, T);
+			for (TTransientPBDRigidParticleHandle<FReal, 3>& Particle : Particles.GetActiveParticlesView())
+			{
+				if (Particle.ObjectState() == EObjectStateType::Dynamic)
+				{
+					Particle.X() = FVec3::Lerp(Particle.Handle()->AuxilaryValue(ParticlePrevXs), Particle.X(), T);
+					Particle.R() = FRotation3::Slerp(Particle.Handle()->AuxilaryValue(ParticlePrevRs), Particle.R(), T);
+				}
+			}
+		}
+		else
+		{
+			for (TTransientPBDRigidParticleHandle<FReal, 3>& Particle : Particles.GetActiveParticlesView())
+			{
+				if (Particle.ObjectState() == EObjectStateType::Dynamic)
+				{
+					const FVec3 XCoM = FParticleUtilitiesXR::GetCoMWorldPosition(&Particle);
+					const FRotation3 RCoM = FParticleUtilitiesXR::GetCoMWorldRotation(&Particle);
+
+					const FVec3 XCoM2 = XCoM - Particle.V() * RewindDt;
+					const FRotation3 RCoM2 = FRotation3::IntegrateRotationWithAngularVelocity(RCoM, -Particle.W(), RewindDt);
+
+					FParticleUtilitiesXR::SetCoMWorldTransform(&Particle, XCoM2, RCoM2);
+				}
+			}
+		}
+
+		for (auto& Particle : Particles.GetActiveKinematicParticlesView())
+		{
+			Particle.X() = Particle.X() - Particle.V() * RewindDt;
+			Particle.R() = FRotation3::IntegrateRotationWithAngularVelocity(Particle.R(), -Particle.W(), RewindDt);
+		}
+	}
+
 	// @todo(ccaulfield): dedupe (PBDRigidsEvolutionGBF)
 	void FPBDMinEvolution::Integrate(FReal Dt)
 	{
@@ -277,23 +368,23 @@ namespace Chaos
 		CollisionDetector.DetectCollisions(Dt);
 	}
 
-	void FPBDMinEvolution::PrepareConstraints(FReal Dt)
+	void FPBDMinEvolution::PrepareIteration(FReal Dt)
 	{
-		SCOPE_CYCLE_COUNTER(STAT_MinEvolution_PrepareConstraints);
+		SCOPE_CYCLE_COUNTER(STAT_MinEvolution_PrepareIteration);
 
-		for (FSimpleConstraintRule* ConstraintRule : PrioritizedConstraintRules)
+		for (FSimpleConstraintRule* ConstraintRule : ConstraintRules)
 		{
-			ConstraintRule->PrepareConstraints(Dt);
+			ConstraintRule->PrepareIteration(Dt);
 		}
 	}
 
-	void FPBDMinEvolution::UnprepareConstraints(FReal Dt)
+	void FPBDMinEvolution::UnprepareIteration(FReal Dt)
 	{
-		SCOPE_CYCLE_COUNTER(STAT_MinEvolution_PrepareConstraints);
+		SCOPE_CYCLE_COUNTER(STAT_MinEvolution_UnprepareIteration);
 
-		for (FSimpleConstraintRule* ConstraintRule : PrioritizedConstraintRules)
+		for (FSimpleConstraintRule* ConstraintRule : ConstraintRules)
 		{
-			ConstraintRule->UnprepareConstraints(Dt);
+			ConstraintRule->UnprepareIteration(Dt);
 		}
 	}
 
@@ -351,6 +442,8 @@ namespace Chaos
 		SCOPE_CYCLE_COUNTER(STAT_MinEvolution_UpdatePositions);
 		for (auto& Particle : Particles.GetActiveParticlesView())
 		{
+			Particle.Handle()->AuxilaryValue(ParticlePrevXs) = Particle.X();
+			Particle.Handle()->AuxilaryValue(ParticlePrevRs) = Particle.R();
 			Particle.X() = Particle.P();
 			Particle.R() = Particle.Q();
 		}
diff --git a/Engine/Source/Runtime/Experimental/Chaos/Private/Chaos/PBDJointConstraints.cpp b/Engine/Source/Runtime/Experimental/Chaos/Private/Chaos/PBDJointConstraints.cpp
--- a/Engine/Source/Runtime/Experimental/Chaos/Private/Chaos/PBDJointConstraints.cpp
+++ b/Engine/Source/Runtime/Experimental/Chaos/Private/Chaos/PBDJointConstraints.cpp
@@ -36,11 +36,14 @@ FAutoConsoleVariableRef CVarChaosJointISPCEnabled(TEXT("p.Chaos.Joint.ISPC"), bC
 namespace Chaos
 {
 	DECLARE_CYCLE_STAT(TEXT("Joints::Sort"), STAT_Joints_Sort, STATGROUP_ChaosJoint);
-	DECLARE_CYCLE_STAT(TEXT("Joints::Prepare"), STAT_Joints_Prepare, STATGROUP_ChaosJoint);
-	DECLARE_CYCLE_STAT(TEXT("Joints::Unprepare"), STAT_Joints_Unprepare, STATGROUP_ChaosJoint);
+	DECLARE_CYCLE_STAT(TEXT("Joints::PrepareTick"), STAT_Joints_PrepareTick, STATGROUP_ChaosJoint);
+	DECLARE_CYCLE_STAT(TEXT("Joints::UnprepareTick"), STAT_Joints_UnprepareTick, STATGROUP_ChaosJoint);
+	DECLARE_CYCLE_STAT(TEXT("Joints::PrepareIterations"), STAT_Joints_PrepareIteration, STATGROUP_ChaosJoint);
+	DECLARE_CYCLE_STAT(TEXT("Joints::UnprepareIteration"), STAT_Joints_UnprepareIteration, STATGROUP_ChaosJoint);
 	DECLARE_CYCLE_STAT(TEXT("Joints::Apply"), STAT_Joints_Apply, STATGROUP_ChaosJoint);
 	DECLARE_CYCLE_STAT(TEXT("Joints::ApplyBatched"), STAT_Joints_ApplyBatched, STATGROUP_ChaosJoint);
 	DECLARE_CYCLE_STAT(TEXT("Joints::ApplyPushOut"), STAT_Joints_ApplyPushOut, STATGROUP_ChaosJoint);
+	DECLARE_CYCLE_STAT(TEXT("Joints::ApplyPushOutBatched"), STAT_Joints_ApplyPushOutBatched, STATGROUP_ChaosJoint);
 
 	//
 	// Constraint Handle
@@ -508,30 +511,81 @@ namespace Chaos
 		}
 	}
 
-
-	void FPBDJointConstraints::PrepareConstraints(FReal Dt)
+	void FPBDJointConstraints::PrepareTick()
 	{
-		SCOPE_CYCLE_COUNTER(STAT_Joints_Prepare);
+		SCOPE_CYCLE_COUNTER(STAT_Joints_PrepareTick);
 
 		if (bChaos_Joint_Batching)
 		{
-			InitSolverJointState();
+			SolverConstraintRowStates.SetNum(SolverConstraintRowDatas.Num());
+			SolverConstraintStates.SetNum(NumConstraints());
 		}
 		else
 		{
 			ConstraintSolvers.SetNum(NumConstraints());
+		}
+	}
 
-			for (int32 ConstraintIndex = 0; ConstraintIndex < NumConstraints(); ++ConstraintIndex)
+	void FPBDJointConstraints::UnprepareTick()
+	{
+		SCOPE_CYCLE_COUNTER(STAT_Joints_UnprepareTick);
+
+		if (bChaos_Joint_Batching)
+		{
+			SolverConstraintRowStates.Empty();
+			SolverConstraintStates.Empty();
+		}
+		else
+		{
+			ConstraintSolvers.Empty();
+		}
+	}
+
+	void FPBDJointConstraints::PrepareIteration(FReal Dt)
+	{
+		SCOPE_CYCLE_COUNTER(STAT_Joints_PrepareIteration);
+
+		if (bChaos_Joint_Batching)
+		{
+			for (int32 JointIndex = 0; JointIndex < NumConstraints(); ++JointIndex)
 			{
-				const FPBDJointSettings& JointSettings = ConstraintSettings[ConstraintIndex];
+				FJointSolverJointState& JointState = SolverConstraintStates[JointIndex];
+				const FPBDJointSettings& JointSettings = ConstraintSettings[JointIndex];
 
-				const FTransformPair& JointFrames = ConstraintFrames[ConstraintIndex];
-				FJointSolverGaussSeidel& Solver = ConstraintSolvers[ConstraintIndex];
+				const FTransformPair& JointFrames = ConstraintFrames[JointIndex];
+				int32 Index0, Index1;
+				GetConstrainedParticleIndices(JointIndex, Index0, Index1);
+				TGenericParticleHandle<FReal, 3> Particle0 = TGenericParticleHandle<FReal, 3>(ConstraintParticles[JointIndex][Index0]);
+				TGenericParticleHandle<FReal, 3> Particle1 = TGenericParticleHandle<FReal, 3>(ConstraintParticles[JointIndex][Index1]);
+
+				JointState.Init(
+					Settings,
+					JointSettings,
+					FParticleUtilitiesXR::GetCoMWorldPosition(Particle0),
+					FParticleUtilitiesXR::GetCoMWorldRotation(Particle0),
+					FParticleUtilitiesXR::GetCoMWorldPosition(Particle1),
+					FParticleUtilitiesXR::GetCoMWorldRotation(Particle1),
+					Particle0->InvM(),
+					Particle0->InvI().GetDiagonal(),
+					Particle1->InvM(),
+					Particle1->InvI().GetDiagonal(),
+					FParticleUtilities::ParticleLocalToCoMLocal(Particle0, JointFrames[Index0]),
+					FParticleUtilities::ParticleLocalToCoMLocal(Particle1, JointFrames[Index1]));
+			}
+		}
+		else
+		{
+			for (int32 JointIndex = 0; JointIndex < NumConstraints(); ++JointIndex)
+			{
+				const FPBDJointSettings& JointSettings = ConstraintSettings[JointIndex];
+
+				const FTransformPair& JointFrames = ConstraintFrames[JointIndex];
+				FJointSolverGaussSeidel& Solver = ConstraintSolvers[JointIndex];
 
 				int32 Index0, Index1;
-				GetConstrainedParticleIndices(ConstraintIndex, Index0, Index1);
-				TGenericParticleHandle<FReal, 3> Particle0 = TGenericParticleHandle<FReal, 3>(ConstraintParticles[ConstraintIndex][Index0]);
-				TGenericParticleHandle<FReal, 3> Particle1 = TGenericParticleHandle<FReal, 3>(ConstraintParticles[ConstraintIndex][Index1]);
+				GetConstrainedParticleIndices(JointIndex, Index0, Index1);
+				TGenericParticleHandle<FReal, 3> Particle0 = TGenericParticleHandle<FReal, 3>(ConstraintParticles[JointIndex][Index0]);
+				TGenericParticleHandle<FReal, 3> Particle1 = TGenericParticleHandle<FReal, 3>(ConstraintParticles[JointIndex][Index1]);
 
 				Solver.Init(
 					Dt,
@@ -552,18 +606,9 @@ namespace Chaos
 	}
 
 
-	void FPBDJointConstraints::UnprepareConstraints(FReal Dt)
+	void FPBDJointConstraints::UnprepareIteration(FReal Dt)
 	{
-		SCOPE_CYCLE_COUNTER(STAT_Joints_Unprepare);
-
-		if (bChaos_Joint_Batching)
-		{
-			DeinitSolverJointState();
-		}
-		else
-		{
-			ConstraintSolvers.Empty();
-		}
+		SCOPE_CYCLE_COUNTER(STAT_Joints_UnprepareIteration);
 	}
 
 	
@@ -632,14 +677,22 @@ namespace Chaos
 
 	bool FPBDJointConstraints::ApplyPushOut(const FReal Dt, const int32 It, const int32 NumIts)
 	{
-		SCOPE_CYCLE_COUNTER(STAT_Joints_ApplyPushOut);
 
 		int32 NumActive = 0;
 		if (Settings.ApplyPushOutPairIterations > 0)
 		{
-			for (int32 ConstraintIndex = 0; ConstraintIndex < NumConstraints(); ++ConstraintIndex)
+			if (bChaos_Joint_Batching)
 			{
-				NumActive += ApplyPushOutSingle(Dt, ConstraintIndex, Settings.ApplyPushOutPairIterations, It, NumIts);
+				SCOPE_CYCLE_COUNTER(STAT_Joints_ApplyPushOutBatched);
+				// @todo(ccaulfield): batch mode pushout
+			}
+			else
+			{
+				SCOPE_CYCLE_COUNTER(STAT_Joints_ApplyPushOut);
+				for (int32 ConstraintIndex = 0; ConstraintIndex < NumConstraints(); ++ConstraintIndex)
+				{
+					NumActive += ApplyPushOutSingle(Dt, ConstraintIndex, Settings.ApplyPushOutPairIterations, It, NumIts);
+				}
 			}
 		}
 
@@ -795,44 +848,6 @@ namespace Chaos
 		SolverConstraintRowDatas.Empty();
 	}
 
-	void FPBDJointConstraints::InitSolverJointState()
-	{
-		SolverConstraintRowStates.SetNum(SolverConstraintRowDatas.Num());
-		SolverConstraintStates.SetNum(NumConstraints());
-
-		for (int32 JointIndex = 0; JointIndex < NumConstraints(); ++JointIndex)
-		{
-			FJointSolverJointState& JointState = SolverConstraintStates[JointIndex];
-			const FPBDJointSettings& JointSettings = ConstraintSettings[JointIndex];
-		
-			const FTransformPair& JointFrames = ConstraintFrames[JointIndex];
-			int32 Index0, Index1;
-			GetConstrainedParticleIndices(JointIndex, Index0, Index1);
-			TGenericParticleHandle<FReal, 3> Particle0 = TGenericParticleHandle<FReal, 3>(ConstraintParticles[JointIndex][Index0]);
-			TGenericParticleHandle<FReal, 3> Particle1 = TGenericParticleHandle<FReal, 3>(ConstraintParticles[JointIndex][Index1]);
-
-			JointState.Init(
-				Settings,
-				JointSettings,
-				FParticleUtilitiesXR::GetCoMWorldPosition(Particle0),
-				FParticleUtilitiesXR::GetCoMWorldRotation(Particle0),
-				FParticleUtilitiesXR::GetCoMWorldPosition(Particle1),
-				FParticleUtilitiesXR::GetCoMWorldRotation(Particle1),
-				Particle0->InvM(),
-				Particle0->InvI().GetDiagonal(),
-				Particle1->InvM(),
-				Particle1->InvI().GetDiagonal(),
-				FParticleUtilities::ParticleLocalToCoMLocal(Particle0, JointFrames[Index0]),
-				FParticleUtilities::ParticleLocalToCoMLocal(Particle1, JointFrames[Index1]));
-		}
-	}
-
-	void FPBDJointConstraints::DeinitSolverJointState()
-	{
-		SolverConstraintRowStates.Empty();
-		SolverConstraintStates.Empty();
-	}
-
 	void FPBDJointConstraints::GatherSolverJointState(int32 JointIndex)
 	{
 		FJointSolverJointState& JointState = SolverConstraintStates[JointIndex];
diff --git a/Engine/Source/Runtime/Experimental/Chaos/Private/Chaos/PBDRigidsEvolutionGBF.cpp b/Engine/Source/Runtime/Experimental/Chaos/Private/Chaos/PBDRigidsEvolutionGBF.cpp
--- a/Engine/Source/Runtime/Experimental/Chaos/Private/Chaos/PBDRigidsEvolutionGBF.cpp
+++ b/Engine/Source/Runtime/Experimental/Chaos/Private/Chaos/PBDRigidsEvolutionGBF.cpp
@@ -115,6 +115,8 @@ void FPBDRigidsEvolutionGBF::Advance(const FReal Dt, const FReal MaxStepDt, cons
 	int32 NumSteps = FMath::CeilToInt(Dt / MaxStepDt);
 	if (NumSteps > 0)
 	{
+		PrepareTick();
+
 		// Determine the step time
 		const FReal StepDt = Dt / (FReal)NumSteps;
 
@@ -133,6 +135,8 @@ void FPBDRigidsEvolutionGBF::Advance(const FReal Dt, const FReal MaxStepDt, cons
 
 			AdvanceOneTimeStep(StepDt, StepFraction);
 		}
+
+		UnprepareTick();
 	}
 }
 
@@ -309,7 +313,7 @@ void FPBDRigidsEvolutionGBF::AdvanceOneTimeStep(const FReal Dt, const FReal Step
 
 	{
 		SCOPE_CYCLE_COUNTER(STAT_Evolution_PrepareConstraints);
-		PrepareConstraints(Dt);
+		PrepareIteration(Dt);
 	}
 
 	{
@@ -402,7 +406,7 @@ void FPBDRigidsEvolutionGBF::AdvanceOneTimeStep(const FReal Dt, const FReal Step
 
 	{
 		SCOPE_CYCLE_COUNTER(STAT_Evolution_UnprepareConstraints);
-		UnprepareConstraints(Dt);
+		UnprepareIteration(Dt);
 	}
 
 	{
diff --git a/Engine/Source/Runtime/Experimental/Chaos/Public/Chaos/Evolution/PBDMinEvolution.h b/Engine/Source/Runtime/Experimental/Chaos/Public/Chaos/Evolution/PBDMinEvolution.h
--- a/Engine/Source/Runtime/Experimental/Chaos/Public/Chaos/Evolution/PBDMinEvolution.h
+++ b/Engine/Source/Runtime/Experimental/Chaos/Public/Chaos/Evolution/PBDMinEvolution.h
@@ -3,6 +3,7 @@
 
 #include "Chaos/Core.h"
 #include "Chaos/ParticleHandleFwd.h"
+#include "Chaos/ArrayCollectionArray.h"
 
 
 namespace Chaos
@@ -31,11 +32,11 @@ namespace Chaos
 		using FEvolutionCallback = TFunction<void()>;
 		using FRigidParticleSOAs = TPBDRigidsSOAs<FReal, 3>;
 
-		FPBDMinEvolution(FRigidParticleSOAs& InParticles, FCollisionDetector& InCollisionDetector, const FReal InBoundsExtension);
+		FPBDMinEvolution(FRigidParticleSOAs& InParticles, TArrayCollectionArray<FVec3>& InPrevX, TArrayCollectionArray<FRotation3>& InPrevR, FCollisionDetector& InCollisionDetector, const FReal InBoundsExtension);
 
 		void AddConstraintRule(FSimpleConstraintRule* Rule);
 
-		void Advance(const FReal StepDt, const int32 NumSteps);
+		void Advance(const FReal StepDt, const int32 NumSteps, const FReal RewindDt);
 		void AdvanceOneTimeStep(const FReal Dt, const FReal StepFraction);
 
 		void SetNumIterations(const int32 NumIts)
@@ -79,11 +80,14 @@ namespace Chaos
 		}
 
 	private:
+		void PrepareTick();
+		void UnprepareTick();
+		void Rewind(FReal Dt, FReal RewindDt);
 		void Integrate(FReal Dt);
 		void ApplyKinematicTargets(FReal Dt, FReal StepFraction);
 		void DetectCollisions(FReal Dt);
-		void PrepareConstraints(FReal Dt);
-		void UnprepareConstraints(FReal Dt);
+		void PrepareIteration(FReal Dt);
+		void UnprepareIteration(FReal Dt);
 		void ApplyConstraints(FReal Dt);
 		void UpdateVelocities(FReal Dt);
 		void ApplyPushOutConstraints(FReal Dt);
@@ -92,6 +96,9 @@ namespace Chaos
 		FRigidParticleSOAs& Particles;
 		FCollisionDetector& CollisionDetector;
 
+		TArrayCollectionArray<FVec3>& ParticlePrevXs;
+		TArrayCollectionArray<FRotation3>& ParticlePrevRs;
+
 		TArray<FSimpleConstraintRule*> ConstraintRules;
 		TArray<FSimpleConstraintRule*> PrioritizedConstraintRules;
 
diff --git a/Engine/Source/Runtime/Experimental/Chaos/Public/Chaos/PBDCollisionConstraints.h b/Engine/Source/Runtime/Experimental/Chaos/Public/Chaos/PBDCollisionConstraints.h
--- a/Engine/Source/Runtime/Experimental/Chaos/Public/Chaos/PBDCollisionConstraints.h
+++ b/Engine/Source/Runtime/Experimental/Chaos/Public/Chaos/PBDCollisionConstraints.h
@@ -137,9 +137,13 @@ public:
 	// General Rule API
 	//
 
-	void PrepareConstraints(FReal Dt) {}
+	void PrepareTick() {}
 
-	void UnprepareConstraints(FReal Dt) {}
+	void UnprepareTick() {}
+
+	void PrepareIteration(FReal Dt) {}
+
+	void UnprepareIteration(FReal Dt) {}
 
 	/**
 	 * Generate all contact constraints.
diff --git a/Engine/Source/Runtime/Experimental/Chaos/Public/Chaos/PBDConstraintRule.h b/Engine/Source/Runtime/Experimental/Chaos/Public/Chaos/PBDConstraintRule.h
--- a/Engine/Source/Runtime/Experimental/Chaos/Public/Chaos/PBDConstraintRule.h
+++ b/Engine/Source/Runtime/Experimental/Chaos/Public/Chaos/PBDConstraintRule.h
@@ -35,11 +35,18 @@ namespace Chaos
 			return L.GetPriority() < R.GetPriority();
 		}
 
-		/** Called once per frame before Apply. Can be used to prepare caches etc. */
-		virtual void PrepareConstraints(FReal Dt) {}
+		/** Called once per frame. Can be used to prepare caches etc. */
+		virtual void PrepareTick() {}
+
+		/** Called once per frame. Should undo whatever is done in PrepareTick (can also free any other transient buffers created after) */
+		virtual void UnprepareTick() {}
+
+		/** Called once per iteration before Apply. Can be used to prepare caches etc. */
+		virtual void PrepareIteration(FReal Dt) {}
+
+		/** Called once per iteration after Apply. Should be used to release any transient stores created in PrepareConstraints. */
+		virtual void UnprepareIteration(FReal Dt) {}
 
-		/** Called once per frame after Apply. Should be used to release any transient stores created in PrepareConstraints. */
-		virtual void UnprepareConstraints(FReal Dt) {}
 
 	protected:
 		int32 Priority;
@@ -70,14 +77,24 @@ namespace Chaos
 		{
 		}
 
-		virtual void PrepareConstraints(FReal Dt) override
+		virtual void PrepareTick() override
+		{
+			Constraints.PrepareTick();
+		}
+
+		virtual void UnprepareTick() override
 		{
-			Constraints.PrepareConstraints(Dt);
+			Constraints.UnprepareTick();
 		}
 
-		virtual void UnprepareConstraints(FReal Dt) override
+		virtual void PrepareIteration(FReal Dt) override
 		{
-			Constraints.UnprepareConstraints(Dt);
+			Constraints.PrepareIteration(Dt);
+		}
+
+		virtual void UnprepareIteration(FReal Dt) override
+		{
+			Constraints.UnprepareIteration(Dt);
 		}
 
 		virtual void UpdatePositionBasedState(const FReal Dt) override
@@ -153,14 +170,24 @@ namespace Chaos
 
 		TPBDConstraintGraphRuleImpl(FConstraints& InConstraints, int32 InPriority);
 
-		virtual void PrepareConstraints(FReal Dt) override
+		virtual void PrepareTick() override
+		{
+			Constraints.PrepareTick();
+		}
+
+		virtual void UnprepareTick() override
+		{
+			Constraints.UnprepareTick();
+		}
+
+		virtual void PrepareIteration(FReal Dt) override
 		{
-			Constraints.PrepareConstraints(Dt);
+			Constraints.PrepareIteration(Dt);
 		}
 
-		virtual void UnprepareConstraints(FReal Dt) override
+		virtual void UnprepareIteration(FReal Dt) override
 		{
-			Constraints.UnprepareConstraints(Dt);
+			Constraints.UnprepareIteration(Dt);
 		}
 
 		virtual void BindToGraph(FPBDConstraintGraph& InContactGraph, uint32 InContainerId) override;
diff --git a/Engine/Source/Runtime/Experimental/Chaos/Public/Chaos/PBDJointConstraints.h b/Engine/Source/Runtime/Experimental/Chaos/Public/Chaos/PBDJointConstraints.h
--- a/Engine/Source/Runtime/Experimental/Chaos/Public/Chaos/PBDJointConstraints.h
+++ b/Engine/Source/Runtime/Experimental/Chaos/Public/Chaos/PBDJointConstraints.h
@@ -147,8 +147,13 @@ namespace Chaos
 		// General Rule API
 		//
 
-		void PrepareConstraints(FReal Dt);
-		void UnprepareConstraints(FReal Dt);
+		void PrepareTick();
+
+		void UnprepareTick();
+
+		void PrepareIteration(FReal Dt);
+
+		void UnprepareIteration(FReal Dt);
 
 		void UpdatePositionBasedState(const FReal Dt);
 
@@ -181,8 +186,6 @@ namespace Chaos
 		
 		void InitSolverJointData();
 		void DeinitSolverJointData();
-		void InitSolverJointState();
-		void DeinitSolverJointState();
 		void GatherSolverJointState(int32 ConstraintIndex);
 		void ScatterSolverJointState(const FReal Dt, int32 ConstraintIndex);
 
diff --git a/Engine/Source/Runtime/Experimental/Chaos/Public/Chaos/PBDPositionConstraints.h b/Engine/Source/Runtime/Experimental/Chaos/Public/Chaos/PBDPositionConstraints.h
--- a/Engine/Source/Runtime/Experimental/Chaos/Public/Chaos/PBDPositionConstraints.h
+++ b/Engine/Source/Runtime/Experimental/Chaos/Public/Chaos/PBDPositionConstraints.h
@@ -163,17 +163,15 @@ namespace Chaos
 		// Island Rule API
 		//
 
-		void PrepareConstraints(FReal Dt)
-		{
-		}
+		void PrepareTick() {}
 
-		void UnprepareConstraints(FReal Dt)
-		{
-		}
+		void UnprepareTick() {}
 
-		void UpdatePositionBasedState(const T Dt)
-		{
-		}
+		void PrepareIteration(FReal Dt) {}
+
+		void UnprepareIteration(FReal Dt) {}
+
+		void UpdatePositionBasedState(const T Dt) {}
 
 		bool Apply(const T Dt, const TArray<FConstraintContainerHandle*>& ConstraintHandles, const int32 It, const int32 NumIts) const;
 
diff --git a/Engine/Source/Runtime/Experimental/Chaos/Public/Chaos/PBDRigidDynamicSpringConstraints.h b/Engine/Source/Runtime/Experimental/Chaos/Public/Chaos/PBDRigidDynamicSpringConstraints.h
--- a/Engine/Source/Runtime/Experimental/Chaos/Public/Chaos/PBDRigidDynamicSpringConstraints.h
+++ b/Engine/Source/Runtime/Experimental/Chaos/Public/Chaos/PBDRigidDynamicSpringConstraints.h
@@ -173,9 +173,13 @@ namespace Chaos
 		// Island Rule API
 		//
 
-		void PrepareConstraints(FReal Dt) {}
+		void PrepareTick() {}
 
-		void UnprepareConstraints(FReal Dt) {}
+		void UnprepareTick() {}
+
+		void PrepareIteration(FReal Dt) {}
+
+		void UnprepareIteration(FReal Dt) {}
 
 		void UpdatePositionBasedState(const T Dt);
 
diff --git a/Engine/Source/Runtime/Experimental/Chaos/Public/Chaos/PBDRigidSpringConstraints.h b/Engine/Source/Runtime/Experimental/Chaos/Public/Chaos/PBDRigidSpringConstraints.h
--- a/Engine/Source/Runtime/Experimental/Chaos/Public/Chaos/PBDRigidSpringConstraints.h
+++ b/Engine/Source/Runtime/Experimental/Chaos/Public/Chaos/PBDRigidSpringConstraints.h
@@ -147,17 +147,15 @@ namespace Chaos
 		// Island Rule API
 		//
 
-		void PrepareConstraints(FReal Dt)
-		{
-		}
+		void PrepareTick() {}
 
-		void UnprepareConstraints(FReal Dt)
-		{
-		}
+		void UnprepareTick() {}
 
-		void UpdatePositionBasedState(const FReal Dt)
-		{
-		}
+		void PrepareIteration(FReal Dt) {}
+
+		void UnprepareIteration(FReal Dt) {}
+
+		void UpdatePositionBasedState(const FReal Dt) {}
 
 		bool Apply(const FReal Dt, const int32 It, const int32 NumIts);
 
diff --git a/Engine/Source/Runtime/Experimental/Chaos/Public/Chaos/PBDRigidsEvolution.h b/Engine/Source/Runtime/Experimental/Chaos/Public/Chaos/PBDRigidsEvolution.h
--- a/Engine/Source/Runtime/Experimental/Chaos/Public/Chaos/PBDRigidsEvolution.h
+++ b/Engine/Source/Runtime/Experimental/Chaos/Public/Chaos/PBDRigidsEvolution.h
@@ -467,19 +467,35 @@ class FPBDRigidsEvolutionBase
 		}
 	}
 
-	void PrepareConstraints(const FReal Dt)
+	void PrepareTick()
 	{
 		for (FPBDConstraintGraphRule* ConstraintRule : ConstraintRules)
 		{
-			ConstraintRule->PrepareConstraints(Dt);
+			ConstraintRule->PrepareTick();
 		}
 	}
 
-	void UnprepareConstraints(const FReal Dt)
+	void UnprepareTick()
 	{
 		for (FPBDConstraintGraphRule* ConstraintRule : ConstraintRules)
 		{
-			ConstraintRule->UnprepareConstraints(Dt);
+			ConstraintRule->UnprepareTick();
+		}
+	}
+
+	void PrepareIteration(const FReal Dt)
+	{
+		for (FPBDConstraintGraphRule* ConstraintRule : ConstraintRules)
+		{
+			ConstraintRule->PrepareIteration(Dt);
+		}
+	}
+
+	void UnprepareIteration(const FReal Dt)
+	{
+		for (FPBDConstraintGraphRule* ConstraintRule : ConstraintRules)
+		{
+			ConstraintRule->UnprepareIteration(Dt);
 		}
 	}
 
diff --git a/Engine/Source/Runtime/Experimental/Chaos/Public/Chaos/PBDRigidsEvolutionGBF.h b/Engine/Source/Runtime/Experimental/Chaos/Public/Chaos/PBDRigidsEvolutionGBF.h
--- a/Engine/Source/Runtime/Experimental/Chaos/Public/Chaos/PBDRigidsEvolutionGBF.h
+++ b/Engine/Source/Runtime/Experimental/Chaos/Public/Chaos/PBDRigidsEvolutionGBF.h
@@ -51,6 +51,9 @@ namespace Chaos
 		static constexpr int32 DefaultNumPushOutIterations = 5;
 		static constexpr int32 DefaultNumPushOutPairIterations = 2;
 
+		// @todo(chaos): Required by clustering - clean up
+		using Base::ApplyPushOut;
+
 		CHAOS_API FPBDRigidsEvolutionGBF(TPBDRigidsSOAs<FReal, 3>& InParticles, THandleArray<FChaosPhysicsMaterial>& SolverPhysicsMaterials, int32 InNumIterations = DefaultNumIterations, int32 InNumPushoutIterations = DefaultNumPushOutIterations, bool InIsSingleThreaded = false);
 		CHAOS_API ~FPBDRigidsEvolutionGBF() {}
 
@@ -98,11 +101,6 @@ namespace Chaos
 		CHAOS_API void Advance(const FReal Dt, const FReal MaxStepDt, const int32 MaxSteps);
 		CHAOS_API void AdvanceOneTimeStep(const FReal dt, const FReal StepFraction = (FReal)1.0);
 
-		using Base::PrepareConstraints;
-		using Base::UnprepareConstraints;
-		using Base::ApplyConstraints;
-		using Base::ApplyPushOut;
-
 		FCollisionConstraints& GetCollisionConstraints() { return CollisionConstraints; }
 		const FCollisionConstraints& GetCollisionConstraints() const { return CollisionConstraints; }
 
diff --git a/Engine/Source/Runtime/Experimental/Chaos/Public/ChaosStats.h b/Engine/Source/Runtime/Experimental/Chaos/Public/ChaosStats.h
--- a/Engine/Source/Runtime/Experimental/Chaos/Public/ChaosStats.h
+++ b/Engine/Source/Runtime/Experimental/Chaos/Public/ChaosStats.h
@@ -12,6 +12,7 @@ DECLARE_STATS_GROUP(TEXT("ChaosDedicated"), STATGROUP_ChaosDedicated, STATCAT_Ad
 DECLARE_STATS_GROUP(TEXT("ChaosEngine"), STATGROUP_ChaosEngine, STATCAT_Advanced);
 DECLARE_STATS_GROUP(TEXT("ChaosCollision"), STATGROUP_ChaosCollision, STATCAT_Advanced);
 DECLARE_STATS_GROUP(TEXT("ChaosJoint"), STATGROUP_ChaosJoint, STATCAT_Advanced);
+DECLARE_STATS_GROUP(TEXT("ChaosMinEvolution"), STATGROUP_ChaosMinEvolution, STATCAT_Advanced);
 
 DECLARE_CYCLE_STAT_EXTERN(TEXT("Physics Tick"), STAT_ChaosTick, STATGROUP_Chaos, CHAOS_API);
 DECLARE_CYCLE_STAT_EXTERN(TEXT("Physics Advance"), STAT_PhysicsAdvance, STATGROUP_Chaos, CHAOS_API);
