Debug drawing for scaled, instanced, convexes.
Definitely pretty much the worst possible way to do this, unsure if it matters, still useful.

diff --git a/Engine/Source/Runtime/Experimental/Chaos/Private/Chaos/ChaosDebugDraw.cpp b/Engine/Source/Runtime/Experimental/Chaos/Private/Chaos/ChaosDebugDraw.cpp
--- a/Engine/Source/Runtime/Experimental/Chaos/Private/Chaos/ChaosDebugDraw.cpp
+++ b/Engine/Source/Runtime/Experimental/Chaos/Private/Chaos/ChaosDebugDraw.cpp
@@ -2,6 +2,8 @@
 #include "Chaos/ChaosDebugDraw.h"
 #include "Chaos/Box.h"
 #include "Chaos/Capsule.h"
+#include "Chaos/Convex.h"
+#include "Chaos/TriangleMesh.h"
 #include "Chaos/DebugDrawQueue.h"
 #include "Chaos/ImplicitObjectScaled.h"
 #include "Chaos/ImplicitObjectTransformed.h"
@@ -53,10 +55,124 @@ namespace Chaos
 
 #if CHAOS_DEBUG_DRAW
 
-		void DrawShapesImpl(const FRigidTransform3& ShapeTransform, const FImplicitObject* Shape, FColor Color)
+		void DrawShapesImpl(const FRigidTransform3& ShapeTransform, const FImplicitObject* Shape, const FColor& Color);
+
+		template <bool bInstanced>
+		void DrawShapesScaledImpl(const FRigidTransform3& ShapeTransform, const FImplicitObject* Shape, const FColor& Color)
+		{
+			const EImplicitObjectType PackedType = Shape->GetType();
+			const EImplicitObjectType InnerType = GetInnerType(PackedType);
+			CHAOS_CHECK(IsScaled(PackedType));
+			CHAOS_CHECK(IsInstanced(PackedType) == bInstanced);
+
+			FRigidTransform3 ScaleTM = FRigidTransform3::Identity;
+			switch (InnerType)
+			{
+			case ImplicitObjectType::Sphere:
+				break;
+			case ImplicitObjectType::Box:
+				break;
+			case ImplicitObjectType::Plane:
+				break;
+			case ImplicitObjectType::Capsule:
+				break;
+			case ImplicitObjectType::Transformed:
+				break;
+			case ImplicitObjectType::Union:
+				break;
+			case ImplicitObjectType::LevelSet:
+				break;
+			case ImplicitObjectType::Unknown:
+				break;
+			case ImplicitObjectType::Convex:
+			{
+				const TImplicitObjectScaled<FConvex, bInstanced>* Scaled = Shape->template GetObject<TImplicitObjectScaled<FConvex, bInstanced>>();
+				ScaleTM.SetScale3D(Scaled->GetScale());
+				DrawShapesImpl(ShapeTransform * ScaleTM, Scaled->GetUnscaledObject(), Color);
+				break;
+			}
+			case ImplicitObjectType::TaperedCylinder:
+				break;
+			case ImplicitObjectType::Cylinder:
+				break;
+			case ImplicitObjectType::TriangleMesh:
+				break;
+			case ImplicitObjectType::HeightField:
+				break;
+			default:
+				break;
+			}
+		}
+
+		void DrawShapesInstancedImpl(const FRigidTransform3& ShapeTransform, const FImplicitObject* Shape, const FColor& Color)
+		{
+			const EImplicitObjectType PackedType = Shape->GetType();
+			const EImplicitObjectType InnerType = GetInnerType(PackedType);
+			CHAOS_CHECK(IsScaled(PackedType) == false);
+			CHAOS_CHECK(IsInstanced(PackedType));
+
+			switch (InnerType)
+			{
+			case ImplicitObjectType::Sphere:
+				break;
+			case ImplicitObjectType::Box:
+				break;
+			case ImplicitObjectType::Plane:
+				break;
+			case ImplicitObjectType::Capsule:
+				break;
+			case ImplicitObjectType::Transformed:
+				break;
+			case ImplicitObjectType::Union:
+				break;
+			case ImplicitObjectType::LevelSet:
+				break;
+			case ImplicitObjectType::Unknown:
+				break;
+			case ImplicitObjectType::Convex:
+			{
+				const TImplicitObjectInstanced<FConvex>* Instanced = Shape->template GetObject<TImplicitObjectInstanced<FConvex>>();
+				DrawShapesImpl(ShapeTransform, Instanced->GetInstancedObject(), Color);
+				break;
+			}
+			case ImplicitObjectType::TaperedCylinder:
+				break;
+			case ImplicitObjectType::Cylinder:
+				break;
+			case ImplicitObjectType::TriangleMesh:
+				break;
+			case ImplicitObjectType::HeightField:
+				break;
+			default:
+				break;
+			}
+		}
+
+		void DrawShapesImpl(const FRigidTransform3& ShapeTransform, const FImplicitObject* Shape, const FColor& Color)
 		{
+			const EImplicitObjectType PackedType = Shape->GetType(); // Type includes scaling and instancing data
+			const EImplicitObjectType InnerType = GetInnerType(Shape->GetType());
+
+			// For scaled shapes, we must unpack the scaled type first
+			if (IsScaled(PackedType))
+			{
+				if (IsInstanced(PackedType))
+				{
+					DrawShapesScaledImpl<true>(ShapeTransform, Shape, Color);
+				}
+				else
+				{
+					DrawShapesScaledImpl<false>(ShapeTransform, Shape, Color);
+				}
+				return;
+			}
+			else if (IsInstanced(PackedType))
+			{
+				DrawShapesInstancedImpl(ShapeTransform, Shape, Color);
+			}
+
 			// @todo(ccaulfield): handle scale throughout
-			switch (Shape->GetType())
+			switch (InnerType)
 			{
 			case ImplicitObjectType::Sphere:
 			{
@@ -103,7 +219,25 @@ namespace Chaos
 			case ImplicitObjectType::Unknown:
 				break;
 			case ImplicitObjectType::Convex:
+			{
+				if (const FConvex* Convex = Shape->template GetObject<FConvex>())
+				{
+					// TODO: This is horrendously slow. Figure out a way to cache
+					// the generated trimeshes on the debug draw queue instance.
+					const TParticles<FReal, 3>& Particles = Convex->GetSurfaceParticles();
+					TTriangleMesh<FReal> Triangles = TTriangleMesh<FReal>::GetConvexHullFromParticles(Particles);
+					for (const TVector<int32, 3>& Elem : Triangles.GetElements())
+					{
+						const FVec3 P0 = ShapeTransform.TransformPosition(Particles.X(Elem[0]));
+						const FVec3 P1 = ShapeTransform.TransformPosition(Particles.X(Elem[1]));
+						const FVec3 P2 = ShapeTransform.TransformPosition(Particles.X(Elem[2]));
+						FDebugDrawQueue::GetInstance().DrawDebugLine(P0, P1, Color, false, -1.f, 0, ShapeThicknesScale * LineThickness);
+						FDebugDrawQueue::GetInstance().DrawDebugLine(P1, P2, Color, false, -1.f, 0, ShapeThicknesScale * LineThickness);
+						FDebugDrawQueue::GetInstance().DrawDebugLine(P2, P0, Color, false, -1.f, 0, ShapeThicknesScale * LineThickness);
+					}
+				}
 				break;
+			}
 			case ImplicitObjectType::TaperedCylinder:
 				break;
 			case ImplicitObjectType::Cylinder:
