Added a script for tagging chaos headers as lvl1. Modification of files in future CL

diff --git a/Engine/Source/Runtime/Experimental/Chaos/Public/Chaos/IncludeLvl1.inl b/Engine/Source/Runtime/Experimental/Chaos/Public/Chaos/IncludeLvl1.inl
new file mode 100644
--- /dev/null
+++ b/Engine/Source/Runtime/Experimental/Chaos/Public/Chaos/IncludeLvl1.inl
@@ -0,0 +1,13 @@
+// Copyright Epic Games, Inc. All Rights Reserved.
+
+#if defined(CHAOS_INCLUDE_LEVEL_0)
+#error "Cannot include level 1 file from level 0"
+#endif
+
+//Level 1 files should NOT be included by unreal .h files (Chaos itself is ok)
+//Unreal headers including Chaos headers leads to very slow compile times.
+//Until we fix all instances of this please use TEMP_HEADER_CHAOS_LEVEL_1 (different name for searchability)
+//Or consider making the included header a Level 0 - this is meant for files that rarely change like forward declares
+#if !defined(CHAOS_INCLUDE_LEVEL_1) || !defined(TEMP_HEADER_CHAOS_LEVEL_1)
+//#error "Cannot include level 1 file. If outside chaos this should be in a cpp file and should be undefined after includes are done"
+#endif
\ No newline at end of file
