Set various chaos related modules to have access to level 1 files

diff --git a/Engine/Plugins/Experimental/ChaosSolverPlugin/Source/ChaosSolverEditor/ChaosSolverEditor.Build.cs b/Engine/Plugins/Experimental/ChaosSolverPlugin/Source/ChaosSolverEditor/ChaosSolverEditor.Build.cs
--- a/Engine/Plugins/Experimental/ChaosSolverPlugin/Source/ChaosSolverEditor/ChaosSolverEditor.Build.cs
+++ b/Engine/Plugins/Experimental/ChaosSolverPlugin/Source/ChaosSolverEditor/ChaosSolverEditor.Build.cs
@@ -29,6 +29,8 @@ namespace UnrealBuildTool.Rules
 					"ToolMenus",
 				}
 				);
+
+			PrivateDefinitions.Add("CHAOS_INCLUDE_LEVEL_1=1");
 		}
 	}
 }
diff --git a/Engine/Plugins/Experimental/FieldSystemPlugin/Source/FieldSyStemEditor/FieldSystemEditor.Build.cs b/Engine/Plugins/Experimental/FieldSystemPlugin/Source/FieldSyStemEditor/FieldSystemEditor.Build.cs
--- a/Engine/Plugins/Experimental/FieldSystemPlugin/Source/FieldSyStemEditor/FieldSystemEditor.Build.cs
+++ b/Engine/Plugins/Experimental/FieldSystemPlugin/Source/FieldSyStemEditor/FieldSystemEditor.Build.cs
@@ -28,6 +28,8 @@ namespace UnrealBuildTool.Rules
 					"ToolMenus",
 				}
 				);
+
+			PrivateDefinitions.Add("CHAOS_INCLUDE_LEVEL_1=1");
 		}
 	}
 }
diff --git a/Engine/Plugins/Experimental/GeometryCollectionPlugin/Source/GeometryCollectionEditor/GeometryCollectionEditor.Build.cs b/Engine/Plugins/Experimental/GeometryCollectionPlugin/Source/GeometryCollectionEditor/GeometryCollectionEditor.Build.cs
--- a/Engine/Plugins/Experimental/GeometryCollectionPlugin/Source/GeometryCollectionEditor/GeometryCollectionEditor.Build.cs
+++ b/Engine/Plugins/Experimental/GeometryCollectionPlugin/Source/GeometryCollectionEditor/GeometryCollectionEditor.Build.cs
@@ -33,6 +33,8 @@ namespace UnrealBuildTool.Rules
 					"ToolMenus",
 				}
 			);
+
+			PrivateDefinitions.Add("CHAOS_INCLUDE_LEVEL_1=1");
 		}
 	}
 }
diff --git a/Engine/Source/Programs/HeadlessChaos/HeadlessChaos.Build.cs b/Engine/Source/Programs/HeadlessChaos/HeadlessChaos.Build.cs
--- a/Engine/Source/Programs/HeadlessChaos/HeadlessChaos.Build.cs
+++ b/Engine/Source/Programs/HeadlessChaos/HeadlessChaos.Build.cs
@@ -42,5 +42,7 @@ public class HeadlessChaos : ModuleRules
 		{
 			PublicDefinitions.Add("GTEST_OS_LINUX=1");
 		}
+
+		PrivateDefinitions.Add("CHAOS_INCLUDE_LEVEL_1=1");
 	}
 }
diff --git a/Engine/Source/Runtime/Experimental/Chaos/Chaos.Build.cs b/Engine/Source/Runtime/Experimental/Chaos/Chaos.Build.cs
--- a/Engine/Source/Runtime/Experimental/Chaos/Chaos.Build.cs
+++ b/Engine/Source/Runtime/Experimental/Chaos/Chaos.Build.cs
@@ -19,6 +19,7 @@ namespace UnrealBuildTool.Rules
             );
 
             PublicDefinitions.Add("COMPILE_WITHOUT_UNREAL_SUPPORT=0");
+			PrivateDefinitions.Add("CHAOS_INCLUDE_LEVEL_1=1");
 
             if (Target.bCompileChaos == true || Target.bUseChaos == true)
             {
diff --git a/Engine/Source/Runtime/Experimental/Chaos/Public/Chaos/IncludeLvl1.inl b/Engine/Source/Runtime/Experimental/Chaos/Public/Chaos/IncludeLvl1.inl
--- a/Engine/Source/Runtime/Experimental/Chaos/Public/Chaos/IncludeLvl1.inl
+++ b/Engine/Source/Runtime/Experimental/Chaos/Public/Chaos/IncludeLvl1.inl
@@ -8,6 +8,6 @@
 //Unreal headers including Chaos headers leads to very slow compile times.
 //Until we fix all instances of this please use TEMP_HEADER_CHAOS_LEVEL_1 (different name for searchability)
 //Or consider making the included header a Level 0 - this is meant for files that rarely change like forward declares
-#if !defined(CHAOS_INCLUDE_LEVEL_1) || !defined(TEMP_HEADER_CHAOS_LEVEL_1)
-//#error "Cannot include level 1 file. If outside chaos this should be in a cpp file and should be undefined after includes are done"
+#if !defined(CHAOS_INCLUDE_LEVEL_1) && !defined(TEMP_HEADER_CHAOS_LEVEL_1)
+//#error "Cannot include level 1 file. If outside chaos this should be in a cpp file and should be undefined after includes are done. Use Non-Unity build to see cpp file"
 #endif
\ No newline at end of file
diff --git a/Engine/Source/Runtime/Experimental/ChaosSolverEngine/ChaosSolverEngine.Build.cs b/Engine/Source/Runtime/Experimental/ChaosSolverEngine/ChaosSolverEngine.Build.cs
--- a/Engine/Source/Runtime/Experimental/ChaosSolverEngine/ChaosSolverEngine.Build.cs
+++ b/Engine/Source/Runtime/Experimental/ChaosSolverEngine/ChaosSolverEngine.Build.cs
@@ -21,6 +21,8 @@ namespace UnrealBuildTool.Rules
                     "RHI"
                 }
 				);
-        }
+
+			PrivateDefinitions.Add("CHAOS_INCLUDE_LEVEL_1=1");
+		}
 	}
 }
diff --git a/Engine/Source/Runtime/Experimental/ChaosSolvers/ChaosSolvers.Build.cs b/Engine/Source/Runtime/Experimental/ChaosSolvers/ChaosSolvers.Build.cs
--- a/Engine/Source/Runtime/Experimental/ChaosSolvers/ChaosSolvers.Build.cs
+++ b/Engine/Source/Runtime/Experimental/ChaosSolvers/ChaosSolvers.Build.cs
@@ -21,6 +21,7 @@ namespace UnrealBuildTool.Rules
             ); 
 
             PublicDefinitions.Add("COMPILE_WITHOUT_UNREAL_SUPPORT=0");
-        }
+			PrivateDefinitions.Add("CHAOS_INCLUDE_LEVEL_1=1");
+		}
     }
 }
diff --git a/Engine/Source/Runtime/Experimental/FieldSystem/Source/FieldSystemCore/FieldSystemCore.Build.cs b/Engine/Source/Runtime/Experimental/FieldSystem/Source/FieldSystemCore/FieldSystemCore.Build.cs
--- a/Engine/Source/Runtime/Experimental/FieldSystem/Source/FieldSystemCore/FieldSystemCore.Build.cs
+++ b/Engine/Source/Runtime/Experimental/FieldSystem/Source/FieldSystemCore/FieldSystemCore.Build.cs
@@ -16,6 +16,8 @@ namespace UnrealBuildTool.Rules
 					"Chaos"
                 }
             );
-        }
+
+			PrivateDefinitions.Add("CHAOS_INCLUDE_LEVEL_1=1");
+		}
     }
 }
diff --git a/Engine/Source/Runtime/Experimental/FieldSystem/Source/FieldSystemEngine/FieldSystemEngine.Build.cs b/Engine/Source/Runtime/Experimental/FieldSystem/Source/FieldSystemEngine/FieldSystemEngine.Build.cs
--- a/Engine/Source/Runtime/Experimental/FieldSystem/Source/FieldSystemEngine/FieldSystemEngine.Build.cs
+++ b/Engine/Source/Runtime/Experimental/FieldSystem/Source/FieldSystemEngine/FieldSystemEngine.Build.cs
@@ -23,6 +23,8 @@ namespace UnrealBuildTool.Rules
                     "FieldSystemSimulationCore"
                 }
                 );
-        }
+
+			PrivateDefinitions.Add("CHAOS_INCLUDE_LEVEL_1=1");
+		}
 	}
 }
diff --git a/Engine/Source/Runtime/Experimental/FieldSystem/Source/FieldSystemSimulationCore/FieldSystemSimulationCore.Build.cs b/Engine/Source/Runtime/Experimental/FieldSystem/Source/FieldSystemSimulationCore/FieldSystemSimulationCore.Build.cs
--- a/Engine/Source/Runtime/Experimental/FieldSystem/Source/FieldSystemSimulationCore/FieldSystemSimulationCore.Build.cs
+++ b/Engine/Source/Runtime/Experimental/FieldSystem/Source/FieldSystemSimulationCore/FieldSystemSimulationCore.Build.cs
@@ -17,6 +17,8 @@ namespace UnrealBuildTool.Rules
 					"FieldSystemCore"
                 }
             );
-        }
+
+			PrivateDefinitions.Add("CHAOS_INCLUDE_LEVEL_1=1");
+		}
     }
 }
diff --git a/Engine/Source/Runtime/Experimental/GeometryCollectionCore/GeometryCollectionCore.Build.cs b/Engine/Source/Runtime/Experimental/GeometryCollectionCore/GeometryCollectionCore.Build.cs
--- a/Engine/Source/Runtime/Experimental/GeometryCollectionCore/GeometryCollectionCore.Build.cs
+++ b/Engine/Source/Runtime/Experimental/GeometryCollectionCore/GeometryCollectionCore.Build.cs
@@ -18,6 +18,8 @@ namespace UnrealBuildTool.Rules
                     "Voronoi"
                 }
             );
-        }
+
+			PrivateDefinitions.Add("CHAOS_INCLUDE_LEVEL_1=1");
+		}
     }
 }
diff --git a/Engine/Source/Runtime/Experimental/GeometryCollectionEngine/GeometryCollectionEngine.Build.cs b/Engine/Source/Runtime/Experimental/GeometryCollectionEngine/GeometryCollectionEngine.Build.cs
--- a/Engine/Source/Runtime/Experimental/GeometryCollectionEngine/GeometryCollectionEngine.Build.cs
+++ b/Engine/Source/Runtime/Experimental/GeometryCollectionEngine/GeometryCollectionEngine.Build.cs
@@ -46,6 +46,8 @@ namespace UnrealBuildTool.Rules
             {
                 PrivateDependencyModuleNames.Add("UnrealEd");
             }
-        }
+
+			PrivateDefinitions.Add("CHAOS_INCLUDE_LEVEL_1=1");
+		}
 	}
 }
diff --git a/Engine/Source/Runtime/Experimental/GeometryCollectionSimulationCore/GeometryCollectionSimulationCore.Build.cs b/Engine/Source/Runtime/Experimental/GeometryCollectionSimulationCore/GeometryCollectionSimulationCore.Build.cs
--- a/Engine/Source/Runtime/Experimental/GeometryCollectionSimulationCore/GeometryCollectionSimulationCore.Build.cs
+++ b/Engine/Source/Runtime/Experimental/GeometryCollectionSimulationCore/GeometryCollectionSimulationCore.Build.cs
@@ -19,6 +19,8 @@ namespace UnrealBuildTool.Rules
                     "FieldSystemCore"
                 }
 				);
+
+			PrivateDefinitions.Add("CHAOS_INCLUDE_LEVEL_1=1");
 		}
 	}
 }
diff --git a/Engine/Source/Runtime/PhysicsCore/PhysicsCore.Build.cs b/Engine/Source/Runtime/PhysicsCore/PhysicsCore.Build.cs
--- a/Engine/Source/Runtime/PhysicsCore/PhysicsCore.Build.cs
+++ b/Engine/Source/Runtime/PhysicsCore/PhysicsCore.Build.cs
@@ -32,5 +32,7 @@ public class PhysicsCore: ModuleRules
                 PublicDependencyModuleNames.Add("APEX");
             }
         }
+
+		PrivateDefinitions.Add("CHAOS_INCLUDE_LEVEL_1=1");
 	}
 }
