Chaos - THeightField<T> is now FHeightField

diff --git a/Engine/Source/Programs/HeadlessChaos/Private/HeadlessChaosTestHeightfield.cpp b/Engine/Source/Programs/HeadlessChaos/Private/HeadlessChaosTestHeightfield.cpp
--- a/Engine/Source/Programs/HeadlessChaos/Private/HeadlessChaosTestHeightfield.cpp
+++ b/Engine/Source/Programs/HeadlessChaos/Private/HeadlessChaosTestHeightfield.cpp
@@ -44,7 +44,7 @@ namespace ChaosTest {
 			auto AlongZTest = [&](const TVec3<T>& Scale)
 			{
 				TArray<T> HeightsCopy = Heights;
-				THeightField<T> Heightfield(MoveTemp(HeightsCopy),TArray<uint8>(),Rows,Columns,Scale);
+				FHeightField Heightfield(MoveTemp(HeightsCopy),TArray<uint8>(),Rows,Columns,Scale);
 				const auto& Bounds = Heightfield.BoundingBox();	//Current API forces us to do this to cache the bounds
 
 				//test straight down raycast
@@ -92,7 +92,7 @@ namespace ChaosTest {
 			auto AlongXTest = [&](const TVec3<T>& Scale)
 			{
 				TArray<T> HeightsCopy = Heights;
-				THeightField<T> Heightfield(MoveTemp(HeightsCopy),TArray<uint8>(),Rows,Columns,Scale);
+				FHeightField Heightfield(MoveTemp(HeightsCopy),TArray<uint8>(),Rows,Columns,Scale);
 				const auto& Bounds = Heightfield.BoundingBox();	//Current API forces us to do this to cache the bounds
 
 				//test along x axis
@@ -134,7 +134,7 @@ namespace ChaosTest {
 			auto AlongYTest = [&](const TVec3<T>& Scale)
 			{
 				TArray<T> HeightsCopy = Heights;
-				THeightField<T> Heightfield(MoveTemp(HeightsCopy),TArray<uint8>(),Rows,Columns,Scale);
+				FHeightField Heightfield(MoveTemp(HeightsCopy),TArray<uint8>(),Rows,Columns,Scale);
 				const auto& Bounds = Heightfield.BoundingBox();	//Current API forces us to do this to cache the bounds
 
 				//test along y axis
@@ -199,7 +199,7 @@ namespace ChaosTest {
 			auto AlongXYTest = [&](const TVec3<T>& Scale)
 			{
 				TArray<T> HeightsCopy = Heights2;
-				THeightField<T> Heightfield(MoveTemp(HeightsCopy),TArray<uint8>(),Rows,Columns,Scale);
+				FHeightField Heightfield(MoveTemp(HeightsCopy),TArray<uint8>(),Rows,Columns,Scale);
 				const auto& Bounds = Heightfield.BoundingBox();	//Current API forces us to do this to cache the bounds
 
 				//test along x-y axis
@@ -236,7 +236,7 @@ namespace ChaosTest {
 				//Pick cells and shoot ray at them
 				//This should always succeed because 0,0 is the lowest and n,n is the heighest
 				TArray<T> HeightsCopy = Heights2;
-				THeightField<T> Heightfield(MoveTemp(HeightsCopy),TArray<uint8>(),Rows,Columns,Scale);
+				FHeightField Heightfield(MoveTemp(HeightsCopy),TArray<uint8>(),Rows,Columns,Scale);
 				const auto& Bounds = Heightfield.BoundingBox();	//Current API forces us to do this to cache the bounds
 
 				T TOI;
diff --git a/Engine/Source/Runtime/Engine/Public/AI/NavigationSystemHelpers.h b/Engine/Source/Runtime/Engine/Public/AI/NavigationSystemHelpers.h
--- a/Engine/Source/Runtime/Engine/Public/AI/NavigationSystemHelpers.h
+++ b/Engine/Source/Runtime/Engine/Public/AI/NavigationSystemHelpers.h
@@ -28,8 +28,7 @@ namespace physx
 #if WITH_CHAOS
 namespace Chaos
 {
-	template<typename T>
-	class THeightField;
+	class FHeightField;
 
 	class FTriangleMeshImplicitObject;
 }
@@ -48,7 +47,7 @@ struct FNavigableGeometryExport
 #if WITH_CHAOS
 	virtual void ExportChaosTriMesh(const Chaos::FTriangleMeshImplicitObject* const TriMesh, const FTransform& LocalToWorld) = 0;
 	virtual void ExportChaosConvexMesh(const FKConvexElem* const Convex, const FTransform& LocalToWorld) = 0;
-	virtual void ExportChaosHeightField(const Chaos::THeightField<float>* const Heightfield, const FTransform& LocalToWorld) = 0;
+	virtual void ExportChaosHeightField(const Chaos::FHeightField* const Heightfield, const FTransform& LocalToWorld) = 0;
 	virtual void ExportChaosHeightFieldSlice(const FNavHeightfieldSamples& PrefetchedHeightfieldSamples, const int32 NumRows, const int32 NumCols, const FTransform& LocalToWorld, const FBox& SliceBox) = 0;
 #endif
 	virtual void ExportRigidBodySetup(UBodySetup& BodySetup, const FTransform& LocalToWorld) = 0;
diff --git a/Engine/Source/Runtime/Experimental/Chaos/Private/Chaos/CollisionResolution.cpp b/Engine/Source/Runtime/Experimental/Chaos/Private/Chaos/CollisionResolution.cpp
--- a/Engine/Source/Runtime/Experimental/Chaos/Private/Chaos/CollisionResolution.cpp
+++ b/Engine/Source/Runtime/Experimental/Chaos/Private/Chaos/CollisionResolution.cpp
@@ -596,13 +596,13 @@ namespace Chaos
 		//
 
 
-		FContactPoint BoxHeightFieldContactPoint(const FAABB3& A, const FRigidTransform3& ATransform, const THeightField<FReal>& B, const FRigidTransform3& BTransform, const FReal CullDistance)
+		FContactPoint BoxHeightFieldContactPoint(const FAABB3& A, const FRigidTransform3& ATransform, const FHeightField& B, const FRigidTransform3& BTransform, const FReal CullDistance)
 		{
 			return GJKImplicitContactPoint< TBox<float, 3> >(TBox<float, 3>(A), ATransform, B, BTransform, CullDistance);
 		}
 
 
-		void UpdateBoxHeightFieldConstraint(const FAABB3& A, const FRigidTransform3& ATransform, const THeightField<FReal>& B, const FRigidTransform3& BTransform, const FReal CullDistance, FRigidBodyPointContactConstraint& Constraint)
+		void UpdateBoxHeightFieldConstraint(const FAABB3& A, const FRigidTransform3& ATransform, const FHeightField& B, const FRigidTransform3& BTransform, const FReal CullDistance, FRigidBodyPointContactConstraint& Constraint)
 		{
 			UpdateContactPoint(Constraint.Manifold, BoxHeightFieldContactPoint(A, ATransform, B, BTransform, CullDistance));
 		}
@@ -618,7 +618,7 @@ namespace Chaos
 		{
 
 			const TBox<FReal, 3> * Object0 = Implicit0->template GetObject<const TBox<FReal, 3> >();
-			const THeightField<FReal> * Object1 = Implicit1->template GetObject<const THeightField<FReal> >();
+			const FHeightField * Object1 = Implicit1->template GetObject<const FHeightField >();
 			if (ensure(Object0 && Object1))
 			{
 				FRigidBodyPointContactConstraint Constraint = FRigidBodyPointContactConstraint(Particle0, Implicit0, LocalTransform0, Particle1, Implicit1, LocalTransform1, EContactShapesType::BoxHeightField);
@@ -853,13 +853,13 @@ namespace Chaos
 		// Sphere - HeightField
 		//
 
-		FContactPoint SphereHeightFieldContactPoint(const TSphere<FReal, 3>& A, const FRigidTransform3& ATransform, const THeightField<FReal>& B, const FRigidTransform3& BTransform, const FReal CullDistance)
+		FContactPoint SphereHeightFieldContactPoint(const TSphere<FReal, 3>& A, const FRigidTransform3& ATransform, const FHeightField& B, const FRigidTransform3& BTransform, const FReal CullDistance)
 		{
 			return GJKImplicitContactPoint< TSphere<float, 3> >(TSphere<float, 3>(A), ATransform, B, BTransform, CullDistance);
 		}
 
 
-		void UpdateSphereHeightFieldConstraint(const TSphere<FReal, 3>& A, const FRigidTransform3& ATransform, const THeightField<FReal>& B, const FRigidTransform3& BTransform, const FReal CullDistance, FRigidBodyPointContactConstraint& Constraint)
+		void UpdateSphereHeightFieldConstraint(const TSphere<FReal, 3>& A, const FRigidTransform3& ATransform, const FHeightField& B, const FRigidTransform3& BTransform, const FReal CullDistance, FRigidBodyPointContactConstraint& Constraint)
 		{
 			UpdateContactPoint(Constraint.Manifold, SphereHeightFieldContactPoint(A, ATransform, B, BTransform, CullDistance));
 		}
@@ -875,7 +875,7 @@ namespace Chaos
 		{
 
 			const TSphere<FReal, 3>* Object0 = Implicit0->template GetObject<const TSphere<FReal, 3> >();
-			const THeightField<FReal> * Object1 = Implicit1->template GetObject<const THeightField<FReal> >();
+			const FHeightField * Object1 = Implicit1->template GetObject<const FHeightField >();
 			if (ensure(Object0 && Object1))
 			{
 				FRigidBodyPointContactConstraint Constraint = FRigidBodyPointContactConstraint(Particle0, Implicit0, LocalTransform0, Particle1, Implicit1, LocalTransform1, EContactShapesType::SphereHeightField);
@@ -1472,19 +1472,19 @@ namespace Chaos
 		//
 
 
-		FContactPoint CapsuleHeightFieldContactPoint(const TCapsule<FReal>& A, const FRigidTransform3& ATransform, const THeightField<FReal>& B, const FRigidTransform3& BTransform, const FReal CullDistance)
+		FContactPoint CapsuleHeightFieldContactPoint(const TCapsule<FReal>& A, const FRigidTransform3& ATransform, const FHeightField& B, const FRigidTransform3& BTransform, const FReal CullDistance)
 		{
 			return GJKImplicitContactPoint< TCapsule<float> >(A, ATransform, B, BTransform, CullDistance);
 		}
 
 
-		void UpdateCapsuleHeightFieldConstraint(const TCapsule<FReal>& A, const FRigidTransform3& ATransform, const THeightField<FReal>& B, const FRigidTransform3& BTransform, const FReal CullDistance, FRigidBodyPointContactConstraint& Constraint)
+		void UpdateCapsuleHeightFieldConstraint(const TCapsule<FReal>& A, const FRigidTransform3& ATransform, const FHeightField& B, const FRigidTransform3& BTransform, const FReal CullDistance, FRigidBodyPointContactConstraint& Constraint)
 		{
 			UpdateContactPoint(Constraint.Manifold, CapsuleHeightFieldContactPoint(A, ATransform, B, BTransform, CullDistance));
 		}
 
 
-		void UpdateCapsuleHeightFieldConstraintSwept(TGeometryParticleHandle<FReal, 3>* Particle0, const TCapsule<FReal>& A, const FRigidTransform3& ATransform, const THeightField<FReal>& B, const FRigidTransform3& BTransform, const FVec3& Dir, const FReal Length, const FReal CullDistance, FRigidBodySweptPointContactConstraint& Constraint)
+		void UpdateCapsuleHeightFieldConstraintSwept(TGeometryParticleHandle<FReal, 3>* Particle0, const TCapsule<FReal>& A, const FRigidTransform3& ATransform, const FHeightField& B, const FRigidTransform3& BTransform, const FVec3& Dir, const FReal Length, const FReal CullDistance, FRigidBodySweptPointContactConstraint& Constraint)
 		{
 			FReal TOI = 1.0f;
 			UpdateContactPoint(Constraint.Manifold, GJKImplicitSweptContactPoint<TCapsule<float> >(A, ATransform, B, BTransform, Dir, Length, CullDistance, TOI));
@@ -1502,7 +1502,7 @@ namespace Chaos
 		{
 
 			const TCapsule<FReal> * Object0 = Implicit0->template GetObject<const TCapsule<FReal> >();
-			const THeightField<FReal> * Object1 = Implicit1->template GetObject<const THeightField<FReal> >();
+			const FHeightField * Object1 = Implicit1->template GetObject<const FHeightField >();
 			if (ensure(Object0 && Object1))
 			{
 				FRigidBodyPointContactConstraint Constraint = FRigidBodyPointContactConstraint(Particle0, Implicit0, LocalTransform0, Particle1, Implicit1, LocalTransform1, EContactShapesType::CapsuleHeightField);
@@ -1523,7 +1523,7 @@ namespace Chaos
 		void ConstructCapsuleHeightFieldConstraintsSwept(TGeometryParticleHandle<FReal, 3>* Particle0, TGeometryParticleHandle<FReal, 3>* Particle1, const FImplicitObject* Implicit0, const FImplicitObject* Implicit1, const FRigidTransform3& LocalTransform0, const FRigidTransform3& LocalTransform1, const FReal CullDistance, const FVec3& Dir, FReal Length, FCollisionConstraintsArray& NewConstraints)
 		{
 			const TCapsule<FReal> * Object0 = Implicit0->template GetObject<const TCapsule<FReal> >();
-			const THeightField<FReal> * Object1 = Implicit1->template GetObject<const THeightField<FReal> >();
+			const FHeightField * Object1 = Implicit1->template GetObject<const FHeightField >();
 			if (ensure(Object0 && Object1))
 			{
 				FRigidBodySweptPointContactConstraint Constraint = FRigidBodySweptPointContactConstraint(Particle0, Implicit0, LocalTransform0, Particle1, Implicit1, LocalTransform1, EContactShapesType::CapsuleHeightField);
@@ -1779,19 +1779,19 @@ namespace Chaos
 		//
 
 
-		FContactPoint ConvexHeightFieldContactPoint(const FImplicitObject& A, const FRigidTransform3& ATransform, const THeightField<FReal>& B, const FRigidTransform3& BTransform, const FReal CullDistance)
+		FContactPoint ConvexHeightFieldContactPoint(const FImplicitObject& A, const FRigidTransform3& ATransform, const FHeightField& B, const FRigidTransform3& BTransform, const FReal CullDistance)
 		{
 			return GJKImplicitContactPoint< FConvex >(A, ATransform, B, BTransform, CullDistance);
 		}
 
 
-		void UpdateConvexHeightFieldConstraint(const FImplicitObject& A, const FRigidTransform3& ATransform, const THeightField<FReal>& B, const FRigidTransform3& BTransform, const FReal CullDistance, FRigidBodyPointContactConstraint& Constraint)
+		void UpdateConvexHeightFieldConstraint(const FImplicitObject& A, const FRigidTransform3& ATransform, const FHeightField& B, const FRigidTransform3& BTransform, const FReal CullDistance, FRigidBodyPointContactConstraint& Constraint)
 		{
 			UpdateContactPoint(Constraint.Manifold, ConvexHeightFieldContactPoint(A, ATransform, B, BTransform, CullDistance));
 		}
 
 
-		void UpdateConvexHeightFieldConstraintSwept(TGeometryParticleHandle<FReal, 3>* Particle0, const FImplicitObject& A, const FRigidTransform3& ATransform, const THeightField<FReal>& B, const FRigidTransform3& BTransform, const FVec3& Dir, const FReal Length, const FReal CullDistance, FRigidBodySweptPointContactConstraint& Constraint)
+		void UpdateConvexHeightFieldConstraintSwept(TGeometryParticleHandle<FReal, 3>* Particle0, const FImplicitObject& A, const FRigidTransform3& ATransform, const FHeightField& B, const FRigidTransform3& BTransform, const FVec3& Dir, const FReal Length, const FReal CullDistance, FRigidBodySweptPointContactConstraint& Constraint)
 		{
 			FReal TOI = 1.0f;
 			UpdateContactPoint(Constraint.Manifold, GJKImplicitSweptContactPoint< FConvex >(A, ATransform, B, BTransform, Dir, Length, CullDistance, TOI));
@@ -1807,7 +1807,7 @@ namespace Chaos
 		template<typename T_TRAITS>
 		void ConstructConvexHeightFieldConstraints(TGeometryParticleHandle<FReal, 3>* Particle0, TGeometryParticleHandle<FReal, 3>* Particle1, const FImplicitObject* Implicit0, const FImplicitObject* Implicit1, const FRigidTransform3& LocalTransform0, const FRigidTransform3& LocalTransform1, const FReal CullDistance, FCollisionConstraintsArray& NewConstraints)
 		{
-			const THeightField<FReal> * Object1 = Implicit1->template GetObject<const THeightField<FReal> >();
+			const FHeightField * Object1 = Implicit1->template GetObject<const FHeightField >();
 			if (ensure(Implicit0->IsConvex() && Object1))
 			{
 				FRigidBodyPointContactConstraint Constraint = FRigidBodyPointContactConstraint(Particle0, Implicit0, LocalTransform0, Particle1, Implicit1, LocalTransform1, EContactShapesType::ConvexHeightField);
@@ -1827,7 +1827,7 @@ namespace Chaos
 
 		void ConstructConvexHeightFieldConstraintsSwept(TGeometryParticleHandle<FReal, 3>* Particle0, TGeometryParticleHandle<FReal, 3>* Particle1, const FImplicitObject* Implicit0, const FImplicitObject* Implicit1, const FRigidTransform3& LocalTransform0, const FRigidTransform3& LocalTransform1, const FReal CullDistance, const FVec3& Dir, FReal Length, FCollisionConstraintsArray& NewConstraints)
 		{
-			const THeightField<FReal> * Object1 = Implicit1->template GetObject<const THeightField<FReal> >();
+			const FHeightField * Object1 = Implicit1->template GetObject<const FHeightField >();
 			if (ensure(Implicit0->IsConvex() && Object1))
 			{
 				FRigidBodySweptPointContactConstraint Constraint = FRigidBodySweptPointContactConstraint(Particle0, Implicit0, LocalTransform0, Particle1, Implicit1, LocalTransform1, EContactShapesType::ConvexHeightField);
@@ -2095,7 +2095,7 @@ namespace Chaos
 				ensure(false);
 				break;
 			case EContactShapesType::SphereHeightField:
-				UpdateSphereHeightFieldConstraint(*Implicit0.template GetObject<TSphere<FReal, 3>>(), WorldTransform0, *Implicit1.template GetObject< THeightField<FReal> >(), WorldTransform1, CullDistance, Constraint);
+				UpdateSphereHeightFieldConstraint(*Implicit0.template GetObject<TSphere<FReal, 3>>(), WorldTransform0, *Implicit1.template GetObject< FHeightField >(), WorldTransform1, CullDistance, Constraint);
 				break;
 			case EContactShapesType::SpherePlane:
 				UpdateSpherePlaneConstraint(*Implicit0.template GetObject<TSphere<FReal, 3>>(), WorldTransform0, *Implicit1.template GetObject<TPlane<FReal, 3>>(), WorldTransform1, CullDistance, Constraint);
@@ -2124,7 +2124,7 @@ namespace Chaos
 				ensure(false);
 				break;
 			case EContactShapesType::CapsuleHeightField:
-				UpdateCapsuleHeightFieldConstraint(*Implicit0.template GetObject<TCapsule<FReal>>(), WorldTransform0, *Implicit1.template GetObject< THeightField<FReal> >(), WorldTransform1, CullDistance, Constraint);
+				UpdateCapsuleHeightFieldConstraint(*Implicit0.template GetObject<TCapsule<FReal>>(), WorldTransform0, *Implicit1.template GetObject< FHeightField >(), WorldTransform1, CullDistance, Constraint);
 				break;
 			case EContactShapesType::BoxBox:
 				UpdateBoxBoxConstraint(Implicit0.template GetObject<TBox<FReal, 3>>()->GetAABB(), WorldTransform0, Implicit1.template GetObject<TBox<FReal, 3>>()->GetAABB(), WorldTransform1, CullDistance, Constraint);
@@ -2147,7 +2147,7 @@ namespace Chaos
 				ensure(false);
 				break;
 			case EContactShapesType::BoxHeightField:
-				UpdateBoxHeightFieldConstraint(Implicit0.template GetObject<TBox<FReal, 3>>()->GetAABB(), WorldTransform0, *Implicit1.template GetObject< THeightField<FReal> >(), WorldTransform1, CullDistance, Constraint);
+				UpdateBoxHeightFieldConstraint(Implicit0.template GetObject<TBox<FReal, 3>>()->GetAABB(), WorldTransform0, *Implicit1.template GetObject< FHeightField >(), WorldTransform1, CullDistance, Constraint);
 				break;
 			case EContactShapesType::BoxPlane:
 				UpdateBoxPlaneConstraint(Implicit0.template GetObject<TBox<FReal, 3>>()->GetAABB(), WorldTransform0, *Implicit1.template GetObject<TPlane<FReal, 3>>(), WorldTransform1, CullDistance, Constraint);
@@ -2169,7 +2169,7 @@ namespace Chaos
 				ensure(false);
 				break;
 			case EContactShapesType::ConvexHeightField:
-				UpdateConvexHeightFieldConstraint(Implicit0, WorldTransform0, *Implicit1.template GetObject< THeightField<FReal> >(), WorldTransform1, CullDistance, Constraint);
+				UpdateConvexHeightFieldConstraint(Implicit0, WorldTransform0, *Implicit1.template GetObject< FHeightField >(), WorldTransform1, CullDistance, Constraint);
 				break;
 			case EContactShapesType::LevelSetLevelSet:
 				UpdateLevelsetLevelsetConstraint<UpdateType>(CullDistance, Constraint);
@@ -2198,11 +2198,11 @@ namespace Chaos
 			{
 				ConstructBoxBoxConstraints<T_TRAITS>(Particle0, Particle1, Implicit0, Implicit1, LocalTransform0, LocalTransform1, CullDistance, NewConstraints);
 			}
-			else if (Implicit0Type == TBox<FReal, 3>::StaticType() && Implicit1Type == THeightField<FReal>::StaticType())
+			else if (Implicit0Type == TBox<FReal, 3>::StaticType() && Implicit1Type == FHeightField::StaticType())
 			{
 				ConstructBoxHeightFieldConstraints<T_TRAITS>(Particle0, Particle1, Implicit0, Implicit1, LocalTransform0, LocalTransform1, CullDistance, NewConstraints);
 			}
-			else if (Implicit0Type == THeightField<FReal>::StaticType() && Implicit1Type == TBox<FReal, 3>::StaticType())
+			else if (Implicit0Type == FHeightField::StaticType() && Implicit1Type == TBox<FReal, 3>::StaticType())
 			{
 				ConstructBoxHeightFieldConstraints<T_TRAITS>(Particle1, Particle0, Implicit1, Implicit0, LocalTransform1, LocalTransform0, CullDistance, NewConstraints);
 			}
@@ -2218,11 +2218,11 @@ namespace Chaos
 			{
 				ConstructSphereSphereConstraints<T_TRAITS>(Particle0, Particle1, Implicit0, Implicit1, LocalTransform0, LocalTransform1, CullDistance, NewConstraints);
 			}
-			else if (Implicit0Type == TSphere<FReal, 3>::StaticType() && Implicit1Type == THeightField<FReal>::StaticType())
+			else if (Implicit0Type == TSphere<FReal, 3>::StaticType() && Implicit1Type == FHeightField::StaticType())
 			{
 				ConstructSphereHeightFieldConstraints<T_TRAITS>(Particle0, Particle1, Implicit0, Implicit1, LocalTransform0, LocalTransform1, CullDistance, NewConstraints);
 			}
-			else if (Implicit0Type == THeightField<FReal>::StaticType() && Implicit1Type == TSphere<FReal, 3>::StaticType())
+			else if (Implicit0Type == FHeightField::StaticType() && Implicit1Type == TSphere<FReal, 3>::StaticType())
 			{
 				ConstructSphereHeightFieldConstraints<T_TRAITS>(Particle1, Particle0, Implicit1, Implicit0, LocalTransform1, LocalTransform0, CullDistance, NewConstraints);
 			}
@@ -2262,7 +2262,7 @@ namespace Chaos
 			{
 				ConstructCapsuleBoxConstraints<T_TRAITS>(Particle1, Particle0, Implicit1, Implicit0, LocalTransform1, LocalTransform0, CullDistance, Context, NewConstraints);
 			}
-			else if (Implicit0Type == TCapsule<FReal>::StaticType() && Implicit1Type == THeightField<FReal>::StaticType())
+			else if (Implicit0Type == TCapsule<FReal>::StaticType() && Implicit1Type == FHeightField::StaticType())
 			{
 				if (bUseCCD)
 				{
@@ -2272,7 +2272,7 @@ namespace Chaos
 
 				ConstructCapsuleHeightFieldConstraints<T_TRAITS>(Particle0, Particle1, Implicit0, Implicit1, LocalTransform0, LocalTransform1, CullDistance, NewConstraints);
 			}
-			else if (Implicit0Type == THeightField<FReal>::StaticType() && Implicit1Type == TCapsule<FReal>::StaticType())
+			else if (Implicit0Type == FHeightField::StaticType() && Implicit1Type == TCapsule<FReal>::StaticType())
 			{
 				if (bUseCCD)
 				{
@@ -2318,7 +2318,7 @@ namespace Chaos
 
 				ConstructCapsuleTriangleMeshConstraints<T_TRAITS>(Particle1, Particle0, Implicit1, Implicit0, LocalTransform1, LocalTransform0, CullDistance, NewConstraints);
 			}
-			else if (bIsConvex0 && Implicit1Type == THeightField<FReal>::StaticType())
+			else if (bIsConvex0 && Implicit1Type == FHeightField::StaticType())
 			{
 				if (bUseCCD)
 				{
@@ -2328,7 +2328,7 @@ namespace Chaos
 
 				ConstructConvexHeightFieldConstraints<T_TRAITS>(Particle0, Particle1, Implicit0, Implicit1, LocalTransform0, LocalTransform1, CullDistance, NewConstraints);
 			}
-			else if (Implicit0Type == THeightField<FReal>::StaticType() && bIsConvex1)
+			else if (Implicit0Type == FHeightField::StaticType() && bIsConvex1)
 			{
 				if (bUseCCD)
 				{
@@ -2735,17 +2735,17 @@ namespace Chaos
 			}
 			Dir /= Length;
 
-			if (Implicit0Type == TBox<FReal, 3>::StaticType() && Implicit1Type == THeightField<FReal>::StaticType())
+			if (Implicit0Type == TBox<FReal, 3>::StaticType() && Implicit1Type == FHeightField::StaticType())
 			{
-				return GetPairTOIHackImpl<TBox<FReal, 3>>(*Implicit0->template GetObject<TBox<FReal, 3>>(), StartTransform0, *Implicit1->template GetObject<THeightField<FReal>>(), Transform1, Dir, Length, OutTOI, OutNormal, OutPhi);
+				return GetPairTOIHackImpl<TBox<FReal, 3>>(*Implicit0->template GetObject<TBox<FReal, 3>>(), StartTransform0, *Implicit1->template GetObject<FHeightField>(), Transform1, Dir, Length, OutTOI, OutNormal, OutPhi);
 			}
-			else if (Implicit0Type == TSphere<FReal, 3>::StaticType() && Implicit1Type == THeightField<FReal>::StaticType())
+			else if (Implicit0Type == TSphere<FReal, 3>::StaticType() && Implicit1Type == FHeightField::StaticType())
 			{
-				return GetPairTOIHackImpl<TSphere<FReal, 3>>(*Implicit0->template GetObject<TSphere<FReal, 3>>(), StartTransform0, *Implicit1->template GetObject<THeightField<FReal>>(), Transform1, Dir, Length, OutTOI, OutNormal, OutPhi);
+				return GetPairTOIHackImpl<TSphere<FReal, 3>>(*Implicit0->template GetObject<TSphere<FReal, 3>>(), StartTransform0, *Implicit1->template GetObject<FHeightField>(), Transform1, Dir, Length, OutTOI, OutNormal, OutPhi);
 			}
-			else if (Implicit0Type == TCapsule<FReal>::StaticType() && Implicit1Type == THeightField<FReal>::StaticType())
+			else if (Implicit0Type == TCapsule<FReal>::StaticType() && Implicit1Type == FHeightField::StaticType())
 			{
-				return GetPairTOIHackImpl<TCapsule<FReal>>(*Implicit0->template GetObject<TCapsule<FReal>>(), StartTransform0, *Implicit1->template GetObject<THeightField<FReal>>(), Transform1, Dir, Length, OutTOI, OutNormal, OutPhi);
+				return GetPairTOIHackImpl<TCapsule<FReal>>(*Implicit0->template GetObject<TCapsule<FReal>>(), StartTransform0, *Implicit1->template GetObject<FHeightField>(), Transform1, Dir, Length, OutTOI, OutNormal, OutPhi);
 			}
 			else if (Implicit0Type == TBox<FReal, 3>::StaticType() && Implicit1Type == FTriangleMeshImplicitObject::StaticType())
 			{
@@ -2759,9 +2759,9 @@ namespace Chaos
 			{
 				return GetPairTOIHackTriMeshImpl<TCapsule<FReal>>(*Implicit0->template GetObject<TCapsule<FReal>>(), StartTransform0, Implicit1, Transform1, Dir, Length, OutTOI, OutNormal, OutPhi);
 			}
-			else if (bIsConvex0 && Implicit1Type == THeightField<FReal>::StaticType())
+			else if (bIsConvex0 && Implicit1Type == FHeightField::StaticType())
 			{
-				return GetPairTOIHackImpl<FConvex>(*Implicit0, StartTransform0, *Implicit1->template GetObject<THeightField<FReal>>(), Transform1, Dir, Length, OutTOI, OutNormal, OutPhi);
+				return GetPairTOIHackImpl<FConvex>(*Implicit0, StartTransform0, *Implicit1->template GetObject<FHeightField>(), Transform1, Dir, Length, OutTOI, OutNormal, OutPhi);
 			}
 			else if (bIsConvex0 && Implicit1Type == FTriangleMeshImplicitObject::StaticType())
 			{
diff --git a/Engine/Source/Runtime/Experimental/Chaos/Private/Chaos/HeightField.cpp b/Engine/Source/Runtime/Experimental/Chaos/Private/Chaos/HeightField.cpp
--- a/Engine/Source/Runtime/Experimental/Chaos/Private/Chaos/HeightField.cpp
+++ b/Engine/Source/Runtime/Experimental/Chaos/Private/Chaos/HeightField.cpp
@@ -12,13 +12,12 @@
 
 namespace Chaos
 {
-	template<typename T>
-	class THeightfieldRaycastVisitor
+	class FHeightfieldRaycastVisitor
 	{
 	public:
 
-		THeightfieldRaycastVisitor(const typename THeightField<T>::FDataType* InData, const TVector<T, 3>& InStart, const TVector<T, 3>& InDir, const T InThickness)
-			: OutTime(TNumericLimits<T>::Max())
+		FHeightfieldRaycastVisitor(const typename FHeightField::FDataType* InData, const FVec3& InStart, const FVec3& InDir, const FReal InThickness)
+			: OutTime(TNumericLimits<FReal>::Max())
 			, OutFaceIndex(INDEX_NONE)
 			, GeomData(InData)
 			, Start(InStart)
@@ -33,24 +32,24 @@ namespace Chaos
 		};
 
 		template<ERaycastType SQType>
-		bool Visit(int32 Payload, T& CurrentLength)
+		bool Visit(int32 Payload, FReal& CurrentLength)
 		{
 			const int32 SubX = Payload % (GeomData->NumCols - 1);
 			const int32 SubY = Payload / (GeomData->NumCols - 1);
 
 			const int32 FullIndex = Payload + SubY;
 
-			const T Radius = Thickness + SMALL_NUMBER;
-			const T Radius2 = Radius * Radius;
+			const FReal Radius = Thickness + SMALL_NUMBER;
+			const FReal Radius2 = Radius * Radius;
 			bool bIntersection = false;
 
-			auto TestTriangle = [&](int32 FaceIndex, const TVector<T, 3>& A, const TVector<T, 3>& B, const TVector<T, 3>& C) -> bool
+			auto TestTriangle = [&](int32 FaceIndex, const FVec3& A, const FVec3& B, const FVec3& C) -> bool
 			{
-				const TVector<T, 3> AB = B - A;
-				const TVector<T, 3> AC = C - A;
+				const FVec3 AB = B - A;
+				const FVec3 AC = C - A;
 
-				TVector<T, 3> Normal = TVector<T, 3>::CrossProduct(AB, AC);
-				const T Len2 = Normal.SafeNormalize();
+				FVec3 Normal = FVec3::CrossProduct(AB, AC);
+				const FReal Len2 = Normal.SafeNormalize();
 
 				if(!ensure(Len2 > SMALL_NUMBER))
 				{
@@ -58,11 +57,11 @@ namespace Chaos
 					return true;
 				}
 
-				const TPlane<T, 3> TrianglePlane(A, Normal);
+				const TPlane<FReal, 3> TrianglePlane(A, Normal);
 
-				TVector<T, 3> ResultPosition(0);
-				TVector<T, 3> ResultNormal(0);
-				T Time = TNumericLimits<T>::Max();
+				FVec3 ResultPosition(0);
+				FVec3 ResultNormal(0);
+				FReal Time = TNumericLimits<FReal>::Max();
 				int32 DummyFaceIndex = INDEX_NONE;
 
 				if(TrianglePlane.Raycast(Start, Dir, CurrentLength, Thickness, Time, ResultPosition, ResultNormal, DummyFaceIndex))
@@ -70,8 +69,8 @@ namespace Chaos
 					if(Time == 0)
 					{
 						// Initial overlap
-						const TVector<T, 3> ClosestPtOnTri = FindClosestPointOnTriangle(TrianglePlane, A, B, C, Start);
-						const T DistToTriangle2 = (Start - ClosestPtOnTri).SizeSquared();
+						const FVec3 ClosestPtOnTri = FindClosestPointOnTriangle(TrianglePlane, A, B, C, Start);
+						const FReal DistToTriangle2 = (Start - ClosestPtOnTri).SizeSquared();
 						if(DistToTriangle2 <= Radius2)
 						{
 							OutTime = 0;
@@ -83,8 +82,8 @@ namespace Chaos
 					}
 					else
 					{
-						const TVector<T, 3> ClosestPtOnTri = FindClosestPointOnTriangle(ResultPosition, A, B, C, ResultPosition);
-						const T DistToTriangle2 = (ResultPosition - ClosestPtOnTri).SizeSquared();
+						const FVec3 ClosestPtOnTri = FindClosestPointOnTriangle(ResultPosition, A, B, C, ResultPosition);
+						const FReal DistToTriangle2 = (ResultPosition - ClosestPtOnTri).SizeSquared();
 						bIntersection = DistToTriangle2 <= SMALL_NUMBER;
 					}
 				}
@@ -92,22 +91,22 @@ namespace Chaos
 				if(SQType == ERaycastType::Sweep && !bIntersection)
 				{
 					//sphere is not immediately touching the triangle, but it could start intersecting the perimeter as it sweeps by
-					TVector<T, 3> BorderPositions[3];
-					TVector<T, 3> BorderNormals[3];
-					T BorderTimes[3];
+					FVec3 BorderPositions[3];
+					FVec3 BorderNormals[3];
+					FReal BorderTimes[3];
 					bool bBorderIntersections[3];
 
-					const TCapsule<T> ABCapsule(A, B, Thickness);
+					const TCapsule<FReal> ABCapsule(A, B, Thickness);
 					bBorderIntersections[0] = ABCapsule.Raycast(Start, Dir, CurrentLength, 0, BorderTimes[0], BorderPositions[0], BorderNormals[0], DummyFaceIndex);
 
-					const TCapsule<T> BCCapsule(B, C, Thickness);
+					const TCapsule<FReal> BCCapsule(B, C, Thickness);
 					bBorderIntersections[1] = BCCapsule.Raycast(Start, Dir, CurrentLength, 0, BorderTimes[1], BorderPositions[1], BorderNormals[1], DummyFaceIndex);
 
-					const TCapsule<T> ACCapsule(A, C, Thickness);
+					const TCapsule<FReal> ACCapsule(A, C, Thickness);
 					bBorderIntersections[2] = ACCapsule.Raycast(Start, Dir, CurrentLength, 0, BorderTimes[2], BorderPositions[2], BorderNormals[2], DummyFaceIndex);
 
 					int32 MinBorderIdx = INDEX_NONE;
-					T MinBorderTime = 0;
+					FReal MinBorderTime = 0;
 
 					for(int32 BorderIdx = 0; BorderIdx < 3; ++BorderIdx)
 					{
@@ -130,8 +129,8 @@ namespace Chaos
 						if(Time == 0)
 						{
 							//we were initially overlapping with triangle plane so no normal was given. Compute it now
-							TVector<T, 3> TmpNormal;
-							const T SignedDistance = TrianglePlane.PhiWithNormal(Start, TmpNormal);
+							FVec3 TmpNormal;
+							const FReal SignedDistance = TrianglePlane.PhiWithNormal(Start, TmpNormal);
 							ResultNormal = SignedDistance >= 0 ? TmpNormal : -TmpNormal;
 						}
 
@@ -165,7 +164,7 @@ namespace Chaos
 				return true;
 			};
 
-			TVector<T, 3> Points[4];
+			FVec3 Points[4];
 			GeomData->GetPointsScaled(FullIndex, Points);
 
 			// Test both triangles that are in this cell, as we could hit both in any order
@@ -175,37 +174,37 @@ namespace Chaos
 			return OutTime > 0;
 		}
 
-		bool VisitRaycast(int32 Payload, T& CurLength)
+		bool VisitRaycast(int32 Payload, FReal& CurLength)
 		{
 			return Visit<ERaycastType::Raycast>(Payload, CurLength);
 		}
 
-		bool VisitSweep(int32 Payload, T& CurLength)
+		bool VisitSweep(int32 Payload, FReal& CurLength)
 		{
 			return Visit<ERaycastType::Sweep>(Payload, CurLength);
 		}
 
-		T OutTime;
-		TVector<T, 3> OutPosition;
-		TVector<T, 3> OutNormal;
+		FReal OutTime;
+		FVec3 OutPosition;
+		FVec3 OutNormal;
 		int32 OutFaceIndex;
 
 	private:
 
-		const typename THeightField<T>::FDataType* GeomData;
+		const typename FHeightField::FDataType* GeomData;
 
-		TVector<T, 3> Start;
-		TVector<T, 3> Dir;
-		T Thickness;
+		FVec3 Start;
+		FVec3 Dir;
+		FReal Thickness;
 	};
 
-	template<typename GeomQueryType, typename T>
+	template<typename GeomQueryType>
 	class THeightfieldSweepVisitor
 	{
 	public:
 
-		THeightfieldSweepVisitor(const typename THeightField<T>::FDataType* InData, const GeomQueryType& InQueryGeom, const TRigidTransform<T, 3>& InStartTM, const TVector<T, 3>& InDir, const T InThickness, bool InComputeMTD)
-			: OutTime(TNumericLimits<T>::Max())
+		THeightfieldSweepVisitor(const typename FHeightField::FDataType* InData, const GeomQueryType& InQueryGeom, const FRigidTransform3& InStartTM, const FVec3& InDir, const FReal InThickness, bool InComputeMTD)
+			: OutTime(TNumericLimits<FReal>::Max())
 			, OutFaceIndex(INDEX_NONE)
 			, HfData(InData)
 			, StartTM(InStartTM)
@@ -215,14 +214,14 @@ namespace Chaos
 			, bComputeMTD(InComputeMTD)
 		{}
 
-		bool VisitSweep(int32 Payload, T& CurrentLength)
+		bool VisitSweep(int32 Payload, FReal& CurrentLength)
 		{
 			const int32 SubX = Payload % (HfData->NumCols - 1);
 			const int32 SubY = Payload / (HfData->NumCols - 1);
 
 			const int32 FullIndex = Payload + SubY;
 
-			auto TestTriangle = [&](int32 FaceIndex, const TVector<T, 3>& A, const TVector<T, 3>& B, const TVector<T, 3>& C) -> bool
+			auto TestTriangle = [&](int32 FaceIndex, const FVec3& A, const FVec3& B, const FVec3& C) -> bool
 			{
 				if(OutTime == 0)
 				{
@@ -231,13 +230,13 @@ namespace Chaos
 
 				//Convert into local space of A to get better precision
 
-				TTriangle<T> Triangle(TVec3<T>(0), B-A, C-A);
+				TTriangle<FReal> Triangle(TVec3<FReal>(0), B-A, C-A);
 
-				T Time;
-				TVector<T, 3> LocalHitPosition;
-				TVector<T, 3> HitNormal;
-				const TRigidTransform<T,3> LocalStartTM(StartTM.GetTranslation() - A,StartTM.GetRotation());
-				if(GJKRaycast2<T>(Triangle, OtherGeom, LocalStartTM, Dir, CurrentLength, Time, LocalHitPosition, HitNormal, Thickness, bComputeMTD))
+				FReal Time;
+				FVec3 LocalHitPosition;
+				FVec3 HitNormal;
+				const FRigidTransform3 LocalStartTM(StartTM.GetTranslation() - A,StartTM.GetRotation());
+				if(GJKRaycast2<FReal>(Triangle, OtherGeom, LocalStartTM, Dir, CurrentLength, Time, LocalHitPosition, HitNormal, Thickness, bComputeMTD))
 				{
 					if(Time < OutTime)
 					{
@@ -262,10 +261,10 @@ namespace Chaos
 								// we adopt the triangle normal but this leaves us with an incorrect MTD from the GJK call
 								// above. #TODO possibly re-do GJK with a plane, or some geom vs.plane special case to solve
 								// both triangles as planes 
-								const TVector<T, 3> AB = B - A;
-								const TVector<T, 3> AC = C - A;
+								const FVec3 AB = B - A;
+								const FVec3 AC = C - A;
 
-								TVector<T, 3> TriNormal = TVector<T, 3>::CrossProduct(AB, AC);
+								FVec3 TriNormal = FVec3::CrossProduct(AB, AC);
 								TriNormal.SafeNormalize();
 
 								OutNormal = TriNormal;
@@ -281,7 +280,7 @@ namespace Chaos
 				return true;
 			};
 
-			TVector<T, 3> Points[4];
+			FVec3 Points[4];
 			HfData->GetPointsScaled(FullIndex, Points);
 
 			bool bContinue = TestTriangle(Payload * 2, Points[0], Points[1], Points[3]);
@@ -293,26 +292,26 @@ namespace Chaos
 			return OutTime > 0;
 		}
 
-		T OutTime;
-		TVector<T, 3> OutPosition;
-		TVector<T, 3> OutNormal;
+		FReal OutTime;
+		FVec3 OutPosition;
+		FVec3 OutNormal;
 		int32 OutFaceIndex;
 
 	private:
 
-		const typename THeightField<T>::FDataType* HfData;
-		const TRigidTransform<T, 3> StartTM;
+		const typename FHeightField::FDataType* HfData;
+		const FRigidTransform3 StartTM;
 		const GeomQueryType& OtherGeom;
-		const TVector<T, 3>& Dir;
-		const T Thickness;
+		const FVec3& Dir;
+		const FReal Thickness;
 		bool bComputeMTD;
 
 	};
 
-	template<typename HeightfieldT, typename BufferType>
-	void BuildGeomData(TArrayView<BufferType> BufferView, TArrayView<uint8> MaterialIndexView, int32 NumRows, int32 NumCols, const TVector<HeightfieldT, 3>& InScale, TUniqueFunction<HeightfieldT(const BufferType)> ToRealFunc, typename THeightField<HeightfieldT>::FDataType& OutData, TAABB<HeightfieldT, 3>& OutBounds)
+	template<typename BufferType>
+	void BuildGeomData(TArrayView<BufferType> BufferView, TArrayView<uint8> MaterialIndexView, int32 NumRows, int32 NumCols, const FVec3& InScale, TUniqueFunction<FReal(const BufferType)> ToRealFunc, typename FHeightField::FDataType& OutData, FAABB3& OutBounds)
 	{
-		using FDataType = typename THeightField<HeightfieldT>::FDataType;
+		using FDataType = typename FHeightField::FDataType;
 
 		using RealType = typename FDataType::RealType;
 
@@ -355,10 +354,10 @@ namespace Chaos
 
 			int32 X = HeightIndex % (NumCols);
 			int32 Y = HeightIndex / (NumCols);
-			TVector<HeightfieldT, 3> Position(RealType(X), RealType(Y), OutData.MinValue + OutData.Heights[HeightIndex] * OutData.HeightPerUnit);
+			FVec3 Position(RealType(X), RealType(Y), OutData.MinValue + OutData.Heights[HeightIndex] * OutData.HeightPerUnit);
 			if(HeightIndex == 0)
 			{
-				OutBounds = TAABB<HeightfieldT, 3>(Position * InScale, Position * InScale);
+				OutBounds = FAABB3(Position * InScale, Position * InScale);
 			}
 			else
 			{
@@ -383,14 +382,14 @@ namespace Chaos
 		}
 	}
 
-	template<typename HeightfieldT, typename BufferType>
-	void EditGeomData(TArrayView<BufferType> BufferView, int32 InBeginRow, int32 InBeginCol, int32 NumRows, int32 NumCols, TUniqueFunction<HeightfieldT(const BufferType)> ToRealFunc, typename THeightField<HeightfieldT>::FDataType& OutData, TAABB<HeightfieldT, 3>& OutBounds)
+	template<typename BufferType>
+	void EditGeomData(TArrayView<BufferType> BufferView, int32 InBeginRow, int32 InBeginCol, int32 NumRows, int32 NumCols, TUniqueFunction<FReal(const BufferType)> ToRealFunc, typename FHeightField::FDataType& OutData, FAABB3& OutBounds)
 	{
-		using FDataType = typename THeightField<HeightfieldT>::FDataType;
+		using FDataType = typename FHeightField::FDataType;
 		using RealType = typename FDataType::RealType;
 
-		RealType MinValue = TNumericLimits<HeightfieldT>::Max();
-		RealType MaxValue = TNumericLimits<HeightfieldT>::Min();
+		RealType MinValue = TNumericLimits<FReal>::Max();
+		RealType MaxValue = TNumericLimits<FReal>::Min();
 
 		for(BufferType& Value : BufferView)
 		{
@@ -433,10 +432,10 @@ namespace Chaos
 
 					int32 X = HeightIndex % (OutData.NumCols);
 					int32 Y = HeightIndex / (OutData.NumCols);
-					TVector<HeightfieldT, 3> Position(RealType(X), RealType(Y), NewMin + OutData.Heights[HeightIndex] * NewHeightPerUnit);
+					FVec3 Position(RealType(X), RealType(Y), NewMin + OutData.Heights[HeightIndex] * NewHeightPerUnit);
 					if(HeightIndex == 0)
 					{
-						OutBounds = TAABB<HeightfieldT, 3>(Position, Position);
+						OutBounds = FAABB3(Position, Position);
 					}
 					else
 					{
@@ -467,31 +466,28 @@ namespace Chaos
 		}
 	}
 
-	template <typename T>
-	THeightField<T>::THeightField(TArray<T>&& Height, TArray<uint8>&& InMaterialIndices, int32 NumRows, int32 NumCols, const TVector<T, 3>& InScale)
+	FHeightField::FHeightField(TArray<FReal>&& Height, TArray<uint8>&& InMaterialIndices, int32 NumRows, int32 NumCols, const FVec3& InScale)
 		: FImplicitObject(EImplicitObject::HasBoundingBox, ImplicitObjectType::HeightField)
 	{
-		BuildGeomData<T, T>(MakeArrayView(Height), MakeArrayView(InMaterialIndices), NumRows, NumCols, TVector<T, 3>(1), [](const T InVal) {return InVal; }, GeomData, LocalBounds);
+		BuildGeomData<float>(MakeArrayView(Height), MakeArrayView(InMaterialIndices), NumRows, NumCols, FVec3(1), [](const FReal InVal) {return InVal; }, GeomData, LocalBounds);
 		CalcBounds();
 		SetScale(InScale);
 	}
 
-	template<typename T>
-	Chaos::THeightField<T>::THeightField(TArrayView<const uint16> InHeights, TArrayView<uint8> InMaterialIndices, int32 InNumRows, int32 InNumCols, const TVector<T, 3>& InScale)
+	Chaos::FHeightField::FHeightField(TArrayView<const uint16> InHeights, TArrayView<uint8> InMaterialIndices, int32 InNumRows, int32 InNumCols, const FVec3& InScale)
 		: FImplicitObject(EImplicitObject::HasBoundingBox, ImplicitObjectType::HeightField)
 	{
-		TUniqueFunction<T(const uint16)> ConversionFunc = [](const T InVal) -> float
+		TUniqueFunction<FReal(const uint16)> ConversionFunc = [](const FReal InVal) -> float
 		{
 			return (float)((int32)InVal - 32768);
 		};
 
-		BuildGeomData<T, const uint16>(InHeights, InMaterialIndices, InNumRows, InNumCols, TVector<T, 3>(1), MoveTemp(ConversionFunc), GeomData, LocalBounds);
+		BuildGeomData<const uint16>(InHeights, InMaterialIndices, InNumRows, InNumCols, FVec3(1), MoveTemp(ConversionFunc), GeomData, LocalBounds);
 		CalcBounds();
 		SetScale(InScale);
 	}
 
-	template<typename T>
-	void Chaos::THeightField<T>::EditHeights(TArrayView<const uint16> InHeights, int32 InBeginRow, int32 InBeginCol, int32 InNumRows, int32 InNumCols)
+	void Chaos::FHeightField::EditHeights(TArrayView<const uint16> InHeights, int32 InBeginRow, int32 InBeginCol, int32 InNumRows, int32 InNumCols)
 	{
 		const int32 NumExpectedValues = InNumRows * InNumCols;
 		const int32 EndRow = InBeginRow + InNumRows - 1;
@@ -499,20 +495,19 @@ namespace Chaos
 
 		if(ensure(InHeights.Num() == NumExpectedValues && InBeginRow >= 0 && InBeginCol >= 0 && EndRow < GeomData.NumRows && EndCol < GeomData.NumCols))
 		{
-			TUniqueFunction<T(const uint16)> ConversionFunc = [](const T InVal) -> float
+			TUniqueFunction<FReal(const uint16)> ConversionFunc = [](const FReal InVal) -> float
 			{
 				return (float)((int32)InVal - 32768);
 			};
 
-			EditGeomData<T, const uint16>(InHeights, InBeginRow, InBeginCol, InNumRows, InNumCols, MoveTemp(ConversionFunc), GeomData, LocalBounds);
+			EditGeomData<const uint16>(InHeights, InBeginRow, InBeginCol, InNumRows, InNumCols, MoveTemp(ConversionFunc), GeomData, LocalBounds);
 
 			// Slow and dumb. TODO: Actually fix CellHeights inside EditGeomData.
 			CalcBounds();
 		}
 	}
 
-	template<typename T>
-	void Chaos::THeightField<T>::EditHeights(TArrayView<T> InHeights, int32 InBeginRow, int32 InBeginCol, int32 InNumRows, int32 InNumCols)
+	void Chaos::FHeightField::EditHeights(TArrayView<FReal> InHeights, int32 InBeginRow, int32 InBeginCol, int32 InNumRows, int32 InNumCols)
 	{
 		const int32 NumExpectedValues = InNumRows * InNumCols;
 		const int32 EndRow = InBeginRow + InNumRows - 1;
@@ -520,25 +515,24 @@ namespace Chaos
 
 		if(ensure(InHeights.Num() == NumExpectedValues && InBeginRow >= 0 && InBeginCol >= 0 && EndRow < GeomData.NumRows && EndCol < GeomData.NumCols))
 		{
-			TUniqueFunction<T(T)> ConversionFunc = [](const T InVal) -> float
+			TUniqueFunction<FReal(FReal)> ConversionFunc = [](const FReal InVal) -> float
 			{
 				return InVal;
 			};
 
-			EditGeomData<T, T>(InHeights, InBeginRow, InBeginCol, InNumRows, InNumCols, MoveTemp(ConversionFunc), GeomData, LocalBounds);
+			EditGeomData<FReal>(InHeights, InBeginRow, InBeginCol, InNumRows, InNumCols, MoveTemp(ConversionFunc), GeomData, LocalBounds);
 
 			// Slow and dumb. TODO: Actually fix CellHeights inside EditGeomData.
 			CalcBounds();
 		}
 	}
 
-	template<typename T>
-	bool Chaos::THeightField<T>::GetCellBounds2D(const TVector<int32, 2> InCoord, FBounds2D& OutBounds, const TVector<T, 2>& InInflate /*= {0}*/) const
+	bool Chaos::FHeightField::GetCellBounds2D(const TVector<int32, 2> InCoord, FBounds2D& OutBounds, const TVector<FReal, 2>& InInflate /*= {0}*/) const
 	{
 		if (FlatGrid.IsValid(InCoord))
 		{
-			OutBounds.Min = TVector<T, 2>(InCoord[0], InCoord[1]);
-			OutBounds.Max = TVector<T, 2>(InCoord[0] + 1, InCoord[1] + 1);
+			OutBounds.Min = TVector<FReal, 2>(InCoord[0], InCoord[1]);
+			OutBounds.Max = TVector<FReal, 2>(InCoord[0] + 1, InCoord[1] + 1);
 			OutBounds.Min -= InInflate;
 			OutBounds.Max += InInflate;
 
@@ -548,26 +542,23 @@ namespace Chaos
 		return false;
 	}
 
-	template<typename T>
-	T Chaos::THeightField<T>::GetHeight(int32 InIndex) const
+	FReal Chaos::FHeightField::GetHeight(int32 InIndex) const
 	{
 		if (CHAOS_ENSURE(InIndex >= 0 && InIndex < GeomData.Heights.Num()))
 		{
 			return GeomData.GetPoint(InIndex).Z;
 		}
 
-		return TNumericLimits<T>::Max();
+		return TNumericLimits<FReal>::Max();
 	}
 
-	template<typename T>
-	T Chaos::THeightField<T>::GetHeight(int32 InX, int32 InY) const
+	FReal Chaos::FHeightField::GetHeight(int32 InX, int32 InY) const
 	{
 		const int32 Index = InY * GeomData.NumCols + InX;
 		return GetHeight(Index);
 	}
 
-	template<typename T>
-	uint8 Chaos::THeightField<T>::GetMaterialIndex(int32 InIndex) const
+	uint8 Chaos::FHeightField::GetMaterialIndex(int32 InIndex) const
 	{
 		if(CHAOS_ENSURE(InIndex >= 0 && InIndex < GeomData.MaterialIndices.Num()))
 		{
@@ -577,36 +568,32 @@ namespace Chaos
 		return TNumericLimits<uint8>::Max();
 	}
 
-	template<typename T>
-	uint8 Chaos::THeightField<T>::GetMaterialIndex(int32 InX, int32 InY) const
+	uint8 Chaos::FHeightField::GetMaterialIndex(int32 InX, int32 InY) const
 	{
 		const int32 Index = InY * GeomData.NumCols + InX;
 		return GetMaterialIndex(Index);
 	}
 
-	template<typename T>
-	bool Chaos::THeightField<T>::IsHole(int32 InIndex) const
+	bool Chaos::FHeightField::IsHole(int32 InIndex) const
 	{
 		return GetMaterialIndex(InIndex) == TNumericLimits<uint8>::Max();
 	}
 
-	template<typename T>
-	bool Chaos::THeightField<T>::IsHole(int32 InCellX, int32 InCellY) const
+	bool Chaos::FHeightField::IsHole(int32 InCellX, int32 InCellY) const
 	{
 		// Convert to single cell index
 		const int32 Index = InCellY * (GeomData.NumCols - 1) + InCellX;
 		return IsHole(Index);
 	}
 
-	template<typename T>
-	T Chaos::THeightField<T>::GetHeightAt(const TVector<T, 2>& InGridLocationLocal) const
+	FReal Chaos::FHeightField::GetHeightAt(const TVector<FReal, 2>& InGridLocationLocal) const
 	{
 		if(CHAOS_ENSURE(InGridLocationLocal == FlatGrid.Clamp(InGridLocationLocal)))
 		{
 			TVector<int32, 2> CellCoord = FlatGrid.Cell(InGridLocationLocal);
 
 			const int32 SingleIndex = CellCoord[1] * (GeomData.NumCols) + CellCoord[0];
-			TVector<T, 3> Pts[4];
+			FVec3 Pts[4];
 			GeomData.GetPoints(SingleIndex, Pts);
 
 			float FractionX = FMath::Frac(InGridLocationLocal[0]);
@@ -656,17 +643,16 @@ namespace Chaos
 		return 0.0f;
 	}
 
-	template<typename T>
-	bool Chaos::THeightField<T>::GetCellBounds3D(const TVector<int32, 2> InCoord, TVector<T, 3>& OutMin, TVector<T, 3>& OutMax, const TVector<T, 3>& InInflate /*= {0}*/) const
+	bool Chaos::FHeightField::GetCellBounds3D(const TVector<int32, 2> InCoord, FVec3& OutMin, FVec3& OutMax, const FVec3& InInflate /*= {0}*/) const
 	{
 		if (FlatGrid.IsValid(InCoord))
 		{
 			//todo: just compute max height, avoid extra work since this is called from tight loop
-			TVec3<T> Min,Max;
+			TVec3<FReal> Min,Max;
 			CalcCellBounds3D(InCoord,Min,Max);
 
-			OutMin = TVec3<T>(InCoord[0], InCoord[1], GeomData.GetMinHeight());
-			OutMax = TVec3<T>(InCoord[0] + 1, InCoord[1] + 1, Max[2]);
+			OutMin = TVec3<FReal>(InCoord[0], InCoord[1], GeomData.GetMinHeight());
+			OutMax = TVec3<FReal>(InCoord[0] + 1, InCoord[1] + 1, Max[2]);
 			OutMin = OutMin - InInflate;
 			OutMax = OutMax + InInflate;
 
@@ -676,16 +662,15 @@ namespace Chaos
 		return false;
 	}
 
-	template<typename T>
-	bool Chaos::THeightField<T>::GetCellBounds2DScaled(const TVector<int32, 2> InCoord, FBounds2D& OutBounds, const TVector<T, 2>& InInflate /*= {0}*/) const
+	bool Chaos::FHeightField::GetCellBounds2DScaled(const TVector<int32, 2> InCoord, FBounds2D& OutBounds, const TVector<FReal, 2>& InInflate /*= {0}*/) const
 	{
 		if (FlatGrid.IsValid(InCoord))
 		{
-			OutBounds.Min = TVector<T, 2>(InCoord[0], InCoord[1]);
-			OutBounds.Max = TVector<T, 2>(InCoord[0] + 1, InCoord[1] + 1);
+			OutBounds.Min = TVector<FReal, 2>(InCoord[0], InCoord[1]);
+			OutBounds.Max = TVector<FReal, 2>(InCoord[0] + 1, InCoord[1] + 1);
 			OutBounds.Min -= InInflate;
 			OutBounds.Max += InInflate;
-			const TVector<float, 2> Scale2D = TVector<T, 2>(GeomData.Scale[0], GeomData.Scale[1]);
+			const TVector<float, 2> Scale2D = TVector<FReal, 2>(GeomData.Scale[0], GeomData.Scale[1]);
 			OutBounds.Min *= Scale2D;
 			OutBounds.Max *= Scale2D;
 			return true;
@@ -694,17 +679,16 @@ namespace Chaos
 		return false;
 	}
 
-	template<typename T>
-	bool Chaos::THeightField<T>::GetCellBounds3DScaled(const TVector<int32, 2> InCoord, TVector<T, 3>& OutMin, TVector<T, 3>& OutMax, const TVector<T, 3>& InInflate /*= {0}*/) const
+	bool Chaos::FHeightField::GetCellBounds3DScaled(const TVector<int32, 2> InCoord, FVec3& OutMin, FVec3& OutMax, const FVec3& InInflate /*= {0}*/) const
 	{
 		if (FlatGrid.IsValid(InCoord))
 		{
 			//todo: just compute max height, avoid extra work since this is called from tight loop
-			TVec3<T> Min,Max;
+			TVec3<FReal> Min,Max;
 			CalcCellBounds3D(InCoord,Min,Max);
 
-			OutMin = TVec3<T>(InCoord[0], InCoord[1], GeomData.GetMinHeight());
-			OutMax = TVec3<T>(InCoord[0] + 1, InCoord[1] + 1, Max[2]);
+			OutMin = TVec3<FReal>(InCoord[0], InCoord[1], GeomData.GetMinHeight());
+			OutMax = TVec3<FReal>(InCoord[0] + 1, InCoord[1] + 1, Max[2]);
 			OutMin = OutMin * GeomData.Scale - InInflate;
 			OutMax = OutMax * GeomData.Scale + InInflate;
 			return true;
@@ -713,22 +697,21 @@ namespace Chaos
 		return false;
 	}
 
-	template<typename T>
-	bool Chaos::THeightField<T>::CalcCellBounds3D(const TVector<int32, 2> InCoord, TVector<T, 3>& OutMin, TVector<T, 3>& OutMax, const TVector<T, 3>& InInflate /*= {0}*/) const
+	bool Chaos::FHeightField::CalcCellBounds3D(const TVector<int32, 2> InCoord, FVec3& OutMin, FVec3& OutMax, const FVec3& InInflate /*= {0}*/) const
 	{
 		if(FlatGrid.IsValid(InCoord))
 		{
 			int32 Index = InCoord[1] * (GeomData.NumCols) + InCoord[0];
-			static TVector<T, 3> Points[4];
+			static FVec3 Points[4];
 			GeomData.GetPoints(Index, Points);
 
 			OutMin = OutMax = Points[0];
 
 			for(int32 PointIndex = 1; PointIndex < 4; ++PointIndex)
 			{
-				const TVector<T, 3>& Point = Points[PointIndex];
-				OutMin = TVector<T, 3>(FGenericPlatformMath::Min(OutMin[0], Point[0]), FGenericPlatformMath::Min(OutMin[1], Point[1]), FGenericPlatformMath::Min(OutMin[2], Point[2]));
-				OutMax = TVector<T, 3>(FGenericPlatformMath::Max(OutMax[0], Point[0]), FGenericPlatformMath::Max(OutMax[1], Point[1]), FGenericPlatformMath::Max(OutMax[2], Point[2]));
+				const FVec3& Point = Points[PointIndex];
+				OutMin = FVec3(FGenericPlatformMath::Min(OutMin[0], Point[0]), FGenericPlatformMath::Min(OutMin[1], Point[1]), FGenericPlatformMath::Min(OutMin[2], Point[2]));
+				OutMax = FVec3(FGenericPlatformMath::Max(OutMax[0], Point[0]), FGenericPlatformMath::Max(OutMax[1], Point[1]), FGenericPlatformMath::Max(OutMax[2], Point[2]));
 			}
 
 			OutMin -= InInflate;
@@ -741,8 +724,7 @@ namespace Chaos
 	}
 
 
-	template <typename T>
-	bool THeightField<T>::GridCast(const TVector<T, 3>& StartPoint, const TVector<T, 3>& Dir, const T Length, THeightfieldRaycastVisitor<T>& Visitor) const
+	bool FHeightField::GridCast(const FVec3& StartPoint, const FVec3& Dir, const FReal Length, FHeightfieldRaycastVisitor& Visitor) const
 	{
 		//Is this check needed?
 		if(Length < 1e-4)
@@ -750,42 +732,42 @@ namespace Chaos
 			return false;
 		}
 
-		T CurrentLength = Length;
-		TVector<T, 2> ClippedFlatRayStart;
-		TVector<T, 2> ClippedFlatRayEnd;
+		FReal CurrentLength = Length;
+		TVector<FReal, 2> ClippedFlatRayStart;
+		TVector<FReal, 2> ClippedFlatRayEnd;
 
 		// Data for fast box cast
-		TVector<T, 3> Min, Max, HitPoint;
+		FVec3 Min, Max, HitPoint;
 		bool bParallel[3];
-		TVector<T, 3> InvDir;
+		FVec3 InvDir;
 
-		T InvCurrentLength = 1 / CurrentLength;
+		FReal InvCurrentLength = 1 / CurrentLength;
 		for(int Axis = 0; Axis < 3; ++Axis)
 		{
 			bParallel[Axis] = FMath::IsNearlyZero(Dir[Axis], 1.e-8f);
 			InvDir[Axis] = bParallel[Axis] ? 0 : 1 / Dir[Axis];
 		}
 
-		T TOI;
+		FReal TOI;
 		const FBounds2D FlatBounds = GetFlatBounds();
-		TAABB<T,3> Bounds(
-			TVec3<T>(FlatBounds.Min[0],FlatBounds.Min[1],GeomData.GetMinHeight() * GeomData.Scale[2]),
-			TVec3<T>(FlatBounds.Max[0],FlatBounds.Max[1],GeomData.GetMaxHeight() * GeomData.Scale[2])
+		TAABB<FReal, 3> Bounds(
+			TVec3<FReal>(FlatBounds.Min[0],FlatBounds.Min[1],GeomData.GetMinHeight() * GeomData.Scale[2]),
+			TVec3<FReal>(FlatBounds.Max[0],FlatBounds.Max[1],GeomData.GetMaxHeight() * GeomData.Scale[2])
 			);
-		TVector<T,3> NextStart;
+		FVec3 NextStart;
 
 		if(Bounds.RaycastFast(StartPoint, Dir, InvDir, bParallel, Length, InvCurrentLength, TOI, NextStart))
 		{
-			const TVector<T,2> Scale2D(GeomData.Scale[0],GeomData.Scale[1]);
+			const TVector<FReal, 2> Scale2D(GeomData.Scale[0],GeomData.Scale[1]);
 			TVector<int32,2> CellIdx = FlatGrid.Cell(TVector<int32,2>(NextStart[0] / Scale2D[0], NextStart[1] / Scale2D[1]));
 
 			// Boundaries might push us one cell over
 			CellIdx = FlatGrid.ClampIndex(CellIdx);
-			const T ZDx = Bounds.Extents()[2];
-			const T ZMidPoint = Bounds.Min()[2] + ZDx * 0.5;
-			const TVector<T,3> ScaledDx(FlatGrid.Dx()[0] * Scale2D[0],FlatGrid.Dx()[1] * Scale2D[1],ZDx);
-			const TVector<T,2> ScaledDx2D(ScaledDx[0],ScaledDx[1]);
-			const TVector<T,2> ScaledMin = FlatGrid.MinCorner() * Scale2D;
+			const FReal ZDx = Bounds.Extents()[2];
+			const FReal ZMidPoint = Bounds.Min()[2] + ZDx * 0.5;
+			const FVec3 ScaledDx(FlatGrid.Dx()[0] * Scale2D[0],FlatGrid.Dx()[1] * Scale2D[1],ZDx);
+			const TVector<FReal, 2> ScaledDx2D(ScaledDx[0],ScaledDx[1]);
+			const TVector<FReal, 2> ScaledMin = FlatGrid.MinCorner() * Scale2D;
 
 			//START
 			do
@@ -794,7 +776,7 @@ namespace Chaos
 				{
 					// Check cell bounds
 					//todo: can do it without raycast
-					if(TAABB<T,3>(Min,Max).RaycastFast(StartPoint,Dir,InvDir,bParallel,CurrentLength,InvCurrentLength,TOI,HitPoint))
+					if(TAABB<FReal, 3>(Min,Max).RaycastFast(StartPoint,Dir,InvDir,bParallel,CurrentLength,InvCurrentLength,TOI,HitPoint))
 					{
 						// Visit the selected cell
 						bool bContinue = Visitor.VisitRaycast(CellIdx[1] * (GeomData.NumCols - 1) + CellIdx[0],CurrentLength);
@@ -809,19 +791,19 @@ namespace Chaos
 				//find next cell
 
 				//We want to know which plane we used to cross into next cell
-				const TVector<T,2> ScaledCellCenter2D = ScaledMin + TVector<T,2>(CellIdx[0] + 0.5,CellIdx[1] + 0.5) * ScaledDx2D;
-				const TVector<T,3> ScaledCellCenter(ScaledCellCenter2D[0], ScaledCellCenter2D[1], ZMidPoint);
+				const TVector<FReal, 2> ScaledCellCenter2D = ScaledMin + TVector<FReal, 2>(CellIdx[0] + 0.5,CellIdx[1] + 0.5) * ScaledDx2D;
+				const FVec3 ScaledCellCenter(ScaledCellCenter2D[0], ScaledCellCenter2D[1], ZMidPoint);
 
-				T Times[3];
-				T BestTime = CurrentLength;
+				FReal Times[3];
+				FReal BestTime = CurrentLength;
 				bool bTerminate = true;
 				for(int Axis = 0; Axis < 3; ++Axis)
 				{
 					if(!bParallel[Axis])
 					{
-						const T CrossPoint = Dir[Axis] > 0 ? ScaledCellCenter[Axis] + ScaledDx[Axis] / 2 : ScaledCellCenter[Axis] - ScaledDx[Axis] / 2;
-						const T Distance = CrossPoint - NextStart[Axis];	//note: CellCenter already has /2, we probably want to use the corner instead
-						const T Time = Distance * InvDir[Axis];
+						const FReal CrossPoint = Dir[Axis] > 0 ? ScaledCellCenter[Axis] + ScaledDx[Axis] / 2 : ScaledCellCenter[Axis] - ScaledDx[Axis] / 2;
+						const FReal Distance = CrossPoint - NextStart[Axis];	//note: CellCenter already has /2, we probably want to use the corner instead
+						const FReal Time = Distance * InvDir[Axis];
 						Times[Axis] = Time;
 						if(Time < BestTime)
 						{
@@ -830,7 +812,7 @@ namespace Chaos
 						}
 					} else
 					{
-						Times[Axis] = TNumericLimits<T>::Max();
+						Times[Axis] = TNumericLimits<FReal>::Max();
 					}
 				}
 
@@ -902,25 +884,24 @@ namespace Chaos
 		int32 DataSize;
 	};
 
-	template<typename T>
 	template<typename SQVisitor>
-	bool THeightField<T>::GridSweep(const TVector<T, 3>& StartPoint, const TVector<T, 3>& Dir, const T Length, const TVector<T, 2> InHalfExtents, SQVisitor& Visitor) const
+	bool FHeightField::GridSweep(const FVec3& StartPoint, const FVec3& Dir, const FReal Length, const TVector<FReal, 2> InHalfExtents, SQVisitor& Visitor) const
 	{
-		T CurrentLength = Length;
+		FReal CurrentLength = Length;
 
 		FBounds2D InflatedBounds = GetFlatBounds();
 		InflatedBounds.Min -= InHalfExtents;
 		InflatedBounds.Max += InHalfExtents;
 
-		TVector<T, 3> HalfExtents3D(InHalfExtents[0], InHalfExtents[1], InHalfExtents[1]);
+		FVec3 HalfExtents3D(InHalfExtents[0], InHalfExtents[1], InHalfExtents[1]);
 
-		const TVector<T, 3> EndPoint = StartPoint + Dir * Length;
-		const TVector<T, 2> Start2D(StartPoint[0], StartPoint[1]);
-		const TVector<T, 2> End2D(EndPoint[0], EndPoint[1]);
-		const TVector<T, 2> Scale2D(GeomData.Scale[0], GeomData.Scale[1]);
+		const FVec3 EndPoint = StartPoint + Dir * Length;
+		const TVector<FReal, 2> Start2D(StartPoint[0], StartPoint[1]);
+		const TVector<FReal, 2> End2D(EndPoint[0], EndPoint[1]);
+		const TVector<FReal, 2> Scale2D(GeomData.Scale[0], GeomData.Scale[1]);
 
-		TVector<T, 2> ClippedStart;
-		TVector<T, 2> ClippedEnd;
+		TVector<FReal, 2> ClippedStart;
+		TVector<FReal, 2> ClippedEnd;
 
 		if(InflatedBounds.ClipLine(StartPoint, StartPoint + Dir * Length, ClippedStart, ClippedEnd))
 		{
@@ -944,7 +925,7 @@ namespace Chaos
 			struct FQueueEntry
 			{
 				TVector<int32, 2> Index;
-				T ToI;
+				FReal ToI;
 			};
 
 			// Tracking data for cells to query (similar to bounding volume approach)
@@ -954,12 +935,12 @@ namespace Chaos
 			Seen.Add(StartCell);
 
 			// Data for fast box cast
-			TVector<T, 3> Min, Max, HitPoint;
+			FVec3 Min, Max, HitPoint;
 			float ToI;
 			bool bParallel[3];
-			TVector<T, 3> InvDir;
+			FVec3 InvDir;
 
-			T InvCurrentLength = 1 / CurrentLength;
+			FReal InvCurrentLength = 1 / CurrentLength;
 			for(int Axis = 0; Axis < 3; ++Axis)
 			{
 				bParallel[Axis] = FMath::IsNearlyZero(Dir[Axis], 1.e-8f);
@@ -1003,7 +984,7 @@ namespace Chaos
 						FBounds2D CellBounds;
 						if(GetCellBounds3DScaled(NeighCoord, Min, Max, HalfExtents3D) && !Seen.Contains(NeighCoord))
 						{
-							if(TAABB<T, 3>(Min,Max).RaycastFast(StartPoint, Dir, InvDir, bParallel, CurrentLength, InvCurrentLength, ToI, HitPoint))
+							if(TAABB<FReal, 3>(Min,Max).RaycastFast(StartPoint, Dir, InvDir, bParallel, CurrentLength, InvCurrentLength, ToI, HitPoint))
 							{
 								Seen.Add(NeighCoord);
 								Queue.Add({NeighCoord, ToI});
@@ -1045,7 +1026,7 @@ namespace Chaos
 					// Check the current cell, if we hit its 3D bound we can move on to narrow phase
 					const TVector<int32, 2> Coord = CellCoord.Index;
 					if(GetCellBounds3DScaled(Coord, Min, Max, HalfExtents3D) &&
-						TAABB<T, 3>(Min,Max).RaycastFast(StartPoint, Dir, InvDir, bParallel, CurrentLength, InvCurrentLength, ToI, HitPoint))
+						TAABB<FReal, 3>(Min,Max).RaycastFast(StartPoint, Dir, InvDir, bParallel, CurrentLength, InvCurrentLength, ToI, HitPoint))
 					{
 						bool bContinue = Visitor.VisitSweep(CellCoord.Index[1] * (GeomData.NumCols - 1) + CellCoord.Index[0], CurrentLength);
 
@@ -1068,7 +1049,7 @@ namespace Chaos
 						{
 							ExpandAxis = 0;
 						}
-						T ExpandSize = HalfExtents3D[ExpandAxis];
+						FReal ExpandSize = HalfExtents3D[ExpandAxis];
 						int32 Steps = FMath::RoundFromZero(ExpandSize / GeomData.Scale[ExpandAxis]);
 
 						Expand(Coord, ThickenDir, Steps);
@@ -1105,16 +1086,15 @@ namespace Chaos
 		return false;
 	}
 
-	template <typename T>
-	bool THeightField<T>::Raycast(const TVector<T, 3>& StartPoint, const TVector<T, 3>& Dir, const T Length, const T Thickness, T& OutTime, TVector<T, 3>& OutPosition, TVector<T, 3>& OutNormal, int32& OutFaceIndex) const
+	bool FHeightField::Raycast(const FVec3& StartPoint, const FVec3& Dir, const FReal Length, const FReal Thickness, FReal& OutTime, FVec3& OutPosition, FVec3& OutNormal, int32& OutFaceIndex) const
 	{
 		OutFaceIndex = INDEX_NONE;
 
-		THeightfieldRaycastVisitor<T> Visitor(&GeomData, StartPoint, Dir, Thickness);
+		FHeightfieldRaycastVisitor Visitor(&GeomData, StartPoint, Dir, Thickness);
 
 		if(Thickness > 0)
 		{
-			GridSweep(StartPoint, Dir, Length, TVector<T, 2>(Thickness), Visitor);
+			GridSweep(StartPoint, Dir, Length, TVector<FReal, 2>(Thickness), Visitor);
 		}
 		else
 		{
@@ -1133,13 +1113,12 @@ namespace Chaos
 		return false;
 	}
 
-	template<typename T>
-	bool Chaos::THeightField<T>::GetGridIntersections(FBounds2D InFlatBounds, TArray<TVector<int32, 2>>& OutInterssctions) const
+	bool Chaos::FHeightField::GetGridIntersections(FBounds2D InFlatBounds, TArray<TVector<int32, 2>>& OutInterssctions) const
 	{
 		OutInterssctions.Reset();
 
 		const FBounds2D FlatBounds = GetFlatBounds();
-		const TVector<T, 2> Scale2D(GeomData.Scale[0], GeomData.Scale[1]);
+		const TVector<FReal, 2> Scale2D(GeomData.Scale[0], GeomData.Scale[1]);
 
 		InFlatBounds.Min = FlatBounds.Clamp(InFlatBounds.Min);
 		InFlatBounds.Max = FlatBounds.Clamp(InFlatBounds.Max);
@@ -1163,33 +1142,31 @@ namespace Chaos
 		return OutInterssctions.Num() > 0;
 	}
 
-	template<typename T>
-	typename Chaos::THeightField<T>::FBounds2D Chaos::THeightField<T>::GetFlatBounds() const
+	typename Chaos::FHeightField::FBounds2D Chaos::FHeightField::GetFlatBounds() const
 	{
 		FBounds2D Result;
-		Result.Min = TVector<T, 2>(CachedBounds.Min()[0], CachedBounds.Min()[1]);
-		Result.Max = TVector<T, 2>(CachedBounds.Max()[0], CachedBounds.Max()[1]);
+		Result.Min = TVector<FReal, 2>(CachedBounds.Min()[0], CachedBounds.Min()[1]);
+		Result.Max = TVector<FReal, 2>(CachedBounds.Max()[0], CachedBounds.Max()[1]);
 		return Result;
 	}
 
-	template <typename T>
-	bool THeightField<T>::Overlap(const TVector<T, 3>& Point, const T Thickness) const
+	bool FHeightField::Overlap(const FVec3& Point, const FReal Thickness) const
 	{
-		auto OverlapTriangle = [&](const TVector<T, 3>& A, const TVector<T, 3>& B, const TVector<T, 3>& C) -> bool
+		auto OverlapTriangle = [&](const FVec3& A, const FVec3& B, const FVec3& C) -> bool
 		{
-			const TVector<T, 3> AB = B - A;
-			const TVector<T, 3> AC = C - A;
-			TVector<T, 3> Normal = TVector<T, 3>::CrossProduct(AB, AC);
-			const T NormalLength = Normal.SafeNormalize();
+			const FVec3 AB = B - A;
+			const FVec3 AC = C - A;
+			FVec3 Normal = FVec3::CrossProduct(AB, AC);
+			const FReal NormalLength = Normal.SafeNormalize();
 
 			if(!ensure(NormalLength > KINDA_SMALL_NUMBER))
 			{
 				return false;
 			}
 
-			const TPlane<T, 3> TriPlane{A, Normal};
-			const TVector<T, 3> ClosestPointOnTri = FindClosestPointOnTriangle(TriPlane, A, B, C, Point);
-			const T Distance2 = (ClosestPointOnTri - Point).SizeSquared();
+			const TPlane<FReal, 3> TriPlane{A, Normal};
+			const FVec3 ClosestPointOnTri = FindClosestPointOnTriangle(TriPlane, A, B, C, Point);
+			const FReal Distance2 = (ClosestPointOnTri - Point).SizeSquared();
 
 			if(Distance2 <= Thickness * Thickness)	//This really only has a hope in working if thickness is > 0
 			{
@@ -1199,15 +1176,15 @@ namespace Chaos
 			return false;
 		};
 
-		TAABB<T, 3> QueryBounds(Point, Point);
+		TAABB<FReal, 3> QueryBounds(Point, Point);
 		QueryBounds.Thicken(Thickness);
 
 		FBounds2D FlatQueryBounds;
-		FlatQueryBounds.Min = TVector<T, 2>(QueryBounds.Min()[0], QueryBounds.Min()[1]);
-		FlatQueryBounds.Max = TVector<T, 2>(QueryBounds.Max()[0], QueryBounds.Max()[1]);
+		FlatQueryBounds.Min = TVector<FReal, 2>(QueryBounds.Min()[0], QueryBounds.Min()[1]);
+		FlatQueryBounds.Max = TVector<FReal, 2>(QueryBounds.Max()[0], QueryBounds.Max()[1]);
 
 		TArray<TVector<int32, 2>> Intersections;
-		TVector<T, 3> Points[4];
+		FVec3 Points[4];
 
 		GetGridIntersections(FlatQueryBounds, Intersections);
 
@@ -1231,26 +1208,25 @@ namespace Chaos
 	}
 
 
-	template <typename T>
 	template <typename GeomType>
-	bool THeightField<T>::GJKContactPointImp(const GeomType& QueryGeom, const TRigidTransform<T, 3>& QueryTM, const T Thickness,
-		TVector<T, 3>& ContactLocation, TVector<T, 3>& ContactNormal, T& ContactPhi) const
+	bool FHeightField::GJKContactPointImp(const GeomType& QueryGeom, const FRigidTransform3& QueryTM, const FReal Thickness,
+		FVec3& ContactLocation, FVec3& ContactNormal, FReal& ContactPhi) const
 	{
-		auto OverlapTriangle = [&](const TVector<T, 3>& A, const TVector<T, 3>& B, const TVector<T, 3>& C,
-			TVector<T, 3>& LocalContactLocation, TVector<T, 3>& LocalContactNormal, T& LocalContactPhi) -> bool
+		auto OverlapTriangle = [&](const FVec3& A, const FVec3& B, const FVec3& C,
+			FVec3& LocalContactLocation, FVec3& LocalContactNormal, FReal& LocalContactPhi) -> bool
 		{
-			const TVector<T, 3> AB = B - A;
-			const TVector<T, 3> AC = C - A;
+			const FVec3 AB = B - A;
+			const FVec3 AC = C - A;
 
-			const TVector<T, 3> Offset = TVector<T, 3>::CrossProduct(AB, AC);
+			const FVec3 Offset = FVec3::CrossProduct(AB, AC);
 
-			TTriangle<T> TriangleConvex(A, B, C);
+			TTriangle<FReal> TriangleConvex(A, B, C);
 
-			T Penetration;
-			TVec3<T> ClosestA, ClosestB, Normal;
-			if (GJKPenetration(TriangleConvex, QueryGeom, QueryTM, Penetration, ClosestA, ClosestB, Normal, (T)0))
+			FReal Penetration;
+			TVec3<FReal> ClosestA, ClosestB, Normal;
+			if (GJKPenetration(TriangleConvex, QueryGeom, QueryTM, Penetration, ClosestA, ClosestB, Normal, (FReal)0))
 			{
-				TVec3<T> TestVector = QueryTM.InverseTransformVector(Normal);
+				TVec3<FReal> TestVector = QueryTM.InverseTransformVector(Normal);
 
 				LocalContactLocation = ClosestB;
 				LocalContactNormal = Normal; 
@@ -1262,21 +1238,21 @@ namespace Chaos
 		};
 
 		bool bResult = false;
-		TAABB<T, 3> QueryBounds = QueryGeom.BoundingBox();
+		TAABB<FReal, 3> QueryBounds = QueryGeom.BoundingBox();
 		QueryBounds.Thicken(Thickness);
 		QueryBounds = QueryBounds.TransformedAABB(QueryTM);
 
 		FBounds2D FlatQueryBounds;
-		FlatQueryBounds.Min = TVector<T, 2>(QueryBounds.Min()[0], QueryBounds.Min()[1]);
-		FlatQueryBounds.Max = TVector<T, 2>(QueryBounds.Max()[0], QueryBounds.Max()[1]);
+		FlatQueryBounds.Min = TVector<FReal, 2>(QueryBounds.Min()[0], QueryBounds.Min()[1]);
+		FlatQueryBounds.Max = TVector<FReal, 2>(QueryBounds.Max()[0], QueryBounds.Max()[1]);
 
 		TArray<TVector<int32, 2>> Intersections;
-		TVector<T, 3> Points[4];
+		FVec3 Points[4];
 
 		GetGridIntersections(FlatQueryBounds, Intersections);
 
-		T LocalContactPhi = FLT_MAX;
-		TVector<T, 3> LocalContactLocation, LocalContactNormal;
+		FReal LocalContactPhi = FLT_MAX;
+		FVec3 LocalContactLocation, LocalContactNormal;
 		for (const TVector<int32, 2>& Cell : Intersections)
 		{
 			const int32 SingleIndex = Cell[1] * GeomData.NumCols + Cell[0];
@@ -1317,81 +1293,72 @@ namespace Chaos
 		return false;
 	}
 
-	template <typename T>
-	bool THeightField<T>::GJKContactPoint(const TBox<T, 3>& QueryGeom, const TRigidTransform<T, 3>& QueryTM, const T Thickness, TVector<T, 3>& ContactLocation, TVector<T, 3>& ContactNormal, T& ContactPhi) const
+	bool FHeightField::GJKContactPoint(const TBox<FReal, 3>& QueryGeom, const FRigidTransform3& QueryTM, const FReal Thickness, FVec3& ContactLocation, FVec3& ContactNormal, FReal& ContactPhi) const
 	{
 		return GJKContactPointImp(QueryGeom, QueryTM, Thickness, ContactLocation, ContactNormal, ContactPhi);
 	}
 
-	template <typename T>
-	bool THeightField<T>::GJKContactPoint(const TSphere<T, 3>& QueryGeom, const TRigidTransform<T, 3>& QueryTM, const T Thickness, TVector<T, 3>& ContactLocation, TVector<T, 3>& ContactNormal, T& ContactPhi) const
+	bool FHeightField::GJKContactPoint(const TSphere<FReal, 3>& QueryGeom, const FRigidTransform3& QueryTM, const FReal Thickness, FVec3& ContactLocation, FVec3& ContactNormal, FReal& ContactPhi) const
 	{
 		return GJKContactPointImp(QueryGeom, QueryTM, Thickness, ContactLocation, ContactNormal, ContactPhi);
 	}
 
-	template <typename T>
-	bool THeightField<T>::GJKContactPoint(const TCapsule<T>& QueryGeom, const TRigidTransform<T, 3>& QueryTM, const T Thickness, TVector<T, 3>& ContactLocation, TVector<T, 3>& ContactNormal, T& ContactPhi) const
+	bool FHeightField::GJKContactPoint(const TCapsule<FReal>& QueryGeom, const FRigidTransform3& QueryTM, const FReal Thickness, FVec3& ContactLocation, FVec3& ContactNormal, FReal& ContactPhi) const
 	{
 		return GJKContactPointImp(QueryGeom, QueryTM, Thickness, ContactLocation, ContactNormal, ContactPhi);
 	}
 
-	template <typename T>
-	bool THeightField<T>::GJKContactPoint(const FConvex& QueryGeom, const TRigidTransform<T, 3>& QueryTM, const T Thickness, TVector<T, 3>& ContactLocation, TVector<T, 3>& ContactNormal, T& ContactPhi) const
+	bool FHeightField::GJKContactPoint(const FConvex& QueryGeom, const FRigidTransform3& QueryTM, const FReal Thickness, FVec3& ContactLocation, FVec3& ContactNormal, FReal& ContactPhi) const
 	{
 		return GJKContactPointImp(QueryGeom, QueryTM, Thickness, ContactLocation, ContactNormal, ContactPhi);
 	}
 
-	template <typename T>
-	bool THeightField<T>::GJKContactPoint(const TImplicitObjectScaled < TBox<T, 3>>& QueryGeom, const TRigidTransform<T, 3>& QueryTM, const T Thickness, TVector<T, 3>& ContactLocation, TVector<T, 3>& ContactNormal, T& ContactPhi) const
+	bool FHeightField::GJKContactPoint(const TImplicitObjectScaled < TBox<FReal, 3>>& QueryGeom, const FRigidTransform3& QueryTM, const FReal Thickness, FVec3& ContactLocation, FVec3& ContactNormal, FReal& ContactPhi) const
 	{
 		return GJKContactPointImp(QueryGeom, QueryTM, Thickness, ContactLocation, ContactNormal, ContactPhi);
 	}
 
-	template <typename T>
-	bool THeightField<T>::GJKContactPoint(const TImplicitObjectScaled < TSphere<T, 3>>& QueryGeom, const TRigidTransform<T, 3>& QueryTM, const T Thickness, TVector<T, 3>& ContactLocation, TVector<T, 3>& ContactNormal, T& ContactPhi) const
+	bool FHeightField::GJKContactPoint(const TImplicitObjectScaled < TSphere<FReal, 3>>& QueryGeom, const FRigidTransform3& QueryTM, const FReal Thickness, FVec3& ContactLocation, FVec3& ContactNormal, FReal& ContactPhi) const
 	{
 		return GJKContactPointImp(QueryGeom, QueryTM, Thickness, ContactLocation, ContactNormal, ContactPhi);
 	}
 
-	template <typename T>
-	bool THeightField<T>::GJKContactPoint(const TImplicitObjectScaled < TCapsule<T>>& QueryGeom, const TRigidTransform<T, 3>& QueryTM, const T Thickness, TVector<T, 3>& ContactLocation, TVector<T, 3>& ContactNormal, T& ContactPhi) const
+	bool FHeightField::GJKContactPoint(const TImplicitObjectScaled < TCapsule<FReal>>& QueryGeom, const FRigidTransform3& QueryTM, const FReal Thickness, FVec3& ContactLocation, FVec3& ContactNormal, FReal& ContactPhi) const
 	{
 		return GJKContactPointImp(QueryGeom, QueryTM, Thickness, ContactLocation, ContactNormal, ContactPhi);
 	}
 
-	template <typename T>
-	bool THeightField<T>::GJKContactPoint(const TImplicitObjectScaled<FConvex>& QueryGeom, const TRigidTransform<T, 3>& QueryTM, const T Thickness, TVector<T, 3>& ContactLocation, TVector<T, 3>& ContactNormal, T& ContactPhi) const
+	bool FHeightField::GJKContactPoint(const TImplicitObjectScaled<FConvex>& QueryGeom, const FRigidTransform3& QueryTM, const FReal Thickness, FVec3& ContactLocation, FVec3& ContactNormal, FReal& ContactPhi) const
 	{
 		return GJKContactPointImp(QueryGeom, QueryTM, Thickness, ContactLocation, ContactNormal, ContactPhi);
 	}
 
 
-	template <typename T>
 	template <typename QueryGeomType>
-	bool THeightField<T>::OverlapGeomImp(const QueryGeomType& QueryGeom, const TRigidTransform<T, 3>& QueryTM, const T Thickness, FMTDInfo* OutMTD) const
+	bool FHeightField::OverlapGeomImp(const QueryGeomType& QueryGeom, const FRigidTransform3& QueryTM, const FReal Thickness, FMTDInfo* OutMTD) const
 	{
 		if (OutMTD)
 		{
-			OutMTD->Normal = TVec3<T>(0);
-			OutMTD->Penetration = TNumericLimits<T>::Lowest();
+			OutMTD->Normal = TVec3<FReal>(0);
+			OutMTD->Penetration = TNumericLimits<FReal>::Lowest();
 		}
 
-		auto OverlapTriangle = [&](const TVector<T, 3>& A, const TVector<T, 3>& B, const TVector<T, 3>& C, FMTDInfo* InnerMTD) -> bool
+		auto OverlapTriangle = [&](const FVec3& A, const FVec3& B, const FVec3& C, FMTDInfo* InnerMTD) -> bool
 		{
-			const TVector<T, 3> AB = B - A;
-			const TVector<T, 3> AC = C - A;
+			const FVec3 AB = B - A;
+			const FVec3 AC = C - A;
 
 			//It's most likely that the query object is in front of the triangle since queries tend to be on the outside.
 			//However, maybe we should check if it's behind the triangle plane. Also, we should enforce this winding in some way
-			const TVector<T, 3> Offset = TVector<T, 3>::CrossProduct(AB, AC);
+			const FVec3 Offset = FVec3::CrossProduct(AB, AC);
 
-			TTriangle<T> TriangleConvex(A, B, C);
+			TTriangle<FReal> TriangleConvex(A, B, C);
 			if (InnerMTD)
 			{
-				TVec3<T> TriangleNormal(0);
-				T Penetration = 0;
-				TVec3<T> ClosestA(0);
-				TVec3<T> ClosestB(0);
+				TVec3<FReal> TriangleNormal(0);
+				FReal Penetration = 0;
+				TVec3<FReal> ClosestA(0);
+				TVec3<FReal> ClosestB(0);
 				if (GJKPenetration(TriangleConvex, QueryGeom, QueryTM, Penetration, ClosestA, ClosestB, TriangleNormal, Thickness))
 				{
 					// Use Deepest MTD.
@@ -1412,16 +1379,16 @@ namespace Chaos
 		};
 
 		bool bResult = false;
-		TAABB<T, 3> QueryBounds = QueryGeom.BoundingBox();
+		TAABB<FReal, 3> QueryBounds = QueryGeom.BoundingBox();
 		QueryBounds.Thicken(Thickness);
 		QueryBounds = QueryBounds.TransformedAABB(QueryTM);
 
 		FBounds2D FlatQueryBounds;
-		FlatQueryBounds.Min = TVector<T, 2>(QueryBounds.Min()[0], QueryBounds.Min()[1]);
-		FlatQueryBounds.Max = TVector<T, 2>(QueryBounds.Max()[0], QueryBounds.Max()[1]);
+		FlatQueryBounds.Min = TVector<FReal, 2>(QueryBounds.Min()[0], QueryBounds.Min()[1]);
+		FlatQueryBounds.Max = TVector<FReal, 2>(QueryBounds.Max()[0], QueryBounds.Max()[1]);
 
 		TArray<TVector<int32, 2>> Intersections;
-		TVector<T, 3> Points[4];
+		FVec3 Points[4];
 
 		GetGridIntersections(FlatQueryBounds, Intersections);
 
@@ -1453,65 +1420,56 @@ namespace Chaos
 		return bOverlaps;
 	}
 
-	template <typename T>
-	bool THeightField<T>::OverlapGeom(const TSphere<T, 3>& QueryGeom, const TRigidTransform<T, 3>& QueryTM, const T Thickness, FMTDInfo* OutMTD) const
+	bool FHeightField::OverlapGeom(const TSphere<FReal, 3>& QueryGeom, const FRigidTransform3& QueryTM, const FReal Thickness, FMTDInfo* OutMTD) const
 	{
 		return OverlapGeomImp(QueryGeom, QueryTM, Thickness, OutMTD);
 	}
 
-	template <typename T>
-	bool THeightField<T>::OverlapGeom(const TBox<T, 3>& QueryGeom, const TRigidTransform<T, 3>& QueryTM, const T Thickness, FMTDInfo* OutMTD) const
+	bool FHeightField::OverlapGeom(const TBox<FReal, 3>& QueryGeom, const FRigidTransform3& QueryTM, const FReal Thickness, FMTDInfo* OutMTD) const
 	{
 		return OverlapGeomImp(QueryGeom, QueryTM, Thickness, OutMTD);
 	}
 
-	template <typename T>
-	bool THeightField<T>::OverlapGeom(const TCapsule<T>& QueryGeom, const TRigidTransform<T, 3>& QueryTM, const T Thickness, FMTDInfo* OutMTD) const
+	bool FHeightField::OverlapGeom(const TCapsule<FReal>& QueryGeom, const FRigidTransform3& QueryTM, const FReal Thickness, FMTDInfo* OutMTD) const
 	{
 		return OverlapGeomImp(QueryGeom, QueryTM, Thickness, OutMTD);
 	}
 
-	template <typename T>
-	bool THeightField<T>::OverlapGeom(const FConvex& QueryGeom, const TRigidTransform<T, 3>& QueryTM, const T Thickness, FMTDInfo* OutMTD) const
+	bool FHeightField::OverlapGeom(const FConvex& QueryGeom, const FRigidTransform3& QueryTM, const FReal Thickness, FMTDInfo* OutMTD) const
 	{
 		return OverlapGeomImp(QueryGeom, QueryTM, Thickness, OutMTD);
 	}
 
-	template <typename T>
-	bool THeightField<T>::OverlapGeom(const TImplicitObjectScaled<TSphere<T, 3>>& QueryGeom, const TRigidTransform<T, 3>& QueryTM, const T Thickness, FMTDInfo* OutMTD) const
+	bool FHeightField::OverlapGeom(const TImplicitObjectScaled<TSphere<FReal, 3>>& QueryGeom, const FRigidTransform3& QueryTM, const FReal Thickness, FMTDInfo* OutMTD) const
 	{
 		return OverlapGeomImp(QueryGeom, QueryTM, Thickness, OutMTD);
 	}
 
-	template <typename T>
-	bool THeightField<T>::OverlapGeom(const TImplicitObjectScaled<TBox<T, 3>>& QueryGeom, const TRigidTransform<T, 3>& QueryTM, const T Thickness, FMTDInfo* OutMTD) const
+	bool FHeightField::OverlapGeom(const TImplicitObjectScaled<TBox<FReal, 3>>& QueryGeom, const FRigidTransform3& QueryTM, const FReal Thickness, FMTDInfo* OutMTD) const
 	{
 		return OverlapGeomImp(QueryGeom, QueryTM, Thickness, OutMTD);
 	}
 
-	template <typename T>
-	bool THeightField<T>::OverlapGeom(const TImplicitObjectScaled<TCapsule<T>>& QueryGeom, const TRigidTransform<T, 3>& QueryTM, const T Thickness, FMTDInfo* OutMTD) const
+	bool FHeightField::OverlapGeom(const TImplicitObjectScaled<TCapsule<FReal>>& QueryGeom, const FRigidTransform3& QueryTM, const FReal Thickness, FMTDInfo* OutMTD) const
 	{
 		return OverlapGeomImp(QueryGeom, QueryTM, Thickness, OutMTD);
 	}
 
-	template <typename T>
-	bool THeightField<T>::OverlapGeom(const TImplicitObjectScaled<FConvex>& QueryGeom, const TRigidTransform<T, 3>& QueryTM, const T Thickness, FMTDInfo* OutMTD) const
+	bool FHeightField::OverlapGeom(const TImplicitObjectScaled<FConvex>& QueryGeom, const FRigidTransform3& QueryTM, const FReal Thickness, FMTDInfo* OutMTD) const
 	{
 		return OverlapGeomImp(QueryGeom, QueryTM, Thickness, OutMTD);
 	}
 
-	template <typename T>
 	template <typename QueryGeomType>
-	bool THeightField<T>::SweepGeomImp(const QueryGeomType& QueryGeom, const TRigidTransform<T, 3>& StartTM, const TVector<T, 3>& Dir, const T Length, T& OutTime, TVector<T, 3>& OutPosition, TVector<T, 3>& OutNormal, int32& OutFaceIndex, const T Thickness, bool bComputeMTD) const
+	bool FHeightField::SweepGeomImp(const QueryGeomType& QueryGeom, const FRigidTransform3& StartTM, const FVec3& Dir, const FReal Length, FReal& OutTime, FVec3& OutPosition, FVec3& OutNormal, int32& OutFaceIndex, const FReal Thickness, bool bComputeMTD) const
 	{
 		bool bHit = false;
-		THeightfieldSweepVisitor<QueryGeomType, T> SQVisitor(&GeomData, QueryGeom, StartTM, Dir, Thickness, bComputeMTD);
-		const TAABB<T, 3> QueryBounds = QueryGeom.BoundingBox();
-		const TVector<T, 3> StartPoint = StartTM.TransformPositionNoScale(QueryBounds.Center());
+		THeightfieldSweepVisitor<QueryGeomType> SQVisitor(&GeomData, QueryGeom, StartTM, Dir, Thickness, bComputeMTD);
+		const TAABB<FReal, 3> QueryBounds = QueryGeom.BoundingBox();
+		const FVec3 StartPoint = StartTM.TransformPositionNoScale(QueryBounds.Center());
 
-		const TVector<T, 3> Inflation3D = QueryBounds.Extents() * 0.5 + TVector<T, 3>(Thickness);
-		GridSweep(StartPoint, Dir, Length, TVector<T, 2>(Inflation3D[0], Inflation3D[1]), SQVisitor);
+		const FVec3 Inflation3D = QueryBounds.Extents() * 0.5 + FVec3(Thickness);
+		GridSweep(StartPoint, Dir, Length, TVector<FReal, 2>(Inflation3D[0], Inflation3D[1]), SQVisitor);
 
 		if(SQVisitor.OutTime <= Length)
 		{
@@ -1525,85 +1483,76 @@ namespace Chaos
 		return bHit;
 	}
 
-	template <typename T>
-	bool THeightField<T>::SweepGeom(const TSphere<T, 3>& QueryGeom, const TRigidTransform<T, 3>& StartTM, const TVector<T, 3>& Dir, const T Length, T& OutTime, TVector<T, 3>& OutPosition, TVector<T, 3>& OutNormal, int32& OutFaceIndex, const T Thickness, bool bComputeMTD) const
+	bool FHeightField::SweepGeom(const TSphere<FReal, 3>& QueryGeom, const FRigidTransform3& StartTM, const FVec3& Dir, const FReal Length, FReal& OutTime, FVec3& OutPosition, FVec3& OutNormal, int32& OutFaceIndex, const FReal Thickness, bool bComputeMTD) const
 	{
 		return SweepGeomImp(QueryGeom, StartTM, Dir, Length, OutTime, OutPosition, OutNormal, OutFaceIndex, Thickness, bComputeMTD);
 	}
 
-	template <typename T>
-	bool THeightField<T>::SweepGeom(const TBox<T, 3>& QueryGeom, const TRigidTransform<T, 3>& StartTM, const TVector<T, 3>& Dir, const T Length, T& OutTime, TVector<T, 3>& OutPosition, TVector<T, 3>& OutNormal, int32& OutFaceIndex, const T Thickness, bool bComputeMTD) const
+	bool FHeightField::SweepGeom(const TBox<FReal, 3>& QueryGeom, const FRigidTransform3& StartTM, const FVec3& Dir, const FReal Length, FReal& OutTime, FVec3& OutPosition, FVec3& OutNormal, int32& OutFaceIndex, const FReal Thickness, bool bComputeMTD) const
 	{
 		return SweepGeomImp(QueryGeom, StartTM, Dir, Length, OutTime, OutPosition, OutNormal, OutFaceIndex, Thickness, bComputeMTD);
 	}
 
-	template <typename T>
-	bool THeightField<T>::SweepGeom(const TCapsule<T>& QueryGeom, const TRigidTransform<T, 3>& StartTM, const TVector<T, 3>& Dir, const T Length, T& OutTime, TVector<T, 3>& OutPosition, TVector<T, 3>& OutNormal, int32& OutFaceIndex, const T Thickness, bool bComputeMTD) const
+	bool FHeightField::SweepGeom(const TCapsule<FReal>& QueryGeom, const FRigidTransform3& StartTM, const FVec3& Dir, const FReal Length, FReal& OutTime, FVec3& OutPosition, FVec3& OutNormal, int32& OutFaceIndex, const FReal Thickness, bool bComputeMTD) const
 	{
 		return SweepGeomImp(QueryGeom, StartTM, Dir, Length, OutTime, OutPosition, OutNormal, OutFaceIndex, Thickness, bComputeMTD);
 	}
 
-	template <typename T>
-	bool THeightField<T>::SweepGeom(const FConvex& QueryGeom, const TRigidTransform<T, 3>& StartTM, const TVector<T, 3>& Dir, const T Length, T& OutTime, TVector<T, 3>& OutPosition, TVector<T, 3>& OutNormal, int32& OutFaceIndex, const T Thickness, bool bComputeMTD) const
+	bool FHeightField::SweepGeom(const FConvex& QueryGeom, const FRigidTransform3& StartTM, const FVec3& Dir, const FReal Length, FReal& OutTime, FVec3& OutPosition, FVec3& OutNormal, int32& OutFaceIndex, const FReal Thickness, bool bComputeMTD) const
 	{
 		return SweepGeomImp(QueryGeom, StartTM, Dir, Length, OutTime, OutPosition, OutNormal, OutFaceIndex, Thickness, bComputeMTD);
 	}
 
-	template <typename T>
-	bool THeightField<T>::SweepGeom(const TImplicitObjectScaled<TSphere<T, 3>>& QueryGeom, const TRigidTransform<T, 3>& StartTM, const TVector<T, 3>& Dir, const T Length, T& OutTime, TVector<T, 3>& OutPosition, TVector<T, 3>& OutNormal, int32& OutFaceIndex, const T Thickness, bool bComputeMTD) const
+	bool FHeightField::SweepGeom(const TImplicitObjectScaled<TSphere<FReal, 3>>& QueryGeom, const FRigidTransform3& StartTM, const FVec3& Dir, const FReal Length, FReal& OutTime, FVec3& OutPosition, FVec3& OutNormal, int32& OutFaceIndex, const FReal Thickness, bool bComputeMTD) const
 	{
 		return SweepGeomImp(QueryGeom, StartTM, Dir, Length, OutTime, OutPosition, OutNormal, OutFaceIndex, Thickness, bComputeMTD);
 	}
 
-	template <typename T>
-	bool THeightField<T>::SweepGeom(const TImplicitObjectScaled<TBox<T, 3>>& QueryGeom, const TRigidTransform<T, 3>& StartTM, const TVector<T, 3>& Dir, const T Length, T& OutTime, TVector<T, 3>& OutPosition, TVector<T, 3>& OutNormal, int32& OutFaceIndex, const T Thickness, bool bComputeMTD) const
+	bool FHeightField::SweepGeom(const TImplicitObjectScaled<TBox<FReal, 3>>& QueryGeom, const FRigidTransform3& StartTM, const FVec3& Dir, const FReal Length, FReal& OutTime, FVec3& OutPosition, FVec3& OutNormal, int32& OutFaceIndex, const FReal Thickness, bool bComputeMTD) const
 	{
 		return SweepGeomImp(QueryGeom, StartTM, Dir, Length, OutTime, OutPosition, OutNormal, OutFaceIndex, Thickness, bComputeMTD);
 	}
 
-	template <typename T>
-	bool THeightField<T>::SweepGeom(const TImplicitObjectScaled<TCapsule<T>>& QueryGeom, const TRigidTransform<T, 3>& StartTM, const TVector<T, 3>& Dir, const T Length, T& OutTime, TVector<T, 3>& OutPosition, TVector<T, 3>& OutNormal, int32& OutFaceIndex, const T Thickness, bool bComputeMTD) const
+	bool FHeightField::SweepGeom(const TImplicitObjectScaled<TCapsule<FReal>>& QueryGeom, const FRigidTransform3& StartTM, const FVec3& Dir, const FReal Length, FReal& OutTime, FVec3& OutPosition, FVec3& OutNormal, int32& OutFaceIndex, const FReal Thickness, bool bComputeMTD) const
 	{
 		return SweepGeomImp(QueryGeom, StartTM, Dir, Length, OutTime, OutPosition, OutNormal, OutFaceIndex, Thickness, bComputeMTD);
 	}
 
-	template <typename T>
-	bool THeightField<T>::SweepGeom(const TImplicitObjectScaled<FConvex>& QueryGeom, const TRigidTransform<T, 3>& StartTM, const TVector<T, 3>& Dir, const T Length, T& OutTime, TVector<T, 3>& OutPosition, TVector<T, 3>& OutNormal, int32& OutFaceIndex, const T Thickness, bool bComputeMTD) const
+	bool FHeightField::SweepGeom(const TImplicitObjectScaled<FConvex>& QueryGeom, const FRigidTransform3& StartTM, const FVec3& Dir, const FReal Length, FReal& OutTime, FVec3& OutPosition, FVec3& OutNormal, int32& OutFaceIndex, const FReal Thickness, bool bComputeMTD) const
 	{
 		return SweepGeomImp(QueryGeom, StartTM, Dir, Length, OutTime, OutPosition, OutNormal, OutFaceIndex, Thickness, bComputeMTD);
 	}
 
-	template <typename T>
-	int32 THeightField<T>::FindMostOpposingFace(const TVector<T, 3>& Position, const TVector<T, 3>& UnitDir, int32 HintFaceIndex, T SearchDist) const
+	int32 FHeightField::FindMostOpposingFace(const FVec3& Position, const FVec3& UnitDir, int32 HintFaceIndex, FReal SearchDist) const
 	{
-		const T SearchDist2 = SearchDist * SearchDist;
+		const FReal SearchDist2 = SearchDist * SearchDist;
 
-		TAABB<T, 3> QueryBounds(Position - TVector<T, 3>(SearchDist), Position + TVector<T, 3>(SearchDist));
+		TAABB<FReal, 3> QueryBounds(Position - FVec3(SearchDist), Position + FVec3(SearchDist));
 		const FBounds2D FlatBounds(QueryBounds);
 		TArray<TVector<int32, 2>> PotentialIntersections;
 		GetGridIntersections(FlatBounds, PotentialIntersections);
 
-		T MostOpposingDot = TNumericLimits<T>::Max();
+		FReal MostOpposingDot = TNumericLimits<FReal>::Max();
 		int32 MostOpposingFace = HintFaceIndex;
 
-		auto CheckTriangle = [&](int32 FaceIndex, const TVector<T, 3>& A, const TVector<T, 3>& B, const TVector<T, 3>& C)
+		auto CheckTriangle = [&](int32 FaceIndex, const FVec3& A, const FVec3& B, const FVec3& C)
 		{
-			const TVector<T, 3> AB = B - A;
-			const TVector<T, 3> AC = C - A;
-			TVector<T, 3> Normal = TVector<T, 3>::CrossProduct(AB, AC);
-			const T NormalLength = Normal.SafeNormalize();
+			const FVec3 AB = B - A;
+			const FVec3 AC = C - A;
+			FVec3 Normal = FVec3::CrossProduct(AB, AC);
+			const FReal NormalLength = Normal.SafeNormalize();
 			if(!ensure(NormalLength > KINDA_SMALL_NUMBER))
 			{
 				//hitting degenerate triangle - should be fixed before we get to this stage
 				return;
 			}
 
-			const TPlane<T, 3> TriPlane{A, Normal};
-			const TVector<T, 3> ClosestPointOnTri = FindClosestPointOnTriangle(TriPlane, A, B, C, Position);
-			const T Distance2 = (ClosestPointOnTri - Position).SizeSquared();
+			const TPlane<FReal, 3> TriPlane{A, Normal};
+			const FVec3 ClosestPointOnTri = FindClosestPointOnTriangle(TriPlane, A, B, C, Position);
+			const FReal Distance2 = (ClosestPointOnTri - Position).SizeSquared();
 			if(Distance2 < SearchDist2)
 			{
-				const T Dot = TVector<T, 3>::DotProduct(Normal, UnitDir);
+				const FReal Dot = FVec3::DotProduct(Normal, UnitDir);
 				if(Dot < MostOpposingDot)
 				{
 					MostOpposingDot = Dot;
@@ -1621,7 +1570,7 @@ namespace Chaos
 
 			const int32 FullIndex = CellIndex + SubY;
 
-			TVector<T, 3> Points[4];
+			FVec3 Points[4];
 
 			GeomData.GetPointsScaled(FullIndex, Points);
 
@@ -1632,8 +1581,7 @@ namespace Chaos
 		return MostOpposingFace;
 	}
 
-	template <typename T>
-	TVector<T, 3> THeightField<T>::FindGeometryOpposingNormal(const TVector<T, 3>& DenormDir, int32 FaceIndex, const TVector<T, 3>& OriginalNormal) const
+	FVec3 FHeightField::FindGeometryOpposingNormal(const FVec3& DenormDir, int32 FaceIndex, const FVec3& OriginalNormal) const
 	{
 		if(ensure(FaceIndex != INDEX_NONE))
 		{
@@ -1642,13 +1590,13 @@ namespace Chaos
 			int32 CellIndex = FaceIndex / 2;
 			int32 CellY = CellIndex / (GeomData.NumCols - 1);
 
-			TVector<T, 3> Points[4];
+			FVec3 Points[4];
 
 			GeomData.GetPointsScaled(CellIndex + CellY, Points);
 
-			TVector<T, 3> A;
-			TVector<T, 3> B;
-			TVector<T, 3> C;
+			FVec3 A;
+			FVec3 B;
+			FVec3 C;
 
 			if(bSecondFace)
 			{
@@ -1663,19 +1611,18 @@ namespace Chaos
 				C = Points[3];
 			}
 
-			const TVector<T, 3> AB = B - A;
-			const TVector<T, 3> AC = C - A;
-			TVector<T, 3> Normal = TVector<T, 3>::CrossProduct(AB, AC);
-			const T Length = Normal.SafeNormalize();
+			const FVec3 AB = B - A;
+			const FVec3 AC = C - A;
+			FVec3 Normal = FVec3::CrossProduct(AB, AC);
+			const FReal Length = Normal.SafeNormalize();
 			ensure(Length);
 			return Normal;
 		}
 
-		return TVector<T, 3>(0, 0, 1);
+		return FVec3(0, 0, 1);
 	}
 
-	template<typename T>
-	void THeightField<T>::CalcBounds()
+	void FHeightField::CalcBounds()
 	{
 		// Flatten out the Z axis
 		FlattenedBounds = GetFlatBounds();
@@ -1687,21 +1634,16 @@ namespace Chaos
 		const int32 NumY = GeomData.NumRows - 1;
 	}
 
-	template<typename T>
-	void THeightField<T>::BuildQueryData()
+	void FHeightField::BuildQueryData()
 	{
 		// NumCols and NumRows are the actual heights, there are n-1 cells between those heights
 		TVector<int32, 2> Cells(GeomData.NumCols - 1, GeomData.NumRows - 1);
 		
-		TVector<T, 2> MinCorner(0, 0);
-		TVector<T, 2> MaxCorner(GeomData.NumCols - 1, GeomData.NumRows - 1);
+		TVector<FReal, 2> MinCorner(0, 0);
+		TVector<FReal, 2> MaxCorner(GeomData.NumCols - 1, GeomData.NumRows - 1);
 		//MaxCorner *= {GeomData.Scale[0], GeomData.Scale[1]};
 
-		FlatGrid = TUniformGrid<T, 2>(MinCorner, MaxCorner, Cells);
+		FlatGrid = TUniformGrid<FReal, 2>(MinCorner, MaxCorner, Cells);
 	}
 
-	template class THeightField<float>;
-
 }
-
-
diff --git a/Engine/Source/Runtime/Experimental/Chaos/Private/Chaos/ImplicitObject.cpp b/Engine/Source/Runtime/Experimental/Chaos/Private/Chaos/ImplicitObject.cpp
--- a/Engine/Source/Runtime/Experimental/Chaos/Private/Chaos/ImplicitObject.cpp
+++ b/Engine/Source/Runtime/Experimental/Chaos/Private/Chaos/ImplicitObject.cpp
@@ -433,7 +433,7 @@ FImplicitObject* FImplicitObject::SerializationFactory(FChaosArchive& Ar, FImpli
 		ensure(Ar.IsLoading() && (Ar.CustomVer(FExternalPhysicsCustomObjectVersion::GUID) < FExternalPhysicsCustomObjectVersion::ScaledGeometryIsConcrete));
 		return new TImplicitObjectScaledGeneric<FReal, 3>();
 	}
-	case ImplicitObjectType::HeightField: if (Ar.IsLoading()) { return new THeightField<FReal>(); } break;
+	case ImplicitObjectType::HeightField: if (Ar.IsLoading()) { return new FHeightField(); } break;
 	case ImplicitObjectType::Cylinder: if (Ar.IsLoading()) { return new TCylinder<FReal>(); } break;
 	default:
 		check(false);
diff --git a/Engine/Source/Runtime/Experimental/Chaos/Public/Chaos/CollisionResolution.h b/Engine/Source/Runtime/Experimental/Chaos/Public/Chaos/CollisionResolution.h
--- a/Engine/Source/Runtime/Experimental/Chaos/Public/Chaos/CollisionResolution.h
+++ b/Engine/Source/Runtime/Experimental/Chaos/Public/Chaos/CollisionResolution.h
@@ -21,11 +21,9 @@ namespace Chaos
 	template <typename T, int d>
 	class TPlane;
 
-	template <typename T>
-	class THeightField;
-
 	class FCollisionContext;
 	class FCollisionConstraintBase;
+	class FHeightField;
 	class FImplicitObject;
 	class FRigidBodyPointContactConstraint;
 	class FRigidBodyMultiPointContactConstraint;
diff --git a/Engine/Source/Runtime/Experimental/Chaos/Public/Chaos/GeometryQueries.h b/Engine/Source/Runtime/Experimental/Chaos/Public/Chaos/GeometryQueries.h
--- a/Engine/Source/Runtime/Experimental/Chaos/Public/Chaos/GeometryQueries.h
+++ b/Engine/Source/Runtime/Experimental/Chaos/Public/Chaos/GeometryQueries.h
@@ -83,7 +83,7 @@ namespace Chaos
 			{
 				case ImplicitObjectType::HeightField:
 				{
-					const THeightField<FReal>& AHeightField = static_cast<const THeightField<FReal>&>(A);
+					const FHeightField& AHeightField = static_cast<const FHeightField&>(A);
 					return AHeightField.OverlapGeom(B, BToATM, Thickness, OutMTD);
 				}
 				case ImplicitObjectType::TriangleMesh:
@@ -197,7 +197,7 @@ namespace Chaos
 			{
 				case ImplicitObjectType::HeightField:
 				{
-					const THeightField<FReal>& AHeightField = static_cast<const THeightField<FReal>&>(A);
+					const FHeightField& AHeightField = static_cast<const FHeightField&>(A);
 					bResult = AHeightField.SweepGeom(B, BToATM, LocalDir, Length, OutTime, LocalPosition, LocalNormal, OutFaceIndex, Thickness, bComputeMTD);
 					break;
 				}
diff --git a/Engine/Source/Runtime/Experimental/Chaos/Public/Chaos/HeightField.h b/Engine/Source/Runtime/Experimental/Chaos/Public/Chaos/HeightField.h
--- a/Engine/Source/Runtime/Experimental/Chaos/Public/Chaos/HeightField.h
+++ b/Engine/Source/Runtime/Experimental/Chaos/Public/Chaos/HeightField.h
@@ -16,80 +16,80 @@
 
 namespace Chaos
 {
-	template<typename T>
-	class THeightfieldRaycastVisitor;
+	class FHeightfieldRaycastVisitor;
 	class FConvex;
 	struct FMTDInfo;
 }
 
 namespace Chaos
 {
-	template<typename T>
-	class THeightField final : public FImplicitObject
+	class CHAOS_API FHeightField final : public FImplicitObject
 	{
 	public:
 		using FImplicitObject::GetTypeName;
 
-		CHAOS_API THeightField(TArray<T>&& Height, TArray<uint8>&& InMaterialIndices, int32 InNumRows, int32 InNumCols, const TVector<T,3>& InScale);
-		CHAOS_API THeightField(TArrayView<const uint16> InHeights, TArrayView<uint8> InMaterialIndices, int32 InNumRows, int32 InNumCols, const TVector<T, 3>& InScale);
-		THeightField(const THeightField& Other) = delete;
-		THeightField(THeightField&& Other) = default;
+		FHeightField(TArray<FReal>&& Height, TArray<uint8>&& InMaterialIndices, int32 InNumRows, int32 InNumCols, const FVec3& InScale);
+		FHeightField(TArrayView<const uint16> InHeights, TArrayView<uint8> InMaterialIndices, int32 InNumRows, int32 InNumCols, const FVec3& InScale);
+		FHeightField(const FHeightField& Other) = delete;
+		
+		// Not required as long as FImplicitObject also has deleted move constructor (adding this causes an error on Linux build)
+		//FHeightField(FHeightField&& Other) = default;
 
-		virtual ~THeightField() {}
+		virtual ~FHeightField() {}
 
 		/** Support for editing a subsection of the heightfield */
-		CHAOS_API void EditHeights(TArrayView<T> InHeights, int32 InBeginRow, int32 InBeginCol, int32 InNumRows, int32 InNumCols);
-		CHAOS_API void EditHeights(TArrayView<const uint16> InHeights, int32 InBeginRow, int32 InBeginCol, int32 InNumRows, int32 InNumCols);
-		CHAOS_API T GetHeight(int32 InIndex) const;
-		CHAOS_API T GetHeight(int32 InX, int32 InY) const;
-		CHAOS_API uint8 GetMaterialIndex(int32 InIndex) const;
-		CHAOS_API uint8 GetMaterialIndex(int32 InX, int32 InY) const;
-		CHAOS_API bool IsHole(int32 InIndex) const;
-		CHAOS_API bool IsHole(int32 InCellX, int32 InCellY) const;
-		CHAOS_API T GetHeightAt(const TVector<T, 2>& InGridLocationLocal) const;
+		void EditHeights(TArrayView<FReal> InHeights, int32 InBeginRow, int32 InBeginCol, int32 InNumRows, int32 InNumCols);
+		void EditHeights(TArrayView<const uint16> InHeights, int32 InBeginRow, int32 InBeginCol, int32 InNumRows, int32 InNumCols);
+		FReal GetHeight(int32 InIndex) const;
+		FReal GetHeight(int32 InX, int32 InY) const;
+		uint8 GetMaterialIndex(int32 InIndex) const;
+		uint8 GetMaterialIndex(int32 InX, int32 InY) const;
+		bool IsHole(int32 InIndex) const;
+		bool IsHole(int32 InCellX, int32 InCellY) const;
+		FReal GetHeightAt(const TVector<FReal, 2>& InGridLocationLocal) const;
 
 		int32 GetNumRows() const { return GeomData.NumRows; }
 		int32 GetNumCols() const { return GeomData.NumCols; }
 
-		virtual T PhiWithNormal(const TVector<T, 3>& x, TVector<T, 3>& Normal) const
+		virtual FReal PhiWithNormal(const FVec3& x, FVec3& Normal) const
 		{
 			check(false);	//not supported yet - might support it in the future or we may change the interface
-			return (T)0;
+			return (FReal)0;
 		}
 
-		CHAOS_API virtual bool Raycast(const TVector<T, 3>& StartPoint, const TVector<T, 3>& Dir, const T Length, const T Thickness, T& OutTime, TVector<T,3>& OutPosition, TVector<T,3>& OutNormal, int32& OutFaceIndex) const override;
-		CHAOS_API virtual bool Overlap(const TVector<T, 3>& Point, const T Thickness) const override;
+		virtual bool Raycast(const FVec3& StartPoint, const FVec3& Dir, const FReal Length, const FReal Thickness, FReal& OutTime, FVec3& OutPosition, FVec3& OutNormal, int32& OutFaceIndex) const override;
+		virtual bool Overlap(const FVec3& Point, const FReal Thickness) const override;
 		
-		CHAOS_API bool OverlapGeom(const TSphere<T, 3>& QueryGeom, const TRigidTransform<T, 3>& QueryTM, const T Thickness, FMTDInfo* OutMTD = nullptr) const;
-		CHAOS_API bool OverlapGeom(const TBox<T, 3>& QueryGeom, const TRigidTransform<T, 3>& QueryTM, const T Thickness, FMTDInfo* OutMTD = nullptr) const;
-		CHAOS_API bool OverlapGeom(const TCapsule<T>& QueryGeom, const TRigidTransform<T, 3>& QueryTM, const T Thickness, FMTDInfo* OutMTD = nullptr) const;
-		CHAOS_API bool OverlapGeom(const FConvex& QueryGeom, const TRigidTransform<T, 3>& QueryTM, const T Thickness, FMTDInfo* OutMTD = nullptr) const;
-		CHAOS_API bool OverlapGeom(const TImplicitObjectScaled<TSphere<T, 3>>& QueryGeom, const TRigidTransform<T, 3>& QueryTM, const T Thickness, FMTDInfo* OutMTD = nullptr) const;
-		CHAOS_API bool OverlapGeom(const TImplicitObjectScaled<TBox<T, 3>>& QueryGeom, const TRigidTransform<T, 3>& QueryTM, const T Thickness, FMTDInfo* OutMTD = nullptr) const;
-		CHAOS_API bool OverlapGeom(const TImplicitObjectScaled<TCapsule<T>>& QueryGeom, const TRigidTransform<T, 3>& QueryTM, const T Thickness, FMTDInfo* OutMTD = nullptr) const;
-		CHAOS_API bool OverlapGeom(const TImplicitObjectScaled<FConvex>& QueryGeom, const TRigidTransform<T, 3>& QueryTM, const T Thickness, FMTDInfo* OutMTD = nullptr) const;
+		bool OverlapGeom(const TSphere<FReal, 3>& QueryGeom, const FRigidTransform3& QueryTM, const FReal Thickness, FMTDInfo* OutMTD = nullptr) const;
+		bool OverlapGeom(const TBox<FReal, 3>& QueryGeom, const FRigidTransform3& QueryTM, const FReal Thickness, FMTDInfo* OutMTD = nullptr) const;
+		bool OverlapGeom(const TCapsule<FReal>& QueryGeom, const FRigidTransform3& QueryTM, const FReal Thickness, FMTDInfo* OutMTD = nullptr) const;
+		bool OverlapGeom(const FConvex& QueryGeom, const FRigidTransform3& QueryTM, const FReal Thickness, FMTDInfo* OutMTD = nullptr) const;
+		bool OverlapGeom(const TImplicitObjectScaled<TSphere<FReal, 3>>& QueryGeom, const FRigidTransform3& QueryTM, const FReal Thickness, FMTDInfo* OutMTD = nullptr) const;
+		bool OverlapGeom(const TImplicitObjectScaled<TBox<FReal, 3>>& QueryGeom, const FRigidTransform3& QueryTM, const FReal Thickness, FMTDInfo* OutMTD = nullptr) const;
+		bool OverlapGeom(const TImplicitObjectScaled<TCapsule<FReal>>& QueryGeom, const FRigidTransform3& QueryTM, const FReal Thickness, FMTDInfo* OutMTD = nullptr) const;
+		bool OverlapGeom(const TImplicitObjectScaled<FConvex>& QueryGeom, const FRigidTransform3& QueryTM, const FReal Thickness, FMTDInfo* OutMTD = nullptr) const;
 
-		CHAOS_API bool SweepGeom(const TSphere<T, 3>& QueryGeom, const TRigidTransform<T, 3>& StartTM, const TVector<T, 3>& Dir, const T Length, T& OutTime, TVector<T, 3>& OutPosition, TVector<T, 3>& OutNormal, int32& OutFaceIndex, const T Thickness = 0, bool bComputeMTD = false) const;
-		CHAOS_API bool SweepGeom(const TBox<T, 3>& QueryGeom, const TRigidTransform<T, 3>& StartTM, const TVector<T, 3>& Dir, const T Length, T& OutTime, TVector<T, 3>& OutPosition, TVector<T, 3>& OutNormal, int32& OutFaceIndex, const T Thickness = 0, bool bComputeMTD = false) const;
-		CHAOS_API bool SweepGeom(const TCapsule<T>& QueryGeom, const TRigidTransform<T, 3>& StartTM, const TVector<T, 3>& Dir, const T Length, T& OutTime, TVector<T, 3>& OutPosition, TVector<T, 3>& OutNormal, int32& OutFaceIndex, const T Thickness = 0, bool bComputeMTD = false) const;
-		CHAOS_API bool SweepGeom(const FConvex& QueryGeom, const TRigidTransform<T, 3>& StartTM, const TVector<T, 3>& Dir, const T Length, T& OutTime, TVector<T, 3>& OutPosition, TVector<T, 3>& OutNormal, int32& OutFaceIndex, const T Thickness = 0, bool bComputeMTD = false) const;
-		CHAOS_API bool SweepGeom(const TImplicitObjectScaled<TSphere<T, 3>>& QueryGeom, const TRigidTransform<T, 3>& StartTM, const TVector<T, 3>& Dir, const T Length, T& OutTime, TVector<T, 3>& OutPosition, TVector<T, 3>& OutNormal, int32& OutFaceIndex, const T Thickness = 0, bool bComputeMTD = false) const;
-		CHAOS_API bool SweepGeom(const TImplicitObjectScaled<TBox<T, 3>>& QueryGeom, const TRigidTransform<T, 3>& StartTM, const TVector<T, 3>& Dir, const T Length, T& OutTime, TVector<T, 3>& OutPosition, TVector<T, 3>& OutNormal, int32& OutFaceIndex, const T Thickness = 0, bool bComputeMTD = false) const;
-		CHAOS_API bool SweepGeom(const TImplicitObjectScaled<TCapsule<T>>& QueryGeom, const TRigidTransform<T, 3>& StartTM, const TVector<T, 3>& Dir, const T Length, T& OutTime, TVector<T, 3>& OutPosition, TVector<T, 3>& OutNormal, int32& OutFaceIndex, const T Thickness = 0, bool bComputeMTD = false) const;
-		CHAOS_API bool SweepGeom(const TImplicitObjectScaled<FConvex>& QueryGeom, const TRigidTransform<T, 3>& StartTM, const TVector<T, 3>& Dir, const T Length, T& OutTime, TVector<T, 3>& OutPosition, TVector<T, 3>& OutNormal, int32& OutFaceIndex, const T Thickness = 0, bool bComputeMTD = false) const;
+		bool SweepGeom(const TSphere<FReal, 3>& QueryGeom, const FRigidTransform3& StartTM, const FVec3& Dir, const FReal Length, FReal& OutTime, FVec3& OutPosition, FVec3& OutNormal, int32& OutFaceIndex, const FReal Thickness = 0, bool bComputeMTD = false) const;
+		bool SweepGeom(const TBox<FReal, 3>& QueryGeom, const FRigidTransform3& StartTM, const FVec3& Dir, const FReal Length, FReal& OutTime, FVec3& OutPosition, FVec3& OutNormal, int32& OutFaceIndex, const FReal Thickness = 0, bool bComputeMTD = false) const;
+		bool SweepGeom(const TCapsule<FReal>& QueryGeom, const FRigidTransform3& StartTM, const FVec3& Dir, const FReal Length, FReal& OutTime, FVec3& OutPosition, FVec3& OutNormal, int32& OutFaceIndex, const FReal Thickness = 0, bool bComputeMTD = false) const;
+		bool SweepGeom(const FConvex& QueryGeom, const FRigidTransform3& StartTM, const FVec3& Dir, const FReal Length, FReal& OutTime, FVec3& OutPosition, FVec3& OutNormal, int32& OutFaceIndex, const FReal Thickness = 0, bool bComputeMTD = false) const;
+		bool SweepGeom(const TImplicitObjectScaled<TSphere<FReal, 3>>& QueryGeom, const FRigidTransform3& StartTM, const FVec3& Dir, const FReal Length, FReal& OutTime, FVec3& OutPosition, FVec3& OutNormal, int32& OutFaceIndex, const FReal Thickness = 0, bool bComputeMTD = false) const;
+		bool SweepGeom(const TImplicitObjectScaled<TBox<FReal, 3>>& QueryGeom, const FRigidTransform3& StartTM, const FVec3& Dir, const FReal Length, FReal& OutTime, FVec3& OutPosition, FVec3& OutNormal, int32& OutFaceIndex, const FReal Thickness = 0, bool bComputeMTD = false) const;
+		bool SweepGeom(const TImplicitObjectScaled<TCapsule<FReal>>& QueryGeom, const FRigidTransform3& StartTM, const FVec3& Dir, const FReal Length, FReal& OutTime, FVec3& OutPosition, FVec3& OutNormal, int32& OutFaceIndex, const FReal Thickness = 0, bool bComputeMTD = false) const;
+		bool SweepGeom(const TImplicitObjectScaled<FConvex>& QueryGeom, const FRigidTransform3& StartTM, const FVec3& Dir, const FReal Length, FReal& OutTime, FVec3& OutPosition, FVec3& OutNormal, int32& OutFaceIndex, const FReal Thickness = 0, bool bComputeMTD = false) const;
 
-		CHAOS_API bool GJKContactPoint(const TBox<T, 3>& QueryGeom, const TRigidTransform<T, 3>& QueryTM, const T Thickness, TVector<T, 3>& ContactLocation, TVector<T, 3>& ContactNormal, T& ContactPhi) const;
-		CHAOS_API bool GJKContactPoint(const TSphere<T, 3>& QueryGeom, const TRigidTransform<T, 3>& QueryTM, const T Thickness, TVector<T, 3>& ContactLocation, TVector<T, 3>& ContactNormal, T& ContactPhi) const;
-		CHAOS_API bool GJKContactPoint(const TCapsule<T>& QueryGeom, const TRigidTransform<T, 3>& QueryTM, const T Thickness, TVector<T, 3>& ContactLocation, TVector<T, 3>& ContactNormal, T& ContactPhi) const;
-		CHAOS_API bool GJKContactPoint(const FConvex& QueryGeom, const TRigidTransform<T, 3>& QueryTM, const T Thickness, TVector<T, 3>& ContactLocation, TVector<T, 3>& ContactNormal, T& ContactPhi) const;
-		CHAOS_API bool GJKContactPoint(const TImplicitObjectScaled<TBox<T, 3>>& QueryGeom, const TRigidTransform<T, 3>& QueryTM, const T Thickness, TVector<T, 3>& ContactLocation, TVector<T, 3>& ContactNormal, T& ContactPhi) const;
-		CHAOS_API bool GJKContactPoint(const TImplicitObjectScaled<TSphere<T, 3>>& QueryGeom, const TRigidTransform<T, 3>& QueryTM, const T Thickness, TVector<T, 3>& ContactLocation, TVector<T, 3>& ContactNormal, T& ContactPhi) const;
-		CHAOS_API bool GJKContactPoint(const TImplicitObjectScaled<TCapsule<T>>& QueryGeom, const TRigidTransform<T, 3>& QueryTM, const T Thickness, TVector<T, 3>& ContactLocation, TVector<T, 3>& ContactNormal, T& ContactPhi) const;
-		CHAOS_API bool GJKContactPoint(const TImplicitObjectScaled<FConvex>& QueryGeom, const TRigidTransform<T, 3>& QueryTM, const T Thickness, TVector<T, 3>& ContactLocation, TVector<T, 3>& ContactNormal, T& ContactPhi) const;
+		bool GJKContactPoint(const TBox<FReal, 3>& QueryGeom, const FRigidTransform3& QueryTM, const FReal Thickness, FVec3& ContactLocation, FVec3& ContactNormal, FReal& ContactPhi) const;
+		bool GJKContactPoint(const TSphere<FReal, 3>& QueryGeom, const FRigidTransform3& QueryTM, const FReal Thickness, FVec3& ContactLocation, FVec3& ContactNormal, FReal& ContactPhi) const;
+		bool GJKContactPoint(const TCapsule<FReal>& QueryGeom, const FRigidTransform3& QueryTM, const FReal Thickness, FVec3& ContactLocation, FVec3& ContactNormal, FReal& ContactPhi) const;
+		bool GJKContactPoint(const FConvex& QueryGeom, const FRigidTransform3& QueryTM, const FReal Thickness, FVec3& ContactLocation, FVec3& ContactNormal, FReal& ContactPhi) const;
+		bool GJKContactPoint(const TImplicitObjectScaled<TBox<FReal, 3>>& QueryGeom, const FRigidTransform3& QueryTM, const FReal Thickness, FVec3& ContactLocation, FVec3& ContactNormal, FReal& ContactPhi) const;
+		bool GJKContactPoint(const TImplicitObjectScaled<TSphere<FReal, 3>>& QueryGeom, const FRigidTransform3& QueryTM, const FReal Thickness, FVec3& ContactLocation, FVec3& ContactNormal, FReal& ContactPhi) const;
+		bool GJKContactPoint(const TImplicitObjectScaled<TCapsule<FReal>>& QueryGeom, const FRigidTransform3& QueryTM, const FReal Thickness, FVec3& ContactLocation, FVec3& ContactNormal, FReal& ContactPhi) const;
+		bool GJKContactPoint(const TImplicitObjectScaled<FConvex>& QueryGeom, const FRigidTransform3& QueryTM, const FReal Thickness, FVec3& ContactLocation, FVec3& ContactNormal, FReal& ContactPhi) const;
 
 
-		CHAOS_API virtual int32 FindMostOpposingFace(const TVector<T, 3>& Position, const TVector<T, 3>& UnitDir, int32 HintFaceIndex, T SearchDist) const override;
-		CHAOS_API virtual TVector<T, 3> FindGeometryOpposingNormal(const TVector<T, 3>& DenormDir, int32 FaceIndex, const TVector<T, 3>& OriginalNormal) const override;
+		virtual int32 FindMostOpposingFace(const FVec3& Position, const FVec3& UnitDir, int32 HintFaceIndex, FReal SearchDist) const override;
+		virtual FVec3 FindGeometryOpposingNormal(const FVec3& DenormDir, int32 FaceIndex, const FVec3& OriginalNormal) const override;
 
 
 
@@ -116,9 +116,9 @@ namespace Chaos
 			return 0;
 		}
 
-		virtual const TAABB<T, 3> BoundingBox() const
+		virtual const FAABB3 BoundingBox() const
 		{
-			CachedBounds = TAABB<T, 3>(LocalBounds.Min() * GeomData.Scale, LocalBounds.Max() * GeomData.Scale);
+			CachedBounds = FAABB3(LocalBounds.Min() * GeomData.Scale, LocalBounds.Max() * GeomData.Scale);
 			return CachedBounds;
 		}
 
@@ -130,7 +130,7 @@ namespace Chaos
 
 			// Saving to an archive is a const operation, but must be non-const
 			// to support loading. Cast const away here to get bytes written
-			const_cast<THeightField<T>*>(this)->Serialize(ChaosAr);
+			const_cast<FHeightField*>(this)->Serialize(ChaosAr);
 
 			return FCrc::MemCrc32(Bytes.GetData(), Bytes.GetAllocatedSize());
 		}
@@ -168,7 +168,7 @@ namespace Chaos
 			}
 		}
 
-		void SetScale(const TVector<T, 3>& InScale)
+		void SetScale(const FVec3& InScale)
 		{
 			GeomData.Scale = InScale;
 		}
@@ -219,7 +219,7 @@ namespace Chaos
 				return Scale[1];
 			}
 
-			FORCEINLINE TVector<T, 3> GetPoint(int32 Index) const
+			FORCEINLINE FVec3 GetPoint(int32 Index) const
 			{
 				const typename FDataType::RealType Height = MinValue + Heights[Index] * HeightPerUnit;
 
@@ -229,12 +229,12 @@ namespace Chaos
 				return {(typename FDataType::RealType)X, (typename FDataType::RealType)Y, Height};
 			}
 
-			FORCEINLINE TVector<T, 3> GetPointScaled(int32 Index) const
+			FORCEINLINE FVec3 GetPointScaled(int32 Index) const
 			{
 				return GetPoint(Index) * Scale;
 			}
 
-			FORCEINLINE void GetPoints(int32 Index, TVector<T, 3> OutPts[4]) const
+			FORCEINLINE void GetPoints(int32 Index, FVec3 OutPts[4]) const
 			{
 				const typename FDataType::RealType H0 = MinValue + Heights[Index] * HeightPerUnit;
 				const typename FDataType::RealType H1 = MinValue + Heights[Index + 1] * HeightPerUnit;
@@ -250,7 +250,7 @@ namespace Chaos
 				OutPts[3] = {(typename FDataType::RealType)X + 1, (typename FDataType::RealType)Y + 1, H3};
 			}
 
-			FORCEINLINE void GetPointsScaled(int32 Index, TVector<T, 3> OutPts[4]) const
+			FORCEINLINE void GetPointsScaled(int32 Index, FVec3 OutPts[4]) const
 			{
 				GetPoints(Index, OutPts);
 
@@ -260,12 +260,12 @@ namespace Chaos
 				OutPts[3] *= Scale;
 			}
 
-			FORCEINLINE T GetMinHeight() const
+			FORCEINLINE FReal GetMinHeight() const
 			{
 				return static_cast<typename FDataType::RealType>(MinValue);
 			}
 
-			FORCEINLINE T GetMaxHeight() const
+			FORCEINLINE FReal GetMaxHeight() const
 			{
 				return static_cast<typename FDataType::RealType>(MaxValue);
 			}
@@ -324,42 +324,42 @@ namespace Chaos
 		// Struct for 2D bounds and associated operations
 		struct FBounds2D
 		{
-			TVector<T, 2> Min;
-			TVector<T, 2> Max;
+			TVector<FReal, 2> Min;
+			TVector<FReal, 2> Max;
 			
 			FBounds2D()
 				: Min(0)
 				, Max(0)
 			{}
 
-			explicit FBounds2D(const TAABB<T, 3>& In3DBounds)
+			explicit FBounds2D(const FAABB3& In3DBounds)
 			{
 				Set(In3DBounds);
 			}
 
-			void Set(const TAABB<T, 3>& In3DBounds)
+			void Set(const FAABB3& In3DBounds)
 			{
 				Min = {In3DBounds.Min()[0], In3DBounds.Min()[1]};
 				Max = {In3DBounds.Max()[0], In3DBounds.Max()[1]};
 			}
 
-			TVector<T, 2> GetExtent() const
+			TVector<FReal, 2> GetExtent() const
 			{
 				return Max - Min;
 			}
 
-			bool IsInside(const TVector<T, 2>& InPoint) const
+			bool IsInside(const TVector<FReal, 2>& InPoint) const
 			{
 				return InPoint[0] >= Min[0] && InPoint[0] <= Max[0] && InPoint[1] >= Min[1] && InPoint[1] <= Max[1];
 			}
 
-			TVector<T, 2> Clamp(const TVector<T, 2>& InToClamp, T InNudge = SMALL_NUMBER) const
+			TVector<FReal, 2> Clamp(const TVector<FReal, 2>& InToClamp, FReal InNudge = SMALL_NUMBER) const
 			{
-				const TVector<T, 2> NudgeVec(InNudge, InNudge);
-				const TVector<T, 2> TestMin = Min + NudgeVec;
-				const TVector<T, 2> TestMax = Max - NudgeVec;
+				const TVector<FReal, 2> NudgeVec(InNudge, InNudge);
+				const TVector<FReal, 2> TestMin = Min + NudgeVec;
+				const TVector<FReal, 2> TestMax = Max - NudgeVec;
 
-				TVector<T, 2> OutVec = InToClamp;
+				TVector<FReal, 2> OutVec = InToClamp;
 
 				OutVec[0] = FMath::Max(OutVec[0], TestMin[0]);
 				OutVec[1] = FMath::Max(OutVec[1], TestMin[1]);
@@ -370,20 +370,20 @@ namespace Chaos
 				return OutVec;
 			}
 
-			bool IntersectLine(const TVector<T, 2>& InStart, const TVector<T, 2>& InEnd)
+			bool IntersectLine(const TVector<FReal, 2>& InStart, const TVector<FReal, 2>& InEnd)
 			{
 				if(IsInside(InStart) || IsInside(InEnd))
 				{
 					return true;
 				}
 
-				const TVector<T, 2> Extent = GetExtent();
+				const TVector<FReal, 2> Extent = GetExtent();
 				float TA, TB;
 
-				if(Utilities::IntersectLineSegments2D(InStart, InEnd, Min, TVector<T, 2>(Min[0] + Extent[0], Min[1]), TA, TB)
-					|| Utilities::IntersectLineSegments2D(InStart, InEnd, Min, TVector<T, 2>(Min[0], Min[1] + Extent[1]), TA, TB)
-					|| Utilities::IntersectLineSegments2D(InStart, InEnd, Max, TVector<T, 2>(Max[0] - Extent[0], Max[1]), TA, TB)
-					|| Utilities::IntersectLineSegments2D(InStart, InEnd, Max, TVector<T, 2>(Max[0], Max[1] - Extent[1]), TA, TB))
+				if(Utilities::IntersectLineSegments2D(InStart, InEnd, Min, TVector<FReal, 2>(Min[0] + Extent[0], Min[1]), TA, TB)
+					|| Utilities::IntersectLineSegments2D(InStart, InEnd, Min, TVector<FReal, 2>(Min[0], Min[1] + Extent[1]), TA, TB)
+					|| Utilities::IntersectLineSegments2D(InStart, InEnd, Max, TVector<FReal, 2>(Max[0] - Extent[0], Max[1]), TA, TB)
+					|| Utilities::IntersectLineSegments2D(InStart, InEnd, Max, TVector<FReal, 2>(Max[0], Max[1] - Extent[1]), TA, TB))
 				{
 					return true;
 				}
@@ -391,10 +391,10 @@ namespace Chaos
 				return false;
 			}
 
-			bool ClipLine(const TVector<T, 3>& InStart, const TVector<T, 3>& InEnd, TVector<T, 2>& OutClippedStart, TVector<T, 2>& OutClippedEnd) const
+			bool ClipLine(const FVec3& InStart, const FVec3& InEnd, TVector<FReal, 2>& OutClippedStart, TVector<FReal, 2>& OutClippedEnd) const
 			{
-				TVector<T, 2> TempStart(InStart[0], InStart[1]);
-				TVector<T, 2> TempEnd(InEnd[0], InEnd[1]);
+				TVector<FReal, 2> TempStart(InStart[0], InStart[1]);
+				TVector<FReal, 2> TempEnd(InEnd[0], InEnd[1]);
 
 				bool bLineIntersects = ClipLine(TempStart, TempEnd);
 
@@ -404,7 +404,7 @@ namespace Chaos
 				return bLineIntersects;
 			}
 
-			bool ClipLine(TVector<T, 2>& InOutStart, TVector<T, 2>& InOutEnd) const
+			bool ClipLine(TVector<FReal, 2>& InOutStart, TVector<FReal, 2>& InOutEnd) const
 			{
 				
 				// Test we don't need to clip at all, quite likely with a heightfield so optimize for it.
@@ -415,7 +415,7 @@ namespace Chaos
 					return true;
 				}
 
-				const TVector<T,2> Dir = InOutEnd - InOutStart;
+				const TVector<FReal, 2> Dir = InOutEnd - InOutStart;
 
 				// Tiny ray not inside so must be outside
 				if(Dir.SizeSquared() < 1e-4)
@@ -424,7 +424,7 @@ namespace Chaos
 				}
 
 				bool bPerpendicular[2];
-				TVector<T,2> InvDir;
+				TVector<FReal, 2> InvDir;
 				for(int Axis = 0; Axis < 2; ++Axis)
 				{
 					bPerpendicular[Axis] = Dir[Axis] == 0;
@@ -435,21 +435,21 @@ namespace Chaos
 
 				if(bStartInside)
 				{
-					const T TimeToExit = ComputeTimeToExit(InOutStart,InvDir);
+					const FReal TimeToExit = ComputeTimeToExit(InOutStart,InvDir);
 					InOutEnd = InOutStart + Dir * TimeToExit;
 					return true;
 				}
 
 				if(bEndInside)
 				{
-					const T TimeToExit = ComputeTimeToExit(InOutEnd,-InvDir);
+					const FReal TimeToExit = ComputeTimeToExit(InOutEnd,-InvDir);
 					InOutStart = InOutEnd - Dir * TimeToExit;
 					return true;
 				}
 
 				//start and end outside, need to see if we even intersect
-				T TimesToEnter[2] = {TNumericLimits<T>::Max(),TNumericLimits<T>::Max()};
-				T TimesToExit[2] = {TNumericLimits<T>::Max(),TNumericLimits<T>::Max()};
+				FReal TimesToEnter[2] = {TNumericLimits<FReal>::Max(),TNumericLimits<FReal>::Max()};
+				FReal TimesToExit[2] = {TNumericLimits<FReal>::Max(),TNumericLimits<FReal>::Max()};
 				
 				for(int Axis = 0; Axis < 2; ++Axis)
 				{
@@ -466,7 +466,7 @@ namespace Chaos
 						{
 							if(InOutStart[Axis] <= Max[Axis])
 							{
-								TimesToEnter[Axis] = FMath::Max<T>(Min[Axis] - InOutStart[Axis], 0) * InvDir[Axis];
+								TimesToEnter[Axis] = FMath::Max<FReal>(Min[Axis] - InOutStart[Axis], 0) * InvDir[Axis];
 								TimesToExit[Axis] = (Max[Axis] - InOutStart[Axis])  * InvDir[Axis];
 							}
 						}
@@ -474,15 +474,15 @@ namespace Chaos
 						{
 							if(InOutStart[Axis] >= Min[Axis])
 							{
-								TimesToEnter[Axis] = FMath::Max<T>(InOutStart[Axis] - Max[Axis],0) * InvDir[Axis];
+								TimesToEnter[Axis] = FMath::Max<FReal>(InOutStart[Axis] - Max[Axis],0) * InvDir[Axis];
 								TimesToExit[Axis] = (InOutStart[Axis] - Min[Axis]) * InvDir[Axis];
 							}
 						}
 					}
 				}
 
-				const T TimeToEnter = FMath::Max(FMath::Abs(TimesToEnter[0]),FMath::Abs(TimesToEnter[1]));
-				const T TimeToExit = FMath::Min(FMath::Abs(TimesToExit[0]),FMath::Abs(TimesToExit[1]));
+				const FReal TimeToEnter = FMath::Max(FMath::Abs(TimesToEnter[0]),FMath::Abs(TimesToEnter[1]));
+				const FReal TimeToExit = FMath::Min(FMath::Abs(TimesToExit[0]),FMath::Abs(TimesToExit[1]));
 
 				if(TimeToExit < TimeToEnter)
 				{
@@ -497,9 +497,9 @@ namespace Chaos
 
 		private:
 			//This helper assumes Start is inside the min/max box and uses InvDir to compute how long it takes to exit
-			T ComputeTimeToExit(const TVector<T,2>& Start,const TVector<T,2>& InvDir) const
+			FReal ComputeTimeToExit(const TVector<FReal, 2>& Start,const TVector<FReal, 2>& InvDir) const
 			{
-				T Times[2] ={TNumericLimits<T>::Max(),TNumericLimits<T>::Max()};
+				FReal Times[2] ={TNumericLimits<FReal>::Max(),TNumericLimits<FReal>::Max()};
 				for(int Axis = 0; Axis < 2; ++Axis)
 				{
 					if(InvDir[Axis] > 0)
@@ -512,56 +512,50 @@ namespace Chaos
 					}
 				}
 
-				const T MinTime = FMath::Min(FMath::Abs(Times[0]),FMath::Abs(Times[1]));
+				const FReal MinTime = FMath::Min(FMath::Abs(Times[0]),FMath::Abs(Times[1]));
 				return MinTime;
 			}
 		};
 
 		// Helpers for accessing bounds
-		CHAOS_API bool GetCellBounds2D(const TVector<int32, 2> InCoord, FBounds2D& OutBounds, const TVector<T, 2>& InInflate = {0}) const;
-		CHAOS_API bool GetCellBounds3D(const TVector<int32, 2> InCoord, TVector<T, 3>& OutMin, TVector<T, 3>& OutMax, const TVector<T, 3>& InInflate = TVector<T, 3>(0)) const;
-		CHAOS_API bool GetCellBounds2DScaled(const TVector<int32, 2> InCoord, FBounds2D& OutBounds, const TVector<T, 2>& InInflate = {0}) const;
-		CHAOS_API bool GetCellBounds3DScaled(const TVector<int32, 2> InCoord, TVector<T, 3>& OutMin, TVector<T, 3>& OutMax, const TVector<T, 3>& InInflate = TVector<T, 3>(0)) const;
-		CHAOS_API bool CalcCellBounds3D(const TVector<int32, 2> InCoord, TVector<T, 3>& OutMin, TVector<T, 3>& OutMax, const TVector<T, 3>& InInflate = TVector<T, 3>(0)) const;
+		bool GetCellBounds2D(const TVector<int32, 2> InCoord, FBounds2D& OutBounds, const TVector<FReal, 2>& InInflate = {0}) const;
+		bool GetCellBounds3D(const TVector<int32, 2> InCoord, FVec3& OutMin, FVec3& OutMax, const FVec3& InInflate = FVec3(0)) const;
+		bool GetCellBounds2DScaled(const TVector<int32, 2> InCoord, FBounds2D& OutBounds, const TVector<FReal, 2>& InInflate = {0}) const;
+		bool GetCellBounds3DScaled(const TVector<int32, 2> InCoord, FVec3& OutMin, FVec3& OutMax, const FVec3& InInflate = FVec3(0)) const;
+		bool CalcCellBounds3D(const TVector<int32, 2> InCoord, FVec3& OutMin, FVec3& OutMax, const FVec3& InInflate = FVec3(0)) const;
 
 		// Query functions - sweep, ray, overlap
 		template<typename SQVisitor>
-		bool GridSweep(const TVector<T, 3>& StartPoint, const TVector<T, 3>& Dir, const T Length, const TVector<T, 2> InHalfExtents, SQVisitor& Visitor) const;
-		CHAOS_API bool GridCast(const TVector<T, 3>& StartPoint, const TVector<T, 3>& Dir, const T Length, THeightfieldRaycastVisitor<T>& Visitor) const;
-		CHAOS_API bool GetGridIntersections(FBounds2D InFlatBounds, TArray<TVector<int32, 2>>& OutInterssctions) const;
+		bool GridSweep(const FVec3& StartPoint, const FVec3& Dir, const FReal Length, const TVector<FReal, 2> InHalfExtents, SQVisitor& Visitor) const;
+		bool GridCast(const FVec3& StartPoint, const FVec3& Dir, const FReal Length, FHeightfieldRaycastVisitor& Visitor) const;
+		bool GetGridIntersections(FBounds2D InFlatBounds, TArray<TVector<int32, 2>>& OutInterssctions) const;
 		
-		CHAOS_API FBounds2D GetFlatBounds() const;
+		FBounds2D GetFlatBounds() const;
 
 		// Grid for queries, faster than bounding volumes for heightfields
-		TUniformGrid<T, 2> FlatGrid;
+		TUniformGrid<FReal, 2> FlatGrid;
 		// Bounds in 2D of the whole heightfield, to clip queries against
 		FBounds2D FlattenedBounds;
 		// 3D bounds for the heightfield, for insertion to the scene structure
-		TAABB<T, 3> LocalBounds;
+		FAABB3 LocalBounds;
 		// Cached when bounds are requested. Mutable to allow GetBounds to be logical const
-		mutable TAABB<T, 3> CachedBounds;
+		mutable FAABB3 CachedBounds;
 
-		CHAOS_API void CalcBounds();
-		CHAOS_API void BuildQueryData();
+		void CalcBounds();
+		void BuildQueryData();
 		
 		// Needed for serialization
-		THeightField() : FImplicitObject(EImplicitObject::HasBoundingBox, ImplicitObjectType::HeightField) {}
+		FHeightField() : FImplicitObject(EImplicitObject::HasBoundingBox, ImplicitObjectType::HeightField) {}
 		friend FImplicitObject;
 
 		template <typename QueryGeomType>
-		bool OverlapGeomImp(const QueryGeomType& QueryGeom, const TRigidTransform<T, 3>& QueryTM, const T Thickness, FMTDInfo* OutMTD = nullptr) const;
+		bool OverlapGeomImp(const QueryGeomType& QueryGeom, const FRigidTransform3& QueryTM, const FReal Thickness, FMTDInfo* OutMTD = nullptr) const;
 
 		template <typename QueryGeomType>
-		bool SweepGeomImp(const QueryGeomType& QueryGeom, const TRigidTransform<T, 3>& StartTM, const TVector<T, 3>& Dir, const T Length, T& OutTime, TVector<T, 3>& OutPosition, TVector<T, 3>& OutNormal, int32& OutFaceIndex, const T Thickness, bool bComputeMTD) const;
+		bool SweepGeomImp(const QueryGeomType& QueryGeom, const FRigidTransform3& StartTM, const FVec3& Dir, const FReal Length, FReal& OutTime, FVec3& OutPosition, FVec3& OutNormal, int32& OutFaceIndex, const FReal Thickness, bool bComputeMTD) const;
 
 		template <typename GeomType>
-		bool GJKContactPointImp(const GeomType& QueryGeom, const TRigidTransform<T, 3>& QueryTM, const T Thickness, TVector<T, 3>& ContactLocation, TVector<T, 3>& ContactNormal, T& ContactPhi) const;
+		bool GJKContactPointImp(const GeomType& QueryGeom, const FRigidTransform3& QueryTM, const FReal Thickness, FVec3& ContactLocation, FVec3& ContactNormal, FReal& ContactPhi) const;
 
 	};
-
-#if PLATFORM_MAC || PLATFORM_LINUX
-	extern template class CHAOS_API THeightField<float>;
-#else
-	extern template class THeightField<float>;
-#endif
 }
diff --git a/Engine/Source/Runtime/Landscape/Classes/LandscapeHeightfieldCollisionComponent.h b/Engine/Source/Runtime/Landscape/Classes/LandscapeHeightfieldCollisionComponent.h
--- a/Engine/Source/Runtime/Landscape/Classes/LandscapeHeightfieldCollisionComponent.h
+++ b/Engine/Source/Runtime/Landscape/Classes/LandscapeHeightfieldCollisionComponent.h
@@ -96,10 +96,10 @@ class ULandscapeHeightfieldCollisionComponent : public UPrimitiveComponent
 
 #if WITH_CHAOS
 		TArray<Chaos::FMaterialHandle> UsedChaosMaterials;
-		TUniquePtr<Chaos::THeightField<float>> Heightfield = nullptr;
-	    TUniquePtr<Chaos::THeightField<float>> HeightfieldSimple = nullptr;
+		TUniquePtr<Chaos::FHeightField> Heightfield = nullptr;
+	    TUniquePtr<Chaos::FHeightField> HeightfieldSimple = nullptr;
 #if WITH_EDITOR
-		TUniquePtr<Chaos::THeightField<float>> EditorHeightfield = nullptr;
+		TUniquePtr<Chaos::FHeightField> EditorHeightfield = nullptr;
 #endif
 #endif
 
diff --git a/Engine/Source/Runtime/Landscape/Private/LandscapeCollision.cpp b/Engine/Source/Runtime/Landscape/Private/LandscapeCollision.cpp
--- a/Engine/Source/Runtime/Landscape/Private/LandscapeCollision.cpp
+++ b/Engine/Source/Runtime/Landscape/Private/LandscapeCollision.cpp
@@ -983,8 +983,8 @@ bool ULandscapeHeightfieldCollisionComponent::CookCollisionData(const FName& For
 		}
 	}
 
-	TUniquePtr<Chaos::THeightField<float>> Heightfield = nullptr;
-	TUniquePtr<Chaos::THeightField<float>> HeightfieldSimple = nullptr;
+	TUniquePtr<Chaos::FHeightField> Heightfield = nullptr;
+	TUniquePtr<Chaos::FHeightField> HeightfieldSimple = nullptr;
 
 	FMemoryWriter Writer(OutData);
 	Chaos::FChaosArchive Ar(Writer);
@@ -993,13 +993,13 @@ bool ULandscapeHeightfieldCollisionComponent::CookCollisionData(const FName& For
 	Ar << bSerializeGenerateSimpleCollision;
 
 	TArrayView<const uint16> ComplexHeightView(Heights, NumSamples);
-	Heightfield = MakeUnique<Chaos::THeightField<float>>(ComplexHeightView, MakeArrayView(MaterialIndices), CollisionSizeVerts, CollisionSizeVerts, Chaos::TVector<float, 3>(1));
+	Heightfield = MakeUnique<Chaos::FHeightField>(ComplexHeightView, MakeArrayView(MaterialIndices), CollisionSizeVerts, CollisionSizeVerts, Chaos::TVector<float, 3>(1));
 	Ar << Heightfield;
 	if(bGenerateSimpleCollision)
 	{
 		// #BGTODO Materials for simple geometry, currently just passing in the default
 		TArrayView<const uint16> SimpleHeightView(Heights + NumSamples, NumSimpleSamples);
-		HeightfieldSimple = MakeUnique<Chaos::THeightField<float>>(SimpleHeightView, MakeArrayView(MaterialIndices.GetData(), 1), CollisionSizeVerts, CollisionSizeVerts, Chaos::TVector<float, 3>(1));
+		HeightfieldSimple = MakeUnique<Chaos::FHeightField>(SimpleHeightView, MakeArrayView(MaterialIndices.GetData(), 1), CollisionSizeVerts, CollisionSizeVerts, Chaos::TVector<float, 3>(1));
 		Ar << HeightfieldSimple;
 	}
 
diff --git a/Engine/Source/Runtime/NavigationSystem/Private/NavMesh/RecastNavMeshGenerator.cpp b/Engine/Source/Runtime/NavigationSystem/Private/NavMesh/RecastNavMeshGenerator.cpp
--- a/Engine/Source/Runtime/NavigationSystem/Private/NavMesh/RecastNavMeshGenerator.cpp
+++ b/Engine/Source/Runtime/NavigationSystem/Private/NavMesh/RecastNavMeshGenerator.cpp
@@ -232,7 +232,7 @@ struct FRecastGeometryExport : public FNavigableGeometryExport
 #if WITH_CHAOS
 	virtual void ExportChaosTriMesh(const Chaos::FTriangleMeshImplicitObject* const TriMesh, const FTransform& LocalToWorld) override;
 	virtual void ExportChaosConvexMesh(const FKConvexElem* const Convex, const FTransform& LocalToWorld) override;
-	virtual void ExportChaosHeightField(const Chaos::THeightField<float>* const Heightfield, const FTransform& LocalToWorld) override;
+	virtual void ExportChaosHeightField(const Chaos::FHeightField* const Heightfield, const FTransform& LocalToWorld) override;
 	virtual void ExportChaosHeightFieldSlice(const FNavHeightfieldSamples& PrefetchedHeightfieldSamples, const int32 NumRows, const int32 NumCols, const FTransform& LocalToWorld, const FBox& SliceBox) override;
 #endif
 	virtual void ExportCustomMesh(const FVector* InVertices, int32 NumVerts, const int32* InIndices, int32 NumIndices, const FTransform& LocalToWorld) override;
@@ -796,7 +796,7 @@ void ExportChaosConvexMesh(const FKConvexElem* const Convex, const FTransform& L
 	#endif // SHOW_NAV_EXPORT_PREVIEW
 }
 
-void ExportChaosHeightField(const Chaos::THeightField<float>* const HeightField, const FTransform& LocalToWorld
+void ExportChaosHeightField(const Chaos::FHeightField* const HeightField, const FTransform& LocalToWorld
 	, TNavStatArray<float>& VertexBuffer, TNavStatArray<int32>& IndexBuffer
 	, FBox& UnrealBounds)
 {
@@ -1448,7 +1448,7 @@ void FRecastGeometryExport::ExportChaosConvexMesh(const FKConvexElem* const Conv
 	RecastGeometryExport::ExportChaosConvexMesh(Convex, LocalToWorld, VertexBuffer, IndexBuffer, Data->Bounds);
 }
 
-void FRecastGeometryExport::ExportChaosHeightField(const Chaos::THeightField<float>* const Heightfield, const FTransform& LocalToWorld)
+void FRecastGeometryExport::ExportChaosHeightField(const Chaos::FHeightField* const Heightfield, const FTransform& LocalToWorld)
 {
 	RecastGeometryExport::ExportChaosHeightField(Heightfield, LocalToWorld, VertexBuffer, IndexBuffer, Data->Bounds);
 }
diff --git a/Engine/Source/Runtime/PhysicsCore/Public/PhysXToChaosUtil.h b/Engine/Source/Runtime/PhysicsCore/Public/PhysXToChaosUtil.h
--- a/Engine/Source/Runtime/PhysicsCore/Public/PhysXToChaosUtil.h
+++ b/Engine/Source/Runtime/PhysicsCore/Public/PhysXToChaosUtil.h
@@ -126,7 +126,7 @@ inline TUniquePtr<Chaos::TImplicitObjectTransformed<float, 3>> PxShapeToChaosGeo
 
 		ShapeTM = U2PTransform(FTransform(LandscapeComponentMatrix));
 
-		InnerObj = MakeUnique<THeightField<float>>(MoveTemp(Height), MoveTemp(MaterialIndices), NumRows, NumCols, TVector<float, 3>(HeightFieldGeom.columnScale, HeightFieldGeom.rowScale, HeightFieldGeom.heightScale));
+		InnerObj = MakeUnique<FHeightField>(MoveTemp(Height), MoveTemp(MaterialIndices), NumRows, NumCols, TVector<float, 3>(HeightFieldGeom.columnScale, HeightFieldGeom.rowScale, HeightFieldGeom.heightScale));
 		break;
 	}
 	case PxGeometryType::eTRIANGLEMESH:
