New internal cluster particles incorrectly setting P,Q

diff --git a/Engine/Source/Runtime/Experimental/Chaos/Private/Chaos/PBDRigidClustering.cpp b/Engine/Source/Runtime/Experimental/Chaos/Private/Chaos/PBDRigidClustering.cpp
--- a/Engine/Source/Runtime/Experimental/Chaos/Private/Chaos/PBDRigidClustering.cpp
+++ b/Engine/Source/Runtime/Experimental/Chaos/Private/Chaos/PBDRigidClustering.cpp
@@ -851,8 +851,8 @@ namespace Chaos
 							NewCluster->SetW(ClusteredParticle->W());
 							NewCluster->SetPreV(ClusteredParticle->PreV());
 							NewCluster->SetPreW(ClusteredParticle->PreW());
-							NewCluster->SetP(ClusteredParticle->X());
-							NewCluster->SetQ(ClusteredParticle->R());
+							NewCluster->SetP(NewCluster->X());
+							NewCluster->SetQ(NewCluster->R());
 
 							ActivatedChildren.Add(NewCluster);
 						}
