Chaos Contact debug draw: Mass Reference frame shown for bodies during DebugDrawCollisions

diff --git a/Engine/Source/Runtime/Experimental/Chaos/Private/Chaos/ChaosDebugDraw.cpp b/Engine/Source/Runtime/Experimental/Chaos/Private/Chaos/ChaosDebugDraw.cpp
--- a/Engine/Source/Runtime/Experimental/Chaos/Private/Chaos/ChaosDebugDraw.cpp
+++ b/Engine/Source/Runtime/Experimental/Chaos/Private/Chaos/ChaosDebugDraw.cpp
@@ -357,6 +357,12 @@ namespace Chaos
 				FDebugDrawQueue::GetInstance().DrawDebugLine(Location, P0, C3, false, KINDA_SMALL_NUMBER, DrawPriority, LineThickness * 0.5f);
 				FDebugDrawQueue::GetInstance().DrawDebugLine(Location, P1, C3, false, KINDA_SMALL_NUMBER, DrawPriority, LineThickness * 0.5f);
 			}
+
+			// Draw the particle (mass frame) coordinates
+			{
+				DrawParticleTransformImpl(FRigidTransform3::Identity, Contact.Particle[0], 0, 1.0f);
+				DrawParticleTransformImpl(FRigidTransform3::Identity, Contact.Particle[1], 0, 1.0f);				
+			}
 		}
 		
 		void DrawCollisionImpl(const FRigidTransform3& SpaceTransform, const FPBDCollisionConstraintHandle* ConstraintHandle, float ColorScale)
