Chaos: remove an ensure for now

diff --git a/Engine/Source/Runtime/Experimental/Chaos/Private/Chaos/PBDConstraintGraph.cpp b/Engine/Source/Runtime/Experimental/Chaos/Private/Chaos/PBDConstraintGraph.cpp
--- a/Engine/Source/Runtime/Experimental/Chaos/Private/Chaos/PBDConstraintGraph.cpp
+++ b/Engine/Source/Runtime/Experimental/Chaos/Private/Chaos/PBDConstraintGraph.cpp
@@ -291,7 +291,9 @@ void FPBDConstraintGraph::UpdateIslands(const TParticleView<TPBDRigidParticles<F
 	for (auto& PBDRigid : PBDRigids)
 	{
 		PBDRigid.Island() = INDEX_NONE;
-		if (!ensure(ParticleToNodeIndex.Contains(PBDRigid.Handle())))
+		// When enabling particle from a break, if the object state is static, then when Enabled the paticle doesn't get added
+		// to the constraint graph however the particle appears in GetNonDisabledDynamicView() so this ensure fires
+		if (!/*ensure*/(ParticleToNodeIndex.Contains(PBDRigid.Handle())))
 		{
 			ParticleAdd(PBDRigid.Handle());
 		}
