Don't attempt to execute global commands when running in single-threaded mode for physics. Removed checks in case we ever do get into the command executor in this mode as it isn't fatal it just wastes a bit of time enqueuing tasks that will not perform any work.

diff --git a/Engine/Source/Runtime/Experimental/ChaosSolvers/Private/ChaosSolversModule.cpp b/Engine/Source/Runtime/Experimental/ChaosSolvers/Private/ChaosSolversModule.cpp
--- a/Engine/Source/Runtime/Experimental/ChaosSolvers/Private/ChaosSolversModule.cpp
+++ b/Engine/Source/Runtime/Experimental/ChaosSolvers/Private/ChaosSolversModule.cpp
@@ -1114,9 +1114,9 @@ void FChaosSolversModule::OnDestroyMaterialMask(Chaos::FMaterialMaskHandle InHan
 
 void FChaosSolversModule::DispatchGlobalCommands()
 {
-	if(Dispatcher->GetMode() == Chaos::EThreadingMode::SingleThread)
+	if(Dispatcher && Dispatcher->GetMode() == Chaos::EThreadingMode::SingleThread)
 	{
-		// Single thread commands run as they are enqueued so no need to proceed here
+		// Single threaded dispatchers fire commands immediately so we don't need to process any commands here
 		return;
 	}
 
diff --git a/Engine/Source/Runtime/Experimental/ChaosSolvers/Private/Framework/PhysicsTickTask.cpp b/Engine/Source/Runtime/Experimental/ChaosSolvers/Private/Framework/PhysicsTickTask.cpp
--- a/Engine/Source/Runtime/Experimental/ChaosSolvers/Private/Framework/PhysicsTickTask.cpp
+++ b/Engine/Source/Runtime/Experimental/ChaosSolvers/Private/Framework/PhysicsTickTask.cpp
@@ -150,7 +150,6 @@ void FPhysicsCommandsTask::DoTask(ENamedThreads::Type CurrentThread, const FGrap
 	check(Dispatcher);
 
 	Dispatcher = Module->GetDispatcher();
-	check(Dispatcher->GetMode() == EChaosThreadingMode::TaskGraph);
 
 	//static FCriticalSection DispatcherExecutionSection;
 	//DispatcherExecutionSection.Lock();
