Hooked up geometry collections to the component body instance filter data so configuring the filter in the details panel now functions correctly. Also defaulted the body instance for geometry collections into the ECC_Destructible query channel.

diff --git a/Engine/Source/Runtime/Experimental/GeometryCollectionEngine/Private/GeometryCollection/GeometryCollectionComponent.cpp b/Engine/Source/Runtime/Experimental/GeometryCollectionEngine/Private/GeometryCollection/GeometryCollectionComponent.cpp
--- a/Engine/Source/Runtime/Experimental/GeometryCollectionEngine/Private/GeometryCollection/GeometryCollectionComponent.cpp
+++ b/Engine/Source/Runtime/Experimental/GeometryCollectionEngine/Private/GeometryCollection/GeometryCollectionComponent.cpp
@@ -156,6 +156,9 @@ UGeometryCollectionComponent::UGeometryCollectionComponent(const FObjectInitiali
 	TransformsAreEqualIndex = 0;
 
 	SetGenerateOverlapEvents(false);
+
+	// By default use the destructible object channel unless the user specifies otherwise
+	BodyInstance.SetObjectType(ECC_Destructible);
 }
 
 Chaos::FPhysicsSolver* GetSolver(const UGeometryCollectionComponent& GeometryCollectionComponent)
@@ -1353,9 +1356,16 @@ void UGeometryCollectionComponent::OnCreatePhysicsState()
 			// end temporary 
 
 			// Set up initial filter data for our particles
-			FCollisionResponseContainer Response;
-			AActor* ComponentOwner = GetOwner();
-			CreateShapeFilterData(ECC_WorldDynamic, FMaskFilter(0), ComponentOwner ? ComponentOwner->GetUniqueID() : 0, Response, GetUniqueID(), 0, InitialQueryFilter, InitialSimFilter, false, false, false, false);
+			// #BGTODO We need a dummy body setup for now to allow the body instance to generate filter information. Change body instance to operate independently.
+			DummyBodySetup = NewObject<UBodySetup>(this, UBodySetup::StaticClass());
+			BodyInstance.BodySetup = DummyBodySetup;
+
+			FBodyCollisionFilterData FilterData;
+			FMaskFilter FilterMask = BodyInstance.GetMaskFilter();
+			BodyInstance.BuildBodyFilterData(FilterData);
+
+			InitialSimFilter = FilterData.SimFilter;
+			InitialQueryFilter = FilterData.QuerySimpleFilter;
 
 			// Enable for complex and simple (no dual representation currently like other meshes)
 			InitialQueryFilter.Word3 |= (EPDF_SimpleCollision | EPDF_ComplexCollision);
