Add temp guards for chaos header includes

diff --git a/Engine/Plugins/Experimental/ChaosCloth/Source/Chaos/Public/ChaosCloth/ChaosClothingSimulation.h b/Engine/Plugins/Experimental/ChaosCloth/Source/Chaos/Public/ChaosCloth/ChaosClothingSimulation.h
--- a/Engine/Plugins/Experimental/ChaosCloth/Source/Chaos/Public/ChaosCloth/ChaosClothingSimulation.h
+++ b/Engine/Plugins/Experimental/ChaosCloth/Source/Chaos/Public/ChaosCloth/ChaosClothingSimulation.h
@@ -3,11 +3,13 @@
 
 #include "ClothingSimulation.h"
 #include "ClothingAsset.h"
+#define TEMP_HEADER_CHAOS_LEVEL_1
 #include "Chaos/ArrayCollectionArray.h"
 #include "Chaos/PBDEvolution.h"
 #include "Chaos/Transform.h"
 #include "Chaos/TriangleMesh.h"
 #include "Chaos/ChaosDebugDrawDeclares.h"
+#undef TEMP_HEADER_CHAOS_LEVEL_1
 #include "ChaosCloth/ChaosClothConfig.h"
 #include "Components/SkeletalMeshComponent.h"
 
diff --git a/Engine/Source/Runtime/Engine/Classes/PhysicsEngine/ConvexElem.h b/Engine/Source/Runtime/Engine/Classes/PhysicsEngine/ConvexElem.h
--- a/Engine/Source/Runtime/Engine/Classes/PhysicsEngine/ConvexElem.h
+++ b/Engine/Source/Runtime/Engine/Classes/PhysicsEngine/ConvexElem.h
@@ -5,7 +5,10 @@
 #include "CoreMinimal.h"
 #include "UObject/ObjectMacros.h"
 #include "PhysicsEngine/ShapeElem.h"
+
+#define TEMP_HEADER_CHAOS_LEVEL_1 1
 #include "Chaos/Convex.h"
+#undef TEMP_HEADER_CHAOS_LEVEL_1
 
 #if WITH_CHAOS
 #include "Chaos/Serializable.h"
@@ -13,6 +16,7 @@
 
 #include "ConvexElem.generated.h"
 
+
 struct FDynamicMeshVertex;
 struct FKBoxElem;
 
diff --git a/Engine/Source/Runtime/Engine/Private/Collision/CollisionDebugDrawing.h b/Engine/Source/Runtime/Engine/Private/Collision/CollisionDebugDrawing.h
--- a/Engine/Source/Runtime/Engine/Private/Collision/CollisionDebugDrawing.h
+++ b/Engine/Source/Runtime/Engine/Private/Collision/CollisionDebugDrawing.h
@@ -9,7 +9,9 @@
 #include "EngineDefines.h"
 #include "PhysicsPublic.h"
 #include "PhysXIncludes.h"
+#define TEMP_HEADER_CHAOS_LEVEL_1
 #include "Chaos/ImplicitObject.h"
+#undef TEMP_HEADER_CHAOS_LEVEL_1
 
 #if PHYSICS_INTERFACE_PHYSX
 /** Draw PhysX geom with overlaps */
diff --git a/Engine/Source/Runtime/Engine/Public/Physics/ImmediatePhysics/ImmediatePhysicsChaos/ImmediatePhysicsSimulation_Chaos.h b/Engine/Source/Runtime/Engine/Public/Physics/ImmediatePhysics/ImmediatePhysicsChaos/ImmediatePhysicsSimulation_Chaos.h
--- a/Engine/Source/Runtime/Engine/Public/Physics/ImmediatePhysics/ImmediatePhysicsChaos/ImmediatePhysicsSimulation_Chaos.h
+++ b/Engine/Source/Runtime/Engine/Public/Physics/ImmediatePhysics/ImmediatePhysicsChaos/ImmediatePhysicsSimulation_Chaos.h
@@ -4,6 +4,7 @@
 
 #include "Physics/ImmediatePhysics/ImmediatePhysicsChaos/ImmediatePhysicsCore_Chaos.h"
 
+#define TEMP_HEADER_CHAOS_LEVEL_1
 #include "Chaos/ChaosDebugDrawDeclares.h"
 #include "Chaos/Collision/NarrowPhase.h"
 #include "Chaos/Collision/ParticlePairBroadPhase.h"
@@ -12,6 +13,7 @@
 #include "Chaos/PBDCollisionConstraints.h"
 #include "Chaos/PBDConstraintRule.h"
 #include "Chaos/PBDJointConstraints.h"
+#undef TEMP_HEADER_CHAOS_LEVEL_1
 
 #include "Engine/EngineTypes.h"
 #include "Templates/UniquePtr.h"
diff --git a/Engine/Source/Runtime/Engine/Public/Physics/PhysicsInterfaceCore.h b/Engine/Source/Runtime/Engine/Public/Physics/PhysicsInterfaceCore.h
--- a/Engine/Source/Runtime/Engine/Public/Physics/PhysicsInterfaceCore.h
+++ b/Engine/Source/Runtime/Engine/Public/Physics/PhysicsInterfaceCore.h
@@ -15,8 +15,10 @@
 
 #elif WITH_CHAOS
 
+#define TEMP_HEADER_CHAOS_LEVEL_1 1
 #include "Physics/Experimental/PhysInterface_Chaos.h"
 #include "Physics/Experimental/PhysScene_Chaos.h"
+#undef TEMP_HEADER_CHAOS_LEVEL_1
 
 #else
 
diff --git a/Engine/Source/Runtime/Experimental/GeometryCollectionCore/Public/GeometryCollection/ManagedArrayTypes.h b/Engine/Source/Runtime/Experimental/GeometryCollectionCore/Public/GeometryCollection/ManagedArrayTypes.h
--- a/Engine/Source/Runtime/Experimental/GeometryCollectionCore/Public/GeometryCollection/ManagedArrayTypes.h
+++ b/Engine/Source/Runtime/Experimental/GeometryCollectionCore/Public/GeometryCollection/ManagedArrayTypes.h
@@ -1,6 +1,5 @@
 // Copyright Epic Games, Inc. All Rights Reserved.
 #pragma once
-#include "Chaos/IncludeLvl1.inl"
 #include "Containers/UnrealString.h"
 #include "GeometryCollection/GeometryCollectionBoneNode.h"
 #include "GeometryCollection/GeometryCollectionSection.h"
@@ -11,10 +10,13 @@
 #include "Math/Vector.h"
 #include "Math/Vector2D.h"
 #include "Math/Box.h"
+
+#define TEMP_HEADER_CHAOS_LEVEL_1
 #include "Chaos/ImplicitObject.h"
 #include "Chaos/BVHParticles.h"
 #include "Chaos/ParticleHandle.h"
 #include "Chaos/ParticleHandleFwd.h"
+#undef TEMP_HEADER_CHAOS_LEVEL_1
 
 inline FArchive& operator<<(FArchive& Ar, TArray<FVector>*& ValueIn)
 {
diff --git a/Engine/Source/Runtime/Experimental/GeometryCollectionEngine/Public/ChaosBlueprint.h b/Engine/Source/Runtime/Experimental/GeometryCollectionEngine/Public/ChaosBlueprint.h
--- a/Engine/Source/Runtime/Experimental/GeometryCollectionEngine/Public/ChaosBlueprint.h
+++ b/Engine/Source/Runtime/Experimental/GeometryCollectionEngine/Public/ChaosBlueprint.h
@@ -1,12 +1,13 @@
 // Copyright Epic Games, Inc. All Rights Reserved.
 #pragma once
 
-#include "Chaos/IncludeLvl1.inl"
 #include "CoreMinimal.h"
 #include "UObject/ObjectMacros.h"
 #include "Chaos/Declares.h"
+#define TEMP_HEADER_CHAOS_LEVEL_1
 #include "Chaos/ChaosSolver.h"
 #include "Chaos/ChaosSolverActor.h"
+#undef TEMP_HEADER_CHAOS_LEVEL_1
 #include "EventManager.h"
 #include "EventsData.h"
 #include "HAL/ThreadSafeBool.h"
diff --git a/Engine/Source/Runtime/Experimental/GeometryCollectionEngine/Public/GeometryCollection/GeometryCollectionComponent.h b/Engine/Source/Runtime/Experimental/GeometryCollectionEngine/Public/GeometryCollection/GeometryCollectionComponent.h
--- a/Engine/Source/Runtime/Experimental/GeometryCollectionEngine/Public/GeometryCollection/GeometryCollectionComponent.h
+++ b/Engine/Source/Runtime/Experimental/GeometryCollectionEngine/Public/GeometryCollection/GeometryCollectionComponent.h
@@ -2,8 +2,8 @@
 
 #pragma once
 
-#include "Chaos/IncludeLvl1.inl"
 #include "Components/MeshComponent.h"
+#define TEMP_HEADER_CHAOS_LEVEL_1
 #include "Chaos/ChaosSolverActor.h"
 #include "Chaos/PBDRigidParticles.h"
 #include "Chaos/Defines.h"
@@ -25,6 +25,7 @@
 #include "Chaos/ChaosGameplayEventDispatcher.h"
 #include "Chaos/ChaosNotifyHandlerInterface.h"
 #include "Chaos/ChaosSolverComponentTypes.h"
+#undef TEMP_HEADER_CHAOS_LEVEL_1
 
 #include "GeometryCollectionComponent.generated.h"
 
diff --git a/Engine/Source/Runtime/Experimental/GeometryCollectionEngine/Public/GeometryCollection/GeometryCollectionParticlesData.h b/Engine/Source/Runtime/Experimental/GeometryCollectionEngine/Public/GeometryCollection/GeometryCollectionParticlesData.h
--- a/Engine/Source/Runtime/Experimental/GeometryCollectionEngine/Public/GeometryCollection/GeometryCollectionParticlesData.h
+++ b/Engine/Source/Runtime/Experimental/GeometryCollectionEngine/Public/GeometryCollection/GeometryCollectionParticlesData.h
@@ -2,17 +2,18 @@
 
 #pragma once
 
-#include "Chaos/IncludeLvl1.inl"
 #include "Misc/Build.h"
 
 #if !(UE_BUILD_SHIPPING || UE_BUILD_TEST)
 #include "Containers/StaticBitArray.h"
+#define TEMP_HEADER_CHAOS_LEVEL_1
 #include "Chaos/ArrayCollectionArray.h"
 #include "Chaos/ImplicitObject.h"
 #include "Chaos/Matrix.h"
 #include "Chaos/PBDRigidClusteredParticles.h"
 #include "Chaos/UniformGrid.h"
 #include "Chaos/Framework/BufferedData.h"
+#undef TEMP_HEADER_CHAOS_LEVEL_1
 #include "Chaos/Declares.h"
 
 class FChaosSolversModule;
diff --git a/Engine/Source/Runtime/Experimental/GeometryCollectionEngine/Public/GeometryCollection/SkeletalMeshSimulationComponent.h b/Engine/Source/Runtime/Experimental/GeometryCollectionEngine/Public/GeometryCollection/SkeletalMeshSimulationComponent.h
--- a/Engine/Source/Runtime/Experimental/GeometryCollectionEngine/Public/GeometryCollection/SkeletalMeshSimulationComponent.h
+++ b/Engine/Source/Runtime/Experimental/GeometryCollectionEngine/Public/GeometryCollection/SkeletalMeshSimulationComponent.h
@@ -1,14 +1,15 @@
 // Copyright Epic Games, Inc. All Rights Reserved.
 #pragma once
 
-#include "Chaos/IncludeLvl1.inl"
 #include "Components/MeshComponent.h"
+#define TEMP_HEADER_CHAOS_LEVEL_1
 #include "Chaos/ChaosSolverActor.h"
+#include "Chaos/ChaosNotifyHandlerInterface.h"
+#undef TEMP_HEADER_CHAOS_LEVEL_1
 #include "GameFramework/Actor.h"
 #include "Physics/Experimental/PhysScene_Chaos.h"
 #include "PhysicalMaterials/Experimental/ChaosPhysicalMaterial.h"
 #include "GeometryCollection/GeometryCollectionSimulationTypes.h"
-#include "Chaos/ChaosNotifyHandlerInterface.h"
 
 #include "SkeletalMeshSimulationComponent.generated.h"
 
diff --git a/Engine/Source/Runtime/Experimental/GeometryCollectionEngine/Public/GeometryCollection/StaticMeshSimulationComponent.h b/Engine/Source/Runtime/Experimental/GeometryCollectionEngine/Public/GeometryCollection/StaticMeshSimulationComponent.h
--- a/Engine/Source/Runtime/Experimental/GeometryCollectionEngine/Public/GeometryCollection/StaticMeshSimulationComponent.h
+++ b/Engine/Source/Runtime/Experimental/GeometryCollectionEngine/Public/GeometryCollection/StaticMeshSimulationComponent.h
@@ -2,9 +2,10 @@
 
 #pragma once
 
-#include "Chaos/IncludeLvl1.inl"
 #include "Components/MeshComponent.h"
+#define TEMP_HEADER_CHAOS_LEVEL_1
 #include "Chaos/ChaosSolverActor.h"
+#undef TEMP_HEADER_CHAOS_LEVEL_1
 #include "GameFramework/Actor.h"
 #include "Physics/Experimental/PhysScene_Chaos.h"
 #include "GeometryCollection/GeometryCollectionSimulationTypes.h"
diff --git a/Engine/Source/Runtime/Landscape/Classes/LandscapeHeightfieldCollisionComponent.h b/Engine/Source/Runtime/Landscape/Classes/LandscapeHeightfieldCollisionComponent.h
--- a/Engine/Source/Runtime/Landscape/Classes/LandscapeHeightfieldCollisionComponent.h
+++ b/Engine/Source/Runtime/Landscape/Classes/LandscapeHeightfieldCollisionComponent.h
@@ -11,7 +11,11 @@
 #include "AI/Navigation/NavigationTypes.h"
 #include "Components/PrimitiveComponent.h"
 #include "Serialization/BulkData.h"
+
+#define TEMP_HEADER_CHAOS_LEVEL_1
 #include "Chaos/HeightField.h"
+#undef TEMP_HEADER_CHAOS_LEVEL_1
+
 #include "LandscapeHeightfieldCollisionComponent.generated.h"
 
 class ALandscapeProxy;
