Fix CIS: Update particlehandle for shape enabling code

diff --git a/Engine/Source/Runtime/Experimental/Chaos/Public/Chaos/ParticleHandle.h b/Engine/Source/Runtime/Experimental/Chaos/Public/Chaos/ParticleHandle.h
--- a/Engine/Source/Runtime/Experimental/Chaos/Public/Chaos/ParticleHandle.h
+++ b/Engine/Source/Runtime/Experimental/Chaos/Public/Chaos/ParticleHandle.h
@@ -1527,12 +1527,21 @@ public:
 		}
 	}
 
-	void SetShapeCollisionDisable(int32 InShapeIndex,bool bInDisable)
+	void SetShapeSimCollisionEnabled(int32 InShapeIndex, bool bInEnabled)
 	{
-		const bool bCurrent = MShapesArray[InShapeIndex]->GetDisable();
-		if(bCurrent != bInDisable)
+		const bool bCurrent = MShapesArray[InShapeIndex]->GetSimEnabled();
+		if(bCurrent != bInEnabled)
 		{
-			MShapesArray[InShapeIndex]->SetDisable(bInDisable);
+			MShapesArray[InShapeIndex]->SetSimEnabled(bInEnabled);
+		}
+	}
+
+	void SetShapeQueryCollisionEnabled(int32 InShapeIndex, bool bInEnabled)
+	{
+		const bool bCurrent = MShapesArray[InShapeIndex]->GetQueryEnabled();
+		if(bCurrent != bInEnabled)
+		{
+			MShapesArray[InShapeIndex]->SetQueryEnabled(bInEnabled);
 		}
 	}
 
@@ -1555,7 +1564,7 @@ public:
 	}
 
 #if CHAOS_CHECKED
-	const FName& DebugName() const { return MNonFrequentData.Read().DebugName(); }
+	const FName DebugName() const { return MNonFrequentData.Read().DebugName(); }
 	void SetDebugName(const FName& InDebugName)
 	{
 		MNonFrequentData.Modify(true,MDirtyFlags,Proxy,[&InDebugName](auto& Data){ Data.SetDebugName(InDebugName);});
@@ -1782,7 +1791,7 @@ public:
 		ShapeMaterials.Empty(Shapes.Num());
 		for (const TUniquePtr<FPerShapeData>& ShapePtr : Shapes)
 		{
-			ShapeCollisionDisableFlags.Add(ShapePtr->GetDisable());
+			ShapeCollisionDisableFlags.Add(!ShapePtr->GetSimEnabled());
 			CollisionTraceType.Add(ShapePtr->GetCollisionTraceType());
 			ShapeSimData.Add(ShapePtr->GetSimData());
 			ShapeQueryData.Add(ShapePtr->GetQueryData());
