Update transforms of particles who have a parent  that is an internal cluster

diff --git a/Engine/Source/Runtime/Experimental/ChaosSolvers/Private/PhysicsProxy/GeometryCollectionPhysicsProxy.cpp b/Engine/Source/Runtime/Experimental/ChaosSolvers/Private/PhysicsProxy/GeometryCollectionPhysicsProxy.cpp
--- a/Engine/Source/Runtime/Experimental/ChaosSolvers/Private/PhysicsProxy/GeometryCollectionPhysicsProxy.cpp
+++ b/Engine/Source/Runtime/Experimental/ChaosSolvers/Private/PhysicsProxy/GeometryCollectionPhysicsProxy.cpp
@@ -1612,6 +1612,29 @@ void FGeometryCollectionPhysicsProxy::BufferPhysicsResults()
 		GeometryCollectionAlgo::GlobalMatrices(
 			TargetResults.Transforms, TargetResults.Parent, TargetResults.GlobalTransforms);
 
+		// update transforms of things parented to internal clusters who would be missed in the above call to GlobalMatrices
+		// #todo: this can be much faster by caching the parent transforms somehow
+		// #todo: it only supports a single level of internal cluster parenting.
+		for (int32 TransformIndex = 0; TransformIndex < TransfromGroupSize; ++TransformIndex)
+		{
+			Chaos::TPBDRigidClusteredParticleHandle<float, 3>* Handle = SolverParticleHandles[TransformIndex];
+
+			// If the particle is disabled and has a parent that is an internal cluster, update transform
+			if (Handle && Handle->Disabled() && SolverClusterID[TransformIndex] && SolverClusterID[TransformIndex]->CastToClustered()->InternalCluster())
+			{
+				FClusterHandle* ParentClusterHandle = SolverClusterID[TransformIndex]->CastToClustered();
+
+				FTransform ParticleToWorld = FTransform(ParentClusterHandle->R(), ParentClusterHandle->X());
+				
+				// #todo: what should the mass to local for the internal cluster be?
+				FTransform MassToLocal = FTransform::Identity;
+				
+				// compute global transform for particle
+				TargetResults.GlobalTransforms[TransformIndex] = Handle->ChildToParent().ToMatrixWithScale();
+				TargetResults.GlobalTransforms[TransformIndex] *= (MassToLocal.GetRelativeTransformReverse(ParticleToWorld).GetRelativeTransform(ActorToWorld)).ToMatrixWithScale();
+			}
+		}
+
 		// Compute world bounds.  This is a loose bounds based on the circumscribed box 
 		// of a bounding sphere for the geometry.		
 		SCOPE_CYCLE_COUNTER(STAT_CalcGlobalGCBounds);
