Back out changelist 12570120. CIS errors

diff --git a/Engine/Source/Runtime/Experimental/Chaos/Public/Chaos/ParticleHandle.h b/Engine/Source/Runtime/Experimental/Chaos/Public/Chaos/ParticleHandle.h
--- a/Engine/Source/Runtime/Experimental/Chaos/Public/Chaos/ParticleHandle.h
+++ b/Engine/Source/Runtime/Experimental/Chaos/Public/Chaos/ParticleHandle.h
@@ -1527,21 +1527,12 @@ public:
 		}
 	}
 
-	void SetShapeSimCollisionEnabled(int32 InShapeIndex, bool bInEnabled)
+	void SetShapeCollisionDisable(int32 InShapeIndex,bool bInDisable)
 	{
-		const bool bCurrent = MShapesArray[InShapeIndex]->GetSimEnabled();
-		if(bCurrent != bInEnabled)
+		const bool bCurrent = MShapesArray[InShapeIndex]->GetDisable();
+		if(bCurrent != bInDisable)
 		{
-			MShapesArray[InShapeIndex]->SetSimEnabled(bInEnabled);
-		}
-	}
-
-	void SetShapeQueryCollisionEnabled(int32 InShapeIndex, bool bInEnabled)
-	{
-		const bool bCurrent = MShapesArray[InShapeIndex]->GetQueryEnabled();
-		if(bCurrent != bInEnabled)
-		{
-			MShapesArray[InShapeIndex]->SetQueryEnabled(bInEnabled);
+			MShapesArray[InShapeIndex]->SetDisable(bInDisable);
 		}
 	}
 
@@ -1564,7 +1555,7 @@ public:
 	}
 
 #if CHAOS_CHECKED
-	const FName DebugName() const { return MNonFrequentData.Read().DebugName(); }
+	const FName& DebugName() const { return MNonFrequentData.Read().DebugName(); }
 	void SetDebugName(const FName& InDebugName)
 	{
 		MNonFrequentData.Modify(true,MDirtyFlags,Proxy,[&InDebugName](auto& Data){ Data.SetDebugName(InDebugName);});
@@ -1791,7 +1782,7 @@ public:
 		ShapeMaterials.Empty(Shapes.Num());
 		for (const TUniquePtr<FPerShapeData>& ShapePtr : Shapes)
 		{
-			ShapeCollisionDisableFlags.Add(!ShapePtr->GetSimEnabled());
+			ShapeCollisionDisableFlags.Add(ShapePtr->GetDisable());
 			CollisionTraceType.Add(ShapePtr->GetCollisionTraceType());
 			ShapeSimData.Add(ShapePtr->GetSimData());
 			ShapeQueryData.Add(ShapePtr->GetQueryData());
