Modify chaos include level guard

diff --git a/Engine/Source/Runtime/Experimental/Chaos/Public/Chaos/IncludeLvl1.inl b/Engine/Source/Runtime/Experimental/Chaos/Public/Chaos/IncludeLvl1.inl
--- a/Engine/Source/Runtime/Experimental/Chaos/Public/Chaos/IncludeLvl1.inl
+++ b/Engine/Source/Runtime/Experimental/Chaos/Public/Chaos/IncludeLvl1.inl
@@ -1,13 +1,27 @@
 // Copyright Epic Games, Inc. All Rights Reserved.
 
-#if defined(CHAOS_INCLUDE_LEVEL_0)
-#error "Cannot include level 1 file from level 0"
-#endif
-
 //Level 1 files should NOT be included by unreal .h files (Chaos itself is ok)
 //Unreal headers including Chaos headers leads to very slow compile times.
 //Until we fix all instances of this please use TEMP_HEADER_CHAOS_LEVEL_1 (different name for searchability)
 //Or consider making the included header a Level 0 - this is meant for files that rarely change like forward declares
-#if !defined(CHAOS_INCLUDE_LEVEL_1) && !defined(TEMP_HEADER_CHAOS_LEVEL_1)
-//#error "Cannot include level 1 file. If outside chaos this should be in a cpp file and should be undefined after includes are done. Use Non-Unity build to see cpp file"
-#endif
\ No newline at end of file
+
+#ifdef CHAOS_LEVEL_CHECK
+#undef CHAOS_LEVEL_CHECK
+#undef CHAOS_LEVEL_CHECK_IMP
+#undef CHAOS_INCLUDE_LEVEL_ACTUAL 
+#endif
+
+#if defined(CHAOS_INCLUDE_LEVEL_1)
+#define CHAOS_INCLUDE_LEVEL_ACTUAL CHAOS_INCLUDE_LEVEL_1
+#elif defined(TEMP_HEADER_CHAOS_LEVEL_1)
+#define CHAOS_INCLUDE_LEVEL_ACTUAL CHAOS_INCLUDE_LEVEL_1
+#endif
+
+#ifndef CHAOS_INCLUDE_LEVEL_ACTUAL
+//#define CHAOS_LEVEL_CHECK_IMP static_assert(false, "Attempting to upgrade CHAOS_INCLUDE_LEVEL");
+#define CHAOS_LEVEL_CHECK_IMP
+#else
+#define CHAOS_LEVEL_CHECK_IMP
+#endif
+
+#define CHAOS_LEVEL_CHECK CHAOS_LEVEL_CHECK_IMP
\ No newline at end of file
