Remove lvl1 include guards from GeometryCollection

diff --git a/Engine/Source/Runtime/Experimental/GeometryCollectionCore/Public/GeometryCollection/GeometryCollection.h b/Engine/Source/Runtime/Experimental/GeometryCollectionCore/Public/GeometryCollection/GeometryCollection.h
--- a/Engine/Source/Runtime/Experimental/GeometryCollectionCore/Public/GeometryCollection/GeometryCollection.h
+++ b/Engine/Source/Runtime/Experimental/GeometryCollectionCore/Public/GeometryCollection/GeometryCollection.h
@@ -2,7 +2,6 @@
 
 #pragma once
 
-#include "Chaos/IncludeLvl1.inl"
 #include "CoreMinimal.h"
 #include "TransformCollection.h"
 #include "Misc/Crc.h"
diff --git a/Engine/Source/Runtime/Experimental/GeometryCollectionCore/Public/GeometryCollection/GeometryCollectionAlgo.h b/Engine/Source/Runtime/Experimental/GeometryCollectionCore/Public/GeometryCollection/GeometryCollectionAlgo.h
--- a/Engine/Source/Runtime/Experimental/GeometryCollectionCore/Public/GeometryCollection/GeometryCollectionAlgo.h
+++ b/Engine/Source/Runtime/Experimental/GeometryCollectionCore/Public/GeometryCollection/GeometryCollectionAlgo.h
@@ -2,7 +2,6 @@
 
 #pragma once
 
-#include "Chaos/IncludeLvl1.inl"
 #include "CoreMinimal.h"
 #include "ManagedArray.h"
 #include "GeometryCollection/TransformCollection.h"
diff --git a/Engine/Source/Runtime/Experimental/GeometryCollectionCore/Public/GeometryCollection/GeometryCollectionBoneNode.h b/Engine/Source/Runtime/Experimental/GeometryCollectionCore/Public/GeometryCollection/GeometryCollectionBoneNode.h
--- a/Engine/Source/Runtime/Experimental/GeometryCollectionCore/Public/GeometryCollection/GeometryCollectionBoneNode.h
+++ b/Engine/Source/Runtime/Experimental/GeometryCollectionCore/Public/GeometryCollection/GeometryCollectionBoneNode.h
@@ -1,7 +1,6 @@
 // Copyright Epic Games, Inc. All Rights Reserved.
 
 #pragma once
-#include "Chaos/IncludeLvl1.inl"
 #include "Containers/Set.h"
 #include "CoreMinimal.h"
 
diff --git a/Engine/Source/Runtime/Experimental/GeometryCollectionCore/Public/GeometryCollection/GeometryCollectionClusteringUtility.h b/Engine/Source/Runtime/Experimental/GeometryCollectionCore/Public/GeometryCollection/GeometryCollectionClusteringUtility.h
--- a/Engine/Source/Runtime/Experimental/GeometryCollectionCore/Public/GeometryCollection/GeometryCollectionClusteringUtility.h
+++ b/Engine/Source/Runtime/Experimental/GeometryCollectionCore/Public/GeometryCollection/GeometryCollectionClusteringUtility.h
@@ -2,7 +2,6 @@
 
 #pragma once
 
-#include "Chaos/IncludeLvl1.inl"
 #include "CoreMinimal.h"
 #include "GeometryCollection/ManagedArrayCollection.h"
 
diff --git a/Engine/Source/Runtime/Experimental/GeometryCollectionCore/Public/GeometryCollection/GeometryCollectionPlugin.h b/Engine/Source/Runtime/Experimental/GeometryCollectionCore/Public/GeometryCollection/GeometryCollectionPlugin.h
--- a/Engine/Source/Runtime/Experimental/GeometryCollectionCore/Public/GeometryCollection/GeometryCollectionPlugin.h
+++ b/Engine/Source/Runtime/Experimental/GeometryCollectionCore/Public/GeometryCollection/GeometryCollectionPlugin.h
@@ -2,7 +2,6 @@
 
 #pragma once
 
-#include "Chaos/IncludeLvl1.inl"
 #include "CoreMinimal.h"
 #include "Modules/ModuleInterface.h"
 #include "Modules/ModuleManager.h"
diff --git a/Engine/Source/Runtime/Experimental/GeometryCollectionCore/Public/GeometryCollection/GeometryCollectionProximityUtility.h b/Engine/Source/Runtime/Experimental/GeometryCollectionCore/Public/GeometryCollection/GeometryCollectionProximityUtility.h
--- a/Engine/Source/Runtime/Experimental/GeometryCollectionCore/Public/GeometryCollection/GeometryCollectionProximityUtility.h
+++ b/Engine/Source/Runtime/Experimental/GeometryCollectionCore/Public/GeometryCollection/GeometryCollectionProximityUtility.h
@@ -2,7 +2,6 @@
 
 #pragma once
 
-#include "Chaos/IncludeLvl1.inl"
 #include "CoreMinimal.h"
 //#include "GeometryCollection/ManagedArrayCollection.h"
 
diff --git a/Engine/Source/Runtime/Experimental/GeometryCollectionCore/Public/GeometryCollection/GeometryCollectionSection.h b/Engine/Source/Runtime/Experimental/GeometryCollectionCore/Public/GeometryCollection/GeometryCollectionSection.h
--- a/Engine/Source/Runtime/Experimental/GeometryCollectionCore/Public/GeometryCollection/GeometryCollectionSection.h
+++ b/Engine/Source/Runtime/Experimental/GeometryCollectionCore/Public/GeometryCollection/GeometryCollectionSection.h
@@ -1,7 +1,6 @@
 // Copyright Epic Games, Inc. All Rights Reserved.
 
 #pragma once
-#include "Chaos/IncludeLvl1.inl"
 #include "Containers/Set.h"
 
 class HHitProxy;
diff --git a/Engine/Source/Runtime/Experimental/GeometryCollectionCore/Public/GeometryCollection/GeometryCollectionUtility.h b/Engine/Source/Runtime/Experimental/GeometryCollectionCore/Public/GeometryCollection/GeometryCollectionUtility.h
--- a/Engine/Source/Runtime/Experimental/GeometryCollectionCore/Public/GeometryCollection/GeometryCollectionUtility.h
+++ b/Engine/Source/Runtime/Experimental/GeometryCollectionCore/Public/GeometryCollection/GeometryCollectionUtility.h
@@ -2,7 +2,6 @@
 
 #pragma once
 
-#include "Chaos/IncludeLvl1.inl"
 #include "CoreMinimal.h"
 #include "GeometryCollection/GeometryCollection.h"
 #include "Templates/SharedPointer.h"
diff --git a/Engine/Source/Runtime/Experimental/GeometryCollectionCore/Public/GeometryCollection/ManagedArray.h b/Engine/Source/Runtime/Experimental/GeometryCollectionCore/Public/GeometryCollection/ManagedArray.h
--- a/Engine/Source/Runtime/Experimental/GeometryCollectionCore/Public/GeometryCollection/ManagedArray.h
+++ b/Engine/Source/Runtime/Experimental/GeometryCollectionCore/Public/GeometryCollection/ManagedArray.h
@@ -1,6 +1,5 @@
 // Copyright Epic Games, Inc. All Rights Reserved.
 #pragma once
-#include "Chaos/IncludeLvl1.inl"
 #include "Containers/Array.h"
 #include "GeometryCollection/GeometryCollectionSection.h"
 #include "Templates/SharedPointer.h"
diff --git a/Engine/Source/Runtime/Experimental/GeometryCollectionCore/Public/GeometryCollection/ManagedArrayCollection.h b/Engine/Source/Runtime/Experimental/GeometryCollectionCore/Public/GeometryCollection/ManagedArrayCollection.h
--- a/Engine/Source/Runtime/Experimental/GeometryCollectionCore/Public/GeometryCollection/ManagedArrayCollection.h
+++ b/Engine/Source/Runtime/Experimental/GeometryCollectionCore/Public/GeometryCollection/ManagedArrayCollection.h
@@ -2,7 +2,6 @@
 
 #pragma once
 
-#include "Chaos/IncludeLvl1.inl"
 #include "GeometryCollection/ManagedArray.h"
 #include "GeometryCollection/ManagedArrayTypes.h"
 #include "UObject/Object.h"
diff --git a/Engine/Source/Runtime/Experimental/GeometryCollectionCore/Public/GeometryCollection/RecordedTransformTrack.h b/Engine/Source/Runtime/Experimental/GeometryCollectionCore/Public/GeometryCollection/RecordedTransformTrack.h
--- a/Engine/Source/Runtime/Experimental/GeometryCollectionCore/Public/GeometryCollection/RecordedTransformTrack.h
+++ b/Engine/Source/Runtime/Experimental/GeometryCollectionCore/Public/GeometryCollection/RecordedTransformTrack.h
@@ -2,7 +2,6 @@
 
 #pragma once
 
-#include "Chaos/IncludeLvl1.inl"
 #include "CoreMinimal.h"
 #include "UObject/ObjectMacros.h"
 #include "HAL/IConsoleManager.h"
diff --git a/Engine/Source/Runtime/Experimental/GeometryCollectionCore/Public/GeometryCollection/TransformCollection.h b/Engine/Source/Runtime/Experimental/GeometryCollectionCore/Public/GeometryCollection/TransformCollection.h
--- a/Engine/Source/Runtime/Experimental/GeometryCollectionCore/Public/GeometryCollection/TransformCollection.h
+++ b/Engine/Source/Runtime/Experimental/GeometryCollectionCore/Public/GeometryCollection/TransformCollection.h
@@ -2,7 +2,6 @@
 
 #pragma once
 
-#include "Chaos/IncludeLvl1.inl"
 #include "CoreMinimal.h"
 #include "GeometryCollection/ManagedArrayCollection.h"
 #include "Misc/Crc.h"
diff --git a/Engine/Source/Runtime/Experimental/GeometryCollectionEngine/Private/GeometryCollection/GeometryCollectionComponentPluginPrivate.h b/Engine/Source/Runtime/Experimental/GeometryCollectionEngine/Private/GeometryCollection/GeometryCollectionComponentPluginPrivate.h
--- a/Engine/Source/Runtime/Experimental/GeometryCollectionEngine/Private/GeometryCollection/GeometryCollectionComponentPluginPrivate.h
+++ b/Engine/Source/Runtime/Experimental/GeometryCollectionEngine/Private/GeometryCollection/GeometryCollectionComponentPluginPrivate.h
@@ -2,7 +2,6 @@
 
 #pragma once
 
-#include "Chaos/IncludeLvl1.inl"
 #include "CoreMinimal.h"
 #include "Stats/Stats.h"
 
diff --git a/Engine/Source/Runtime/Experimental/GeometryCollectionEngine/Private/GeometryCollection/GeometryCollectionSQAccelerator.h b/Engine/Source/Runtime/Experimental/GeometryCollectionEngine/Private/GeometryCollection/GeometryCollectionSQAccelerator.h
--- a/Engine/Source/Runtime/Experimental/GeometryCollectionEngine/Private/GeometryCollection/GeometryCollectionSQAccelerator.h
+++ b/Engine/Source/Runtime/Experimental/GeometryCollectionEngine/Private/GeometryCollection/GeometryCollectionSQAccelerator.h
@@ -3,7 +3,6 @@
 
 
 #if !WITH_CHAOS_NEEDS_TO_BE_FIXED
-#include "Chaos/IncludeLvl1.inl"
 #include "SQAccelerator.h"
 
 class UGeometryCollectionComponent;
diff --git a/Engine/Source/Runtime/Experimental/GeometryCollectionEngine/Private/GeometryCollection/GeometryCollectionSceneProxy.h b/Engine/Source/Runtime/Experimental/GeometryCollectionEngine/Private/GeometryCollection/GeometryCollectionSceneProxy.h
--- a/Engine/Source/Runtime/Experimental/GeometryCollectionEngine/Private/GeometryCollection/GeometryCollectionSceneProxy.h
+++ b/Engine/Source/Runtime/Experimental/GeometryCollectionEngine/Private/GeometryCollection/GeometryCollectionSceneProxy.h
@@ -1,7 +1,6 @@
 // Copyright Epic Games, Inc. All Rights Reserved.
 #pragma once
 
-#include "Chaos/IncludeLvl1.inl"
 #include "DynamicMeshBuilder.h"
 #include "EngineGlobals.h"
 #include "PrimitiveViewRelevance.h"
diff --git a/Engine/Source/Runtime/Experimental/GeometryCollectionEngine/Public/ChaosBreakingEventFilter.h b/Engine/Source/Runtime/Experimental/GeometryCollectionEngine/Public/ChaosBreakingEventFilter.h
--- a/Engine/Source/Runtime/Experimental/GeometryCollectionEngine/Public/ChaosBreakingEventFilter.h
+++ b/Engine/Source/Runtime/Experimental/GeometryCollectionEngine/Public/ChaosBreakingEventFilter.h
@@ -1,7 +1,6 @@
 // Copyright Epic Games, Inc. All Rights Reserved.
 #pragma once
 
-#include "Chaos/IncludeLvl1.inl"
 #include "CoreMinimal.h"
 #include "UObject/ObjectMacros.h"
 #include "ChaosFilter.h"
diff --git a/Engine/Source/Runtime/Experimental/GeometryCollectionEngine/Public/ChaosCollisionEventFilter.h b/Engine/Source/Runtime/Experimental/GeometryCollectionEngine/Public/ChaosCollisionEventFilter.h
--- a/Engine/Source/Runtime/Experimental/GeometryCollectionEngine/Public/ChaosCollisionEventFilter.h
+++ b/Engine/Source/Runtime/Experimental/GeometryCollectionEngine/Public/ChaosCollisionEventFilter.h
@@ -1,7 +1,6 @@
 // Copyright Epic Games, Inc. All Rights Reserved.
 #pragma once
 
-#include "Chaos/IncludeLvl1.inl"
 #include "CoreMinimal.h"
 #include "UObject/ObjectMacros.h"
 #include "ChaosFilter.h"
diff --git a/Engine/Source/Runtime/Experimental/GeometryCollectionEngine/Public/ChaosFilter.h b/Engine/Source/Runtime/Experimental/GeometryCollectionEngine/Public/ChaosFilter.h
--- a/Engine/Source/Runtime/Experimental/GeometryCollectionEngine/Public/ChaosFilter.h
+++ b/Engine/Source/Runtime/Experimental/GeometryCollectionEngine/Public/ChaosFilter.h
@@ -1,7 +1,6 @@
 // Copyright Epic Games, Inc. All Rights Reserved.
 #pragma once
 
-#include "Chaos/IncludeLvl1.inl"
 #include "CoreMinimal.h"
 
 /**
diff --git a/Engine/Source/Runtime/Experimental/GeometryCollectionEngine/Public/ChaosTrailingEventFilter.h b/Engine/Source/Runtime/Experimental/GeometryCollectionEngine/Public/ChaosTrailingEventFilter.h
--- a/Engine/Source/Runtime/Experimental/GeometryCollectionEngine/Public/ChaosTrailingEventFilter.h
+++ b/Engine/Source/Runtime/Experimental/GeometryCollectionEngine/Public/ChaosTrailingEventFilter.h
@@ -1,7 +1,6 @@
 // Copyright Epic Games, Inc. All Rights Reserved.
 #pragma once
 
-#include "Chaos/IncludeLvl1.inl"
 #include "CoreMinimal.h"
 #include "UObject/ObjectMacros.h"
 #include "ChaosFilter.h"
diff --git a/Engine/Source/Runtime/Experimental/GeometryCollectionEngine/Public/GeometryCollection/DerivedDataGeometryCollectionCooker.h b/Engine/Source/Runtime/Experimental/GeometryCollectionEngine/Public/GeometryCollection/DerivedDataGeometryCollectionCooker.h
--- a/Engine/Source/Runtime/Experimental/GeometryCollectionEngine/Public/GeometryCollection/DerivedDataGeometryCollectionCooker.h
+++ b/Engine/Source/Runtime/Experimental/GeometryCollectionEngine/Public/GeometryCollection/DerivedDataGeometryCollectionCooker.h
@@ -2,7 +2,6 @@
 
 #pragma once
 
-#include "Chaos/IncludeLvl1.inl"
 #include "CoreMinimal.h"
 #if WITH_EDITOR
 	#include "DerivedDataPluginInterface.h"
diff --git a/Engine/Source/Runtime/Experimental/GeometryCollectionEngine/Public/GeometryCollection/GeometryCollectionActor.h b/Engine/Source/Runtime/Experimental/GeometryCollectionEngine/Public/GeometryCollection/GeometryCollectionActor.h
--- a/Engine/Source/Runtime/Experimental/GeometryCollectionEngine/Public/GeometryCollection/GeometryCollectionActor.h
+++ b/Engine/Source/Runtime/Experimental/GeometryCollectionEngine/Public/GeometryCollection/GeometryCollectionActor.h
@@ -4,7 +4,6 @@
 
 /** This class represents an GeometryCollection Actor. */
 
-#include "Chaos/IncludeLvl1.inl"
 #include "CoreMinimal.h"
 #include "GameFramework/Actor.h"
 #include "GeometryCollection/GeometryCollection.h"
diff --git a/Engine/Source/Runtime/Experimental/GeometryCollectionEngine/Public/GeometryCollection/GeometryCollectionCache.h b/Engine/Source/Runtime/Experimental/GeometryCollectionEngine/Public/GeometryCollection/GeometryCollectionCache.h
--- a/Engine/Source/Runtime/Experimental/GeometryCollectionEngine/Public/GeometryCollection/GeometryCollectionCache.h
+++ b/Engine/Source/Runtime/Experimental/GeometryCollectionEngine/Public/GeometryCollection/GeometryCollectionCache.h
@@ -2,7 +2,6 @@
 
 #pragma once
 
-#include "Chaos/IncludeLvl1.inl"
 #include "CoreMinimal.h"
 #include "GeometryCollection/RecordedTransformTrack.h"
 
diff --git a/Engine/Source/Runtime/Experimental/GeometryCollectionEngine/Public/GeometryCollection/GeometryCollectionComponentPlugin.h b/Engine/Source/Runtime/Experimental/GeometryCollectionEngine/Public/GeometryCollection/GeometryCollectionComponentPlugin.h
--- a/Engine/Source/Runtime/Experimental/GeometryCollectionEngine/Public/GeometryCollection/GeometryCollectionComponentPlugin.h
+++ b/Engine/Source/Runtime/Experimental/GeometryCollectionEngine/Public/GeometryCollection/GeometryCollectionComponentPlugin.h
@@ -2,7 +2,6 @@
 
 #pragma once
 
-#include "Chaos/IncludeLvl1.inl"
 #include "CoreMinimal.h"
 #include "Modules/ModuleInterface.h"
 #include "Modules/ModuleManager.h"
diff --git a/Engine/Source/Runtime/Experimental/GeometryCollectionEngine/Public/GeometryCollection/GeometryCollectionCreationParameters.h b/Engine/Source/Runtime/Experimental/GeometryCollectionEngine/Public/GeometryCollection/GeometryCollectionCreationParameters.h
--- a/Engine/Source/Runtime/Experimental/GeometryCollectionEngine/Public/GeometryCollection/GeometryCollectionCreationParameters.h
+++ b/Engine/Source/Runtime/Experimental/GeometryCollectionEngine/Public/GeometryCollection/GeometryCollectionCreationParameters.h
@@ -1,7 +1,6 @@
 // Copyright Epic Games, Inc. All Rights Reserved.
 
 #pragma once
-#include "Chaos/IncludeLvl1.inl"
 #include "CoreMinimal.h"
 
 
diff --git a/Engine/Source/Runtime/Experimental/GeometryCollectionEngine/Public/GeometryCollection/GeometryCollectionDebugDrawActor.h b/Engine/Source/Runtime/Experimental/GeometryCollectionEngine/Public/GeometryCollection/GeometryCollectionDebugDrawActor.h
--- a/Engine/Source/Runtime/Experimental/GeometryCollectionEngine/Public/GeometryCollection/GeometryCollectionDebugDrawActor.h
+++ b/Engine/Source/Runtime/Experimental/GeometryCollectionEngine/Public/GeometryCollection/GeometryCollectionDebugDrawActor.h
@@ -2,7 +2,6 @@
 
 #pragma once
 
-#include "Chaos/IncludeLvl1.inl"
 #include "GameFramework/Actor.h"
 #include "Components/BillboardComponent.h"
 
diff --git a/Engine/Source/Runtime/Experimental/GeometryCollectionEngine/Public/GeometryCollection/GeometryCollectionDebugDrawComponent.h b/Engine/Source/Runtime/Experimental/GeometryCollectionEngine/Public/GeometryCollection/GeometryCollectionDebugDrawComponent.h
--- a/Engine/Source/Runtime/Experimental/GeometryCollectionEngine/Public/GeometryCollection/GeometryCollectionDebugDrawComponent.h
+++ b/Engine/Source/Runtime/Experimental/GeometryCollectionEngine/Public/GeometryCollection/GeometryCollectionDebugDrawComponent.h
@@ -6,7 +6,6 @@
 #define GEOMETRYCOLLECTION_DEBUG_DRAW !(UE_BUILD_SHIPPING || UE_BUILD_TEST)
 #endif
 
-#include "Chaos/IncludeLvl1.inl"
 #include "Components/MeshComponent.h"
 
 #include "GeometryCollectionDebugDrawComponent.generated.h"
diff --git a/Engine/Source/Runtime/Experimental/GeometryCollectionEngine/Public/GeometryCollection/GeometryCollectionEditorSelection.h b/Engine/Source/Runtime/Experimental/GeometryCollectionEngine/Public/GeometryCollection/GeometryCollectionEditorSelection.h
--- a/Engine/Source/Runtime/Experimental/GeometryCollectionEngine/Public/GeometryCollection/GeometryCollectionEditorSelection.h
+++ b/Engine/Source/Runtime/Experimental/GeometryCollectionEngine/Public/GeometryCollection/GeometryCollectionEditorSelection.h
@@ -1,7 +1,6 @@
 // Copyright Epic Games, Inc. All Rights Reserved.
 
 #pragma once
-#include "Chaos/IncludeLvl1.inl"
 
 #ifndef GEOMETRYCOLLECTION_EDITOR_SELECTION
 #define GEOMETRYCOLLECTION_EDITOR_SELECTION WITH_EDITOR
diff --git a/Engine/Source/Runtime/Experimental/GeometryCollectionEngine/Public/GeometryCollection/GeometryCollectionEngineTypes.h b/Engine/Source/Runtime/Experimental/GeometryCollectionEngine/Public/GeometryCollection/GeometryCollectionEngineTypes.h
--- a/Engine/Source/Runtime/Experimental/GeometryCollectionEngine/Public/GeometryCollection/GeometryCollectionEngineTypes.h
+++ b/Engine/Source/Runtime/Experimental/GeometryCollectionEngine/Public/GeometryCollection/GeometryCollectionEngineTypes.h
@@ -3,7 +3,6 @@
 #pragma once
 
 
-#include "Chaos/IncludeLvl1.inl"
 #include "UObject/ObjectMacros.h"
 #include "GeometryCollectionProxyData.h"
 #include "GeometryCollection/GeometryCollection.h"
diff --git a/Engine/Source/Runtime/Experimental/GeometryCollectionEngine/Public/GeometryCollection/GeometryCollectionEngineUtility.h b/Engine/Source/Runtime/Experimental/GeometryCollectionEngine/Public/GeometryCollection/GeometryCollectionEngineUtility.h
--- a/Engine/Source/Runtime/Experimental/GeometryCollectionEngine/Public/GeometryCollection/GeometryCollectionEngineUtility.h
+++ b/Engine/Source/Runtime/Experimental/GeometryCollectionEngine/Public/GeometryCollection/GeometryCollectionEngineUtility.h
@@ -2,7 +2,6 @@
 
 #pragma once
 
-#include "Chaos/IncludeLvl1.inl"
 #include "CoreMinimal.h"
 
 class FGeometryCollection;
diff --git a/Engine/Source/Runtime/Experimental/GeometryCollectionEngine/Public/GeometryCollection/GeometryCollectionHitProxy.h b/Engine/Source/Runtime/Experimental/GeometryCollectionEngine/Public/GeometryCollection/GeometryCollectionHitProxy.h
--- a/Engine/Source/Runtime/Experimental/GeometryCollectionEngine/Public/GeometryCollection/GeometryCollectionHitProxy.h
+++ b/Engine/Source/Runtime/Experimental/GeometryCollectionEngine/Public/GeometryCollection/GeometryCollectionHitProxy.h
@@ -1,7 +1,6 @@
 // Copyright Epic Games, Inc. All Rights Reserved.
 #pragma once
 
-#include "Chaos/IncludeLvl1.inl"
 #include "GeometryCollectionEditorSelection.h"
 
 #if GEOMETRYCOLLECTION_EDITOR_SELECTION
diff --git a/Engine/Source/Runtime/Experimental/GeometryCollectionEngine/Public/GeometryCollection/GeometryCollectionObject.h b/Engine/Source/Runtime/Experimental/GeometryCollectionEngine/Public/GeometryCollection/GeometryCollectionObject.h
--- a/Engine/Source/Runtime/Experimental/GeometryCollectionEngine/Public/GeometryCollection/GeometryCollectionObject.h
+++ b/Engine/Source/Runtime/Experimental/GeometryCollectionEngine/Public/GeometryCollection/GeometryCollectionObject.h
@@ -2,7 +2,6 @@
 
 #pragma once
 
-#include "Chaos/IncludeLvl1.inl"
 #include "CoreMinimal.h"
 #include "Misc/Crc.h"
 #include "GeometryCollection/ManagedArray.h"
diff --git a/Engine/Source/Runtime/Experimental/GeometryCollectionEngine/Public/GeometryCollection/GeometryCollectionRenderLevelSetActor.h b/Engine/Source/Runtime/Experimental/GeometryCollectionEngine/Public/GeometryCollection/GeometryCollectionRenderLevelSetActor.h
--- a/Engine/Source/Runtime/Experimental/GeometryCollectionEngine/Public/GeometryCollection/GeometryCollectionRenderLevelSetActor.h
+++ b/Engine/Source/Runtime/Experimental/GeometryCollectionEngine/Public/GeometryCollection/GeometryCollectionRenderLevelSetActor.h
@@ -2,7 +2,6 @@
 
 #pragma once
 
-#include "Chaos/IncludeLvl1.inl"
 #include "CoreMinimal.h"
 #include "GameFramework/Actor.h"
 #include "GeometryCollection/GeometryCollection.h"
diff --git a/Engine/Source/Runtime/Experimental/GeometryCollectionEngine/Public/GeometryCollection/PhysicsAssetSimulation.h b/Engine/Source/Runtime/Experimental/GeometryCollectionEngine/Public/GeometryCollection/PhysicsAssetSimulation.h
--- a/Engine/Source/Runtime/Experimental/GeometryCollectionEngine/Public/GeometryCollection/PhysicsAssetSimulation.h
+++ b/Engine/Source/Runtime/Experimental/GeometryCollectionEngine/Public/GeometryCollection/PhysicsAssetSimulation.h
@@ -1,7 +1,6 @@
 // Copyright Epic Games, Inc. All Rights Reserved.
 #pragma once
 
-#include "Chaos/IncludeLvl1.inl"
 #include "CoreMinimal.h"
 #include "Chaos/ChaosNotifyHandlerInterface.h"
 #include "GeometryCollection/GeometryCollectionSimulationTypes.h"
diff --git a/Engine/Source/Runtime/Experimental/GeometryCollectionSimulationCore/Public/GeometryCollection/GeometryCollectionCollisionStructureManager.h b/Engine/Source/Runtime/Experimental/GeometryCollectionSimulationCore/Public/GeometryCollection/GeometryCollectionCollisionStructureManager.h
--- a/Engine/Source/Runtime/Experimental/GeometryCollectionSimulationCore/Public/GeometryCollection/GeometryCollectionCollisionStructureManager.h
+++ b/Engine/Source/Runtime/Experimental/GeometryCollectionSimulationCore/Public/GeometryCollection/GeometryCollectionCollisionStructureManager.h
@@ -2,7 +2,6 @@
 
 #pragma once
 
-#include "Chaos/IncludeLvl1.inl"
 #include "CoreMinimal.h"
 #include "GeometryCollection/GeometryCollection.h"
 #include "GeometryCollection/GeometryCollectionSimulationCoreTypes.h"
diff --git a/Engine/Source/Runtime/Experimental/GeometryCollectionSimulationCore/Public/GeometryCollection/GeometryCollectionPhysicsFieldCommands.h b/Engine/Source/Runtime/Experimental/GeometryCollectionSimulationCore/Public/GeometryCollection/GeometryCollectionPhysicsFieldCommands.h
--- a/Engine/Source/Runtime/Experimental/GeometryCollectionSimulationCore/Public/GeometryCollection/GeometryCollectionPhysicsFieldCommands.h
+++ b/Engine/Source/Runtime/Experimental/GeometryCollectionSimulationCore/Public/GeometryCollection/GeometryCollectionPhysicsFieldCommands.h
@@ -3,7 +3,6 @@
 #pragma once
 
 
-#include "Chaos/IncludeLvl1.inl"
 #include "GeometryCollection/ManagedArray.h"
 #include "Field/FieldSystemTypes.h"
 #include "Field/FieldSystem.h"
diff --git a/Engine/Source/Runtime/Experimental/GeometryCollectionSimulationCore/Public/GeometryCollection/GeometryCollectionSimulationCorePlugin.h b/Engine/Source/Runtime/Experimental/GeometryCollectionSimulationCore/Public/GeometryCollection/GeometryCollectionSimulationCorePlugin.h
--- a/Engine/Source/Runtime/Experimental/GeometryCollectionSimulationCore/Public/GeometryCollection/GeometryCollectionSimulationCorePlugin.h
+++ b/Engine/Source/Runtime/Experimental/GeometryCollectionSimulationCore/Public/GeometryCollection/GeometryCollectionSimulationCorePlugin.h
@@ -2,7 +2,6 @@
 
 #pragma once
 
-#include "Chaos/IncludeLvl1.inl"
 #include "CoreMinimal.h"
 #include "Modules/ModuleInterface.h"
 #include "Modules/ModuleManager.h"
diff --git a/Engine/Source/Runtime/Experimental/GeometryCollectionSimulationCore/Public/GeometryCollection/GeometryCollectionSimulationCoreTypes.h b/Engine/Source/Runtime/Experimental/GeometryCollectionSimulationCore/Public/GeometryCollection/GeometryCollectionSimulationCoreTypes.h
--- a/Engine/Source/Runtime/Experimental/GeometryCollectionSimulationCore/Public/GeometryCollection/GeometryCollectionSimulationCoreTypes.h
+++ b/Engine/Source/Runtime/Experimental/GeometryCollectionSimulationCore/Public/GeometryCollection/GeometryCollectionSimulationCoreTypes.h
@@ -2,7 +2,6 @@
 
 #pragma once
 
-#include "Chaos/IncludeLvl1.inl"
 #include "Containers/UnrealString.h"
 #include "Field/FieldSystem.h"
 #include "GeometryCollection/RecordedTransformTrack.h"
diff --git a/Engine/Source/Runtime/Experimental/GeometryCollectionSimulationCore/Public/GeometryCollection/GeometryCollectionSimulationTypes.h b/Engine/Source/Runtime/Experimental/GeometryCollectionSimulationCore/Public/GeometryCollection/GeometryCollectionSimulationTypes.h
--- a/Engine/Source/Runtime/Experimental/GeometryCollectionSimulationCore/Public/GeometryCollection/GeometryCollectionSimulationTypes.h
+++ b/Engine/Source/Runtime/Experimental/GeometryCollectionSimulationCore/Public/GeometryCollection/GeometryCollectionSimulationTypes.h
@@ -1,7 +1,6 @@
 // Copyright Epic Games, Inc. All Rights Reserved.
 #pragma once
 
-#include "Chaos/IncludeLvl1.inl"
 #include "UObject/ObjectMacros.h"
 #include "Chaos/PBDRigidParticles.h"
 
