Chaos
- Joint elliptical limits done properly
- Elliptical limit unit tests
- Clamp max FPS in solver (at 100Hz) to fix high FPS collision resolution issues

diff --git a/Engine/Source/Programs/HeadlessChaos/Private/HeadlessChaosTestJointUtilities.cpp b/Engine/Source/Programs/HeadlessChaos/Private/HeadlessChaosTestJointUtilities.cpp
--- a/Engine/Source/Programs/HeadlessChaos/Private/HeadlessChaosTestJointUtilities.cpp
+++ b/Engine/Source/Programs/HeadlessChaos/Private/HeadlessChaosTestJointUtilities.cpp
@@ -116,4 +116,112 @@ namespace ChaosTest
 		}
 	}
 
+	void TestEllipticalConeAxisError(const FVec3& Axis, FReal AngleDeg, FReal LimitYDeg, FReal LimitZDeg, const FVec3& ExpectedAxis, FReal ExpectedErrorDeg)
+	{
+		FReal Angle = FMath::DegreesToRadians(AngleDeg);
+		FReal LimitY = FMath::DegreesToRadians(LimitYDeg);
+		FReal LimitZ = FMath::DegreesToRadians(LimitZDeg);
+		FRotation3 RTwist = FRotation3::FromIdentity();
+		FRotation3 RSwing = FRotation3::FromAxisAngle(Axis, Angle);
+
+		FVec3 AxisLocal;
+		FReal Error;
+		FPBDJointUtilities::GetEllipticalConeAxisErrorLocal(RTwist, RSwing, LimitY, LimitZ, AxisLocal, Error);
+
+		FReal ErrorDeg = FMath::RadiansToDegrees(Error);
+
+		// About 0.5deg tolerance on error angle
+		EXPECT_NEAR(ExpectedErrorDeg, ErrorDeg, 0.5f) << "Axis: (" << Axis.X << ", " << Axis.Y << ", " << Axis.Z << "); Angle = " << AngleDeg;
+		EXPECT_NEAR(AxisLocal.Size(), 1.0f, KINDA_SMALL_NUMBER);
+
+		// About 1deg error tolerance on axis
+		FReal AxisDot = FVec3::DotProduct(AxisLocal, ExpectedAxis);
+		EXPECT_NEAR(AxisDot, 1.0f, 0.02f);
+	}
+
+
+	GTEST_TEST(JointUtilitiesTes, TestEllipticalConeAxisError)
+	{
+		// Test Swing along Y Minor Axis
+		{
+			FVec3 Axis = FVec3(0, 1, 0);
+			FReal AngleDeg = 40.0f;
+			FReal LimitYDeg = 20.0f;
+			FReal LimitZDeg = 30.0f;
+			TestEllipticalConeAxisError(Axis, AngleDeg, LimitYDeg, LimitZDeg, Axis, AngleDeg - LimitYDeg);
+		}
+
+		// Test Swing along Y Major Axis
+		{
+			FVec3 Axis = FVec3(0, 1, 0);
+			FReal AngleDeg = 40.0f;
+			FReal LimitYDeg = 30.0f;
+			FReal LimitZDeg = 20.0f;
+			TestEllipticalConeAxisError(Axis, AngleDeg, LimitYDeg, LimitZDeg, Axis, AngleDeg - LimitYDeg);
+		}
+
+		// Test Swing along Z Minor Axis
+		{
+			FVec3 Axis = FVec3(0, 0, 1);
+			FReal AngleDeg = 40.0f;
+			FReal LimitYDeg = 30.0f;
+			FReal LimitZDeg = 20.0f;
+			TestEllipticalConeAxisError(Axis, AngleDeg, LimitYDeg, LimitZDeg, Axis, AngleDeg - LimitZDeg);
+		}
+
+		// Test Swing along Z Major Axis
+		{
+			FVec3 Axis = FVec3(0, 0, 1);
+			FReal AngleDeg = 40.0f;
+			FReal LimitYDeg = 20.0f;
+			FReal LimitZDeg = 30.0f;
+			TestEllipticalConeAxisError(Axis, AngleDeg, LimitYDeg, LimitZDeg, Axis, AngleDeg - LimitZDeg);
+		}
+
+		// Test Cicular
+		{
+			for (int32 Deg = 0; Deg < 360; ++Deg)
+			{
+				FReal AxisDeg = (FReal)Deg;
+				FVec3 Axis = FRotation3::FromAxisAngle(FVec3(1, 0, 0), FMath::DegreesToRadians(AxisDeg)) *  FVec3(0, 1, 0);
+				FReal AngleDeg = 40.0f;
+				FReal LimitDeg = 20.0f;
+				TestEllipticalConeAxisError(Axis, AngleDeg, LimitDeg, LimitDeg, Axis, AngleDeg - LimitDeg);
+			}
+		}
+
+		// Test Elliptical
+		{
+			FReal ExpectedErrorDeg = 20.0f;
+			FReal LimitYDeg = 20.0f;
+			FReal LimitZDeg = 30.0f;
+			const int32 NumSteps = 360;
+			for (int32 Step = 0; Step < NumSteps; ++Step)
+			{
+				FReal S = (FReal)Step / (FReal)NumSteps;
+
+				// Find the point on the ellipse for parameter S, and then a point with the specified error along the normal
+				// https://mathworld.wolfram.com/Ellipse.html
+				FReal CosS = FMath::Cos(S * 2.0f * PI);
+				FReal SinS = FMath::Sin(S * 2.0f * PI);
+				FReal TDenom = FMath::Sqrt(LimitYDeg * LimitYDeg * CosS * CosS + LimitZDeg * LimitZDeg * SinS * SinS);
+				FVec2 P = FVec2(LimitZDeg * CosS, LimitYDeg * SinS);						// point on elipse for parameter S
+				FVec2 T = FVec2(-LimitZDeg * SinS / TDenom, LimitYDeg * CosS / TDenom);		// tangent to ellipse at P
+				FVec2 N = FVec2(T.Y, -T.X).GetSafeNormal();									// normal to ellipse at P
+
+				// Point "error" degrees away from ellipse
+				FVec3 NetP = FVec3(0.0f, P.X, P.Y) + FVec3(0.0f, N.X, N.Y) * ExpectedErrorDeg;
+
+				// Get the axis and to rotate about to get to the off-ellipse point we want
+				FReal AngleDeg = NetP.Size();
+				FVec3 NetPDir = NetP / AngleDeg;
+				FVec3 Axis = FVec3::CrossProduct(FVec3(1, 0, 0), NetPDir);
+
+				FVec3 ExpectedAxis = FVec3(0.0f, T.X, T.Y);
+				TestEllipticalConeAxisError(Axis, AngleDeg, LimitYDeg, LimitZDeg, ExpectedAxis, ExpectedErrorDeg);
+			}
+		}
+
+	}
+
 }
\ No newline at end of file
diff --git a/Engine/Source/Runtime/Engine/Private/PhysicsEngine/ImmediatePhysics/ImmediatePhysicsChaos/ImmediatePhysicsSimulation_Chaos.cpp b/Engine/Source/Runtime/Engine/Private/PhysicsEngine/ImmediatePhysics/ImmediatePhysicsChaos/ImmediatePhysicsSimulation_Chaos.cpp
--- a/Engine/Source/Runtime/Engine/Private/PhysicsEngine/ImmediatePhysics/ImmediatePhysicsChaos/ImmediatePhysicsSimulation_Chaos.cpp
+++ b/Engine/Source/Runtime/Engine/Private/PhysicsEngine/ImmediatePhysics/ImmediatePhysicsChaos/ImmediatePhysicsSimulation_Chaos.cpp
@@ -49,8 +49,10 @@ FAutoConsoleVariableRef CVarChaosImmPhysIterations(TEXT("p.Chaos.ImmPhys.Iterati
 FAutoConsoleVariableRef CVarChaosImmPhysPushOutIterations(TEXT("p.Chaos.ImmPhys.PushOutIterations"), ChaosImmediate_Evolution_PushOutIterations, TEXT("Override number of solver push-out loops (if >= 0)"));
 FAutoConsoleVariableRef CVarChaosImmPhysBoundsExtension(TEXT("p.Chaos.ImmPhys.BoundsExtension"), ChaosImmediate_Evolution_BoundsExtension, TEXT("Bounds are grown by this fraction of their size (should be >= 0.0)"));
 
+float ChaosImmediate_Evolution_MinStepTime = 0.01f;
 float ChaosImmediate_Evolution_FixedStepTime = -1.0f;
 float ChaosImmediate_Evolution_FixedStepTolerance = 0.05f;
+FAutoConsoleVariableRef CVarChaosImmPhysMinStepTime(TEXT("p.Chaos.ImmPhys.MinStepTime"), ChaosImmediate_Evolution_MinStepTime, TEXT("If non-zero, then if step time is lower than this, go into fixed step mode with this timestep."));
 FAutoConsoleVariableRef CVarChaosImmPhysFixedStepTime(TEXT("p.Chaos.ImmPhys.FixedStepTime"), ChaosImmediate_Evolution_FixedStepTime, TEXT("Override fixed step time mode: fixed step time (if positive); variable time mode (if zero); asset defined (if negative)"));
 FAutoConsoleVariableRef CVarChaosImmPhysFixedStepTolerance(TEXT("p.Chaos.ImmPhys.FixedStepTolerance"), ChaosImmediate_Evolution_FixedStepTolerance, TEXT("Tiem remainder required to add a new step (fraction of FixedStepTime)"));
 
@@ -642,12 +644,21 @@ namespace ImmediatePhysics_Chaos
 		DebugDrawDynamicParticles(2, 2, FColor(192, 192, 0));
 		DebugDrawConstraints(2, 2, 0.7f);
 
-		// Fixed timestep mode
-		// @todo(ccaulfield): integrate this better with the non-fixed step calculate above
+		// Fixed timestep mode DT (Only used if > 0.0f)
+		FReal FixedStepTime = Implementation->FixedStepTime;
+
+		// Use fixed step mode anyway if StepTime is too low
+		// This can prevent collision push resolution introducing large velocities at small DTs
+		if ((FixedStepTime <= 0.0f) && (StepTime < ChaosImmediate_Evolution_MinStepTime))
+		{
+			FixedStepTime = ChaosImmediate_Evolution_MinStepTime;
+		}
+
+		// If using FixedStep mode, calculate the number of steps and how far to rewind (if at all)
 		FReal RewindTime = 0.0f;
-		if (Implementation->FixedStepTime > 0)
+		if (FixedStepTime > 0)
 		{
-			StepTime = Implementation->FixedStepTime;
+			StepTime = FixedStepTime;
 			NumSteps = FMath::FloorToInt(DeltaTime / StepTime);
 			FReal RemainderTime = DeltaTime - NumSteps * StepTime;
 			if (RemainderTime > ChaosImmediate_Evolution_FixedStepTolerance* StepTime)
diff --git a/Engine/Source/Runtime/Experimental/Chaos/Private/Chaos/Joint/PBDJointSolverGaussSeidel.cpp b/Engine/Source/Runtime/Experimental/Chaos/Private/Chaos/Joint/PBDJointSolverGaussSeidel.cpp
--- a/Engine/Source/Runtime/Experimental/Chaos/Private/Chaos/Joint/PBDJointSolverGaussSeidel.cpp
+++ b/Engine/Source/Runtime/Experimental/Chaos/Private/Chaos/Joint/PBDJointSolverGaussSeidel.cpp
@@ -16,6 +16,9 @@
 float Chaos_Joint_DegenerateRotationLimit = -0.998f;	// Cos(176deg)
 FAutoConsoleVariableRef CVarChaosJointDegenerateRotationLimit(TEXT("p.Chaos.Joint.DegenerateRotationLimit"), Chaos_Joint_DegenerateRotationLimit, TEXT("Cosine of the swing angle that is considered degerenerate (default Cos(176deg))"));
 
+bool bChaos_Joint_EllipticalFix = true;
+FAutoConsoleVariableRef CVarChaosJointEllipticalFix(TEXT("p.Chaos.Joint.EllipticalFix"), bChaos_Joint_EllipticalFix, TEXT("Enable the proper elliptical joint axis/error calculation"));
+
 namespace Chaos
 {
 
@@ -1103,25 +1106,38 @@ namespace Chaos
 		const FPBDJointSettings& JointSettings,
 		const bool bUseSoftLimit)
 	{
-		// Calculate swing angle and axis
-		FReal SwingAngle;
 		FVec3 SwingAxisLocal;
-		FPBDJointUtilities::GetConeAxisAngleLocal(Rs[0], Rs[1], SolverSettings.SwingTwistAngleTolerance, SwingAxisLocal, SwingAngle);
-		const FVec3 SwingAxis = Rs[0] * SwingAxisLocal;
-
-		// Calculate swing angle error
-		FReal SwingAngleMax = FPBDJointUtilities::GetConeAngleLimit(JointSettings, SwingAxisLocal, SwingAngle) + AngleTolerance;
 		FReal DSwingAngle = 0.0f;
-		if (SwingAngle > SwingAngleMax)
+
+		if (!bChaos_Joint_EllipticalFix)
 		{
-			DSwingAngle = SwingAngle - SwingAngleMax;
+			// Calculate swing angle and axis
+			FReal SwingAngle;
+			FPBDJointUtilities::GetConeAxisAngleLocal(Rs[0], Rs[1], SolverSettings.SwingTwistAngleTolerance, SwingAxisLocal, SwingAngle);
+
+			// Calculate swing angle error
+			FReal SwingAngleMax = FPBDJointUtilities::GetConeAngleLimit(JointSettings, SwingAxisLocal, SwingAngle) + AngleTolerance;
+			if (SwingAngle > SwingAngleMax)
+			{
+				DSwingAngle = SwingAngle - SwingAngleMax;
+			}
+			else if (SwingAngle < -SwingAngleMax)
+			{
+				DSwingAngle = SwingAngle + SwingAngleMax;
+			}
+
+			UE_LOG(LogChaosJoint, VeryVerbose, TEXT("    Cone Angle %f [Limit %f]"), FMath::RadiansToDegrees(SwingAngle), FMath::RadiansToDegrees(SwingAngleMax));
 		}
-		else if (SwingAngle < -SwingAngleMax)
+		else
 		{
-			DSwingAngle = SwingAngle + SwingAngleMax;
+			const FReal Swing1Limit = JointSettings.AngularLimits[(int32)EJointAngularConstraintIndex::Swing1];
+			const FReal Swing2Limit = JointSettings.AngularLimits[(int32)EJointAngularConstraintIndex::Swing2];
+			FPBDJointUtilities::GetEllipticalConeAxisErrorLocal(Rs[0], Rs[1], Swing2Limit, Swing1Limit, SwingAxisLocal, DSwingAngle);
+
+			UE_LOG(LogChaosJoint, VeryVerbose, TEXT("    Cone Error %f [Limits %f %f]"), FMath::RadiansToDegrees(DSwingAngle), FMath::RadiansToDegrees(Swing2Limit), FMath::RadiansToDegrees(Swing1Limit));
 		}
 
-		UE_LOG(LogChaosJoint, VeryVerbose, TEXT("    Cone Angle %f [Limit %f]"), FMath::RadiansToDegrees(SwingAngle), FMath::RadiansToDegrees(SwingAngleMax));
+		const FVec3 SwingAxis = Rs[0] * SwingAxisLocal;
 
 		// Apply swing correction to each body
 		if (FMath::Abs(DSwingAngle) > 0)
diff --git a/Engine/Source/Runtime/Experimental/Chaos/Private/Chaos/PBDJointConstraintUtilities.cpp b/Engine/Source/Runtime/Experimental/Chaos/Private/Chaos/PBDJointConstraintUtilities.cpp
--- a/Engine/Source/Runtime/Experimental/Chaos/Private/Chaos/PBDJointConstraintUtilities.cpp
+++ b/Engine/Source/Runtime/Experimental/Chaos/Private/Chaos/PBDJointConstraintUtilities.cpp
@@ -8,6 +8,91 @@
 
 namespace Chaos
 {
+	FReal GetSwingAngleY(const FRotation3& RSwing)
+	{
+		return 4.0f * FMath::Atan2(RSwing.Y, 1.0f + RSwing.W);
+	}
+
+	FReal GetSwingAngleZ(const FRotation3& RSwing)
+	{
+		return 4.0f * FMath::Atan2(RSwing.Z, 1.0f + RSwing.W);
+	}
+
+	// Iterative find near point on ellipse. Point InP, Radii R.
+	// https://www.geometrictools.com/Documentation/DistancePointEllipseEllipsoid.pdf
+	FVec2 NearPointOnEllipse(const FVec2& P, const FVec2& R, const int MaxIts = 20, const FReal Tolerance = 1.e-4f)
+	{
+		// Map point into first quadrant
+		FVec2 PAbs = P.GetAbs();
+
+		// Check for point on minor axis
+		const FReal Epsilon = 1.e-6f;
+		if (R.X >= R.Y)
+		{
+			if (PAbs.Y < Epsilon)
+			{
+				return FVec2((P.X > 0.0f) ? R.X : -R.X, 0.0f);
+			}
+		}
+		else
+		{
+			if (PAbs.X < Epsilon)
+			{
+				return FVec2(0.0f, (P.Y > 0.0f)? R.Y : -R.Y);
+			}
+		}
+
+		// Iterate to find nearest point
+		FVec2 R2 = R * R;
+		FVec2 RP = R * PAbs;
+		FReal T = FMath::Max(RP.X - R2.X, RP.Y - R2.Y);
+		FVec2 D;
+		for (int32 It = 0; It < MaxIts; ++It)
+		{
+			D = FVec2(1.0f / (T + R2.X), 1 / (T + R2.Y));
+			FVec2 RPD = RP * D;
+
+			FVec2 FV = RPD * RPD;
+			FReal F = FV.X + FV.Y - 1.0f;
+
+			if (F < Tolerance)
+			{
+				return (R2 * P) * D;
+			}
+
+			FReal DF = -2.0f * FVec2::DotProduct(FV, D);
+			T = T - F / DF;
+		}
+
+		// Too many iterations - return current value clamped
+		FVec2 S = (R2 * P) * D;
+		FVec2 SN = S / R;	
+		return S / SN.Size();
+	}
+
+	bool GetEllipticalAxisError(const FVec3& SwingAxisRot, const FVec3& EllipseNormal, const FVec3& TwistAxis, FVec3& AxisLocal, FReal& Error)
+	{
+		const FReal R2 = SwingAxisRot.SizeSquared();
+		const FReal A = 1.0f - R2;
+		const FReal B = 1.0f / (1.0f + R2);
+		const FReal B2 = B * B;
+		const FReal V1 = 2.0f * A * B2;
+		const FVec3 V2 = FVec3(A, 2.0f * SwingAxisRot.Z, -2.0f * SwingAxisRot.Y);
+		const FReal RD = FVec3::DotProduct(SwingAxisRot, EllipseNormal);
+		const FReal DV1 = -4.0f * RD * (3.0f - R2) * B2 * B;
+		const FVec3 DV2 = FVec3(-2.0f * RD, 2.0f * EllipseNormal.Z, -2.0f * EllipseNormal.Y);
+		
+		const FVec3 Line = V1 * V2 - FVec3(1, 0, 0);
+		FVec3 Normal = V1 * DV2 + DV1 * V2;
+		if (Utilities::NormalizeSafe(Normal))
+		{
+			AxisLocal = FVec3::CrossProduct(Line, Normal);
+			Error = -FVec3::DotProduct(FVec3::CrossProduct(Line, AxisLocal), TwistAxis);
+			return true;
+		}
+		return false;
+	}
+
 	void FPBDJointUtilities::GetSphericalAxisDelta(
 		const FVec3& X0,
 		const FVec3& X1,
@@ -82,8 +167,9 @@ namespace Chaos
 		FRotation3 R01Twist, R01Swing;
 		FPBDJointUtilities::DecomposeSwingTwistLocal(R0, R1, R01Swing, R01Twist);
 		TwistAngle = R01Twist.GetAngle();
-		Swing1Angle = 4.0f * FMath::Atan2(R01Swing.Z, 1.0f + R01Swing.W);
-		Swing2Angle = 4.0f * FMath::Atan2(R01Swing.Y, 1.0f + R01Swing.W);
+		// @todo(ccaulfield): makes assumptions about Swing1 and Swing2 axes - fix this
+		Swing1Angle = GetSwingAngleZ(R01Swing);
+		Swing2Angle = GetSwingAngleY(R01Swing);
 	}
 
 	FReal FPBDJointUtilities::GetTwistAngle(const FRotation3& InTwist)
@@ -135,6 +221,70 @@ namespace Chaos
 		}
 	}
 
+	void FPBDJointUtilities::GetCircularConeAxisErrorLocal(
+		const FRotation3& R0,
+		const FRotation3& R1,
+		const FReal SwingLimit,
+		FVec3& AxisLocal,
+		FReal& Error)
+	{
+		FRotation3 R01Twist, R01Swing;
+		FPBDJointUtilities::DecomposeSwingTwistLocal(R0, R1, R01Swing, R01Twist);
+
+		FReal Angle = 0.0f;
+		R01Swing.ToAxisAndAngleSafe(AxisLocal, Angle, FJointConstants::Swing1Axis(), 1.e-6f);
+
+		Error = 0.0f;
+		if (Angle > SwingLimit)
+		{
+			Error = Angle - SwingLimit;
+		}
+		else if (Angle < -SwingLimit)
+		{
+			Error = Angle + SwingLimit;
+		}
+	}
+
+	void FPBDJointUtilities::GetEllipticalConeAxisErrorLocal(
+		const FRotation3& R0,
+		const FRotation3& R1,
+		const FReal SwingLimitY,
+		const FReal SwingLimitZ,
+		FVec3& AxisLocal,
+		FReal& Error)
+	{
+		if (FMath::IsNearlyEqual(SwingLimitY, SwingLimitZ, 1.e-3f))
+		{ 
+			GetCircularConeAxisErrorLocal(R0, R1, SwingLimitY, AxisLocal, Error);
+			return;
+		}
+
+		AxisLocal = FJointConstants::Swing1Axis();
+		Error = 0.0f;
+
+		FRotation3 R01Twist, R01Swing;
+		FPBDJointUtilities::DecomposeSwingTwistLocal(R0, R1, R01Swing, R01Twist);
+
+		const FVec2 SwingAngles = FVec2(GetSwingAngleY(R01Swing), GetSwingAngleZ(R01Swing));
+		const FVec2 SwingLimits = FVec2(SwingLimitY, SwingLimitZ);
+
+		// Transform onto a circle to see if we are within the ellipse
+		const FVec2 CircleMappedAngles = SwingAngles / SwingLimits;
+		if (CircleMappedAngles.SizeSquared() > 1.0f)
+		{
+			// Map the swing to a position on the elliptical limits
+			const FVec2 ClampedSwingAngles = NearPointOnEllipse(SwingAngles, SwingLimits);
+
+			// Get the ellipse normal
+			const FVec2 ClampedNormal = ClampedSwingAngles / (SwingLimits * SwingLimits);
+
+			// Calculate the axis and error
+			const FVec3 TwistAxis = R01Swing.GetAxisX();
+			const FVec3 SwingRotAxis = FVec3(0.0f, FMath::Tan(ClampedSwingAngles.X / 4.0f), FMath::Tan(ClampedSwingAngles.Y / 4.0f));
+			const FVec3 EllipseNormal = FVec3(0.0f, ClampedNormal.X, ClampedNormal.Y);
+			GetEllipticalAxisError(SwingRotAxis, EllipseNormal, TwistAxis, AxisLocal, Error);
+		}
+	}
 
 	void FPBDJointUtilities::GetLockedSwingAxisAngle(
 		const FRotation3& R0,
diff --git a/Engine/Source/Runtime/Experimental/Chaos/Public/Chaos/PBDJointConstraintUtilities.h b/Engine/Source/Runtime/Experimental/Chaos/Public/Chaos/PBDJointConstraintUtilities.h
--- a/Engine/Source/Runtime/Experimental/Chaos/Public/Chaos/PBDJointConstraintUtilities.h
+++ b/Engine/Source/Runtime/Experimental/Chaos/Public/Chaos/PBDJointConstraintUtilities.h
@@ -66,6 +66,21 @@ namespace Chaos
 			FVec3& AxisLocal,
 			FReal& Angle);
 
+		static void GetCircularConeAxisErrorLocal(
+			const FRotation3& R0,
+			const FRotation3& R1,
+			const FReal SwingLimit,
+			FVec3& AxisLocal,
+			FReal& Error);
+
+		static void GetEllipticalConeAxisErrorLocal(
+			const FRotation3& R0,
+			const FRotation3& R1,
+			const FReal SwingLimitY,
+			const FReal SwingLimitZ,
+			FVec3& AxisLocal,
+			FReal& Error);
+
 		static void GetLockedSwingAxisAngle(
 			const FRotation3& R0,
 			const FRotation3& R1,
diff --git a/Engine/Source/Runtime/Experimental/ChaosCore/Public/Chaos/Core.h b/Engine/Source/Runtime/Experimental/ChaosCore/Public/Chaos/Core.h
--- a/Engine/Source/Runtime/Experimental/ChaosCore/Public/Chaos/Core.h
+++ b/Engine/Source/Runtime/Experimental/ChaosCore/Public/Chaos/Core.h
@@ -19,6 +19,7 @@ namespace Chaos
 	 */
 	using FReal = float;
 
+	using FVec2 = TVector<FReal, 2>;
 	using FVec3 = TVector<FReal, 3>;
 	using FMatrix33 = PMatrix<FReal, 3, 3>;
 	using FRigidTransform3 = TRigidTransform<FReal, 3>;
