GeoemtryCollection :

- Added UnparentTransform to the GeometryCollection API

diff --git a/Engine/Source/Runtime/Experimental/GeometryCollectionCore/Private/GeometryCollection/GeometryCollectionAlgo.cpp b/Engine/Source/Runtime/Experimental/GeometryCollectionCore/Private/GeometryCollection/GeometryCollectionAlgo.cpp
--- a/Engine/Source/Runtime/Experimental/GeometryCollectionCore/Private/GeometryCollection/GeometryCollectionAlgo.cpp
+++ b/Engine/Source/Runtime/Experimental/GeometryCollectionCore/Private/GeometryCollection/GeometryCollectionAlgo.cpp
@@ -270,6 +270,33 @@ namespace GeometryCollectionAlgo
 		ensure(!HasCycle(Parents, SelectedBones));
 	}
 
+	void UnparentTransform(FManagedArrayCollection* Collection, const int32 ChildIndex)
+	{
+		if (Collection)
+		{
+			if (Collection->HasAttribute(FTransformCollection::ChildrenAttribute, FTransformCollection::TransformGroup))
+			{
+				if (Collection->HasAttribute(FTransformCollection::ParentAttribute, FTransformCollection::TransformGroup))
+				{
+					int32 NumTransforms = Collection->NumElements(FTransformCollection::TransformGroup);
+
+					if (0 < ChildIndex && ChildIndex < NumTransforms)
+					{
+						TManagedArray<int32>& Parent = Collection->GetAttribute<int32>(FTransformCollection::ParentAttribute, FTransformCollection::TransformGroup);
+						TManagedArray< TArray<int32> >& Children = Collection->GetAttribute< TArray<int32> >(FTransformCollection::ChildrenAttribute, FTransformCollection::TransformGroup);
+
+						int32 ParentIndex = Parent[ChildIndex];
+						if (0 <= ParentIndex && ParentIndex < NumTransforms)
+						{
+							Children[ParentIndex].Remove(ChildIndex);
+							Parent[ChildIndex] = INDEX_NONE;
+						}
+					}
+				}
+			}
+		}
+	}
+
 	void GlobalMatricesRecursive(const int32 Index, const TManagedArray<int32>& Parents, const TManagedArray<FTransform>& Transform, TArray<bool>& IsTransformComputed, TArray<FMatrix>& OutGlobalTransforms)
 	{
 		if (!IsTransformComputed[Index])
diff --git a/Engine/Source/Runtime/Experimental/GeometryCollectionCore/Private/GeometryCollection/TransformCollection.cpp b/Engine/Source/Runtime/Experimental/GeometryCollectionCore/Private/GeometryCollection/TransformCollection.cpp
--- a/Engine/Source/Runtime/Experimental/GeometryCollectionCore/Private/GeometryCollection/TransformCollection.cpp
+++ b/Engine/Source/Runtime/Experimental/GeometryCollectionCore/Private/GeometryCollection/TransformCollection.cpp
@@ -119,6 +119,12 @@ void FTransformCollection::ParentTransforms(const int32 TransformIndex, const TA
 	GeometryCollectionAlgo::ParentTransforms(this, TransformIndex, SelectedBones);
 }
 
+void FTransformCollection::UnparentTransform(const int32 ChildIndex)
+{
+	GeometryCollectionAlgo::UnparentTransform(this, ChildIndex);
+}
+
+
 
 void FTransformCollection::RelativeTransformation(const int32& Index, const FTransform& LocalOffset)
 {
diff --git a/Engine/Source/Runtime/Experimental/GeometryCollectionCore/Public/GeometryCollection/GeometryCollectionAlgo.h b/Engine/Source/Runtime/Experimental/GeometryCollectionCore/Public/GeometryCollection/GeometryCollectionAlgo.h
--- a/Engine/Source/Runtime/Experimental/GeometryCollectionCore/Public/GeometryCollection/GeometryCollectionAlgo.h
+++ b/Engine/Source/Runtime/Experimental/GeometryCollectionCore/Public/GeometryCollection/GeometryCollectionAlgo.h
@@ -113,6 +113,14 @@ namespace GeometryCollectionAlgo
 	GEOMETRYCOLLECTIONCORE_API 
 	ParentTransforms(FTransformCollection* GeometryCollection, const int32 TransformIndex, const TArray<int32>& SelectedBones);
 
+
+	/*
+	*  Unparent the child index from its parent
+	*/
+	void
+	GEOMETRYCOLLECTIONCORE_API
+	UnparentTransform(FManagedArrayCollection* Collection, const int32 ChildIndex);
+
 	/*
 	*  Find the average position of the transforms.
 	*/
diff --git a/Engine/Source/Runtime/Experimental/GeometryCollectionCore/Public/GeometryCollection/TransformCollection.h b/Engine/Source/Runtime/Experimental/GeometryCollectionCore/Public/GeometryCollection/TransformCollection.h
--- a/Engine/Source/Runtime/Experimental/GeometryCollectionCore/Public/GeometryCollection/TransformCollection.h
+++ b/Engine/Source/Runtime/Experimental/GeometryCollectionCore/Public/GeometryCollection/TransformCollection.h
@@ -74,6 +74,7 @@ public:
 	*/
 	void ParentTransforms(const int32 TransformIndex, const int32 ChildIndex);
 	void ParentTransforms(const int32 TransformIndex, const TArray<int32>& SelectedBones);
+	void UnparentTransform(const int32 ChildIndex);
 
 	/*
 	* RelativeTransformation
