Chaos : Destruction

Buffer physics results for unioned clusteres

diff --git a/Engine/Source/Runtime/Experimental/Chaos/Private/Chaos/PBDRigidClustering.cpp b/Engine/Source/Runtime/Experimental/Chaos/Private/Chaos/PBDRigidClustering.cpp
--- a/Engine/Source/Runtime/Experimental/Chaos/Private/Chaos/PBDRigidClustering.cpp
+++ b/Engine/Source/Runtime/Experimental/Chaos/Private/Chaos/PBDRigidClustering.cpp
@@ -609,15 +609,18 @@ namespace Chaos
 				int32 ClusterGroupID = Group.Key;
 				TArray<TPBDRigidClusteredParticleHandle<T, 3>* > Handles = Group.Value;
 
-				if (!NewClusterGroups.Contains(ClusterGroupID))
-					NewClusterGroups.Add(ClusterGroupID, TArray < TPBDRigidParticleHandle<T, 3>*>());
-
-				TArray<TPBDRigidParticleHandle<T, 3>*> ClusterBodies;
-				for (TPBDRigidClusteredParticleHandle<T, 3>* ActiveCluster : Handles)
+				if (Handles.Num() > 1)
 				{
-					TSet<TPBDRigidParticleHandle<T, 3>*> Children = ReleaseClusterParticles(ActiveCluster, nullptr, true);
-					NewClusterGroups[ClusterGroupID].Append(Children.Array());
-					for (auto& Child : Children) ClusterParents.Add(Child, ActiveCluster);
+					if (!NewClusterGroups.Contains(ClusterGroupID))
+						NewClusterGroups.Add(ClusterGroupID, TArray < TPBDRigidParticleHandle<T, 3>*>());
+
+					TArray<TPBDRigidParticleHandle<T, 3>*> ClusterBodies;
+					for (TPBDRigidClusteredParticleHandle<T, 3>* ActiveCluster : Handles)
+					{
+						TSet<TPBDRigidParticleHandle<T, 3>*> Children = ReleaseClusterParticles(ActiveCluster, nullptr, true);
+						NewClusterGroups[ClusterGroupID].Append(Children.Array());
+						for (auto& Child : Children) ClusterParents.Add(Child, ActiveCluster);
+					}
 				}
 			}
 
@@ -625,7 +628,8 @@ namespace Chaos
 			{
 				int32 ClusterGroupID = Group.Key;
 				TArray< TPBDRigidParticleHandle<T, 3> *> ActiveCluster = Group.Value;
-				TPBDRigidParticleHandle<T, 3>* NewCluster = CreateClusterParticle(-FMath::Abs(ClusterGroupID), MoveTemp(Group.Value));
+				TPBDRigidClusteredParticleHandle<T, 3>* NewCluster = CreateClusterParticle(-FMath::Abs(ClusterGroupID), MoveTemp(Group.Value));
+				NewCluster->SetInternalCluster(true);
 				for (auto& Constituent : ActiveCluster) MEvolution.DoInternalParticleInitilization( ClusterParents[Constituent], NewCluster);
 			}
 			ClusterUnionMap.Empty();
diff --git a/Engine/Source/Runtime/Experimental/ChaosSolvers/Private/PhysicsProxy/GeometryCollectionPhysicsProxy.cpp b/Engine/Source/Runtime/Experimental/ChaosSolvers/Private/PhysicsProxy/GeometryCollectionPhysicsProxy.cpp
--- a/Engine/Source/Runtime/Experimental/ChaosSolvers/Private/PhysicsProxy/GeometryCollectionPhysicsProxy.cpp
+++ b/Engine/Source/Runtime/Experimental/ChaosSolvers/Private/PhysicsProxy/GeometryCollectionPhysicsProxy.cpp
@@ -1424,76 +1424,30 @@ void FGeometryCollectionPhysicsProxy::BufferPhysicsResults()
 	using namespace Chaos;
 	SCOPE_CYCLE_COUNTER(STAT_CacheResultGeomCollection);
 
-	//FGeometryCollectionResults& TargetResults = Results.GetPhysicsDataForWrite();
-	const FGeometryDynamicCollection& PTDynamicCollection = PhysicsThreadCollection;
+	IsObjectDynamic = false;
 	FGeometryCollectionResults& TargetResults = *PhysToGameInterchange.AccessProducerBuffer();
 
-	if (TargetResults.NumTransformGroup() != PTDynamicCollection.NumElements(FGeometryCollection::TransformGroup))
+	int32 NumTransformGroupElements = PhysicsThreadCollection.NumElements(FGeometryCollection::TransformGroup);
+	if (TargetResults.NumTransformGroup() != NumTransformGroupElements)
 	{
-		TargetResults.InitArrays(PTDynamicCollection);
+		TargetResults.InitArrays(PhysicsThreadCollection);
 
 		// Base particle index to calculate index from a global particle index on the game thread
 		TargetResults.BaseIndex = BaseParticleIndex;
 		TargetResults.NumParticlesAdded = NumParticles;
 	}
 
-	{ // tmp scope
-		IsObjectDynamic = false;
-		int32 Idx = 0;
-		for (auto* Handle : SolverParticleHandles)
-		{
-			if (!Handle || Handle->Disabled())
-			{
-				TargetResults.DisabledStates[Idx++] = true;
-				continue;
-			}
-			TargetResults.DisabledStates[Idx++] = false;
-			IsObjectDynamic = true;
-		}
-	} // tmp scope
-
-	// Advertise to game thread
-	TargetResults.IsObjectDynamic = IsObjectDynamic;
-	TargetResults.IsObjectLoading = IsObjectLoading;
 	
 	const FTransform& ActorToWorld = Parameters.WorldTransform;
+	const TManagedArray<int32>& Parent = PhysicsThreadCollection.Parent;
+	const TManagedArray<TSet<int32>>& Children = PhysicsThreadCollection.Children;
 
-	const int32 TransfromGroupSize = PTDynamicCollection.NumElements(FGeometryCollection::TransformGroup);
-	if (EndFrameUnparentingBuffer.Num() < TransfromGroupSize)
-	{
-		// It's not worth shrinking the EndFrameUnparentingBuffer array, 
-		// at least until the solver supports deleting bodies.
-		EndFrameUnparentingBuffer.Init(INDEX_NONE, TransfromGroupSize);
-	}
-	const TManagedArray<int32>& Parent = PTDynamicCollection.Parent;
-	const TManagedArray<TSet<int32>>& Children = PTDynamicCollection.Children;
+	{ SCOPE_CYCLE_COUNTER(STAT_CalcParticleToWorld);
 
-	{ // tmp scope
-		SCOPE_CYCLE_COUNTER(STAT_CalcParticleToWorld);
-		// Fill particle to world transforms
-//		TargetResults.ParticleToWorldTransforms.SetNum(NumParticles);
-		//for (int32 TransformIndex = 0; TransformIndex < NumParticles; ++TransformIndex)
-		//{
-		//	//only update roots and first children
-		//	if (TargetResults.Parent[TransformIndex] == INDEX_NONE)
-		//	{
-		//		if (auto* Handle = SolverParticleHandles[TransformIndex])
-		//		{
-		//			const FTransform Xf(Handle->R(), Handle->X());
-		//			//const int32 ParticleIndex = BaseParticleIndex + TransformIndex;
-		//			TargetResults.ParticleToWorldTransforms[TransformIndex] = Xf;
-		//		}
-		//	}
-		//}
-
-		//const TArrayCollectionArray<Chaos::ClusterId>& ClusterID = GetSolver()->GetRigidClustering().GetClusterIdsArray();
-		//const TArrayCollectionArray<Chaos::TRigidTransform<float, 3>>& ClusterChildToParentMap = GetSolver()->GetRigidClustering().GetChildToParentMap();
-		//const TArrayCollectionArray<FMultiChildProxyId>& MultiChildProxyIdArray = GetSolver()->GetRigidClustering().GetMultiChildProxyIdArray();
-		//const TArrayCollectionArray<TUniquePtr<TMultiChildProxyData<float, 3>>>& MultiChildProxyDataArray = GetSolver()->GetRigidClustering().GetMultiChildProxyDataArray();
-
-		for (int32 TransformIndex = 0; TransformIndex < TransfromGroupSize; ++TransformIndex)
+		for (int32 TransformGroupIndex = 0; TransformGroupIndex < NumTransformGroupElements; ++TransformGroupIndex)
 		{
-			Chaos::TPBDRigidClusteredParticleHandle<float, 3>* Handle = SolverParticleHandles[TransformIndex];
+			TargetResults.DisabledStates[TransformGroupIndex] = true;
+			Chaos::TPBDRigidClusteredParticleHandle<float, 3>* Handle = SolverParticleHandles[TransformGroupIndex];
 			if (!Handle)
 			{
 				continue;
@@ -1506,18 +1460,18 @@ void FGeometryCollectionPhysicsProxy::BufferPhysicsResults()
 				switch (ObjectState)
 				{
 				case Chaos::EObjectStateType::Kinematic:
-					TargetResults.DynamicState[TransformIndex] = (int)EObjectStateTypeEnum::Chaos_Object_Kinematic;
+					TargetResults.DynamicState[TransformGroupIndex] = (int)EObjectStateTypeEnum::Chaos_Object_Kinematic;
 					break;
 				case Chaos::EObjectStateType::Static:
-					TargetResults.DynamicState[TransformIndex] = (int)EObjectStateTypeEnum::Chaos_Object_Static;
+					TargetResults.DynamicState[TransformGroupIndex] = (int)EObjectStateTypeEnum::Chaos_Object_Static;
 					break;
 				case Chaos::EObjectStateType::Sleeping:
-					TargetResults.DynamicState[TransformIndex] = (int)EObjectStateTypeEnum::Chaos_Object_Sleeping;
+					TargetResults.DynamicState[TransformGroupIndex] = (int)EObjectStateTypeEnum::Chaos_Object_Sleeping;
 					break;
 				case Chaos::EObjectStateType::Dynamic:
 				case Chaos::EObjectStateType::Uninitialized:
 				default:
-					TargetResults.DynamicState[TransformIndex] = (int)EObjectStateTypeEnum::Chaos_Object_Dynamic;
+					TargetResults.DynamicState[TransformGroupIndex] = (int)EObjectStateTypeEnum::Chaos_Object_Dynamic;
 					break;
 				}
 			}
@@ -1537,33 +1491,32 @@ void FGeometryCollectionPhysicsProxy::BufferPhysicsResults()
 				// single transform in the GeometryCollection, and all children are stored in the local space
 				// of the parent cluster.
 	
-				FTransform& ParticleToWorld = TargetResults.ParticleToWorldTransforms[TransformIndex];
+				FTransform& ParticleToWorld = TargetResults.ParticleToWorldTransforms[TransformGroupIndex];
 				ParticleToWorld = FTransform(Handle->R(), Handle->X());
 
-				TargetResults.Transforms[TransformIndex] = PTDynamicCollection.MassToLocal[TransformIndex].GetRelativeTransformReverse(ParticleToWorld).GetRelativeTransform(ActorToWorld);
-				TargetResults.Transforms[TransformIndex].NormalizeRotation();
+				TargetResults.Transforms[TransformGroupIndex] = PhysicsThreadCollection.MassToLocal[TransformGroupIndex].GetRelativeTransformReverse(ParticleToWorld).GetRelativeTransform(ActorToWorld);
+				TargetResults.Transforms[TransformGroupIndex].NormalizeRotation();
+
+				// Indicate that this object needs to be updated and the proxy is active.
+				TargetResults.DisabledStates[TransformGroupIndex] = false;
+				IsObjectDynamic = true;
 
 				// If the parent of this NON DISABLED body is set to anything other than INDEX_NONE,
 				// then it was just unparented, likely either by rigid clustering or by fields.  We
 				// need to force all such enabled rigid bodies out of the transform hierarchy.
-				if (Parent[TransformIndex] != INDEX_NONE)
+				TargetResults.Parent[TransformGroupIndex] = INDEX_NONE;
+				if (PhysicsThreadCollection.Parent[TransformGroupIndex] != INDEX_NONE)
 				{
-					const int32 ParentIndex = Parent[TransformIndex];
-					// Children in the hierarchy are stored in a TSet, which is not thread safe.  
-					// So we retain indices to remove afterwards.
-					EndFrameUnparentingBuffer[TransformIndex] = ParentIndex;
-					
-					// Once a child is released there is no longer a useful parent relationship as the child is
-					// now entirely responsible for its own simulation. Sever the link here so we don't attempt
-					// to pull in a parent transform later in the pipeline.
-					PhysicsThreadCollection.Parent[TransformIndex] = INDEX_NONE;
+					//GeometryCollectionAlgo::UnparentTransform(&PhysicsThreadCollection,TransformGroupIndex);
+					PhysicsThreadCollection.Children[PhysicsThreadCollection.Parent[TransformGroupIndex]].Remove(TransformGroupIndex);
+					PhysicsThreadCollection.Parent[TransformGroupIndex] = INDEX_NONE;
 				}
 
 				// When a leaf node rigid body is removed from a cluster, the rigid
 				// body will become active and needs its clusterID updated.  This just
 				// syncs the clusterID all the time.
 				TPBDRigidParticleHandle<float, 3>* ClusterParentId = Handle->ClusterIds().Id;
-				SolverClusterID[TransformIndex] = ClusterParentId;
+				SolverClusterID[TransformGroupIndex] = ClusterParentId;
 			}
 			else // Handle->Disabled()
 			{
@@ -1571,70 +1524,67 @@ void FGeometryCollectionPhysicsProxy::BufferPhysicsResults()
 				// parent body is not tracked within the geometry collection. So we need to
 				// pull the rigid bodies out of the transform hierarchy, and just drive
 				// the positions directly from the solvers cluster particle. 
-				TPBDRigidParticleHandle<float, 3>* ClusterParentId = Handle->ClusterIds().Id;
-				if(SolverClusterID[TransformIndex] != ClusterParentId)
+				if (TPBDRigidParticleHandle<float, 3>* ClusterParentBase = Handle->ClusterIds().Id)
 				{
-					// Force all driven rigid bodies out of the transform hierarchy
-					if (Parent[TransformIndex] != INDEX_NONE)
+					if(Chaos::TPBDRigidClusteredParticleHandle<float, 3>* ClusterParent = ClusterParentBase->CastToClustered() )
 					{
-						const int32 ParentIndex = Parent[TransformIndex];
-						// Children in the hierarchy are stored in a TSet, which is not thread safe.  So we retain
-						// indices to remove afterwards.
-						EndFrameUnparentingBuffer[TransformIndex] = ParentIndex;
+
+						// syncronize parents if it has changed.
+						if (SolverClusterID[TransformGroupIndex] != ClusterParent)
+						{
+							// Force all driven rigid bodies out of the transform hierarchy
+							if (Parent[TransformGroupIndex] != INDEX_NONE)
+							{
+								// If the parent of this NON DISABLED body is set to anything other than INDEX_NONE,
+								// then it was just unparented, likely either by rigid clustering or by fields.  We
+								// need to force all such enabled rigid bodies out of the transform hierarchy.
+								TargetResults.Parent[TransformGroupIndex] = INDEX_NONE;
+								
+								//GeometryCollectionAlgo::UnparentTransform(&PhysicsThreadCollection, ChildIndex);
+								PhysicsThreadCollection.Children[PhysicsThreadCollection.Parent[TransformGroupIndex]].Remove(TransformGroupIndex);
+								PhysicsThreadCollection.Parent[TransformGroupIndex] = INDEX_NONE;
+
+								// Indicate that this object needs to be updated and the proxy is active.
+								TargetResults.DisabledStates[TransformGroupIndex] = false;
+								IsObjectDynamic = true;
+							}
+							SolverClusterID[TransformGroupIndex] = Handle->ClusterIds().Id;
+						}
+
+						if (ClusterParent->InternalCluster() )
+						{
+							Chaos::TPBDRigidClusteredParticleHandle<float, 3>* ProxyElementHandle = SolverParticleHandles[TransformGroupIndex];
+							const FTransform ClusterToWorld = FTransform(ClusterParent->R(), ClusterParent->X());
+							const FTransform ChildToCluster = FTransform(ProxyElementHandle->ChildToParent());
+							const FTransform ClusterChildToWorld = ChildToCluster * ClusterToWorld;
+							FTransform MassToLocal = FTransform::Identity;// @todo(chaos): what should the mass to local for the internal cluster be?
+
+							// GeomToActor = ActorToWorld.Inv() * ClusterChildToWorld * MassToLocal.Inv();
+							TargetResults.Transforms[TransformGroupIndex] = MassToLocal.GetRelativeTransformReverse(ClusterChildToWorld).GetRelativeTransform(ActorToWorld);
+							TargetResults.Transforms[TransformGroupIndex].NormalizeRotation();
+
+							// Indicate that this object needs to be updated and the proxy is active.
+							TargetResults.DisabledStates[TransformGroupIndex] = false;
+							IsObjectDynamic = true;
+						}
 					}
-					SolverClusterID[TransformIndex] = Handle->ClusterIds().Id;
 				}
 			}// end if
+
+			PhysicsThreadCollection.Active[TransformGroupIndex] = !TargetResults.DisabledStates[TransformGroupIndex];
 		} // end for
-	} // tmp scope
+	} // STAT_CalcParticleToWorld scope
 
-	for (int32 TransformGroupIndex = 0; TransformGroupIndex < TransfromGroupSize; TransformGroupIndex++)
-	{
-		const int32 ParentIndex = EndFrameUnparentingBuffer[TransformGroupIndex];
-		if (ParentIndex >= 0)
-		{
-			// We reuse EndFrameUnparentingBuffer potentially without reinitialization, 
-			// so reset this index to INDEX_NONE before it gets paged out.
-			EndFrameUnparentingBuffer[TransformGroupIndex] = INDEX_NONE;
 
-			// todo : this should be on the output buffer -
-			//Children[ParentIndex].Remove(TransformGroupIndex);
-			//Parent[TransformGroupIndex] = INDEX_NONE;
-		}
-		TargetResults.Parent[TransformGroupIndex] = Parent[TransformGroupIndex];
-	}
 
 	// If object is dynamic, compute global matrices	
-	if (IsObjectDynamic || TargetResults.GlobalTransforms.Num() == 0)
+ 	if (IsObjectDynamic || TargetResults.GlobalTransforms.Num() == 0)
 	{
 		SCOPE_CYCLE_COUNTER(STAT_CalcGlobalGCMatrices);
 		check(TargetResults.Transforms.Num() == TargetResults.Parent.Num());
 		GeometryCollectionAlgo::GlobalMatrices(
 			TargetResults.Transforms, TargetResults.Parent, TargetResults.GlobalTransforms);
 
-		// update transforms of things parented to internal clusters who would be missed in the above call to GlobalMatrices
-		// #todo: this can be much faster by caching the parent transforms somehow
-		// #todo: it only supports a single level of internal cluster parenting.
-		for (int32 TransformIndex = 0; TransformIndex < TransfromGroupSize; ++TransformIndex)
-		{
-			Chaos::TPBDRigidClusteredParticleHandle<float, 3>* Handle = SolverParticleHandles[TransformIndex];
-
-			// If the particle is disabled and has a parent that is an internal cluster, update transform
-			if (Handle && Handle->Disabled() && SolverClusterID[TransformIndex] && SolverClusterID[TransformIndex]->CastToClustered()->InternalCluster())
-			{
-				FClusterHandle* ParentClusterHandle = SolverClusterID[TransformIndex]->CastToClustered();
-
-				FTransform ParticleToWorld = FTransform(ParentClusterHandle->R(), ParentClusterHandle->X());
-				
-				// #todo: what should the mass to local for the internal cluster be?
-				FTransform MassToLocal = FTransform::Identity;
-				
-				// compute global transform for particle
-				TargetResults.GlobalTransforms[TransformIndex] = Handle->ChildToParent().ToMatrixWithScale();
-				TargetResults.GlobalTransforms[TransformIndex] *= (MassToLocal.GetRelativeTransformReverse(ParticleToWorld).GetRelativeTransform(ActorToWorld)).ToMatrixWithScale();
-			}
-		}
-
 		// Compute world bounds.  This is a loose bounds based on the circumscribed box 
 		// of a bounding sphere for the geometry.		
 		SCOPE_CYCLE_COUNTER(STAT_CalcGlobalGCBounds);
@@ -1647,8 +1597,9 @@ void FGeometryCollectionPhysicsProxy::BufferPhysicsResults()
 		TargetResults.WorldBounds = FBoxSphereBounds(BoundingBox);
 	}
 
-	// We pulled state out of particles.  Push results into PTDynamicCollection, 
-	// so that it has an updated view.
+	// Advertise to game thread
+	TargetResults.IsObjectDynamic = IsObjectDynamic;
+	TargetResults.IsObjectLoading = IsObjectLoading;
 }
 
 void FGeometryCollectionPhysicsProxy::FlipBuffer()
@@ -1683,31 +1634,30 @@ void FGeometryCollectionPhysicsProxy::PullFromPhysicsState()
 	}
 
 	// Remove const-ness as we're going to do the ExchangeArrays() thing.
-	FGeometryCollectionResults& TR = *const_cast<FGeometryCollectionResults*>(TargetResultPtr);
+	FGeometryCollectionResults& TargetResults = *const_cast<FGeometryCollectionResults*>(TargetResultPtr);
 
 	FGeometryDynamicCollection& DynamicCollection = GameThreadCollection;
 
 	// We should never be changing the number of entries, this would break other 
 	// attributes in the transform group.
 	const int32 NumTransforms = DynamicCollection.Transform.Num();
-	if (ensure(NumTransforms == TR.Transforms.Num()))	
+	if (ensure(NumTransforms == TargetResults.Transforms.Num()))
 	{
-		for(int32 TmIndex = 0; TmIndex < NumTransforms; ++TmIndex)
+		for (int32 TransformGroupIndex = 0; TransformGroupIndex < NumTransforms; ++TransformGroupIndex)
 		{
-			if(!TR.DisabledStates[TmIndex])
+			if (!TargetResults.DisabledStates[TransformGroupIndex])
 			{
-				const FTransform& LocalTransform = TR.Transforms[TmIndex];
-				const FTransform& ParticleToWorld = TR.ParticleToWorldTransforms[TmIndex];
+				DynamicCollection.Parent[TransformGroupIndex] = TargetResults.Parent[TransformGroupIndex];
+				const FTransform& LocalTransform = TargetResults.Transforms[TransformGroupIndex];
+				const FTransform& ParticleToWorld = TargetResults.ParticleToWorldTransforms[TransformGroupIndex];
 
-				DynamicCollection.Transform[TmIndex] = LocalTransform;
-				GTParticles[TmIndex]->SetX(ParticleToWorld.GetTranslation());
-				GTParticles[TmIndex]->SetR(ParticleToWorld.GetRotation());
-				DynamicCollection.Parent[TmIndex] = FGeometryCollection::Invalid;
+				DynamicCollection.Transform[TransformGroupIndex] = LocalTransform;
+				GTParticles[TransformGroupIndex]->SetX(ParticleToWorld.GetTranslation());
+				GTParticles[TransformGroupIndex]->SetR(ParticleToWorld.GetRotation());
 			}
 
-			DynamicCollection.DynamicState[TmIndex] = TR.DynamicState[TmIndex];
-
-			DynamicCollection.Active[TmIndex] = !TR.DisabledStates[TmIndex];
+			DynamicCollection.DynamicState[TransformGroupIndex] = TargetResults.DynamicState[TransformGroupIndex];
+			DynamicCollection.Active[TransformGroupIndex] = !TargetResults.DisabledStates[TransformGroupIndex];
 		}
 
 		//question: why do we need this? Sleeping objects will always have to update GPU
@@ -1715,90 +1665,11 @@ void FGeometryCollectionPhysicsProxy::PullFromPhysicsState()
 
 		if (CacheSyncFunc)
 		{
-			CacheSyncFunc(TR);
+			CacheSyncFunc(TargetResults);
 		}
 	}
 }
 
-void FGeometryCollectionPhysicsProxy::UpdateGeometryCollection( const FGeometryCollectionResults& GCResults, FGeometryDynamicCollection* Collection)
-{
-	check(IsInGameThread());
-	if (!Collection)
-	{
-		Collection = &GameThreadCollection;
-	}
-
-	TManagedArray<FTransform>& Transform = Collection->Transform;
-	TManagedArray<int32>& Parent = Collection->Parent;
-	TManagedArray<TSet<int32>>& Children = Collection->Children;
-	TManagedArray<int32>& SimulationType = Collection->SimulationType;
-	TManagedArray<int32>& DynamicState = Collection->GetAttribute<int32>("DynamicState", FGeometryCollection::TransformGroup);
-
-	// Particles X and R are aligned with center of mass and inertia principal axes.
-	// Renderer doesn't know about this and simply does ActorToWorld * GeomToActor * LocalSpaceVerts
-	// In proper math multiplication order:
-	//		ParticleToWorld = ActorToWorld * GeomToActor * MassToLocal
-	//		GeomToWorld = ActorToWorld * GeomToActor
-	//		=> GeomToWorld = ParticleToWorld * MassToLocal.Inv()
-	//		=> GeomToActor = ActorToWorld.Inv() * ParticleToWorld * MassToLocal.Inv()
-
-	const FTransform& ActorToWorld = Parameters.WorldTransform;
-	const int32 TransformSize = Collection->NumElements(FGeometryCollection::TransformGroup);
-	for (int32 TransformGroupIndex = 0; TransformGroupIndex < TransformSize; ++TransformGroupIndex)
-	{
-		if (!Collection->SimulatableParticles[TransformGroupIndex])
-			continue;
-
-		auto* Handle = SolverParticleHandles[TransformGroupIndex];
-
-		// dynamic state is also updated by the solver during field interaction. 
-		if(!Handle->Sleeping())
-		{
-			DynamicState[TransformGroupIndex] = GCResults.DynamicState[TransformGroupIndex];
-		}
-
-		// Update the transform and parent hierarchy of the active rigid bodies. Active bodies can be either
-		// rigid geometry defined from the leaf nodes of the collection, or cluster bodies that drive an entire
-		// branch of the hierarchy within the GeometryCollection.
-		// - Active bodies are directly driven from the global position of the corresponding
-		//   rigid bodies within the solver ( cases where RigidBodyID[TransformGroupIndex] is not disabled ). 
-		// - Deactivated bodies are driven from the transforms of their active parents. However the solver can
-		//   take ownership of the parents during the simulation, so it might be necessary to force deactivated
-		//   bodies out of the collections hierarchy during the simulation.  
-
-		// Static analysis build complains that the else clause does the same thing as the if/then,
-		// so remove the conditional for now...
-#if 0
-		if(!GCResults.DisabledStates[TransformGroupIndex])
-		{
-#endif
-	// Update the transform of the active body. The active body can be either a single rigid
-	// or a collection of rigidly attached geometries (Clustering). The cluster is represented as a
-	// single transform in the GeometryCollection, and all children are stored in the local space
-	// of the parent cluster.
-	// ... When setting cluster transforms it is expected that the MassToLocal is identity.
-	//     Cluster initialization will set the vertices in the MassSpace of the rigid body.
-	// ... When setting individual rigid bodies that are not clustered, the MassToLocal will be 
-	//     non-Identity, and will reflect the difference between the geometric center of the geometry
-	//     and that corresponding rigid bodies center of mass. 
-
-			Transform[TransformGroupIndex] = GCResults.Transforms[TransformGroupIndex];
-
-			// When a leaf node rigid body is removed from a cluster the rigid
-			// body will become active and needs its clusterID updated. This just
-			// syncs the clusterID all the time. 
-//			CollectionClusterID[TransformGroupIndex] = ClusterID[RigidBodyIndex].Id;
-#if 0
-		}
-		else //if (Particles.Disabled(RigidBodyIndex))
-		{
-			Transform[TransformGroupIndex] = GCResults.Transforms[TransformGroupIndex];
-		}
-#endif
-	}
-}
-
-
 //==============================================================================
 // STATIC SETUP FUNCTIONS
 //==============================================================================
diff --git a/Engine/Source/Runtime/Experimental/ChaosSolvers/Public/PhysicsProxy/GeometryCollectionPhysicsProxy.h b/Engine/Source/Runtime/Experimental/ChaosSolvers/Public/PhysicsProxy/GeometryCollectionPhysicsProxy.h
--- a/Engine/Source/Runtime/Experimental/ChaosSolvers/Public/PhysicsProxy/GeometryCollectionPhysicsProxy.h
+++ b/Engine/Source/Runtime/Experimental/ChaosSolvers/Public/PhysicsProxy/GeometryCollectionPhysicsProxy.h
@@ -198,14 +198,6 @@ protected:
 		const TArray<int32>& ChildTransformGroupIndices,
 		const Chaos::FClusterCreationParameters<float> & Parameters);
 
-	/** 
-	 * Copy \p Results into \p TargetCollection, or \c Parameters.DynamicCollection 
-	 * if \c nullptr.
-	 */
-	void UpdateGeometryCollection(
-		const FGeometryCollectionResults& Results,
-		FGeometryDynamicCollection* TargetCollection=nullptr);
-
 	/**
 	 */
 	void GetRelevantHandles(
@@ -253,7 +245,6 @@ private:
 	FCollisionFilterData SimFilter;
 	FCollisionFilterData QueryFilter;
 
-	TArray<int32> EndFrameUnparentingBuffer;
 	// This is a subset of the geometry group that are used in the transform hierarchy to represent geometry
 	TArray<FBox> ValidGeometryBoundingBoxes;
 	TArray<int32> ValidGeometryTransformIndices;
