dynamic state support for clustered particles

diff --git a/Engine/Source/Runtime/Experimental/ChaosSolvers/Private/PhysicsProxy/FieldSystemPhysicsProxy.cpp b/Engine/Source/Runtime/Experimental/ChaosSolvers/Private/PhysicsProxy/FieldSystemPhysicsProxy.cpp
--- a/Engine/Source/Runtime/Experimental/ChaosSolvers/Private/PhysicsProxy/FieldSystemPhysicsProxy.cpp
+++ b/Engine/Source/Runtime/Experimental/ChaosSolvers/Private/PhysicsProxy/FieldSystemPhysicsProxy.cpp
@@ -170,7 +170,8 @@ void FFieldSystemPhysicsProxy::FieldParameterUpdateCallback(
 						// not be dynamic.  Har.
 						if(Chaos::TPBDRigidParticleHandle<float, 3>* RigidHandle = Handle->CastToRigidParticle())
 						{
-							const bool bIsGC = (Handle->GetParticleType() == Chaos::EParticleType::GeometryCollection);
+							const bool bIsGC = (Handle->GetParticleType() == Chaos::EParticleType::GeometryCollection) || 
+								(Handle->GetParticleType() == Chaos::EParticleType::Clustered && !Handle->CastToClustered()->InternalCluster());
 							const EObjectStateType HandleState = RigidHandle->ObjectState();
 
 							// Non-Geometry Collection rigid bodies are more restricted as to how
