Chaos: Disabled experimental contact solve uses the accumulated impulse (instead of the delta impulse) when clamping
Only works for cases where there are more than one contact point per iteration that persists
- Much less jitter with 4 iterations
- More iterations => less jitter now
- Bodies on inclined planes will now start sliding at the correct angle (given a friction coefficient)

diff --git a/Engine/Source/Runtime/Experimental/Chaos/Private/Chaos/PBDCollisionConstraints.cpp b/Engine/Source/Runtime/Experimental/Chaos/Private/Chaos/PBDCollisionConstraints.cpp
--- a/Engine/Source/Runtime/Experimental/Chaos/Private/Chaos/PBDCollisionConstraints.cpp
+++ b/Engine/Source/Runtime/Experimental/Chaos/Private/Chaos/PBDCollisionConstraints.cpp
@@ -59,6 +59,9 @@ namespace Chaos
 	float DefaultCollisionRestitution = 0;
 	FAutoConsoleVariableRef CVarDefaultCollisionRestitution(TEXT("p.DefaultCollisionRestitution"), DefaultCollisionRestitution, TEXT("Collision restitution default value if no materials are found."));
 
+	int32 Chaos_Collision_UseAccumulatedImpulseClipSolve = 0; // Experimental: This requires multiple contact points per iteration per pair, and making sure the contact points don't move too much in body space
+	FAutoConsoleVariableRef CVarChaosCollisionOriginalSolve(TEXT("p.Chaos.Collision.UseAccumulatedImpulseClipSolve"), Chaos_Collision_UseAccumulatedImpulseClipSolve, TEXT("Use experimental Accumulated impulse clipped contact solve"));
+
 #if INTEL_ISPC
 	bool bChaos_Collision_ISPC_Enabled = false;
 	FAutoConsoleVariableRef CVarChaosCollisionISPCEnabled(TEXT("p.Chaos.Collision.ISPC"), bChaos_Collision_ISPC_Enabled, TEXT("Whether to use ISPC optimizations in the Collision Solver"));
@@ -433,6 +436,13 @@ namespace Chaos
 	{
 		SCOPE_CYCLE_COUNTER(STAT_Collisions_UpdatePointConstraints);
 
+		// Make sure the cull distance is enough if we switched to Accumulated Impulse clipping		
+		const int MinCullDistanceForImpulseClipping = 5;
+		if (Chaos_Collision_UseAccumulatedImpulseClipSolve && MCullDistance < MinCullDistanceForImpulseClipping)
+		{
+			MCullDistance = MinCullDistanceForImpulseClipping;
+		}
+
 		// @todo(chaos): parallelism needs to be optional
 
 		//PhysicsParallelFor(Handles.Num(), [&](int32 ConstraintHandleIndex)
diff --git a/Engine/Source/Runtime/Experimental/Chaos/Private/Chaos/PBDCollisionConstraintsContact.cpp b/Engine/Source/Runtime/Experimental/Chaos/Private/Chaos/PBDCollisionConstraintsContact.cpp
--- a/Engine/Source/Runtime/Experimental/Chaos/Private/Chaos/PBDCollisionConstraintsContact.cpp
+++ b/Engine/Source/Runtime/Experimental/Chaos/Private/Chaos/PBDCollisionConstraintsContact.cpp
@@ -19,6 +19,8 @@ namespace Chaos
 	extern bool bChaos_Collision_ISPC_Enabled;
 #endif
 
+	extern int32 Chaos_Collision_UseAccumulatedImpulseClipSolve;
+
 	namespace Collisions
 	{
 		int32 Chaos_Collision_EnergyClampEnabled = 1;
@@ -59,6 +61,201 @@ namespace Chaos
 			UpdateConstraintFromManifold(Constraint, Transform0, Transform1, CullDistance);
 		}
 
+		void ApplyAngularFriction(
+			  FVec3& Impulse // In and Out
+			, FVec3& AngularImpulse // In and Out
+			, FReal AngularFriction
+			, const FVec3& RelativeAngularVelocity
+			, const FVec3& ContactVelocityChange
+			, const FVec3& ContactNormal
+			, const FVec3& VectorToPoint1
+			, const FVec3& VectorToPoint2
+			, const TPBDRigidParticleHandle<FReal, 3>* PBDRigid0
+			, const TPBDRigidParticleHandle<FReal, 3>* PBDRigid1
+			, bool bIsRigidDynamic0
+			, bool bIsRigidDynamic1
+			, const FRotation3& Q0
+			, const FRotation3& Q1
+			, const FMatrix33& WorldSpaceInvI1
+			, const FMatrix33& WorldSpaceInvI2
+			)
+		{
+			FReal AngularNormal = FVec3::DotProduct(RelativeAngularVelocity, ContactNormal);
+			FVec3 AngularTangent = RelativeAngularVelocity - AngularNormal * ContactNormal;
+			FReal NormalVelocityChange = FVec3::DotProduct(ContactVelocityChange, ContactNormal);
+			FVec3 FinalAngularVelocity = FMath::Sign(AngularNormal) * FMath::Max((FReal)0, FMath::Abs(AngularNormal) - AngularFriction * NormalVelocityChange) * ContactNormal + FMath::Max((FReal)0, AngularTangent.Size() - AngularFriction * NormalVelocityChange) * AngularTangent.GetSafeNormal();
+			FVec3 Delta = FinalAngularVelocity - RelativeAngularVelocity;
+			if (!bIsRigidDynamic0 && bIsRigidDynamic1)
+			{
+				FMatrix33 WorldSpaceI2 = (Q1 * FMatrix::Identity) * PBDRigid1->I() * (Q1 * FMatrix::Identity).GetTransposed();
+				FVec3 ImpulseDelta = PBDRigid1->M() * FVec3::CrossProduct(VectorToPoint2, Delta);
+				Impulse += ImpulseDelta;
+				AngularImpulse += WorldSpaceI2 * Delta - FVec3::CrossProduct(VectorToPoint2, ImpulseDelta);
+			}
+			else if (bIsRigidDynamic0 && !bIsRigidDynamic1)
+			{
+				FMatrix33 WorldSpaceI1 = (Q0 * FMatrix::Identity) * PBDRigid0->I() * (Q0 * FMatrix::Identity).GetTransposed();
+				FVec3 ImpulseDelta = PBDRigid0->M() * FVec3::CrossProduct(VectorToPoint1, Delta);
+				Impulse += ImpulseDelta;
+				AngularImpulse += WorldSpaceI1 * Delta - FVec3::CrossProduct(VectorToPoint1, ImpulseDelta);
+			}
+			else if (bIsRigidDynamic0 && bIsRigidDynamic1)
+			{
+				FMatrix33 Cross1(0, VectorToPoint1.Z, -VectorToPoint1.Y, -VectorToPoint1.Z, 0, VectorToPoint1.X, VectorToPoint1.Y, -VectorToPoint1.X, 0);
+				FMatrix33 Cross2(0, VectorToPoint2.Z, -VectorToPoint2.Y, -VectorToPoint2.Z, 0, VectorToPoint2.X, VectorToPoint2.Y, -VectorToPoint2.X, 0);
+				FMatrix33 CrossI1 = Cross1 * WorldSpaceInvI1;
+				FMatrix33 CrossI2 = Cross2 * WorldSpaceInvI2;
+				FMatrix33 Diag1 = CrossI1 * Cross1.GetTransposed() + CrossI2 * Cross2.GetTransposed();
+				Diag1.M[0][0] += PBDRigid0->InvM() + PBDRigid1->InvM();
+				Diag1.M[1][1] += PBDRigid0->InvM() + PBDRigid1->InvM();
+				Diag1.M[2][2] += PBDRigid0->InvM() + PBDRigid1->InvM();
+				FMatrix33 OffDiag1 = (CrossI1 + CrossI2) * -1;
+				FMatrix33 Diag2 = (WorldSpaceInvI1 + WorldSpaceInvI2).Inverse();
+				FMatrix33 OffDiag1Diag2 = OffDiag1 * Diag2;
+				FVec3 ImpulseDelta = FMatrix33((Diag1 - OffDiag1Diag2 * OffDiag1.GetTransposed()).Inverse()) * ((OffDiag1Diag2 * -1) * Delta);
+				Impulse += ImpulseDelta;
+				AngularImpulse += Diag2 * (Delta - FMatrix33(OffDiag1.GetTransposed()) * ImpulseDelta);
+			}
+		}
+		
+		FVec3 ApplyContactAccImpulse(FCollisionContact& Contact,
+			TGenericParticleHandle<FReal, 3> Particle0,
+			TGenericParticleHandle<FReal, 3> Particle1,
+			const FContactIterationParameters& IterationParameters,
+			const FContactParticleParameters& ParticleParameters,
+			FVec3& AccumulatedImpulse)
+		{
+
+			TPBDRigidParticleHandle<FReal, 3>* PBDRigid0 = Particle0->CastToRigidParticle();
+			TPBDRigidParticleHandle<FReal, 3>* PBDRigid1 = Particle1->CastToRigidParticle();
+
+			const bool bIsRigidDynamic0 = PBDRigid0 && PBDRigid0->ObjectState() == EObjectStateType::Dynamic;
+			const bool bIsRigidDynamic1 = PBDRigid1 && PBDRigid1->ObjectState() == EObjectStateType::Dynamic;
+
+			const FVec3 ZeroVector = FVec3(0);
+			FVec3 P0 = FParticleUtilities::GetCoMWorldPosition(Particle0);
+			FVec3 P1 = FParticleUtilities::GetCoMWorldPosition(Particle1);
+			FRotation3 Q0 = FParticleUtilities::GetCoMWorldRotation(Particle0);
+			FRotation3 Q1 = FParticleUtilities::GetCoMWorldRotation(Particle1);
+
+			const FVec3 VectorToPoint1 = Contact.Location - P0;
+			const FVec3 VectorToPoint2 = Contact.Location - P1;
+			const FVec3 Body1Velocity = FParticleUtilities::GetVelocityAtCoMRelativePosition(Particle0, VectorToPoint1);
+			const FVec3 Body2Velocity = FParticleUtilities::GetVelocityAtCoMRelativePosition(Particle1, VectorToPoint2);
+			const FVec3 RelativeVelocity = Body1Velocity - Body2Velocity;  // Do not early out on negative normal velocity since there can still be an accumulated impulse
+			*IterationParameters.NeedsAnotherIteration = true;
+
+			const FMatrix33 WorldSpaceInvI1 = bIsRigidDynamic0 ? Utilities::ComputeWorldSpaceInertia(Q0, PBDRigid0->InvI()) * Contact.InvInertiaScale0 : FMatrix33(0);
+			const FMatrix33 WorldSpaceInvI2 = bIsRigidDynamic1 ? Utilities::ComputeWorldSpaceInertia(Q1, PBDRigid1->InvI()) * Contact.InvInertiaScale1 : FMatrix33(0);
+			const FMatrix33 Factor =
+				(bIsRigidDynamic0 ? ComputeFactorMatrix3(VectorToPoint1, WorldSpaceInvI1, PBDRigid0->InvM()) : FMatrix33(0)) +
+				(bIsRigidDynamic1 ? ComputeFactorMatrix3(VectorToPoint2, WorldSpaceInvI2, PBDRigid1->InvM()) : FMatrix33(0));
+			FVec3 Impulse;
+			FVec3 AngularImpulse(0);
+
+			// Resting contact if very close to the surface
+			const bool bApplyRestitution = (RelativeVelocity.Size() > (2 * 980 * IterationParameters.Dt));
+			const FReal Restitution = (bApplyRestitution) ? Contact.Restitution : (FReal)0;
+			const FReal Friction = Contact.Friction;
+			const FReal AngularFriction = Contact.AngularFriction;
+
+			const FVec3 VelocityTarget = (-Restitution*FVec3::DotProduct(RelativeVelocity, Contact.Normal)) * Contact.Normal;
+			const FVec3 VelocityChange = VelocityTarget - RelativeVelocity;
+			const FMatrix33 FactorInverse = Factor.Inverse();
+			Impulse = FactorInverse * VelocityChange;  // Delta Impulse = (J.M^(-1).J^(T))^(-1).(ContactVelocityError)
+
+			// clip the impulse so that the accumulated impulse is not in the wrong direction and in the friction cone
+			// Clipping the accumulated impulse instead of the delta impulse (for the current iteration) is very important for jitter
+			{
+				const FVec3 OldAccumulatedImpulse = AccumulatedImpulse;
+				const FVec3 NewAccumulatedImpulse = AccumulatedImpulse + Impulse;
+				FVec3 ClippedAccumulatedImpulse(0.0f);
+				//Project to normal
+				const float NewAccImpNormalSize = FVec3::DotProduct(NewAccumulatedImpulse, Contact.Normal);
+				if (NewAccImpNormalSize > 0) // Clipping impulses to be positive
+				{
+					if (Friction <= 0)
+					{
+						ClippedAccumulatedImpulse = NewAccImpNormalSize * Contact.Normal;
+						Impulse = ClippedAccumulatedImpulse - OldAccumulatedImpulse;
+					}
+					else
+					{
+						const FVec3 ImpulseTangential = NewAccumulatedImpulse - NewAccImpNormalSize * Contact.Normal;
+						const FReal ImpulseTangentialSize = ImpulseTangential.Size();
+						const FReal MaximumImpulseTangential = Friction * NewAccImpNormalSize;
+						if (ImpulseTangentialSize <= MaximumImpulseTangential)
+						{
+							//Static friction case. We are done Impulse is correct, unless we need to add Angular friction:
+							if (AngularFriction)
+							{
+								ApplyAngularFriction(
+									Impulse, // In and out
+									AngularImpulse, // In and out
+									AngularFriction,
+									Particle0->W() - Particle1->W(),
+									VelocityChange,
+									Contact.Normal,
+									VectorToPoint1,
+									VectorToPoint2,
+									PBDRigid0,
+									PBDRigid1,
+									bIsRigidDynamic0,
+									bIsRigidDynamic1,
+									Q0,
+									Q1,
+									WorldSpaceInvI1,
+									WorldSpaceInvI2);
+							}
+						}
+						else
+						{
+							ClippedAccumulatedImpulse = (MaximumImpulseTangential / ImpulseTangentialSize) * ImpulseTangential + NewAccImpNormalSize * Contact.Normal;
+							Impulse = ClippedAccumulatedImpulse - OldAccumulatedImpulse;
+						}
+					}
+				}
+				else
+				{
+					// Clipped Impulse is 0
+					Impulse = -OldAccumulatedImpulse; // Important case: the delta impulse (Impulse) should remove the accumulated impulse
+				}
+			}
+				
+			AccumulatedImpulse += Impulse;
+			if (bIsRigidDynamic0)
+			{
+				// Velocity update for next step
+				const FVec3 NetAngularImpulse = FVec3::CrossProduct(VectorToPoint1, Impulse) + AngularImpulse;
+				const FVec3 DV = PBDRigid0->InvM() * Impulse;
+				const FVec3 DW = WorldSpaceInvI1 * NetAngularImpulse;
+				PBDRigid0->V() += DV;
+				PBDRigid0->W() += DW;
+				// Position update as part of pbd
+				P0 += (DV * IterationParameters.Dt);
+				Q0 += FRotation3::FromElements(DW, 0.f) * Q0 * IterationParameters.Dt * FReal(0.5);
+				Q0.Normalize();
+				FParticleUtilities::SetCoMWorldTransform(PBDRigid0, P0, Q0);
+			}
+			if (bIsRigidDynamic1)
+			{
+				// Velocity update for next step
+				const FVec3 NetAngularImpulse = FVec3::CrossProduct(VectorToPoint2, -Impulse) - AngularImpulse;
+				const FVec3 DV = -PBDRigid1->InvM() * Impulse;
+				const FVec3 DW = WorldSpaceInvI2 * NetAngularImpulse;
+				PBDRigid1->V() += DV;
+				PBDRigid1->W() += DW;
+				// Position update as part of pbd
+				P1 += (DV * IterationParameters.Dt);
+				Q1 += FRotation3::FromElements(DW, 0.f) * Q1 * IterationParameters.Dt * FReal(0.5);
+				Q1.Normalize();
+				FParticleUtilities::SetCoMWorldTransform(PBDRigid1, P1, Q1);
+			}
+			
+			return AccumulatedImpulse;
+		}
+
+		// Apply contacts, impulse clipping is done on delta impulses as apposed to Accumulated impulses
 		FVec3 ApplyContact(FCollisionContact& Contact,
 			TGenericParticleHandle<FReal, 3> Particle0, 
 			TGenericParticleHandle<FReal, 3> Particle1,
@@ -118,42 +315,23 @@ namespace Chaos
 						Impulse = MinimalImpulse;
 						if (AngularFriction)
 						{
-							FVec3 RelativeAngularVelocity = Particle0->W() - Particle1->W();
-							FReal AngularNormal = FVec3::DotProduct(RelativeAngularVelocity, Contact.Normal);
-							FVec3 AngularTangent = RelativeAngularVelocity - AngularNormal * Contact.Normal;
-							FVec3 FinalAngularVelocity = FMath::Sign(AngularNormal) * FMath::Max((FReal)0, FMath::Abs(AngularNormal) - AngularFriction * NormalVelocityChange) * Contact.Normal + FMath::Max((FReal)0, AngularTangent.Size() - AngularFriction * NormalVelocityChange) * AngularTangent.GetSafeNormal();
-							FVec3 Delta = FinalAngularVelocity - RelativeAngularVelocity;
-							if (!bIsRigidDynamic0 && bIsRigidDynamic1)
-							{
-								FMatrix33 WorldSpaceI2 = (Q1 * FMatrix::Identity) * PBDRigid1->I() * (Q1 * FMatrix::Identity).GetTransposed();
-								FVec3 ImpulseDelta = PBDRigid1->M() * FVec3::CrossProduct(VectorToPoint2, Delta);
-								Impulse += ImpulseDelta;
-								AngularImpulse += WorldSpaceI2 * Delta - FVec3::CrossProduct(VectorToPoint2, ImpulseDelta);
-							}
-							else if (bIsRigidDynamic0 && !bIsRigidDynamic1)
-							{
-								FMatrix33 WorldSpaceI1 = (Q0 * FMatrix::Identity) * PBDRigid0->I() * (Q0 * FMatrix::Identity).GetTransposed();
-								FVec3 ImpulseDelta = PBDRigid0->M() * FVec3::CrossProduct(VectorToPoint1, Delta);
-								Impulse += ImpulseDelta;
-								AngularImpulse += WorldSpaceI1 * Delta - FVec3::CrossProduct(VectorToPoint1, ImpulseDelta);
-							}
-							else if (bIsRigidDynamic0 && bIsRigidDynamic1)
-							{
-								FMatrix33 Cross1(0, VectorToPoint1.Z, -VectorToPoint1.Y, -VectorToPoint1.Z, 0, VectorToPoint1.X, VectorToPoint1.Y, -VectorToPoint1.X, 0);
-								FMatrix33 Cross2(0, VectorToPoint2.Z, -VectorToPoint2.Y, -VectorToPoint2.Z, 0, VectorToPoint2.X, VectorToPoint2.Y, -VectorToPoint2.X, 0);
-								FMatrix33 CrossI1 = Cross1 * WorldSpaceInvI1;
-								FMatrix33 CrossI2 = Cross2 * WorldSpaceInvI2;
-								FMatrix33 Diag1 = CrossI1 * Cross1.GetTransposed() + CrossI2 * Cross2.GetTransposed();
-								Diag1.M[0][0] += PBDRigid0->InvM() + PBDRigid1->InvM();
-								Diag1.M[1][1] += PBDRigid0->InvM() + PBDRigid1->InvM();
-								Diag1.M[2][2] += PBDRigid0->InvM() + PBDRigid1->InvM();
-								FMatrix33 OffDiag1 = (CrossI1 + CrossI2) * -1;
-								FMatrix33 Diag2 = (WorldSpaceInvI1 + WorldSpaceInvI2).Inverse();
-								FMatrix33 OffDiag1Diag2 = OffDiag1 * Diag2;
-								FVec3 ImpulseDelta = FMatrix33((Diag1 - OffDiag1Diag2 * OffDiag1.GetTransposed()).Inverse())* ((OffDiag1Diag2 * -1) * Delta);
-								Impulse += ImpulseDelta;
-								AngularImpulse += Diag2 * (Delta - FMatrix33(OffDiag1.GetTransposed()) * ImpulseDelta);
-							}
+							ApplyAngularFriction(
+								Impulse, // In and out
+								AngularImpulse, // In and out
+								AngularFriction,
+								Particle0->W() - Particle1->W(),
+								VelocityChange,
+								Contact.Normal,
+								VectorToPoint1,
+								VectorToPoint2,
+								PBDRigid0,
+								PBDRigid1,
+								bIsRigidDynamic0,
+								bIsRigidDynamic1,
+								Q0,
+								Q1,
+								WorldSpaceInvI1,
+								WorldSpaceInvI2);
 						}
 					}
 					else
@@ -397,7 +575,8 @@ namespace Chaos
 					return;
 				}
 
-				if (Constraint.GetPhi() >= ParticleParameters.ShapePadding)
+				// Do not early out here in the case of Accumulated impulse solve
+				if (Constraint.GetPhi() >= ParticleParameters.ShapePadding && !Chaos_Collision_UseAccumulatedImpulseClipSolve)
 				{
 					return;
 				}
@@ -423,7 +602,16 @@ namespace Chaos
 
 				if (bUseVelocityMode)
 				{
-					Constraint.AccumulatedImpulse += ApplyContact(Constraint.Manifold, Particle0, Particle1, IterationParameters, ParticleParameters);
+					if (Chaos_Collision_UseAccumulatedImpulseClipSolve)
+					{
+						// This version of Apply contact is different from the original in that it clips accumulated impulses instead of delta impulses
+						// Todo: we need multiple contact points per pair for this to be effective
+						ApplyContactAccImpulse(Constraint.Manifold, Particle0, Particle1, IterationParameters, ParticleParameters, Constraint.AccumulatedImpulse);
+					}
+					else
+					{
+						Constraint.AccumulatedImpulse += ApplyContact(Constraint.Manifold, Particle0, Particle1, IterationParameters, ParticleParameters);
+					}
 				}
 				else
 				{
