Fix CIS.

diff --git a/Engine/Source/Runtime/Experimental/ChaosSolvers/Public/PBDRigidsSolver.h b/Engine/Source/Runtime/Experimental/ChaosSolvers/Public/PBDRigidsSolver.h
--- a/Engine/Source/Runtime/Experimental/ChaosSolvers/Public/PBDRigidsSolver.h
+++ b/Engine/Source/Runtime/Experimental/ChaosSolvers/Public/PBDRigidsSolver.h
@@ -16,6 +16,7 @@
 #include "Chaos/PerParticleGravity.h"
 #include "Chaos/ParticleHandle.h"
 #include "Chaos/Transform.h"
+#include "Chaos/Framework/PhysicsProxy.h"
 #include "CoreMinimal.h"
 #include "Containers/Queue.h"
 #include "EventManager.h"
