Fix CIS.

diff --git a/Engine/Source/Runtime/Experimental/Chaos/Public/Chaos/RigidParticles.h b/Engine/Source/Runtime/Experimental/Chaos/Public/Chaos/RigidParticles.h
--- a/Engine/Source/Runtime/Experimental/Chaos/Public/Chaos/RigidParticles.h
+++ b/Engine/Source/Runtime/Experimental/Chaos/Public/Chaos/RigidParticles.h
@@ -143,12 +143,11 @@ class TRigidParticles : public TKinematicGeometryParticles<T, d>
 	FORCEINLINE const T InvM(const int32 Index) const { return MInvM[Index]; }
 	FORCEINLINE T& InvM(const int32 Index) { return MInvM[Index]; }
 
-	const TVector<T,d>& CenterOfMass(const int32 Index) const { return MCenterOfMass[Index]; }
-	TVector<T,d>& CenterOfMass(const int32 Index) { return MCenterOfMass[Index]; }
-
-	const TRotation<T,d>& RotationOfMass(const int32 Index) const { return MRotationOfMass[Index]; }
-	TRotation<T,d>& RotationOfMass(const int32 Index) { return MRotationOfMass[Index]; }
+	FORCEINLINE const TVector<T,d>& CenterOfMass(const int32 Index) const { return MCenterOfMass[Index]; }
+	FORCEINLINE TVector<T,d>& CenterOfMass(const int32 Index) { return MCenterOfMass[Index]; }
 
+	FORCEINLINE const TRotation<T,d>& RotationOfMass(const int32 Index) const { return MRotationOfMass[Index]; }
+	FORCEINLINE TRotation<T,d>& RotationOfMass(const int32 Index) { return MRotationOfMass[Index]; }
 
 	FORCEINLINE const T& LinearEtherDrag(const int32 index) const { return MLinearEtherDrag[index]; }
 	FORCEINLINE T& LinearEtherDrag(const int32 index) { return MLinearEtherDrag[index]; }
