Add overrides for collision jitter testing.

diff --git a/Engine/Source/Runtime/Experimental/Chaos/Private/Chaos/PBDRigidsEvolution.cpp b/Engine/Source/Runtime/Experimental/Chaos/Private/Chaos/PBDRigidsEvolution.cpp
--- a/Engine/Source/Runtime/Experimental/Chaos/Private/Chaos/PBDRigidsEvolution.cpp
+++ b/Engine/Source/Runtime/Experimental/Chaos/Private/Chaos/PBDRigidsEvolution.cpp
@@ -13,6 +13,12 @@ FAutoConsoleVariableRef CVarChaosRigidsEvolutionApplyAllowEarlyOut(TEXT("p.Chaos
 int32 ChaosRigidsEvolutionApplyPushoutAllowEarlyOutCVar = 1;
 FAutoConsoleVariableRef CVarChaosRigidsEvolutionApplyPushoutAllowEarlyOut(TEXT("p.ChaosRigidsEvolutionApplyPushoutAllowEarlyOut"), ChaosRigidsEvolutionApplyPushoutAllowEarlyOutCVar, TEXT("Allow Chaos Rigids Evolution apply-pushout iterations to early out when resolved.[def:1]"));
 
+int32 ChaosNumPushOutIterationsOverride = -1;
+FAutoConsoleVariableRef CVarChaosNumPushOutIterationsOverride(TEXT("p.ChaosNumPushOutIterationsOverride"), ChaosNumPushOutIterationsOverride, TEXT("Override for num push out iterations if >= 0 [def:-1]"));
+
+int32 ChaosNumContactIterationsOverride = 1;
+FAutoConsoleVariableRef CVarChaosNumContactIterationsOverride(TEXT("p.ChaosNumContactIterationsOverride"), ChaosNumContactIterationsOverride, TEXT("Override for num contact iterations if >= 0. [def:-1]"));
+
 namespace Chaos
 {
 	CHAOS_API int32 FixBadAccelerationStructureRemoval = 1;
diff --git a/Engine/Source/Runtime/Experimental/Chaos/Public/Chaos/PBDRigidsEvolution.h b/Engine/Source/Runtime/Experimental/Chaos/Public/Chaos/PBDRigidsEvolution.h
--- a/Engine/Source/Runtime/Experimental/Chaos/Public/Chaos/PBDRigidsEvolution.h
+++ b/Engine/Source/Runtime/Experimental/Chaos/Public/Chaos/PBDRigidsEvolution.h
@@ -18,7 +18,8 @@
 
 extern int32 ChaosRigidsEvolutionApplyAllowEarlyOutCVar;
 extern int32 ChaosRigidsEvolutionApplyPushoutAllowEarlyOutCVar;
-
+extern int32 ChaosNumPushOutIterationsOverride;
+extern int32 ChaosNumContactIterationsOverride;
 
 // Declaring so it can be friended for tests.
 namespace ChaosTest { void TestPendingSpatialDataHandlePointerConflict(); } 
@@ -512,13 +513,14 @@ class FPBDRigidsEvolutionBase
 	{
 		UpdateAccelerationStructures(Island);
 
+		int32 LocalNumIterations = ChaosNumContactIterationsOverride >= 0 ? ChaosNumContactIterationsOverride : NumIterations;
 		// @todo(ccaulfield): track whether we are sufficiently solved and can early-out
-		for (int i = 0; i < NumIterations; ++i)
+		for (int i = 0; i < LocalNumIterations; ++i)
 		{
 			bool bNeedsAnotherIteration = false;
 			for (FPBDConstraintGraphRule* ConstraintRule : PrioritizedConstraintRules)
 			{
-				bNeedsAnotherIteration |= ConstraintRule->ApplyConstraints(Dt, Island, i, NumIterations);
+				bNeedsAnotherIteration |= ConstraintRule->ApplyConstraints(Dt, Island, i, LocalNumIterations);
 			}
 
 			if (ChaosRigidsEvolutionApplyAllowEarlyOutCVar && !bNeedsAnotherIteration)
@@ -689,13 +691,14 @@ protected:
 
 	void ApplyPushOut(const FReal Dt, int32 Island)
 	{
+		int32 LocalNumPushOutIterations = ChaosNumPushOutIterationsOverride >= 0 ? ChaosNumPushOutIterationsOverride : NumPushOutIterations;
 		bool bNeedsAnotherIteration = true;
-		for (int32 It = 0; It < NumPushOutIterations; ++It)
+		for (int32 It = 0; It < LocalNumPushOutIterations; ++It)
 		{
 			bNeedsAnotherIteration = false;
 			for (FPBDConstraintGraphRule* ConstraintRule : PrioritizedConstraintRules)
 			{
-				bNeedsAnotherIteration |= ConstraintRule->ApplyPushOut(Dt, Island, It, NumPushOutIterations);
+				bNeedsAnotherIteration |= ConstraintRule->ApplyPushOut(Dt, Island, It, LocalNumPushOutIterations);
 			}
 
 			if (ChaosRigidsEvolutionApplyPushoutAllowEarlyOutCVar && !bNeedsAnotherIteration)
