remove jitter due to bug with triple buffer

diff --git a/Engine/Source/Runtime/Experimental/ChaosSolvers/Private/PhysicsProxy/GeometryCollectionPhysicsProxy.cpp b/Engine/Source/Runtime/Experimental/ChaosSolvers/Private/PhysicsProxy/GeometryCollectionPhysicsProxy.cpp
--- a/Engine/Source/Runtime/Experimental/ChaosSolvers/Private/PhysicsProxy/GeometryCollectionPhysicsProxy.cpp
+++ b/Engine/Source/Runtime/Experimental/ChaosSolvers/Private/PhysicsProxy/GeometryCollectionPhysicsProxy.cpp
@@ -1444,6 +1444,11 @@ void FGeometryCollectionPhysicsProxy::BufferPhysicsResults()
 
 	{ SCOPE_CYCLE_COUNTER(STAT_CalcParticleToWorld);
 
+		// initialize Target Results
+		TargetResults.Transforms.Init(PhysicsThreadCollection.Transform);
+		TargetResults.Children.Init(PhysicsThreadCollection.Children);
+		TargetResults.Parent.Init(PhysicsThreadCollection.Parent);
+
 		for (int32 TransformGroupIndex = 0; TransformGroupIndex < NumTransformGroupElements; ++TransformGroupIndex)
 		{
 			TargetResults.DisabledStates[TransformGroupIndex] = true;
@@ -1497,6 +1502,8 @@ void FGeometryCollectionPhysicsProxy::BufferPhysicsResults()
 				TargetResults.Transforms[TransformGroupIndex] = PhysicsThreadCollection.MassToLocal[TransformGroupIndex].GetRelativeTransformReverse(ParticleToWorld).GetRelativeTransform(ActorToWorld);
 				TargetResults.Transforms[TransformGroupIndex].NormalizeRotation();
 
+				PhysicsThreadCollection.Transform[TransformGroupIndex] = TargetResults.Transforms[TransformGroupIndex];
+
 				// Indicate that this object needs to be updated and the proxy is active.
 				TargetResults.DisabledStates[TransformGroupIndex] = false;
 				IsObjectDynamic = true;
@@ -1563,6 +1570,8 @@ void FGeometryCollectionPhysicsProxy::BufferPhysicsResults()
 							TargetResults.Transforms[TransformGroupIndex] = MassToLocal.GetRelativeTransformReverse(ClusterChildToWorld).GetRelativeTransform(ActorToWorld);
 							TargetResults.Transforms[TransformGroupIndex].NormalizeRotation();
 
+							PhysicsThreadCollection.Transform[TransformGroupIndex] = TargetResults.Transforms[TransformGroupIndex];
+
 							// Indicate that this object needs to be updated and the proxy is active.
 							TargetResults.DisabledStates[TransformGroupIndex] = false;
 							IsObjectDynamic = true;
