Fix static analysis

diff --git a/Engine/Source/Runtime/Experimental/Chaos/Public/Chaos/IncludeLvl1.inl b/Engine/Source/Runtime/Experimental/Chaos/Public/Chaos/IncludeLvl1.inl
--- a/Engine/Source/Runtime/Experimental/Chaos/Public/Chaos/IncludeLvl1.inl
+++ b/Engine/Source/Runtime/Experimental/Chaos/Public/Chaos/IncludeLvl1.inl
@@ -5,6 +5,7 @@
 //Until we fix all instances of this please use TEMP_HEADER_CHAOS_LEVEL_1 (different name for searchability)
 //Or consider making the included header a Level 0 - this is meant for files that rarely change like forward declares
 
+#if 0
 #ifdef CHAOS_LEVEL_CHECK
 #undef CHAOS_LEVEL_CHECK
 #undef CHAOS_LEVEL_CHECK_IMP
@@ -14,14 +15,20 @@
 #if defined(CHAOS_INCLUDE_LEVEL_1)
 #define CHAOS_INCLUDE_LEVEL_ACTUAL CHAOS_INCLUDE_LEVEL_1
 #elif defined(TEMP_HEADER_CHAOS_LEVEL_1)
-#define CHAOS_INCLUDE_LEVEL_ACTUAL CHAOS_INCLUDE_LEVEL_1
+#define CHAOS_INCLUDE_LEVEL_ACTUAL TEMP_HEADER_CHAOS_LEVEL_1
 #endif
 
 #ifndef CHAOS_INCLUDE_LEVEL_ACTUAL
-//#define CHAOS_LEVEL_CHECK_IMP static_assert(false, "Attempting to upgrade CHAOS_INCLUDE_LEVEL");
-#define CHAOS_LEVEL_CHECK_IMP
+#define CHAOS_LEVEL_CHECK_IMP static_assert(false, "Attempting to upgrade CHAOS_INCLUDE_LEVEL");
+//#define CHAOS_LEVEL_CHECK_IMP
 #else
 #define CHAOS_LEVEL_CHECK_IMP
 #endif
 
-#define CHAOS_LEVEL_CHECK CHAOS_LEVEL_CHECK_IMP
\ No newline at end of file
+#define CHAOS_LEVEL_CHECK CHAOS_LEVEL_CHECK_IMP
+#endif
+
+//TODO: remove
+#ifndef CHAOS_LEVEL_CHECK
+#define CHAOS_LEVEL_CHECK
+#endif 
