Conflicts resolved: Chaos: Disabled experimental contact solve uses the accumulated impulse (instead of the delta impulse) when clamping
Only works for cases where there are more than one contact point per iteration that persists
- Much less jitter with 4 iterations
- More iterations => less jitter now
- Bodies on inclined planes will now start sliding at the correct angle (given a friction coefficient)

#[fyi] Michael.Lentine, Brice.Criswell,  Chris.Caulfield

diff --git a/Engine/Source/Runtime/Experimental/Chaos/Private/Chaos/PBDCollisionConstraints.cpp b/Engine/Source/Runtime/Experimental/Chaos/Private/Chaos/PBDCollisionConstraints.cpp
--- a/Engine/Source/Runtime/Experimental/Chaos/Private/Chaos/PBDCollisionConstraints.cpp
+++ b/Engine/Source/Runtime/Experimental/Chaos/Private/Chaos/PBDCollisionConstraints.cpp
@@ -56,6 +56,9 @@ namespace Chaos
 	float DefaultCollisionFriction = 0;
 	FAutoConsoleVariableRef CVarDefaultCollisionFriction(TEXT("p.DefaultCollisionFriction"), DefaultCollisionFriction, TEXT("Collision friction default value if no materials are found."));
 
+	int32 Chaos_Collision_UseAccumulatedImpulseClipSolve = 0; // Experimental: This requires multiple contact points per iteration per pair, and making sure the contact points don't move too much in body space
+	FAutoConsoleVariableRef CVarChaosCollisionOriginalSolve(TEXT("p.Chaos.Collision.UseAccumulatedImpulseClipSolve"), Chaos_Collision_UseAccumulatedImpulseClipSolve, TEXT("Use experimental Accumulated impulse clipped contact solve"));
+
 	float DefaultCollisionRestitution = 0;
 	FAutoConsoleVariableRef CVarDefaultCollisionRestitution(TEXT("p.DefaultCollisionRestitution"), DefaultCollisionRestitution, TEXT("Collision restitution default value if no materials are found."));
 
