CIS: removed inertia scale that does not exist anymore

diff --git a/Engine/Source/Runtime/Experimental/Chaos/Private/Chaos/PBDCollisionConstraintsContact.cpp b/Engine/Source/Runtime/Experimental/Chaos/Private/Chaos/PBDCollisionConstraintsContact.cpp
--- a/Engine/Source/Runtime/Experimental/Chaos/Private/Chaos/PBDCollisionConstraintsContact.cpp
+++ b/Engine/Source/Runtime/Experimental/Chaos/Private/Chaos/PBDCollisionConstraintsContact.cpp
@@ -145,8 +145,8 @@ namespace Chaos
 			const FVec3 RelativeVelocity = Body1Velocity - Body2Velocity;  // Do not early out on negative normal velocity since there can still be an accumulated impulse
 			*IterationParameters.NeedsAnotherIteration = true;
 
-			const FMatrix33 WorldSpaceInvI1 = bIsRigidDynamic0 ? Utilities::ComputeWorldSpaceInertia(Q0, PBDRigid0->InvI()) * Contact.InvInertiaScale0 : FMatrix33(0);
-			const FMatrix33 WorldSpaceInvI2 = bIsRigidDynamic1 ? Utilities::ComputeWorldSpaceInertia(Q1, PBDRigid1->InvI()) * Contact.InvInertiaScale1 : FMatrix33(0);
+			const FMatrix33 WorldSpaceInvI1 = bIsRigidDynamic0 ? Utilities::ComputeWorldSpaceInertia(Q0, PBDRigid0->InvI()) : FMatrix33(0);
+			const FMatrix33 WorldSpaceInvI2 = bIsRigidDynamic1 ? Utilities::ComputeWorldSpaceInertia(Q1, PBDRigid1->InvI()) : FMatrix33(0);
 			const FMatrix33 Factor =
 				(bIsRigidDynamic0 ? ComputeFactorMatrix3(VectorToPoint1, WorldSpaceInvI1, PBDRigid0->InvM()) : FMatrix33(0)) +
 				(bIsRigidDynamic1 ? ComputeFactorMatrix3(VectorToPoint2, WorldSpaceInvI2, PBDRigid1->InvM()) : FMatrix33(0));
