CIS: Removed duplicate CVAR

diff --git a/Engine/Source/Runtime/Experimental/Chaos/Private/Chaos/PBDCollisionConstraints.cpp b/Engine/Source/Runtime/Experimental/Chaos/Private/Chaos/PBDCollisionConstraints.cpp
--- a/Engine/Source/Runtime/Experimental/Chaos/Private/Chaos/PBDCollisionConstraints.cpp
+++ b/Engine/Source/Runtime/Experimental/Chaos/Private/Chaos/PBDCollisionConstraints.cpp
@@ -56,9 +56,6 @@ namespace Chaos
 	float DefaultCollisionFriction = 0;
 	FAutoConsoleVariableRef CVarDefaultCollisionFriction(TEXT("p.DefaultCollisionFriction"), DefaultCollisionFriction, TEXT("Collision friction default value if no materials are found."));
 
-	int32 Chaos_Collision_UseAccumulatedImpulseClipSolve = 0; // Experimental: This requires multiple contact points per iteration per pair, and making sure the contact points don't move too much in body space
-	FAutoConsoleVariableRef CVarChaosCollisionOriginalSolve(TEXT("p.Chaos.Collision.UseAccumulatedImpulseClipSolve"), Chaos_Collision_UseAccumulatedImpulseClipSolve, TEXT("Use experimental Accumulated impulse clipped contact solve"));
-
 	float DefaultCollisionRestitution = 0;
 	FAutoConsoleVariableRef CVarDefaultCollisionRestitution(TEXT("p.DefaultCollisionRestitution"), DefaultCollisionRestitution, TEXT("Collision restitution default value if no materials are found."));
 
