Add more reasonable fallback for inertia.

diff --git a/Engine/Source/Runtime/Experimental/ChaosSolvers/Private/PhysicsProxy/GeometryCollectionPhysicsProxy.cpp b/Engine/Source/Runtime/Experimental/ChaosSolvers/Private/PhysicsProxy/GeometryCollectionPhysicsProxy.cpp
--- a/Engine/Source/Runtime/Experimental/ChaosSolvers/Private/PhysicsProxy/GeometryCollectionPhysicsProxy.cpp
+++ b/Engine/Source/Runtime/Experimental/ChaosSolvers/Private/PhysicsProxy/GeometryCollectionPhysicsProxy.cpp
@@ -1847,6 +1847,17 @@ void FGeometryCollectionPhysicsProxy::InitializeSharedCollisionStructures(
 				{
 					CalculateVolumeAndCenterOfMass(MassSpaceParticles, TriMesh->GetElements(), MassProperties.Volume, MassProperties.CenterOfMass);
 
+					if(MassProperties.Volume == 0)
+					{
+						FVector Extents = BoundingBox[GeometryIndex].GetExtent();
+						MassProperties.Volume = Extents.X * Extents.Y * Extents.Z;
+						float ExtentsYZ = Extents.Y * Extents.Y + Extents.Z * Extents.Z;
+						float ExtentsXZ = Extents.X * Extents.X + Extents.Z * Extents.Z;
+						float ExtentsXY = Extents.X * Extents.X + Extents.Y * Extents.Y;
+						MassProperties.InertiaTensor = PMatrix<float, 3, 3>(ExtentsYZ / 12., ExtentsXZ / 12., ExtentsXY / 12.);
+						MassProperties.CenterOfMass = BoundingBox[GeometryIndex].GetCenter();
+					}
+
 					if (MassProperties.Volume < MinVolume)
 					{
 						// For rigid bodies outside of range just defaut to a clamped bounding box, and warn the user.
