Fix inertia init.

diff --git a/Engine/Source/Runtime/Experimental/ChaosSolvers/Private/PhysicsProxy/GeometryCollectionPhysicsProxy.cpp b/Engine/Source/Runtime/Experimental/ChaosSolvers/Private/PhysicsProxy/GeometryCollectionPhysicsProxy.cpp
--- a/Engine/Source/Runtime/Experimental/ChaosSolvers/Private/PhysicsProxy/GeometryCollectionPhysicsProxy.cpp
+++ b/Engine/Source/Runtime/Experimental/ChaosSolvers/Private/PhysicsProxy/GeometryCollectionPhysicsProxy.cpp
@@ -243,18 +243,18 @@ void PopulateSimulatedParticle(
 	{
 		Handle->SetObjectState(Chaos::EObjectStateType::Uninitialized);
 
-		if (ensureMsgf(FMath::IsWithinInclusive(MassIn, SharedParams.MinimumMassClamp, SharedParams.MaximumMassClamp),
+		if (!ensureMsgf(FMath::IsWithinInclusive(MassIn, SharedParams.MinimumMassClamp, SharedParams.MaximumMassClamp),
 			TEXT("Clamped mass[%3.5f] to range [%3.5f,%3.5f]"), MassIn, SharedParams.MinimumMassClamp, SharedParams.MaximumMassClamp))
 		{
 			MassIn = FMath::Clamp(MassIn, SharedParams.MinimumMassClamp, SharedParams.MaximumMassClamp);
 		}
 
-		if (ensureMsgf(!FMath::IsNaN(InertiaTensorVec[0]) && !FMath::IsNaN(InertiaTensorVec[1]) && !FMath::IsNaN(InertiaTensorVec[2]),
+		if (!ensureMsgf(!FMath::IsNaN(InertiaTensorVec[0]) && !FMath::IsNaN(InertiaTensorVec[1]) && !FMath::IsNaN(InertiaTensorVec[2]),
 			TEXT("Nan Tensor, reset to unit tesor")))
 		{
 			InertiaTensorVec = FVector(1);
 		}
-		else if (ensureMsgf(FMath::IsWithinInclusive(InertiaTensorVec[0], SharedParams.MinimumInertiaTensorDiagonalClamp, SharedParams.MaximumInertiaTensorDiagonalClamp)
+		else if (!ensureMsgf(FMath::IsWithinInclusive(InertiaTensorVec[0], SharedParams.MinimumInertiaTensorDiagonalClamp, SharedParams.MaximumInertiaTensorDiagonalClamp)
 			&& FMath::IsWithinInclusive(InertiaTensorVec[1], SharedParams.MinimumInertiaTensorDiagonalClamp, SharedParams.MaximumInertiaTensorDiagonalClamp)
 			&& FMath::IsWithinInclusive(InertiaTensorVec[2], SharedParams.MinimumInertiaTensorDiagonalClamp, SharedParams.MaximumInertiaTensorDiagonalClamp),
 			TEXT("Clamped Inertia tensor[%3.5f,%3.5f,%3.5f]. Clamped each element to [%3.5f, %3.5f,]"), InertiaTensorVec[0], InertiaTensorVec[1], InertiaTensorVec[2],
@@ -1846,6 +1846,7 @@ void FGeometryCollectionPhysicsProxy::InitializeSharedCollisionStructures(
 				else
 				{
 					CalculateVolumeAndCenterOfMass(MassSpaceParticles, TriMesh->GetElements(), MassProperties.Volume, MassProperties.CenterOfMass);
+					InertiaComputationNeeded[GeometryIndex] = true;
 
 					if(MassProperties.Volume == 0)
 					{
@@ -1856,6 +1857,7 @@ void FGeometryCollectionPhysicsProxy::InitializeSharedCollisionStructures(
 						float ExtentsXY = Extents.X * Extents.X + Extents.Y * Extents.Y;
 						MassProperties.InertiaTensor = PMatrix<float, 3, 3>(ExtentsYZ / 12., ExtentsXZ / 12., ExtentsXY / 12.);
 						MassProperties.CenterOfMass = BoundingBox[GeometryIndex].GetCenter();
+						InertiaComputationNeeded[GeometryIndex] = false;
 					}
 
 					if (MassProperties.Volume < MinVolume)
@@ -1863,16 +1865,17 @@ void FGeometryCollectionPhysicsProxy::InitializeSharedCollisionStructures(
 						// For rigid bodies outside of range just defaut to a clamped bounding box, and warn the user.
 						MassProperties.Volume = MinVolume;
 						CollectionMassToLocal[TransformGroupIndex] = FTransform(FQuat::Identity, BoundingBox[GeometryIndex].GetCenter());
+						InertiaComputationNeeded[GeometryIndex] = false;
 					}
 					else if (MaxVolume < MassProperties.Volume)
 					{
 						// For rigid bodies outside of range just defaut to a clamped bounding box, and warn the user
 						MassProperties.Volume = MaxVolume;
 						CollectionMassToLocal[TransformGroupIndex] = FTransform(FQuat::Identity, BoundingBox[GeometryIndex].GetCenter());
+						InertiaComputationNeeded[GeometryIndex] = false;
 					}
 					else
 					{
-						InertiaComputationNeeded[GeometryIndex] = true;
 						CollectionMassToLocal[TransformGroupIndex] = FTransform(FQuat::Identity, MassProperties.CenterOfMass);
 					}
 
