IWYU changes

diff --git a/Engine/Plugins/Experimental/ChaosEditor/Source/FractureEditor/Public/FractureEditorModeToolkit.h b/Engine/Plugins/Experimental/ChaosEditor/Source/FractureEditor/Public/FractureEditorModeToolkit.h
--- a/Engine/Plugins/Experimental/ChaosEditor/Source/FractureEditor/Public/FractureEditorModeToolkit.h
+++ b/Engine/Plugins/Experimental/ChaosEditor/Source/FractureEditor/Public/FractureEditorModeToolkit.h
@@ -7,7 +7,6 @@
 #include "Toolkits/BaseToolkit.h"
 #include "Widgets/Input/SCheckBox.h"
 #include "Widgets/Layout/SSplitter.h"
-#include "GeometryCollection/GeometryCollectionActor.h"
 
 class IDetailsView;
 class SScrollBox;
@@ -16,6 +15,14 @@ class UEditableMesh;
 struct FFractureContext;
 class FGeometryCollection;
 class SGeometryCollectionOutliner;
+class AGeometryCollectionActor;
+class UGeometryCollectionComponent;
+class UGeometryCollection;
+
+namespace GeometryCollection
+{
+enum class ESelectionMode: uint8;
+}
 
 class FFractureEditorModeToolkit : public FModeToolkit, public FGCObject
 {
diff --git a/Engine/Plugins/Experimental/ChaosNiagara/Source/ChaosNiagara/Private/NiagaraDataInterfaceChaosDestruction.cpp b/Engine/Plugins/Experimental/ChaosNiagara/Source/ChaosNiagara/Private/NiagaraDataInterfaceChaosDestruction.cpp
--- a/Engine/Plugins/Experimental/ChaosNiagara/Source/ChaosNiagara/Private/NiagaraDataInterfaceChaosDestruction.cpp
+++ b/Engine/Plugins/Experimental/ChaosNiagara/Source/ChaosNiagara/Private/NiagaraDataInterfaceChaosDestruction.cpp
@@ -15,6 +15,7 @@
 #include "NiagaraEmitterInstanceBatcher.h"
 
 #include <memory>
+#include "GeometryCollection/GeometryCollectionComponent.h"
 
 #define LOCTEXT_NAMESPACE "ChaosNiagaraDestructionDataInterface"
 //#pragma optimize("", off)
diff --git a/Engine/Plugins/Experimental/GeometryCollectionPlugin/Source/GeometryCollectionEditor/Private/GeometryCollection/DetailCustomizations/WarningMessageCustomization.cpp b/Engine/Plugins/Experimental/GeometryCollectionPlugin/Source/GeometryCollectionEditor/Private/GeometryCollection/DetailCustomizations/WarningMessageCustomization.cpp
--- a/Engine/Plugins/Experimental/GeometryCollectionPlugin/Source/GeometryCollectionEditor/Private/GeometryCollection/DetailCustomizations/WarningMessageCustomization.cpp
+++ b/Engine/Plugins/Experimental/GeometryCollectionPlugin/Source/GeometryCollectionEditor/Private/GeometryCollection/DetailCustomizations/WarningMessageCustomization.cpp
@@ -11,6 +11,7 @@
 #include "EditorFontGlyphs.h"
 #include "GeometryCollection/GeometryCollectionActor.h"
 #include "GeometryCollection/GeometryCollectionDebugDrawActor.h"
+#include "GeometryCollection/GeometryCollectionDebugDrawComponent.h"
 
 TSharedRef<IPropertyTypeCustomization> FWarningMessageCustomization::MakeInstance()
 {
diff --git a/Engine/Plugins/Experimental/GeometryCollectionPlugin/Source/GeometryCollectionEditor/Private/GeometryCollection/GeometryCollectionSelectRigidBodyEdMode.cpp b/Engine/Plugins/Experimental/GeometryCollectionPlugin/Source/GeometryCollectionEditor/Private/GeometryCollection/GeometryCollectionSelectRigidBodyEdMode.cpp
--- a/Engine/Plugins/Experimental/GeometryCollectionPlugin/Source/GeometryCollectionEditor/Private/GeometryCollection/GeometryCollectionSelectRigidBodyEdMode.cpp
+++ b/Engine/Plugins/Experimental/GeometryCollectionPlugin/Source/GeometryCollectionEditor/Private/GeometryCollection/GeometryCollectionSelectRigidBodyEdMode.cpp
@@ -12,6 +12,7 @@
 #include "EditorViewportClient.h"
 #include "EngineUtils.h"
 #include "PropertyHandle.h"
+#include "GeometryCollection/GeometryCollectionComponent.h"
 
 DEFINE_LOG_CATEGORY_STATIC(LogGeometryCollectionSelectRigidBodyEdMode, Log, All);
 
diff --git a/Engine/Source/Runtime/Experimental/GeometryCollectionEngine/Private/ChaosBlueprint.cpp b/Engine/Source/Runtime/Experimental/GeometryCollectionEngine/Private/ChaosBlueprint.cpp
--- a/Engine/Source/Runtime/Experimental/GeometryCollectionEngine/Private/ChaosBlueprint.cpp
+++ b/Engine/Source/Runtime/Experimental/GeometryCollectionEngine/Private/ChaosBlueprint.cpp
@@ -4,6 +4,8 @@
 #include "PhysicsSolver.h"
 #include "Async/Async.h"
 #include "PhysicsProxy/GeometryCollectionPhysicsProxy.h"
+#include "GeometryCollection/GeometryCollectionActor.h"
+#include "GeometryCollection/GeometryCollectionComponent.h"
 
 #define DISPATCH_BLUEPRINTS_IMMEDIATE 1
 
diff --git a/Engine/Source/Runtime/Experimental/GeometryCollectionEngine/Private/GeometryCollection/GeometryCollectionActor.cpp b/Engine/Source/Runtime/Experimental/GeometryCollectionEngine/Private/GeometryCollection/GeometryCollectionActor.cpp
--- a/Engine/Source/Runtime/Experimental/GeometryCollectionEngine/Private/GeometryCollection/GeometryCollectionActor.cpp
+++ b/Engine/Source/Runtime/Experimental/GeometryCollectionEngine/Private/GeometryCollection/GeometryCollectionActor.cpp
@@ -19,6 +19,7 @@ GeometryCollectionActor.cpp: AGeometryCollectionActor methods.
 #include "Math/Box.h"
 #include "Physics/PhysicsInterfaceCore.h"
 #include "PhysicsSolver.h"
+#include "GeometryCollection/GeometryCollectionDebugDrawComponent.h"
 
 
 DEFINE_LOG_CATEGORY_STATIC(AGeometryCollectionActorLogging, Log, All);
diff --git a/Engine/Source/Runtime/Experimental/GeometryCollectionEngine/Private/GeometryCollection/GeometryCollectionComponent.cpp b/Engine/Source/Runtime/Experimental/GeometryCollectionEngine/Private/GeometryCollection/GeometryCollectionComponent.cpp
--- a/Engine/Source/Runtime/Experimental/GeometryCollectionEngine/Private/GeometryCollection/GeometryCollectionComponent.cpp
+++ b/Engine/Source/Runtime/Experimental/GeometryCollectionEngine/Private/GeometryCollection/GeometryCollectionComponent.cpp
@@ -21,6 +21,7 @@
 #include "PhysicsProxy/GeometryCollectionPhysicsProxy.h"
 #include "PhysicsSolver.h"
 #include "Physics/PhysicsFiltering.h"
+#include "PhysicalMaterials/Experimental/ChaosPhysicalMaterial.h"
 
 #if WITH_EDITOR
 #include "AssetToolsModule.h"
diff --git a/Engine/Source/Runtime/Experimental/GeometryCollectionEngine/Private/GeometryCollection/GeometryCollectionDebugDrawActor.cpp b/Engine/Source/Runtime/Experimental/GeometryCollectionEngine/Private/GeometryCollection/GeometryCollectionDebugDrawActor.cpp
--- a/Engine/Source/Runtime/Experimental/GeometryCollectionEngine/Private/GeometryCollection/GeometryCollectionDebugDrawActor.cpp
+++ b/Engine/Source/Runtime/Experimental/GeometryCollectionEngine/Private/GeometryCollection/GeometryCollectionDebugDrawActor.cpp
@@ -20,6 +20,7 @@
 #include "HAL/IConsoleManager.h"
 #include "PBDRigidsSolver.h"
 #include "PhysicsSolver.h"  // #if TODO_REIMPLEMENT_GET_RIGID_PARTICLES
+#include "GeometryCollection/GeometryCollectionDebugDrawComponent.h"
 
 DEFINE_LOG_CATEGORY_STATIC(LogGeometryCollectionDebugDrawActor, Log, All);
 
diff --git a/Engine/Source/Runtime/Experimental/GeometryCollectionEngine/Public/ChaosBlueprint.h b/Engine/Source/Runtime/Experimental/GeometryCollectionEngine/Public/ChaosBlueprint.h
--- a/Engine/Source/Runtime/Experimental/GeometryCollectionEngine/Public/ChaosBlueprint.h
+++ b/Engine/Source/Runtime/Experimental/GeometryCollectionEngine/Public/ChaosBlueprint.h
@@ -9,13 +9,12 @@
 #include "EventManager.h"
 #include "EventsData.h"
 #include "HAL/ThreadSafeBool.h"
-#include "GeometryCollection/GeometryCollectionActor.h"
 #include "ChaosCollisionEventFilter.h"
 #include "ChaosBreakingEventFilter.h"
 #include "ChaosTrailingEventFilter.h"
 #include "ChaosBlueprint.generated.h"
 
-
+class AGeometryCollectionActor;
 
 /**
 * Called when new destruction events are available for collisions. Collision listening must be enabled to get callbacks on this delegate.
diff --git a/Engine/Source/Runtime/Experimental/GeometryCollectionEngine/Public/GeometryCollection/GeometryCollectionActor.h b/Engine/Source/Runtime/Experimental/GeometryCollectionEngine/Public/GeometryCollection/GeometryCollectionActor.h
--- a/Engine/Source/Runtime/Experimental/GeometryCollectionEngine/Public/GeometryCollection/GeometryCollectionActor.h
+++ b/Engine/Source/Runtime/Experimental/GeometryCollectionEngine/Public/GeometryCollection/GeometryCollectionActor.h
@@ -6,16 +6,14 @@
 
 #include "CoreMinimal.h"
 #include "GameFramework/Actor.h"
-#include "GeometryCollection/GeometryCollection.h"
-#include "GeometryCollection/GeometryCollectionComponent.h"
-#include "GeometryCollection/GeometryCollectionDebugDrawComponent.h"
-#include "Physics/Experimental/PhysScene_Chaos.h"
-#include "Physics/ImmediatePhysics/ImmediatePhysicsActorHandle.h"
 #include "UObject/ObjectMacros.h"
 
 #include "GeometryCollectionActor.generated.h"
 
 
+class UGeometryCollectionComponent;
+class UGeometryCollectionDebugDrawComponent;
+struct FHitResult;
 
 UCLASS()
 class GEOMETRYCOLLECTIONENGINE_API AGeometryCollectionActor: public AActor
diff --git a/Engine/Source/Runtime/Experimental/GeometryCollectionEngine/Public/GeometryCollection/GeometryCollectionComponent.h b/Engine/Source/Runtime/Experimental/GeometryCollectionEngine/Public/GeometryCollection/GeometryCollectionComponent.h
--- a/Engine/Source/Runtime/Experimental/GeometryCollectionEngine/Public/GeometryCollection/GeometryCollectionComponent.h
+++ b/Engine/Source/Runtime/Experimental/GeometryCollectionEngine/Public/GeometryCollection/GeometryCollectionComponent.h
@@ -3,8 +3,6 @@
 #pragma once
 
 #include "Components/MeshComponent.h"
-#include "Chaos/ChaosSolverActor.h"
-#include "Chaos/PBDRigidParticles.h"
 #include "Chaos/Defines.h"
 #include "Field/FieldSystem.h"
 #include "Field/FieldSystemActor.h"
@@ -20,7 +18,6 @@
 #include "GeometryCollectionEditorSelection.h"
 #include "GeometryCollection/RecordedTransformTrack.h"
 #include "Templates/UniquePtr.h"
-#include "PhysicalMaterials/Experimental/ChaosPhysicalMaterial.h"
 #include "Chaos/ChaosGameplayEventDispatcher.h"
 #include "Chaos/ChaosNotifyHandlerInterface.h"
 #include "Chaos/ChaosSolverComponentTypes.h"
@@ -33,6 +30,8 @@ class FGeometryCollectionPhysicsProxy;
 class UGeometryCollectionComponent;
 class UBoxComponent;
 class UGeometryCollectionCache;
+class UChaosPhysicalMaterial;
+class AChaosSolverActor;
 
 DECLARE_DYNAMIC_MULTICAST_DELEGATE_OneParam(FOnChaosBreakEvent, const FChaosBreakEvent&, BreakEvent);
 
diff --git a/Engine/Source/Runtime/Experimental/GeometryCollectionSimulationCore/Public/GeometryCollection/GeometryCollectionSimulationTypes.h b/Engine/Source/Runtime/Experimental/GeometryCollectionSimulationCore/Public/GeometryCollection/GeometryCollectionSimulationTypes.h
--- a/Engine/Source/Runtime/Experimental/GeometryCollectionSimulationCore/Public/GeometryCollection/GeometryCollectionSimulationTypes.h
+++ b/Engine/Source/Runtime/Experimental/GeometryCollectionSimulationCore/Public/GeometryCollection/GeometryCollectionSimulationTypes.h
@@ -2,7 +2,6 @@
 #pragma once
 
 #include "UObject/ObjectMacros.h"
-#include "Chaos/PBDRigidParticles.h"
 
 #include "GeometryCollectionSimulationTypes.generated.h"
 
