Reverse winding order for mass properties calculations.

diff --git a/Engine/Source/Runtime/Experimental/ChaosSolvers/Private/PhysicsProxy/GeometryCollectionPhysicsProxy.cpp b/Engine/Source/Runtime/Experimental/ChaosSolvers/Private/PhysicsProxy/GeometryCollectionPhysicsProxy.cpp
--- a/Engine/Source/Runtime/Experimental/ChaosSolvers/Private/PhysicsProxy/GeometryCollectionPhysicsProxy.cpp
+++ b/Engine/Source/Runtime/Experimental/ChaosSolvers/Private/PhysicsProxy/GeometryCollectionPhysicsProxy.cpp
@@ -121,7 +121,7 @@ Chaos::TTriangleMesh<float>* CreateTriangleMesh(
 		if (Visible[Idx])
 		{
 			const FIntVector& Tri = Indices[Idx];
-			Faces.Add(Chaos::TVector<int32, 3>(Tri.X, Tri.Y, Tri.Z));
+			Faces.Add(Chaos::TVector<int32, 3>(Tri.Z, Tri.Y, Tri.X));
 		}
 	}
 	return new Chaos::TTriangleMesh<float>(MoveTemp(Faces)); // Culls geometrically degenerate faces
