IWYU SkeletalMeshSimulationComponent

diff --git a/Engine/Source/Runtime/Experimental/GeometryCollectionEngine/Private/GeometryCollection/SkeletalMeshSimulationComponent.cpp b/Engine/Source/Runtime/Experimental/GeometryCollectionEngine/Private/GeometryCollection/SkeletalMeshSimulationComponent.cpp
--- a/Engine/Source/Runtime/Experimental/GeometryCollectionEngine/Private/GeometryCollection/SkeletalMeshSimulationComponent.cpp
+++ b/Engine/Source/Runtime/Experimental/GeometryCollectionEngine/Private/GeometryCollection/SkeletalMeshSimulationComponent.cpp
@@ -28,6 +28,10 @@
 #include "Modules/ModuleManager.h"
 #include "Chaos/ChaosGameplayEventDispatcher.h"
 
+#include "Chaos/ChaosSolverActor.h"
+#include "PhysicalMaterials/Experimental/ChaosPhysicalMaterial.h"
+#include "Physics/Experimental/PhysScene_Chaos.h"
+
 
 //DEFINE_LOG_CATEGORY_STATIC(USkeletalMeshSimulationComponentLogging, NoLogging, All);
 
@@ -68,6 +72,14 @@ USkeletalMeshSimulationComponent::USkeletalMeshSimulationComponent(const FObject
 	ChaosMaterial = MakeUnique<Chaos::FChaosPhysicsMaterial>();
 }
 
+USkeletalMeshSimulationComponent::USkeletalMeshSimulationComponent(FVTableHelper& Helper)
+	: Super(Helper)
+{
+
+}
+
+USkeletalMeshSimulationComponent::~USkeletalMeshSimulationComponent() = default;
+
 Chaos::FPhysicsSolver* GetSolver(const USkeletalMeshSimulationComponent& SkeletalMeshSimulationComponent)
 {
 #if INCLUDE_CHAOS
diff --git a/Engine/Source/Runtime/Experimental/GeometryCollectionEngine/Public/GeometryCollection/SkeletalMeshSimulationComponent.h b/Engine/Source/Runtime/Experimental/GeometryCollectionEngine/Public/GeometryCollection/SkeletalMeshSimulationComponent.h
--- a/Engine/Source/Runtime/Experimental/GeometryCollectionEngine/Public/GeometryCollection/SkeletalMeshSimulationComponent.h
+++ b/Engine/Source/Runtime/Experimental/GeometryCollectionEngine/Public/GeometryCollection/SkeletalMeshSimulationComponent.h
@@ -2,17 +2,22 @@
 #pragma once
 
 #include "Components/MeshComponent.h"
-#include "Chaos/ChaosSolverActor.h"
 #include "Chaos/ChaosNotifyHandlerInterface.h"
 #include "GameFramework/Actor.h"
-#include "Physics/Experimental/PhysScene_Chaos.h"
-#include "PhysicalMaterials/Experimental/ChaosPhysicalMaterial.h"
 #include "GeometryCollection/GeometryCollectionSimulationTypes.h"
 
 #include "SkeletalMeshSimulationComponent.generated.h"
 
 class FSkeletalMeshSimulationComponentPhysicsProxy;
 
+namespace Chaos
+{
+class FChaosPhysicsMaterial;
+}
+
+class AChaosSolverActor;
+class UChaosPhysicalMaterial;
+
 /**
 *	USkeletalMeshSimulationComponent
 */
@@ -20,6 +25,8 @@ UCLASS(ClassGroup = Physics, Experimental, meta = (BlueprintSpawnableComponent))
 class GEOMETRYCOLLECTIONENGINE_API USkeletalMeshSimulationComponent : public UActorComponent, public IChaosNotifyHandlerInterface
 {
 	GENERATED_UCLASS_BODY()
+	USkeletalMeshSimulationComponent(FVTableHelper& Helper);
+	virtual ~USkeletalMeshSimulationComponent();
 
 public:
 	//
