IWYU StaticMeshSimulationComponent

diff --git a/Engine/Source/Runtime/Experimental/GeometryCollectionEngine/Private/GeometryCollection/StaticMeshSimulationComponent.cpp b/Engine/Source/Runtime/Experimental/GeometryCollectionEngine/Private/GeometryCollection/StaticMeshSimulationComponent.cpp
--- a/Engine/Source/Runtime/Experimental/GeometryCollectionEngine/Private/GeometryCollection/StaticMeshSimulationComponent.cpp
+++ b/Engine/Source/Runtime/Experimental/GeometryCollectionEngine/Private/GeometryCollection/StaticMeshSimulationComponent.cpp
@@ -17,6 +17,8 @@
 #include "PhysicsProxy/StaticMeshPhysicsProxy.h"
 #include "PhysicsSolver.h"
 #include "Chaos/ChaosGameplayEventDispatcher.h"
+#include "Chaos/ChaosSolverActor.h"
+#include "PhysicalMaterials/Experimental/ChaosPhysicalMaterial.h"
 
 
 DEFINE_LOG_CATEGORY_STATIC(UStaticMeshSimulationComponentLogging, NoLogging, All);
@@ -41,6 +43,14 @@ UStaticMeshSimulationComponent::UStaticMeshSimulationComponent(const FObjectInit
 	ChaosMaterial = MakeUnique<Chaos::FChaosPhysicsMaterial>();
 }
 
+UStaticMeshSimulationComponent::UStaticMeshSimulationComponent(FVTableHelper& Helper)
+	: Super(Helper)
+{
+
+}
+
+UStaticMeshSimulationComponent::~UStaticMeshSimulationComponent() = default;
+
 // We tick to detect components that unreal has moved, so we can update the solver.
 // The better solution long-term is to tie into UPrimitiveComponent::OnUpdateTransform() like we do for physx
 void UStaticMeshSimulationComponent::TickComponent(float DeltaTime, enum ELevelTick TickType, FActorComponentTickFunction *ThisTickFunction)
diff --git a/Engine/Source/Runtime/Experimental/GeometryCollectionEngine/Public/GeometryCollection/StaticMeshSimulationComponent.h b/Engine/Source/Runtime/Experimental/GeometryCollectionEngine/Public/GeometryCollection/StaticMeshSimulationComponent.h
--- a/Engine/Source/Runtime/Experimental/GeometryCollectionEngine/Public/GeometryCollection/StaticMeshSimulationComponent.h
+++ b/Engine/Source/Runtime/Experimental/GeometryCollectionEngine/Public/GeometryCollection/StaticMeshSimulationComponent.h
@@ -2,17 +2,20 @@
 
 #pragma once
 
-#include "Components/MeshComponent.h"
-#include "Chaos/ChaosSolverActor.h"
 #include "GameFramework/Actor.h"
-#include "Physics/Experimental/PhysScene_Chaos.h"
 #include "GeometryCollection/GeometryCollectionSimulationTypes.h"
-#include "PhysicalMaterials/Experimental/ChaosPhysicalMaterial.h"
 #include "Chaos/ChaosNotifyHandlerInterface.h"
 
 #include "StaticMeshSimulationComponent.generated.h"
 
 class FStaticMeshPhysicsProxy;
+class AChaosSolverActor;
+class UChaosPhysicalMaterial;
+
+namespace Chaos
+{
+class FChaosPhysicsMaterial;
+}
 
 /**
 *	UStaticMeshSimulationComponent
@@ -21,7 +24,8 @@ UCLASS(ClassGroup = Physics, Experimental, meta = (BlueprintSpawnableComponent))
 class GEOMETRYCOLLECTIONENGINE_API UStaticMeshSimulationComponent : public UActorComponent, public IChaosNotifyHandlerInterface
 {
 	GENERATED_UCLASS_BODY()
-
+	UStaticMeshSimulationComponent(FVTableHelper& Helper);
+	virtual ~UStaticMeshSimulationComponent();
 public:
 
 	virtual void TickComponent(float DeltaTime, enum ELevelTick TickType, FActorComponentTickFunction *ThisTickFunction);
