Missing include for IWYU

diff --git a/Engine/Source/Runtime/Experimental/Chaos/Public/Chaos/PBDAnimDriveConstraint.h b/Engine/Source/Runtime/Experimental/Chaos/Public/Chaos/PBDAnimDriveConstraint.h
--- a/Engine/Source/Runtime/Experimental/Chaos/Public/Chaos/PBDAnimDriveConstraint.h
+++ b/Engine/Source/Runtime/Experimental/Chaos/Public/Chaos/PBDAnimDriveConstraint.h
@@ -2,6 +2,7 @@
 #pragma once
 
 #include "Chaos/ParticleRule.h"
+#include "Chaos/Framework/Parallel.h"
 
 namespace Chaos
 {
