Missing forward declare and include

diff --git a/Engine/Source/Runtime/Experimental/GeometryCollectionEngine/Public/GeometryCollection/StaticMeshSimulationComponent.h b/Engine/Source/Runtime/Experimental/GeometryCollectionEngine/Public/GeometryCollection/StaticMeshSimulationComponent.h
--- a/Engine/Source/Runtime/Experimental/GeometryCollectionEngine/Public/GeometryCollection/StaticMeshSimulationComponent.h
+++ b/Engine/Source/Runtime/Experimental/GeometryCollectionEngine/Public/GeometryCollection/StaticMeshSimulationComponent.h
@@ -11,6 +11,7 @@
 class FStaticMeshPhysicsProxy;
 class AChaosSolverActor;
 class UChaosPhysicalMaterial;
+class FPhysScene_Chaos;
 
 namespace Chaos
 {
