Added culled evaluation list to field evaluation to reduct the cost of applying fields and avoid pulling in uninitialized data during the field application. Culled results currently only applied to disable fields.
#rnx

diff --git a/Engine/Source/Runtime/Experimental/ChaosSolvers/Private/PhysicsProxy/FieldSystemPhysicsProxy.cpp b/Engine/Source/Runtime/Experimental/ChaosSolvers/Private/PhysicsProxy/FieldSystemPhysicsProxy.cpp
--- a/Engine/Source/Runtime/Experimental/ChaosSolvers/Private/PhysicsProxy/FieldSystemPhysicsProxy.cpp
+++ b/Engine/Source/Runtime/Experimental/ChaosSolvers/Private/PhysicsProxy/FieldSystemPhysicsProxy.cpp
@@ -577,21 +577,19 @@ void FFieldSystemPhysicsProxy::FieldParameterUpdateCallback(
 					{
 						TArrayView<FVector> SamplePointsView(&(SamplePoints[0]), SamplePoints.Num());
 						TArrayView<ContextIndex> SampleIndicesView(&(SampleIndices[0]), SampleIndices.Num());
-						FFieldContext Context(
-							SampleIndicesView,
-							SamplePointsView,
-							Command.MetaData);
+
+						FFieldContext Context(SampleIndicesView, SamplePointsView, Command.MetaData);
 
 						TArray<float> LocalResults;
 						LocalResults.AddUninitialized(Handles.Num());
+
 						TArrayView<float> ResultsView(&(LocalResults[0]), LocalResults.Num());
-						static_cast<const FFieldNode<float> *>(Command.RootNode.Get())->Evaluate(
-							Context, ResultsView);
+						
+						static_cast<const FFieldNode<float> *>(Command.RootNode.Get())->Evaluate(Context, ResultsView);
 
-						int32 i = 0;
-						for (Chaos::TGeometryParticleHandle<float, 3>* Handle : Handles)
+						for(const ContextIndex& Index : Context.GetEvaluatedSamples())
 						{
-							Chaos::TPBDRigidParticleHandle<float, 3>* RigidHandle = Handle->CastToRigidParticle();
+							Chaos::TPBDRigidParticleHandle<float, 3>* RigidHandle = Handles[Index.Sample]->CastToRigidParticle();
 							
 							if (RigidHandle && RigidHandle->ObjectState() == Chaos::EObjectStateType::Dynamic && ResultsView.Num() > 0)
 							{						
@@ -600,9 +598,8 @@ void FFieldSystemPhysicsProxy::FieldParameterUpdateCallback(
 								{
 
 									TUniquePtr<Chaos::FChaosPhysicsMaterial> NewMaterial = MakeUnique< Chaos::FChaosPhysicsMaterial>();
-									NewMaterial->DisabledLinearThreshold = ResultsView[i];
-									NewMaterial->DisabledAngularThreshold = ResultsView[i];
-
+									NewMaterial->DisabledLinearThreshold = ResultsView[Index.Result];
+									NewMaterial->DisabledAngularThreshold = ResultsView[Index.Result];
 
 									CurrentSolver->GetEvolution()->SetPhysicsMaterial(RigidHandle, MakeSerializable(NewMaterial));
 									CurrentSolver->GetEvolution()->SetPerParticlePhysicsMaterial(RigidHandle, NewMaterial);
@@ -611,14 +608,13 @@ void FFieldSystemPhysicsProxy::FieldParameterUpdateCallback(
 								{
 									const TUniquePtr<FChaosPhysicsMaterial> &InstanceMaterial = CurrentSolver->GetEvolution()->GetPerParticlePhysicsMaterial(RigidHandle);
 		
-									if (ResultsView[i] != InstanceMaterial->DisabledLinearThreshold)
+									if (ResultsView[Index.Result] != InstanceMaterial->DisabledLinearThreshold)
 									{
-										InstanceMaterial->DisabledLinearThreshold = ResultsView[i];
-										InstanceMaterial->DisabledAngularThreshold = ResultsView[i];
+										InstanceMaterial->DisabledLinearThreshold = ResultsView[Index.Result];
+										InstanceMaterial->DisabledAngularThreshold = ResultsView[Index.Result];
 									}
 								}
 							}
-							++i;
 						}
 					}									
 				}
diff --git a/Engine/Source/Runtime/Experimental/FieldSystem/Source/FieldSystemCore/Private/Field/FieldSystemNodes.cpp b/Engine/Source/Runtime/Experimental/FieldSystem/Source/FieldSystemCore/Private/Field/FieldSystemNodes.cpp
--- a/Engine/Source/Runtime/Experimental/FieldSystem/Source/FieldSystemCore/Private/Field/FieldSystemNodes.cpp
+++ b/Engine/Source/Runtime/Experimental/FieldSystem/Source/FieldSystemCore/Private/Field/FieldSystemNodes.cpp
@@ -61,7 +61,7 @@ FFieldNodeBase * FieldNodeFactory(FFieldNodeBase::EFieldType BaseType,FFieldNode
 
 
 template<class T>
-void SerializeInternal(FArchive& Ar, TUniquePtr< FFieldNode<T> > & Field)
+void SerializeInternal(FArchive& Ar, TUniquePtr< FFieldNode<T> >& Field)
 {
 	uint8 dType = Field.IsValid() ? (uint8)Field->Type() : (uint8)FFieldNodeBase::EFieldType::EField_None;
 	Ar << dType;
@@ -80,7 +80,7 @@ void SerializeInternal(FArchive& Ar, TUniquePtr< FFieldNode<T> > & Field)
 	}
 }
 template<class Enum>
-void SerializeInternal(FArchive& Ar, Enum & Var)
+void SerializeInternal(FArchive& Ar, Enum& Var)
 {
 	uint8 Type = (uint8)Var;
 	Ar << Type;
@@ -91,7 +91,7 @@ void SerializeInternal(FArchive& Ar, Enum & Var)
 /**
 * FUniformInteger
 */
-void FUniformInteger::Evaluate(const FFieldContext & Context, TArrayView<int32> & Results) const
+void FUniformInteger::Evaluate(FFieldContext& Context, TArrayView<int32>& Results) const
 {
 	int32 MagnitudeVal = Magnitude;
 
@@ -120,38 +120,45 @@ bool FUniformInteger::operator==(const FFieldNodeBase& Node)
 /**
 * FRadialIntMask
 */
-void FRadialIntMask::Evaluate(const FFieldContext & Context, TArrayView<int32> & Results) const
+void FRadialIntMask::Evaluate(FFieldContext& Context, TArrayView<int32>& Results) const
 {
 	float Radius2 = Radius * Radius;
 
 	int32 NumSamples = Context.SampleIndices.Num();
 	for (int32 SampleIndex = 0; SampleIndex < NumSamples; SampleIndex++)
 	{
-		const ContextIndex & Index = Context.SampleIndices[SampleIndex];
+		const ContextIndex& Index = Context.SampleIndices[SampleIndex];
 		{
 			float Result;
 			float Delta2 = (Position - Context.Samples[Index.Sample]).SizeSquared();
-			if (Delta2 < Radius2)
+
+			if(Delta2 < Radius2)
+			{
 				Result = InteriorValue;
+			}
 			else
+			{
 				Result = ExteriorValue;
+			}
 
-			switch (SetMaskCondition) {
+			switch (SetMaskCondition) 
+			{
 			case ESetMaskConditionType::Field_Set_Always:
 				Results[Index.Result] = Result;
 				break;
 			case ESetMaskConditionType::Field_Set_IFF_NOT_Interior:
-				if (Results[Index.Result] != InteriorValue) {
+				if (Results[Index.Result] != InteriorValue) 
+				{
 					Results[Index.Result] = Result;
 				}
 				break;
 			case ESetMaskConditionType::Field_Set_IFF_NOT_Exterior:
-				if (Results[Index.Result] != ExteriorValue) {
+				if (Results[Index.Result] != ExteriorValue) 
+				{
 					Results[Index.Result] = Result;
 				}
 				break;
 			}
-		
 		}
 	}
 }
@@ -182,7 +189,7 @@ bool FRadialIntMask::operator==(const FFieldNodeBase& Node)
 /**
 * FUniformScalar
 */
-void FUniformScalar::Evaluate(const FFieldContext & Context, TArrayView<float> & Results) const
+void FUniformScalar::Evaluate(FFieldContext& Context, TArrayView<float>& Results) const
 {
 	float MagnitudeVal = Magnitude;
 
@@ -212,7 +219,7 @@ bool FUniformScalar::operator==(const FFieldNodeBase& Node)
 * FRadialFalloff
 */
 template<>
-void FRadialFalloff::Evaluator<EFieldFalloffType::Field_FallOff_None>(const FFieldContext & Context, TArrayView<float> & Results) const
+void FRadialFalloff::Evaluator<EFieldFalloffType::Field_FallOff_None>(const FFieldContext& Context, TArrayView<float>& Results) const
 {
 	float Radius2 = Radius * Radius;
 
@@ -221,7 +228,7 @@ void FRadialFalloff::Evaluator<EFieldFalloffType::Field_FallOff_None>(const FFie
 		int32 NumSamples = Context.SampleIndices.Num();
 		for (int32 SampleIndex = 0; SampleIndex < NumSamples; SampleIndex++)
 		{
-			const ContextIndex & Index = Context.SampleIndices[SampleIndex];
+			const ContextIndex& Index = Context.SampleIndices[SampleIndex];
 			{
 				Results[Index.Result] = Default;
 				float Delta2 = (Position - Context.Samples[Index.Sample]).SizeSquared();
@@ -234,7 +241,7 @@ void FRadialFalloff::Evaluator<EFieldFalloffType::Field_FallOff_None>(const FFie
 	}
 }
 template<>
-void FRadialFalloff::Evaluator<EFieldFalloffType::Field_Falloff_Linear>(const FFieldContext & Context, TArrayView<float> & Results) const
+void FRadialFalloff::Evaluator<EFieldFalloffType::Field_Falloff_Linear>(const FFieldContext& Context, TArrayView<float>& Results) const
 {
 	float Radius2 = Radius * Radius;
 	if (Radius2 > 0.f)
@@ -245,7 +252,7 @@ void FRadialFalloff::Evaluator<EFieldFalloffType::Field_Falloff_Linear>(const FF
 		int32 NumSamples = Context.SampleIndices.Num();
 		for (int32 SampleIndex = 0; SampleIndex < NumSamples; SampleIndex++)
 		{
-			const ContextIndex & Index = Context.SampleIndices[SampleIndex];
+			const ContextIndex& Index = Context.SampleIndices[SampleIndex];
 			{
 				Results[Index.Result] = Default;
 				float Delta2 = (Position - Context.Samples[Index.Sample]).SizeSquared();
@@ -270,12 +277,12 @@ void FRadialFalloff::Evaluator<EFieldFalloffType::Field_Falloff_Linear>(const FF
 	}
 }
 template<>
-void FRadialFalloff::Evaluator<EFieldFalloffType::Field_Falloff_Squared>(const FFieldContext & Context, TArrayView<float> & Results) const
+void FRadialFalloff::Evaluator<EFieldFalloffType::Field_Falloff_Squared>(const FFieldContext& Context, TArrayView<float>& Results) const
 {
 	int32 NumSamples = Context.SampleIndices.Num();
 	for (int32 SampleIndex = 0; SampleIndex < NumSamples; SampleIndex++)
 	{
-		const ContextIndex & Index = Context.SampleIndices[SampleIndex];
+		const ContextIndex& Index = Context.SampleIndices[SampleIndex];
 		{
 			Results[Index.Result] = Default;
 			float Delta = (Position - Context.Samples[Index.Sample]).Size();
@@ -288,12 +295,12 @@ void FRadialFalloff::Evaluator<EFieldFalloffType::Field_Falloff_Squared>(const F
 	}
 }
 template<>
-void FRadialFalloff::Evaluator<EFieldFalloffType::Field_Falloff_Inverse>(const FFieldContext & Context, TArrayView<float> & Results) const
+void FRadialFalloff::Evaluator<EFieldFalloffType::Field_Falloff_Inverse>(const FFieldContext& Context, TArrayView<float>& Results) const
 {
 	int32 NumSamples = Context.SampleIndices.Num();
 	for (int32 SampleIndex = 0; SampleIndex < NumSamples; SampleIndex++)
 	{
-		const ContextIndex & Index = Context.SampleIndices[SampleIndex];
+		const ContextIndex& Index = Context.SampleIndices[SampleIndex];
 		{
 			Results[Index.Result] = Default;
 			float Delta = (Position - Context.Samples[Index.Sample]).Size();
@@ -309,12 +316,12 @@ void FRadialFalloff::Evaluator<EFieldFalloffType::Field_Falloff_Inverse>(const F
 	}
 }
 template<>
-void FRadialFalloff::Evaluator<EFieldFalloffType::Field_Falloff_Logarithmic>(const FFieldContext & Context, TArrayView<float> & Results) const
+void FRadialFalloff::Evaluator<EFieldFalloffType::Field_Falloff_Logarithmic>(const FFieldContext& Context, TArrayView<float>& Results) const
 {
 	int32 NumSamples = Context.SampleIndices.Num();
 	for (int32 SampleIndex = 0; SampleIndex < NumSamples; SampleIndex++)
 	{
-		const ContextIndex & Index = Context.SampleIndices[SampleIndex];
+		const ContextIndex& Index = Context.SampleIndices[SampleIndex];
 		{
 			Results[Index.Result] = Default;
 			float Delta = (Position - Context.Samples[Index.Sample]).Size();
@@ -326,7 +333,7 @@ void FRadialFalloff::Evaluator<EFieldFalloffType::Field_Falloff_Logarithmic>(con
 		}
 	}
 }
-void FRadialFalloff::Evaluate(const FFieldContext & Context, TArrayView<float> & Results) const
+void FRadialFalloff::Evaluate(FFieldContext& Context, TArrayView<float>& Results) const
 {
 	switch (Falloff)
 	{
@@ -383,12 +390,12 @@ bool FRadialFalloff::operator==(const FFieldNodeBase& Node)
 * FPlaneFalloff
 */
 template<>
-void FPlaneFalloff::Evaluator<EFieldFalloffType::Field_FallOff_None>(const FFieldContext & Context, const FPlane & Plane, TArrayView<float> & Results) const
+void FPlaneFalloff::Evaluator<EFieldFalloffType::Field_FallOff_None>(const FFieldContext& Context, const FPlane& Plane, TArrayView<float>& Results) const
 {
 	int32 NumSamples = Context.SampleIndices.Num();
 	for (int32 SampleIndex = 0; SampleIndex < NumSamples; SampleIndex++)
 	{
-		const ContextIndex & Index = Context.SampleIndices[SampleIndex];
+		const ContextIndex& Index = Context.SampleIndices[SampleIndex];
 		{
 			Results[Index.Result] = Default;
 			float LocalDistance = Plane.PlaneDot(Context.Samples[Index.Sample]);
@@ -400,7 +407,7 @@ void FPlaneFalloff::Evaluator<EFieldFalloffType::Field_FallOff_None>(const FFiel
 	}
 }
 template<>
-void FPlaneFalloff::Evaluator<EFieldFalloffType::Field_Falloff_Linear>(const FFieldContext & Context, const FPlane & Plane, TArrayView<float> & Results) const
+void FPlaneFalloff::Evaluator<EFieldFalloffType::Field_Falloff_Linear>(const FFieldContext& Context, const FPlane& Plane, TArrayView<float>& Results) const
 {
 	float Scalar = (MaxRange - MinRange);
 	bool bNonUnitScalar = FMath::Abs(Scalar - 1.f) > KINDA_SMALL_NUMBER;
@@ -408,7 +415,7 @@ void FPlaneFalloff::Evaluator<EFieldFalloffType::Field_Falloff_Linear>(const FFi
 	int32 NumSamples = Context.SampleIndices.Num();
 	for (int32 SampleIndex = 0; SampleIndex < NumSamples; SampleIndex++)
 	{
-		const ContextIndex & Index = Context.SampleIndices[SampleIndex];
+		const ContextIndex& Index = Context.SampleIndices[SampleIndex];
 		{
 			Results[Index.Result] = Default;
 			float LocalDistance = Plane.PlaneDot(Context.Samples[Index.Sample]);
@@ -425,12 +432,12 @@ void FPlaneFalloff::Evaluator<EFieldFalloffType::Field_Falloff_Linear>(const FFi
 	}
 }
 template<>
-void FPlaneFalloff::Evaluator<EFieldFalloffType::Field_Falloff_Squared>(const FFieldContext & Context, const FPlane & Plane, TArrayView<float> & Results) const
+void FPlaneFalloff::Evaluator<EFieldFalloffType::Field_Falloff_Squared>(const FFieldContext& Context, const FPlane& Plane, TArrayView<float>& Results) const
 {
 	int32 NumSamples = Context.SampleIndices.Num();
 	for (int32 SampleIndex = 0; SampleIndex < NumSamples; SampleIndex++)
 	{
-		const ContextIndex & Index = Context.SampleIndices[SampleIndex];
+		const ContextIndex& Index = Context.SampleIndices[SampleIndex];
 		{
 			Results[Index.Result] = Default;
 			float LocalDistance = Plane.PlaneDot(Context.Samples[Index.Sample]);
@@ -443,14 +450,14 @@ void FPlaneFalloff::Evaluator<EFieldFalloffType::Field_Falloff_Squared>(const FF
 	}
 }
 template<>
-void FPlaneFalloff::Evaluator<EFieldFalloffType::Field_Falloff_Inverse>(const FFieldContext & Context, const FPlane & Plane, TArrayView<float> & Results) const
+void FPlaneFalloff::Evaluator<EFieldFalloffType::Field_Falloff_Inverse>(const FFieldContext& Context, const FPlane& Plane, TArrayView<float>& Results) const
 {
 	float NegMagnitude = -Magnitude;
 
 	int32 NumSamples = Context.SampleIndices.Num();
 	for (int32 SampleIndex = 0; SampleIndex < NumSamples; SampleIndex++)
 	{
-		const ContextIndex & Index = Context.SampleIndices[SampleIndex];
+		const ContextIndex& Index = Context.SampleIndices[SampleIndex];
 		{
 			Results[Index.Result] = Default;
 			float LocalDistance = Plane.PlaneDot(Context.Samples[Index.Sample]);
@@ -463,12 +470,12 @@ void FPlaneFalloff::Evaluator<EFieldFalloffType::Field_Falloff_Inverse>(const FF
 	}
 }
 template<>
-void FPlaneFalloff::Evaluator<EFieldFalloffType::Field_Falloff_Logarithmic>(const FFieldContext & Context, const FPlane & Plane, TArrayView<float> & Results) const
+void FPlaneFalloff::Evaluator<EFieldFalloffType::Field_Falloff_Logarithmic>(const FFieldContext& Context, const FPlane& Plane, TArrayView<float>& Results) const
 {
 	int32 NumSamples = Context.SampleIndices.Num();
 	for (int32 SampleIndex = 0; SampleIndex < NumSamples; SampleIndex++)
 	{
-		const ContextIndex & Index = Context.SampleIndices[SampleIndex];
+		const ContextIndex& Index = Context.SampleIndices[SampleIndex];
 		{
 			Results[Index.Result] = Default;
 			float LocalDistance = Plane.PlaneDot(Context.Samples[Index.Sample]);
@@ -481,7 +488,7 @@ void FPlaneFalloff::Evaluator<EFieldFalloffType::Field_Falloff_Logarithmic>(cons
 	}
 }
 void
-FPlaneFalloff::Evaluate(const FFieldContext & Context, TArrayView<float> & Results) const
+FPlaneFalloff::Evaluate(FFieldContext& Context, TArrayView<float>& Results) const
 {
 	FPlane Plane(Position, Normal);
 	switch (Falloff)
@@ -539,13 +546,13 @@ bool FPlaneFalloff::operator==(const FFieldNodeBase& Node)
 * FBoxFalloff
 */
 template<>
-void FBoxFalloff::Evaluator<EFieldFalloffType::Field_FallOff_None>(const FFieldContext & Context, TArrayView<float> & Results) const
+void FBoxFalloff::Evaluator<EFieldFalloffType::Field_FallOff_None>(const FFieldContext& Context, TArrayView<float>& Results) const
 {
 	int32 NumSamples = Context.SampleIndices.Num();
 	FBox UnitBox(FVector(-50), FVector(50));
 	for (int32 SampleIndex = 0; SampleIndex < NumSamples; SampleIndex++)
 	{
-		const ContextIndex & Index = Context.SampleIndices[SampleIndex];
+		const ContextIndex& Index = Context.SampleIndices[SampleIndex];
 		{
 			Results[Index.Result] = Default;
 
@@ -558,7 +565,7 @@ void FBoxFalloff::Evaluator<EFieldFalloffType::Field_FallOff_None>(const FFieldC
 	}
 }
 template<>
-void FBoxFalloff::Evaluator<EFieldFalloffType::Field_Falloff_Linear>(const FFieldContext & Context, TArrayView<float> & Results) const
+void FBoxFalloff::Evaluator<EFieldFalloffType::Field_Falloff_Linear>(const FFieldContext& Context, TArrayView<float>& Results) const
 {
 	float Scalar = (MaxRange - MinRange);
 	bool bNonUnitScalar = FMath::Abs(Scalar - 1.f) > KINDA_SMALL_NUMBER;
@@ -567,7 +574,7 @@ void FBoxFalloff::Evaluator<EFieldFalloffType::Field_Falloff_Linear>(const FFiel
 	FBox UnitBox(FVector(-50), FVector(50));
 	for (int32 SampleIndex = 0; SampleIndex < NumSamples; SampleIndex++)
 	{
-		const ContextIndex & Index = Context.SampleIndices[SampleIndex];
+		const ContextIndex& Index = Context.SampleIndices[SampleIndex];
 		{
 			Results[Index.Result] = Default;
 
@@ -594,13 +601,13 @@ void FBoxFalloff::Evaluator<EFieldFalloffType::Field_Falloff_Linear>(const FFiel
 		}
 	}
 }template<>
-void FBoxFalloff::Evaluator<EFieldFalloffType::Field_Falloff_Squared>(const FFieldContext & Context, TArrayView<float> & Results) const
+void FBoxFalloff::Evaluator<EFieldFalloffType::Field_Falloff_Squared>(const FFieldContext& Context, TArrayView<float>& Results) const
 {
 	int32 NumSamples = Context.SampleIndices.Num();
 	FBox UnitBox(FVector(-50), FVector(50));
 	for (int32 SampleIndex = 0; SampleIndex < NumSamples; SampleIndex++)
 	{
-		const ContextIndex & Index = Context.SampleIndices[SampleIndex];
+		const ContextIndex& Index = Context.SampleIndices[SampleIndex];
 		{
 			Results[Index.Result] = Default;
 
@@ -614,13 +621,13 @@ void FBoxFalloff::Evaluator<EFieldFalloffType::Field_Falloff_Squared>(const FFie
 	}
 }
 template<>
-void FBoxFalloff::Evaluator<EFieldFalloffType::Field_Falloff_Inverse>(const FFieldContext & Context, TArrayView<float> & Results) const
+void FBoxFalloff::Evaluator<EFieldFalloffType::Field_Falloff_Inverse>(const FFieldContext& Context, TArrayView<float>& Results) const
 {
 	int32 NumSamples = Context.SampleIndices.Num();
 	FBox UnitBox(FVector(-50), FVector(50));
 	for (int32 SampleIndex = 0; SampleIndex < NumSamples; SampleIndex++)
 	{
-		const ContextIndex & Index = Context.SampleIndices[SampleIndex];
+		const ContextIndex& Index = Context.SampleIndices[SampleIndex];
 		{
 			FVector LocalPoint = Transform.InverseTransformPosition(Context.Samples[Index.Sample]);
 			if (UnitBox.IsInside(LocalPoint))
@@ -637,13 +644,13 @@ void FBoxFalloff::Evaluator<EFieldFalloffType::Field_Falloff_Inverse>(const FFie
 	}
 }
 template<>
-void FBoxFalloff::Evaluator<EFieldFalloffType::Field_Falloff_Logarithmic>(const FFieldContext & Context, TArrayView<float> & Results) const
+void FBoxFalloff::Evaluator<EFieldFalloffType::Field_Falloff_Logarithmic>(const FFieldContext& Context, TArrayView<float>& Results) const
 {
 	int32 NumSamples = Context.SampleIndices.Num();
 	FBox UnitBox(FVector(-50), FVector(50));
 	for (int32 SampleIndex = 0; SampleIndex < NumSamples; SampleIndex++)
 	{
-		const ContextIndex & Index = Context.SampleIndices[SampleIndex];
+		const ContextIndex& Index = Context.SampleIndices[SampleIndex];
 		{
 			Results[Index.Result] = Default;
 
@@ -657,7 +664,7 @@ void FBoxFalloff::Evaluator<EFieldFalloffType::Field_Falloff_Logarithmic>(const
 	}
 }
 void
-FBoxFalloff::Evaluate(const FFieldContext & Context, TArrayView<float> & Results) const
+FBoxFalloff::Evaluate(FFieldContext& Context, TArrayView<float>& Results) const
 {
 	switch (Falloff)
 	{
@@ -710,7 +717,7 @@ bool FBoxFalloff::operator==(const FFieldNodeBase& Node)
 * FNoiseField
 */
 void
-FNoiseField::Evaluate(const FFieldContext & Context, TArrayView<float> & Results) const
+FNoiseField::Evaluate(FFieldContext& Context, TArrayView<float>& Results) const
 {
 	int32 NumSamples = Context.SampleIndices.Num();
 
@@ -718,7 +725,7 @@ FNoiseField::Evaluate(const FFieldContext & Context, TArrayView<float> & Results
 	FBox Bounds(ForceInit);
 	for (int32 SampleIndex = 0; SampleIndex < NumSamples; SampleIndex++)
 	{
-		const ContextIndex & Index = Context.SampleIndices[SampleIndex];
+		const ContextIndex& Index = Context.SampleIndices[SampleIndex];
 		Bounds += Context.Samples[Index.Sample];
 	}
 	float Scale = 0.5;
@@ -728,7 +735,7 @@ FNoiseField::Evaluate(const FFieldContext & Context, TArrayView<float> & Results
 	float MinResult = FLT_MAX, MaxResult = -FLT_MAX;
 	for (int32 SampleIndex = 0; SampleIndex < NumSamples; SampleIndex++)
 	{
-		const ContextIndex & Index = Context.SampleIndices[SampleIndex];
+		const ContextIndex& Index = Context.SampleIndices[SampleIndex];
 		FVector Delta = Transform.TransformVector(Context.Samples[Index.Sample] - Bounds.Min);
 
 		float i = BoundsSize.X? (Delta.X / BoundsSize.X) : 0;
@@ -749,7 +756,7 @@ FNoiseField::Evaluate(const FFieldContext & Context, TArrayView<float> & Results
 	{
 		for (int32 SampleIndex = 0; SampleIndex < NumSamples; SampleIndex++)
 		{
-			const ContextIndex & Index = Context.SampleIndices[SampleIndex];
+			const ContextIndex& Index = Context.SampleIndices[SampleIndex];
 			Results[Index.Result] = Ratio * (Results[Index.Result] - MinResult) + MinRange;
 		}
 	}
@@ -778,7 +785,7 @@ bool FNoiseField::operator==(const FFieldNodeBase& Node)
 /**
 * FUniformVector
 */
-void FUniformVector::Evaluate(const FFieldContext & Context, TArrayView<FVector> & Results) const
+void FUniformVector::Evaluate(FFieldContext& Context, TArrayView<FVector>& Results) const
 {
 	FVector DirectionVal = Direction;
 	float MagnitudeVal =Magnitude;
@@ -810,12 +817,12 @@ bool FUniformVector::operator==(const FFieldNodeBase& Node)
 /**
 * FRadialVector
 */
-void FRadialVector::Evaluate(const FFieldContext & Context, TArrayView<FVector> & Results) const
+void FRadialVector::Evaluate(FFieldContext& Context, TArrayView<FVector>& Results) const
 {
 	int32 NumSamples = Context.SampleIndices.Num();
 	for (int32 SampleIndex = 0; SampleIndex < NumSamples; SampleIndex++)
 	{
-		const ContextIndex & Index = Context.SampleIndices[SampleIndex];
+		const ContextIndex& Index = Context.SampleIndices[SampleIndex];
 		{
 			
 			Results[Index.Result] = Magnitude * (Context.Samples[Index.Sample] - Position).GetSafeNormal();
@@ -844,12 +851,12 @@ bool FRadialVector::operator==(const FFieldNodeBase& Node)
 /**
 * FRandomVector
 */
-void FRandomVector::Evaluate(const FFieldContext & Context, TArrayView<FVector> & Results) const
+void FRandomVector::Evaluate(FFieldContext& Context, TArrayView<FVector>& Results) const
 {
 	int32 NumSamples = Context.SampleIndices.Num();
 	for (int32 SampleIndex = 0; SampleIndex < NumSamples; SampleIndex++)
 	{
-		const ContextIndex & Index = Context.SampleIndices[SampleIndex];
+		const ContextIndex& Index = Context.SampleIndices[SampleIndex];
 		{
 			Results[Index.Result].X = FMath::RandRange(-1.f, 1.f);
 			Results[Index.Result].Y = FMath::RandRange(-1.f, 1.f);
@@ -878,14 +885,13 @@ bool FRandomVector::operator==(const FFieldNodeBase& Node)
 /**
 * FSumScalar
 */
-void FSumScalar::Evaluate(const FFieldContext & ContextIn, TArrayView<float> & Results) const
+void FSumScalar::Evaluate(FFieldContext& ContextIn, TArrayView<float>& Results) const
 {
-	FFieldContext Context = ContextIn;
 	TUniquePtr<FFieldSystemMetaDataResults<float> > ResultsData(new FFieldSystemMetaDataResults<float>(Results));
-	Context.MetaData.Add(FFieldSystemMetaData::EMetaType::ECommandData_Results, ResultsData.Get());
+	FScopedFieldContextMetaData ScopedMetaData(ContextIn, ResultsData.Get());
 
-	int32 NumResults = Results.Num();
-	int32 NumSamples = Context.SampleIndices.Num();
+	const int32 NumResults = Results.Num();
+	const int32 NumSamples = ContextIn.SampleIndices.Num();
 
 	const FFieldNode<float> * RightField = ScalarRight.Get();
 	const FFieldNode<float> * LeftField = ScalarLeft.Get();
@@ -893,12 +899,13 @@ void FSumScalar::Evaluate(const FFieldContext & ContextIn, TArrayView<float> & R
 	float MagnitudeVal = Magnitude;
 	if (LeftField != nullptr && RightField != nullptr)
 	{
-		TArray<float> Buffer;
-		Buffer.SetNumUninitialized(2 * NumResults);
+		TArray<float> ResultsBuffer;
+		ResultsBuffer.SetNumUninitialized(2 * NumResults);
 		TArrayView<float> Buffers[2] = {
-			TArrayView<float>(&Buffer[0],NumResults),
-			TArrayView<float>(&Buffer[NumResults],NumResults),
+			TArrayView<float>(&ResultsBuffer[0],NumResults),
+			TArrayView<float>(&ResultsBuffer[NumResults],NumResults),
 		};
+
 		TArray<const FFieldNode<float> * > FieldNodes = { LeftField,RightField };
 
 		for (int32 i = 0; i < 2; i++)
@@ -906,7 +913,7 @@ void FSumScalar::Evaluate(const FFieldContext & ContextIn, TArrayView<float> & R
 			if (ensureMsgf(FieldNodes[i]->Type() == FFieldNode<float>::StaticType(),
 				TEXT("Field system SumScalar expects float input arrays.")))
 			{
-				FieldNodes[i]->Evaluate(Context, Buffers[i]);
+				FieldNodes[i]->Evaluate(ContextIn, Buffers[i]);
 			}
 		}
 
@@ -915,7 +922,7 @@ void FSumScalar::Evaluate(const FFieldContext & ContextIn, TArrayView<float> & R
 		case EFieldOperationType::Field_Multiply:
 			for (int32 SampleIndex = 0; SampleIndex < NumSamples; SampleIndex++)
 			{
-				const ContextIndex & Index = Context.SampleIndices[SampleIndex];
+				const ContextIndex& Index = ContextIn.SampleIndices[SampleIndex];
 				{
 					Results[Index.Result] = Buffers[1][Index.Result] * Buffers[0][Index.Result];
 				}
@@ -924,7 +931,7 @@ void FSumScalar::Evaluate(const FFieldContext & ContextIn, TArrayView<float> & R
 		case EFieldOperationType::Field_Divide:
 			for (int32 SampleIndex = 0; SampleIndex < NumSamples; SampleIndex++)
 			{
-				const ContextIndex & Index = Context.SampleIndices[SampleIndex];
+				const ContextIndex& Index = ContextIn.SampleIndices[SampleIndex];
 				{
 					Results[Index.Result] = Buffers[0][Index.Result] / Buffers[1][Index.Result];
 				}
@@ -933,7 +940,7 @@ void FSumScalar::Evaluate(const FFieldContext & ContextIn, TArrayView<float> & R
 		case EFieldOperationType::Field_Add:
 			for (int32 SampleIndex = 0; SampleIndex < NumSamples; SampleIndex++)
 			{
-				const ContextIndex & Index = Context.SampleIndices[SampleIndex];
+				const ContextIndex& Index = ContextIn.SampleIndices[SampleIndex];
 				{
 					Results[Index.Result] = Buffers[1][Index.Result] + Buffers[0][Index.Result];
 				}
@@ -942,7 +949,7 @@ void FSumScalar::Evaluate(const FFieldContext & ContextIn, TArrayView<float> & R
 		case EFieldOperationType::Field_Substract:
 			for (int32 SampleIndex = 0; SampleIndex < NumSamples; SampleIndex++)
 			{
-				const ContextIndex & Index = Context.SampleIndices[SampleIndex];
+				const ContextIndex& Index = ContextIn.SampleIndices[SampleIndex];
 				{
 					Results[Index.Result] = Buffers[0][Index.Result] - Buffers[1][Index.Result];
 				}
@@ -953,20 +960,20 @@ void FSumScalar::Evaluate(const FFieldContext & ContextIn, TArrayView<float> & R
 	else if (LeftField != nullptr && ensureMsgf(ScalarLeft->Type() == FFieldNode<float>::StaticType(),
 		TEXT("Field system SumScalar expects float input arrays.")))
 	{
-		LeftField->Evaluate(Context, Results);
+		LeftField->Evaluate(ContextIn, Results);
 
 	}
 	else if (ScalarRight != nullptr &&  ensureMsgf(ScalarRight->Type() == FFieldNode<float>::StaticType(),
 		TEXT("Field system SumScalar expects float input arrays.")))
 	{
-		ScalarRight->Evaluate(Context, Results);
+		ScalarRight->Evaluate(ContextIn, Results);
 	}
 
 	if (MagnitudeVal != 1.0)
 	{
 		for (int32 SampleIndex = 0; SampleIndex < NumSamples; SampleIndex++)
 		{
-			const ContextIndex & Index = Context.SampleIndices[SampleIndex];
+			const ContextIndex& Index = ContextIn.SampleIndices[SampleIndex];
 			{
 				Results[Index.Result] *= MagnitudeVal;
 			}
@@ -999,14 +1006,13 @@ bool FSumScalar::operator==(const FFieldNodeBase& Node)
 /**
 * FSumVector
 */
-void FSumVector::Evaluate(const FFieldContext & ContextIn, TArrayView<FVector> & Results) const
+void FSumVector::Evaluate(FFieldContext& ContextIn, TArrayView<FVector>& Results) const
 {
-	FFieldContext Context = ContextIn;
-	TUniquePtr<FFieldSystemMetaDataResults<FVector> > ResultsData(new FFieldSystemMetaDataResults<FVector>(Results));
-	Context.MetaData.Add(FFieldSystemMetaData::EMetaType::ECommandData_Results, ResultsData.Get());
+	TUniquePtr<FFieldSystemMetaDataResults<FVector>> ResultsData(new FFieldSystemMetaDataResults<FVector>(Results));
+	FScopedFieldContextMetaData ScopedMetaData(ContextIn, ResultsData.Get());
 
 	int32 NumResults = Results.Num();
-	int32 NumSamples = Context.SampleIndices.Num();
+	int32 NumSamples = ContextIn.SampleIndices.Num();
 
 	const FFieldNode<float> * ScalarField = Scalar.Get();
 	const FFieldNode<FVector> * RightVectorField = VectorRight.Get();
@@ -1022,15 +1028,15 @@ void FSumVector::Evaluate(const FFieldContext & ContextIn, TArrayView<FVector> &
 			TArrayView<FVector>(&Buffer[NumResults],NumResults),
 		};
 
-		LeftVectorField->Evaluate(Context, Buffers[0]);
-		RightVectorField->Evaluate(Context, Buffers[1]);
+		LeftVectorField->Evaluate(ContextIn, Buffers[0]);
+		RightVectorField->Evaluate(ContextIn, Buffers[1]);
 
 		switch (Operation)
 		{
 		case EFieldOperationType::Field_Multiply:
 			for (int32 SampleIndex = 0; SampleIndex < NumSamples; SampleIndex++)
 			{
-				const ContextIndex & Index = Context.SampleIndices[SampleIndex];
+				const ContextIndex& Index = ContextIn.SampleIndices[SampleIndex];
 				{
 					Results[Index.Result] = Buffers[1][Index.Result] * Buffers[0][Index.Result];
 				}
@@ -1039,7 +1045,7 @@ void FSumVector::Evaluate(const FFieldContext & ContextIn, TArrayView<FVector> &
 		case EFieldOperationType::Field_Divide:
 			for (int32 SampleIndex = 0; SampleIndex < NumSamples; SampleIndex++)
 			{
-		  const ContextIndex & Index = Context.SampleIndices[SampleIndex];
+		  const ContextIndex& Index = ContextIn.SampleIndices[SampleIndex];
 				{
 					Results[Index.Result] = Buffers[0][Index.Result] / Buffers[1][Index.Result];
 				}
@@ -1048,7 +1054,7 @@ void FSumVector::Evaluate(const FFieldContext & ContextIn, TArrayView<FVector> &
 		case EFieldOperationType::Field_Add:
 			for (int32 SampleIndex = 0; SampleIndex < NumSamples; SampleIndex++)
 			{
-		  const ContextIndex & Index = Context.SampleIndices[SampleIndex];
+		  const ContextIndex& Index = ContextIn.SampleIndices[SampleIndex];
 				{
 					Results[Index.Result] = Buffers[1][Index.Result] + Buffers[0][Index.Result];
 				}
@@ -1057,7 +1063,7 @@ void FSumVector::Evaluate(const FFieldContext & ContextIn, TArrayView<FVector> &
 		case EFieldOperationType::Field_Substract:
 			for (int32 SampleIndex = 0; SampleIndex < NumSamples; SampleIndex++)
 			{
-		  const ContextIndex & Index = Context.SampleIndices[SampleIndex];
+		  const ContextIndex& Index = ContextIn.SampleIndices[SampleIndex];
 				{
 					Results[Index.Result] = Buffers[0][Index.Result] - Buffers[1][Index.Result];
 				}
@@ -1067,11 +1073,11 @@ void FSumVector::Evaluate(const FFieldContext & ContextIn, TArrayView<FVector> &
 	}
 	else if (LeftVectorField != nullptr)
 	{
-		LeftVectorField->Evaluate(Context, Results);
+		LeftVectorField->Evaluate(ContextIn, Results);
 	}
 	else if (RightVectorField != nullptr)
 	{
-		RightVectorField->Evaluate(Context, Results);
+		RightVectorField->Evaluate(ContextIn, Results);
 	}
 
 	if (ScalarField != nullptr)
@@ -1080,11 +1086,11 @@ void FSumVector::Evaluate(const FFieldContext & ContextIn, TArrayView<FVector> &
 		Buffer.SetNumUninitialized(NumResults);
 		TArrayView<float> BufferView(&Buffer[0], NumResults);
 
-		ScalarField->Evaluate(Context, BufferView);
+		ScalarField->Evaluate(ContextIn, BufferView);
 
 		for (int32 SampleIndex = 0; SampleIndex < NumSamples; SampleIndex++)
 		{
-			 const ContextIndex & Index = Context.SampleIndices[SampleIndex];
+			 const ContextIndex& Index = ContextIn.SampleIndices[SampleIndex];
 			{
 				Results[Index.Result] *= Buffer[Index.Result];
 			}
@@ -1095,7 +1101,7 @@ void FSumVector::Evaluate(const FFieldContext & ContextIn, TArrayView<FVector> &
 	{
 		for (int32 SampleIndex = 0; SampleIndex < NumSamples; SampleIndex++)
 		{
-			const ContextIndex & Index = Context.SampleIndices[SampleIndex];
+			const ContextIndex& Index = ContextIn.SampleIndices[SampleIndex];
 			{
 				Results[Index.Result] *= MagnitudeVal;
 			}
@@ -1132,7 +1138,7 @@ bool FSumVector::operator==(const FFieldNodeBase& Node)
 * FConversionField<InT,OutT>
 */
 template<class InT, class OutT>
-void FConversionField<InT,OutT>::Evaluate(const FFieldContext & Context, TArrayView<OutT> & Results) const
+void FConversionField<InT,OutT>::Evaluate(FFieldContext& Context, TArrayView<OutT>& Results) const
 {
 	int32 NumSamples = Context.SampleIndices.Num();
 
@@ -1143,7 +1149,7 @@ void FConversionField<InT,OutT>::Evaluate(const FFieldContext & Context, TArrayV
 
 	for (int32 SampleIndex = 0; SampleIndex < NumSamples; SampleIndex++)
 	{
-		const ContextIndex & Index = Context.SampleIndices[SampleIndex];
+		const ContextIndex& Index = Context.SampleIndices[SampleIndex];
 		Results[Index.Result] = (OutT)Array[Index.Result];
 	}
 }
@@ -1172,7 +1178,7 @@ template class FConversionField<float, int32>;
 *  FCullingField<T>
 */
 template<class T>
-void FCullingField<T>::Evaluate(const FFieldContext & Context, TArrayView<T> & Results) const
+void FCullingField<T>::Evaluate(FFieldContext& Context, TArrayView<T>& Results) const
 {
 	int32 NumResults = Results.Num();
 	int32 NumSamples = Context.SampleIndices.Num();
@@ -1186,6 +1192,13 @@ void FCullingField<T>::Evaluate(const FFieldContext & Context, TArrayView<T> & R
 		if (ensureMsgf(CullingField->Type() == FFieldNode<float>::StaticType(),
 			TEXT("Field Node CullingFields Culling input expects a float input array.")))
 		{
+			FFieldSystemMetaDataCulling* CullingData = static_cast<FFieldSystemMetaDataCulling*>(Context.MetaData[FFieldSystemMetaData::EMetaType::ECommandData_Culling]);
+			
+			if(CullingData)
+			{
+				CullingData->bCullingActive = true;
+			}
+
 			TArray<float> EvaluationBuffer;
 			EvaluationBuffer.Init(0.f, NumResults);
 			TArrayView<float> EvaluationBufferView(&(EvaluationBuffer[0]), NumResults);
@@ -1194,7 +1207,7 @@ void FCullingField<T>::Evaluate(const FFieldContext & Context, TArrayView<T> & R
 			int NewEvaluationSize = 0;
 			for (int32 SampleIndex = 0; SampleIndex < NumSamples; SampleIndex++)
 			{
-				const ContextIndex & Index = Context.SampleIndices[SampleIndex];
+				const ContextIndex& Index = Context.SampleIndices[SampleIndex];
 				{
 					if (Operation == EFieldCullingOperationType::Field_Culling_Outside)
 					{
@@ -1215,7 +1228,7 @@ void FCullingField<T>::Evaluate(const FFieldContext & Context, TArrayView<T> & R
 			IndexBuffer.SetNumUninitialized(NewEvaluationSize);
 			for (int32 SampleIndex = 0, j=0; SampleIndex < NumSamples; SampleIndex++)
 			{
-				const ContextIndex & Index = Context.SampleIndices[SampleIndex];
+				const ContextIndex& Index = Context.SampleIndices[SampleIndex];
 				{
 					if (Operation == EFieldCullingOperationType::Field_Culling_Outside)
 					{
@@ -1223,6 +1236,11 @@ void FCullingField<T>::Evaluate(const FFieldContext & Context, TArrayView<T> & R
 						{
 							IndexBuffer[j] = Context.SampleIndices[SampleIndex];
 							j++;
+
+							if(CullingData)
+							{
+								CullingData->EvaluatedIndexBuffer.Add(Context.SampleIndices[SampleIndex]);
+							}
 						}
 					}
 					else 
@@ -1231,6 +1249,11 @@ void FCullingField<T>::Evaluate(const FFieldContext & Context, TArrayView<T> & R
 						{
 							IndexBuffer[j] = Context.SampleIndices[SampleIndex];
 							j++;
+
+							if(CullingData)
+							{
+								CullingData->EvaluatedIndexBuffer.Add(Context.SampleIndices[SampleIndex]);
+							}
 						}
 					}
 				}
@@ -1280,7 +1303,7 @@ template class FIELDSYSTEMCORE_API FCullingField<FVector>;
 * FReturnResultsTerminal<T>
 */
 template<class T>
-void FReturnResultsTerminal<T>::Evaluate(const FFieldContext & Context, TArrayView<T> & Results) const
+void FReturnResultsTerminal<T>::Evaluate(FFieldContext& Context, TArrayView<T>& Results) const
 {
 	if (ensureMsgf(Context.MetaData.Contains(FFieldSystemMetaData::EMetaType::ECommandData_Results),
 		TEXT("Return results nodes can only be used upstream from a 'results expector', for example as an input "
diff --git a/Engine/Source/Runtime/Experimental/FieldSystem/Source/FieldSystemCore/Public/Field/FieldSystem.h b/Engine/Source/Runtime/Experimental/FieldSystem/Source/FieldSystemCore/Public/Field/FieldSystem.h
--- a/Engine/Source/Runtime/Experimental/FieldSystem/Source/FieldSystemCore/Public/Field/FieldSystem.h
+++ b/Engine/Source/Runtime/Experimental/FieldSystem/Source/FieldSystemCore/Public/Field/FieldSystem.h
@@ -9,6 +9,40 @@
 #include "Field/FieldSystemTypes.h"
 #include "Math/Vector.h"
 
+/**
+* FFieldContext
+*   The Context is passed into the field evaluation pipeline during evaluation. The Nodes
+*   will have access to the samples and indices for evaluation. The MetaData is a optional
+*   data package that the nodes will use during evaluation, the context does not assume
+*   ownership of the metadata but assumes it will remain in scope during evaluation.
+*/
+struct FIELDSYSTEMCORE_API ContextIndex
+{
+	ContextIndex(int32 InSample = INDEX_NONE, int32 InResult = INDEX_NONE)
+		: Sample(InSample)
+		, Result(InResult)
+	{}
+
+	static void ContiguousIndices(
+		TArray<ContextIndex>& Array,
+		const int NumParticles,
+		const bool bForce = true)
+	{
+		if(bForce)
+		{
+			Array.SetNum(NumParticles);
+			for(int32 i = 0; i < Array.Num(); ++i)
+			{
+				Array[i].Result = i;
+				Array[i].Sample = i;
+			}
+		}
+	}
+
+	int32 Sample;
+	int32 Result;
+};
+
 /**
 * MetaData
 *
@@ -26,7 +60,8 @@ public:
 		ECommandData_None = 0,
 		ECommandData_ProcessingResolution,
 		ECommandData_Results,
-		ECommandData_Iteration
+		ECommandData_Iteration,
+		ECommandData_Culling
 	};
 
 
@@ -67,39 +102,30 @@ public:
 	int32 Iterations;
 };
 
-
-/**
-* FFieldContext
-*   The Context is passed into the field evaluation pipeline during evaluation. The Nodes
-*   will have access to the samples and indices for evaluation. The MetaData is a optional
-*   data package that the nodes will use during evaluation, the context does not assume 
-*   ownership of the metadata but assumes it will remain in scope during evaluation. 
-*/
-struct FIELDSYSTEMCORE_API ContextIndex 
+class FIELDSYSTEMCORE_API FFieldSystemMetaDataCulling : public FFieldSystemMetaData
 {
-	ContextIndex(int32 InSample=INDEX_NONE, int32 InResult=INDEX_NONE)
-		: Sample(InSample)
-		, Result(InResult) 
-	{}
+public:
+	explicit FFieldSystemMetaDataCulling(const int32 PotentialSize)
+		: bCullingActive(false)
+		, MaxSize(PotentialSize)
+	{
+		EvaluatedIndexBuffer.Reserve(MaxSize);
+	};
 
-	static void ContiguousIndices(
-		TArray<ContextIndex>& Array, 
-		const int NumParticles, 
-		const bool bForce = true)
+	virtual ~FFieldSystemMetaDataCulling() = default;
+
+	virtual EMetaType Type() const
 	{
-		if (bForce)
-		{
-			Array.SetNum(NumParticles);
-			for (int32 i = 0; i < Array.Num(); ++i)
-			{
-				Array[i].Result = i;
-				Array[i].Sample = i;
-			}
-		}
+		return EMetaType::ECommandData_Culling;
+	}
+	virtual FFieldSystemMetaData* NewCopy() const
+	{
+		return new FFieldSystemMetaDataCulling(MaxSize);
 	}
 
-	int32 Sample;
-	int32 Result;
+	bool bCullingActive;
+	int32 MaxSize;
+	TArray<ContextIndex> EvaluatedIndexBuffer;
 };
 
 struct FIELDSYSTEMCORE_API FFieldContext
@@ -108,22 +134,47 @@ struct FIELDSYSTEMCORE_API FFieldContext
 	typedef  TMap<FFieldSystemMetaData::EMetaType, FFieldSystemMetaData * > PointerMap;
 
 	FFieldContext() = delete;
+
+	FFieldContext(const FFieldContext&) = delete;
+	FFieldContext(FFieldContext&&) = delete;
+	FFieldContext& operator =(const FFieldContext&) = delete;
+	FFieldContext & operator =(FFieldContext&&) = delete;
+
 	FFieldContext(const TArrayView< ContextIndex >& SampleIndicesIn, const TArrayView<FVector>& SamplesIn,
 		const UniquePointerMap & MetaDataIn )
 		: SampleIndices(SampleIndicesIn)
 		, Samples(SamplesIn)
+
 	{
 		for (const TPair<FFieldSystemMetaData::EMetaType, TUniquePtr<FFieldSystemMetaData>>& Meta : MetaDataIn)
 		{
 			MetaData.Add(Meta.Key) = Meta.Value.Get();
 		}
+
+		CullingData = MakeUnique<FFieldSystemMetaDataCulling>(SampleIndices.Num());
+		MetaData.Add(FFieldSystemMetaData::EMetaType::ECommandData_Culling, CullingData.Get());
 	}
 	FFieldContext(const TArrayView< ContextIndex >& SampleIndicesIn, const TArrayView<FVector>& SamplesIn,
 		const PointerMap & MetaDataIn)
 		: SampleIndices(SampleIndicesIn)
 		, Samples(SamplesIn)
 		, MetaData(MetaDataIn)
-	{}
+	{
+		CullingData = MakeUnique<FFieldSystemMetaDataCulling>(SampleIndices.Num());
+		MetaData.Add(FFieldSystemMetaData::EMetaType::ECommandData_Culling, CullingData.Get());
+	}
+
+	TArrayView<ContextIndex> GetEvaluatedSamples()
+	{
+		if(!CullingData->bCullingActive)
+		{
+			// No culling took place
+			return SampleIndices;
+		}
+
+		// Culling fields created an evaluation set
+		return MakeArrayView(CullingData->EvaluatedIndexBuffer);
+	}
 
 	//
 	// Ryan - TODO: This concept of having discreet sample data needs to change.  
@@ -132,12 +183,44 @@ struct FIELDSYSTEMCORE_API FFieldContext
 	// traversed also needs to change; possibly to some load balanced threaded iterator 
 	// or task based paradigm.
 
-	//const TArrayView<THandle::TTransientHandle*>& Handles;
 	const TArrayView<ContextIndex>& SampleIndices;
 	const TArrayView<FVector>& Samples;
 	PointerMap MetaData;
+	TUniquePtr<FFieldSystemMetaDataCulling> CullingData;
 };
 
+/** 
+ * Limits the application of a meta data object to a single scope.
+ * This has the effect of exposing metadata to downstream nodes but making sure
+ * upstream nodes cannot see it.
+ */
+class FIELDSYSTEMCORE_API FScopedFieldContextMetaData
+{
+	FScopedFieldContextMetaData() = delete;
+	FScopedFieldContextMetaData(const FScopedFieldContextMetaData&) = delete;
+	FScopedFieldContextMetaData(FScopedFieldContextMetaData&&) = delete;
+	FScopedFieldContextMetaData& operator=(const FScopedFieldContextMetaData&) = delete;
+	FScopedFieldContextMetaData& operator=(FScopedFieldContextMetaData&&) = delete;
+
+public:
+
+	explicit FScopedFieldContextMetaData(FFieldContext& InContext, FFieldSystemMetaData* InMetaData)
+		: TargetContext(InContext)
+	{
+		check(InMetaData);
+		MetaType = InMetaData->Type();
+		TargetContext.MetaData.Add(MetaType, InMetaData);
+	}
+
+	~FScopedFieldContextMetaData()
+	{
+		TargetContext.MetaData.Remove(MetaType);
+	}
+
+private:
+	FFieldSystemMetaData::EMetaType MetaType;
+	FFieldContext& TargetContext;
+};
 
 /**
 * FFieldNodeBase
@@ -210,7 +293,7 @@ public:
 	
 	virtual ~FFieldNode() {}
 
-	virtual void Evaluate(const FFieldContext &, TArrayView<T> & Results) const = 0;
+	virtual void Evaluate(FFieldContext&, TArrayView<T>& Results) const = 0;
 
 	static EFieldType StaticType();
 	virtual EFieldType Type() const { return StaticType(); }
@@ -245,7 +328,7 @@ public:
 
 	// Commands are copied when moved from the one thread to 
 	// another. This requires a full copy of all associated data. 
-	FFieldSystemCommand(const FFieldSystemCommand & Other)
+	FFieldSystemCommand(const FFieldSystemCommand& Other)
 		: TargetAttribute(Other.RootNode ? Other.TargetAttribute:"")
 		, RootNode(Other.RootNode?Other.RootNode->NewCopy():nullptr)
 	{
diff --git a/Engine/Source/Runtime/Experimental/FieldSystem/Source/FieldSystemCore/Public/Field/FieldSystemNodes.h b/Engine/Source/Runtime/Experimental/FieldSystem/Source/FieldSystemCore/Public/Field/FieldSystemNodes.h
--- a/Engine/Source/Runtime/Experimental/FieldSystem/Source/FieldSystemCore/Public/Field/FieldSystemNodes.h
+++ b/Engine/Source/Runtime/Experimental/FieldSystem/Source/FieldSystemCore/Public/Field/FieldSystemNodes.h
@@ -24,7 +24,7 @@ public:
 	virtual FFieldNodeBase * NewCopy() const override { return new FUniformInteger(Magnitude); }
 	virtual ~FUniformInteger() {}
 
-	void Evaluate(const FFieldContext &, TArrayView<int32> & Results) const override;
+	void Evaluate(FFieldContext&, TArrayView<int32>& Results) const override;
 	virtual bool operator==(const FFieldNodeBase& Node);
 
 	/** Serialization API */
@@ -61,7 +61,7 @@ public:
 	virtual FFieldNodeBase * NewCopy() const override { return new FRadialIntMask(Radius,Position,InteriorValue,ExteriorValue,SetMaskCondition); }
 	virtual ~FRadialIntMask() {}
 
-	void Evaluate(const FFieldContext &, TArrayView<int32> & Results) const override;
+	void Evaluate(FFieldContext&, TArrayView<int32>& Results) const override;
 	virtual bool operator==(const FFieldNodeBase& Node);
 
 	/** Serialization API */
@@ -94,7 +94,7 @@ public:
 
 	virtual ~FUniformScalar() {}
 
-	void Evaluate(const FFieldContext &, TArrayView<float> & Results) const override;
+	void Evaluate(FFieldContext&, TArrayView<float>& Results) const override;
 	virtual bool operator==(const FFieldNodeBase& Node);
 
 	/** Serialization API */
@@ -137,7 +137,7 @@ public:
 
 	virtual ~FRadialFalloff() {}
 
-	void Evaluate(const FFieldContext &, TArrayView<float> & Results) const override;
+	void Evaluate(FFieldContext&, TArrayView<float>& Results) const override;
 	virtual bool operator==(const FFieldNodeBase& Node);
 
 	/** Serialization API */
@@ -153,7 +153,7 @@ public:
 	EFieldFalloffType Falloff;
 
 protected:
-	template<EFieldFalloffType> void Evaluator(const FFieldContext & Context, TArrayView<float> & Results) const;
+	template<EFieldFalloffType> void Evaluator(const FFieldContext& Context, TArrayView<float>& Results) const;
 };
 
 
@@ -189,7 +189,7 @@ public:
 	virtual FFieldNodeBase * NewCopy() const override { return new FPlaneFalloff(Magnitude, MinRange, MaxRange, Default, Distance, Position, Normal, Falloff); }
 	virtual ~FPlaneFalloff() {}
 
-	void Evaluate(const FFieldContext &, TArrayView<float> & Results) const override;
+	void Evaluate(FFieldContext&, TArrayView<float>& Results) const override;
 	virtual bool operator==(const FFieldNodeBase& Node);
 
 	/** Serialization API */
@@ -206,7 +206,7 @@ public:
 	EFieldFalloffType Falloff;
 
 protected :
-	template<EFieldFalloffType> void Evaluator(const FFieldContext & Context, const FPlane & Plane, TArrayView<float> & Results) const;
+	template<EFieldFalloffType> void Evaluator(const FFieldContext& Context, const FPlane& Plane, TArrayView<float>& Results) const;
 
 };
 
@@ -238,7 +238,7 @@ public:
 	virtual FFieldNodeBase * NewCopy() const override { return new FBoxFalloff(Magnitude, MinRange, MaxRange, Default, Transform, Falloff); }
 	virtual ~FBoxFalloff() {}
 
-	void Evaluate(const FFieldContext &, TArrayView<float> & Results) const override;
+	void Evaluate(FFieldContext&, TArrayView<float>& Results) const override;
 	virtual bool operator==(const FFieldNodeBase& Node);
 
 	/** Serialization API */
@@ -253,7 +253,7 @@ public:
 	EFieldFalloffType Falloff;
 
 protected:
-	template<EFieldFalloffType> void Evaluator(const FFieldContext & Context, TArrayView<float> & Results) const;
+	template<EFieldFalloffType> void Evaluator(const FFieldContext& Context, TArrayView<float>& Results) const;
 
 };
 
@@ -276,7 +276,7 @@ public:
 	virtual FFieldNodeBase * NewCopy() const override { return new FNoiseField(MinRange,MaxRange,Transform); }
 	virtual ~FNoiseField() {}
 
-	void Evaluate(const FFieldContext &, TArrayView<float> & Results) const override;	
+	void Evaluate(FFieldContext&, TArrayView<float>& Results) const override;	
 	virtual bool operator==(const FFieldNodeBase& Node);
 
 	/** Serialization API */
@@ -310,7 +310,7 @@ public:
 	virtual FFieldNodeBase * NewCopy() const override { return new FUniformVector(Magnitude,Direction); }
 	virtual ~FUniformVector() {}
 
-	void Evaluate(const FFieldContext &, TArrayView<FVector> & Results) const override;
+	void Evaluate(FFieldContext&, TArrayView<FVector>& Results) const override;
 	virtual bool operator==(const FFieldNodeBase& Node);
 
 	/** Serialization API */
@@ -341,7 +341,7 @@ public:
 	virtual FFieldNodeBase * NewCopy() const override { return new FRadialVector(Magnitude,Position); }
 	virtual ~FRadialVector() {}
 
-	void Evaluate(const FFieldContext &, TArrayView<FVector> & Results) const override;
+	void Evaluate(FFieldContext&, TArrayView<FVector>& Results) const override;
 	virtual bool operator==(const FFieldNodeBase& Node);
 
 	/** Serialization API */
@@ -371,7 +371,7 @@ public:
 	virtual FFieldNodeBase * NewCopy() const override { return new FRandomVector(Magnitude); }
 	virtual ~FRandomVector() {}
 
-	void Evaluate(const FFieldContext &, TArrayView<FVector> & Results) const override;
+	void Evaluate(FFieldContext&, TArrayView<FVector>& Results) const override;
 	virtual bool operator==(const FFieldNodeBase& Node);
 
 	/** Serialization API */
@@ -411,7 +411,7 @@ public:
 	}
 	virtual ~FSumScalar() {}
 
-	void Evaluate(const FFieldContext &, TArrayView<float> & Results) const override;
+	void Evaluate(FFieldContext&, TArrayView<float>& Results) const override;
 	virtual bool operator==(const FFieldNodeBase& Node);
 
 	/** Serialization API */
@@ -456,7 +456,7 @@ public:
 	}
 	virtual ~FSumVector() {}
 
-	void Evaluate(const FFieldContext &, TArrayView<FVector> & Results) const override;
+	void Evaluate(FFieldContext&, TArrayView<FVector>& Results) const override;
 	virtual bool operator==(const FFieldNodeBase& Node);
 
 	/** Serialization API */
@@ -491,7 +491,7 @@ public:
 	}
 	virtual ~FConversionField() {}
 
-	void Evaluate(const FFieldContext &, TArrayView<OutT> & Results) const override;
+	void Evaluate(FFieldContext&, TArrayView<OutT>& Results) const override;
 	virtual bool operator==(const FFieldNodeBase& Node);
 
 	/** Serialization API */
@@ -531,7 +531,7 @@ public:
 	}
 	virtual ~FCullingField() {}
 
-	void Evaluate(const FFieldContext &, TArrayView<T> & Results) const override;
+	void Evaluate(FFieldContext&, TArrayView<T>& Results) const override;
 	virtual bool operator==(const FFieldNodeBase& Node);
 
 	/** Serialization API */
@@ -568,7 +568,7 @@ public:
 	}
 	virtual ~FReturnResultsTerminal() {}
 
-	void Evaluate(const FFieldContext &, TArrayView<T> & Results) const override;
+	void Evaluate(FFieldContext&, TArrayView<T>& Results) const override;
 	virtual bool operator==(const FFieldNodeBase& Node);
 
 	/** Serialization API */
