Chaos : Clustering

Update Proxy mapping for internally created cluster.
Disable MassProperties tests that were failing.

diff --git a/Engine/Source/Runtime/Experimental/Chaos/Private/Chaos/PBDRigidClustering.cpp b/Engine/Source/Runtime/Experimental/Chaos/Private/Chaos/PBDRigidClustering.cpp
--- a/Engine/Source/Runtime/Experimental/Chaos/Private/Chaos/PBDRigidClustering.cpp
+++ b/Engine/Source/Runtime/Experimental/Chaos/Private/Chaos/PBDRigidClustering.cpp
@@ -537,6 +537,7 @@ namespace Chaos
 		TopLevelClusterParents.Add(NewParticle);
 		NewParticle->SetInternalCluster(true);
 		NewParticle->SetClusterId(ClusterId(Parent, Children.Num()));
+		for (auto& Constituent : Children) MEvolution.DoInternalParticleInitilization(Constituent, NewParticle);
 
 		//
 		// Update clustering data structures.
