Copy the ErrorReporter rather than taking a reference, so that the prefix doesn't get overwritten when called in parallel.  Don't call PhiWithNormal() on a level set, if SignedDistance() will do the job.

diff --git a/Engine/Source/Runtime/Experimental/Chaos/Private/Chaos/Levelset.cpp b/Engine/Source/Runtime/Experimental/Chaos/Private/Chaos/Levelset.cpp
--- a/Engine/Source/Runtime/Experimental/Chaos/Private/Chaos/Levelset.cpp
+++ b/Engine/Source/Runtime/Experimental/Chaos/Private/Chaos/Levelset.cpp
@@ -283,9 +283,7 @@ T TLevelSet<T, d>::ComputeLevelSetError(const TParticles<T, d>& InParticles, con
 
 			for (int j = 0; j < d; ++j)
 			{
-				TVector<T, 3> UnusedGridNormal;
-
-				DistErrorValues[i] += FMath::Abs(PhiWithNormal(InParticles.X(CurrMeshFace[j]), UnusedGridNormal));
+				DistErrorValues[i] += FMath::Abs(SignedDistance(InParticles.X(CurrMeshFace[j])));
 			}
 
 			// per triangle error average of 3 corners and center distance to surface according to MPhi
diff --git a/Engine/Source/Runtime/Experimental/GeometryCollectionSimulationCore/Private/GeometryCollection/GeometryCollectionCollisionStructureManager.cpp b/Engine/Source/Runtime/Experimental/GeometryCollectionSimulationCore/Private/GeometryCollection/GeometryCollectionCollisionStructureManager.cpp
--- a/Engine/Source/Runtime/Experimental/GeometryCollectionSimulationCore/Private/GeometryCollection/GeometryCollectionCollisionStructureManager.cpp
+++ b/Engine/Source/Runtime/Experimental/GeometryCollectionSimulationCore/Private/GeometryCollection/GeometryCollectionCollisionStructureManager.cpp
@@ -186,7 +186,7 @@ void FCollisionStructureManager::UpdateImplicitFlags(
 
 FCollisionStructureManager::FImplicit* 
 FCollisionStructureManager::NewImplicit(
-	Chaos::FErrorReporter& ErrorReporter,
+	Chaos::FErrorReporter ErrorReporter,
 	const Chaos::TParticles<float, 3>& MeshParticles,
 	const Chaos::TTriangleMesh<float>& TriMesh,
 	const FBox& CollisionBounds,
@@ -239,7 +239,7 @@ FCollisionStructureManager::NewImplicitSphere(
 
 FCollisionStructureManager::FImplicit*
 FCollisionStructureManager::NewImplicitLevelset(
-	Chaos::FErrorReporter& ErrorReporter,
+	Chaos::FErrorReporter ErrorReporter,
 	const Chaos::TParticles<float, 3>& MeshParticles,
 	const Chaos::TTriangleMesh<float>& TriMesh,
 	const FBox& CollisionBounds,
@@ -284,7 +284,7 @@ FCollisionStructureManager::NewImplicitLevelset(
 }
 
 Chaos::TLevelSet<float, 3>* FCollisionStructureManager::NewLevelset(
-	Chaos::FErrorReporter& ErrorReporter,
+	Chaos::FErrorReporter ErrorReporter,
 	const Chaos::TParticles<float, 3>& MeshParticles,
 	const Chaos::TTriangleMesh<float>& TriMesh,
 	const FBox& CollisionBounds,
diff --git a/Engine/Source/Runtime/Experimental/GeometryCollectionSimulationCore/Public/GeometryCollection/GeometryCollectionCollisionStructureManager.h b/Engine/Source/Runtime/Experimental/GeometryCollectionSimulationCore/Public/GeometryCollection/GeometryCollectionCollisionStructureManager.h
--- a/Engine/Source/Runtime/Experimental/GeometryCollectionSimulationCore/Public/GeometryCollection/GeometryCollectionCollisionStructureManager.h
+++ b/Engine/Source/Runtime/Experimental/GeometryCollectionSimulationCore/Public/GeometryCollection/GeometryCollectionCollisionStructureManager.h
@@ -77,7 +77,7 @@ public:
 	 *	\p ImplicitType - type of implicit shape to build.
 	 */
 	static FImplicit * NewImplicit(
-		Chaos::FErrorReporter& ErrorReporter,
+		Chaos::FErrorReporter ErrorReporter,
 		const Chaos::TParticles<float, 3>& MeshParticles,
 		const Chaos::TTriangleMesh<float>& TriMesh,
 		const FBox& CollisionBounds,
@@ -99,7 +99,7 @@ public:
 		const ECollisionTypeEnum CollisionType);
 
 	static FImplicit* NewImplicitLevelset(
-		Chaos::FErrorReporter& ErrorReporter,
+		Chaos::FErrorReporter ErrorReporter,
 		const Chaos::TParticles<float, 3>& MeshParticles,
 		const Chaos::TTriangleMesh<float>& TriMesh,
 		const FBox& CollisionBounds,
@@ -109,7 +109,7 @@ public:
 		const ECollisionTypeEnum CollisionType);
 
 	static Chaos::TLevelSet<float, 3>* NewLevelset(
-		Chaos::FErrorReporter& ErrorReporter,
+		Chaos::FErrorReporter ErrorReporter,
 		const Chaos::TParticles<float, 3>& MeshParticles,
 		const Chaos::TTriangleMesh<float>& TriMesh,
 		const FBox& CollisionBounds,
