Fail on level set rasterization in cases of high error by default.

diff --git a/Engine/Source/Runtime/Experimental/Chaos/Private/Chaos/Levelset.cpp b/Engine/Source/Runtime/Experimental/Chaos/Private/Chaos/Levelset.cpp
--- a/Engine/Source/Runtime/Experimental/Chaos/Private/Chaos/Levelset.cpp
+++ b/Engine/Source/Runtime/Experimental/Chaos/Private/Chaos/Levelset.cpp
@@ -19,7 +19,7 @@
 int32 OutputFailedLevelSetDebugData = 0;
 FAutoConsoleVariableRef CVarOutputFailedLevelSetDebugData(TEXT("p.LevelSetOutputFailedDebugData"), OutputFailedLevelSetDebugData, TEXT("Output debug obj files for level set and mesh when error tolerances are too high"));
 
-int32 FailureOnHighError = 0;
+int32 FailureOnHighError = 1;
 FAutoConsoleVariableRef CVarFailureOnHighError(TEXT("p.LevelSetFailureOnHighError"), FailureOnHighError, TEXT("Set level sets with high error to null in the solver"));
 
 float AvgDistErrorTolerance = .05;
