Rotate GeometryCollection DDC to pick up changes to collision and surface winding from yesterday.

diff --git a/Engine/Source/Runtime/Experimental/GeometryCollectionEngine/Private/GeometryCollection/DerivedDataGeometryCollectionCooker.cpp b/Engine/Source/Runtime/Experimental/GeometryCollectionEngine/Private/GeometryCollection/DerivedDataGeometryCollectionCooker.cpp
--- a/Engine/Source/Runtime/Experimental/GeometryCollectionEngine/Private/GeometryCollection/DerivedDataGeometryCollectionCooker.cpp
+++ b/Engine/Source/Runtime/Experimental/GeometryCollectionEngine/Private/GeometryCollection/DerivedDataGeometryCollectionCooker.cpp
@@ -59,7 +59,7 @@ const TCHAR* FDerivedDataGeometryCollectionCooker::GetVersionString() const
 		return OverrideVersion;	//force load old ddc if found. Not recommended
 	}
 
-	return TEXT("6C8317B9BC9C4DA1A38AEA126F8F24D0");
+	return TEXT("9336C4DCDFF54701A5D322C73EF6D5A4");
 }
 
 FString FDerivedDataGeometryCollectionCooker::GetPluginSpecificCacheKeySuffix() const
