Made UCullingField report its type as the type of the underlying culled field so operator nodes can work with culled data sets.

diff --git a/Engine/Source/Runtime/Experimental/FieldSystem/Source/FieldSystemEngine/Private/Field/FieldSystemObjects.cpp b/Engine/Source/Runtime/Experimental/FieldSystem/Source/FieldSystemEngine/Private/Field/FieldSystemObjects.cpp
--- a/Engine/Source/Runtime/Experimental/FieldSystem/Source/FieldSystemEngine/Private/Field/FieldSystemObjects.cpp
+++ b/Engine/Source/Runtime/Experimental/FieldSystem/Source/FieldSystemEngine/Private/Field/FieldSystemObjects.cpp
@@ -450,6 +450,10 @@ UCullingField* UCullingField::SetCullingField(const UFieldNodeBase* InCulling, c
 	return this;
 }
 
+FFieldNodeBase::EFieldType UCullingField::Type() const
+{
+	return Field ? Field->Type() : Super::Type();
+}
 
 // UReturnResultsField
 FFieldNodeBase* UReturnResultsTerminal::NewEvaluationGraph(TArray<const UFieldNodeBase*>& Nodes) const
diff --git a/Engine/Source/Runtime/Experimental/FieldSystem/Source/FieldSystemEngine/Public/Field/FieldSystemObjects.h b/Engine/Source/Runtime/Experimental/FieldSystem/Source/FieldSystemEngine/Public/Field/FieldSystemObjects.h
--- a/Engine/Source/Runtime/Experimental/FieldSystem/Source/FieldSystemEngine/Public/Field/FieldSystemObjects.h
+++ b/Engine/Source/Runtime/Experimental/FieldSystem/Source/FieldSystemEngine/Public/Field/FieldSystemObjects.h
@@ -608,6 +608,8 @@ public:
 
 	virtual FFieldNodeBase* NewEvaluationGraph(TArray<const UFieldNodeBase*>& Nodes) const override;
 
+	virtual FFieldNodeBase::EFieldType Type() const override;
+
 	UFUNCTION(BlueprintPure, Category = "Field")
 	UCullingField* SetCullingField( const UFieldNodeBase* Culling, const UFieldNodeBase* Field, EFieldCullingOperationType Operation);
 
@@ -619,7 +621,6 @@ public:
 
 	UPROPERTY(BlueprintReadWrite, EditAnywhere, Category = "Field")
 	TEnumAsByte<EFieldCullingOperationType> Operation;
-
 };
 
 /**
