Update all chaos field command types to use reduced evaluated sample sets

diff --git a/Engine/Source/Runtime/Experimental/ChaosSolvers/Private/PhysicsProxy/FieldSystemPhysicsProxy.cpp b/Engine/Source/Runtime/Experimental/ChaosSolvers/Private/PhysicsProxy/FieldSystemPhysicsProxy.cpp
--- a/Engine/Source/Runtime/Experimental/ChaosSolvers/Private/PhysicsProxy/FieldSystemPhysicsProxy.cpp
+++ b/Engine/Source/Runtime/Experimental/ChaosSolvers/Private/PhysicsProxy/FieldSystemPhysicsProxy.cpp
@@ -117,10 +117,12 @@ void FFieldSystemPhysicsProxy::FieldParameterUpdateCallback(
 						SamplePointsView,
 						Command.MetaData);
 
-					//
-					//  Sample the dynamic state array in the field
-					//
-					TArray<int32> DynamicState; // TODO: 32 bits seems rather excessive!!!
+					// Sample the dynamic state array in the field
+
+					// #BGTODO We're initializing every particle in the simulation here even though we're probably
+					// going to cull it in the field eval - can probably be smarter about this.
+					
+					TArray<int32> DynamicState; // #BGTODO Enum class support (so we can size the underlying type to be more appropriate)
 					DynamicState.AddUninitialized(Handles.Num());					
 					int32 i = 0;
 					for(Chaos::TGeometryParticleHandle<float,3>* Handle : Handles)
@@ -154,9 +156,10 @@ void FFieldSystemPhysicsProxy::FieldParameterUpdateCallback(
 					}
 
 					bool StateChanged = false;
-					i = 0;
-					for (Chaos::TGeometryParticleHandle<float, 3>* Handle : Handles)
+					for(const ContextIndex& Index : Context.GetEvaluatedSamples())
 					{
+						Chaos::TGeometryParticleHandle<float, 3>* Handle = Handles[Index.Sample];
+
 						// Lower level particle handles, like TGeometryParticleHandle and 
 						// TKinematicParticleHandle, infer their dynamic state by whether or not
 						// promotion to a derived c++ handle type succeeds or fails.
@@ -167,11 +170,12 @@ void FFieldSystemPhysicsProxy::FieldParameterUpdateCallback(
 						// handle has a getter and a setter for that data.  So, at least for now,
 						// we're just going to ignore non-dynamic particles.  This has the added
 						// benefit of not needing to deal with the floor, as it's pretty likely to
-						// not be dynamic.  Har.
+						// not be dynamic.
 						if(Chaos::TPBDRigidParticleHandle<float, 3>* RigidHandle = Handle->CastToRigidParticle())
 						{
 							const bool bIsGC = (Handle->GetParticleType() == Chaos::EParticleType::GeometryCollection) || 
 								(Handle->GetParticleType() == Chaos::EParticleType::Clustered && !Handle->CastToClustered()->InternalCluster());
+
 							const EObjectStateType HandleState = RigidHandle->ObjectState();
 
 							// Non-Geometry Collection rigid bodies are more restricted as to how
@@ -181,11 +185,10 @@ void FFieldSystemPhysicsProxy::FieldParameterUpdateCallback(
 							// downgrading from Dynamic in those cases.
 							if (!bIsGC && HandleState != Chaos::EObjectStateType::Dynamic)
 							{
-								++i;
 								continue;
 							}
 
-							const int32 FieldState = DynamicStateView[i];
+							const int32 FieldState = DynamicStateView[Index.Result];
 							if (FieldState == (int32)EObjectStateTypeEnum::Chaos_Object_Dynamic)
 							{
 								if ((HandleState == Chaos::EObjectStateType::Static ||
@@ -230,8 +233,8 @@ void FFieldSystemPhysicsProxy::FieldParameterUpdateCallback(
 								}
 							}
 						} // handle is dynamic
-						++i;
-					} // end for all handles
+					} // end for all samples
+
 					if (StateChanged)
 					{
 						// regenerate views
@@ -305,11 +308,10 @@ void FFieldSystemPhysicsProxy::FieldParameterUpdateCallback(
 
 #if TODO_REIMPLEMENT_RIGID_CLUSTERING
 					// transfer results to rigid system.
-					int32 NumSamples = Context.SampleIndices.Num();
-					for (int32 SampleIndex = 0; SampleIndex < NumSamples; SampleIndex++)
+					for(const ContextIndex& Index : Context.GetEvaluatedSamples())
 					{
-						int32 RigidBodyIndex = Context.SampleIndices[SampleIndex].Result;
-						if (DynamicStateView[RigidBodyIndex] == 0 && Particles.Disabled(RigidBodyIndex))
+						int32 RigidBodyIndex = Index.Result;
+						if(DynamicStateView[RigidBodyIndex] == 0 && Particles.Disabled(RigidBodyIndex))
 						{
 							ensure(CurrentSolver->GetRigidClustering().GetClusterIdsArray()[RigidBodyIndex].Id == INDEX_NONE);
 							CurrentSolver->GetEvolution()->EnableParticle(RigidBodyIndex, INDEX_NONE);
@@ -368,12 +370,13 @@ void FFieldSystemPhysicsProxy::FieldParameterUpdateCallback(
 						{
 							TMap<TGeometryParticleHandle<float, 3>*, float> Map;
 							
-							for (int Index=0; Index < StrainSamples.Num(); Index++)
+							for(const ContextIndex& Index : Context.GetEvaluatedSamples())
 							{
-								if (StrainSamples[Index] > 0)
+								if(StrainSamples[Index.Result] > 0)
 								{
-									Map.Add(Handles[Index], StrainSamples[Index]);
+									Map.Add(Handles[Index.Sample], StrainSamples[Index.Result]);
 								}
+
 							}
 
 							// Capture the results from the breaking model to post-process
@@ -427,19 +430,16 @@ void FFieldSystemPhysicsProxy::FieldParameterUpdateCallback(
 						static_cast<const FFieldNode<float>*>(Command.RootNode.Get())->Evaluate(
 							Context, ResultsView);
 						
-
 						bool HasDisabled = false;
-						int32 i = 0;
-						for (Chaos::TGeometryParticleHandle<float, 3>* Handle : Handles)
+						for(const ContextIndex& Index : Context.GetEvaluatedSamples())
 						{
-							Chaos::TPBDRigidParticleHandle<float, 3>* RigidHandle = Handle->CastToRigidParticle();
+							Chaos::TPBDRigidParticleHandle<float, 3>* RigidHandle = Handles[Index.Sample]->CastToRigidParticle();
 							
-							if (ResultsView[i] > 0.0)
+							if (RigidHandle && ResultsView[Index.Result] > 0.0)
 							{
-								CurrentSolver->GetEvolution()->DisableParticle(Handle);
+								CurrentSolver->GetEvolution()->DisableParticle(RigidHandle);
 								HasDisabled = true;
 							}
-							++i;
 						}
 					}
 				}
@@ -462,20 +462,18 @@ void FFieldSystemPhysicsProxy::FieldParameterUpdateCallback(
 							Command.MetaData);
 
 						TArray<FVector> LocalResults;
-						LocalResults.AddUninitialized(Handles.Num());					
+						LocalResults.AddUninitialized(Handles.Num());
 						TArrayView<FVector> ResultsView(&(LocalResults[0]), LocalResults.Num());
-						static_cast<const FFieldNode<FVector> *>(Command.RootNode.Get())->Evaluate(
-							Context, ResultsView);
 
-						int32 i = 0;
-						for (Chaos::TGeometryParticleHandle<float, 3>* Handle : Handles)
+						static_cast<const FFieldNode<FVector> *>(Command.RootNode.Get())->Evaluate(Context, ResultsView);
+
+						for(const ContextIndex& Index : Context.GetEvaluatedSamples())
 						{
-							Chaos::TPBDRigidParticleHandle<float, 3>* RigidHandle = Handle->CastToRigidParticle();
+							Chaos::TPBDRigidParticleHandle<float, 3>* RigidHandle = Handles[Index.Sample]->CastToRigidParticle();
 							if(RigidHandle && RigidHandle->ObjectState() == Chaos::EObjectStateType::Dynamic)
 							{
-								RigidHandle->V() += ResultsView[i];
+								RigidHandle->V() += ResultsView[Index.Result];
 							}
-							++i;
 						}
 					}
 				}
@@ -525,13 +523,12 @@ void FFieldSystemPhysicsProxy::FieldParameterUpdateCallback(
 						TArray<float> LocalResults;
 						LocalResults.AddZeroed(Handles.Num());
 						TArrayView<float> ResultsView(&(LocalResults[0]), LocalResults.Num());
-						static_cast<const FFieldNode<float>*>(Command.RootNode.Get())->Evaluate(
-							Context, ResultsView);
 
-						int32 i = 0;
-						for (Chaos::TGeometryParticleHandle<float, 3>* Handle : Handles)
+						static_cast<const FFieldNode<float>*>(Command.RootNode.Get())->Evaluate(Context, ResultsView);
+
+						for(const ContextIndex& Index : Context.GetEvaluatedSamples())
 						{
-							Chaos::TPBDRigidParticleHandle<float, 3>* RigidHandle = Handle->CastToRigidParticle();
+							Chaos::TPBDRigidParticleHandle<float, 3>* RigidHandle = Handles[Index.Sample]->CastToRigidParticle();
 
 							if (RigidHandle && ResultsView.Num() > 0)
 							{
@@ -540,8 +537,8 @@ void FFieldSystemPhysicsProxy::FieldParameterUpdateCallback(
 								{
 
 									TUniquePtr<Chaos::FChaosPhysicsMaterial> NewMaterial = MakeUnique< Chaos::FChaosPhysicsMaterial>();
-									NewMaterial->SleepingLinearThreshold = ResultsView[i];
-									NewMaterial->SleepingAngularThreshold = ResultsView[i];
+									NewMaterial->SleepingLinearThreshold = ResultsView[Index.Result];
+									NewMaterial->SleepingAngularThreshold = ResultsView[Index.Result];
 
 
 									CurrentSolver->GetEvolution()->SetPhysicsMaterial(RigidHandle, MakeSerializable(NewMaterial));
@@ -551,15 +548,13 @@ void FFieldSystemPhysicsProxy::FieldParameterUpdateCallback(
 								{
 									const TUniquePtr<FChaosPhysicsMaterial>& InstanceMaterial = CurrentSolver->GetEvolution()->GetPerParticlePhysicsMaterial(RigidHandle);
 
-									if (ResultsView[i] != InstanceMaterial->DisabledLinearThreshold)
+									if (ResultsView[Index.Result] != InstanceMaterial->DisabledLinearThreshold)
 									{
-										InstanceMaterial->SleepingLinearThreshold = ResultsView[i];
-										InstanceMaterial->SleepingAngularThreshold = ResultsView[i];
+										InstanceMaterial->SleepingLinearThreshold = ResultsView[Index.Result];
+										InstanceMaterial->SleepingAngularThreshold = ResultsView[Index.Result];
 									}
 								}
 							}
-
-							++i;
 						}
 					}				
 				}
@@ -672,7 +667,7 @@ void FFieldSystemPhysicsProxy::FieldParameterUpdateCallback(
 						static_cast<const FFieldNode<int32> *>(Command.RootNode.Get())->Evaluate(Context, ResultsView);
 
 #if TODO_REIMPLEMENT_KINEMATIC_PROXY
-						for (const ContextIndex& CIndex : IndicesArray)
+						for (const ContextIndex& CIndex : Context.GetEvaluatedSamples())
 						{
 							const int32 i = CIndex.Result;
 							if (Results[i])
@@ -723,7 +718,7 @@ void FFieldSystemPhysicsProxy::FieldParameterUpdateCallback(
 						static_cast<const FFieldNode<FVector> *>(Command.RootNode.Get())->Evaluate(Context, ResultsView);
 
 #if TODO_REIMPLEMENT_KINEMATIC_PROXY
-						for (const ContextIndex& CIndex : IndicesArray)
+						for (const ContextIndex& CIndex : Context.GetEvaluatedSamples())
 						{
 							const int32 i = CIndex.Result;
 							if (Results[i] != FVector(FLT_MAX))
@@ -833,7 +828,7 @@ void FFieldSystemPhysicsProxy::FieldParameterUpdateCallback(
 						TArrayView<float> ResultsView(&(Results[0]), Results.Num());
 						static_cast<const FFieldNode<float> *>(Command.RootNode.Get())->Evaluate(Context, ResultsView);
 
-						for (const ContextIndex& CIndex : IndicesArray)
+						for (const ContextIndex& CIndex : Context.GetEvaluatedSamples())
 						{
 							const int32 i = CIndex.Result;
 							if (Results[i] != FLT_MAX)
@@ -922,25 +917,24 @@ void FFieldSystemPhysicsProxy::FieldForcesUpdateCallback(
 						TArrayView<FVector> ForceView(&(LocalForce[0]), LocalForce.Num());
 						static_cast<const FFieldNode<FVector> *>(Command.RootNode.Get())->Evaluate(Context, ForceView);
 		
-						int32 i = 0;
-						for (Chaos::TGeometryParticleHandle<float, 3>* Handle : Handles)
+						for(const ContextIndex& Index : Context.GetEvaluatedSamples())
 						{
-							Chaos::TPBDRigidParticleHandle<float, 3>* RigidHandle = Handle->CastToRigidParticle();
+							Chaos::TPBDRigidParticleHandle<float, 3>* RigidHandle = Handles[Index.Sample]->CastToRigidParticle();
+
 							if(RigidHandle && (RigidHandle->ObjectState() == Chaos::EObjectStateType::Dynamic || RigidHandle->ObjectState() == Chaos::EObjectStateType::Sleeping))
 							{				
 								if (RigidHandle->Sleeping())
 								{
 									RigidHandle->SetObjectState(Chaos::EObjectStateType::Dynamic);
 								}
-								RigidHandle->F() += ForceView[i];
+								RigidHandle->F() += ForceView[Index.Result];
 							}
-							++i;
 						}
 
 #if TODO_REIMPLEMENT_WAKE_ISLANDS
 						// @todo(ccaulfield): encapsulation: add WakeParticles (and therefore islands) functionality to Evolution
 						TSet<int32> IslandsToActivate;
-						for (const ContextIndex& CIndex : IndicesArray)
+						for (const ContextIndex& CIndex : Context.GetEvaluatedSamples())
 						{
 							const int32 i = CIndex.Result;
 							if (ForceView[i] != FVector(0) && Particles.ObjectState(i) == Chaos::EObjectStateType::Sleeping && !Particles.Disabled(i) && IslandsToActivate.Find(Particles.Island(i)) == nullptr)
@@ -986,25 +980,23 @@ void FFieldSystemPhysicsProxy::FieldForcesUpdateCallback(
 						TArrayView<FVector> TorqueView(&(LocalTorque[0]), LocalTorque.Num());
 						static_cast<const FFieldNode<FVector> *>(Command.RootNode.Get())->Evaluate(Context, TorqueView);
 		
-						int32 i = 0;
-						for (Chaos::TGeometryParticleHandle<float, 3>* Handle : Handles)
+						for(const ContextIndex& Index : Context.GetEvaluatedSamples())
 						{
-							Chaos::TPBDRigidParticleHandle<float, 3>* RigidHandle = Handle->CastToRigidParticle();
+							Chaos::TPBDRigidParticleHandle<float, 3>* RigidHandle = Handles[Index.Sample]->CastToRigidParticle();
 							if(RigidHandle && (RigidHandle->ObjectState() == Chaos::EObjectStateType::Dynamic || RigidHandle->ObjectState() == Chaos::EObjectStateType::Sleeping))
 							{
 								if (RigidHandle->Sleeping())
 								{
 									RigidHandle->SetObjectState(Chaos::EObjectStateType::Dynamic);
 								}
-								RigidHandle->Torque() += TorqueView[i];
+								RigidHandle->Torque() += TorqueView[Index.Result];
 							}
-							++i;
 						}
 
 #if TODO_REIMPLEMENT_WAKE_ISLANDS
 						// @todo(ccaulfield): encapsulation: add WakeParticles (and therefore islands) functionality to Evolution
 						TSet<int32> IslandsToActivate;
-						for (const ContextIndex& CIndex : IndicesArray)
+						for (const ContextIndex& CIndex : Context.GetEvaluatedSamples())
 						{
 							const int32 i = CIndex.Result;
 							if (TorqueView[i] != FVector(0) && Particles.ObjectState(i) == Chaos::EObjectStateType::Sleeping && !Particles.Disabled(i) && IslandsToActivate.Find(Particles.Island(i)) == nullptr)
