Added a failsafe and some logging with the failsafe is triggered in EPA.

diff --git a/Engine/Source/Runtime/Experimental/Chaos/Public/Chaos/EPA.h b/Engine/Source/Runtime/Experimental/Chaos/Public/Chaos/EPA.h
--- a/Engine/Source/Runtime/Experimental/Chaos/Public/Chaos/EPA.h
+++ b/Engine/Source/Runtime/Experimental/Chaos/Public/Chaos/EPA.h
@@ -305,7 +305,10 @@ void EPAComputeVisibilityBorder(TArray<TEPAEntry<T>>& Entries, int32 EntryIdx, c
 		}
 	}
 
-	while (ToVisitStack.Num())
+	int32 Iteration = 0;
+	const int32 MaxIteration = 10000;
+
+	while (ToVisitStack.Num() && Iteration++ < MaxIteration)
 	{
 		const FEPAFloodEntry FloodEntry = ToVisitStack.Pop(false);
 		TEPAEntry<T>& Entry = Entries[FloodEntry.EntryIdx];
@@ -328,6 +331,11 @@ void EPAComputeVisibilityBorder(TArray<TEPAEntry<T>>& Entries, int32 EntryIdx, c
 			}
 		}
 	}
+
+	if(Iteration >= MaxIteration)
+	{
+		UE_LOG(LogChaos,Warning,TEXT("EPAComputeVisibilityBorder reached max iteration - something is wrong"));
+	}
 }
 
 template <typename T>
