Enabled scene queries for children of internal clusters

diff --git a/Engine/Source/Runtime/Experimental/Chaos/Public/Chaos/PBDRigidsEvolution.h b/Engine/Source/Runtime/Experimental/Chaos/Public/Chaos/PBDRigidsEvolution.h
--- a/Engine/Source/Runtime/Experimental/Chaos/Public/Chaos/PBDRigidsEvolution.h
+++ b/Engine/Source/Runtime/Experimental/Chaos/Public/Chaos/PBDRigidsEvolution.h
@@ -343,8 +343,29 @@ class FPBDRigidsEvolutionBase
 		const TPBDRigidParticleHandleImp<FReal, 3, bPersistent>* AsRigid = Particle.CastToRigidParticle();
 		if(AsRigid && AsRigid->Disabled())
 		{
-			// Disabled particles take no immediate part in sim or query so shouldn't be added to the acceleration
-			return;
+			TPBDRigidClusteredParticleHandleImp<FReal, 3, bPersistent>* AsClustered = Particle.CastToClustered();
+
+			if(AsClustered)
+			{
+				// For clustered particles, they may appear disabled but they're being driven by an internal (solver-owned) cluster parent.
+				// If this is the case we let the spatial data update with those particles, otherwise skip.
+				// #BGTODO consider converting MDisabled into a bitfield for multiple disable types (Disabled, DisabledDriven, etc.)
+				if(TPBDRigidParticleHandle<float, 3>* ClusterParentBase = AsClustered->ClusterIds().Id)
+				{
+					if(Chaos::TPBDRigidClusteredParticleHandle<float, 3>* ClusterParent = ClusterParentBase->CastToClustered())
+					{
+						if(!ClusterParent->InternalCluster())
+						{
+							return;
+						}
+					}
+				}
+			}
+			else
+			{
+				// Disabled particles take no immediate part in sim or query so shouldn't be added to the acceleration
+				return;
+			}
 		}
 
 		//TODO: distinguish between new particles and dirty particles
diff --git a/Engine/Source/Runtime/Experimental/ChaosSolvers/Private/PhysicsProxy/GeometryCollectionPhysicsProxy.cpp b/Engine/Source/Runtime/Experimental/ChaosSolvers/Private/PhysicsProxy/GeometryCollectionPhysicsProxy.cpp
--- a/Engine/Source/Runtime/Experimental/ChaosSolvers/Private/PhysicsProxy/GeometryCollectionPhysicsProxy.cpp
+++ b/Engine/Source/Runtime/Experimental/ChaosSolvers/Private/PhysicsProxy/GeometryCollectionPhysicsProxy.cpp
@@ -1575,6 +1575,10 @@ void FGeometryCollectionPhysicsProxy::BufferPhysicsResults()
 							// Indicate that this object needs to be updated and the proxy is active.
 							TargetResults.DisabledStates[TransformGroupIndex] = false;
 							IsObjectDynamic = true;
+
+							ProxyElementHandle->X() = ClusterChildToWorld.GetTranslation();
+							ProxyElementHandle->R() = ClusterChildToWorld.GetRotation();
+							GetSolver()->GetEvolution()->DirtyParticle(*ProxyElementHandle);
 						}
 					}
 				}
