fixed dynamic state field when used as an initialization field

diff --git a/Engine/Source/Runtime/Experimental/ChaosSolvers/Private/PhysicsProxy/GeometryCollectionPhysicsProxy.cpp b/Engine/Source/Runtime/Experimental/ChaosSolvers/Private/PhysicsProxy/GeometryCollectionPhysicsProxy.cpp
--- a/Engine/Source/Runtime/Experimental/ChaosSolvers/Private/PhysicsProxy/GeometryCollectionPhysicsProxy.cpp
+++ b/Engine/Source/Runtime/Experimental/ChaosSolvers/Private/PhysicsProxy/GeometryCollectionPhysicsProxy.cpp
@@ -2497,9 +2497,42 @@ void FGeometryCollectionPhysicsProxy::ProcessCommands(FParticlesType& Particles,
 							SamplesView,
 							Command.MetaData };
 
-						TArrayView<int32> DynamicStateView(&(Collection.DynamicState[0]), Collection.DynamicState.Num());
-						static_cast<const FFieldNode<int32> *>(Command.RootNode.Get())->Evaluate(Context, DynamicStateView);
+						TArray<int32> DynamicState; // #BGTODO Enum class support (so we can size the underlying type to be more appropriate)
+						DynamicState.AddUninitialized(Handles.Num());
+						int32 i = 0;
+						for (Chaos::TGeometryParticleHandle<float, 3> * Handle : Handles)
+						{
+							const Chaos::EObjectStateType ObjectState = Handle->ObjectState();
+							switch (ObjectState)
+							{
+							case Chaos::EObjectStateType::Kinematic:
+								DynamicState[i++] = (int)EObjectStateTypeEnum::Chaos_Object_Kinematic;
+								break;
+							case Chaos::EObjectStateType::Static:
+								DynamicState[i++] = (int)EObjectStateTypeEnum::Chaos_Object_Static;
+								break;
+							case Chaos::EObjectStateType::Sleeping:
+								DynamicState[i++] = (int)EObjectStateTypeEnum::Chaos_Object_Sleeping;
+								break;
+							case Chaos::EObjectStateType::Dynamic:
+							case Chaos::EObjectStateType::Uninitialized:
+							default:
+								DynamicState[i++] = (int)EObjectStateTypeEnum::Chaos_Object_Dynamic;
+								break;
+							}
+						}
+						
+						TArrayView<int32> DynamicStateView(&(DynamicState[0]), DynamicState.Num());						
+						static_cast<const FFieldNode<int32>*>(Command.RootNode.Get())->Evaluate(Context, DynamicStateView);
+						
+						for (const ContextIndex& Index : Context.GetEvaluatedSamples())
+						{
+							Chaos::TGeometryParticleHandle<float, 3>* Handle = Handles[Index.Sample];
+
+							Collection.DynamicState[HandleToTransformGroupIndex[Handle->CastToRigidParticle()]] = DynamicStateView[Index.Result];
+						} // end for all samples
 
+						// push values back to dynamic collection					
 						PushKinematicStateToSolver(Particles);
 					}
 				}
