Chaos - Joint batchmode fixes and finish ISPC

diff --git a/Engine/Source/Runtime/Experimental/Chaos/Private/Chaos/Joint/PBDJointSolverGaussSeidel.ispc b/Engine/Source/Runtime/Experimental/Chaos/Private/Chaos/Joint/PBDJointSolverGaussSeidel.ispc
--- a/Engine/Source/Runtime/Experimental/Chaos/Private/Chaos/Joint/PBDJointSolverGaussSeidel.ispc
+++ b/Engine/Source/Runtime/Experimental/Chaos/Private/Chaos/Joint/PBDJointSolverGaussSeidel.ispc
@@ -562,6 +562,147 @@ export void BatchUpdateDerivedState(uniform FJointSolverJointState Joints[], uni
 	BatchUpdateDerivedStateImpl(Joints, JointIndexBegin, JointIndexEnd);
 }
 
+static inline void BatchApplyPositionConstraint1(
+	const uniform float Dt,
+	uniform FJointSolverJointState& JointState,
+	uniform FJointSolverConstraintRowData& RowData,
+	uniform FJointSolverConstraintRowState& RowState)
+{
+	uniform float II = 0.0f;
+	uniform FVector AngularAxis0 = ZeroVector;
+	uniform FVector AngularAxis1 = ZeroVector;
+
+	if (JointState.InvMs[0] > 0.0f)
+	{
+		AngularAxis0 = VectorCross(JointState.Xs[0] - JointState.Ps[0], RowState.Axis);
+		const uniform FVector IA0 = Multiply(JointState.InvIs[0], AngularAxis0);
+		const uniform float II0 = VectorDot(AngularAxis0, IA0);
+		II += JointState.InvMs[0] + II0;
+	}
+	if (JointState.InvMs[1] > 0.0f)
+	{
+		AngularAxis1 = VectorCross(JointState.Xs[1] - JointState.Ps[1], RowState.Axis);
+		const uniform FVector IA1 = Multiply(JointState.InvIs[1], AngularAxis1);
+		const uniform float II1 = VectorDot(AngularAxis1, IA1);
+		II += JointState.InvMs[1] + II1;
+	}
+
+	if (RowData.bIsSoft)
+	{
+		uniform float VelDt = 0;
+		if (RowData.Damping > KINDA_SMALL_NUMBER)
+		{
+			const uniform FVector V0 = CalculateVelocity(JointState.PrevXs[0], JointState.Xs[0], 1.0f);
+			const uniform FVector V1 = CalculateVelocity(JointState.PrevXs[1], JointState.Xs[1], 1.0f);
+			VelDt = VectorDot(V0 - V1, RowState.Axis);
+		}
+
+		const uniform float SpringMassScale = (RowData.bIsAccelerationMode) ? 1.0f / (JointState.InvMs[0] + JointState.InvMs[1]) : 1.0f;
+		const uniform float S = SpringMassScale * RowData.Stiffness * Dt * Dt;
+		const uniform float D = SpringMassScale * RowData.Damping * Dt;
+		const uniform float Multiplier = 1.0f / ((S + D) * II + 1.0f);
+		const uniform float DLambda = Multiplier * (S * RowState.Error - D * VelDt - RowState.Lambda);
+
+		RowState.DPs[0] = (JointState.InvMs[0] * DLambda) * RowState.Axis;
+		RowState.DPs[1] = (-JointState.InvMs[1] * DLambda) * RowState.Axis;
+		RowState.DRs[0] = DLambda * Multiply(JointState.InvIs[0], AngularAxis0);
+		RowState.DRs[1] = -DLambda * Multiply(JointState.InvIs[1], AngularAxis1);
+		RowState.Lambda += DLambda;
+	}
+	else
+	{
+		const uniform FVector DX = RowState.Axis * (RowData.Stiffness * RowState.Error / II);
+
+		// Apply constraint correction
+		RowState.DPs[0] = JointState.InvMs[0] * DX;
+		RowState.DPs[1] = JointState.InvMs[1] * VectorNegate(DX);
+		RowState.DRs[0] = Multiply(JointState.InvIs[0], VectorCross(JointState.Xs[0] - JointState.Ps[0], DX));
+		RowState.DRs[1] = Multiply(JointState.InvIs[1], VectorCross(JointState.Xs[1] - JointState.Ps[1], VectorNegate(DX)));
+	}
+}
+
+static inline void BatchApplyDelta(
+	uniform FJointSolverJointState* uniform JointStates, 
+	uniform int JointIndex)
+{
+	if (JointStates[JointIndex].InvMs[0] > 0.0f)
+	{
+		const uniform FVector4 Qs0 = JointStates[JointIndex].Qs[0];
+		const uniform FVector4 DR0 = SetVector4(JointStates[JointIndex].DRs[0], 0);
+		const uniform FVector4 DQ0 = VectorQuaternionMultiply2(DR0, Qs0) * 0.5f;
+		const uniform FVector4 Q0 = VectorNormalizeQuaternion(Qs0 + DQ0);
+		JointStates[JointIndex].Ps[0] = JointStates[JointIndex].Ps[0] + JointStates[JointIndex].DPs[0];
+		JointStates[JointIndex].Qs[0] = Q0;
+	}
+
+	if (JointStates[JointIndex].InvMs[1] > 0.0f)
+	{
+		const uniform FVector4 Qs1 = JointStates[JointIndex].Qs[1];
+		const uniform FVector4 DR1 = SetVector4(JointStates[JointIndex].DRs[1], 0);
+		const uniform FVector4 DQ1 = VectorQuaternionMultiply2(DR1, Qs1) * 0.5f;
+		const uniform FVector4 Q1 = VectorNormalizeQuaternion(Qs1 + DQ1);
+		JointStates[JointIndex].Ps[1] = JointStates[JointIndex].Ps[1] + JointStates[JointIndex].DPs[1];
+		JointStates[JointIndex].Qs[1] = Q1;
+	}
+
+	JointStates[JointIndex].Qs[1] = VectorQuaternionEnforceShortestArcWith(JointStates[JointIndex].Qs[1], JointStates[JointIndex].Qs[0]);
+}
+
+export uniform int BatchApplyPositionConstraints(
+	uniform float Dt, 
+	uniform FJointSolverJointState* uniform JointStates, 
+	uniform FJointSolverConstraintRowData* uniform RowDatas, 
+	uniform FJointSolverConstraintRowState* uniform RowStates, 
+	uniform int JointIndexBegin, 
+	uniform int JointIndexEnd, 
+	uniform int RowIndexBegin, 
+	uniform int RowIndexEnd)
+{
+	uniform int32 NumActiveConstraints = 0;
+
+	// Calculate the delta position and rotation to fix the constraint errors
+	// NOTE: Rows are in Joint order, so we can loop over them directly rather than looping over joints and then their rows.
+	for (uniform int32 RowIndex = RowIndexBegin; RowIndex < RowIndexEnd; /* no op */)
+	{
+		const uniform int32 RowJointIndex = RowDatas[RowIndex].JointIndex;
+		//if (RowDatas[RowIndex].NumRows == 3)
+		//{
+		//	uniform FVector Error = FVec3(RowStates[RowIndex].Error, RowStates[RowIndex + 1].Error, RowStates[RowIndex + 2].Error);
+		//	if (Error.Size() > JointStates[RowJointIndex].PositionTolerance)
+		//	{
+		//		ApplyPositionConstraint3(Dt, JointStates[RowJointIndex], RowDatas[RowIndex], RowStates[RowIndex], RowStates[RowIndex + 1], RowStates[RowIndex + 2]);
+		//		JointStates[RowJointIndex].DPs[0] = JointStates[RowJointIndex].DPs[0] + RowStates[RowIndex].DPs[0];
+		//		JointStates[RowJointIndex].DPs[1] = JointStates[RowJointIndex].DPs[1] + RowStates[RowIndex].DPs[1];
+		//		JointStates[RowJointIndex].DRs[0] = JointStates[RowJointIndex].DRs[0] + RowStates[RowIndex].DRs[0];
+		//		JointStates[RowJointIndex].DRs[1] = JointStates[RowJointIndex].DRs[1] + RowStates[RowIndex].DRs[1];
+		//		NumActiveConstraints += 3;
+		//	}
+		//	RowIndex += 3;
+		//}
+		//else if (RowDatas[RowIndex].NumRows == 1)
+		{
+			if (abs(RowStates[RowIndex].Error) > JointStates[RowJointIndex].PositionTolerance)
+			{
+				BatchApplyPositionConstraint1(Dt, JointStates[RowJointIndex], RowDatas[RowIndex], RowStates[RowIndex]);
+				JointStates[RowJointIndex].DPs[0] = JointStates[RowJointIndex].DPs[0] + RowStates[RowIndex].DPs[0];
+				JointStates[RowJointIndex].DPs[1] = JointStates[RowJointIndex].DPs[1] + RowStates[RowIndex].DPs[1];
+				JointStates[RowJointIndex].DRs[0] = JointStates[RowJointIndex].DRs[0] + RowStates[RowIndex].DRs[0];
+				JointStates[RowJointIndex].DRs[1] = JointStates[RowJointIndex].DRs[1] + RowStates[RowIndex].DRs[1];
+				++NumActiveConstraints;
+			}
+			++RowIndex;
+		}
+	}
+
+	// Accumulate and apply the deltas for each joint
+	for (uniform int32 JointIndex = JointIndexBegin; JointIndex < JointIndexEnd; ++JointIndex)
+	{
+		BatchApplyDelta(JointStates, JointIndex);
+	}
+
+	return NumActiveConstraints;
+}
+
 static inline void BatchApplyRotationConstraint(
 	const uniform float Dt,
 	uniform FJointSolverJointState& JointState,
diff --git a/Engine/Source/Runtime/Experimental/Chaos/Private/Chaos/PBDJointConstraints.cpp b/Engine/Source/Runtime/Experimental/Chaos/Private/Chaos/PBDJointConstraints.cpp
--- a/Engine/Source/Runtime/Experimental/Chaos/Private/Chaos/PBDJointConstraints.cpp
+++ b/Engine/Source/Runtime/Experimental/Chaos/Private/Chaos/PBDJointConstraints.cpp
@@ -498,6 +498,12 @@ namespace Chaos
 
 	void FPBDJointConstraints::UpdatePositionBasedState(const FReal Dt)
 	{
+	}
+
+	void FPBDJointConstraints::PrepareTick()
+	{
+		SCOPE_CYCLE_COUNTER(STAT_Joints_PrepareTick);
+
 		if (bJointsDirty || (bIsBatched != bChaos_Joint_Batching))
 		{
 			DeinitSolverJointData();
@@ -509,16 +515,16 @@ namespace Chaos
 			bIsBatched = bChaos_Joint_Batching;
 			bJointsDirty = false;
 		}
-	}
-
-	void FPBDJointConstraints::PrepareTick()
-	{
-		SCOPE_CYCLE_COUNTER(STAT_Joints_PrepareTick);
 
 		if (bChaos_Joint_Batching)
 		{
 			SolverConstraintRowStates.SetNum(SolverConstraintRowDatas.Num());
 			SolverConstraintStates.SetNum(NumConstraints());
+
+			for (FJointSolverConstraintRowState RowState : SolverConstraintRowStates)
+			{
+				RowState.TickReset();
+			}
 		}
 		else
 		{
@@ -931,10 +937,24 @@ namespace Chaos
 			}
 
 			// Solve and apply the position constraints for all Joints in the batch
-			if (JointIndexEnd > JointIndexBegin)
+			const int32 LinearRowIndexBegin = SolverConstraints[JointIndexBegin].GetLinearRowIndexBegin();
+			const int32 LinearRowIndexEnd = SolverConstraints[JointIndexEnd - 1].GetLinearRowIndexEnd();
+			if (bChaos_Joint_ISPC_Enabled)
+			{
+#if INTEL_ISPC
+				ispc::BatchApplyPositionConstraints(
+					Dt,
+					(ispc::FJointSolverJointState*)SolverConstraintStates.GetData(),
+					(ispc::FJointSolverConstraintRowData*)SolverConstraintRowDatas.GetData(),
+					(ispc::FJointSolverConstraintRowState*)SolverConstraintRowStates.GetData(),
+					JointIndexBegin,
+					JointIndexEnd,
+					LinearRowIndexBegin,
+					LinearRowIndexEnd);
+#endif
+			}
+			else
 			{
-				const int32 RowIndexBegin = SolverConstraints[JointIndexBegin].GetLinearRowIndexBegin();
-				const int32 RowIndexEnd = SolverConstraints[JointIndexEnd - 1].GetLinearRowIndexEnd();
 				NumActive += FJointSolver::ApplyPositionConstraints(
 					Dt, 
 					SolverConstraintStates, 
@@ -942,8 +962,8 @@ namespace Chaos
 					SolverConstraintRowStates, 
 					JointIndexBegin,
 					JointIndexEnd,
-					RowIndexBegin,
-					RowIndexEnd);
+					LinearRowIndexBegin,
+					LinearRowIndexEnd);
 			}
 
 			// Reset accumulators and update derived state
@@ -975,8 +995,8 @@ namespace Chaos
 			}
 
 			// Solve and apply the rotation constraints for all Joints in the batch
-			const int32 RowIndexBegin = SolverConstraints[JointIndexBegin].GetAngularRowIndexBegin();
-			const int32 RowIndexEnd = SolverConstraints[JointIndexEnd - 1].GetAngularRowIndexEnd();
+			const int32 AngularRowIndexBegin = SolverConstraints[JointIndexBegin].GetAngularRowIndexBegin();
+			const int32 AngularRowIndexEnd = SolverConstraints[JointIndexEnd - 1].GetAngularRowIndexEnd();
 			if (bChaos_Joint_ISPC_Enabled)
 			{
 #if INTEL_ISPC
@@ -987,8 +1007,8 @@ namespace Chaos
 					(ispc::FJointSolverConstraintRowState*)SolverConstraintRowStates.GetData(),
 					JointIndexBegin,
 					JointIndexEnd,
-					RowIndexBegin,
-					RowIndexEnd);
+					AngularRowIndexBegin,
+					AngularRowIndexEnd);
 #endif
 			}
 			else
@@ -1000,8 +1020,8 @@ namespace Chaos
 					SolverConstraintRowStates, 
 					JointIndexBegin,
 					JointIndexEnd,
-					RowIndexBegin,
-					RowIndexEnd);
+					AngularRowIndexBegin,
+					AngularRowIndexEnd);
 			}
 		}
 
diff --git a/Engine/Source/Runtime/Experimental/Chaos/Public/Chaos/Joint/JointSolverConstraints.h b/Engine/Source/Runtime/Experimental/Chaos/Public/Chaos/Joint/JointSolverConstraints.h
--- a/Engine/Source/Runtime/Experimental/Chaos/Public/Chaos/Joint/JointSolverConstraints.h
+++ b/Engine/Source/Runtime/Experimental/Chaos/Public/Chaos/Joint/JointSolverConstraints.h
@@ -122,6 +122,9 @@ namespace Chaos
 		{
 		}
 
+		// Reset state that persists over all iterations, but must be cleared at start of each tick.
+		FORCEINLINE void TickReset();
+
 		// Reset calculated values ready for next iteration. Note: Lambda is not reset here,
 		// it accumulates over the whole timestep.
 		FORCEINLINE void IterationReset();
diff --git a/Engine/Source/Runtime/Experimental/Chaos/Public/Chaos/Joint/JointSolverConstraints.inl b/Engine/Source/Runtime/Experimental/Chaos/Public/Chaos/Joint/JointSolverConstraints.inl
--- a/Engine/Source/Runtime/Experimental/Chaos/Public/Chaos/Joint/JointSolverConstraints.inl
+++ b/Engine/Source/Runtime/Experimental/Chaos/Public/Chaos/Joint/JointSolverConstraints.inl
@@ -147,6 +147,11 @@ namespace Chaos
 	///////////////////////////////////////////////////////////////////////
 	//
 
+	FORCEINLINE void FJointSolverConstraintRowState::TickReset()
+	{
+		Lambda = 0.0f;
+	}
+
 	// Reset calculated values ready for next iteration. Note: Lambda is not reset here,
 	// it accumulates over the whole timestep.
 	FORCEINLINE void FJointSolverConstraintRowState::IterationReset()
