Formulate distance error metrics relative to the size of a grid cell.  Abandon comparing face normals to grid normals, and instead ensure consistency with normals with the bounding box.  Disable CVarFailureOnHighError by default.

diff --git a/Engine/Source/Runtime/Experimental/Chaos/Private/Chaos/Levelset.cpp b/Engine/Source/Runtime/Experimental/Chaos/Private/Chaos/Levelset.cpp
--- a/Engine/Source/Runtime/Experimental/Chaos/Private/Chaos/Levelset.cpp
+++ b/Engine/Source/Runtime/Experimental/Chaos/Private/Chaos/Levelset.cpp
@@ -19,13 +19,13 @@
 int32 OutputFailedLevelSetDebugData = 0;
 FAutoConsoleVariableRef CVarOutputFailedLevelSetDebugData(TEXT("p.LevelSetOutputFailedDebugData"), OutputFailedLevelSetDebugData, TEXT("Output debug obj files for level set and mesh when error tolerances are too high"));
 
-int32 FailureOnHighError = 1;
+int32 FailureOnHighError = 0;
 FAutoConsoleVariableRef CVarFailureOnHighError(TEXT("p.LevelSetFailureOnHighError"), FailureOnHighError, TEXT("Set level sets with high error to null in the solver"));
 
-float AvgDistErrorTolerance = .05;
+float AvgDistErrorTolerance = 1.f;
 FAutoConsoleVariableRef CVarAvgDistErrorTolerance(TEXT("p.LevelSetAvgDistErrorTolerance"), AvgDistErrorTolerance, TEXT("Error tolerance for average distance between the triangles and generated levelset.  Note this is a fraction of the average bounding box dimensions."));
 
-float MaxDistErrorTolerance = .1;
+float MaxDistErrorTolerance = 1.f;
 FAutoConsoleVariableRef CVarMaxDistErrorTolerance(TEXT("p.LevelSetMaxDistErrorTolerance"), MaxDistErrorTolerance, TEXT("Max error for the highest error triangle generated from a levelset.  Note this is a fraction of the average bounding box dimensions."));
 
 float AvgAngleErrorTolerance = 1.;
@@ -259,8 +259,18 @@ T TLevelSet<T, d>::ComputeLevelSetError(const TParticles<T, d>& InParticles, con
 	TArray<T> DistErrorValues;
 	DistErrorValues.AddDefaulted(Mesh.GetNumElements());
 
-	TArray<T> AngleErrorValues;
-	AngleErrorValues.AddDefaulted(Mesh.GetNumElements());
+	// Testing that the grid normal points generally in the same direction as the face normal
+	// depends on the geometry being very clean.  It's entirely reasonable that the geometry may
+	// have geometry on the interior that gets filled in on the grid, at which point the gradient
+	// should no longer be beholden to the direction of the surface normals.  So this test only 
+	// really makes sense for points on the zero isocontour.  But even then, I think this is more
+	// of a test that the geometry is well formed and consistent than whether the grid normals
+	// point the right direction.  And so, until we have time to formulate a better test, we 
+	// simply test that the level set normals point rougly in the same directions as the bounding 
+	// box normals, at the corners of the bounding box.  See below...
+
+	//TArray<T> AngleErrorValues;
+	//AngleErrorValues.AddDefaulted(Mesh.GetNumElements());
 
 	TArray<T> TriangleArea;
 	TriangleArea.AddDefaulted(Mesh.GetNumElements());
@@ -271,8 +281,9 @@ T TLevelSet<T, d>::ComputeLevelSetError(const TParticles<T, d>& InParticles, con
 		const TVector<int32, 3> CurrMeshFace = Faces[i];
 		const TVector<T, 3> MeshFaceCenter = (InParticles.X(CurrMeshFace[0]) + InParticles.X(CurrMeshFace[1]) + InParticles.X(CurrMeshFace[2])) / 3.f;
 
-		TVector<T, 3> GridNormal;
-		T phi = PhiWithNormal(MeshFaceCenter, GridNormal);
+		//TVector<T, 3> GridNormal;
+		//T phi = PhiWithNormal(MeshFaceCenter, GridNormal);
+		const T phi = SignedDistance(MeshFaceCenter);
 
 		// ignore triangles where the the center is more than 2 voxels inside
 		// #note: this biases the statistics since what we really want to do is preprocess for interior triangles, but
@@ -290,11 +301,10 @@ T TLevelSet<T, d>::ComputeLevelSetError(const TParticles<T, d>& InParticles, con
 			DistErrorValues[i] /= 4.f;
 
 			// angle error computed by angle between mesh face normal and level set gradient
-			TVector<T, 3> MeshFaceNormal = Normals[i];
-			MeshFaceNormal.SafeNormalize();
-
-			GridNormal.SafeNormalize();
-			AngleErrorValues[i] = FMath::Acos(TVector<T, 3>::DotProduct(MeshFaceNormal, GridNormal));
+			//TVector<T, 3> MeshFaceNormal = Normals[i];
+			//MeshFaceNormal.SafeNormalize();
+			//GridNormal.SafeNormalize();
+			//AngleErrorValues[i] = FMath::Acos(TVector<T, 3>::DotProduct(MeshFaceNormal, GridNormal));
 
 			// triangle area used for weighted average
 			TriangleArea[i] = .5 * sqrt(TVector<T, d>::CrossProduct(InParticles.X(CurrMeshFace[1]) - InParticles.X(CurrMeshFace[0]), InParticles.X(CurrMeshFace[2]) - InParticles.X(CurrMeshFace[0])).SizeSquared());
@@ -302,7 +312,7 @@ T TLevelSet<T, d>::ComputeLevelSetError(const TParticles<T, d>& InParticles, con
 	});
 
 	float TotalDistError = 0.f;
-	float TotalAngleError = 0.f;
+	//float TotalAngleError = 0.f;
 	float TotalTriangleArea = 0.f;
 	float MaxError = -1. * FLT_MAX;
 	for (int i = 0; i < Mesh.GetNumElements(); ++i)
@@ -314,7 +324,7 @@ T TLevelSet<T, d>::ComputeLevelSetError(const TParticles<T, d>& InParticles, con
 
 		// weight the error values by the area
 		TotalDistError += DistErrorValues[i] * TriangleArea[i];
-		TotalAngleError += AngleErrorValues[i] * TriangleArea[i];
+		//TotalAngleError += AngleErrorValues[i] * TriangleArea[i];
 		TotalTriangleArea += TriangleArea[i];
 	}
 
@@ -322,7 +332,7 @@ T TLevelSet<T, d>::ComputeLevelSetError(const TParticles<T, d>& InParticles, con
 	if (TotalTriangleArea < 1e-5)
 	{
 		MaxDistError = MAX_flt;
-		AngleError = MAX_flt;
+		//AngleError = MAX_flt;
 		return MAX_flt;
 	}
 
@@ -337,14 +347,60 @@ T TLevelSet<T, d>::ComputeLevelSetError(const TParticles<T, d>& InParticles, con
 	if (AvgExtents < 1e-5)
 	{
 		MaxDistError = MAX_flt;
-		AngleError = MAX_flt;
+		//AngleError = MAX_flt;
 		return MAX_flt;
 	}
 
 	AvgDistError /= AvgExtents;
 	MaxDistError = MaxError / AvgExtents;
 
-	AngleError = TotalAngleError / TotalTriangleArea;
+	//AngleError = TotalAngleError / TotalTriangleArea;
+
+	// Test the normal directions at the corners of the bounding box that they
+	// point outward.
+	const TAABB<T, d> BBox = BoundingBox();
+	const TVector<T, 3>& MinPt = BBox.Min();
+	const TVector<T, 3>& MaxPt = BBox.Max();
+	
+	bool Fail = false;
+	TVector<T, 3> Pt = MinPt;
+	TVector<T, d> BoxNorm, LSNorm;
+	for (int i = 0; i < 8; i++)
+	{
+		// i
+		// 0 - (min, min, min) MinPt
+		// 1 - (max, min, min)
+		// 2 - (min, max, min)
+		// 3 - (min, min, max)
+
+		// 4 - (max, max, max) MaxPt
+		// 5 - (min, max, max)
+		// 6 - (max, min, max)
+		// 7 - (max, max, min)
+
+		if (i <= 3)
+		{
+			Pt = MinPt;
+			if (i == 1) Pt[0] = MaxPt[0];
+			else if (i == 2) Pt[1] = MaxPt[1];
+			else if (i == 3) Pt[2] = MaxPt[2];
+		}
+		else
+		{
+			Pt = MaxPt;
+			if (i == 1) Pt[0] = MinPt[0];
+			else if (i == 2) Pt[1] = MinPt[1];
+			else if (i == 3) Pt[2] = MinPt[2];
+		}
+
+		BBox.PhiWithNormal(Pt, BoxNorm);
+		PhiWithNormal(Pt, LSNorm);
+		const float Dot = Chaos::TVector<T, 3>::DotProduct(BoxNorm, LSNorm);
+		if (Dot < 0)
+		{
+			AngleError += FMath::Abs(FMath::Acos(Dot));
+		}
+	}
 
 	return AvgDistError;
 }
@@ -512,7 +568,7 @@ bool TLevelSet<T, d>::CheckData(FErrorReporter& ErrorReporter, const TParticles<
 	float AvgDistError = ComputeLevelSetError(InParticles, Normals, Mesh, AvgAngleError, MaxDistError);
 
 	// Report high error, but don't report it as an invalid level set
-	if (AvgDistError > AvgDistErrorTolerance || AvgAngleError > AvgAngleErrorTolerance || MaxDistError > MaxDistErrorTolerance)
+	if (AvgDistError > AvgDistErrorTolerance*MGrid.Dx().Size() || AvgAngleError > AvgAngleErrorTolerance || MaxDistError > MaxDistErrorTolerance*MGrid.Dx().Size())
 	{
 		if (OutputFailedLevelSetDebugData)
 		{
@@ -522,13 +578,20 @@ bool TLevelSet<T, d>::CheckData(FErrorReporter& ErrorReporter, const TParticles<
 
 		if (FailureOnHighError)
 		{
-			FString ErrorStr = FString::Printf(TEXT("High error for level set: AvgDistError: %f, MaxDistError: %f, AvgAngleError: %f"), AvgDistError, MaxDistError, AvgAngleError);
+			FString ErrorStr = FString::Printf(TEXT("High error for level set: AvgDistError: %f (Max: %f*%f), MaxDistError: %f (Max: %f*%f), AvgAngleError: %f (Max: %f)"), 
+				AvgDistError, AvgDistErrorTolerance, MGrid.Dx().Size(),
+				MaxDistError, MaxDistErrorTolerance, MGrid.Dx().Size(),
+				AvgAngleError, AvgAngleErrorTolerance);
 			ErrorReporter.ReportError(*ErrorStr);
 			return false;
 		}
 		else
 		{
-			UE_LOG(LogChaos, Log, TEXT("%s: High error for level set: AvgDistError: %f, MaxDistError: %f, AvgAngleError: %f"), *ErrorReporter.GetPrefix(), AvgDistError, MaxDistError, AvgAngleError);
+			UE_LOG(LogChaos, Log, TEXT("%s: High error for level set: AvgDistError: %f (Max: %f*%f), MaxDistError: %f (Max: %f*%f), AvgAngleError: %f (Max: %f)"), 
+				*ErrorReporter.GetPrefix(), 
+				AvgDistError, AvgDistErrorTolerance, MGrid.Dx().Size(),
+				MaxDistError, MaxDistErrorTolerance, MGrid.Dx().Size(),
+				AvgAngleError, AvgAngleErrorTolerance);
 		}
 	}
 
