Allow chaos solvers to tick in editor builds when running dedicated server only

diff --git a/Engine/Source/Runtime/Experimental/ChaosSolvers/Private/ChaosSolversModule.cpp b/Engine/Source/Runtime/Experimental/ChaosSolvers/Private/ChaosSolversModule.cpp
--- a/Engine/Source/Runtime/Experimental/ChaosSolvers/Private/ChaosSolversModule.cpp
+++ b/Engine/Source/Runtime/Experimental/ChaosSolvers/Private/ChaosSolversModule.cpp
@@ -177,7 +177,7 @@ FChaosSolversModule::FChaosSolversModule()
 	, bModuleInitialized(false)
 {
 #if WITH_EDITOR
-	if(!IsRunningGame())
+	if(!(IsRunningDedicatedServer() || IsRunningGame()))
 	{
 		// In the editor we begin with everything paused so we don't needlessly tick
 		// the physics solvers until PIE begins. Delegates are bound in FPhysScene_ChaosPauseHandler
