Expose more read API for rewind data

diff --git a/Engine/Source/Runtime/Experimental/ChaosSolvers/Public/RewindData.h b/Engine/Source/Runtime/Experimental/ChaosSolvers/Public/RewindData.h
--- a/Engine/Source/Runtime/Experimental/ChaosSolvers/Public/RewindData.h
+++ b/Engine/Source/Runtime/Experimental/ChaosSolvers/Public/RewindData.h
@@ -160,12 +160,56 @@ public:
 		return NonFrequentData.IsSet() ? MakeSerializable(NonFrequentData.Read().Geometry()) : Particle.Geometry();
 	}
 
+	template <typename TParticle>
+	void* UserData(const TParticle& Particle) const
+	{
+		return NonFrequentData.IsSet() ? NonFrequentData.Read().UserData() : Particle.UserData();
+	}
+
+	template <typename TParticle>
+	FUniqueIdx UniqueIdx(const TParticle& Particle) const
+	{
+		return NonFrequentData.IsSet() ? NonFrequentData.Read().UniqueIdx() : Particle.UniqueIdx();
+	}
+	
+	template <typename TParticle>
+	FSpatialAccelerationIdx SpatialIdx(const TParticle& Particle) const
+	{
+		return NonFrequentData.IsSet() ? NonFrequentData.Read().SpatialIdx() : Particle.SpatialIdx();
+	}
+
+#if CHAOS_CHECKED
+	template <typename TParticle>
+	FName DebugName(const TParticle& Particle) const
+	{
+		return NonFrequentData.IsSet() ? NonFrequentData.Read().DebugName() : Particle.DebugName();
+	}
+#endif
+
 	template <typename TParticle>
 	const FVec3& F(const TParticle& Particle) const
 	{
 		return Dynamics.IsSet() ? Dynamics.Read().F() : Particle.CastToRigidParticle()->F();
 	}
 
+	template <typename TParticle>
+	const FVec3& Torque(const TParticle& Particle) const
+	{
+		return Dynamics.IsSet() ? Dynamics.Read().Torque() : Particle.CastToRigidParticle()->Torque();
+	}
+
+	template <typename TParticle>
+	const FVec3& LinearImpulse(const TParticle& Particle) const
+	{
+		return Dynamics.IsSet() ? Dynamics.Read().LinearImpulse() : Particle.CastToRigidParticle()->LinearImpulse();
+	}
+
+	template <typename TParticle>
+	const FVec3& AngularImpulse(const TParticle& Particle) const
+	{
+		return Dynamics.IsSet() ? Dynamics.Read().AngularImpulse() : Particle.CastToRigidParticle()->AngularImpulse();
+	}
+
 	void SyncSimWritablePropsFromSim(FDirtyPropertiesManager& Manager,int32 Idx,const TPBDRigidParticleHandle<FReal,3>& Rigid)
 	{
 		FParticleDirtyFlags Flags;
@@ -397,11 +441,48 @@ public:
 		return State.Geometry(Particle);
 	}
 
+	void* UserData() const
+	{
+		return State.UserData(Particle);
+	}
+
+	FUniqueIdx UniqueIdx() const
+	{
+		return State.UniqueIdx(Particle);
+	}
+
+	FSpatialAccelerationIdx SpatialIdx() const
+	{
+		return State.SpatialIdx(Particle);
+	}
+
+#if CHAOS_CHECKED
+	FName DebugName() const
+	{
+		return State.DebugName(Particle);
+	}
+#endif
+
 	const FVec3& F() const
 	{
 		return State.F(Particle);
 	}
 
+	const FVec3& Torque() const
+	{
+		return State.Torque(Particle);
+	}
+
+	const FVec3& LinearImpulse() const
+	{
+		return State.LinearImpulse(Particle);
+	}
+
+	const FVec3& AngularImpulse() const
+	{
+		return State.AngularImpulse(Particle);
+	}
+
 	const TGeometryParticle<FReal,3>& GetParticle() const
 	{
 		return Particle;
