Rewind data records all particle data (excluding shapes)

diff --git a/Engine/Source/Runtime/Engine/Private/PhysicsEngine/Experimental/PhysInterface_Chaos.cpp b/Engine/Source/Runtime/Engine/Private/PhysicsEngine/Experimental/PhysInterface_Chaos.cpp
--- a/Engine/Source/Runtime/Engine/Private/PhysicsEngine/Experimental/PhysInterface_Chaos.cpp
+++ b/Engine/Source/Runtime/Engine/Private/PhysicsEngine/Experimental/PhysInterface_Chaos.cpp
@@ -872,7 +872,7 @@ bool FPhysInterface_Chaos::IsGravityEnabled_AssumesLocked(const FPhysicsActorHan
 {
 	if (Chaos::TPBDRigidParticle<float, 3 >* RigidParticle = InActorReference->CastToRigidParticle())
 	{
-		return RigidParticle->IsGravityEnabled();
+		return RigidParticle->GravityEnabled();
 	}
 	return false;
 }
diff --git a/Engine/Source/Runtime/Experimental/Chaos/Public/Chaos/ParticleHandle.h b/Engine/Source/Runtime/Experimental/Chaos/Public/Chaos/ParticleHandle.h
--- a/Engine/Source/Runtime/Experimental/Chaos/Public/Chaos/ParticleHandle.h
+++ b/Engine/Source/Runtime/Experimental/Chaos/Public/Chaos/ParticleHandle.h
@@ -2022,7 +2022,7 @@ public:
 		MMiscData.Modify(true,MDirtyFlags,Proxy,[InDisabled](auto& Data){ Data.bDisabled = InDisabled;});
 	}*/
 
-	bool IsGravityEnabled() const { return MMiscData.Read().GravityEnabled(); }
+	bool GravityEnabled() const { return MMiscData.Read().GravityEnabled(); }
 	void SetGravityEnabled(const bool InGravityEnabled)
 	{
 		MMiscData.Modify(true,MDirtyFlags,Proxy,[InGravityEnabled](auto& Data){ Data.SetGravityEnabled (InGravityEnabled);});
@@ -2275,7 +2275,7 @@ public:
 		, MObjectState(InParticle.ObjectState())
 		//, MDisabled(InParticle.Disabled())
 		, MToBeRemovedOnFracture(InParticle.ToBeRemovedOnFracture())
-		, MGravityEnabled(InParticle.IsGravityEnabled())
+		, MGravityEnabled(InParticle.GravityEnabled())
 		, MInitialized(InParticle.IsInitialized())
 	{
 		Type = EParticleType::Rigid;
@@ -2345,7 +2345,7 @@ public:
 			MObjectState = InParticle.ObjectState();
 			//MDisabled = InParticle.Disabled();
 			MToBeRemovedOnFracture = InParticle.ToBeRemovedOnFracture();
-			MGravityEnabled = InParticle.IsGravityEnabled();
+			MGravityEnabled = InParticle.GravityEnabled();
 			MInitialized = InParticle.IsInitialized();
 			Type = EParticleType::Rigid;
 		}
diff --git a/Engine/Source/Runtime/Experimental/ChaosSolvers/Public/RewindData.h b/Engine/Source/Runtime/Experimental/ChaosSolvers/Public/RewindData.h
--- a/Engine/Source/Runtime/Experimental/ChaosSolvers/Public/RewindData.h
+++ b/Engine/Source/Runtime/Experimental/ChaosSolvers/Public/RewindData.h
@@ -154,6 +154,72 @@ public:
 		return Velocities.IsSet() ? Velocities.Read().W() : Particle.CastToKinematicParticle()->W();
 	}
 
+	template <typename TParticle>
+	FReal LinearEtherDrag(const TParticle& Particle) const
+	{
+		return DynamicsMisc.IsSet() ? DynamicsMisc.Read().LinearEtherDrag() : Particle.CastToRigidParticle()->LinearEtherDrag();
+	}
+
+	template <typename TParticle>
+	FReal AngularEtherDrag(const TParticle& Particle) const
+	{
+		return DynamicsMisc.IsSet() ? DynamicsMisc.Read().AngularEtherDrag() : Particle.CastToRigidParticle()->AngularEtherDrag();
+	}
+
+	template <typename TParticle>
+	EObjectStateType ObjectState(const TParticle& Particle) const
+	{
+		return DynamicsMisc.IsSet() ? DynamicsMisc.Read().ObjectState() : Particle.CastToRigidParticle()->ObjectState();
+	}
+
+	template <typename TParticle>
+	bool GravityEnabled(const TParticle& Particle) const
+	{
+		return DynamicsMisc.IsSet() ? DynamicsMisc.Read().GravityEnabled() : Particle.CastToRigidParticle()->GravityEnabled();
+	}
+
+	template <typename TParticle>
+	int32 CollisionGroup(const TParticle& Particle) const
+	{
+		return DynamicsMisc.IsSet() ? DynamicsMisc.Read().CollisionGroup() : Particle.CastToRigidParticle()->CollisionGroup();
+	}
+
+	template <typename TParticle>
+	const FVec3& CenterOfMass(const TParticle& Particle) const
+	{
+		return MassProps.IsSet() ? MassProps.Read().CenterOfMass() : Particle.CastToRigidParticle()->CenterOfMass();
+	}
+
+	template <typename TParticle>
+	const FRotation3& RotationOfMass(const TParticle& Particle) const
+	{
+		return MassProps.IsSet() ? MassProps.Read().RotationOfMass() : Particle.CastToRigidParticle()->RotationOfMass();
+	}
+
+	template <typename TParticle>
+	const FMatrix33& I(const TParticle& Particle) const
+	{
+		return MassProps.IsSet() ? MassProps.Read().I() : Particle.CastToRigidParticle()->I();
+	}
+
+	template <typename TParticle>
+	const FMatrix33& InvI(const TParticle& Particle) const
+	{
+		return MassProps.IsSet() ? MassProps.Read().InvI() : Particle.CastToRigidParticle()->InvI();
+	}
+
+	template <typename TParticle>
+	FReal M(const TParticle& Particle) const
+	{
+		return MassProps.IsSet() ? MassProps.Read().M() : Particle.CastToRigidParticle()->M();
+	}
+
+	template <typename TParticle>
+	FReal InvM(const TParticle& Particle) const
+	{
+		return MassProps.IsSet() ? MassProps.Read().InvM() : Particle.CastToRigidParticle()->InvM();
+	}
+
 	template <typename TParticle>
 	TSerializablePtr<FImplicitObject> Geometry(const TParticle& Particle) const
 	{
@@ -243,8 +309,6 @@ public:
 
 	void SyncToParticle(TGeometryParticle<FReal,3>& Particle) const
 	{
-		//todo: set properties directly instead of assigning sub-properties
-
 		ParticlePositionRotation.SyncToParticle([&Particle](const auto& Data)
 		{
 			Particle.SetXR(Data);
@@ -258,7 +322,6 @@ public:
 			});
 		}
 
-
 		NonFrequentData.SyncToParticle([&Particle](const auto& Data)
 		{
 			Particle.SetNonFrequentData(Data);
@@ -270,6 +333,16 @@ public:
 			{
 				Rigid->SetDynamics(Data);
 			});
+
+			DynamicsMisc.SyncToParticle([Rigid](const auto& Data)
+			{
+				Rigid->SetDynamicMisc(Data);
+			});
+
+			MassProps.SyncToParticle([Rigid](const auto& Data)
+			{
+				Rigid->SetMassProps(Data);
+			});
 		}
 	}
 
@@ -305,6 +378,11 @@ public:
 		{
 			Data.CopyFrom(*Handle);
 		});
+
+		NonFrequentData.SyncRemoteData(Manager,Idx,Dirty.ParticleData,[Handle](FParticleNonFrequentData& Data)
+		{
+			Data.CopyFrom(*Handle);
+		});
 	}
 
 	void SyncIfDirty(FDirtyPropertiesManager& Manager,int32 Idx,const TGeometryParticle<FReal,3>& InParticle, const FGeometryParticleStateBase& RewindState)
@@ -331,12 +409,23 @@ public:
 			}
 		}
 		
-		if(RewindState.NonFrequentData.IsSet())
+		if(auto Rigid = Particle->CastToRigidParticle())
 		{
-			NonFrequentData.SyncRemoteDataForced(Manager,Idx,[Particle](FParticleNonFrequentData& Data)
+			if(RewindState.DynamicsMisc.IsSet())
 			{
-				Data.CopyFrom(*Particle);
-			});
+				DynamicsMisc.SyncRemoteDataForced(Manager,Idx,[Rigid](FParticleDynamicMisc& Data)
+				{
+					Data.CopyFrom(*Rigid);
+				});
+			}
+
+			if(RewindState.MassProps.IsSet())
+			{
+				MassProps.SyncRemoteDataForced(Manager,Idx,[Rigid](FParticleMassProps& Data)
+				{
+					Data.CopyFrom(*Rigid);
+				});
+			}
 		}
 	}
 
@@ -361,6 +450,18 @@ public:
 			bCoalesced = true;
 		}
 
+		if(!MassProps.IsSet() && LatestState.MassProps.IsSet())
+		{
+			MassProps = LatestState.MassProps;
+			bCoalesced = true;
+		}
+
+		if(!DynamicsMisc.IsSet() && LatestState.DynamicsMisc.IsSet())
+		{
+			DynamicsMisc = LatestState.DynamicsMisc;
+			bCoalesced = true;
+		}
+
 		//dynamics do not coalesce since they are always written when dirty
 
 		return bCoalesced;
@@ -392,13 +493,8 @@ private:
 	TParticleStateProperty<FParticleNonFrequentData,EParticleProperty::NonFrequentData> NonFrequentData;
 	TParticleStateProperty<FParticleVelocities,EParticleProperty::Velocities> Velocities;
 	TParticleStateProperty<FParticleDynamics,EParticleProperty::Dynamics> Dynamics;
-	/*
-	PARTICLE_PROPERTY(XR,FParticlePositionRotation)
-		PARTICLE_PROPERTY(Velocities,FParticleVelocities)
-		PARTICLE_PROPERTY(Dynamics,FParticleDynamics)
-		PARTICLE_PROPERTY(Misc,FParticleMisc)
-		PARTICLE_PROPERTY(NonFrequentData,FParticleNonFrequentData)
-		PARTICLE_PROPERTY(MassProps,FParticleMassProps)*/
+	TParticleStateProperty<FParticleDynamicMisc,EParticleProperty::DynamicMisc> DynamicsMisc;
+	TParticleStateProperty<FParticleMassProps,EParticleProperty::MassProps> MassProps;
 };
 
 class FGeometryParticleState
@@ -436,6 +532,61 @@ public:
 		return State.W(Particle);
 	}
 
+	FReal LinearEtherDrag() const
+	{
+		return State.LinearEtherDrag(Particle);
+	}
+
+	FReal AngularEtherDrag() const
+	{
+		return State.AngularEtherDrag(Particle);
+	}
+
+	EObjectStateType ObjectState() const
+	{
+		return State.ObjectState(Particle);
+	}
+
+	bool GravityEnabled() const
+	{
+		return State.GravityEnabled(Particle);
+	}
+
+	int32 CollisionGroup() const
+	{
+		return State.CollisionGroup(Particle);
+	}
+
+	const FVec3& CenterOfMass() const
+	{
+		return State.CenterOfMass(Particle);
+	}
+
+	const FRotation3& RotationOfMass() const
+	{
+		return State.RotationOfMass(Particle);
+	}
+
+	const FMatrix33& I() const
+	{
+		return State.I(Particle);
+	}
+
+	const FMatrix33& InvI() const
+	{
+		return State.InvI(Particle);
+	}
+
+	FReal M() const
+	{
+		return State.M(Particle);
+	}
+
+	FReal InvM() const
+	{
+		return State.InvM(Particle);
+	}
+
 	TSerializablePtr<FImplicitObject> Geometry() const
 	{
 		return State.Geometry(Particle);
