Chaos: Fixes to get event demo level working again

diff --git a/Engine/Plugins/Experimental/ChaosNiagara/Source/ChaosNiagara/Classes/NiagaraDataInterfaceChaosDestruction.h b/Engine/Plugins/Experimental/ChaosNiagara/Source/ChaosNiagara/Classes/NiagaraDataInterfaceChaosDestruction.h
--- a/Engine/Plugins/Experimental/ChaosNiagara/Source/ChaosNiagara/Classes/NiagaraDataInterfaceChaosDestruction.h
+++ b/Engine/Plugins/Experimental/ChaosNiagara/Source/ChaosNiagara/Classes/NiagaraDataInterfaceChaosDestruction.h
@@ -86,7 +86,7 @@ struct FSolverData
 		: Solver(nullptr)
 	{}
 
-	FPhysScene* PhysScene;
+	FPhysScene_Chaos* PhysScene;
 	Chaos::FPhysicsSolver* Solver;
 };
 
diff --git a/Engine/Plugins/Experimental/ChaosNiagara/Source/ChaosNiagara/Private/NiagaraDataInterfaceChaosDestruction.cpp b/Engine/Plugins/Experimental/ChaosNiagara/Source/ChaosNiagara/Private/NiagaraDataInterfaceChaosDestruction.cpp
--- a/Engine/Plugins/Experimental/ChaosNiagara/Source/ChaosNiagara/Private/NiagaraDataInterfaceChaosDestruction.cpp
+++ b/Engine/Plugins/Experimental/ChaosNiagara/Source/ChaosNiagara/Private/NiagaraDataInterfaceChaosDestruction.cpp
@@ -455,8 +455,8 @@ bool UNiagaraDataInterfaceChaosDestruction::InitPerInstanceData(void* PerInstanc
 					int32 NewIdx = Solvers.Add(FSolverData());
 
 					FSolverData& SolverData = Solvers[NewIdx];
-					SolverData.PhysScene = World->GetPhysicsScene();
-					SolverData.Solver = SolverData.PhysScene->GetScene().GetSolver();
+					SolverData.PhysScene = &World->GetPhysicsScene()->GetScene();
+					SolverData.Solver = SolverData.PhysScene->GetSolver();
 
 					RegisterWithSolverEventManager(SolverData.Solver);
 				}
@@ -464,27 +464,27 @@ bool UNiagaraDataInterfaceChaosDestruction::InitPerInstanceData(void* PerInstanc
 		}
 #endif
 	}
-	else
+	//else
 	{
-//#todo : are solver actors going to exist going forwards?
-//		for (AChaosSolverActor* SolverActor : ChaosSolverActorSet)
-//		{
-//			if (SolverActor)
-//			{
-//				if (Chaos::FPhysicsSolver* Solver = SolverActor->GetSolver())
-//				{
-//					int32 NewIdx = Solvers.Add(FSolverData());
-//
-//					FSolverData& SolverData = Solvers[NewIdx];
-//					SolverData.PhysScene = SolverActor->GetPhysicsScene();
-//					SolverData.Solver = Solver;
-//
-//#if INCLUDE_CHAOS
-//					RegisterWithSolverEventManager(Solver);
-//#endif
-//				}
-//			}
-//		}
+		// #todo : are solver actors going to exist going forwards?
+		for (AChaosSolverActor* SolverActor : ChaosSolverActorSet)
+		{
+			if (SolverActor)
+			{
+				if (Chaos::FPhysicsSolver* Solver = SolverActor->GetSolver())
+				{
+					int32 NewIdx = Solvers.Add(FSolverData());
+
+					FSolverData& SolverData = Solvers[NewIdx];
+					SolverData.PhysScene = SolverActor->GetPhysicsScene().Get();
+					SolverData.Solver = Solver;
+
+#if INCLUDE_CHAOS
+					RegisterWithSolverEventManager(Solver);
+#endif
+				}
+			}
+		}
 	}
 
 	ResetInstData(InstData);
@@ -1201,8 +1201,6 @@ void UNiagaraDataInterfaceChaosDestruction::HandleBreakingEvents(const Chaos::FB
 	ensure(IsInGameThread());
 	Chaos::FBreakingDataArray const& BreakingDataIn = Event.BreakingData.AllBreakingsArray;
 
-	BreakingEvents.Reset();
-
 	// Copy data from Event
 	BreakingEvents.InsertZeroed(0, BreakingDataIn.Num());
 
@@ -1233,8 +1231,16 @@ void UNiagaraDataInterfaceChaosDestruction::HandleBreakingEvents(const Chaos::FB
 			UGeometryCollectionComponent* GeometryCollectionComponent = nullptr;
 			UMaterialInterface* Material = nullptr;
 #if INCLUDE_CHAOS
-			GeometryCollectionComponent = Solvers[0].PhysScene->GetScene().GetOwningComponent<UGeometryCollectionComponent>(DataIn.ParticleProxy);
-			Material = GeometryCollectionComponent->GetMaterial(MaterialID);
+			for (auto& Solver : Solvers)
+			{
+				GeometryCollectionComponent = Solver.PhysScene->GetOwningComponent<UGeometryCollectionComponent>(DataIn.ParticleProxy);
+				if (GeometryCollectionComponent)
+					break;
+			}
+			if (GeometryCollectionComponent)
+			{
+				Material = GeometryCollectionComponent->GetMaterial(MaterialID);
+			}
 #endif
 			if (Material)
 			{
@@ -1740,6 +1746,8 @@ bool UNiagaraDataInterfaceChaosDestruction::BreakingCallback(FNDIChaosDestructio
 		IdxSolver++;
 	}
 
+	BreakingEvents.Reset();
+
 	INC_DWORD_STAT_BY(STAT_NiagaraNumParticlesSpawnedFromBreakings, InstData->PositionArray.Num());
 
 	return false;
diff --git a/Engine/Source/Runtime/Experimental/ChaosSolverEngine/Private/Chaos/ChaosEventListenerComponent.cpp b/Engine/Source/Runtime/Experimental/ChaosSolverEngine/Private/Chaos/ChaosEventListenerComponent.cpp
--- a/Engine/Source/Runtime/Experimental/ChaosSolverEngine/Private/Chaos/ChaosEventListenerComponent.cpp
+++ b/Engine/Source/Runtime/Experimental/ChaosSolverEngine/Private/Chaos/ChaosEventListenerComponent.cpp
@@ -16,19 +16,3 @@ void UChaosEventListenerComponent::TickComponent(float DeltaTime, enum ELevelTic
 	Super::TickComponent(DeltaTime, TickType, ThisTickFunction);
 }
 
-const AChaosSolverActor* UChaosEventListenerComponent::GetSolverActor() const
-{
-	return GetTypedOuter<AChaosSolverActor>();
-}
-
-const Chaos::FPhysicsSolver* UChaosEventListenerComponent::GetSolver() const
-{
-	const AChaosSolverActor* A = GetSolverActor();
-	return A ? A->GetSolver() : nullptr;
-}
-
-const TSharedPtr<FPhysScene_Chaos> UChaosEventListenerComponent::GetPhysicsScene() const
-{
-	const AChaosSolverActor* A = GetSolverActor();
-	return A ? A->GetPhysicsScene() : nullptr;
-}
diff --git a/Engine/Source/Runtime/Experimental/ChaosSolverEngine/Private/Chaos/ChaosGameplayEventDispatcher.cpp b/Engine/Source/Runtime/Experimental/ChaosSolverEngine/Private/Chaos/ChaosGameplayEventDispatcher.cpp
--- a/Engine/Source/Runtime/Experimental/ChaosSolverEngine/Private/Chaos/ChaosGameplayEventDispatcher.cpp
+++ b/Engine/Source/Runtime/Experimental/ChaosSolverEngine/Private/Chaos/ChaosGameplayEventDispatcher.cpp
@@ -218,7 +218,7 @@ void UChaosGameplayEventDispatcher::RegisterChaosEvents()
 		{
 			Chaos::FEventManager* EventManager = Solver->GetEventManager();
 			EventManager->RegisterHandler<Chaos::FCollisionEventData>(Chaos::EEventType::Collision, this, &UChaosGameplayEventDispatcher::HandleCollisionEvents);
-			//EventManager->RegisterHandler<Chaos::FBreakingEventData>(Chaos::EEventType::Breaking, this, &UChaosGameplayEventDispatcher::HandleBreakingEvents);
+			EventManager->RegisterHandler<Chaos::FBreakingEventData>(Chaos::EEventType::Breaking, this, &UChaosGameplayEventDispatcher::HandleBreakingEvents);
 			EventManager->RegisterHandler<Chaos::FSleepingEventData>(Chaos::EEventType::Sleeping, this, &UChaosGameplayEventDispatcher::HandleSleepingEvents);
 		}
 	}
@@ -237,7 +237,7 @@ void UChaosGameplayEventDispatcher::UnregisterChaosEvents()
 			{
 				Chaos::FEventManager* EventManager = Solver->GetEventManager();
 				EventManager->UnregisterHandler(Chaos::EEventType::Collision, this);
-				//EventManager->UnregisterHandler(Chaos::EEventType::Breaking, this);
+				EventManager->UnregisterHandler(Chaos::EEventType::Breaking, this);
 				EventManager->UnregisterHandler(Chaos::EEventType::Sleeping, this);
 			}
 		}
@@ -249,7 +249,7 @@ void UChaosGameplayEventDispatcher::HandleCollisionEvents(const Chaos::FCollisio
 {
 	SCOPE_CYCLE_COUNTER(STAT_DispatchCollisionEvents);
 
-	FPhysScene_Chaos const* const Scene = GetPhysicsScene().Get();
+	FPhysScene_Chaos& Scene = GetWorld()->GetPhysicsScene()->GetScene();
 
 	PendingChaosCollisionNotifies.Reset();
 	ContactPairToPendingNotifyMap.Reset();
@@ -269,7 +269,7 @@ void UChaosGameplayEventDispatcher::HandleCollisionEvents(const Chaos::FCollisio
 				const FChaosHandlerSet& HandlerSet = It.Value();
 
 				UPrimitiveComponent* const Comp0 = Cast<UPrimitiveComponent>(It.Key());
-				IPhysicsProxyBase* const PhysicsProxy0 = Scene->GetOwnedPhysicsProxy(Comp0);
+				IPhysicsProxyBase* const PhysicsProxy0 = Scene.GetOwnedPhysicsProxy(Comp0);
 				TArray<int32> const* const CollisionIndices = PhysicsProxyToCollisionIndicesMap.Find(PhysicsProxy0);
 				if (CollisionIndices)
 				{
@@ -281,7 +281,6 @@ void UChaosGameplayEventDispatcher::HandleCollisionEvents(const Chaos::FCollisio
 						Chaos::TCollisionData<float, 3> const& CollisionDataItem = CollisionData[CollisionIdx];
 						IPhysicsProxyBase* const PhysicsProxy1 = bSwapOrder ? CollisionDataItem.ParticleProxy : CollisionDataItem.LevelsetProxy;
 
-						if (HandlerSet.bLegacyComponentNotify)
 						{
 							bool bNewEntry = false;
 							FCollisionNotifyInfo& NotifyInfo = GetPendingCollisionForContactPair(PhysicsProxy0, PhysicsProxy1, bNewEntry);
@@ -294,7 +293,7 @@ void UChaosGameplayEventDispatcher::HandleCollisionEvents(const Chaos::FCollisio
 
 							if (bNewEntry)
 							{
-								UPrimitiveComponent* const Comp1 = Scene->GetOwningComponent<UPrimitiveComponent>(PhysicsProxy1);
+								UPrimitiveComponent* const Comp1 = Scene.GetOwningComponent<UPrimitiveComponent>(PhysicsProxy1);
 
 								// fill in legacy contact data
 								NotifyInfo.bCallEvent0 = true;
@@ -322,7 +321,7 @@ void UChaosGameplayEventDispatcher::HandleCollisionEvents(const Chaos::FCollisio
 
 							if (bNewEntry)
 							{
-								UPrimitiveComponent* const Comp1 = Scene->GetOwningComponent<UPrimitiveComponent>(PhysicsProxy1);
+								UPrimitiveComponent* const Comp1 = Scene.GetOwningComponent<UPrimitiveComponent>(PhysicsProxy1);
 
 								// fill in Chaos contact data
 								ChaosNotifyInfo.CollisionInfo.Component = Comp0;
diff --git a/Engine/Source/Runtime/Experimental/ChaosSolverEngine/Public/Chaos/ChaosEventListenerComponent.h b/Engine/Source/Runtime/Experimental/ChaosSolverEngine/Public/Chaos/ChaosEventListenerComponent.h
--- a/Engine/Source/Runtime/Experimental/ChaosSolverEngine/Public/Chaos/ChaosEventListenerComponent.h
+++ b/Engine/Source/Runtime/Experimental/ChaosSolverEngine/Public/Chaos/ChaosEventListenerComponent.h
@@ -36,11 +36,6 @@ public:
 
 protected:
 
-	const Chaos::FPhysicsSolver* GetSolver() const;
-	const TSharedPtr<FPhysScene_Chaos> GetPhysicsScene() const;
-
-	const AChaosSolverActor* GetSolverActor() const;
-
 	/** Used to know when the physics thread has updated the collision info for processing on the game thread */
 	float LastCollisionTickTime;
 };
diff --git a/Engine/Source/Runtime/Experimental/GeometryCollectionEngine/Private/GeometryCollection/GeometryCollectionComponent.cpp b/Engine/Source/Runtime/Experimental/GeometryCollectionEngine/Private/GeometryCollection/GeometryCollectionComponent.cpp
--- a/Engine/Source/Runtime/Experimental/GeometryCollectionEngine/Private/GeometryCollection/GeometryCollectionComponent.cpp
+++ b/Engine/Source/Runtime/Experimental/GeometryCollectionEngine/Private/GeometryCollection/GeometryCollectionComponent.cpp
@@ -160,6 +160,8 @@ UGeometryCollectionComponent::UGeometryCollectionComponent(const FObjectInitiali
 
 	// By default use the destructible object channel unless the user specifies otherwise
 	BodyInstance.SetObjectType(ECC_Destructible);
+
+	EventDispatcher = ObjectInitializer.CreateDefaultSubobject<UChaosGameplayEventDispatcher>(this, TEXT("GameplayEventDispatcher"));
 }
 
 Chaos::FPhysicsSolver* GetSolver(const UGeometryCollectionComponent& GeometryCollectionComponent)
@@ -434,57 +436,41 @@ void UGeometryCollectionComponent::RegisterForEvents()
 {
 	if (BodyInstance.bNotifyRigidBodyCollision || bNotifyBreaks || bNotifyCollisions)
 	{
-		if (AChaosSolverActor* const SolverActor = GetPhysicsSolverActor())
+		if (bNotifyCollisions || BodyInstance.bNotifyRigidBodyCollision)
 		{
-			if (UChaosGameplayEventDispatcher* const EventDispatcher = SolverActor->GetGameplayEventDispatcher())
-			{
-				if (bNotifyCollisions || BodyInstance.bNotifyRigidBodyCollision)
-				{
-					EventDispatcher->RegisterForCollisionEvents(this, this);
-				}
+			EventDispatcher->RegisterForCollisionEvents(this, this);
+			GetWorld()->GetPhysicsScene()->GetScene().GetSolver()->SetGenerateCollisionData(true);
+		}
 
-				if (bNotifyBreaks)
-				{
-					EventDispatcher->RegisterForBreakEvents(this, &DispatchGeometryCollectionBreakEvent);
-				}
-			}
+		if (bNotifyBreaks)
+		{
+			EventDispatcher->RegisterForBreakEvents(this, &DispatchGeometryCollectionBreakEvent);
+			GetWorld()->GetPhysicsScene()->GetScene().GetSolver()->SetGenerateBreakingData(true);
 		}
 	}
 }
 
 void UGeometryCollectionComponent::UpdateRBCollisionEventRegistration()
 {
-	if (AChaosSolverActor* const SolverActor = GetPhysicsSolverActor())
+	if (bNotifyCollisions || BodyInstance.bNotifyRigidBodyCollision)
 	{
-		if (UChaosGameplayEventDispatcher* const EventDispatcher = SolverActor->GetGameplayEventDispatcher())
-		{
-			if (bNotifyCollisions || BodyInstance.bNotifyRigidBodyCollision)
-			{
-				EventDispatcher->RegisterForCollisionEvents(this, this);
-			}
-			else
-			{
-				EventDispatcher->UnRegisterForCollisionEvents(this, this);
-			}
-		}
+		EventDispatcher->RegisterForCollisionEvents(this, this);
+	}
+	else
+	{
+		EventDispatcher->UnRegisterForCollisionEvents(this, this);
 	}
 }
 
 void UGeometryCollectionComponent::UpdateBreakEventRegistration()
 {
-	if (AChaosSolverActor* const SolverActor = GetPhysicsSolverActor())
+	if (bNotifyBreaks)
 	{
-		if (UChaosGameplayEventDispatcher* const EventDispatcher = SolverActor->GetGameplayEventDispatcher())
-		{
-			if (bNotifyBreaks)
-			{
-				EventDispatcher->RegisterForBreakEvents(this, &DispatchGeometryCollectionBreakEvent);
-			}
-			else
-			{
-				EventDispatcher->UnRegisterForBreakEvents(this);
-			}
-		}
+		EventDispatcher->RegisterForBreakEvents(this, &DispatchGeometryCollectionBreakEvent);
+	}
+	else
+	{
+		EventDispatcher->UnRegisterForBreakEvents(this);
 	}
 }
 
diff --git a/Engine/Source/Runtime/Experimental/GeometryCollectionEngine/Public/GeometryCollection/GeometryCollectionComponent.h b/Engine/Source/Runtime/Experimental/GeometryCollectionEngine/Public/GeometryCollection/GeometryCollectionComponent.h
--- a/Engine/Source/Runtime/Experimental/GeometryCollectionEngine/Public/GeometryCollection/GeometryCollectionComponent.h
+++ b/Engine/Source/Runtime/Experimental/GeometryCollectionEngine/Public/GeometryCollection/GeometryCollectionComponent.h
@@ -634,6 +634,8 @@ private:
 	TArray<bool> TransformsAreEqual;	
 	int32 TransformsAreEqualIndex;
 
+	UChaosGameplayEventDispatcher* EventDispatcher;
+
 #if GEOMETRYCOLLECTION_EDITOR_SELECTION
 	bool bIsTransformSelectionModeEnabled;
 #endif  // #if GEOMETRYCOLLECTION_EDITOR_SELECTION
