CIS fix

diff --git a/Engine/Source/Runtime/Experimental/ChaosSolverEngine/Private/Chaos/ChaosGameplayEventDispatcher.cpp b/Engine/Source/Runtime/Experimental/ChaosSolverEngine/Private/Chaos/ChaosGameplayEventDispatcher.cpp
--- a/Engine/Source/Runtime/Experimental/ChaosSolverEngine/Private/Chaos/ChaosGameplayEventDispatcher.cpp
+++ b/Engine/Source/Runtime/Experimental/ChaosSolverEngine/Private/Chaos/ChaosGameplayEventDispatcher.cpp
@@ -249,6 +249,8 @@ void UChaosGameplayEventDispatcher::HandleCollisionEvents(const Chaos::FCollisio
 {
 	SCOPE_CYCLE_COUNTER(STAT_DispatchCollisionEvents);
 
+#if INCLUDE_CHAOS
+
 	FPhysScene_Chaos& Scene = GetWorld()->GetPhysicsScene()->GetScene();
 
 	PendingChaosCollisionNotifies.Reset();
@@ -363,6 +365,8 @@ void UChaosGameplayEventDispatcher::HandleCollisionEvents(const Chaos::FCollisio
 
 	// Tell the world and actors about the collisions
 	DispatchPendingCollisionNotifies();
+
+#endif
 }
 
 void UChaosGameplayEventDispatcher::HandleBreakingEvents(const Chaos::FBreakingEventData& Event)
diff --git a/Engine/Source/Runtime/Experimental/GeometryCollectionEngine/Private/GeometryCollection/GeometryCollectionComponent.cpp b/Engine/Source/Runtime/Experimental/GeometryCollectionEngine/Private/GeometryCollection/GeometryCollectionComponent.cpp
--- a/Engine/Source/Runtime/Experimental/GeometryCollectionEngine/Private/GeometryCollection/GeometryCollectionComponent.cpp
+++ b/Engine/Source/Runtime/Experimental/GeometryCollectionEngine/Private/GeometryCollection/GeometryCollectionComponent.cpp
@@ -439,13 +439,17 @@ void UGeometryCollectionComponent::RegisterForEvents()
 		if (bNotifyCollisions || BodyInstance.bNotifyRigidBodyCollision)
 		{
 			EventDispatcher->RegisterForCollisionEvents(this, this);
+#if INCLUDE_CHAOS
 			GetWorld()->GetPhysicsScene()->GetScene().GetSolver()->SetGenerateCollisionData(true);
+#endif
 		}
 
 		if (bNotifyBreaks)
 		{
 			EventDispatcher->RegisterForBreakEvents(this, &DispatchGeometryCollectionBreakEvent);
+#if INCLUDE_CHAOS
 			GetWorld()->GetPhysicsScene()->GetScene().GetSolver()->SetGenerateBreakingData(true);
+#endif
 		}
 	}
 }
