Rotate geometry collection DDC again to update with previous changes to levelsets from CL#12714484

diff --git a/Engine/Source/Runtime/Experimental/GeometryCollectionEngine/Private/GeometryCollection/DerivedDataGeometryCollectionCooker.cpp b/Engine/Source/Runtime/Experimental/GeometryCollectionEngine/Private/GeometryCollection/DerivedDataGeometryCollectionCooker.cpp
--- a/Engine/Source/Runtime/Experimental/GeometryCollectionEngine/Private/GeometryCollection/DerivedDataGeometryCollectionCooker.cpp
+++ b/Engine/Source/Runtime/Experimental/GeometryCollectionEngine/Private/GeometryCollection/DerivedDataGeometryCollectionCooker.cpp
@@ -59,7 +59,7 @@ const TCHAR* FDerivedDataGeometryCollectionCooker::GetVersionString() const
 		return OverrideVersion;	//force load old ddc if found. Not recommended
 	}
 
-	return TEXT("9336C4DCDFF54701A5D322C73EF6D5A4");
+	return TEXT("A8A2C0FB45084FCB922FEC1139E11341");
 }
 
 FString FDerivedDataGeometryCollectionCooker::GetPluginSpecificCacheKeySuffix() const
