Fixed unculled fields applying dynamic parameters to inactive geometry collection particles

diff --git a/Engine/Source/Runtime/Experimental/ChaosSolvers/Private/PhysicsProxy/FieldSystemPhysicsProxy.cpp b/Engine/Source/Runtime/Experimental/ChaosSolvers/Private/PhysicsProxy/FieldSystemPhysicsProxy.cpp
--- a/Engine/Source/Runtime/Experimental/ChaosSolvers/Private/PhysicsProxy/FieldSystemPhysicsProxy.cpp
+++ b/Engine/Source/Runtime/Experimental/ChaosSolvers/Private/PhysicsProxy/FieldSystemPhysicsProxy.cpp
@@ -921,12 +921,13 @@ void FFieldSystemPhysicsProxy::FieldForcesUpdateCallback(
 						{
 							Chaos::TPBDRigidParticleHandle<float, 3>* RigidHandle = Handles[Index.Sample]->CastToRigidParticle();
 
-							if(RigidHandle && (RigidHandle->ObjectState() == Chaos::EObjectStateType::Dynamic || RigidHandle->ObjectState() == Chaos::EObjectStateType::Sleeping))
+							if(RigidHandle && !RigidHandle->Disabled() && (RigidHandle->ObjectState() == Chaos::EObjectStateType::Dynamic || RigidHandle->ObjectState() == Chaos::EObjectStateType::Sleeping))
 							{				
 								if (RigidHandle->Sleeping())
 								{
 									RigidHandle->SetObjectState(Chaos::EObjectStateType::Dynamic);
 								}
+
 								RigidHandle->F() += ForceView[Index.Result];
 							}
 						}
