Handle FImplicitObjectUnionClustered collision constraint generation

diff --git a/Engine/Source/Runtime/Experimental/Chaos/Private/Chaos/CollisionResolution.cpp b/Engine/Source/Runtime/Experimental/Chaos/Private/Chaos/CollisionResolution.cpp
--- a/Engine/Source/Runtime/Experimental/Chaos/Private/Chaos/CollisionResolution.cpp
+++ b/Engine/Source/Runtime/Experimental/Chaos/Private/Chaos/CollisionResolution.cpp
@@ -2604,6 +2604,17 @@ namespace Chaos
 				}
 				return;
 			}
+
+			if (Implicit0OuterType == FImplicitObjectUnionClustered::StaticType())
+			{
+				const FImplicitObjectUnionClustered* Union0 = Implicit0->template GetObject<FImplicitObjectUnionClustered>();
+				for (const auto& Child0 : Union0->GetObjects())
+				{
+					ConstructConstraints(Particle0, Particle1, Child0.Get(), Implicit1, Transform0, Transform1, CullDistance, Context, NewConstraints);
+				}
+				return;
+			}
+
 			if (Implicit1OuterType == FImplicitObjectUnion::StaticType())
 			{
 				const FImplicitObjectUnion* Union1 = Implicit1->template GetObject<FImplicitObjectUnion>();
@@ -2614,6 +2625,17 @@ namespace Chaos
 				return;
 			}
 
+
+			if (Implicit1OuterType == FImplicitObjectUnionClustered::StaticType())
+			{
+				const FImplicitObjectUnionClustered* Union1 = Implicit1->template GetObject<FImplicitObjectUnionClustered>();
+				for (const auto& Child1 : Union1->GetObjects())
+				{
+					ConstructConstraints(Particle0, Particle1, Implicit0, Child1.Get(), Transform0, Transform1, CullDistance, Context, NewConstraints);
+				}
+				return;
+			}
+
 			// Check shape pair filtering if enable
 			if (Context.bFilteringEnabled && !DoCollide(Implicit0Type, Particle0->GetImplicitShape(Implicit0), Implicit1Type, Particle1->GetImplicitShape(Implicit1)))
 			{
