Add flag to set collsiions type as levelset.

diff --git a/Engine/Source/Runtime/Engine/Classes/PhysicsEngine/BodyInstance.h b/Engine/Source/Runtime/Engine/Classes/PhysicsEngine/BodyInstance.h
--- a/Engine/Source/Runtime/Engine/Classes/PhysicsEngine/BodyInstance.h
+++ b/Engine/Source/Runtime/Engine/Classes/PhysicsEngine/BodyInstance.h
@@ -296,6 +296,10 @@ public:
 	UPROPERTY(EditAnywhere, AdvancedDisplay, BlueprintReadOnly, Category = Collision)
 	uint8 bUseCCD : 1;
 
+	/** If true ignore analytic collisions and treat objects as a general implicit surface */
+	UPROPERTY(EditAnywhere, AdvancedDisplay, BlueprintReadOnly, Category = Collision)
+	uint8 bIgnoreAnalyticCollisions : 1;
+
 	/**	Should 'Hit' events fire when this object collides during physics simulation. */
 	UPROPERTY(EditAnywhere, BlueprintReadOnly, Category = Collision, meta = (DisplayName = "Simulation Generates Hit Events"))
 	uint8 bNotifyRigidBodyCollision : 1;
diff --git a/Engine/Source/Runtime/Engine/Private/PhysicsEngine/BodyInstance.cpp b/Engine/Source/Runtime/Engine/Private/PhysicsEngine/BodyInstance.cpp
--- a/Engine/Source/Runtime/Engine/Private/PhysicsEngine/BodyInstance.cpp
+++ b/Engine/Source/Runtime/Engine/Private/PhysicsEngine/BodyInstance.cpp
@@ -99,6 +99,13 @@ FAutoConsoleVariableRef CVarEnableDeferredPhysicsCreation(
 	ECVF_Default
 );
 
+TAutoConsoleVariable<int32> CVarIgnoreAnalyticCollisionsOverride(
+	TEXT("p.IgnoreAnalyticCollisionsOverride"), 
+	0, 
+	TEXT("Overrides the default for ignroing analytic collsions."),
+	ECVF_ReadOnly
+);
+
 bool FBodyInstance::UseDeferredPhysicsBodyCreation()
 {
 	// @todo: this could be improved by tracking all oustanding delayed creations, and return true if there are any at all. 
@@ -106,7 +113,6 @@ bool FBodyInstance::UseDeferredPhysicsBodyCreation()
 	return GUseDeferredPhysicsBodyCreation != 0;
 }
 
-
 using namespace PhysicsInterfaceTypes;
 
 bool IsRigidBodyKinematic_AssumesLocked(const FPhysicsActorHandle& InActorRef)
@@ -342,6 +348,7 @@ FBodyInstance::FBodyInstance()
 	, SleepFamily(ESleepFamily::Normal)
 	, DOFMode(0)
 	, bUseCCD(false)
+	, bIgnoreAnalyticCollisions(false)
 	, bNotifyRigidBodyCollision(false)
 	, bSimulatePhysics(false)
 	, bOverrideMass(false)
@@ -1052,7 +1059,6 @@ void FInitBodiesHelperBase::CreateActor_AssumesLocked(FBodyInstance* Instance, c
 	else
 	{
 		FPhysicsInterface::CreateActor(ActorParams, Instance->ActorHandle);
-
 		FPhysicsInterface::SetCcdEnabled_AssumesLocked(Instance->ActorHandle, Instance->bUseCCD);
 		FPhysicsInterface::SetIsKinematic_AssumesLocked(Instance->ActorHandle, !Instance->ShouldInstanceSimulatingPhysics());
 
@@ -1083,10 +1089,13 @@ bool FInitBodiesHelperBase::CreateShapes_AssumesLocked(FBodyInstance* Instance)
 	// #PHYS2 Call interface AddGeometry
 	BodySetup->AddShapesToRigidActor_AssumesLocked(Instance, Instance->Scale3D, SimplePhysMat, ComplexPhysMats, ComplexPhysMatMasks, BodyCollisionData, FTransform::Identity);
 
+#if WITH_CHAOS
+	FPhysicsInterface::SetIgnoreAnalyticCollisions_AssumesLocked(Instance->ActorHandle, CVarIgnoreAnalyticCollisionsOverride.GetValueOnGameThread() ? true : Instance->bIgnoreAnalyticCollisions);
+#endif
+
 	const int32 NumShapes = FPhysicsInterface::GetNumShapes(Instance->ActorHandle);
 	bInitFail |= NumShapes == 0;
 
-
 	return bInitFail;
 }
 
diff --git a/Engine/Source/Runtime/Engine/Private/PhysicsEngine/Experimental/PhysInterface_Chaos.cpp b/Engine/Source/Runtime/Engine/Private/PhysicsEngine/Experimental/PhysInterface_Chaos.cpp
--- a/Engine/Source/Runtime/Engine/Private/PhysicsEngine/Experimental/PhysInterface_Chaos.cpp
+++ b/Engine/Source/Runtime/Engine/Private/PhysicsEngine/Experimental/PhysInterface_Chaos.cpp
@@ -583,6 +583,11 @@ void FPhysInterface_Chaos::SetCcdEnabled_AssumesLocked(const FPhysicsActorHandle
     //check(bIsCcdEnabled == false);
 }
 
+void FPhysInterface_Chaos::SetIgnoreAnalyticCollisions_AssumesLocked(const FPhysicsActorHandle& InActorReference, bool bIgnoreAnalyticCollisions)
+{
+	InActorReference->SetIgnoreAnalyticCollisions(bIgnoreAnalyticCollisions);
+}
+
 FTransform FPhysInterface_Chaos::GetGlobalPose_AssumesLocked(const FPhysicsActorHandle& InActorReference)
 {
 	return Chaos::TRigidTransform<float, 3>(InActorReference->X(), InActorReference->R());
diff --git a/Engine/Source/Runtime/Engine/Public/Physics/Experimental/PhysInterface_Chaos.h b/Engine/Source/Runtime/Engine/Public/Physics/Experimental/PhysInterface_Chaos.h
--- a/Engine/Source/Runtime/Engine/Public/Physics/Experimental/PhysInterface_Chaos.h
+++ b/Engine/Source/Runtime/Engine/Public/Physics/Experimental/PhysInterface_Chaos.h
@@ -241,6 +241,7 @@ public:
 
 	static void SetIsKinematic_AssumesLocked(const FPhysicsActorHandle& InActorReference, bool bIsKinematic);
 	static void SetCcdEnabled_AssumesLocked(const FPhysicsActorHandle& InActorReference, bool bIsCcdEnabled);
+	static void SetIgnoreAnalyticCollisions_AssumesLocked(const FPhysicsActorHandle& InActorReference, bool bIsCcdEnabled);
 
 	static FTransform GetGlobalPose_AssumesLocked(const FPhysicsActorHandle& InActorReference);
 	static void SetGlobalPose_AssumesLocked(const FPhysicsActorHandle& InActorReference, const FTransform& InNewPose, bool bAutoWake = true);
diff --git a/Engine/Source/Runtime/Experimental/Chaos/Private/Chaos/CollisionResolution.cpp b/Engine/Source/Runtime/Experimental/Chaos/Private/Chaos/CollisionResolution.cpp
--- a/Engine/Source/Runtime/Experimental/Chaos/Private/Chaos/CollisionResolution.cpp
+++ b/Engine/Source/Runtime/Experimental/Chaos/Private/Chaos/CollisionResolution.cpp
@@ -2187,8 +2187,8 @@ namespace Chaos
 			// @todo(chaos): We use GetInnerType here because TriMeshes are left with their "Instanced" wrapper, unlike all other instanced implicits. Should we strip the instance on Tri Mesh too?
 			EImplicitObjectType Implicit0Type = Implicit0 ? GetInnerType(Implicit0->GetCollisionType()) : ImplicitObjectType::Unknown;
 			EImplicitObjectType Implicit1Type = Implicit1 ? GetInnerType(Implicit1->GetCollisionType()) : ImplicitObjectType::Unknown;
-			bool bIsConvex0 = Implicit0 && Implicit0->IsConvex();
-			bool bIsConvex1 = Implicit1 && Implicit1->IsConvex();
+			bool bIsConvex0 = Implicit0 && Implicit0->IsConvex() && Implicit0Type != ImplicitObjectType::LevelSet;
+			bool bIsConvex1 = Implicit1 && Implicit1->IsConvex() && Implicit1Type != ImplicitObjectType::LevelSet;
 
 			FReal LengthCCD = 0.0f;
 			FVec3 DirCCD(0.0f);
diff --git a/Engine/Source/Runtime/Experimental/Chaos/Public/Chaos/ParticleHandle.h b/Engine/Source/Runtime/Experimental/Chaos/Public/Chaos/ParticleHandle.h
--- a/Engine/Source/Runtime/Experimental/Chaos/Public/Chaos/ParticleHandle.h
+++ b/Engine/Source/Runtime/Experimental/Chaos/Public/Chaos/ParticleHandle.h
@@ -1578,6 +1578,84 @@ public:
 		MapImplicitShapes();
 	}
 
+	inline FImplicitObject* GetInstancedImplicitHelper(FImplicitObject* Implicit0)
+	{
+		EImplicitObjectType Implicit0OuterType = Implicit0->GetType();
+
+		if (Implicit0OuterType == TImplicitObjectInstanced<FConvex>::StaticType())
+		{
+			return const_cast<FConvex*>(Implicit0->template GetObject<TImplicitObjectInstanced<FConvex>>()->GetInstancedObject());
+		}
+		else if (Implicit0OuterType == TImplicitObjectInstanced<TBox<FReal, 3>>::StaticType())
+		{
+			return const_cast<TBox<FReal, 3>*>(Implicit0->template GetObject<TImplicitObjectInstanced<TBox<FReal, 3>>>()->GetInstancedObject());
+		}
+		else if (Implicit0OuterType == TImplicitObjectInstanced<TCapsule<FReal>>::StaticType())
+		{
+			return const_cast<TCapsule<FReal>*>(Implicit0->template GetObject<TImplicitObjectInstanced<TCapsule<FReal>>>()->GetInstancedObject());
+		}
+		else if (Implicit0OuterType == TImplicitObjectInstanced<TSphere<FReal, 3>>::StaticType())
+		{
+			return const_cast<TSphere<FReal, 3>*>(Implicit0->template GetObject<TImplicitObjectInstanced<TSphere<FReal, 3>>>()->GetInstancedObject());
+		}
+		else if (Implicit0OuterType == TImplicitObjectInstanced<FTriangleMeshImplicitObject>::StaticType())
+		{
+			return const_cast<FTriangleMeshImplicitObject*>(Implicit0->template GetObject<TImplicitObjectInstanced<FTriangleMeshImplicitObject>>()->GetInstancedObject());
+		}
+
+		return nullptr;
+	}
+
+	void SetIgnoreAnalyticCollisionsImp(FImplicitObject* Implicit, bool bIgnoreAnalyticCollisions)
+	{
+		check(Implicit);
+		if (Implicit->GetType() == FImplicitObjectUnion::StaticType())
+		{
+			FImplicitObjectUnion* Union = Implicit->template GetObject<FImplicitObjectUnion>();
+			for (const auto& Child : Union->GetObjects())
+			{
+				SetIgnoreAnalyticCollisionsImp(Child.Get(), bIgnoreAnalyticCollisions);
+			}
+		}
+		else if (Implicit->GetType() == TImplicitObjectTransformed<T, d>::StaticType())
+		{
+			TImplicitObjectTransformed<FReal, 3>* TransformedImplicit = Implicit->template GetObject<TImplicitObjectTransformed<FReal, 3>>();
+			SetIgnoreAnalyticCollisionsImp(TransformedImplicit, bIgnoreAnalyticCollisions);
+		}
+		else if ((uint32)Implicit->GetType() & ImplicitObjectType::IsInstanced)
+		{
+			SetIgnoreAnalyticCollisionsImp(GetInstancedImplicitHelper(Implicit), bIgnoreAnalyticCollisions);
+		}
+		else
+		{
+			if (const auto* PerShapeData = GetImplicitShape(Implicit))
+			{
+				if (PerShapeData->bDisable)
+				{
+					return;
+				}
+			}
+			if (bIgnoreAnalyticCollisions)
+			{
+				Implicit->SetCollsionType(Chaos::ImplicitObjectType::LevelSet);
+				//Implicit->SetConvex(false);
+			}
+			else
+			{
+				Implicit->SetCollsionType(MGeometry->GetType());
+				// @todo (mlentine): Need to in theory set convex properly here
+			}
+		}
+	}
+
+	void SetIgnoreAnalyticCollisions(bool bIgnoreAnalyticCollisions)
+	{
+		if (MGeometry)
+		{
+			SetIgnoreAnalyticCollisionsImp(MGeometry.Get(), bIgnoreAnalyticCollisions);
+		}
+	}
+
 	TSerializablePtr<FImplicitObject> Geometry() const { return MakeSerializable(MNonFrequentData.Read().Geometry()); }
 
 	const FShapesArray& ShapesArray() const { return MShapesArray; }
