Speculative CIS fix.

diff --git a/Engine/Source/Runtime/Experimental/Chaos/Private/Chaos/ParticleHandle.cpp b/Engine/Source/Runtime/Experimental/Chaos/Private/Chaos/ParticleHandle.cpp
--- a/Engine/Source/Runtime/Experimental/Chaos/Private/Chaos/ParticleHandle.cpp
+++ b/Engine/Source/Runtime/Experimental/Chaos/Private/Chaos/ParticleHandle.cpp
@@ -2,6 +2,9 @@
 
 #include "Chaos/ParticleHandle.h"
 
+#include "Chaos/ImplicitObjectTransformed.h"
+#include "Chaos/ImplicitObjectScaled.h"
+#include "Chaos/TriangleMeshImplicitObject.h"
 #include "Chaos/Framework/PhysicsProxy.h"
 #include "Chaos/Framework/PhysicsSolverBase.h"
 #include "Chaos/Framework/PhysicsProxyBase.h"
@@ -67,7 +70,76 @@ namespace Chaos
 		}
 	}
 
+	inline FImplicitObject* GetInstancedImplicitHelper(FImplicitObject* Implicit0)
+	{
+		EImplicitObjectType Implicit0OuterType = Implicit0->GetType();
+
+		if (Implicit0OuterType == TImplicitObjectInstanced<FConvex>::StaticType())
+		{
+			return const_cast<FConvex*>(Implicit0->template GetObject<TImplicitObjectInstanced<FConvex>>()->GetInstancedObject());
+		}
+		else if (Implicit0OuterType == TImplicitObjectInstanced<TBox<FReal, 3>>::StaticType())
+		{
+			return const_cast<TBox<FReal, 3>*>(Implicit0->template GetObject<TImplicitObjectInstanced<TBox<FReal, 3>>>()->GetInstancedObject());
+		}
+		else if (Implicit0OuterType == TImplicitObjectInstanced<TCapsule<FReal>>::StaticType())
+		{
+			return const_cast<TCapsule<FReal>*>(Implicit0->template GetObject<TImplicitObjectInstanced<TCapsule<FReal>>>()->GetInstancedObject());
+		}
+		else if (Implicit0OuterType == TImplicitObjectInstanced<TSphere<FReal, 3>>::StaticType())
+		{
+			return const_cast<TSphere<FReal, 3>*>(Implicit0->template GetObject<TImplicitObjectInstanced<TSphere<FReal, 3>>>()->GetInstancedObject());
+		}
+		else if (Implicit0OuterType == TImplicitObjectInstanced<FTriangleMeshImplicitObject>::StaticType())
+		{
+			return const_cast<FTriangleMeshImplicitObject*>(Implicit0->template GetObject<TImplicitObjectInstanced<FTriangleMeshImplicitObject>>()->GetInstancedObject());
+		}
+
+		return nullptr;
+	}
 
+	template <typename T, int d>
+	void Chaos::TGeometryParticle<T, d>::SetIgnoreAnalyticCollisionsImp(FImplicitObject* Implicit, bool bIgnoreAnalyticCollisions)
+	{
+		check(Implicit);
+		if (Implicit->GetType() == FImplicitObjectUnion::StaticType())
+		{
+			FImplicitObjectUnion* Union = Implicit->template GetObject<FImplicitObjectUnion>();
+			for (const auto& Child : Union->GetObjects())
+			{
+				SetIgnoreAnalyticCollisionsImp(Child.Get(), bIgnoreAnalyticCollisions);
+			}
+		}
+		else if (Implicit->GetType() == TImplicitObjectTransformed<T, d>::StaticType())
+		{
+			TImplicitObjectTransformed<FReal, 3>* TransformedImplicit = Implicit->template GetObject<TImplicitObjectTransformed<FReal, 3>>();
+			SetIgnoreAnalyticCollisionsImp(TransformedImplicit, bIgnoreAnalyticCollisions);
+		}
+		else if ((uint32)Implicit->GetType() & ImplicitObjectType::IsInstanced)
+		{
+			SetIgnoreAnalyticCollisionsImp(GetInstancedImplicitHelper(Implicit), bIgnoreAnalyticCollisions);
+		}
+		else
+		{
+			if (const auto* PerShapeData = GetImplicitShape(Implicit))
+			{
+				if (PerShapeData->bDisable)
+				{
+					return;
+				}
+			}
+			if (bIgnoreAnalyticCollisions)
+			{
+				Implicit->SetCollsionType(Chaos::ImplicitObjectType::LevelSet);
+				//Implicit->SetConvex(false);
+			}
+			else
+			{
+				Implicit->SetCollsionType(MGeometry->GetType());
+				// @todo (mlentine): Need to in theory set convex properly here
+			}
+		}
+	}
 
 	template class CHAOS_API TGeometryParticleData<float, 3>;
 	template class CHAOS_API TGeometryParticle<float, 3>;
diff --git a/Engine/Source/Runtime/Experimental/Chaos/Public/Chaos/ParticleHandle.h b/Engine/Source/Runtime/Experimental/Chaos/Public/Chaos/ParticleHandle.h
--- a/Engine/Source/Runtime/Experimental/Chaos/Public/Chaos/ParticleHandle.h
+++ b/Engine/Source/Runtime/Experimental/Chaos/Public/Chaos/ParticleHandle.h
@@ -17,8 +17,6 @@ class IPhysicsProxyBase;
 
 namespace Chaos
 {
-
-
 template <typename T, int d>
 struct TGeometryParticleParameters
 {
@@ -1576,76 +1574,7 @@ public:
 		MapImplicitShapes();
 	}
 
-	inline FImplicitObject* GetInstancedImplicitHelper(FImplicitObject* Implicit0)
-	{
-		EImplicitObjectType Implicit0OuterType = Implicit0->GetType();
-
-		if (Implicit0OuterType == TImplicitObjectInstanced<FConvex>::StaticType())
-		{
-			return const_cast<FConvex*>(Implicit0->template GetObject<TImplicitObjectInstanced<FConvex>>()->GetInstancedObject());
-		}
-		else if (Implicit0OuterType == TImplicitObjectInstanced<TBox<FReal, 3>>::StaticType())
-		{
-			return const_cast<TBox<FReal, 3>*>(Implicit0->template GetObject<TImplicitObjectInstanced<TBox<FReal, 3>>>()->GetInstancedObject());
-		}
-		else if (Implicit0OuterType == TImplicitObjectInstanced<TCapsule<FReal>>::StaticType())
-		{
-			return const_cast<TCapsule<FReal>*>(Implicit0->template GetObject<TImplicitObjectInstanced<TCapsule<FReal>>>()->GetInstancedObject());
-		}
-		else if (Implicit0OuterType == TImplicitObjectInstanced<TSphere<FReal, 3>>::StaticType())
-		{
-			return const_cast<TSphere<FReal, 3>*>(Implicit0->template GetObject<TImplicitObjectInstanced<TSphere<FReal, 3>>>()->GetInstancedObject());
-		}
-		else if (Implicit0OuterType == TImplicitObjectInstanced<FTriangleMeshImplicitObject>::StaticType())
-		{
-			return const_cast<FTriangleMeshImplicitObject*>(Implicit0->template GetObject<TImplicitObjectInstanced<FTriangleMeshImplicitObject>>()->GetInstancedObject());
-		}
-
-		return nullptr;
-	}
-
-	void SetIgnoreAnalyticCollisionsImp(FImplicitObject* Implicit, bool bIgnoreAnalyticCollisions)
-	{
-		check(Implicit);
-		if (Implicit->GetType() == FImplicitObjectUnion::StaticType())
-		{
-			FImplicitObjectUnion* Union = Implicit->template GetObject<FImplicitObjectUnion>();
-			for (const auto& Child : Union->GetObjects())
-			{
-				SetIgnoreAnalyticCollisionsImp(Child.Get(), bIgnoreAnalyticCollisions);
-			}
-		}
-		else if (Implicit->GetType() == TImplicitObjectTransformed<T, d>::StaticType())
-		{
-			TImplicitObjectTransformed<FReal, 3>* TransformedImplicit = Implicit->template GetObject<TImplicitObjectTransformed<FReal, 3>>();
-			SetIgnoreAnalyticCollisionsImp(TransformedImplicit, bIgnoreAnalyticCollisions);
-		}
-		else if ((uint32)Implicit->GetType() & ImplicitObjectType::IsInstanced)
-		{
-			SetIgnoreAnalyticCollisionsImp(GetInstancedImplicitHelper(Implicit), bIgnoreAnalyticCollisions);
-		}
-		else
-		{
-			if (const auto* PerShapeData = GetImplicitShape(Implicit))
-			{
-				if (PerShapeData->bDisable)
-				{
-					return;
-				}
-			}
-			if (bIgnoreAnalyticCollisions)
-			{
-				Implicit->SetCollsionType(Chaos::ImplicitObjectType::LevelSet);
-				//Implicit->SetConvex(false);
-			}
-			else
-			{
-				Implicit->SetCollsionType(MGeometry->GetType());
-				// @todo (mlentine): Need to in theory set convex properly here
-			}
-		}
-	}
-
+	void SetIgnoreAnalyticCollisionsImp(FImplicitObject* Implicit, bool bIgnoreAnalyticCollisions);
 	void SetIgnoreAnalyticCollisions(bool bIgnoreAnalyticCollisions)
 	{
 		if (MGeometry)
