Chaos : Clustering

- Restrict processing of union clusters to only active objects.

diff --git a/Engine/Source/Runtime/Experimental/Chaos/Private/Chaos/PBDRigidClustering.cpp b/Engine/Source/Runtime/Experimental/Chaos/Private/Chaos/PBDRigidClustering.cpp
--- a/Engine/Source/Runtime/Experimental/Chaos/Private/Chaos/PBDRigidClustering.cpp
+++ b/Engine/Source/Runtime/Experimental/Chaos/Private/Chaos/PBDRigidClustering.cpp
@@ -618,9 +618,12 @@ namespace Chaos
 					TArray<TPBDRigidParticleHandle<T, 3>*> ClusterBodies;
 					for (TPBDRigidClusteredParticleHandle<T, 3>* ActiveCluster : Handles)
 					{
-						TSet<TPBDRigidParticleHandle<T, 3>*> Children = ReleaseClusterParticles(ActiveCluster, nullptr, true);
-						NewClusterGroups[ClusterGroupID].Append(Children.Array());
-						for (auto& Child : Children) ClusterParents.Add(Child, ActiveCluster);
+						if (!ActiveCluster->Disabled())
+						{
+							TSet<TPBDRigidParticleHandle<T, 3>*> Children = ReleaseClusterParticles(ActiveCluster, nullptr, true);
+							NewClusterGroups[ClusterGroupID].Append(Children.Array());
+							for (auto& Child : Children) ClusterParents.Add(Child, ActiveCluster);
+						}
 					}
 				}
 			}
