Chaos: Re-imlement internal strain in fields

diff --git a/Engine/Source/Runtime/Experimental/ChaosSolvers/Private/PhysicsProxy/FieldSystemPhysicsProxy.cpp b/Engine/Source/Runtime/Experimental/ChaosSolvers/Private/PhysicsProxy/FieldSystemPhysicsProxy.cpp
--- a/Engine/Source/Runtime/Experimental/ChaosSolvers/Private/PhysicsProxy/FieldSystemPhysicsProxy.cpp
+++ b/Engine/Source/Runtime/Experimental/ChaosSolvers/Private/PhysicsProxy/FieldSystemPhysicsProxy.cpp
@@ -616,6 +616,42 @@ void FFieldSystemPhysicsProxy::FieldParameterUpdateCallback(
 
 				CommandsToRemove.Add(CommandIndex);
 			}
+			else if (Command.TargetAttribute == GetFieldPhysicsName(EFieldPhysicsType::Field_InternalClusterStrain))
+			{
+				SCOPE_CYCLE_COUNTER(STAT_ParamUpdateField_InternalClusterStrain);
+				if (ensureMsgf(Command.RootNode->Type() == FFieldNode<float>::StaticType(),
+					TEXT("Field based evaluation of the simulations 'ExternalClusterStrain' parameter expects scalar field inputs.")))
+				{
+					if (Handles.Num())
+					{
+						TArrayView<FVector> SamplePointsView(&(SamplePoints[0]), SamplePoints.Num());
+						TArrayView<ContextIndex> SampleIndicesView(&(SampleIndices[0]), SampleIndices.Num());
+
+						FFieldContext Context{
+							SampleIndicesView, // @todo(chaos) important: an empty index array should evaluate everything
+							SamplePointsView,
+							Command.MetaData
+						};
+
+						TArray<float> LocalResults;
+						LocalResults.AddZeroed(Handles.Num());
+						TArrayView<float> ResultsView(&(LocalResults[0]), LocalResults.Num());
+
+						static_cast<const FFieldNode<float>*>(Command.RootNode.Get())->Evaluate(Context, ResultsView);
+
+						for (const ContextIndex& Index : Context.GetEvaluatedSamples())
+						{
+							Chaos::TPBDRigidClusteredParticleHandle<float, 3>* RigidHandle = Handles[Index.Sample]->CastToClustered();
+							if (RigidHandle && RigidHandle->ObjectState() == Chaos::EObjectStateType::Dynamic)
+							{
+								RigidHandle->Strain() += ResultsView[Index.Result];
+							}
+						}
+
+					}
+				}
+				CommandsToRemove.Add(CommandIndex);
+			}
 			else if (Command.TargetAttribute == GetFieldPhysicsName(EFieldPhysicsType::Field_CollisionGroup))
 			{
 				if (ensureMsgf(Command.RootNode->Type() == FFieldNode<int32>::StaticType(),
