Move gravity enabled into pbd rigid particles directly. Helps with data marshaling and network rewind

diff --git a/Engine/Source/Programs/HeadlessChaos/Private/HeadlessChaosTestConstraints.cpp b/Engine/Source/Programs/HeadlessChaos/Private/HeadlessChaosTestConstraints.cpp
--- a/Engine/Source/Programs/HeadlessChaos/Private/HeadlessChaosTestConstraints.cpp
+++ b/Engine/Source/Programs/HeadlessChaos/Private/HeadlessChaosTestConstraints.cpp
@@ -47,7 +47,7 @@ namespace ChaosTest {
 			THandleArray<FChaosPhysicsMaterial> PhysicalMaterials;
 			FPBDRigidsEvolutionGBF Evolution(Particles, PhysicalMaterials);
 			TArray<TPBDRigidParticleHandle<FReal, 3>*> Dynamics = Evolution.CreateDynamicParticles(1);
-			Evolution.GetGravityForces().SetEnabled(*Dynamics[0], false);
+			Dynamics[0]->SetGravityEnabled(false);
 
 			TArray<FVec3> Positions = { FVec3(1) };
 			TPBDPositionConstraints<FReal, 3> PositionConstraints(MoveTemp(Positions), MoveTemp(Dynamics), 0.5f);
diff --git a/Engine/Source/Programs/HeadlessChaos/Private/HeadlessChaosTestSolverProxies.cpp b/Engine/Source/Programs/HeadlessChaos/Private/HeadlessChaosTestSolverProxies.cpp
--- a/Engine/Source/Programs/HeadlessChaos/Private/HeadlessChaosTestSolverProxies.cpp
+++ b/Engine/Source/Programs/HeadlessChaos/Private/HeadlessChaosTestSolverProxies.cpp
@@ -38,7 +38,7 @@ namespace ChaosTest {
 		TUniquePtr<Chaos::TPBDRigidParticle<float, 3>> Particle = Chaos::TPBDRigidParticle<float, 3>::CreateParticle();
 		Particle->SetGeometry(Sphere);
 		Particle->SetX(TVector<float, 3>(0, 0, 0));
-
+		Particle->SetGravityEnabled(false);
 		Solver->RegisterObject(Particle.Get());
 
 		Particle->SetV(TVector<float, 3>(0, 0, 10));
diff --git a/Engine/Source/Runtime/Experimental/Chaos/Public/Chaos/ParticleHandle.h b/Engine/Source/Runtime/Experimental/Chaos/Public/Chaos/ParticleHandle.h
--- a/Engine/Source/Runtime/Experimental/Chaos/Public/Chaos/ParticleHandle.h
+++ b/Engine/Source/Runtime/Experimental/Chaos/Public/Chaos/ParticleHandle.h
@@ -38,7 +38,7 @@ struct TPBDRigidParticleParameters : public TKinematicGeometryParticleParameters
 	TPBDRigidParticleParameters()
 		: TKinematicGeometryParticleParameters<T, d>()
 		, bStartSleeping(false)
-		, bGravityEnabled(false)
+		, bGravityEnabled(true)
 	{}
 	bool bStartSleeping;
 	bool bGravityEnabled;
@@ -84,6 +84,7 @@ void PBDRigidParticleHandleImpDefaultConstruct(TPBDRigidParticleHandleImp<T, d,
 	Concrete.SetLinearEtherDrag(0.f);
 	Concrete.SetAngularEtherDrag(0.f);
 	Concrete.SetObjectStateLowLevel(Params.bStartSleeping ? EObjectStateType::Sleeping : EObjectStateType::Dynamic);
+	Concrete.SetGravityEnabled(Params.bGravityEnabled);
 }
 
 template <typename T, int d>
@@ -755,6 +756,7 @@ public:
 		SetLinearEtherDrag(DynamicMisc.LinearEtherDrag());
 		SetAngularEtherDrag(DynamicMisc.AngularEtherDrag());
 		SetCollisionGroup(DynamicMisc.CollisionGroup());
+		SetGravityEnabled(DynamicMisc.GravityEnabled());
 	}
 
 	const PMatrix<T, d, d>& I() const { return PBDRigidParticles->I(ParticleIdx); }
@@ -807,6 +809,10 @@ public:
 	const TPBDRigidParticleHandleImp<T, d, true>* Handle() const { return PBDRigidParticles->Handle(ParticleIdx); }
 	TPBDRigidParticleHandleImp<T, d, true>* Handle() { return PBDRigidParticles->Handle(ParticleIdx); }
 
+	bool GravityEnabled() const { return PBDRigidParticles->GravityEnabled(ParticleIdx); }
+
+	void SetGravityEnabled(bool bEnabled){ PBDRigidParticles->GravityEnabled(ParticleIdx) = bEnabled; }
+
 	static constexpr EParticleType StaticType() { return EParticleType::Rigid; }
 };
 
diff --git a/Engine/Source/Runtime/Experimental/Chaos/Public/Chaos/PerParticleGravity.h b/Engine/Source/Runtime/Experimental/Chaos/Public/Chaos/PerParticleGravity.h
--- a/Engine/Source/Runtime/Experimental/Chaos/Public/Chaos/PerParticleGravity.h
+++ b/Engine/Source/Runtime/Experimental/Chaos/Public/Chaos/PerParticleGravity.h
@@ -34,29 +34,12 @@ namespace Chaos
 
 		void Apply(TTransientPBDRigidParticleHandle<T, d>& Particle, const T Dt) const override //-V762
 		{
-			if (!OptOut.Contains(Particle.Handle()))
+			if(Particle.GravityEnabled())
 			{
 				Particle.F() += MAcceleration * Particle.M();
 			}
 		}
 
-		bool GetEnabled(const TGeometryParticleHandleImp<T, d, true>& Handle)
-		{
-			return !OptOut.Contains(&Handle);
-		}
-
-		void SetEnabled(const TGeometryParticleHandleImp<T, d, true>& Handle, bool bEnabled)
-		{
-			if (bEnabled)
-			{
-				OptOut.Remove(&Handle);
-			}
-			else
-			{
-				OptOut.Add(&Handle);
-			}
-		}
-
 		void SetAcceleration(const TVector<T, d>& Acceleration)
 		{ MAcceleration = Acceleration; }
 
@@ -65,6 +48,5 @@ namespace Chaos
 
 	private:
 		TVector<T, d> MAcceleration;
-		TSet< const TGeometryParticleHandleImp<T, d, true>* > OptOut;
 	};
 }
diff --git a/Engine/Source/Runtime/Experimental/Chaos/Public/Chaos/RigidParticles.h b/Engine/Source/Runtime/Experimental/Chaos/Public/Chaos/RigidParticles.h
--- a/Engine/Source/Runtime/Experimental/Chaos/Public/Chaos/RigidParticles.h
+++ b/Engine/Source/Runtime/Experimental/Chaos/Public/Chaos/RigidParticles.h
@@ -78,6 +78,7 @@ class TRigidParticles : public TKinematicGeometryParticles<T, d>
 		TArrayCollection::AddArray(&MObjectState);
 		TArrayCollection::AddArray(&MIsland);
 		TArrayCollection::AddArray(&MToBeRemovedOnFracture);
+		TArrayCollection::AddArray(&MGravityEnabled);
 	}
 	TRigidParticles(const TRigidParticles<T, d>& Other) = delete;
 	CHAOS_API TRigidParticles(TRigidParticles<T, d>&& Other)
@@ -94,6 +95,7 @@ class TRigidParticles : public TKinematicGeometryParticles<T, d>
 		, MCollisionParticles(MoveTemp(Other.MCollisionParticles))
 		, MCollisionGroup(MoveTemp(Other.MCollisionGroup))
 		, MObjectState(MoveTemp(Other.MObjectState))
+		, MGravityEnabled(MoveTemp(Other.MGravityEnabled))
 	{
 		TArrayCollection::AddArray(&MF);
 		TArrayCollection::AddArray(&MT);
@@ -113,6 +115,7 @@ class TRigidParticles : public TKinematicGeometryParticles<T, d>
 		TArrayCollection::AddArray(&MObjectState);
 		TArrayCollection::AddArray(&MIsland);
 		TArrayCollection::AddArray(&MToBeRemovedOnFracture);
+		TArrayCollection::AddArray(&MGravityEnabled);
 	}
 
 	CHAOS_API virtual ~TRigidParticles()
@@ -176,6 +179,9 @@ class TRigidParticles : public TKinematicGeometryParticles<T, d>
 	FORCEINLINE const bool ToBeRemovedOnFracture(const int32 Index) const { return MToBeRemovedOnFracture[Index]; }
 	FORCEINLINE bool& ToBeRemovedOnFracture(const int32 Index) { return MToBeRemovedOnFracture[Index]; }
 
+	FORCEINLINE const bool& GravityEnabled(const int32 Index) const { return MGravityEnabled[Index]; }
+	FORCEINLINE bool& GravityEnabled(const int32 Index) { return MGravityEnabled[Index]; }
+
 	FORCEINLINE TArray<TSleepData<T, d>>& GetSleepData() { return MSleepData; }
 	FORCEINLINE	void AddSleepData(TGeometryParticleHandle<T, d>* Particle, bool Sleeping)
 	{ 
@@ -236,6 +242,7 @@ class TRigidParticles : public TKinematicGeometryParticles<T, d>
 		}
 
 		Ar << MCollisionParticles << MCollisionGroup << MIsland << MDisabled << MObjectState;
+		//todo: add gravity enabled when we decide how we want to handle serialization
 	}
 
   private:
@@ -257,6 +264,7 @@ class TRigidParticles : public TKinematicGeometryParticles<T, d>
 	TArrayCollectionArray<bool> MDisabled;
 	TArrayCollectionArray<bool> MToBeRemovedOnFracture;
 	TArrayCollectionArray<EObjectStateType> MObjectState;
+	TArrayCollectionArray<bool> MGravityEnabled;
 
 	TArray<TSleepData<T, d>> MSleepData;
 	FRWLock SleepDataLock;
diff --git a/Engine/Source/Runtime/Experimental/ChaosSolvers/Private/PhysicsProxy/SingleParticlePhysicsProxy.cpp b/Engine/Source/Runtime/Experimental/ChaosSolvers/Private/PhysicsProxy/SingleParticlePhysicsProxy.cpp
--- a/Engine/Source/Runtime/Experimental/ChaosSolvers/Private/PhysicsProxy/SingleParticlePhysicsProxy.cpp
+++ b/Engine/Source/Runtime/Experimental/ChaosSolvers/Private/PhysicsProxy/SingleParticlePhysicsProxy.cpp
@@ -99,7 +99,6 @@ void PushToPhysicsStateImp(const Chaos::FDirtyPropertiesManager& Manager, TParti
 			if(auto NewData = ParticleData.FindDynamicMisc(Manager,DataIdx))
 			{
 				Solver->GetEvolution()->SetParticleObjectState(RigidHandle,NewData->ObjectState());
-				Solver->GetEvolution()->GetGravityForces().SetEnabled(*RigidHandle,NewData->GravityEnabled());
 
 				RigidHandle->SetDynamicMisc(*NewData);
 			}
