Physics rewind
-Wrap dirty property manager and index into struct to fix issues where read/write indices were conflated, causing asserts when property changes came from the GT

#rnx

diff --git a/Engine/Source/Runtime/Experimental/ChaosSolvers/Public/RewindData.h b/Engine/Source/Runtime/Experimental/ChaosSolvers/Public/RewindData.h
--- a/Engine/Source/Runtime/Experimental/ChaosSolvers/Public/RewindData.h
+++ b/Engine/Source/Runtime/Experimental/ChaosSolvers/Public/RewindData.h
@@ -17,67 +17,85 @@ enum EDesyncResult
 	NeedInfo //one of the entries is missing. Need more context to determine whether desynced
 };
 
+// Wraps FDirtyPropertiesManager and its DataIdx to avoid confusion between Source and offset Dest indices
+struct FDirtyPropData
+{
+	FDirtyPropData(FDirtyPropertiesManager* InManager, int32 InDataIdx)
+		: Ptr(InManager), DataIdx(InDataIdx) { }
+
+	FDirtyPropertiesManager* Ptr;
+	int32 DataIdx;
+};
+
+struct FConstDirtyPropData
+{
+	FConstDirtyPropData(const FDirtyPropertiesManager* InManager, int32 InDataIdx)
+		: Ptr(InManager), DataIdx(InDataIdx) { }
+
+	const FDirtyPropertiesManager* Ptr;
+	int32 DataIdx;
+};
+
 template <typename T,EParticleProperty PropName>
 class TParticleStateProperty
 {
 public:
-
+	
 	TParticleStateProperty()
-		: Manager(nullptr)
+		: Manager(nullptr, INDEX_NONE)
 	{
 	}
 
-	TParticleStateProperty(FDirtyPropertiesManager* InManager,int32 InIdx)
+	TParticleStateProperty(const FDirtyPropData& InManager)
 		: Manager(InManager)
-		,Idx(InIdx)
 	{
 	}
 
 	const T& Read() const
 	{
-		const TDirtyElementPool<T>& Pool = Manager->GetParticlePool<T,PropName>();
-		return Pool.GetElement(Idx);
+		const TDirtyElementPool<T>& Pool = Manager.Ptr->GetParticlePool<T,PropName>();
+		return Pool.GetElement(Manager.DataIdx);
 	}
 
 	template <typename LambdaWrite>
 	void SyncToParticle(const LambdaWrite& WriteFunc) const
 	{
-		if(Manager)
+		if(Manager.Ptr)
 		{
-			const TDirtyElementPool<T>& Pool = Manager->GetParticlePool<T,PropName>();
-			const T& Value = Pool.GetElement(Idx);
+			const TDirtyElementPool<T>& Pool = Manager.Ptr->GetParticlePool<T,PropName>();
+			const T& Value = Pool.GetElement(Manager.DataIdx);
 			WriteFunc(Value);
 		}
 	}
 
 	template <typename LambdaSet>
-	void SyncRemoteDataForced(FDirtyPropertiesManager& InManager,int32 InIdx,const LambdaSet& SetFunc)
+	void SyncRemoteDataForced(const FDirtyPropData& InManager, const LambdaSet& SetFunc)
 	{
-		Manager = &InManager;
-		Idx = InIdx;
-		T& NewVal = Manager->GetParticlePool<T,PropName>().GetElement(Idx);
+		Manager = InManager;
+		T& NewVal = Manager.Ptr->GetParticlePool<T,PropName>().GetElement(Manager.DataIdx);
 		SetFunc(NewVal);
 	}
 
 	template <typename LambdaSet>
-	void SyncRemoteData(FDirtyPropertiesManager& InManager,int32 InIdx, const FParticleDirtyData& DirtyData, const LambdaSet& SetFunc)
+	void SyncRemoteData(const FDirtyPropData& InManager, const FParticleDirtyData& DirtyData, const LambdaSet& SetFunc)
 	{
+		checkSlow(InManager.Ptr != nullptr);
 		if(DirtyData.IsDirty(ParticlePropToFlag(PropName)))
 		{
-			SyncRemoteDataForced(InManager,InIdx,SetFunc);
+			SyncRemoteDataForced(InManager, SetFunc);
 		}
 	}
 
 	bool IsSet() const
 	{
-		return Manager != nullptr;
+		return Manager.Ptr != nullptr;
 	}
 
 	template <typename TParticleHandle>
-	bool IsInSync(const FDirtyPropertiesManager& SrcManager, const int32 DataIdxIn, const FParticleDirtyFlags Flags, const TParticleHandle& Handle) const
+	bool IsInSync(const FConstDirtyPropData& SrcManager, const FParticleDirtyFlags Flags, const TParticleHandle& Handle) const
 	{
-		const T* RecordedEntry = Manager ? &GetValue(*Manager,Idx) : nullptr;
-		const T* NewEntry = Flags.IsDirty(ParticlePropToFlag(PropName)) ? &GetValue(SrcManager,DataIdxIn) : nullptr;
+		const T* RecordedEntry = Manager.Ptr ? &GetValue(Manager.Ptr, Manager.DataIdx) : nullptr;
+		const T* NewEntry = Flags.IsDirty(ParticlePropToFlag(PropName)) ? &GetValue(SrcManager.Ptr, SrcManager.DataIdx) : nullptr;
 
 		if(NewEntry)
 		{
@@ -108,12 +126,11 @@ public:
 	}
 
 private:
-	FDirtyPropertiesManager* Manager;
-	int32 Idx;
-
-	static const T& GetValue(const FDirtyPropertiesManager& InManager, const int32 InIdx)
+	FDirtyPropData Manager;
+	
+	static const T& GetValue(const FDirtyPropertiesManager* Ptr, int32 DataIdx)
 	{
-		return InManager.GetParticlePool<T,PropName>().GetElement(InIdx);
+		return Ptr->GetParticlePool<T,PropName>().GetElement(DataIdx);
 	}
 };
 
@@ -276,7 +293,7 @@ public:
 		return Dynamics.IsSet() ? Dynamics.Read().AngularImpulse() : Particle.CastToRigidParticle()->AngularImpulse();
 	}
 
-	void SyncSimWritablePropsFromSim(FDirtyPropertiesManager& Manager,int32 Idx,const TPBDRigidParticleHandle<FReal,3>& Rigid)
+	void SyncSimWritablePropsFromSim(FDirtyPropData Manager,const TPBDRigidParticleHandle<FReal,3>& Rigid)
 	{
 		FParticleDirtyFlags Flags;
 		Flags.MarkDirty(EParticleFlags::XR);
@@ -284,26 +301,26 @@ public:
 		FParticleDirtyData Dirty;
 		Dirty.SetFlags(Flags);
 
-		ParticlePositionRotation.SyncRemoteData(Manager,Idx,Dirty,[&Rigid](auto& Data)
+		ParticlePositionRotation.SyncRemoteData(Manager,Dirty,[&Rigid](auto& Data)
 		{
 			Data.CopyFrom(Rigid);
 		});
 
-		Velocities.SyncRemoteData(Manager,Idx,Dirty,[&Rigid](auto& Data)
+		Velocities.SyncRemoteData(Manager,Dirty,[&Rigid](auto& Data)
 		{
 			Data.SetV(Rigid.PreV());
 			Data.SetW(Rigid.PreW());
 		});
 	}
 
-	void SyncDirtyDynamics(FDirtyPropertiesManager& DestManager,int32 DataIdx,const FParticleDirtyData& Dirty,const FDirtyPropertiesManager& SrcManager)
+	void SyncDirtyDynamics(FDirtyPropData& DestManager,const FParticleDirtyData& Dirty,const FConstDirtyPropData& SrcManager)
 	{
 		FParticleDirtyData DirtyFlags;
 		DirtyFlags.SetFlags(Dirty.GetFlags());
-
-		Dynamics.SyncRemoteData(DestManager,DataIdx,DirtyFlags,[&Dirty,&SrcManager,DataIdx](auto& Data)
+		
+		Dynamics.SyncRemoteData(DestManager,DirtyFlags,[&Dirty,&SrcManager](auto& Data)
 		{
-			Data = Dirty.GetDynamics(SrcManager,DataIdx);
+			Data = Dirty.GetDynamics(*SrcManager.Ptr,SrcManager.DataIdx);
 		});
 	}
 
@@ -346,7 +363,7 @@ public:
 		}
 	}
 
-	void SyncPrevFrame(FDirtyPropertiesManager& Manager,int32 Idx,const FDirtyProxy& Dirty)
+	void SyncPrevFrame(FDirtyPropData& Manager,const FDirtyProxy& Dirty)
 	{
 		//syncs the data before it was made dirty
 		//for sim-writable props this is only possible if those props are immutable from the sim side (sleeping, not simulated, etc...)
@@ -360,33 +377,33 @@ public:
 		//not sure if that's a valid operation, we probably need to catch sleep/awake and handle it in a special way
 		if(bSyncSimWritable)
 		{
-			ParticlePositionRotation.SyncRemoteData(Manager,Idx,Dirty.ParticleData,[Handle](FParticlePositionRotation& Data)
+			ParticlePositionRotation.SyncRemoteData(Manager,Dirty.ParticleData,[Handle](FParticlePositionRotation& Data)
 			{
 				Data.CopyFrom(*Handle);
 			});
 
 			if(auto Kinematic = Handle->CastToKinematicParticle())
 			{
-				Velocities.SyncRemoteData(Manager,Idx,Dirty.ParticleData,[Kinematic](auto& Data)
+				Velocities.SyncRemoteData(Manager,Dirty.ParticleData,[Kinematic](auto& Data)
 				{
 					Data.CopyFrom(*Kinematic);
 				});
 			}
 		}
 
-		NonFrequentData.SyncRemoteData(Manager,Idx,Dirty.ParticleData,[Handle](FParticleNonFrequentData& Data)
+		NonFrequentData.SyncRemoteData(Manager,Dirty.ParticleData,[Handle](FParticleNonFrequentData& Data)
 		{
 			Data.CopyFrom(*Handle);
 		});
 
 		if(auto Rigid = Handle->CastToRigidParticle())
 		{
-			DynamicsMisc.SyncRemoteData(Manager,Idx,Dirty.ParticleData,[Rigid](FParticleDynamicMisc& Data)
+			DynamicsMisc.SyncRemoteData(Manager,Dirty.ParticleData,[Rigid](FParticleDynamicMisc& Data)
 			{
 				Data.CopyFrom(*Rigid);
 			});
 
-			MassProps.SyncRemoteData(Manager,Idx,Dirty.ParticleData,[Rigid](FParticleMassProps& Data)
+			MassProps.SyncRemoteData(Manager,Dirty.ParticleData,[Rigid](FParticleMassProps& Data)
 			{
 				Data.CopyFrom(*Rigid);
 			});
@@ -394,14 +411,14 @@ public:
 		
 	}
 
-	void SyncIfDirty(FDirtyPropertiesManager& Manager,int32 Idx,const TGeometryParticle<FReal,3>& InParticle, const FGeometryParticleStateBase& RewindState)
+	void SyncIfDirty(const FDirtyPropData& Manager,const TGeometryParticle<FReal,3>& InParticle, const FGeometryParticleStateBase& RewindState)
 	{
 		ensure(IsInGameThread());
 		const auto Particle = &InParticle;
 
 		if(RewindState.ParticlePositionRotation.IsSet())
 		{
-			ParticlePositionRotation.SyncRemoteDataForced(Manager,Idx,[Particle](FParticlePositionRotation& Data)
+			ParticlePositionRotation.SyncRemoteDataForced(Manager,[Particle](FParticlePositionRotation& Data)
 			{
 				Data.CopyFrom(*Particle);
 			});
@@ -411,7 +428,7 @@ public:
 		{
 			if(RewindState.Velocities.IsSet())
 			{
-				Velocities.SyncRemoteDataForced(Manager,Idx,[Kinematic](auto& Data)
+				Velocities.SyncRemoteDataForced(Manager,[Kinematic](auto& Data)
 				{
 					Data.CopyFrom(*Kinematic);
 				});
@@ -422,7 +439,7 @@ public:
 		{
 			if(RewindState.DynamicsMisc.IsSet())
 			{
-				DynamicsMisc.SyncRemoteDataForced(Manager,Idx,[Rigid](FParticleDynamicMisc& Data)
+				DynamicsMisc.SyncRemoteDataForced(Manager,[Rigid](FParticleDynamicMisc& Data)
 				{
 					Data.CopyFrom(*Rigid);
 				});
@@ -430,7 +447,7 @@ public:
 
 			if(RewindState.MassProps.IsSet())
 			{
-				MassProps.SyncRemoteDataForced(Manager,Idx,[Rigid](FParticleMassProps& Data)
+				MassProps.SyncRemoteDataForced(Manager,[Rigid](FParticleMassProps& Data)
 				{
 					Data.CopyFrom(*Rigid);
 				});
@@ -476,23 +493,23 @@ public:
 		return bCoalesced;
 	}
 
-	bool IsDesynced(const FDirtyPropertiesManager& SrcManager, const int32 DataIdxIn, const TGeometryParticleHandle<FReal,3>& Handle, const FParticleDirtyFlags Flags) const
+	bool IsDesynced(const FConstDirtyPropData& SrcManager, const TGeometryParticleHandle<FReal,3>& Handle, const FParticleDirtyFlags Flags) const
 	{
 		bool Desynced = false;
 		{
-			if(!ParticlePositionRotation.IsInSync(SrcManager,DataIdxIn,Flags,Handle))
+			if(!ParticlePositionRotation.IsInSync(SrcManager,Flags,Handle))
 			{
 				return true;
 			}
 
-			if(!NonFrequentData.IsInSync(SrcManager,DataIdxIn,Flags,Handle))
+			if(!NonFrequentData.IsInSync(SrcManager,Flags,Handle))
 			{
 				return true;
 			}
 
 			if(auto Kinematic = Handle.CastToKinematicParticle())
 			{
-				if(!Velocities.IsInSync(SrcManager,DataIdxIn,Flags,*Kinematic))
+				if(!Velocities.IsInSync(SrcManager,Flags,*Kinematic))
 				{
 					return true;
 				}
@@ -500,17 +517,17 @@ public:
 
 			if(auto Rigid = Handle.CastToRigidParticle())
 			{
-				if(!Dynamics.IsInSync(SrcManager,DataIdxIn,Flags,*Rigid))
+				if(!Dynamics.IsInSync(SrcManager,Flags,*Rigid))
 				{
 					return true;
 				}
 
-				if(!DynamicsMisc.IsInSync(SrcManager,DataIdxIn,Flags,*Rigid))
+				if(!DynamicsMisc.IsInSync(SrcManager,Flags,*Rigid))
 				{
 					return true;
 				}
 
-				if(!MassProps.IsInSync(SrcManager,DataIdxIn,Flags,*Rigid))
+				if(!MassProps.IsInSync(SrcManager,Flags,*Rigid))
 				{
 					return true;
 				}
@@ -677,9 +694,9 @@ public:
 		State = InState;
 	}
 
-	bool IsDesynced(const FDirtyPropertiesManager& SrcManager, const int32 DataIdxIn, const TGeometryParticleHandle<FReal,3>& Handle, const FParticleDirtyFlags Flags) const
+	bool IsDesynced(const FConstDirtyPropData& SrcManager, const TGeometryParticleHandle<FReal,3>& Handle, const FParticleDirtyFlags Flags) const
 	{
-		return State.IsDesynced(SrcManager,DataIdxIn,Handle,Flags);
+		return State.IsDesynced(SrcManager,Handle,Flags);
 	}
 
 private:
@@ -709,6 +726,7 @@ public:
 	}
 
 	int32 Capacity() const { return Managers.Capacity(); }
+	int32 CurrentFrame() const { return CurFrame; }
 
 	FReal GetDeltaTimeForFrame(int32 Frame) const
 	{
@@ -755,7 +773,8 @@ public:
 			
 				if(const FGeometryParticleStateBase* RewindState = GetStateAtFrameImp(DirtyParticleInfo, Frame))
 				{
-					LatestState.SyncIfDirty(*DestManager,DataIdx++,*DirtyParticleInfo.Particle,*RewindState);
+
+					LatestState.SyncIfDirty(FDirtyPropData(DestManager,DataIdx++),*DirtyParticleInfo.Particle,*RewindState);
 					CoalesceBack(DirtyParticleInfo.Frames, CurFrame);
 
 					RewindState->SyncToParticle(*DirtyParticleInfo.Particle);
@@ -913,10 +932,11 @@ public:
 		//Some sim-writable properties can change without the GT knowing about it, see PushPTDirtyData
 
 		//User called PrepareManagerForFrame for this frame so use it
-		FDirtyPropertiesManager& DestManager = *Managers[CurFrame].Manager;
+		FConstDirtyPropData SrcManagerWrapper(&SrcManager, SrcDataIdx);
+		FDirtyPropData DestManagerWrapper(Managers[CurFrame].Manager.Get(), DestDataIdx);
 		bNeedsSave = true;
 		
-		auto ProcessProxy = [this,&SrcManager, DestDataIdx, SrcDataIdx, Dirty, &DestManager](const auto Proxy)
+		auto ProcessProxy = [this,&SrcManagerWrapper, SrcDataIdx, Dirty, &DestManagerWrapper](const auto Proxy)
 		{
 			const auto PTParticle = Proxy->GetHandle();
 			FDirtyParticleInfo& Info = FindOrAddParticle(*Proxy->GetParticle(),PTParticle->UniqueIdx());
@@ -929,7 +949,7 @@ public:
 				FGeometryParticleState FutureState(*Proxy->GetParticle());
 				if(GetFutureStateAtFrame(FutureState,CurFrame) == EFutureQueryResult::Ok)
 				{
-					if(FutureState.IsDesynced(SrcManager, SrcDataIdx, *PTParticle, Dirty.ParticleData.GetFlags()))
+					if(FutureState.IsDesynced(SrcManagerWrapper, *PTParticle, Dirty.ParticleData.GetFlags()))
 					{
 						Info.Desync(CurFrame-1, LatestFrame);
 					}
@@ -948,7 +968,7 @@ public:
 				if(!bResim || FramesSaved > 0)
 				{
 					FGeometryParticleStateBase& LatestState = Info.AddFrame(CurFrame-1);
-					LatestState.SyncPrevFrame(DestManager,DestDataIdx,Dirty);
+					LatestState.SyncPrevFrame(DestManagerWrapper,Dirty);
 					CoalesceBack(Info.Frames,CurFrame-1);
 				}
 			}
@@ -957,7 +977,7 @@ public:
 			if(Dirty.ParticleData.IsDirty(EParticleFlags::Dynamics))
 			{
 				FGeometryParticleStateBase& LatestState = Info.AddFrame(CurFrame);
-				LatestState.SyncDirtyDynamics(DestManager,DestDataIdx,Dirty.ParticleData,SrcManager);
+				LatestState.SyncDirtyDynamics(DestManagerWrapper,Dirty.ParticleData,SrcManagerWrapper);
 			}
 		};
 
@@ -998,11 +1018,11 @@ public:
 			Info.LastDirtyFrame = CurFrame-1;
 
 			//User called PrepareManagerForFrame (or PrepareFrameForPTDirty) for the previous frame, so use it
-			FDirtyPropertiesManager& DestManager = *Managers[CurFrame-1].Manager;
+			FDirtyPropData DestManagerWrapper(Managers[CurFrame-1].Manager.Get(), DestDataIdx);
 
 			//sim-writable properties changed at head, so we must write down what they were
 			FGeometryParticleStateBase& LatestState = Info.AddFrame(CurFrame-1);
-			LatestState.SyncSimWritablePropsFromSim(DestManager,DestDataIdx,Rigid);
+			LatestState.SyncSimWritablePropsFromSim(DestManagerWrapper,Rigid);
 
 			//update any previous frames that were pointing at head
 			CoalesceBack(Info.Frames, CurFrame-1);
