Fix CIS.

diff --git a/Engine/Source/Runtime/Experimental/Chaos/Private/Chaos/CollisionResolution.cpp b/Engine/Source/Runtime/Experimental/Chaos/Private/Chaos/CollisionResolution.cpp
--- a/Engine/Source/Runtime/Experimental/Chaos/Private/Chaos/CollisionResolution.cpp
+++ b/Engine/Source/Runtime/Experimental/Chaos/Private/Chaos/CollisionResolution.cpp
@@ -2610,7 +2610,7 @@ namespace Chaos
 				const FImplicitObjectUnionClustered* Union0 = Implicit0->template GetObject<FImplicitObjectUnionClustered>();
 				for (const auto& Child0 : Union0->GetObjects())
 				{
-					ConstructConstraints(Particle0, Particle1, Child0.Get(), Implicit1, Transform0, Transform1, CullDistance, Context, NewConstraints);
+					ConstructConstraints(Particle0, Particle1, Child0.Get(), Implicit1, LocalTransform0, LocalTransform1, CullDistance, Context, NewConstraints);
 				}
 				return;
 			}
@@ -2631,7 +2631,7 @@ namespace Chaos
 				const FImplicitObjectUnionClustered* Union1 = Implicit1->template GetObject<FImplicitObjectUnionClustered>();
 				for (const auto& Child1 : Union1->GetObjects())
 				{
-					ConstructConstraints(Particle0, Particle1, Implicit0, Child1.Get(), Transform0, Transform1, CullDistance, Context, NewConstraints);
+					ConstructConstraints(Particle0, Particle1, Implicit0, Child1.Get(), LocalTransform0, LocalTransform1, CullDistance, Context, NewConstraints);
 				}
 				return;
 			}
diff --git a/Engine/Source/Runtime/Experimental/Chaos/Private/Chaos/ParticleHandle.cpp b/Engine/Source/Runtime/Experimental/Chaos/Private/Chaos/ParticleHandle.cpp
--- a/Engine/Source/Runtime/Experimental/Chaos/Private/Chaos/ParticleHandle.cpp
+++ b/Engine/Source/Runtime/Experimental/Chaos/Private/Chaos/ParticleHandle.cpp
@@ -123,7 +123,7 @@ namespace Chaos
 		{
 			if (const auto* PerShapeData = GetImplicitShape(Implicit))
 			{
-				if (PerShapeData->bDisable)
+				if (!PerShapeData->GetSimEnabled())
 				{
 					return;
 				}
@@ -135,7 +135,7 @@ namespace Chaos
 			}
 			else
 			{
-				Implicit->SetCollsionType(MGeometry->GetType());
+				Implicit->SetCollsionType(Implicit->GetType());
 				// @todo (mlentine): Need to in theory set convex properly here
 			}
 		}
diff --git a/Engine/Source/Runtime/Experimental/Chaos/Public/Chaos/ParticleHandle.h b/Engine/Source/Runtime/Experimental/Chaos/Public/Chaos/ParticleHandle.h
--- a/Engine/Source/Runtime/Experimental/Chaos/Public/Chaos/ParticleHandle.h
+++ b/Engine/Source/Runtime/Experimental/Chaos/Public/Chaos/ParticleHandle.h
@@ -1583,9 +1583,9 @@ public:
 	void SetIgnoreAnalyticCollisionsImp(FImplicitObject* Implicit, bool bIgnoreAnalyticCollisions);
 	void SetIgnoreAnalyticCollisions(bool bIgnoreAnalyticCollisions)
 	{
-		if (MGeometry)
+		if (MNonFrequentData.Read().Geometry())
 		{
-			SetIgnoreAnalyticCollisionsImp(MGeometry.Get(), bIgnoreAnalyticCollisions);
+			SetIgnoreAnalyticCollisionsImp(MNonFrequentData.Read().Geometry().Get(), bIgnoreAnalyticCollisions);
 		}
 	}
 
