Chaos : Cluster Transforms
- Stopped clobbering MassToLocal on creation.
- Removed overwriting of the DynamicCollections transform array on initialization, instead it now preserves the original animation transformation hierarchy.
- Incorporated use of MassToLocal in BufferPhysicsResults.
- Clean up MassToLocal usage. MassToLocal is still used in the initialization and output processes.

diff --git a/Engine/Source/Runtime/Experimental/ChaosSolvers/Private/PhysicsProxy/GeometryCollectionPhysicsProxy.cpp b/Engine/Source/Runtime/Experimental/ChaosSolvers/Private/PhysicsProxy/GeometryCollectionPhysicsProxy.cpp
--- a/Engine/Source/Runtime/Experimental/ChaosSolvers/Private/PhysicsProxy/GeometryCollectionPhysicsProxy.cpp
+++ b/Engine/Source/Runtime/Experimental/ChaosSolvers/Private/PhysicsProxy/GeometryCollectionPhysicsProxy.cpp
@@ -219,7 +219,6 @@ void PopulateSimulatedParticle(
 	const FCollisionFilterData QueryFilterIn,
 	float MassIn,
 	FVector InertiaTensorVec,
-	const FTransform& MassOffset,
 	const FTransform& WorldTransform, 
 	const uint8 DynamicState, 
 	const int16 CollisionGroup)
@@ -749,7 +748,6 @@ void FGeometryCollectionPhysicsProxy::InitializeBodiesPT(Chaos::FPBDRigidsSolver
 					QueryFilter,
 					Mass[TransformGroupIndex],
 					InertiaTensor[TransformGroupIndex],
-					MassToLocal[TransformGroupIndex],
 					WorldTransform,
 					(uint8)DynamicState[TransformGroupIndex],
 					CollisionGroup[TransformGroupIndex]);
@@ -1056,11 +1054,6 @@ FGeometryCollectionPhysicsProxy::BuildClusters(
 		}
 	}
 
-	// In theory we should be computing this and passing in to avoid inertia computation at 
-	// runtime. If we do this we must account for leaf particles that have already been created in world space
-	//const FTransform ParticleTM(Particles.R(NewSolverClusterID), Particles.X(NewSolverClusterID));	
-	MassToLocal[CollectionClusterIndex] = FTransform::Identity;
-
 	check(Parameters.RestCollection);
 	const TManagedArray<float>& Mass = 
 		Parameters.RestCollection->GetAttribute<float>("Mass", FTransformCollection::TransformGroup);
@@ -1076,7 +1069,6 @@ FGeometryCollectionPhysicsProxy::BuildClusters(
 		QueryFilter,
 		Parent->M() > 0.0 ? Parent->M() : Mass[CollectionClusterIndex], 
 		Parent->I().GetDiagonal() != Chaos::TVector<float,3>(0.0) ? Parent->I().GetDiagonal() : InertiaTensor[CollectionClusterIndex],
-		MassToLocal[CollectionClusterIndex],
 		ParticleTM, 
 		(uint8)DynamicState[CollectionClusterIndex], 
 		0); // CollisionGroup
@@ -1112,19 +1104,6 @@ FGeometryCollectionPhysicsProxy::BuildClusters(
 		const int32 ChildTransformGroupIndex = ChildTransformGroupIndices[Idx];
 		SolverClusterHandles[ChildTransformGroupIndex] = Parent;
 
-		const FTransform ChildTransform(Child->R(), Child->X());
-		if(Children[ChildTransformGroupIndex].Num()) // clustered local transform
-		{
-			Transform[ChildTransformGroupIndex] = ChildTransform.GetRelativeTransform(ParticleTM);
-		}
-		else // rigid local transform
-		{
-			const FTransform RestTransform = 
-				Parameters.RestCollection->Transform[ChildTransformGroupIndex] * ParentTransform * Parameters.WorldTransform;
-			Transform[ChildTransformGroupIndex] = RestTransform.GetRelativeTransform(ParticleTM);
-		}
-		Transform[ChildTransformGroupIndex].NormalizeRotation();
-
 		MinCollisionGroup = FMath::Min(Child->CollisionGroup(), MinCollisionGroup);
 	}
 	Parent->SetCollisionGroup(MinCollisionGroup);
@@ -1498,8 +1477,9 @@ void FGeometryCollectionPhysicsProxy::BufferPhysicsResults()
 	
 				FTransform& ParticleToWorld = TargetResults.ParticleToWorldTransforms[TransformGroupIndex];
 				ParticleToWorld = FTransform(Handle->R(), Handle->X());
+				const FTransform MassToLocal = PhysicsThreadCollection.MassToLocal[TransformGroupIndex];
 
-				TargetResults.Transforms[TransformGroupIndex] = PhysicsThreadCollection.MassToLocal[TransformGroupIndex].GetRelativeTransformReverse(ParticleToWorld).GetRelativeTransform(ActorToWorld);
+				TargetResults.Transforms[TransformGroupIndex] = MassToLocal.GetRelativeTransformReverse(ParticleToWorld).GetRelativeTransform(ActorToWorld);
 				TargetResults.Transforms[TransformGroupIndex].NormalizeRotation();
 
 				PhysicsThreadCollection.Transform[TransformGroupIndex] = TargetResults.Transforms[TransformGroupIndex];
@@ -1561,13 +1541,13 @@ void FGeometryCollectionPhysicsProxy::BufferPhysicsResults()
 						if (ClusterParent->InternalCluster() )
 						{
 							Chaos::TPBDRigidClusteredParticleHandle<float, 3>* ProxyElementHandle = SolverParticleHandles[TransformGroupIndex];
-							const FTransform ClusterToWorld = FTransform(ClusterParent->R(), ClusterParent->X());
-							const FTransform ChildToCluster = FTransform(ProxyElementHandle->ChildToParent());
-							const FTransform ClusterChildToWorld = ChildToCluster * ClusterToWorld;
-							FTransform MassToLocal = FTransform::Identity;// @todo(chaos): what should the mass to local for the internal cluster be?
+
+							FTransform& ParticleToWorld = TargetResults.ParticleToWorldTransforms[TransformGroupIndex];
+							ParticleToWorld = ProxyElementHandle->ChildToParent() * FTransform(ClusterParent->R(), ClusterParent->X()); // aka ClusterChildToWorld
 
 							// GeomToActor = ActorToWorld.Inv() * ClusterChildToWorld * MassToLocal.Inv();
-							TargetResults.Transforms[TransformGroupIndex] = MassToLocal.GetRelativeTransformReverse(ClusterChildToWorld).GetRelativeTransform(ActorToWorld);
+							const FTransform MassToLocal = PhysicsThreadCollection.MassToLocal[TransformGroupIndex];
+							TargetResults.Transforms[TransformGroupIndex] = MassToLocal.GetRelativeTransformReverse(ParticleToWorld).GetRelativeTransform(ActorToWorld);
 							TargetResults.Transforms[TransformGroupIndex].NormalizeRotation();
 
 							PhysicsThreadCollection.Transform[TransformGroupIndex] = TargetResults.Transforms[TransformGroupIndex];
@@ -1576,8 +1556,8 @@ void FGeometryCollectionPhysicsProxy::BufferPhysicsResults()
 							TargetResults.DisabledStates[TransformGroupIndex] = false;
 							IsObjectDynamic = true;
 
-							ProxyElementHandle->X() = ClusterChildToWorld.GetTranslation();
-							ProxyElementHandle->R() = ClusterChildToWorld.GetRotation();
+							ProxyElementHandle->X() = ParticleToWorld.GetTranslation();
+							ProxyElementHandle->R() = ParticleToWorld.GetRotation();
 							GetSolver()->GetEvolution()->DirtyParticle(*ProxyElementHandle);
 						}
 					}
@@ -2144,7 +2124,6 @@ void FGeometryCollectionPhysicsProxy::InitializeSharedCollisionStructures(
 					FCollisionFilterData(),		// QueryFilter
 					CollectionMass[TransformGroupIndex],
 					CollectionInertiaTensor[TransformGroupIndex], 
-					CollectionMassToLocal[TransformGroupIndex],
 					CollectionSpaceTransforms[TransformGroupIndex],
 					(uint8)EObjectStateTypeEnum::Chaos_Object_Dynamic, 
 					INDEX_NONE); // CollisionGroup
@@ -2187,9 +2166,6 @@ void FGeometryCollectionPhysicsProxy::InitializeSharedCollisionStructures(
 					FTransform(CollectionSpaceParticles->R(ClusterTransformIdx), 
 							   CollectionSpaceParticles->X(ClusterTransformIdx));
 
-				// Compute MassToLocal as if the transform hierarchy stays fixed. 
-				// In reality we modify the transform hierarchy so that MassToLocal 
-				// is identity at runtime.
 				CollectionMassToLocal[ClusterTransformIdx] = 
 					ClusterMassToCollection.GetRelativeTransform(
 						CollectionSpaceTransforms[ClusterTransformIdx]);
