Add angular friction override.

diff --git a/Engine/Source/Runtime/Experimental/Chaos/Private/Chaos/PBDCollisionConstraints.cpp b/Engine/Source/Runtime/Experimental/Chaos/Private/Chaos/PBDCollisionConstraints.cpp
--- a/Engine/Source/Runtime/Experimental/Chaos/Private/Chaos/PBDCollisionConstraints.cpp
+++ b/Engine/Source/Runtime/Experimental/Chaos/Private/Chaos/PBDCollisionConstraints.cpp
@@ -49,6 +49,9 @@ namespace Chaos
 
 	float CollisionRestitutionOverride = -1.0f;
 	FAutoConsoleVariableRef CVarCollisionRestitutionOverride(TEXT("p.CollisionRestitution"), CollisionRestitutionOverride, TEXT("Collision restitution for all contacts if >= 0"));
+	
+	float CollisionAngularFrictionOverride = -1.0f;
+	FAutoConsoleVariableRef CVarCollisionAngularFrictionOverride(TEXT("p.CollisionAngularFriction"), CollisionAngularFrictionOverride, TEXT("Collision angular friction for all contacts if >= 0"));
 
 	CHAOS_API int32 EnableCollisions = 1;
 	FAutoConsoleVariableRef CVarEnableCollisions(TEXT("p.EnableCollisions"), EnableCollisions, TEXT("Enable/Disable collisions on the Chaos solver."));
@@ -206,6 +209,10 @@ namespace Chaos
 		{
 			Contact.Restitution = CollisionRestitutionOverride;
 		}
+		if (CollisionAngularFrictionOverride >= 0)
+		{
+			Contact.AngularFriction = CollisionAngularFrictionOverride;
+		}
 	}
 
 	void FPBDCollisionConstraints::AddConstraint(const FRigidBodyPointContactConstraint& InConstraint)
