Move p.HackMaxVelocity2 hotfix value to code.

diff --git a/Engine/Source/Runtime/Experimental/Chaos/Private/Chaos/PBDRigidsEvolutionGBF.cpp b/Engine/Source/Runtime/Experimental/Chaos/Private/Chaos/PBDRigidsEvolutionGBF.cpp
--- a/Engine/Source/Runtime/Experimental/Chaos/Private/Chaos/PBDRigidsEvolutionGBF.cpp
+++ b/Engine/Source/Runtime/Experimental/Chaos/Private/Chaos/PBDRigidsEvolutionGBF.cpp
@@ -34,7 +34,7 @@ namespace Chaos
 float HackMaxAngularVelocity = 1000.f;
 FAutoConsoleVariableRef CVarHackMaxAngularVelocity(TEXT("p.HackMaxAngularVelocity"), HackMaxAngularVelocity, TEXT("Max cap on angular velocity: rad/s. This is only a temp solution and should not be relied on as a feature. -1.f to disable"));
 
-float HackMaxVelocity = -1.f;
+float HackMaxVelocity = 2500.f;
 FAutoConsoleVariableRef CVarHackMaxVelocity(TEXT("p.HackMaxVelocity2"), HackMaxVelocity, TEXT("Max cap on velocity: cm/s. This is only a temp solution and should not be relied on as a feature. -1.f to disable"));
 
 
