default p.UnionsHaveCollisionParticles 1
This flag was causing collisions with internal clusters to not work because they were not inheriting collision particles from the children

diff --git a/Engine/Source/Runtime/Experimental/Chaos/Private/Chaos/PBDRigidClustering.cpp b/Engine/Source/Runtime/Experimental/Chaos/Private/Chaos/PBDRigidClustering.cpp
--- a/Engine/Source/Runtime/Experimental/Chaos/Private/Chaos/PBDRigidClustering.cpp
+++ b/Engine/Source/Runtime/Experimental/Chaos/Private/Chaos/PBDRigidClustering.cpp
@@ -510,7 +510,7 @@ namespace Chaos
 		return NewParticle;
 	}
 
-	int32 UnionsHaveCollisionParticles = 0;
+	int32 UnionsHaveCollisionParticles = 1;
 	FAutoConsoleVariableRef CVarUnionsHaveCollisionParticles(TEXT("p.UnionsHaveCollisionParticles"), UnionsHaveCollisionParticles, TEXT(""));
 
 	DECLARE_CYCLE_STAT(TEXT("TPBDRigidClustering<>::CreateClusterParticleFromClusterChildren"), STAT_CreateClusterParticleFromClusterChildren, STATGROUP_Chaos);
