Make static analysis happier.

diff --git a/Engine/Source/Runtime/Experimental/Chaos/Private/Chaos/Levelset.cpp b/Engine/Source/Runtime/Experimental/Chaos/Private/Chaos/Levelset.cpp
--- a/Engine/Source/Runtime/Experimental/Chaos/Private/Chaos/Levelset.cpp
+++ b/Engine/Source/Runtime/Experimental/Chaos/Private/Chaos/Levelset.cpp
@@ -388,9 +388,9 @@ T TLevelSet<T, d>::ComputeLevelSetError(const TParticles<T, d>& InParticles, con
 		else
 		{
 			Pt = MaxPt;
-			if (i == 1) Pt[0] = MinPt[0];
-			else if (i == 2) Pt[1] = MinPt[1];
-			else if (i == 3) Pt[2] = MinPt[2];
+			if (i == 5) Pt[0] = MinPt[0];
+			else if (i == 6) Pt[1] = MinPt[1];
+			else if (i == 7) Pt[2] = MinPt[2];
 		}
 
 		BBox.PhiWithNormal(Pt, BoxNorm);
