Make ensure less strong to avoid breaking build

diff --git a/Engine/Source/Runtime/Experimental/Chaos/Public/Chaos/Plane.h b/Engine/Source/Runtime/Experimental/Chaos/Public/Chaos/Plane.h
--- a/Engine/Source/Runtime/Experimental/Chaos/Public/Chaos/Plane.h
+++ b/Engine/Source/Runtime/Experimental/Chaos/Public/Chaos/Plane.h
@@ -3,6 +3,7 @@
 
 #include "Chaos/ImplicitObject.h"
 #include "ChaosArchive.h"
+#include "ChaosCheck.h"
 
 namespace Chaos
 {
@@ -44,7 +45,7 @@ public:
 	bool Raycast(const TVector<T, d>& StartPoint, const TVector<T, d>& Dir, const T Length, const T Thickness, T& OutTime, TVector<T, d>& OutPosition, TVector<T, d>& OutNormal, int32& OutFaceIndex) const
 	{
 		ensure(FMath::IsNearlyEqual(Dir.SizeSquared(),1, KINDA_SMALL_NUMBER));
-		ensure(Length > 0);
+		CHAOS_ENSURE(Length > 0);
 		OutFaceIndex = INDEX_NONE;
 
 		const T SignedDist = TVector<T, d>::DotProduct(StartPoint - MX, MNormal);
