Fixed phantom collisions in Chaos - No longer only removing dynamic particles from active particle list on disable and don't force dirty every geometry collection particle on push exec.

diff --git a/Engine/Source/Runtime/Experimental/Chaos/Public/Chaos/PBDRigidsSOAs.h b/Engine/Source/Runtime/Experimental/Chaos/Public/Chaos/PBDRigidsSOAs.h
--- a/Engine/Source/Runtime/Experimental/Chaos/Public/Chaos/PBDRigidsSOAs.h
+++ b/Engine/Source/Runtime/Experimental/Chaos/Public/Chaos/PBDRigidsSOAs.h
@@ -261,7 +261,6 @@ public:
 			}
 
 			// All active particles RIGID particles
-			if (Particle->ObjectState() == EObjectStateType::Dynamic)
 			{
 				RemoveFromMapAndArray(PBDRigid, ActiveParticlesToIndex, ActiveParticlesArray);
 			}
diff --git a/Engine/Source/Runtime/Experimental/ChaosSolvers/Private/PBDRigidsSolver.cpp b/Engine/Source/Runtime/Experimental/ChaosSolvers/Private/PBDRigidsSolver.cpp
--- a/Engine/Source/Runtime/Experimental/ChaosSolvers/Private/PBDRigidsSolver.cpp
+++ b/Engine/Source/Runtime/Experimental/ChaosSolvers/Private/PBDRigidsSolver.cpp
@@ -677,14 +677,32 @@ namespace Chaos
 	void PushPhysicsStateExec(FPBDRigidsSolver* Solver, FGeometryCollectionPhysicsProxy* Proxy, Chaos::IDispatcher* Dispatcher)
 	{
 		Proxy->NewData();
+
 		auto Cmd = [Proxy, Solver](Chaos::FPersistentPhysicsTask* PhysThread)
 		{
-			auto* Evolution = Solver->GetEvolution();
+			FPBDRigidsSolver::FPBDRigidsEvolution* Evolution = Solver->GetEvolution();
+			
 			TArray<FGeometryCollectionPhysicsProxy::FClusterHandle*>& Handles = Proxy->GetSolverParticleHandles();
-			for (auto* Handle : Handles) if (Handle) Evolution->DirtyParticle(*Handle);
+			for(FGeometryCollectionPhysicsProxy::FClusterHandle* Handle : Handles)
+			{
+				if(Handle && !Handle->Disabled())
+				{
+					Evolution->DirtyParticle(*Handle);
+				}
+			}
+
 			Proxy->PushToPhysicsState(nullptr);
 		};
-		if (Dispatcher) Dispatcher->EnqueueCommandImmediate(Cmd);  else Cmd(nullptr);
+
+		if(Dispatcher)
+		{
+			Dispatcher->EnqueueCommandImmediate(Cmd);
+		}
+		else
+		{
+			Cmd(nullptr);
+		}
+
 		Proxy->ClearAccumulatedData();
 		Solver->RemoveDirtyProxy(Proxy);
 	}
@@ -692,15 +710,21 @@ namespace Chaos
 	void PushPhysicsStateExec(FPBDRigidsSolver* Solver, FFieldSystemPhysicsProxy* Proxy, Chaos::IDispatcher* Dispatcher)
 	{
 		Chaos::FParticleData* ProxyData = Proxy->NewData();
+
 		auto Cmd = [Proxy, Solver, ProxyData](Chaos::FPersistentPhysicsTask* PhysThread)
 		{
 			Proxy->PushToPhysicsState(nullptr);
 		};
 
-		if (Dispatcher)
+		if(Dispatcher)
+		{
 			Dispatcher->EnqueueCommandImmediate(Cmd);
+		}
 		else
+		{
 			Cmd(nullptr);
+		}
+
 		Proxy->ClearAccumulatedData();
 		Solver->RemoveDirtyProxy(Proxy);
 	}
