Chaos
- RBAN joint angular projection (Disabled - WIP)

diff --git a/Engine/Source/Runtime/Engine/Private/PhysicsEngine/ImmediatePhysics/ImmediatePhysicsChaos/ImmediatePhysicsJointHandle_Chaos.cpp b/Engine/Source/Runtime/Engine/Private/PhysicsEngine/ImmediatePhysics/ImmediatePhysicsChaos/ImmediatePhysicsJointHandle_Chaos.cpp
--- a/Engine/Source/Runtime/Engine/Private/PhysicsEngine/ImmediatePhysics/ImmediatePhysicsChaos/ImmediatePhysicsJointHandle_Chaos.cpp
+++ b/Engine/Source/Runtime/Engine/Private/PhysicsEngine/ImmediatePhysics/ImmediatePhysicsChaos/ImmediatePhysicsJointHandle_Chaos.cpp
@@ -38,7 +38,7 @@ namespace ImmediatePhysics_Chaos
 	float ChaosImmediate_JointMinLinearProjection = 0.0f;
 	float ChaosImmediate_JointMaxLinearProjection = 1.0f;
 	float ChaosImmediate_JointMinAngularProjection = 0.0f;
-	float ChaosImmediate_JointMaxAngularProjection = 0.1f;
+	float ChaosImmediate_JointMaxAngularProjection = 0.0f;
 
 	FAutoConsoleVariableRef CVarJointStiffness(TEXT("p.Chaos.ImmPhys.JointStiffness"), ChaosImmediate_JointStiffness, TEXT("Hard-joint solver stiffness."));
 	FAutoConsoleVariableRef CVarLinearDriveStiffnessScale(TEXT("p.Chaos.ImmPhys.LinearDriveStiffnessScale"), ChaosImmediate_LinearDriveStiffnessScale, TEXT("Conversion factor for Linear drive stiffness."));
diff --git a/Engine/Source/Runtime/Engine/Private/PhysicsEngine/ImmediatePhysics/ImmediatePhysicsChaos/ImmediatePhysicsSimulation_Chaos.cpp b/Engine/Source/Runtime/Engine/Private/PhysicsEngine/ImmediatePhysics/ImmediatePhysicsChaos/ImmediatePhysicsSimulation_Chaos.cpp
--- a/Engine/Source/Runtime/Engine/Private/PhysicsEngine/ImmediatePhysics/ImmediatePhysicsChaos/ImmediatePhysicsSimulation_Chaos.cpp
+++ b/Engine/Source/Runtime/Engine/Private/PhysicsEngine/ImmediatePhysics/ImmediatePhysicsChaos/ImmediatePhysicsSimulation_Chaos.cpp
@@ -79,8 +79,8 @@ float ChaosImmediate_Joint_AngleTolerance = 0.001f;
 int32 ChaosImmediate_Joint_EnableTwistLimits = 1;
 int32 ChaosImmediate_Joint_EnableSwingLimits = 1;
 int32 ChaosImmediate_Joint_EnableDrives = 1;
-float ChaosImmediate_Joint_LinearProjection = 0.0f;
-float ChaosImmediate_Joint_AngularProjection = 0.0f;
+float ChaosImmediate_Joint_LinearProjection = -1.0f;
+float ChaosImmediate_Joint_AngularProjection = -1.0f;
 float ChaosImmediate_Joint_Stiffness = 1.0f;
 float ChaosImmediate_Joint_SoftLinearStiffness = 0.0f;
 float ChaosImmediate_Joint_SoftTwistStiffness = 0.0f;
@@ -94,6 +94,8 @@ float ChaosImmediate_Joint_AngularDriveDamping = 0.0f;
 float ChaosImmediate_Joint_MinParentMassRatio = 0.2f;
 float ChaosImmediate_Joint_MaxInertiaRatio = 5.0f;
 float ChaosImmediate_Joint_AngularPositionCorrection = 0.0f;
+float ChaosImmediate_Joint_ProjectionInvMassScale = 0.0f;
+float ChaosImmediate_Joint_VelProjectionInvMassScale = 1.0f;
 FAutoConsoleVariableRef CVarChaosImmPhysJointPairIterations(TEXT("p.Chaos.ImmPhys.Joint.PairIterations"), ChaosImmediate_Joint_PairIterations, TEXT("Override joint pair iterations (if >= 0)"));
 FAutoConsoleVariableRef CVarChaosImmPhysJointPushOutPairIterations(TEXT("p.Chaos.ImmPhys.Joint.PushOutPairIterations"), ChaosImmediate_Joint_PushOutPairIterations, TEXT("Override joint push-out pair iterations (if >= 0)"));
 FAutoConsoleVariableRef CVarChaosImmPhysJointSwingTwistAngleTolerance(TEXT("p.Chaos.ImmPhys.Joint.SwingTwistAngleTolerance"), ChaosImmediate_Joint_SwingTwistAngleTolerance, TEXT("SwingTwistAngleTolerance."));
@@ -102,8 +104,8 @@ FAutoConsoleVariableRef CVarChaosImmPhysJointAngleTolerance(TEXT("p.Chaos.ImmPhy
 FAutoConsoleVariableRef CVarChaosImmPhysJointEnableTwistLimits(TEXT("p.Chaos.ImmPhys.Joint.EnableTwistLimits"), ChaosImmediate_Joint_EnableTwistLimits, TEXT("EnableTwistLimits."));
 FAutoConsoleVariableRef CVarChaosImmPhysJointEnableSwingLimits(TEXT("p.Chaos.ImmPhys.Joint.EnableSwingLimits"), ChaosImmediate_Joint_EnableSwingLimits, TEXT("EnableSwingLimits."));
 FAutoConsoleVariableRef CVarChaosImmPhysJointEnableDrives(TEXT("p.Chaos.ImmPhys.Joint.EnableDrives"), ChaosImmediate_Joint_EnableDrives, TEXT("EnableDrives."));
-FAutoConsoleVariableRef CVarChaosImmPhysJointLinearProjection(TEXT("p.Chaos.ImmPhys.Joint.LinearProjection"), ChaosImmediate_Joint_LinearProjection, TEXT("6Dof joint projection amount override (if > 0)."));
-FAutoConsoleVariableRef CVarChaosImmPhysJointAngularProjection(TEXT("p.Chaos.ImmPhys.Joint.AngularProjection"), ChaosImmediate_Joint_AngularProjection, TEXT("6Dof joint projection amount override (if > 0)."));
+FAutoConsoleVariableRef CVarChaosImmPhysJointLinearProjection(TEXT("p.Chaos.ImmPhys.Joint.LinearProjection"), ChaosImmediate_Joint_LinearProjection, TEXT("6Dof joint projection amount override (if >= 0)."));
+FAutoConsoleVariableRef CVarChaosImmPhysJointAngularProjection(TEXT("p.Chaos.ImmPhys.Joint.AngularProjection"), ChaosImmediate_Joint_AngularProjection, TEXT("6Dof joint projection amount override (if >= 0)."));
 FAutoConsoleVariableRef CVarChaosImmPhysJointStiffness(TEXT("p.Chaos.ImmPhys.Joint.Stiffness"), ChaosImmediate_Joint_Stiffness, TEXT("6Dof joint stiffness override (if > 0)."));
 FAutoConsoleVariableRef CVarChaosImmPhysJointSoftLinearStiffness(TEXT("p.Chaos.ImmPhys.Joint.SoftLinearStiffness"), ChaosImmediate_Joint_SoftLinearStiffness, TEXT("6Dof joint soft linear stiffness override (if > 0)."));
 FAutoConsoleVariableRef CVarChaosImmPhysJointSoftTwistStiffness(TEXT("p.Chaos.ImmPhys.Joint.SoftTwistStiffness"), ChaosImmediate_Joint_SoftTwistStiffness, TEXT("6Dof joint SoftTwist stiffness override (if > 0)."));
@@ -117,6 +119,8 @@ FAutoConsoleVariableRef CVarChaosImmPhysJointAngularDriveDamping(TEXT("p.Chaos.I
 FAutoConsoleVariableRef CVarChaosImmPhysJointMinParentMassRatio(TEXT("p.Chaos.ImmPhys.Joint.MinParentMassRatio"), ChaosImmediate_Joint_MinParentMassRatio, TEXT("6Dof joint MinParentMassRatio (if > 0)"));
 FAutoConsoleVariableRef CVarChaosImmPhysJointMaxInertiaRatio(TEXT("p.Chaos.ImmPhys.Joint.MaxInertiaRatio"), ChaosImmediate_Joint_MaxInertiaRatio, TEXT("6Dof joint MaxInertiaRatio (if > 0)"));
 FAutoConsoleVariableRef CVarChaosImmPhysJointAngularPositionCorrection(TEXT("p.Chaos.ImmPhys.Joint.AngularPositionCorrection"), ChaosImmediate_Joint_AngularPositionCorrection, TEXT("6Dof joint post-rotation constraint position correction amount [0-1]"));
+FAutoConsoleVariableRef CVarChaosImmPhysJointProjectionInvMassScale(TEXT("p.Chaos.ImmPhys.Joint.ProjectionInvMassScale"), ChaosImmediate_Joint_ProjectionInvMassScale, TEXT("Scale applied to parent body's inverse mass during projection [0-1]"));
+FAutoConsoleVariableRef CVarChaosImmPhysJointVelProjectionInvMassScale(TEXT("p.Chaos.ImmPhys.Joint.VelProjectionInvMassScale"), ChaosImmediate_Joint_VelProjectionInvMassScale, TEXT("Scale applied to parent body's inverse mass during projection for velocity fixup [0-1]"));
 
 //
 // Even more temp that the above...
@@ -134,12 +138,16 @@ FAutoConsoleVariableRef CVarChaosImmPhysUsePositionSolver(TEXT("p.Chaos.ImmPhys.
 #if UE_BUILD_DEBUG
 int32 ChaosImmediate_DebugDrawParticles = 0;
 int32 ChaosImmediate_DebugDrawShapes = 1;
+int32 ChaosImmediate_DebugDrawShowKinematics = 1;
+int32 ChaosImmediate_DebugDrawShowDynamics = 1;
 int32 ChaosImmediate_DebugDrawBounds = 0;
 int32 ChaosImmediate_DebugDrawCollisions = 0;
 int32 ChaosImmediate_DebugDrawJoints = 1;
 #else
 int32 ChaosImmediate_DebugDrawParticles = 0;
 int32 ChaosImmediate_DebugDrawShapes = 0;
+int32 ChaosImmediate_DebugDrawShowKinematics = 1;
+int32 ChaosImmediate_DebugDrawShowDynamics = 1;
 int32 ChaosImmediate_DebugDrawBounds = 0;
 int32 ChaosImmediate_DebugDrawCollisions = 0;
 int32 ChaosImmediate_DebugDrawJoints = 0;
@@ -153,6 +161,8 @@ int32 ChaosImmediate_DebugDrawJointFeatures = 0;
 
 FAutoConsoleVariableRef CVarChaosImmPhysDebugDrawParticles(TEXT("p.Chaos.ImmPhys.DebugDrawParticles"), ChaosImmediate_DebugDrawParticles, TEXT("Draw Particle Transforms (0 = never; 1 = end of frame; 2 = begin and end of frame; 3 = post-integate, post-apply and post-applypushout;)."));
 FAutoConsoleVariableRef CVarChaosImmPhysDebugDrawShapes(TEXT("p.Chaos.ImmPhys.DebugDrawShapes"), ChaosImmediate_DebugDrawShapes, TEXT("Draw Shapes (0 = never; 1 = end of frame; 2 = begin and end of frame; 3 = post-integate, post-apply and post-applypushout;"));
+FAutoConsoleVariableRef CVarChaosImmPhysDebugDrawShowKinematics(TEXT("p.Chaos.ImmPhys.DebugDrawShowKinematics"), ChaosImmediate_DebugDrawShowKinematics, TEXT("Show kinematics if shape debug draw is enabled"));
+FAutoConsoleVariableRef CVarChaosImmPhysDebugDrawShowDynamics(TEXT("p.Chaos.ImmPhys.DebugDrawShowDynamics"), ChaosImmediate_DebugDrawShowDynamics, TEXT("Show dynamics if shape debug draw is enabled"));
 FAutoConsoleVariableRef CVarChaosImmPhysDebugDrawBounds(TEXT("p.Chaos.ImmPhys.DebugDrawBounds"), ChaosImmediate_DebugDrawBounds, TEXT("Draw Particle Bounds (0 = never; 1 = end of frame; 2 = begin and end of frame; 3 = post-integate, post-apply and post-applypushout;)."));
 FAutoConsoleVariableRef CVarChaosImmPhysDebugDrawCollisions(TEXT("p.Chaos.ImmPhys.DebugDrawCollisions"), ChaosImmediate_DebugDrawCollisions, TEXT("Draw Collisions (0 = never; 1 = end of frame; 2 = begin and end of frame; 3 = post-integate, post-apply and post-applypushout;)"));
 FAutoConsoleVariableRef CVarChaosImmPhysDebugDrawJoints(TEXT("p.Chaos.ImmPhys.DebugDrawJoints"), ChaosImmediate_DebugDrawJoints, TEXT("Draw Joints. (0 = never; 1 = end of frame; 2 = begin and end of frame; 3 = post-integate, post-apply and post-applypushout; 4 = each Apply step)."));
@@ -582,6 +592,8 @@ namespace ImmediatePhysics_Chaos
 			JointsSettings.MinParentMassRatio = ChaosImmediate_Joint_MinParentMassRatio;
 			JointsSettings.MaxInertiaRatio = ChaosImmediate_Joint_MaxInertiaRatio;
 			JointsSettings.AngularConstraintPositionCorrection = ChaosImmediate_Joint_AngularPositionCorrection;
+			JointsSettings.ProjectionInvMassScale = ChaosImmediate_Joint_ProjectionInvMassScale;
+			JointsSettings.VelProjectionInvMassScale = ChaosImmediate_Joint_VelProjectionInvMassScale;
 			JointsSettings.bEnableTwistLimits = ChaosImmediate_Joint_EnableTwistLimits != 0;
 			JointsSettings.bEnableSwingLimits = ChaosImmediate_Joint_EnableSwingLimits != 0;
 			JointsSettings.bEnableDrives = ChaosImmediate_Joint_EnableDrives != 0;
@@ -724,6 +736,10 @@ namespace ImmediatePhysics_Chaos
 		using namespace Chaos;
 		if (FDebugDrawQueue::IsDebugDrawingEnabled())
 		{
+			if (!ChaosImmediate_DebugDrawShowKinematics)
+			{
+				return;
+			}
 			if ((ChaosImmediate_DebugDrawParticles >= MinDebugLevel) && (ChaosImmediate_DebugDrawParticles <= MaxDebugLevel))
 			{
 				DebugDraw::DrawParticleTransforms(Implementation->SimulationSpaceTransform, Implementation->Particles.GetActiveKinematicParticlesView());
@@ -747,6 +763,10 @@ namespace ImmediatePhysics_Chaos
 		using namespace Chaos;
 		if (FDebugDrawQueue::IsDebugDrawingEnabled())
 		{
+			if (!ChaosImmediate_DebugDrawShowDynamics)
+			{
+				return;
+			}
 			if ((ChaosImmediate_DebugDrawParticles >= MinDebugLevel) && (ChaosImmediate_DebugDrawParticles <= MaxDebugLevel))
 			{
 				DebugDraw::DrawParticleTransforms(Implementation->SimulationSpaceTransform, Implementation->Particles.GetActiveParticlesView());
diff --git a/Engine/Source/Runtime/Experimental/Chaos/Private/Chaos/ChaosDebugDraw.cpp b/Engine/Source/Runtime/Experimental/Chaos/Private/Chaos/ChaosDebugDraw.cpp
--- a/Engine/Source/Runtime/Experimental/Chaos/Private/Chaos/ChaosDebugDraw.cpp
+++ b/Engine/Source/Runtime/Experimental/Chaos/Private/Chaos/ChaosDebugDraw.cpp
@@ -36,6 +36,8 @@ namespace Chaos
 		float FontScale = 1.5f;
 		float ShapeThicknesScale = 1.0f;
 		float PointSize = 2.0f;
+		float VelScale = 0.0f;
+		float AngVelScale = 0.0f;
 		int DrawPriority = 10.0f;
 
 		FAutoConsoleVariableRef CVarArrowSize(TEXT("p.Chaos.DebugDrawArrowSize"), ArrowSize, TEXT("ArrowSize."));
@@ -47,6 +49,8 @@ namespace Chaos
 		FAutoConsoleVariableRef CVarConstraintAxisLen(TEXT("p.Chaos.DebugDrawConstraintAxisLen"), ConstraintAxisLen, TEXT("ConstraintAxisLen."));
 		FAutoConsoleVariableRef CVarLineThickness(TEXT("p.Chaos.DebugDrawLineThickness"), LineThickness, TEXT("LineThickness."));
 		FAutoConsoleVariableRef CVarLineShapeThickness(TEXT("p.Chaos.DebugDrawShapeLineThicknessScale"), ShapeThicknesScale, TEXT("Shape lineThickness multiplier."));
+		FAutoConsoleVariableRef CVarVelScale(TEXT("p.Chaos.DebugDrawVelScale"), VelScale, TEXT("If >0 show velocity when drawing particle transforms."));
+		FAutoConsoleVariableRef CVarAngVelScale(TEXT("p.Chaos.DebugDrawAngVelScale"), AngVelScale, TEXT("If >0 show angular velocity when drawing particle transforms."));
 		FAutoConsoleVariableRef CVarScale(TEXT("p.Chaos.DebugDrawScale"), DrawScale, TEXT("Scale applied to all Chaos Debug Draw line lengths etc."));
 
 		//
@@ -298,7 +302,16 @@ namespace Chaos
 		
 			if (Index >= 0)
 			{
-				FDebugDrawQueue::GetInstance().DrawDebugString(PCOM + FontHeight * FVec3(0, 0, 1), FString::Format(TEXT("{0}{1}"), { Particle->IsKinematic()? TEXT("K"): TEXT("D"), Index }), nullptr, FColor::White, KINDA_SMALL_NUMBER, false, FontScale);
+				//FDebugDrawQueue::GetInstance().DrawDebugString(PCOM + FontHeight * FVec3(0, 0, 1), FString::Format(TEXT("{0}{1}"), { Particle->IsKinematic()? TEXT("K"): TEXT("D"), Index }), nullptr, FColor::White, KINDA_SMALL_NUMBER, false, FontScale);
+			}
+
+			if ((VelScale > 0.0f) && (Particle->V().Size() > KINDA_SMALL_NUMBER))
+			{
+				FDebugDrawQueue::GetInstance().DrawDebugLine(PCOM, PCOM + Particle->V() * VelScale, Red, false, KINDA_SMALL_NUMBER, DrawPriority, LineThickness);
+			}
+			if ((AngVelScale > 0.0f) && (Particle->W().Size() > KINDA_SMALL_NUMBER))
+			{
+				FDebugDrawQueue::GetInstance().DrawDebugLine(PCOM, PCOM + Particle->W() * AngVelScale, Green, false, KINDA_SMALL_NUMBER, DrawPriority, LineThickness);
 			}
 		}
 
@@ -360,8 +373,8 @@ namespace Chaos
 
 			// Draw the particle (mass frame) coordinates
 			{
-				DrawParticleTransformImpl(FRigidTransform3::Identity, Contact.Particle[0], 0, 1.0f);
-				DrawParticleTransformImpl(FRigidTransform3::Identity, Contact.Particle[1], 0, 1.0f);				
+				//DrawParticleTransformImpl(FRigidTransform3::Identity, Contact.Particle[0], INDEX_NONE, 1.0f);
+				//DrawParticleTransformImpl(FRigidTransform3::Identity, Contact.Particle[1], INDEX_NONE, 1.0f);				
 			}
 		}
 		
diff --git a/Engine/Source/Runtime/Experimental/Chaos/Private/Chaos/Joint/JointConstraintsCVars.cpp b/Engine/Source/Runtime/Experimental/Chaos/Private/Chaos/Joint/JointConstraintsCVars.cpp
new file mode 100644
--- /dev/null
+++ b/Engine/Source/Runtime/Experimental/Chaos/Private/Chaos/Joint/JointConstraintsCVars.cpp
@@ -0,0 +1,23 @@
+// Copyright Epic Games, Inc. All Rights Reserved.
+#include "Chaos/Joint/JointConstraintsCVars.h"
+#include "HAL/IConsoleManager.h"
+
+#if INTEL_ISPC && !UE_BUILD_SHIPPING
+bool bChaos_Joint_ISPC_Enabled = true;
+FAutoConsoleVariableRef CVarChaosJointISPCEnabled(TEXT("p.Chaos.Joint.ISPC"), bChaos_Joint_ISPC_Enabled, TEXT("Whether to use ISPC optimizations in the Joint Solver"));
+#endif
+
+bool bChaos_Joint_EarlyOut_Enabled = true;
+FAutoConsoleVariableRef CVarChaosJointEarlyOutEnabled(TEXT("p.Chaos.Joint.EarlyOut"), bChaos_Joint_EarlyOut_Enabled, TEXT("Whether to iterating when joints report being solved"));
+
+bool bChaos_Joint_Batching = false;
+FAutoConsoleVariableRef CVarChaosJointBatching(TEXT("p.Chaos.Joint.Batching"), bChaos_Joint_Batching, TEXT(""));
+
+int32 bChaos_Joint_MaxBatchSize = 1000;
+FAutoConsoleVariableRef CVarChaosJointBatchSize(TEXT("p.Chaos.Joint.MaxBatchSize"), bChaos_Joint_MaxBatchSize, TEXT(""));
+
+float Chaos_Joint_DegenerateRotationLimit = -0.998f;	// Cos(176deg)
+FAutoConsoleVariableRef CVarChaosJointDegenerateRotationLimit(TEXT("p.Chaos.Joint.DegenerateRotationLimit"), Chaos_Joint_DegenerateRotationLimit, TEXT("Cosine of the swing angle that is considered degerenerate (default Cos(176deg))"));
+
+bool bChaos_Joint_EllipticalFix = true;
+FAutoConsoleVariableRef CVarChaosJointEllipticalFix(TEXT("p.Chaos.Joint.EllipticalFix"), bChaos_Joint_EllipticalFix, TEXT("Enable the proper elliptical joint axis/error calculation"));
diff --git a/Engine/Source/Runtime/Experimental/Chaos/Private/Chaos/Joint/JointConstraintsCVars.h b/Engine/Source/Runtime/Experimental/Chaos/Private/Chaos/Joint/JointConstraintsCVars.h
new file mode 100644
--- /dev/null
+++ b/Engine/Source/Runtime/Experimental/Chaos/Private/Chaos/Joint/JointConstraintsCVars.h
@@ -0,0 +1,20 @@
+// Copyright Epic Games, Inc. All Rights Reserved.
+#pragma once
+
+#include "CoreMinimal.h"
+
+// Support ISPC enable/disable in non-shipping builds
+#if !INTEL_ISPC
+	const bool bChaos_Joint_ISPC_Enabled = false;
+#elif UE_BUILD_SHIPPING
+	const bool bChaos_Joint_ISPC_Enabled = true;
+#else
+	extern bool bChaos_Joint_ISPC_Enabled;
+#endif
+
+extern bool bChaos_Joint_EarlyOut_Enabled;
+extern bool bChaos_Joint_Batching;
+extern int32 bChaos_Joint_MaxBatchSize;
+
+extern float Chaos_Joint_DegenerateRotationLimit;
+extern bool bChaos_Joint_EllipticalFix;
diff --git a/Engine/Source/Runtime/Experimental/Chaos/Private/Chaos/Joint/PBDJointSolverGaussSeidel.cpp b/Engine/Source/Runtime/Experimental/Chaos/Private/Chaos/Joint/PBDJointSolverGaussSeidel.cpp
--- a/Engine/Source/Runtime/Experimental/Chaos/Private/Chaos/Joint/PBDJointSolverGaussSeidel.cpp
+++ b/Engine/Source/Runtime/Experimental/Chaos/Private/Chaos/Joint/PBDJointSolverGaussSeidel.cpp
@@ -1,6 +1,7 @@
 // Copyright Epic Games, Inc. All Rights Reserved.
 #include "Chaos/Joint/PBDJointSolverGaussSeidel.h"
 #include "Chaos/Joint/ChaosJointLog.h"
+#include "Chaos/Joint/JointConstraintsCVars.h"
 #include "Chaos/Joint/JointSolverConstraints.h"
 #include "Chaos/ParticleHandle.h"
 #include "Chaos/PBDJointConstraintUtilities.h"
@@ -12,13 +13,6 @@
 
 //PRAGMA_DISABLE_OPTIMIZATION
 
-
-float Chaos_Joint_DegenerateRotationLimit = -0.998f;	// Cos(176deg)
-FAutoConsoleVariableRef CVarChaosJointDegenerateRotationLimit(TEXT("p.Chaos.Joint.DegenerateRotationLimit"), Chaos_Joint_DegenerateRotationLimit, TEXT("Cosine of the swing angle that is considered degerenerate (default Cos(176deg))"));
-
-bool bChaos_Joint_EllipticalFix = true;
-FAutoConsoleVariableRef CVarChaosJointEllipticalFix(TEXT("p.Chaos.Joint.EllipticalFix"), bChaos_Joint_EllipticalFix, TEXT("Enable the proper elliptical joint axis/error calculation"));
-
 namespace Chaos
 {
 
@@ -125,6 +119,8 @@ namespace Chaos
 
 		PositionTolerance = SolverSettings.PositionTolerance;
 		AngleTolerance = SolverSettings.AngleTolerance;
+		ProjectionInvMassScale = SolverSettings.ProjectionInvMassScale;
+		VelProjectionInvMassScale = SolverSettings.VelProjectionInvMassScale;
 
 		InitDerivedState();
 	}
@@ -230,10 +226,42 @@ namespace Chaos
 		const FPBDJointSettings& JointSettings)
 	{
 		int32 NumActive = 0;
+		const FReal LinearProjection = FPBDJointUtilities::GetLinearProjection(SolverSettings, JointSettings);
+		const FReal AngularProjection = FPBDJointUtilities::GetAngularProjection(SolverSettings, JointSettings);
+		if ((LinearProjection == 0.0f) && (AngularProjection == 0.0f))
+		{
+			return NumActive;
+		}
+
+		const TVector<EJointMotionType, 3>& LinearMotion = JointSettings.LinearMotionTypes;
+		const EJointMotionType TwistMotion = JointSettings.AngularMotionTypes[(int32)EJointAngularConstraintIndex::Twist];
+		const EJointMotionType Swing1Motion = JointSettings.AngularMotionTypes[(int32)EJointAngularConstraintIndex::Swing1];
+		const EJointMotionType Swing2Motion = JointSettings.AngularMotionTypes[(int32)EJointAngularConstraintIndex::Swing2];
+		const bool bTwistSoft = JointSettings.bSoftTwistLimitsEnabled;
+		const bool bSwingSoft = JointSettings.bSoftSwingLimitsEnabled;
 
-		//NumActive += ApplyRotationProjection(Dt, SolverSettings, JointSettings);
+		const bool bPointProject = 
+			(LinearProjection > 0.0f) 
+			&& (LinearMotion[0] == EJointMotionType::Locked)
+			&& (LinearMotion[1] == EJointMotionType::Locked)
+			&& (LinearMotion[2] == EJointMotionType::Locked);
+		bool bConeProject = 
+			SolverSettings.bEnableSwingLimits 
+			&& (AngularProjection > 0.0f) 
+			&& !bSwingSoft 
+			&& (Swing1Motion == EJointMotionType::Limited) 
+			&& (Swing2Motion == EJointMotionType::Limited);
 
-		NumActive += ApplyPositionProjections(Dt, SolverSettings, JointSettings);
+		if (bPointProject && bConeProject)
+		{
+			// Fixed point constraint with cone limits
+			NumActive += ApplyPointConeProjection(Dt, SolverSettings, JointSettings);
+		}
+		else if (bPointProject)
+		{
+			// Fixed point constraint with no angular limits
+			NumActive += ApplyPointProjection(Dt, SolverSettings, JointSettings);
+		}
 
 		return NumActive;
 	}
@@ -399,50 +427,6 @@ namespace Chaos
 	}
 
 
-	int32 FJointSolverGaussSeidel::ApplyRotationProjections(
-		const FReal Dt,
-		const FPBDJointSolverSettings& SolverSettings,
-		const FPBDJointSettings& JointSettings)
-	{
-		int32 NumActive = 0;
-
-		EJointMotionType TwistMotion = JointSettings.AngularMotionTypes[(int32)EJointAngularConstraintIndex::Twist];
-		EJointMotionType Swing1Motion = JointSettings.AngularMotionTypes[(int32)EJointAngularConstraintIndex::Swing1];
-		EJointMotionType Swing2Motion = JointSettings.AngularMotionTypes[(int32)EJointAngularConstraintIndex::Swing2];
-		bool bTwistSoft = JointSettings.bSoftTwistLimitsEnabled;
-		bool bSwingSoft = JointSettings.bSoftSwingLimitsEnabled;
-
-		if (SolverSettings.bEnableTwistLimits)
-		{
-			if ((TwistMotion == EJointMotionType::Locked) || ((TwistMotion == EJointMotionType::Limited) && !bTwistSoft))
-			{
-				NumActive += ApplyTwistProjection(Dt, SolverSettings, JointSettings);
-			}
-		}
-
-		if (SolverSettings.bEnableSwingLimits)
-		{
-			if ((Swing1Motion == EJointMotionType::Limited) && (Swing2Motion == EJointMotionType::Limited) && !bSwingSoft)
-			{
-				NumActive += ApplyConeProjection(Dt, SolverSettings, JointSettings);
-			}
-			else
-			{
-				if ((Swing1Motion == EJointMotionType::Locked) || ((Swing1Motion == EJointMotionType::Limited) && !bSwingSoft))
-				{
-					NumActive += ApplySwingProjection(Dt, SolverSettings, JointSettings, EJointAngularConstraintIndex::Swing1);
-				}
-				if ((Swing2Motion == EJointMotionType::Locked) || ((Swing2Motion == EJointMotionType::Limited) && !bSwingSoft))
-				{
-					NumActive += ApplySwingProjection(Dt, SolverSettings, JointSettings, EJointAngularConstraintIndex::Swing2);
-				}
-			}
-		}
-
-		return NumActive;
-	}
-
-
 	int32 FJointSolverGaussSeidel::ApplyPositionConstraints(
 		const FReal Dt,
 		const FPBDJointSolverSettings& SolverSettings,
@@ -590,39 +574,6 @@ namespace Chaos
 	}
 
 
-	int32 FJointSolverGaussSeidel::ApplyPositionProjections(
-		const FReal Dt,
-		const FPBDJointSolverSettings& SolverSettings,
-		const FPBDJointSettings& JointSettings)
-	{
-		int32 NumActive = 0;
-
-		const TVector<EJointMotionType, 3>& LinearMotion = JointSettings.LinearMotionTypes;
-		const TVector<bool, 3> bLinearLocked =
-		{
-			(LinearMotion[0] == EJointMotionType::Locked),
-			(LinearMotion[1] == EJointMotionType::Locked),
-			(LinearMotion[2] == EJointMotionType::Locked),
-		};
-		const TVector<bool, 3> bLinearLimted =
-		{
-			(LinearMotion[0] == EJointMotionType::Limited),
-			(LinearMotion[1] == EJointMotionType::Limited),
-			(LinearMotion[2] == EJointMotionType::Limited),
-		};
-
-		if (bLinearLocked[0] && bLinearLocked[1] && bLinearLocked[2])
-		{
-			// Hard point constraint (most common case)
-			NumActive += ApplyPointPositionProjection(Dt, SolverSettings, JointSettings);
-		}
-
-		// @todo(ccaulfield): support projection for other linear constraint configs
-
-		return NumActive;
-	}
-
-
 	//
 	//
 	//////////////////////////////////////////////////////////////////////////
@@ -638,6 +589,8 @@ namespace Chaos
 		//UE_LOG(LogChaosJoint, VeryVerbose, TEXT("      Apply DP%d %f %f %f"), BodyIndex, DP.X, DP.Y, DP.Z);
 
 		Ps[BodyIndex] += Stiffness * DP;
+
+		Xs[BodyIndex] += Stiffness * DP;
 	}
 
 
@@ -651,6 +604,9 @@ namespace Chaos
 
 		Ps[0] += Stiffness * DP0;
 		Ps[1] += Stiffness * DP1;
+
+		Xs[0] += Stiffness * DP0;
+		Xs[1] += Stiffness * DP1;
 	}
 
 
@@ -751,10 +707,8 @@ namespace Chaos
 	{
 		const FVec3 AngularAxis0 = FVec3::CrossProduct(Xs[0] - Ps[0], Axis);
 		const FVec3 AngularAxis1 = FVec3::CrossProduct(Xs[1] - Ps[1], Axis);
-		const FMatrix33 InvI0 = Utilities::ComputeWorldSpaceInertia(Qs[0], InvILs[0]);
-		const FMatrix33 InvI1 = Utilities::ComputeWorldSpaceInertia(Qs[1], InvILs[1]);
-		const FVec3 IA0 = Utilities::Multiply(InvI0, AngularAxis0);
-		const FVec3 IA1 = Utilities::Multiply(InvI1, AngularAxis1);
+		const FVec3 IA0 = Utilities::Multiply(InvIs[0], AngularAxis0);
+		const FVec3 IA1 = Utilities::Multiply(InvIs[1], AngularAxis1);
 
 		// Joint-space inverse mass
 		const FReal II0 = FVec3::DotProduct(AngularAxis0, IA0);
@@ -766,8 +720,8 @@ namespace Chaos
 		// Apply constraint correction
 		const FVec3 DP0 = InvMs[0] * DX;
 		const FVec3 DP1 = -InvMs[1] * DX;
-		const FVec3 DR0 = Utilities::Multiply(InvI0, FVec3::CrossProduct(Xs[0] - Ps[0], DX));
-		const FVec3 DR1 = Utilities::Multiply(InvI1, FVec3::CrossProduct(Xs[1] - Ps[1], -DX));
+		const FVec3 DR0 = Utilities::Multiply(InvIs[0], FVec3::CrossProduct(Xs[0] - Ps[0], DX));
+		const FVec3 DR1 = Utilities::Multiply(InvIs[1], FVec3::CrossProduct(Xs[1] - Ps[1], -DX));
 
 		ApplyPositionDelta(Stiffness, DP0, DP1);
 		ApplyRotationDelta(Stiffness, DR0, DR1);
@@ -791,15 +745,11 @@ namespace Chaos
 		}
 		else
 		{
-			// World-space inverse mass
-			const FMatrix33 InvI0 = Utilities::ComputeWorldSpaceInertia(Qs[0], InvILs[0]);
-			const FMatrix33 InvI1 = Utilities::ComputeWorldSpaceInertia(Qs[1], InvILs[1]);
-	
 			// Joint-space inverse mass
 			const FVec3 AngularAxis0 = FVec3::CrossProduct(Xs[0] - Ps[0], Axis);
 			const FVec3 AngularAxis1 = FVec3::CrossProduct(Xs[1] - Ps[1], Axis);
-			const FVec3 IA0 = Utilities::Multiply(InvI0, AngularAxis0);
-			const FVec3 IA1 = Utilities::Multiply(InvI1, AngularAxis1);
+			const FVec3 IA0 = Utilities::Multiply(InvIs[0], AngularAxis0);
+			const FVec3 IA1 = Utilities::Multiply(InvIs[1], AngularAxis1);
 			const FReal II0 = FVec3::DotProduct(AngularAxis0, IA0);
 			const FReal II1 = FVec3::DotProduct(AngularAxis1, IA1);
 			const FReal II = (InvMs[0] + II0 + InvMs[1] + II1);
@@ -820,8 +770,8 @@ namespace Chaos
 	
 			const FVec3 DP0 = (InvMs[0] * DLambda) * Axis;
 			const FVec3 DP1 = (-InvMs[1] * DLambda) * Axis;
-			const FVec3 DR0 = DLambda * Utilities::Multiply(InvI0, AngularAxis0);
-			const FVec3 DR1 = -DLambda * Utilities::Multiply(InvI1, AngularAxis1);
+			const FVec3 DR0 = DLambda * Utilities::Multiply(InvIs[0], AngularAxis0);
+			const FVec3 DR1 = -DLambda * Utilities::Multiply(InvIs[1], AngularAxis1);
 	
 			Lambda += DLambda;
 			ApplyPositionDelta((FReal)1, DP0, DP1);
@@ -829,17 +779,27 @@ namespace Chaos
 		}
 	}
 	
-	
-	void FJointSolverGaussSeidel::ApplyRotationConstraint(
+
+	void FJointSolverGaussSeidel::ApplyRotationConstraintKD(
+		const int32 KIndex,
+		const int32 DIndex,
+		const FReal Stiffness,
+		const FVec3& Axis,
+		const FReal Angle)
+	{
+		const FVec3 DR1 = -Angle * Axis;
+		ApplyRotationDelta(DIndex, Stiffness, DR1);
+	}
+
+
+	void FJointSolverGaussSeidel::ApplyRotationConstraintDD(
 		const FReal Stiffness,
 		const FVec3& Axis,
 		const FReal Angle)
 	{
 		// World-space inverse mass
-		const FMatrix33 InvI0 = Utilities::ComputeWorldSpaceInertia(Qs[0], InvILs[0]);
-		const FMatrix33 InvI1 = Utilities::ComputeWorldSpaceInertia(Qs[1], InvILs[1]);
-		const FVec3 IA0 = Utilities::Multiply(InvI0, Axis);
-		const FVec3 IA1 = Utilities::Multiply(InvI1, Axis);
+		const FVec3 IA0 = Utilities::Multiply(InvIs[0], Axis);
+		const FVec3 IA1 = Utilities::Multiply(InvIs[1], Axis);
 
 		// Joint-space inverse mass
 		const FReal II0 = FVec3::DotProduct(Axis, IA0);
@@ -854,6 +814,26 @@ namespace Chaos
 	}
 
 
+	void FJointSolverGaussSeidel::ApplyRotationConstraint(
+		const FReal Stiffness,
+		const FVec3& Axis,
+		const FReal Angle)
+	{
+		if (InvMs[0] == 0)
+		{
+			ApplyRotationConstraintKD(0, 1, Stiffness, Axis, Angle);
+		}
+		else if (InvMs[1] == 0)
+		{
+			ApplyRotationConstraintKD(1, 0, Stiffness, Axis, -Angle);
+		}
+		else
+		{
+			ApplyRotationConstraintDD(Stiffness, Axis, Angle);
+		}
+	}
+
+
 	// See "XPBD: Position-Based Simulation of Compliant Constrained Dynamics"
 	void FJointSolverGaussSeidel::ApplyRotationConstraintSoftKD(
 		const int32 KIndex,
@@ -975,7 +955,7 @@ namespace Chaos
 		}
 		else if (InvMs[1] == 0)
 		{
-			ApplyRotationConstraintSoftKD(1, 0, Dt, Stiffness, Damping, bAccelerationMode, Axis, Angle, Lambda);
+			ApplyRotationConstraintSoftKD(1, 0, Dt, Stiffness, Damping, bAccelerationMode, Axis, -Angle, Lambda);
 		}
 		else
 		{
@@ -1196,7 +1176,67 @@ namespace Chaos
 		const FPBDJointSolverSettings& SolverSettings,
 		const FPBDJointSettings& JointSettings)
 	{
-		return 0;
+		const FReal Stiffness = FPBDJointUtilities::GetAngularProjection(SolverSettings, JointSettings);
+		if (Stiffness == 0.0f)
+		{
+			return 0;
+		}
+
+		const FReal Swing1Limit = JointSettings.AngularLimits[(int32)EJointAngularConstraintIndex::Swing1];
+		const FReal Swing2Limit = JointSettings.AngularLimits[(int32)EJointAngularConstraintIndex::Swing2];
+
+		// Fix the rotation - treat parent as infinite mass
+		{
+			// Calculate the swing error and axis
+			FVec3 SwingAxisLocal;
+			FReal DSwingAngle = 0.0f;
+			FPBDJointUtilities::GetEllipticalConeAxisErrorLocal(Rs[0], Rs[1], Swing2Limit, Swing1Limit, SwingAxisLocal, DSwingAngle);
+			const FVec3 SwingAxis = Rs[0] * SwingAxisLocal;
+
+			if (DSwingAngle < AngleTolerance)
+			{
+				return 0;
+			}
+
+			ApplyRotationConstraintKD(0, 1, Stiffness, SwingAxis, DSwingAngle);
+		}
+
+		// Fix the relative velocity - this time use the actual inertias rather than treating the parent as infinite mass
+		{
+			// Calculate the swing error and axis
+			FVec3 SwingAxisLocal;
+			FReal DSwingAngle = 0.0f;
+			FPBDJointUtilities::GetEllipticalConeAxisErrorLocal(Rs[0], Rs[1], Swing2Limit, Swing1Limit, SwingAxisLocal, DSwingAngle);
+			const FVec3 SwingAxis = Rs[0] * SwingAxisLocal;
+
+			const FReal W0 = FVec3::DotProduct(Ws[0], SwingAxis);
+			const FReal W1 = FVec3::DotProduct(Ws[1], SwingAxis);
+			const FVec3 DW = (W1 - W0) * SwingAxis;
+
+			FReal II0 = 0.0f;
+			FReal II1 = 0.0f;
+			if (InvMs[0] > 0.0f)
+			{
+				II0 = FVec3::DotProduct(SwingAxis, Utilities::Multiply(InvIs[0], SwingAxis));
+			}
+			if (InvMs[1] > 0.0f)
+			{
+				II1 = FVec3::DotProduct(SwingAxis, Utilities::Multiply(InvIs[1], SwingAxis));
+			}
+
+			if (InvMs[0] > 0.0f)
+			{
+				const FVec3 DW0 = (II0 / (II0 + II1)) * DW;
+				ApplyVelocityDelta(0, Stiffness, FVec3(0), DW0);
+			}
+			if (InvMs[1] > 0.0f)
+			{
+				const FVec3 DW1 = -(II1 / (II0 + II1)) * DW;
+				ApplyVelocityDelta(1, Stiffness, FVec3(0), DW1);
+			}
+		}
+
+		return 1;
 	}
 
 
@@ -1689,46 +1729,69 @@ namespace Chaos
 	}
 
 
-	int32 FJointSolverGaussSeidel::ApplyPointPositionProjection(
+	// Projection for a fixed point constraint and no rotation constraint
+	int32 FJointSolverGaussSeidel::ApplyPointProjection(
 		const FReal Dt,
 		const FPBDJointSolverSettings& SolverSettings,
 		const FPBDJointSettings& JointSettings)
 	{
 		FReal LinearProjection = FPBDJointUtilities::GetLinearProjection(SolverSettings, JointSettings);
-		if (LinearProjection > 0)
+		const FVec3 CX = Xs[1] - Xs[0];
+		bool bApplyProjection = (CX.Size() > KINDA_SMALL_NUMBER); // + PositionTolerance);
+		if (bApplyProjection)
 		{
-			const FVec3 CX = Xs[1] - Xs[0];
-			bool bApplyProjection = (CX.Size() > KINDA_SMALL_NUMBER); // + PositionTolerance);
-			if (bApplyProjection && bChaos_Joint_ISPC_Enabled)
+			if (bChaos_Joint_ISPC_Enabled)
 			{
 #if INTEL_ISPC
 				ispc::ApplyPointPositionProjection((ispc::FJointSolverGaussSeidel*)this, (ispc::FVector&)CX, LinearProjection);
 #endif
 			}
-			else if (bApplyProjection)
+			else
 			{
-				// Apply a position correction with the parent body set to infinite mass
+				// Apply a position correction
 				{
-					const FMatrix33 J = Utilities::ComputeJointFactorMatrix(Xs[1] - Ps[1], InvIs[1], InvMs[1]);
+					const FReal InvM0 = ProjectionInvMassScale * InvMs[0];
+					const FMatrix33 InvI0 = ProjectionInvMassScale * InvIs[0];
+
+					FMatrix33 J = FMatrix33(0, 0, 0);
+					if (InvM0 > SMALL_NUMBER)
+					{
+						J += Utilities::ComputeJointFactorMatrix(Xs[0] - Ps[0], InvI0, InvM0);
+					}
+					if (InvMs[1] > SMALL_NUMBER)
+					{
+						J += Utilities::ComputeJointFactorMatrix(Xs[1] - Ps[1], InvIs[1], InvMs[1]);
+					}
 					const FMatrix33 IJ = J.Inverse();
 					const FVec3 DX = Utilities::Multiply(IJ, CX);
-					const FVec3 DP1 = -InvMs[1] * DX;
-					const FVec3 DR1 = Utilities::Multiply(InvIs[1], FVec3::CrossProduct(Xs[1] - Ps[1], -DX));
-					ApplyDelta(1, LinearProjection, DP1, DR1);
+
+					if (InvM0 > SMALL_NUMBER)
+					{
+						const FVec3 DP0 = InvM0 * DX;
+						const FVec3 DR0 = Utilities::Multiply(InvI0, FVec3::CrossProduct(Xs[0] - Ps[0], DX));
+						ApplyDelta(0, LinearProjection, DP0, DR0);
+					}
+
+					if (InvMs[1] > SMALL_NUMBER)
+					{
+						const FVec3 DP1 = -InvMs[1] * DX;
+						const FVec3 DR1 = Utilities::Multiply(InvIs[1], FVec3::CrossProduct(Xs[1] - Ps[1], -DX));
+						ApplyDelta(1, LinearProjection, DP1, DR1);
+					}
 				}
 
 				// Apply a velocity correction to set the relative velocity at the joint to zero
-				// NOTE: Velocity correction must come after position correction and derived state must be updated accordingly.
-				// NOTE: No parent mass scale here - velocity correction is shared between bodies normally
-				// (if we don't do this, velocity tends to accumulate down chains leading to whip-like behaviour)
 				{
+					const FReal InvM0 = VelProjectionInvMassScale * InvMs[0];
+					const FMatrix33 InvI0 = VelProjectionInvMassScale * InvIs[0];
+
 					// Calculate joint-space mass
 					FMatrix33 J = FMatrix33(0, 0, 0);
-					if (InvMs[0] > 0)
+					if (InvM0 > SMALL_NUMBER)
 					{
-						J = J + Utilities::ComputeJointFactorMatrix(Xs[0] - Ps[0], InvIs[0], InvMs[0]);
+						J = J + Utilities::ComputeJointFactorMatrix(Xs[0] - Ps[0], InvI0, InvM0);
 					}
-					if (InvMs[1] > 0)
+					if (InvMs[1] > SMALL_NUMBER)
 					{
 						J = J + Utilities::ComputeJointFactorMatrix(Xs[1] - Ps[1], InvIs[1], InvMs[1]);
 					}
@@ -1737,16 +1800,15 @@ namespace Chaos
 					// Velocity correction
 					const FVec3 V0 = Vs[0] + FVec3::CrossProduct(Ws[0], Xs[0] - Ps[0]);
 					const FVec3 V1 = Vs[1] + FVec3::CrossProduct(Ws[1], Xs[1] - Ps[1]);
-					const FVec3 CV = V1 - V0;
-					const FVec3 DV = Utilities::Multiply(IJ, CV);
+					const FVec3 DV = Utilities::Multiply(IJ, V1 - V0);
 
-					if (InvMs[0] > 0)
+					if (InvM0 > SMALL_NUMBER)
 					{
-						const FVec3 DV0 = InvMs[0] * DV;
-						const FVec3 DW0 = Utilities::Multiply(InvIs[0], FVec3::CrossProduct(Xs[0] - Ps[0], DV));
+						const FVec3 DV0 = InvM0 * DV;
+						const FVec3 DW0 = Utilities::Multiply(InvI0, FVec3::CrossProduct(Xs[0] - Ps[0], DV));
 						ApplyVelocityDelta(0, LinearProjection, DV0, DW0);
 					}
-					if (InvMs[1] > 0)
+					if (InvMs[1] > SMALL_NUMBER)
 					{
 						const FVec3 DV1 = -InvMs[1] * DV;
 						const FVec3 DW1 = Utilities::Multiply(InvIs[1], FVec3::CrossProduct(Xs[1] - Ps[1], -DV));
@@ -1754,8 +1816,241 @@ namespace Chaos
 					}
 				}
 			}
-			return bApplyProjection ? 1 : 0;
+			return 1;
 		}
 		return 0;
 	}
+
+	// Projection for fixed point constraint and cone limited angular constraint
+	int32 FJointSolverGaussSeidel::ApplyPointConeProjection(
+		const FReal Dt,
+		const FPBDJointSolverSettings& SolverSettings,
+		const FPBDJointSettings& JointSettings)
+	{
+		FReal LinearProjection = FPBDJointUtilities::GetLinearProjection(SolverSettings, JointSettings);
+		const FReal AngularProjection = FPBDJointUtilities::GetAngularProjection(SolverSettings, JointSettings);
+
+		int32 NumApplied = 0;
+		bool bConeActive = false;
+
+		// Correct the position error
+		const FVec3 CX = Xs[1] - Xs[0];
+		if (CX.Size() > PositionTolerance)
+		{
+			const FReal InvM0 = ProjectionInvMassScale * InvMs[0];
+			const FMatrix33 InvI0 = ProjectionInvMassScale * InvIs[0];
+
+			FMatrix33 J = FMatrix33(0, 0, 0);
+			if (InvM0 > SMALL_NUMBER)
+			{
+				J += Utilities::ComputeJointFactorMatrix(Xs[0] - Ps[0], InvI0, InvM0);
+			}
+			if (InvMs[1] > SMALL_NUMBER)
+			{
+				J += Utilities::ComputeJointFactorMatrix(Xs[1] - Ps[1], InvIs[1], InvMs[1]);
+			}
+			const FMatrix33 IJ = J.Inverse();
+			const FVec3 DX = Utilities::Multiply(IJ, CX);
+
+			if (InvM0 > SMALL_NUMBER)
+			{
+				const FVec3 DP0 = InvM0 * DX;
+				const FVec3 DR0 = Utilities::Multiply(InvI0, FVec3::CrossProduct(Xs[0] - Ps[0], DX));
+				ApplyDelta(0, LinearProjection, DP0, DR0);
+			}
+			if (InvMs[1] > SMALL_NUMBER)
+			{
+				const FVec3 DP1 = -InvMs[1] * DX;
+				const FVec3 DR1 = Utilities::Multiply(InvIs[1], FVec3::CrossProduct(Xs[1] - Ps[1], -DX));
+				ApplyDelta(1, LinearProjection, DP1, DR1);
+			}
+
+			++NumApplied;
+		}
+
+		// Correct the cone error, maintaining position constraint
+		const FReal Swing1Limit = JointSettings.AngularLimits[(int32)EJointAngularConstraintIndex::Swing1];
+		const FReal Swing2Limit = JointSettings.AngularLimits[(int32)EJointAngularConstraintIndex::Swing2];
+		FVec3 SwingAxis = FVec3(0);
+		{
+			const FReal InvM0 = ProjectionInvMassScale * InvMs[0];
+			const FMatrix33 InvI0 = ProjectionInvMassScale * InvIs[0];
+			FVec3 SwingAxisLocal;
+			FReal DSwingAngle = 0.0f;
+			FPBDJointUtilities::GetEllipticalConeAxisErrorLocal(Rs[0], Rs[1], Swing2Limit, Swing1Limit, SwingAxisLocal, DSwingAngle);
+			SwingAxis = Rs[0] * SwingAxisLocal;
+			if (DSwingAngle > AngleTolerance)
+			{
+				if (InvM0 > SMALL_NUMBER)
+				{
+					const FVec3 DR0 = DSwingAngle * SwingAxis;
+					const FVec3 DP0 = FVec3::CrossProduct(DR0, Xs[0] - Ps[0]);
+					ApplyDelta(0, AngularProjection, DP0, DR0);
+				}
+				if (InvMs[1] > SMALL_NUMBER)
+				{
+					const FVec3 DR1 = -DSwingAngle * SwingAxis;
+					const FVec3 DP1 = -FVec3::CrossProduct(DR1, Xs[1] - Ps[1]);
+					ApplyDelta(1, AngularProjection, DP1, DR1);
+				}
+
+				++NumApplied;
+				bConeActive = true;
+			}
+		}
+
+		static bool b4x4 = false;
+
+		// Velocity correction
+		// NON-ITERATIVE VERSION
+		if (b4x4)
+		{
+			const FReal InvM0 = VelProjectionInvMassScale * InvMs[0];
+			const FMatrix33 InvI0 = VelProjectionInvMassScale * InvIs[0];
+
+			if (InvM0 > SMALL_NUMBER)
+			{
+				// If the parent body moved, we need to recalculate the swing axis
+				FVec3 SwingAxisLocal;
+				FReal DSwingAngle = 0.0f;
+				FPBDJointUtilities::GetEllipticalConeAxisErrorLocal(Rs[0], Rs[1], Swing2Limit, Swing1Limit, SwingAxisLocal, DSwingAngle);
+				SwingAxis = Rs[0] * SwingAxisLocal;
+			}
+
+			// Calculate velocity and angular velocity error at joint
+			const FVec3 V0 = Vs[0] + FVec3::CrossProduct(Ws[0], Xs[0] - Ps[0]);
+			const FVec3 V1 = Vs[1] + FVec3::CrossProduct(Ws[1], Xs[1] - Ps[1]);
+			const FReal W0 = FVec3::DotProduct(Ws[0], SwingAxis);
+			const FReal W1 = FVec3::DotProduct(Ws[1], SwingAxis);
+			const FVec3 CV = V1 - V0;
+			const FReal CW = FMath::Max(0.0f, W1 - W0);
+			const TVector<FReal, 4> C = TVector<FReal, 4>(CV.X, CV.Y, CV.Z, CW);
+
+			// Calculate joint-space mass from the Jacobian
+			FMatrix33 J00 = FMatrix33(0, 0, 0);
+			FVec3 J01 = FVec3(0);
+			FReal J11 = 1.0f;
+			if (InvM0 > SMALL_NUMBER)
+			{
+				J00 += Utilities::ComputeJointFactorMatrix(Xs[0] - Ps[0], InvI0, InvM0);
+				if (bConeActive && (CW > 0.0f))
+				{
+					J01 = FVec3::CrossProduct(Xs[0] - Ps[0], Utilities::Multiply(InvI0, SwingAxis));
+					J11 = FVec3::DotProduct(SwingAxis, Utilities::Multiply(InvI0, SwingAxis));
+				}
+			}
+			if (InvMs[1] > SMALL_NUMBER)
+			{
+				J00 += Utilities::ComputeJointFactorMatrix(Xs[1] - Ps[1], InvIs[1], InvMs[1]);
+				if (bConeActive && (CW > 0.0f))
+				{
+					J01 = FVec3::CrossProduct(Xs[1] - Ps[1], Utilities::Multiply(InvIs[1], SwingAxis));
+					J11 = FVec3::DotProduct(SwingAxis, Utilities::Multiply(InvIs[1], SwingAxis));
+				}
+			}
+			PMatrix<FReal, 4, 4>& J = reinterpret_cast<PMatrix<FReal, 4, 4>&>(J00);
+			J.M[3][0] = J01.X;
+			J.M[3][1] = J01.Y;
+			J.M[3][2] = J01.Z;
+			J.M[0][3] = -J01.X;
+			J.M[1][3] = -J01.Y;
+			J.M[2][3] = -J01.Z;
+			J.M[3][3] = J11;
+			PMatrix<FReal, 4, 4> IJ = J.Inverse();
+
+			// Calculate the joint-space velocity and angular velocity correction
+			const TVector<FReal, 4> D = IJ.TransformFVector4(C);
+			const FVec3 DV = FVec3(D.X, D.Y, D.Z);
+			const FReal DW = D.W;
+
+			// Apply the world-space velocity and angular velocity correction
+			if (InvM0 > SMALL_NUMBER)
+			{
+				FVec3 DV0 = InvM0 * DV;
+				FVec3 DW0 = FVec3::CrossProduct(Xs[0] - Ps[0], Utilities::Multiply(InvI0, DV)) - Utilities::Multiply(InvI0, SwingAxis) * DW;
+				ApplyVelocityDelta(0, 1.0f, DV0, DW0);
+			}
+			if (InvMs[1] > SMALL_NUMBER)
+			{
+				FVec3 DV1 = -InvMs[1] * DV;
+				FVec3 DW1 = FVec3::CrossProduct(Xs[1] - Ps[1], Utilities::Multiply(InvIs[1], DV)) - Utilities::Multiply(InvIs[1], SwingAxis) * DW;
+				ApplyVelocityDelta(1, 1.0f, DV1, DW1);
+			}
+		}
+
+		// Velocity correction
+		// ITERATIVE VERSION
+		if (!b4x4)
+		{
+			const FReal InvM0 = VelProjectionInvMassScale * InvMs[0];
+			const FMatrix33 InvI0 = VelProjectionInvMassScale * InvIs[0];
+
+			FMatrix33 J = FMatrix33(0, 0, 0);
+			if (InvM0 > 0)
+			{
+				J = J + Utilities::ComputeJointFactorMatrix(Xs[0] - Ps[0], InvI0, InvM0);
+			}
+			if (InvMs[1] > 0)
+			{
+				J = J + Utilities::ComputeJointFactorMatrix(Xs[1] - Ps[1], InvIs[1], InvMs[1]);
+			}
+			const FMatrix33 IJ = J.Inverse();
+
+			FVec3 IA0 = FVec3(0);
+			FVec3 IA1 = FVec3(0);
+			FReal II0 = 0.0f;
+			FReal II1 = 0.0f;
+			if (bConeActive)
+			{
+				if (InvM0 > 0.0f)
+				{
+					II0 = FVec3::DotProduct(SwingAxis, Utilities::Multiply(InvI0, SwingAxis));
+				}
+				if (InvMs[1] > 0.0f)
+				{
+					II1 = FVec3::DotProduct(SwingAxis, Utilities::Multiply(InvIs[1], SwingAxis));
+				}
+			}
+
+			static int32 NumVelIts = 10;
+			for (int32 VelIt = 0; VelIt < NumVelIts; ++VelIt)
+			{
+				const FVec3 V0 = Vs[0] + FVec3::CrossProduct(Ws[0], Xs[0] - Ps[0]);
+				const FVec3 V1 = Vs[1] + FVec3::CrossProduct(Ws[1], Xs[1] - Ps[1]);
+				const FReal W0 = FVec3::DotProduct(Ws[0], SwingAxis);
+				const FReal W1 = FVec3::DotProduct(Ws[1], SwingAxis);
+				const FVec3 DV = Utilities::Multiply(IJ, V1 - V0);
+				const FVec3 DW = (W1 - W0) * SwingAxis;
+
+				if (InvM0 > 0)
+				{
+					FVec3 DV0 = InvM0 * DV;
+					FVec3 DW0 = Utilities::Multiply(InvI0, FVec3::CrossProduct(Xs[0] - Ps[0], DV));
+
+					if (bConeActive && ((W1 - W0) > 0.0f))
+					{
+						DW0 += (II0 / (II0 + II1)) * DW;
+						DV0 += FVec3(0);// -FVec3::CrossProduct(DW0, Xs[0] - Ps[0]);
+					}
+
+					ApplyVelocityDelta(0, 1.0f, DV0, DW0);
+				}
+				if (InvMs[1] > 0)
+				{
+					FVec3 DV1 = -InvMs[1] * DV;
+					FVec3 DW1 = FVec3::CrossProduct(Xs[1] - Ps[1], Utilities::Multiply(InvIs[1], -DV));
+
+					if (bConeActive && ((W1 - W0) > 0.0f))
+					{
+						DW1 += -(II1 / (II0 + II1)) * DW;
+						DV1 += FVec3(0);// -FVec3::CrossProduct(DW0, Xs[0] - Ps[0]);
+					}
+
+					ApplyVelocityDelta(1, 1.0f, DV1, DW1);
+				}
+			}
+		}
+
+		return NumApplied;
+	}
 }
diff --git a/Engine/Source/Runtime/Experimental/Chaos/Private/Chaos/Joint/PBDJointSolverGaussSeidel.ispc b/Engine/Source/Runtime/Experimental/Chaos/Private/Chaos/Joint/PBDJointSolverGaussSeidel.ispc
--- a/Engine/Source/Runtime/Experimental/Chaos/Private/Chaos/Joint/PBDJointSolverGaussSeidel.ispc
+++ b/Engine/Source/Runtime/Experimental/Chaos/Private/Chaos/Joint/PBDJointSolverGaussSeidel.ispc
@@ -46,6 +46,9 @@ struct FJointSolverGaussSeidel
 	
 	float PositionTolerance;					// Distance error below which we consider a constraint or drive solved
 	float AngleTolerance;						// Angle error below which we consider a constraint or drive solved
+
+	float ProjectionInvMassScale;
+	float VelProjectionInvMassScale;
 };
 
 export uniform int32 SizeofFJointSolverGaussSeidel()
@@ -411,23 +414,46 @@ export void ApplyPointPositionProjection(
 {
 	// Apply a position correction with the parent body set to infinite mass
 	{
-		const uniform FMatrix33 J = ComputeJointFactorMatrix(M->Xs[1] - M->Ps[1], M->InvIs[1], M->InvMs[1]);
+		const uniform float InvM0 = M->InvMs[0] * M->ProjectionInvMassScale;
+		const uniform FMatrix33 InvI0 = Multiply(M->InvIs[0], M->ProjectionInvMassScale);
+
+		uniform FMatrix33 J = SetMatrix33(0, 0, 0);
+		if (InvM0 > SMALL_NUMBER)
+		{
+			J = AddAB(J, ComputeJointFactorMatrix(M->Xs[0] - M->Ps[0], InvI0, InvM0));
+		}
+		if (M->InvMs[1] > SMALL_NUMBER)
+		{
+			J = AddAB(J, ComputeJointFactorMatrix(M->Xs[1] - M->Ps[1], M->InvIs[1], M->InvMs[1]));
+		}
 		const uniform FMatrix33 IJ = MatrixInverse(J);
-		const uniform FVector DX = VectorNegate(Multiply(IJ, CX));
-		const uniform FVector DP1 = M->InvMs[1] * DX;
-		const uniform FVector DR1 = Multiply(M->InvIs[1], VectorCross(M->Xs[1] - M->Ps[1], DX));
-		ApplyDelta(M, 1, LinearProjection, DP1, DR1);
+		const uniform FVector DX = Multiply(IJ, CX);
+
+		if (InvM0 > SMALL_NUMBER)
+		{
+			const uniform FVector DP0 = InvM0 * DX;
+			const uniform FVector DR0 = Multiply(InvI0, VectorCross(M->Xs[0] - M->Ps[0], DX));
+			ApplyDelta(M, 0, LinearProjection, DP0, DR0);
+		}
+		if (M->InvMs[1] > SMALL_NUMBER)
+		{
+			const uniform FVector NDX = VectorNegate(DX);
+			const uniform FVector DP1 = M->InvMs[1] * NDX;
+			const uniform FVector DR1 = Multiply(M->InvIs[1], VectorCross(M->Xs[1] - M->Ps[1], NDX));
+			ApplyDelta(M, 1, LinearProjection, DP1, DR1);
+		}
 	}
 
 	// Apply a velocity correction to set the relative velocity at the joint to zero
-	// NOTE: No parent mass scale here - velocity correction is shared between bodies normally
-	// (if we don't do this, velocity tends to accumulate down chains leading to whip-like behaviour)
 	{
+		const uniform float InvM0 = M->VelProjectionInvMassScale * M->InvMs[0];
+		const uniform FMatrix33 InvI0 = Multiply(M->InvIs[0], M->VelProjectionInvMassScale);
+
 		// Calculate joint-space mass
 		uniform FMatrix33 J = SetMatrix33(0, 0, 0);
-		if (M->InvMs[0] > 0)
+		if (InvM0 > 0)
 		{
-			J = AddAB(J, ComputeJointFactorMatrix(M->Xs[0] - M->Ps[0], M->InvIs[0], M->InvMs[0]));
+			J = AddAB(J, ComputeJointFactorMatrix(M->Xs[0] - M->Ps[0], InvI0, InvM0));
 		}
 		if (M->InvMs[1] > 0)
 		{
@@ -441,10 +467,10 @@ export void ApplyPointPositionProjection(
 		const uniform FVector CV = V1 - V0;
 		const uniform FVector DV = Multiply(IJ, CV);
 
-		if (M->InvMs[0] > 0)
+		if (InvM0 > 0)
 		{
-			const uniform FVector DV0 = M->InvMs[0] * DV;
-			const uniform FVector DW0 = Multiply(M->InvIs[0], VectorCross(M->Xs[0] - M->Ps[0], DV));
+			const uniform FVector DV0 = InvM0 * DV;
+			const uniform FVector DW0 = Multiply(InvI0, VectorCross(M->Xs[0] - M->Ps[0], DV));
 			ApplyVelocityDelta(M, 0, LinearProjection, DV0, DW0);
 		}
 		if (M->InvMs[1] > 0)
diff --git a/Engine/Source/Runtime/Experimental/Chaos/Private/Chaos/PBDCollisionConstraints.ispc b/Engine/Source/Runtime/Experimental/Chaos/Private/Chaos/PBDCollisionConstraints.ispc
--- a/Engine/Source/Runtime/Experimental/Chaos/Private/Chaos/PBDCollisionConstraints.ispc
+++ b/Engine/Source/Runtime/Experimental/Chaos/Private/Chaos/PBDCollisionConstraints.ispc
@@ -765,22 +765,6 @@ export uniform int32 SizeofFCollisionContact()
 	return sizeof(uniform FCollisionContact);
 }
 
-// @todo(ccaulfield): move to ispc libs
-inline uniform FMatrix33 Multiply(const uniform FMatrix33& L, const uniform float R)
-{
-	uniform FMatrix33 M;
-	for (uniform int Index = 0; Index < 12; ++Index)
-	{
-		M.M[Index] = L.M[Index] * R;
-	}
-	M.M[12] = 0.0f;
-	M.M[13] = 0.0f;
-	M.M[14] = 0.0f;
-	M.M[15] = 1.0f;
-	return M;
-}
-
-
 // Get the CoM position given the Actor position
 static inline uniform FVector GetCoMWorldPosition(
 	const uniform FVector& PActorWorld,
diff --git a/Engine/Source/Runtime/Experimental/Chaos/Private/Chaos/PBDJointConstraintUtilities.cpp b/Engine/Source/Runtime/Experimental/Chaos/Private/Chaos/PBDJointConstraintUtilities.cpp
--- a/Engine/Source/Runtime/Experimental/Chaos/Private/Chaos/PBDJointConstraintUtilities.cpp
+++ b/Engine/Source/Runtime/Experimental/Chaos/Private/Chaos/PBDJointConstraintUtilities.cpp
@@ -540,14 +540,14 @@ namespace Chaos
 		const FPBDJointSolverSettings& SolverSettings,
 		const FPBDJointSettings& JointSettings)
 	{
-		return (SolverSettings.LinearProjection > 0.0f) ? SolverSettings.LinearProjection : JointSettings.LinearProjection;
+		return (SolverSettings.LinearProjection >= 0.0f) ? SolverSettings.LinearProjection : JointSettings.LinearProjection;
 	}
 
 	FReal FPBDJointUtilities::GetAngularProjection(
 		const FPBDJointSolverSettings& SolverSettings,
 		const FPBDJointSettings& JointSettings)
 	{
-		return (SolverSettings.AngularProjection > 0.0f) ? SolverSettings.AngularProjection : JointSettings.AngularProjection;
+		return (SolverSettings.AngularProjection >= 0.0f) ? SolverSettings.AngularProjection : JointSettings.AngularProjection;
 	}
 
 	bool FPBDJointUtilities::GetLinearSoftAccelerationMode(
diff --git a/Engine/Source/Runtime/Experimental/Chaos/Private/Chaos/PBDJointConstraints.cpp b/Engine/Source/Runtime/Experimental/Chaos/Private/Chaos/PBDJointConstraints.cpp
--- a/Engine/Source/Runtime/Experimental/Chaos/Private/Chaos/PBDJointConstraints.cpp
+++ b/Engine/Source/Runtime/Experimental/Chaos/Private/Chaos/PBDJointConstraints.cpp
@@ -4,6 +4,7 @@
 #include "Chaos/DebugDrawQueue.h"
 #include "Chaos/Joint/ChaosJointLog.h"
 #include "Chaos/Joint/ColoringGraph.h"
+#include "Chaos/Joint/JointConstraintsCVars.h"
 #include "Chaos/Particle/ParticleUtilities.h"
 #include "Chaos/ParticleHandle.h"
 #include "Chaos/PBDJointConstraintUtilities.h"
@@ -19,20 +20,6 @@
 
 //PRAGMA_DISABLE_OPTIMIZATION
 
-bool bChaos_Joint_EarlyOut_Enabled = true;
-FAutoConsoleVariableRef CVarChaosJointEarlyOutEnabled(TEXT("p.Chaos.Joint.EarlyOut"), bChaos_Joint_EarlyOut_Enabled, TEXT("Whether to iterating when joints report being solved"));
-
-bool bChaos_Joint_Batching = false;
-FAutoConsoleVariableRef CVarChaosJointBatching(TEXT("p.Chaos.Joint.Batching"), bChaos_Joint_Batching, TEXT(""));
-
-int32 bChaos_Joint_MaxBatchSize = 1000;
-FAutoConsoleVariableRef CVarChaosJointBatchSize(TEXT("p.Chaos.Joint.MaxBatchSize"), bChaos_Joint_MaxBatchSize, TEXT(""));
-
-#if INTEL_ISPC && !UE_BUILD_SHIPPING
-bool bChaos_Joint_ISPC_Enabled = true;
-FAutoConsoleVariableRef CVarChaosJointISPCEnabled(TEXT("p.Chaos.Joint.ISPC"), bChaos_Joint_ISPC_Enabled, TEXT("Whether to use ISPC optimizations in the Joint Solver"));
-#endif
-
 namespace Chaos
 {
 	DECLARE_CYCLE_STAT(TEXT("Joints::Sort"), STAT_Joints_Sort, STATGROUP_ChaosJoint);
@@ -225,6 +212,8 @@ namespace Chaos
 		, SoftTwistDamping(0)
 		, SoftSwingStiffness(0)
 		, SoftSwingDamping(0)
+		, ProjectionInvMassScale(0)
+		, VelProjectionInvMassScale(0)
 	{
 	}
 
diff --git a/Engine/Source/Runtime/Experimental/Chaos/Public/Chaos/Joint/PBDJointSolverGaussSeidel.h b/Engine/Source/Runtime/Experimental/Chaos/Public/Chaos/Joint/PBDJointSolverGaussSeidel.h
--- a/Engine/Source/Runtime/Experimental/Chaos/Public/Chaos/Joint/PBDJointSolverGaussSeidel.h
+++ b/Engine/Source/Runtime/Experimental/Chaos/Public/Chaos/Joint/PBDJointSolverGaussSeidel.h
@@ -130,16 +130,6 @@ namespace Chaos
 			const FPBDJointSolverSettings& SolverSettings,
 			const FPBDJointSettings& JointSettings);
 
-		int32 ApplyPositionProjections(
-			const FReal Dt,
-			const FPBDJointSolverSettings& SolverSettings,
-			const FPBDJointSettings& JointSettings);
-
-		int32 ApplyRotationProjections(
-			const FReal Dt,
-			const FPBDJointSolverSettings& SolverSettings,
-			const FPBDJointSettings& JointSettings);
-
 
 		void ApplyPositionDelta(
 			const int32 BodyIndex,
@@ -199,6 +189,18 @@ namespace Chaos
 			const FVec3& Axis,
 			const FReal Angle);
 
+		void ApplyRotationConstraintKD(
+			const int32 KIndex,
+			const int32 DIndex,
+			const FReal Stiffness,
+			const FVec3& Axis,
+			const FReal Angle);
+
+		void ApplyRotationConstraintDD(
+			const FReal Stiffness,
+			const FVec3& Axis,
+			const FReal Angle);
+
 		void ApplyRotationConstraintSoft(
 			const FReal Dt,
 			const FReal Stiffness,
@@ -357,7 +359,12 @@ namespace Chaos
 			const FPBDJointSolverSettings& SolverSettings,
 			const FPBDJointSettings& JointSettings);
 
-		int32 ApplyPointPositionProjection(
+		int32 ApplyPointProjection(
+			const FReal Dt,
+			const FPBDJointSolverSettings& SolverSettings,
+			const FPBDJointSettings& JointSettings);
+
+		int32 ApplyPointConeProjection(
 			const FReal Dt,
 			const FPBDJointSolverSettings& SolverSettings,
 			const FPBDJointSettings& JointSettings);
@@ -392,6 +399,9 @@ namespace Chaos
 
 		FReal PositionTolerance;					// Distance error below which we consider a constraint or drive solved
 		FReal AngleTolerance;						// Angle error below which we consider a constraint or drive solved
+
+		FReal ProjectionInvMassScale;
+		FReal VelProjectionInvMassScale;
 	};
 
 }
\ No newline at end of file
diff --git a/Engine/Source/Runtime/Experimental/Chaos/Public/Chaos/PBDJointConstraintTypes.h b/Engine/Source/Runtime/Experimental/Chaos/Public/Chaos/PBDJointConstraintTypes.h
--- a/Engine/Source/Runtime/Experimental/Chaos/Public/Chaos/PBDJointConstraintTypes.h
+++ b/Engine/Source/Runtime/Experimental/Chaos/Public/Chaos/PBDJointConstraintTypes.h
@@ -185,6 +185,8 @@ namespace Chaos
 		FReal SoftTwistDamping;
 		FReal SoftSwingStiffness;
 		FReal SoftSwingDamping;
+		FReal ProjectionInvMassScale;
+		FReal VelProjectionInvMassScale;
 	};
 
 	class FJointSolverResult
diff --git a/Engine/Source/Runtime/Experimental/Chaos/Public/Chaos/PBDJointConstraintUtilities.h b/Engine/Source/Runtime/Experimental/Chaos/Public/Chaos/PBDJointConstraintUtilities.h
--- a/Engine/Source/Runtime/Experimental/Chaos/Public/Chaos/PBDJointConstraintUtilities.h
+++ b/Engine/Source/Runtime/Experimental/Chaos/Public/Chaos/PBDJointConstraintUtilities.h
@@ -251,11 +251,3 @@ namespace Chaos
 	};
 }
 
-// Support ISPC enable/disable in non-shipping builds
-#if !INTEL_ISPC
-const bool bChaos_Joint_ISPC_Enabled = false;
-#elif UE_BUILD_SHIPPING
-const bool bChaos_Joint_ISPC_Enabled = true;
-#else
-extern bool bChaos_Joint_ISPC_Enabled;
-#endif
diff --git a/Engine/Source/Runtime/Experimental/ChaosCore/Public/Chaos/Matrix33.isph b/Engine/Source/Runtime/Experimental/ChaosCore/Public/Chaos/Matrix33.isph
--- a/Engine/Source/Runtime/Experimental/ChaosCore/Public/Chaos/Matrix33.isph
+++ b/Engine/Source/Runtime/Experimental/ChaosCore/Public/Chaos/Matrix33.isph
@@ -92,6 +92,15 @@ inline varying FVector Multiply(const uniform FMatrix33& L, const varying FVecto
 	return SetVector(Result.V[0], Result.V[1], Result.V[2]);
 }
 
+inline uniform FMatrix33 Multiply(const uniform FMatrix33& L, const uniform float R)
+{
+	return SetMatrix33(
+		L.M[0] * R, L.M[1] * R, L.M[2] * R,
+		L.M[4] * R, L.M[5] * R, L.M[6] * R,
+		L.M[8] * R, L.M[9] * R, L.M[10] * R
+	);
+}
+
 inline uniform FMatrix33 AddAB(const uniform FMatrix33& L, const uniform FMatrix33& R)
 {
 	const uniform FVector4 L0 = SetVector4(L.M[0], L.M[1], L.M[2], L.M[3]);
