Chaos, remove cluster array from dynamic view since cluster particles are being integrated twice.

diff --git a/Engine/Source/Runtime/Experimental/Chaos/Public/Chaos/PBDRigidsSOAs.h b/Engine/Source/Runtime/Experimental/Chaos/Public/Chaos/PBDRigidsSOAs.h
--- a/Engine/Source/Runtime/Experimental/Chaos/Public/Chaos/PBDRigidsSOAs.h
+++ b/Engine/Source/Runtime/Experimental/Chaos/Public/Chaos/PBDRigidsSOAs.h
@@ -842,7 +842,7 @@ private:
 			TArray<TSOAView<TPBDRigidParticles<T, d>>> TmpArray = 
 			{ 
 				{&ActiveParticlesArray},
-				{&NonDisabledClusteredArray},
+			//	{&NonDisabledClusteredArray},  Cluster particles appear in the ActiveParticlesArray
 				{&StaticGeometryCollectionArray},
 				{&KinematicGeometryCollectionArray},
 				{&DynamicGeometryCollectionArray}
