Fix cloth hang.

diff --git a/Engine/Source/Runtime/Experimental/Chaos/Private/Chaos/ParticleHandle.cpp b/Engine/Source/Runtime/Experimental/Chaos/Private/Chaos/ParticleHandle.cpp
--- a/Engine/Source/Runtime/Experimental/Chaos/Private/Chaos/ParticleHandle.cpp
+++ b/Engine/Source/Runtime/Experimental/Chaos/Private/Chaos/ParticleHandle.cpp
@@ -113,7 +113,7 @@ namespace Chaos
 		else if (Implicit->GetType() == TImplicitObjectTransformed<T, d>::StaticType())
 		{
 			TImplicitObjectTransformed<FReal, 3>* TransformedImplicit = Implicit->template GetObject<TImplicitObjectTransformed<FReal, 3>>();
-			SetIgnoreAnalyticCollisionsImp(TransformedImplicit, bIgnoreAnalyticCollisions);
+			SetIgnoreAnalyticCollisionsImp(const_cast<FImplicitObject*>(TransformedImplicit->GetTransformedObject()), bIgnoreAnalyticCollisions);
 		}
 		else if ((uint32)Implicit->GetType() & ImplicitObjectType::IsInstanced)
 		{
