Reimplemented solver floor for custom Chaos solver actors

diff --git a/Engine/Source/Runtime/Experimental/ChaosSolverEngine/Private/Chaos/ChaosSolverActor.cpp b/Engine/Source/Runtime/Experimental/ChaosSolverEngine/Private/Chaos/ChaosSolverActor.cpp
--- a/Engine/Source/Runtime/Experimental/ChaosSolverEngine/Private/Chaos/ChaosSolverActor.cpp
+++ b/Engine/Source/Runtime/Experimental/ChaosSolverEngine/Private/Chaos/ChaosSolverActor.cpp
@@ -316,6 +316,20 @@ void AChaosSolverActor::BeginPlay()
 			InSolver->SetPaused(false);
 #endif  // #if CHAOS_WITH_PAUSABLE_SOLVER
 #endif  // #if TODO_REIMPLEMENT_SOLVER_SETTINGS_ACCESSORS
+
+			// Add a floor if specified on the actor
+			if(InHasFloor)
+			{
+				Chaos::TGeometryParticle<float, 3>* Particle = Chaos::TGeometryParticle<float, 3>::CreateParticle().Release();
+				Particle->SetGeometry(TUniquePtr<Chaos::TPlane<float, 3>>(new Chaos::TPlane<float, 3>(FVector(0), FVector(0, 0, 1))));
+				Particle->SetObjectState(Chaos::EObjectStateType::Static);
+				Particle->SetX(Chaos::TVector<float, 3>(0.f, 0.f, InFloorHeight));
+				FCollisionFilterData FilterData;
+				FilterData.Word1 = 0xFFFF;
+				FilterData.Word3 = 0xFFFF;
+				Particle->SetShapeSimData(0, FilterData);
+				InSolver->RegisterObject(Particle);
+			}
 		});
 	}
 
