Make BoundingVolume acceleration structure only rebuild when time sliced structures are done.

diff --git a/Engine/Source/Runtime/Experimental/Chaos/Private/Chaos/PBDRigidsEvolution.cpp b/Engine/Source/Runtime/Experimental/Chaos/Private/Chaos/PBDRigidsEvolution.cpp
--- a/Engine/Source/Runtime/Experimental/Chaos/Private/Chaos/PBDRigidsEvolution.cpp
+++ b/Engine/Source/Runtime/Experimental/Chaos/Private/Chaos/PBDRigidsEvolution.cpp
@@ -85,8 +85,46 @@ namespace Chaos
 			return ConfigSettings.BroadphaseType >= 3 ? 3 : 1;
 		}
 
+		virtual bool IsBucketTimeSliced(uint16 BucketIdx) const
+		{
+			// TODO: Unduplicate switch statement here with CreateAccelerationPerBucket_Threaded and refactor so that bucket index mapping is better.
+			switch (BucketIdx)
+			{
+			case 0:
+			{
+				if (ConfigSettings.BroadphaseType == 0)
+				{
+					// BVType
+					return false;
+				}
+				else if (ConfigSettings.BroadphaseType == 1 || ConfigSettings.BroadphaseType == 3)
+				{
+					// AABBTreeType
+					return true;
+				}
+				else if (ConfigSettings.BroadphaseType == 4 || ConfigSettings.BroadphaseType == 2)
+				{
+					// AABBTreeOfGridsType
+					return true;
+				}
+			}
+			case 1:
+			{
+				// BVType
+				ensure(ConfigSettings.BroadphaseType == 3 || ConfigSettings.BroadphaseType == 4);
+				return false;
+			}
+			default:
+			{
+				check(false);
+				return false;
+			}
+			}
+		}
+
 		virtual TUniquePtr<ISpatialAcceleration<TAccelerationStructureHandle<FReal, 3>, FReal, 3>> CreateAccelerationPerBucket_Threaded(const TConstParticleView<FSpatialAccelerationCache>& Particles, uint16 BucketIdx, bool ForceFullBuild) override
 		{
+			// TODO: Unduplicate switch statement here with IsBucketTimeSliced and refactor so that bucket index mapping is better.
 			switch (BucketIdx)
 			{
 			case 0:
@@ -169,6 +207,7 @@ namespace Chaos
 	DECLARE_CYCLE_STAT(TEXT("SwapAccelerationStructures"), STAT_SwapAccelerationStructures, STATGROUP_Chaos);
 	DECLARE_CYCLE_STAT(TEXT("AccelerationStructureTimeSlice"), STAT_AccelerationStructureTimeSlice, STATGROUP_Chaos);
 	DECLARE_CYCLE_STAT(TEXT("CreateInitialAccelerationStructure"), STAT_CreateInitialAccelerationStructure, STATGROUP_Chaos);
+	DECLARE_CYCLE_STAT(TEXT("CreateNonSlicedStructures"), STAT_CreateNonSlicedStructures, STATGROUP_Chaos);
 
 
 	FPBDRigidsEvolutionBase::FChaosAccelerationStructureTask::FChaosAccelerationStructureTask(
@@ -234,6 +273,8 @@ namespace Chaos
 
 		uint8 ActiveBucketsMask = SpatialCollectionFactory.GetActiveBucketsMask();
 		TArray<TSOAView<FSpatialAccelerationCache>> ViewsPerBucket[8];
+		TArray<uint8> TimeSlicedBucketsToCreate;
+		TArray<uint8> NonTimeSlicedBucketsToCreate;
 
 		bool IsTimeSlicingProgressing = false;
 
@@ -260,11 +301,20 @@ namespace Chaos
 				{
 					AccelerationStructure->RemoveSubstructure(SpatialIdx);
 				}
+
+				if (SpatialCollectionFactory.IsBucketTimeSliced(BucketIdx))
+				{
+					TimeSlicedBucketsToCreate.Add(SpatialIdx.Bucket);
+				}
+				else
+				{
+					NonTimeSlicedBucketsToCreate.Add(SpatialIdx.Bucket);
+				}
 			}
 		}
 
 		//todo: creation can go wide, insertion to collection cannot
-		for (uint8 BucketIdx = 0; BucketIdx < 8; ++BucketIdx)
+		for (uint8 BucketIdx : TimeSlicedBucketsToCreate)
 		{
 			if (ViewsPerBucket[BucketIdx].Num())
 			{
@@ -274,10 +324,7 @@ namespace Chaos
 				auto NewStruct = SpatialCollectionFactory.CreateAccelerationPerBucket_Threaded(ParticleView, BucketIdx, IsForceFullBuild);
 
 				// we kicked of the creation of a new structure and it's going to time-slice the work
-				if (!NewStruct->IsAsyncTimeSlicingComplete())
-				{
-					IsTimeSlicingProgressing = true;
-				}
+				IsTimeSlicingProgressing = true;
 
 				AccelerationStructure->AddSubstructure(MoveTemp(NewStruct), BucketIdx);
 
@@ -289,6 +336,21 @@ namespace Chaos
 		// If it's not progressing then it is finished so we can perform the final copy if required
 		if (!IsTimeSlicingProgressing)
 		{
+			//todo: creation can go wide, insertion to collection cannot
+			for (uint8 BucketIdx : NonTimeSlicedBucketsToCreate)
+			{
+				if (ViewsPerBucket[BucketIdx].Num())
+				{
+					SCOPE_CYCLE_COUNTER(STAT_CreateNonSlicedStructures);
+
+					auto ParticleView = MakeConstParticleView(MoveTemp(ViewsPerBucket[BucketIdx]));
+					auto NewStruct = SpatialCollectionFactory.CreateAccelerationPerBucket_Threaded(ParticleView, BucketIdx, IsForceFullBuild);
+
+					AccelerationStructure->AddSubstructure(MoveTemp(NewStruct), BucketIdx);
+
+				}
+			}
+
 			if (!bIsSingleThreaded)
 			{
 				// This operation is slow!
diff --git a/Engine/Source/Runtime/Experimental/Chaos/Public/Chaos/PBDRigidsEvolution.h b/Engine/Source/Runtime/Experimental/Chaos/Public/Chaos/PBDRigidsEvolution.h
--- a/Engine/Source/Runtime/Experimental/Chaos/Public/Chaos/PBDRigidsEvolution.h
+++ b/Engine/Source/Runtime/Experimental/Chaos/Public/Chaos/PBDRigidsEvolution.h
@@ -205,6 +205,9 @@ struct CHAOS_API ISpatialAccelerationCollectionFactory
 	//Create an empty acceleration collection with the desired buckets. Chaos enqueues acceleration structure operations per bucket
 	virtual TUniquePtr<ISpatialAccelerationCollection<TAccelerationStructureHandle<FReal, 3>, FReal, 3>> CreateEmptyCollection() = 0;
 
+	// Determines if bucket implements time slicing.
+	virtual bool IsBucketTimeSliced(uint16 BucketIdx) const = 0;
+
 	//Chaos creates new acceleration structures per bucket. Factory can change underlying type at runtime as well as number of buckets to AB test
 	virtual TUniquePtr<ISpatialAcceleration<TAccelerationStructureHandle<FReal, 3>, FReal, 3>> CreateAccelerationPerBucket_Threaded(const TConstParticleView<FSpatialAccelerationCache>& Particles, uint16 BucketIdx, bool ForceFullBuild) = 0;
 
