Chaos : Destruction

- Fix for crash in break events.

diff --git a/Engine/Source/Runtime/Experimental/ChaosSolvers/Private/EventDefaults.cpp b/Engine/Source/Runtime/Experimental/ChaosSolvers/Private/EventDefaults.cpp
--- a/Engine/Source/Runtime/Experimental/ChaosSolvers/Private/EventDefaults.cpp
+++ b/Engine/Source/Runtime/Experimental/ChaosSolvers/Private/EventDefaults.cpp
@@ -76,6 +76,8 @@ namespace Chaos
 						{
 							for (IPhysicsProxyBase* Proxy : *Proxies)
 							{
+								if (NumValidCollisions >= CollisionRule.NumConstraints()) break;
+
 								if (Proxy != nullptr)
 								{
 									if (ensure(!Constraint.AccumulatedImpulse.ContainsNaN() && FMath::IsFinite(Constraint.GetPhi())))
@@ -114,6 +116,8 @@ namespace Chaos
 						{
 							for (IPhysicsProxyBase* Proxy : *Proxies)
 							{
+								if (NumValidCollisions >= CollisionRule.NumConstraints()) break;
+
 								if (Proxy != nullptr)
 								{
 									if (ensure(!Constraint.AccumulatedImpulse.ContainsNaN() && FMath::IsFinite(Constraint.GetPhi())))
