Chaos : Destruction

- Temporary change to prevent deactivated bodies from colliding. This will allow artists to continue testing while the correct particle invalidation is investigated. Currently PBDRigidCluster is disabling particles at the end of Advance, and DirtlyParticle does not seem to be removing the deactivated particles from the ActiveParticles list.

diff --git a/Engine/Source/Runtime/Experimental/Chaos/Public/Chaos/Collision/SpatialAccelerationBroadPhase.h b/Engine/Source/Runtime/Experimental/Chaos/Public/Chaos/Collision/SpatialAccelerationBroadPhase.h
--- a/Engine/Source/Runtime/Experimental/Chaos/Public/Chaos/Collision/SpatialAccelerationBroadPhase.h
+++ b/Engine/Source/Runtime/Experimental/Chaos/Public/Chaos/Collision/SpatialAccelerationBroadPhase.h
@@ -198,6 +198,13 @@ namespace Chaos
 						continue;
 					}
 
+					// HACK : This should not be happening if the disabled particles are properly removed from the active particles list. 
+					if (Particle1.Disabled() || Particle2Generic->Disabled())
+					{
+						continue;
+					}
+
+
 					// Sleeping won't collide against another sleeping and sleeping vs dynamic gets picked up by the other direction.
 					const bool bIsParticle2Kinematic = Particle2.CastToKinematicParticle() &&
 						(Particle2.ObjectState() == EObjectStateType::Kinematic &&
