Chaos : Destruction

- Update the Bounding Box properties within UpdateGeometry.
- Update UnionClusters to reflect the ClusterProperties configuration from 4.23

diff --git a/Engine/Source/Runtime/Experimental/Chaos/Private/Chaos/PBDRigidClustering.cpp b/Engine/Source/Runtime/Experimental/Chaos/Private/Chaos/PBDRigidClustering.cpp
--- a/Engine/Source/Runtime/Experimental/Chaos/Private/Chaos/PBDRigidClustering.cpp
+++ b/Engine/Source/Runtime/Experimental/Chaos/Private/Chaos/PBDRigidClustering.cpp
@@ -630,11 +630,13 @@ namespace Chaos
 
 			for (TTuple<int32, TArray<TPBDRigidParticleHandle<T, 3>* >>& Group : NewClusterGroups)
 			{
-				int32 ClusterGroupID = Group.Key;
-				TArray< TPBDRigidParticleHandle<T, 3> *> ActiveCluster = Group.Value;
-				TPBDRigidClusteredParticleHandle<T, 3>* NewCluster = CreateClusterParticle(-FMath::Abs(ClusterGroupID), MoveTemp(Group.Value));
-				NewCluster->SetInternalCluster(true);
-				for (auto& Constituent : ActiveCluster) MEvolution.DoInternalParticleInitilization( ClusterParents[Constituent], NewCluster);
+				int32 ClusterGroupID = FMath::Abs(Group.Key);
+				TArray< TPBDRigidParticleHandle<T, 3>*> ActiveCluster = Group.Value;
+				FClusterCreationParameters<T> Parameters(0.3, 100, false, !!UnionsHaveCollisionParticles);  Parameters.ConnectionMethod = MClusterUnionConnectionType;
+				TPBDRigidClusteredParticleHandleImp<float, 3, true>* Handle = CreateClusterParticle(-ClusterGroupID, MoveTemp(Group.Value), Parameters, TSharedPtr<FImplicitObject, ESPMode::ThreadSafe>());
+				Handle->SetInternalCluster(true);
+				MEvolution.SetPhysicsMaterial(Handle, MEvolution.GetPhysicsMaterial(ActiveCluster[0]));
+				for (auto& Constituent : ActiveCluster) MEvolution.DoInternalParticleInitilization( ClusterParents[Constituent], Handle);
 			}
 			ClusterUnionMap.Empty();
 		}
@@ -1513,6 +1515,17 @@ namespace Chaos
 				Parent->CollisionParticles()->UpdateAccelerationStructures();
 			}
 		}
+
+		if (TSerializablePtr<FImplicitObject> Implicit = Parent->Geometry())
+		{
+			// strange hacked initilization that seems misplaced and ill thought
+			Parent->SetHasBounds(true);
+			Parent->SetLocalBounds(Implicit->BoundingBox());
+			const Chaos::TAABB<float, 3>& LocalBounds = Parent->LocalBounds();
+			const Chaos::TRigidTransform<float, 3> Xf(Parent->X(), Parent->R());
+			const Chaos::TAABB<float, 3> TransformedBBox = LocalBounds.TransformedAABB(Xf);
+			Parent->SetWorldSpaceInflatedBounds(TransformedBBox);
+		}
 	}
 
 	float MinImpulseForStrainEval = 980 * 2 * 1.f / 30.f; //ignore impulses caused by just keeping object on ground. This is a total hack, we should not use accumulated impulse directly. Instead we need to look at delta v along constraint normal
