updating default collision/pushout parameters to match old settings

diff --git a/Engine/Source/Runtime/Experimental/Chaos/Public/Chaos/PBDRigidsEvolutionGBF.h b/Engine/Source/Runtime/Experimental/Chaos/Public/Chaos/PBDRigidsEvolutionGBF.h
--- a/Engine/Source/Runtime/Experimental/Chaos/Public/Chaos/PBDRigidsEvolutionGBF.h
+++ b/Engine/Source/Runtime/Experimental/Chaos/Public/Chaos/PBDRigidsEvolutionGBF.h
@@ -48,7 +48,7 @@ namespace Chaos
 
 		static constexpr int32 DefaultNumIterations = 1;
 		static constexpr int32 DefaultNumPairIterations = 1;
-		static constexpr int32 DefaultNumPushOutIterations = 5;
+		static constexpr int32 DefaultNumPushOutIterations = 3;
 		static constexpr int32 DefaultNumPushOutPairIterations = 2;
 
 		// @todo(chaos): Required by clustering - clean up
diff --git a/Engine/Source/Runtime/Experimental/ChaosSolvers/Private/PBDRigidsSolver.cpp b/Engine/Source/Runtime/Experimental/ChaosSolvers/Private/PBDRigidsSolver.cpp
--- a/Engine/Source/Runtime/Experimental/ChaosSolvers/Private/PBDRigidsSolver.cpp
+++ b/Engine/Source/Runtime/Experimental/ChaosSolvers/Private/PBDRigidsSolver.cpp
@@ -37,10 +37,10 @@ FAutoConsoleVariableRef CVarChaosSolverUseParticlePool(TEXT("p.Chaos.Solver.UseP
 int32 ChaosSolverParticlePoolNumFrameUntilShrink = 30;
 FAutoConsoleVariableRef CVarChaosSolverParticlePoolNumFrameUntilShrink(TEXT("p.Chaos.Solver.ParticlePoolNumFrameUntilShrink"), ChaosSolverParticlePoolNumFrameUntilShrink, TEXT("Num Frame until we can potentially shrink the pool"));
 
-int32 ChaosSolverCollisionDefaultIterationsCVar = 4;
+int32 ChaosSolverCollisionDefaultIterationsCVar = 1;
 FAutoConsoleVariableRef CVarChaosSolverCollisionDefaultIterations(TEXT("p.ChaosSolverCollisionDefaultIterations"), ChaosSolverCollisionDefaultIterationsCVar, TEXT("Default collision iterations for the solver.[def:1]"));
 
-int32 ChaosSolverCollisionDefaultPushoutIterationsCVar = 4;
+int32 ChaosSolverCollisionDefaultPushoutIterationsCVar = 3;
 FAutoConsoleVariableRef CVarChaosSolverCollisionDefaultPushoutIterations(TEXT("p.ChaosSolverCollisionDefaultPushoutIterations"), ChaosSolverCollisionDefaultPushoutIterationsCVar, TEXT("Default collision pushout iterations for the solver.[def:1]"));
 
 int32 ChaosSolverCleanupCommandsOnDestruction = 1;
