Fixed external strains always triggering breaks in geometry collections no matter how weak they were

diff --git a/Engine/Source/Runtime/Experimental/Chaos/Private/Chaos/PBDRigidClustering.cpp b/Engine/Source/Runtime/Experimental/Chaos/Private/Chaos/PBDRigidClustering.cpp
--- a/Engine/Source/Runtime/Experimental/Chaos/Private/Chaos/PBDRigidClustering.cpp
+++ b/Engine/Source/Runtime/Experimental/Chaos/Private/Chaos/PBDRigidClustering.cpp
@@ -729,12 +729,28 @@ namespace Chaos
 		for (int32 ChildIdx = Children.Num() - 1; ChildIdx >= 0; --ChildIdx)
 		{
 			TPBDRigidClusteredParticleHandle<T, d>* Child = Children[ChildIdx]->CastToClustered();
+			
 			if (!Child)
+			{
 				continue;
-			if ((ExternalStrainMap && (ExternalStrainMap->Find(Child))) ||
-			   (!ExternalStrainMap && Child->CollisionImpulses() >= Child->Strain()) ||
-				bForceRelease)
+			}
+
+			Chaos::FReal ChildStrain = 0.0;
+
+			if(ExternalStrainMap)
+			{
+				const Chaos::FReal* MapStrain = ExternalStrainMap->Find(Child);
+				ChildStrain = MapStrain ? *MapStrain : Child->CollisionImpulses();
+			}
+			else
+			{
+				ChildStrain = Child->CollisionImpulses();
+			}
+
+			if (ChildStrain >= Child->Strain() || bForceRelease)
 			{
+				//UE_LOG(LogTemp, Warning, TEXT("Releasing child %d from parent %p due to strain %.5f Exceeding internal strain %.5f (Source: %s)"), ChildIdx, ClusteredParticle, ChildStrain, Child->Strain(), bForceRelease ? TEXT("Forced by caller") : ExternalStrainMap ? TEXT("External") : TEXT("Collision"));
+
 				// The piece that hits just breaks off - we may want more control 
 				// by looking at the edges of this piece which would give us cleaner 
 				// breaks (this approach produces more rubble)
