Fix angular friction and add a default ether.

diff --git a/Engine/Source/Runtime/Experimental/Chaos/Private/Chaos/PBDCollisionConstraintsContact.cpp b/Engine/Source/Runtime/Experimental/Chaos/Private/Chaos/PBDCollisionConstraintsContact.cpp
--- a/Engine/Source/Runtime/Experimental/Chaos/Private/Chaos/PBDCollisionConstraintsContact.cpp
+++ b/Engine/Source/Runtime/Experimental/Chaos/Private/Chaos/PBDCollisionConstraintsContact.cpp
@@ -87,14 +87,14 @@ namespace Chaos
 			FVec3 Delta = FinalAngularVelocity - RelativeAngularVelocity;
 			if (!bIsRigidDynamic0 && bIsRigidDynamic1)
 			{
-				FMatrix33 WorldSpaceI2 = (Q1 * FMatrix::Identity) * PBDRigid1->I() * (Q1 * FMatrix::Identity).GetTransposed();
+				FMatrix33 WorldSpaceI2 = Utilities::ComputeWorldSpaceInertia(Q1, PBDRigid1->I());
 				FVec3 ImpulseDelta = PBDRigid1->M() * FVec3::CrossProduct(VectorToPoint2, Delta);
 				Impulse += ImpulseDelta;
 				AngularImpulse += WorldSpaceI2 * Delta - FVec3::CrossProduct(VectorToPoint2, ImpulseDelta);
 			}
 			else if (bIsRigidDynamic0 && !bIsRigidDynamic1)
 			{
-				FMatrix33 WorldSpaceI1 = (Q0 * FMatrix::Identity) * PBDRigid0->I() * (Q0 * FMatrix::Identity).GetTransposed();
+				FMatrix33 WorldSpaceI1 = Utilities::ComputeWorldSpaceInertia(Q0, PBDRigid0->I());
 				FVec3 ImpulseDelta = PBDRigid0->M() * FVec3::CrossProduct(VectorToPoint1, Delta);
 				Impulse += ImpulseDelta;
 				AngularImpulse += WorldSpaceI1 * Delta - FVec3::CrossProduct(VectorToPoint1, ImpulseDelta);
@@ -307,9 +307,8 @@ namespace Chaos
 					FReal NormalVelocityChange = FVec3::DotProduct(VelocityChange, Contact.Normal);
 					FMatrix33 FactorInverse = Factor.Inverse();
 					FVec3 MinimalImpulse = FactorInverse * VelocityChange;
-					const FReal MinimalImpulseDotNormal = FVec3::DotProduct(MinimalImpulse, Contact.Normal);
-					const FReal TangentialSize = (MinimalImpulse - MinimalImpulseDotNormal * Contact.Normal).Size();
-					if (TangentialSize <= Friction * MinimalImpulseDotNormal)
+					const FReal TangentialSize = (VelocityChange - NormalVelocityChange * Contact.Normal).Size();
+					if (TangentialSize <= Friction * NormalVelocityChange)
 					{
 						//within friction cone so just solve for static friction stopping the object
 						Impulse = MinimalImpulse;
diff --git a/Engine/Source/Runtime/Experimental/Chaos/Private/Chaos/PerParticleEtherDrag.cpp b/Engine/Source/Runtime/Experimental/Chaos/Private/Chaos/PerParticleEtherDrag.cpp
new file mode 100644
--- /dev/null
+++ b/Engine/Source/Runtime/Experimental/Chaos/Private/Chaos/PerParticleEtherDrag.cpp
@@ -0,0 +1,14 @@
+// Copyright Epic Games, Inc. All Rights Reserved.
+
+#include "Chaos/PerParticleEtherDrag.h"
+
+namespace Chaos
+{
+
+float LinearEtherDragOverride = -1.f;
+FAutoConsoleVariableRef CVarLinearEtherDragOverride(TEXT("p.LinearEtherDragOverride"), LinearEtherDragOverride, TEXT("Set an override linear ether drag value. -1.f to disable"));
+
+float AngularEtherDragOverride = -1.f;
+FAutoConsoleVariableRef CVarAngularEtherDragOverride(TEXT("p.AngularEtherDragOverride"), AngularEtherDragOverride, TEXT("Set an override angular ether drag value. -1.f to disable"));
+
+}
diff --git a/Engine/Source/Runtime/Experimental/Chaos/Public/Chaos/PerParticleEtherDrag.h b/Engine/Source/Runtime/Experimental/Chaos/Public/Chaos/PerParticleEtherDrag.h
--- a/Engine/Source/Runtime/Experimental/Chaos/Public/Chaos/PerParticleEtherDrag.h
+++ b/Engine/Source/Runtime/Experimental/Chaos/Public/Chaos/PerParticleEtherDrag.h
@@ -6,6 +6,9 @@
 
 namespace Chaos
 {
+	extern float LinearEtherDragOverride;
+	extern float AngularEtherDragOverride;
+
 	template<class T, int d>
 	class TPerParticleEtherDrag : public TPerParticleRule<T, d>
 	{
@@ -15,10 +18,12 @@ namespace Chaos
 
 		inline void ApplyHelper(FVec3& V, FVec3& W, FReal LinearDamp, FReal AngularDamp, T Dt) const
 		{
-			const FReal LinearMultiplier = FMath::Max(FReal(0), FReal(1) - (LinearDamp * Dt));
+			const FReal LinearDrag = LinearEtherDragOverride >= 0 ? LinearEtherDragOverride : LinearDamp * Dt;
+			const FReal LinearMultiplier = FMath::Max(FReal(0), FReal(1) - LinearDrag);
 			V *= LinearMultiplier;
 
-			const FReal AngularMultiplier = FMath::Max(FReal(0), FReal(1) - (AngularDamp * Dt));
+			const FReal AngularDrag = AngularEtherDragOverride >= 0 ? AngularEtherDragOverride : AngularDamp * Dt;
+			const FReal AngularMultiplier = FMath::Max(FReal(0), FReal(1) - AngularDrag);
 			W *= AngularMultiplier;
 		}
 
