Fix CIS.

diff --git a/Engine/Source/Runtime/Experimental/Chaos/Private/Chaos/PerParticleEtherDrag.cpp b/Engine/Source/Runtime/Experimental/Chaos/Private/Chaos/PerParticleEtherDrag.cpp
--- a/Engine/Source/Runtime/Experimental/Chaos/Private/Chaos/PerParticleEtherDrag.cpp
+++ b/Engine/Source/Runtime/Experimental/Chaos/Private/Chaos/PerParticleEtherDrag.cpp
@@ -1,6 +1,7 @@
 // Copyright Epic Games, Inc. All Rights Reserved.
 
 #include "Chaos/PerParticleEtherDrag.h"
+#include "HAL/IConsoleManager.h"
 
 namespace Chaos
 {
