Edigrate Dev-Editor updates to InteractiveToolsFramework, GeometryProcessing, MeshModelingToolset, and ModelingTools Editor Mode Plugins. Update a few files in new MeshPaint Plugin for UInteractiveToo

diff --git a/Engine/Plugins/Experimental/ModelingToolsEditorMode/Source/ModelingToolsEditorMode/Private/ModelingToolsEditorMode.cpp b/Engine/Plugins/Experimental/ModelingToolsEditorMode/Source/ModelingToolsEditorMode/Private/ModelingToolsEditorMode.cpp
--- a/Engine/Plugins/Experimental/ModelingToolsEditorMode/Source/ModelingToolsEditorMode/Private/ModelingToolsEditorMode.cpp
+++ b/Engine/Plugins/Experimental/ModelingToolsEditorMode/Source/ModelingToolsEditorMode/Private/ModelingToolsEditorMode.cpp
@@ -12,6 +12,7 @@
 //#include "MeshSurfacePointTool.h"
 //#include "MeshVertexDragTool.h"
 #include "DynamicMeshSculptTool.h"
+#include "MeshVertexSculptTool.h"
 #include "EditMeshPolygonsTool.h"
 #include "DeformMeshPolygonsTool.h"
 #include "ConvertToPolygonsTool.h"
@@ -28,6 +29,9 @@
 #include "MergeMeshesTool.h"
 #include "VoxelCSGMeshesTool.h"
 #include "PlaneCutTool.h"
+#include "SelfUnionMeshesTool.h"
+#include "CSGMeshesTool.h"
+#include "HoleFillTool.h"
 #include "PolygonOnMeshTool.h"
 #include "DisplaceMeshTool.h"
 #include "MeshSpaceDeformerTool.h"
@@ -44,6 +48,7 @@
 #include "CombineMeshesTool.h"
 #include "AlignObjectsTool.h"
 #include "EditUVIslandsTool.h"
+#include "MeshAttributePaintTool.h"
 
 #include "ParameterizeMeshTool.h"
 
@@ -296,7 +301,7 @@ bool FModelingToolsEditorMode::MouseLeave(FEditorViewportClient* ViewportClient,
 //
 // FStylusStateTracker registers itself as a listener for stylus events and implements
 // the IToolStylusStateProviderAPI interface, which allows MeshSurfacePointTool implementations
-// to query for the pen pressure. 
+ // to query for the pen pressure.
 //
 // This is kind of a hack. Unfortunately the current Stylus module is a Plugin so it
 // cannot be used in the base ToolsFramework, and we need this in the Mode as a workaround.
@@ -453,12 +458,55 @@ void FModelingToolsEditorMode::Enter()
 	// register tool set
 
 	//
-	// make shape tools
+	// primitive tools
 	//
-	auto AddPrimitiveToolBuilder = NewObject<UAddPrimitiveToolBuilder>();
-	AddPrimitiveToolBuilder->AssetAPI = ToolsContext->GetAssetAPI();
-	RegisterToolFunc(ToolManagerCommands.BeginAddPrimitiveTool, TEXT("AddPrimitiveTool"), AddPrimitiveToolBuilder);
+	auto RegisterPrimitiveToolFunc  =
+		[this, &RegisterToolFunc](TSharedPtr<FUICommandInfo> UICommand,
+								  FString&& ToolIdentifier,
+								  UAddPrimitiveToolBuilder::EMakeMeshShapeType ShapeTypeIn)
+	{
+		auto AddPrimitiveToolBuilder = NewObject<UAddPrimitiveToolBuilder>();
+		AddPrimitiveToolBuilder->AssetAPI = ToolsContext->GetAssetAPI();
+		AddPrimitiveToolBuilder->ShapeType = ShapeTypeIn;
+		RegisterToolFunc(UICommand, ToolIdentifier, AddPrimitiveToolBuilder);
+	};
+	RegisterPrimitiveToolFunc(ToolManagerCommands.BeginAddBoxPrimitiveTool,
+							  TEXT("BeginAddBoxPrimitiveTool"),
+							  UAddPrimitiveToolBuilder::EMakeMeshShapeType::Box);
+	RegisterPrimitiveToolFunc(ToolManagerCommands.BeginAddCylinderPrimitiveTool,
+							  TEXT("BeginAddCylinderPrimitiveTool"),
+							  UAddPrimitiveToolBuilder::EMakeMeshShapeType::Cylinder);
+	RegisterPrimitiveToolFunc(ToolManagerCommands.BeginAddConePrimitiveTool,
+							  TEXT("BeginAddConePrimitiveTool"),
+							  UAddPrimitiveToolBuilder::EMakeMeshShapeType::Cone);
+	RegisterPrimitiveToolFunc(ToolManagerCommands.BeginAddArrowPrimitiveTool,
+							  TEXT("BeginAddArrowPrimitiveTool"),
+							  UAddPrimitiveToolBuilder::EMakeMeshShapeType::Arrow);
+	RegisterPrimitiveToolFunc(ToolManagerCommands.BeginAddRectanglePrimitiveTool,
+							  TEXT("BeginAddRectanglePrimitiveTool"),
+							  UAddPrimitiveToolBuilder::EMakeMeshShapeType::Rectangle);
+	RegisterPrimitiveToolFunc(ToolManagerCommands.BeginAddRoundedRectanglePrimitiveTool,
+							  TEXT("BeginAddRoundedRectanglePrimitiveTool"),
+							  UAddPrimitiveToolBuilder::EMakeMeshShapeType::RoundedRectangle);
+	RegisterPrimitiveToolFunc(ToolManagerCommands.BeginAddDiscPrimitiveTool,
+							  TEXT("BeginAddDiscPrimitiveTool"),
+							  UAddPrimitiveToolBuilder::EMakeMeshShapeType::Disc);
+	RegisterPrimitiveToolFunc(ToolManagerCommands.BeginAddPuncturedDiscPrimitiveTool,
+							  TEXT("BeginAddPuncturedDiscPrimitiveTool"),
+							  UAddPrimitiveToolBuilder::EMakeMeshShapeType::PuncturedDisc);
+	RegisterPrimitiveToolFunc(ToolManagerCommands.BeginAddTorusPrimitiveTool,
+							  TEXT("BeginAddTorusPrimitiveTool"),
+							  UAddPrimitiveToolBuilder::EMakeMeshShapeType::Torus);
+	RegisterPrimitiveToolFunc(ToolManagerCommands.BeginAddSpherePrimitiveTool,
+							  TEXT("BeginAddSpherePrimitiveTool"),
+							  UAddPrimitiveToolBuilder::EMakeMeshShapeType::Sphere);
+	RegisterPrimitiveToolFunc(ToolManagerCommands.BeginAddSphericalBoxPrimitiveTool,
+							  TEXT("BeginAddSphericalBoxPrimitiveTool"),
+							  UAddPrimitiveToolBuilder::EMakeMeshShapeType::SphericalBox);
 
+	//
+	// make shape tools
+	//
 	auto AddPatchToolBuilder = NewObject<UAddPatchToolBuilder>();
 	AddPatchToolBuilder->AssetAPI = ToolsContext->GetAssetAPI();
 	RegisterToolFunc(ToolManagerCommands.BeginAddPatchTool, TEXT("AddPatchTool"), AddPatchToolBuilder);
@@ -478,10 +526,9 @@ void FModelingToolsEditorMode::Enter()
 
 	//
 	// vertex deform tools
-	// 
+	//
 
-	auto MoveVerticesToolBuilder = NewObject<UDynamicMeshSculptToolBuilder>();
-	MoveVerticesToolBuilder->bEnableRemeshing = false;
+	auto MoveVerticesToolBuilder = NewObject<UMeshVertexSculptToolBuilder>();
 	MoveVerticesToolBuilder->StylusAPI = StylusStateTracker.Get();
 	RegisterToolFunc(ToolManagerCommands.BeginSculptMeshTool, TEXT("MoveVerticesTool"), MoveVerticesToolBuilder);
 
@@ -529,6 +576,9 @@ void FModelingToolsEditorMode::Enter()
 	RemoveOccludedTrianglesToolBuilder->AssetAPI = ToolsContext->GetAssetAPI();
 	RegisterToolFunc(ToolManagerCommands.BeginRemoveOccludedTrianglesTool, TEXT("RemoveOccludedTrianglesTool"), RemoveOccludedTrianglesToolBuilder);
 
+	auto HoleFillToolBuilder = NewObject<UHoleFillToolBuilder>();
+	RegisterToolFunc(ToolManagerCommands.BeginHoleFillTool, TEXT("HoleFillTool"), HoleFillToolBuilder);
+
 	auto UVProjectionToolBuilder = NewObject<UUVProjectionToolBuilder>();
 	UVProjectionToolBuilder->AssetAPI = ToolsContext->GetAssetAPI();
 	RegisterToolFunc(ToolManagerCommands.BeginUVProjectionTool, TEXT("UVProjectionTool"), UVProjectionToolBuilder);
@@ -545,6 +595,14 @@ void FModelingToolsEditorMode::Enter()
 	VoxelCSGMeshesToolBuilder->AssetAPI = ToolsContext->GetAssetAPI();
 	RegisterToolFunc(ToolManagerCommands.BeginVoxelBooleanTool, TEXT("VoxelCSGMeshesTool"), VoxelCSGMeshesToolBuilder);
 
+	auto SelfUnionMeshesToolBuilder = NewObject<USelfUnionMeshesToolBuilder>();
+	SelfUnionMeshesToolBuilder->AssetAPI = ToolsContext->GetAssetAPI();
+	RegisterToolFunc(ToolManagerCommands.BeginSelfUnionTool, TEXT("SelfUnionMeshesTool"), SelfUnionMeshesToolBuilder);
+
+	auto CSGMeshesToolBuilder = NewObject<UCSGMeshesToolBuilder>();
+	CSGMeshesToolBuilder->AssetAPI = ToolsContext->GetAssetAPI();
+	RegisterToolFunc(ToolManagerCommands.BeginMeshBooleanTool, TEXT("CSGMeshesTool"), CSGMeshesToolBuilder);
+
 	auto PlaneCutToolBuilder = NewObject<UPlaneCutToolBuilder>();
 	PlaneCutToolBuilder->AssetAPI = ToolsContext->GetAssetAPI();
 	RegisterToolFunc(ToolManagerCommands.BeginPlaneCutTool, TEXT("PlaneCutTool"), PlaneCutToolBuilder);
@@ -570,6 +628,7 @@ void FModelingToolsEditorMode::Enter()
 	EditMeshMaterialsToolBuilder->AssetAPI = ToolsContext->GetAssetAPI();
 	RegisterToolFunc(ToolManagerCommands.BeginEditMeshMaterialsTool, TEXT("EditMaterialsTool"), EditMeshMaterialsToolBuilder);
 	
+	RegisterToolFunc(ToolManagerCommands.BeginMeshAttributePaintTool, TEXT("MeshAttributePaintTool"), NewObject<UMeshAttributePaintToolBuilder>());
 
 	// analysis tools
 
@@ -707,4 +766,4 @@ void FModelingToolsEditorMode::ConfigureRealTimeViewportsOverride(bool bEnable)
 
 
 
-#undef LOCTEXT_NAMESPACE
\ No newline at end of file
+#undef LOCTEXT_NAMESPACE
diff --git a/Engine/Plugins/Experimental/ModelingToolsEditorMode/Source/ModelingToolsEditorMode/Private/ModelingToolsEditorModeStyle.cpp b/Engine/Plugins/Experimental/ModelingToolsEditorMode/Source/ModelingToolsEditorMode/Private/ModelingToolsEditorModeStyle.cpp
--- a/Engine/Plugins/Experimental/ModelingToolsEditorMode/Source/ModelingToolsEditorMode/Private/ModelingToolsEditorModeStyle.cpp
+++ b/Engine/Plugins/Experimental/ModelingToolsEditorMode/Source/ModelingToolsEditorMode/Private/ModelingToolsEditorModeStyle.cpp
@@ -135,6 +135,10 @@ void FModelingToolsEditorModeStyle::Initialize()
 		StyleSet->Set("ModelingToolsManagerCommands.BeginVoxelMergeTool.Small", 		new IMAGE_PLUGIN_BRUSH("Icons/VoxMerge_40x", 		Icon20x20));
 		StyleSet->Set("ModelingToolsManagerCommands.BeginVoxelBooleanTool", 			new IMAGE_PLUGIN_BRUSH("Icons/VoxBoolean_40x", 		Icon20x20));
 		StyleSet->Set("ModelingToolsManagerCommands.BeginVoxelBooleanTool.Small", 		new IMAGE_PLUGIN_BRUSH("Icons/VoxBoolean_40x", 		Icon20x20));
+		StyleSet->Set("ModelingToolsManagerCommands.BeginSelfUnionTool",				new IMAGE_PLUGIN_BRUSH("Icons/MeshMerge_40x",		Icon20x20));
+		StyleSet->Set("ModelingToolsManagerCommands.BeginSelfUnionTool.Small",			new IMAGE_PLUGIN_BRUSH("Icons/MeshMerge_40x",		Icon20x20));
+		StyleSet->Set("ModelingToolsManagerCommands.BeginMeshBooleanTool",				new IMAGE_PLUGIN_BRUSH("Icons/Boolean_40x",			Icon20x20));
+		StyleSet->Set("ModelingToolsManagerCommands.BeginMeshBooleanTool.Small",		new IMAGE_PLUGIN_BRUSH("Icons/Boolean_40x",			Icon20x20));
 		StyleSet->Set("ModelingToolsManagerCommands.BeginPlaneCutTool", 				new IMAGE_PLUGIN_BRUSH("Icons/PlaneCut_40x", 		Icon20x20));
 		StyleSet->Set("ModelingToolsManagerCommands.BeginPlaneCutTool.Small", 			new IMAGE_PLUGIN_BRUSH("Icons/PlaneCut_40x", 		Icon20x20));
 		StyleSet->Set("ModelingToolsManagerCommands.BeginMeshSelectionTool", 			new IMAGE_PLUGIN_BRUSH("Icons/MeshSelect_40x",		Icon20x20));
diff --git a/Engine/Plugins/Experimental/ModelingToolsEditorMode/Source/ModelingToolsEditorMode/Private/ModelingToolsEditorModeToolkit.cpp b/Engine/Plugins/Experimental/ModelingToolsEditorMode/Source/ModelingToolsEditorMode/Private/ModelingToolsEditorModeToolkit.cpp
--- a/Engine/Plugins/Experimental/ModelingToolsEditorMode/Source/ModelingToolsEditorMode/Private/ModelingToolsEditorModeToolkit.cpp
+++ b/Engine/Plugins/Experimental/ModelingToolsEditorMode/Source/ModelingToolsEditorMode/Private/ModelingToolsEditorModeToolkit.cpp
@@ -26,7 +26,7 @@
 // if set to 1, then on mode initialization we include buttons for prototype modeling tools
 static TAutoConsoleVariable<int32> CVarEnablePrototypeModelingTools(
 	TEXT("modeling.EnablePrototypes"),
-	0,
+	1,
 	TEXT("Enable unsupported Experimental prototype Modeling Tools"));
 
 
@@ -316,6 +316,7 @@ UEdModeInteractiveToolsContext* FModelingToolsEditorModeToolkit::GetToolsContext
 	return GetToolsEditorMode()->GetToolsContext();
 }
 
+static const FName PrimitiveTabName(TEXT("Primitives"));
 static const FName CreateTabName(TEXT("Create"));
 static const FName EditTabName(TEXT("Edit"));
 static const FName SculptTabName(TEXT("Sculpt"));
@@ -327,8 +328,8 @@ static const FName DeformTabName(TEXT("Deform"));
 static const FName PrototypesTabName(TEXT("Prototypes"));
 
 
-const TArray<FName> FModelingToolsEditorModeToolkit::PaletteNames_Standard = { CreateTabName, TransformTabName, DeformTabName, PolyGroupsTabName, TrianglesTabName, UVNormalTabName };
-const TArray<FName> FModelingToolsEditorModeToolkit::PaletteNames_Experimental = { CreateTabName, TransformTabName, DeformTabName, PolyGroupsTabName, TrianglesTabName, UVNormalTabName, PrototypesTabName };
+const TArray<FName> FModelingToolsEditorModeToolkit::PaletteNames_Standard = { PrimitiveTabName, CreateTabName, TransformTabName, DeformTabName, PolyGroupsTabName, TrianglesTabName, UVNormalTabName };
+const TArray<FName> FModelingToolsEditorModeToolkit::PaletteNames_Experimental = { PrimitiveTabName, CreateTabName, TransformTabName, DeformTabName, PolyGroupsTabName, TrianglesTabName, UVNormalTabName, PrototypesTabName };
 
 
 void FModelingToolsEditorModeToolkit::GetToolPaletteNames(TArray<FName>& InPaletteName) const
@@ -370,17 +371,31 @@ void FModelingToolsEditorModeToolkit::BuildToolPalette_Standard(FName PaletteInd
 	ToolbarBuilder.AddSeparator();
 
 
-	if (PaletteIndex == CreateTabName)
+	if (PaletteIndex == PrimitiveTabName)
+	{
+		ToolbarBuilder.AddToolBarButton(Commands.BeginAddBoxPrimitiveTool);
+		ToolbarBuilder.AddToolBarButton(Commands.BeginAddCylinderPrimitiveTool);
+		ToolbarBuilder.AddToolBarButton(Commands.BeginAddConePrimitiveTool);
+		ToolbarBuilder.AddToolBarButton(Commands.BeginAddArrowPrimitiveTool);
+		ToolbarBuilder.AddToolBarButton(Commands.BeginAddRectanglePrimitiveTool);
+		ToolbarBuilder.AddToolBarButton(Commands.BeginAddRoundedRectanglePrimitiveTool);
+		ToolbarBuilder.AddToolBarButton(Commands.BeginAddDiscPrimitiveTool);
+		ToolbarBuilder.AddToolBarButton(Commands.BeginAddPuncturedDiscPrimitiveTool);
+		ToolbarBuilder.AddToolBarButton(Commands.BeginAddTorusPrimitiveTool);
+		ToolbarBuilder.AddToolBarButton(Commands.BeginAddSpherePrimitiveTool);
+		ToolbarBuilder.AddToolBarButton(Commands.BeginAddSphericalBoxPrimitiveTool);
+	}
+	else if (PaletteIndex == CreateTabName)
 	{
 		ToolbarBuilder.AddToolBarButton(Commands.BeginTransformMeshesTool);
 		ToolbarBuilder.AddToolBarButton(Commands.BeginMeshSelectionTool);
 		ToolbarBuilder.AddSeparator();
-		ToolbarBuilder.AddToolBarButton(Commands.BeginAddPrimitiveTool);
 		ToolbarBuilder.AddToolBarButton(Commands.BeginDrawPolygonTool);
 		ToolbarBuilder.AddToolBarButton(Commands.BeginDrawPolyPathTool);
 		ToolbarBuilder.AddSeparator();
 		ToolbarBuilder.AddToolBarButton(Commands.BeginVoxelMergeTool);
 		ToolbarBuilder.AddToolBarButton(Commands.BeginVoxelBooleanTool);
+		ToolbarBuilder.AddToolBarButton(Commands.BeginMeshBooleanTool);
 	}
 	else if (PaletteIndex == TransformTabName)
 	{
@@ -424,6 +439,7 @@ void FModelingToolsEditorModeToolkit::BuildToolPalette_Standard(FName PaletteInd
 		ToolbarBuilder.AddToolBarButton(Commands.BeginEditMeshMaterialsTool);
 		ToolbarBuilder.AddToolBarButton(Commands.BeginWeldEdgesTool);
 		ToolbarBuilder.AddToolBarButton(Commands.BeginRemoveOccludedTrianglesTool);
+		ToolbarBuilder.AddToolBarButton(Commands.BeginHoleFillTool);
 		ToolbarBuilder.AddSeparator();
 		ToolbarBuilder.AddToolBarButton(Commands.BeginMeshInspectorTool);
 	}
@@ -470,18 +486,32 @@ void FModelingToolsEditorModeToolkit::BuildToolPalette_Experimental(FName Palett
 
 	ToolbarBuilder.AddSeparator();
 
-
-	if (PaletteIndex == CreateTabName)
+	if (PaletteIndex == PrimitiveTabName)
+	{
+		ToolbarBuilder.AddToolBarButton(Commands.BeginAddBoxPrimitiveTool);
+		ToolbarBuilder.AddToolBarButton(Commands.BeginAddCylinderPrimitiveTool);
+		ToolbarBuilder.AddToolBarButton(Commands.BeginAddConePrimitiveTool);
+		ToolbarBuilder.AddToolBarButton(Commands.BeginAddArrowPrimitiveTool);
+		ToolbarBuilder.AddToolBarButton(Commands.BeginAddRectanglePrimitiveTool);
+		ToolbarBuilder.AddToolBarButton(Commands.BeginAddRoundedRectanglePrimitiveTool);
+		ToolbarBuilder.AddToolBarButton(Commands.BeginAddDiscPrimitiveTool);
+		ToolbarBuilder.AddToolBarButton(Commands.BeginAddPuncturedDiscPrimitiveTool);
+		ToolbarBuilder.AddToolBarButton(Commands.BeginAddTorusPrimitiveTool);
+		ToolbarBuilder.AddToolBarButton(Commands.BeginAddSpherePrimitiveTool);
+		ToolbarBuilder.AddToolBarButton(Commands.BeginAddSphericalBoxPrimitiveTool);
+	}
+	else if (PaletteIndex == CreateTabName)
 	{
 		ToolbarBuilder.AddToolBarButton(Commands.BeginTransformMeshesTool);
 		ToolbarBuilder.AddToolBarButton(Commands.BeginMeshSelectionTool);
 		ToolbarBuilder.AddSeparator();
-		ToolbarBuilder.AddToolBarButton(Commands.BeginAddPrimitiveTool);
 		ToolbarBuilder.AddToolBarButton(Commands.BeginDrawPolygonTool);
 		ToolbarBuilder.AddToolBarButton(Commands.BeginDrawPolyPathTool);
 		ToolbarBuilder.AddSeparator();
 		ToolbarBuilder.AddToolBarButton(Commands.BeginVoxelMergeTool);
 		ToolbarBuilder.AddToolBarButton(Commands.BeginVoxelBooleanTool);
+		ToolbarBuilder.AddToolBarButton(Commands.BeginSelfUnionTool);
+		ToolbarBuilder.AddToolBarButton(Commands.BeginMeshBooleanTool);
 	}
 	else if (PaletteIndex == TransformTabName)
 	{
@@ -525,6 +555,7 @@ void FModelingToolsEditorModeToolkit::BuildToolPalette_Experimental(FName Palett
 		ToolbarBuilder.AddToolBarButton(Commands.BeginEditMeshMaterialsTool);
 		ToolbarBuilder.AddToolBarButton(Commands.BeginWeldEdgesTool);
 		ToolbarBuilder.AddToolBarButton(Commands.BeginRemoveOccludedTrianglesTool);
+		ToolbarBuilder.AddToolBarButton(Commands.BeginHoleFillTool);
 		ToolbarBuilder.AddSeparator();
 		ToolbarBuilder.AddToolBarButton(Commands.BeginMeshInspectorTool);
 	}
@@ -564,6 +595,8 @@ void FModelingToolsEditorModeToolkit::BuildToolPalette_Experimental(FName Palett
 		ToolbarBuilder.AddSeparator();
 		ToolbarBuilder.AddToolBarButton(Commands.BeginAddPatchTool);
 		ToolbarBuilder.AddToolBarButton(Commands.BeginShapeSprayTool);
+		ToolbarBuilder.AddToolBarButton(Commands.BeginMeshAttributePaintTool);
+
 	}
 }
 
diff --git a/Engine/Plugins/Experimental/ModelingToolsEditorMode/Source/ModelingToolsEditorMode/Private/ModelingToolsManagerActions.cpp b/Engine/Plugins/Experimental/ModelingToolsEditorMode/Source/ModelingToolsEditorMode/Private/ModelingToolsManagerActions.cpp
--- a/Engine/Plugins/Experimental/ModelingToolsEditorMode/Source/ModelingToolsEditorMode/Private/ModelingToolsManagerActions.cpp
+++ b/Engine/Plugins/Experimental/ModelingToolsEditorMode/Source/ModelingToolsEditorMode/Private/ModelingToolsManagerActions.cpp
@@ -21,7 +21,18 @@ FModelingToolsManagerCommands::FModelingToolsManagerCommands() :
 
 void FModelingToolsManagerCommands::RegisterCommands()
 {
-	UI_COMMAND(BeginAddPrimitiveTool, "Primitive", "Start the Add Primitive Tool", EUserInterfaceActionType::Button, FInputChord());
+	UI_COMMAND(BeginAddBoxPrimitiveTool, "Box", "Add Boxes", EUserInterfaceActionType::Button, FInputChord());
+	UI_COMMAND(BeginAddCylinderPrimitiveTool, "Cylinder", "Add Cylinders", EUserInterfaceActionType::Button, FInputChord());
+	UI_COMMAND(BeginAddConePrimitiveTool, "Cone", "Add Cones", EUserInterfaceActionType::Button, FInputChord());
+	UI_COMMAND(BeginAddArrowPrimitiveTool, "Arrow", "Add Arrows", EUserInterfaceActionType::Button, FInputChord());
+	UI_COMMAND(BeginAddRectanglePrimitiveTool, "Rectangle", "Add Rectangles", EUserInterfaceActionType::Button, FInputChord());
+	UI_COMMAND(BeginAddRoundedRectanglePrimitiveTool, "Rounded Rectangle", "Add Rounded Rectangles", EUserInterfaceActionType::Button, FInputChord());
+	UI_COMMAND(BeginAddDiscPrimitiveTool, "Disc", "Add Discs", EUserInterfaceActionType::Button, FInputChord());
+	UI_COMMAND(BeginAddPuncturedDiscPrimitiveTool, "Punctured Disc", "Add PuncturedDiscs", EUserInterfaceActionType::Button, FInputChord());
+	UI_COMMAND(BeginAddTorusPrimitiveTool, "Torus", "Add Torii", EUserInterfaceActionType::Button, FInputChord());
+	UI_COMMAND(BeginAddSphericalBoxPrimitiveTool, "Spherical Box", "Add Box Parametrized Spheres", EUserInterfaceActionType::Button, FInputChord());
+	UI_COMMAND(BeginAddSpherePrimitiveTool, "Sphere", "Add Lat/Long Parametrized Spheres", EUserInterfaceActionType::Button, FInputChord());
+
 	UI_COMMAND(BeginAddPatchTool, "Patch", "Start the Add Patch Tool", EUserInterfaceActionType::Button, FInputChord());
 	UI_COMMAND(BeginDrawPolygonTool, "Polygon", "Start the Draw Polygon Tool", EUserInterfaceActionType::Button, FInputChord());
 	UI_COMMAND(BeginDrawPolyPathTool, "PolyPath", "Start the Draw PolyPath Tool", EUserInterfaceActionType::Button, FInputChord());
@@ -50,7 +61,10 @@ void FModelingToolsManagerCommands::RegisterCommands()
 	UI_COMMAND(BeginUVLayoutTool, "UVLayout", "Start the UV Layout Tool", EUserInterfaceActionType::Button, FInputChord());
 	UI_COMMAND(BeginVoxelMergeTool, "VoxMerge", "Start the Voxel Merge Tool", EUserInterfaceActionType::Button, FInputChord());
 	UI_COMMAND(BeginVoxelBooleanTool, "VoxBoolean", "Start the Voxel Boolean Tool", EUserInterfaceActionType::Button, FInputChord());
+	UI_COMMAND(BeginSelfUnionTool, "Mesh Merge", "Start the Mesh-Based (Intersection-Resolving) Merge Tool", EUserInterfaceActionType::Button, FInputChord());
+	UI_COMMAND(BeginMeshBooleanTool, "Mesh Boolean", "Start the Mesh-Based Boolean Tool", EUserInterfaceActionType::Button, FInputChord());
 	UI_COMMAND(BeginPlaneCutTool, "PlaneCut", "Start the Plane Cut Tool", EUserInterfaceActionType::Button, FInputChord());
+	UI_COMMAND(BeginHoleFillTool, "HoleFill", "Start the Hole Fill Tool", EUserInterfaceActionType::Button, FInputChord());
 	UI_COMMAND(BeginPolygonCutTool, "PolyCut", "Start the Polygon Cut Tool", EUserInterfaceActionType::Button, FInputChord());
 	UI_COMMAND(BeginGlobalUVGenerateTool, "AutoUnwrap", "Start the Global UV Unwrap Tool", EUserInterfaceActionType::Button, FInputChord());
 	UI_COMMAND(BeginGroupUVGenerateTool, "GroupUnwrap", "Stat the PolyGroup UV Unwrap Tool", EUserInterfaceActionType::Button, FInputChord());
@@ -61,6 +75,7 @@ void FModelingToolsManagerCommands::RegisterCommands()
 	UI_COMMAND(BeginPolyGroupsTool, "PolyGroups", "Start the PolyGroups Tool", EUserInterfaceActionType::Button, FInputChord());
 	UI_COMMAND(BeginEditMeshMaterialsTool, "Edit Mats", "Start the Material Editing Tool", EUserInterfaceActionType::Button, FInputChord());
 	UI_COMMAND(BeginTransformUVIslandsTool, "Transform UVs", "Start the UV Island Transformation Tool", EUserInterfaceActionType::Button, FInputChord());
+	UI_COMMAND(BeginMeshAttributePaintTool, "Paint Maps", "Start the Paint Maps Tool", EUserInterfaceActionType::Button, FInputChord());
 	UI_COMMAND(BeginAttributeEditorTool, "Edit Attribs", "Start the Attribute Editor Tool", EUserInterfaceActionType::Button, FInputChord());
 
 	UI_COMMAND(AcceptActiveTool, "Accept", "Accept the active tool", EUserInterfaceActionType::Button, FInputChord());
diff --git a/Engine/Plugins/Experimental/ModelingToolsEditorMode/Source/ModelingToolsEditorMode/Public/ModelingToolsActions.h b/Engine/Plugins/Experimental/ModelingToolsEditorMode/Source/ModelingToolsEditorMode/Public/ModelingToolsActions.h
--- a/Engine/Plugins/Experimental/ModelingToolsEditorMode/Source/ModelingToolsEditorMode/Public/ModelingToolsActions.h
+++ b/Engine/Plugins/Experimental/ModelingToolsEditorMode/Source/ModelingToolsEditorMode/Public/ModelingToolsActions.h
@@ -66,49 +66,20 @@ public:
 // would ever be used at the same time. And we cannot create/register TCommands at runtime.
 // So, we have to define a separate TCommands instance for each Tool. This is unfortunate.
 
+#define DECLARE_TOOL_ACTION_COMMANDS(CommandsClassName) \
+class CommandsClassName : public TInteractiveToolCommands<CommandsClassName> \
+{\
+public:\
+	CommandsClassName();\
+	virtual void GetToolDefaultObjectList(TArray<UInteractiveTool*>& ToolCDOs) override;\
+};\
+
+
+DECLARE_TOOL_ACTION_COMMANDS(FSculptToolActionCommands);
+DECLARE_TOOL_ACTION_COMMANDS(FVertexSculptToolActionCommands);
+DECLARE_TOOL_ACTION_COMMANDS(FTransformToolActionCommands);
+DECLARE_TOOL_ACTION_COMMANDS(FDrawPolygonToolActionCommands);
+DECLARE_TOOL_ACTION_COMMANDS(FMeshSelectionToolActionCommands);
+DECLARE_TOOL_ACTION_COMMANDS(FMeshPlaneCutToolActionCommands);
+DECLARE_TOOL_ACTION_COMMANDS(FEditMeshPolygonsToolActionCommands);
 
-class FSculptToolActionCommands : public TInteractiveToolCommands<FSculptToolActionCommands>
-{
-public:
-	FSculptToolActionCommands();
-	virtual void GetToolDefaultObjectList(TArray<UInteractiveTool*>& ToolCDOs) override;
-};
-
-
-class FTransformToolActionCommands : public TInteractiveToolCommands<FTransformToolActionCommands>
-{
-public:
-	FTransformToolActionCommands();
-	virtual void GetToolDefaultObjectList(TArray<UInteractiveTool*>& ToolCDOs) override;
-};
-
-
-class FDrawPolygonToolActionCommands : public TInteractiveToolCommands<FDrawPolygonToolActionCommands>
-{
-public:
-	FDrawPolygonToolActionCommands();
-	virtual void GetToolDefaultObjectList(TArray<UInteractiveTool*>& ToolCDOs) override;
-};
-
-
-class FMeshSelectionToolActionCommands : public TInteractiveToolCommands<FMeshSelectionToolActionCommands>
-{
-public:
-	FMeshSelectionToolActionCommands();
-	virtual void GetToolDefaultObjectList(TArray<UInteractiveTool*>& ToolCDOs) override;
-};
-
-
-class FMeshPlaneCutToolActionCommands : public TInteractiveToolCommands<FMeshPlaneCutToolActionCommands>
-{
-public:
-	FMeshPlaneCutToolActionCommands();
-	virtual void GetToolDefaultObjectList(TArray<UInteractiveTool*>& ToolCDOs) override;
-};
-
-class FEditMeshPolygonsToolActionCommands : public TInteractiveToolCommands<FEditMeshPolygonsToolActionCommands>
-{
-public:
-	FEditMeshPolygonsToolActionCommands();
-	virtual void GetToolDefaultObjectList(TArray<UInteractiveTool*>& ToolCDOs) override;
-};
\ No newline at end of file
diff --git a/Engine/Plugins/Experimental/ModelingToolsEditorMode/Source/ModelingToolsEditorMode/Public/ModelingToolsManagerActions.h b/Engine/Plugins/Experimental/ModelingToolsEditorMode/Source/ModelingToolsEditorMode/Public/ModelingToolsManagerActions.h
--- a/Engine/Plugins/Experimental/ModelingToolsEditorMode/Source/ModelingToolsEditorMode/Public/ModelingToolsManagerActions.h
+++ b/Engine/Plugins/Experimental/ModelingToolsEditorMode/Source/ModelingToolsEditorMode/Public/ModelingToolsManagerActions.h
@@ -15,7 +15,18 @@ public:
 	FModelingToolsManagerCommands();
 
 
-	TSharedPtr<FUICommandInfo> BeginAddPrimitiveTool;
+	TSharedPtr<FUICommandInfo> BeginAddBoxPrimitiveTool;
+	TSharedPtr<FUICommandInfo> BeginAddCylinderPrimitiveTool;
+	TSharedPtr<FUICommandInfo> BeginAddConePrimitiveTool;
+	TSharedPtr<FUICommandInfo> BeginAddArrowPrimitiveTool;
+	TSharedPtr<FUICommandInfo> BeginAddRectanglePrimitiveTool;
+	TSharedPtr<FUICommandInfo> BeginAddRoundedRectanglePrimitiveTool;
+	TSharedPtr<FUICommandInfo> BeginAddDiscPrimitiveTool;
+	TSharedPtr<FUICommandInfo> BeginAddPuncturedDiscPrimitiveTool;
+	TSharedPtr<FUICommandInfo> BeginAddTorusPrimitiveTool;
+	TSharedPtr<FUICommandInfo> BeginAddSphericalBoxPrimitiveTool;
+	TSharedPtr<FUICommandInfo> BeginAddSpherePrimitiveTool;
+
 	TSharedPtr<FUICommandInfo> BeginAddPatchTool;
 	TSharedPtr<FUICommandInfo> BeginDrawPolygonTool;
 	TSharedPtr<FUICommandInfo> BeginDrawPolyPathTool;
@@ -43,9 +54,12 @@ public:
 	TSharedPtr<FUICommandInfo> BeginUVProjectionTool;
 	TSharedPtr<FUICommandInfo> BeginUVLayoutTool;
 	TSharedPtr<FUICommandInfo> BeginPlaneCutTool;
+	TSharedPtr<FUICommandInfo> BeginHoleFillTool;
 	TSharedPtr<FUICommandInfo> BeginPolygonCutTool;
 	TSharedPtr<FUICommandInfo> BeginVoxelMergeTool;
 	TSharedPtr<FUICommandInfo> BeginVoxelBooleanTool;
+	TSharedPtr<FUICommandInfo> BeginSelfUnionTool;
+	TSharedPtr<FUICommandInfo> BeginMeshBooleanTool;
 	TSharedPtr<FUICommandInfo> BeginMeshSelectionTool;
 
 	TSharedPtr<FUICommandInfo> BeginMeshInspectorTool;
@@ -55,6 +69,7 @@ public:
 	TSharedPtr<FUICommandInfo> BeginPolyGroupsTool;
 	TSharedPtr<FUICommandInfo> BeginEditMeshMaterialsTool;
 	TSharedPtr<FUICommandInfo> BeginTransformUVIslandsTool;
+	TSharedPtr<FUICommandInfo> BeginMeshAttributePaintTool;
 	TSharedPtr<FUICommandInfo> BeginAttributeEditorTool;
 
 	TSharedPtr<FUICommandInfo> AcceptActiveTool;
@@ -66,4 +81,4 @@ public:
 	 * Initialize commands
 	 */
 	virtual void RegisterCommands() override;
-};
\ No newline at end of file
+};
diff --git a/Engine/Plugins/Experimental/PlanarCutPlugin/Source/PlanarCut/Private/PlanarCut.cpp b/Engine/Plugins/Experimental/PlanarCutPlugin/Source/PlanarCut/Private/PlanarCut.cpp
--- a/Engine/Plugins/Experimental/PlanarCutPlugin/Source/PlanarCut/Private/PlanarCut.cpp
+++ b/Engine/Plugins/Experimental/PlanarCutPlugin/Source/PlanarCut/Private/PlanarCut.cpp
@@ -1748,7 +1748,7 @@ void CutWithPlanarCellsHelper(
 				{
 					auto FuncDistSq = [&](int B) { return V.DistanceSquared(Arrangement.Graph.GetVertex(B)); };
 					TPair<int, double> NearestPt = NoisePointHash.FindNearestInRadius(V, PointSpacing*.99, FuncDistSq);
-					return NearestPt.Key != NoisePointHash.InvalidValue();
+					return NearestPt.Key != NoisePointHash.GetInvalidValue();
 				};
 				for (int32 VertIdx = 0; VertIdx < Arrangement.Graph.MaxVertexID(); VertIdx++)
 				{
diff --git a/Engine/Plugins/MeshPainting/Source/MeshPaintingToolset/Private/MeshTexturePaintingTool.cpp b/Engine/Plugins/MeshPainting/Source/MeshPaintingToolset/Private/MeshTexturePaintingTool.cpp
--- a/Engine/Plugins/MeshPainting/Source/MeshPaintingToolset/Private/MeshTexturePaintingTool.cpp
+++ b/Engine/Plugins/MeshPainting/Source/MeshPaintingToolset/Private/MeshTexturePaintingTool.cpp
@@ -195,10 +195,8 @@ void UMeshTexturePaintingTool::Render(IToolsContextRenderAPI* RenderAPI)
 
 }
 
-void UMeshTexturePaintingTool::Tick(float DeltaTime)
+void UMeshTexturePaintingTool::OnTick(float DeltaTime)
 {
-	Super::Tick(DeltaTime);
-
 	if (UMeshToolManager* MeshToolManager = Cast<UMeshToolManager>(GetToolManager()))
 	{
 		if (MeshToolManager->bNeedsRecache)
diff --git a/Engine/Plugins/MeshPainting/Source/MeshPaintingToolset/Private/MeshVertexPaintingTool.cpp b/Engine/Plugins/MeshPainting/Source/MeshPaintingToolset/Private/MeshVertexPaintingTool.cpp
--- a/Engine/Plugins/MeshPainting/Source/MeshPaintingToolset/Private/MeshVertexPaintingTool.cpp
+++ b/Engine/Plugins/MeshPainting/Source/MeshPaintingToolset/Private/MeshVertexPaintingTool.cpp
@@ -169,10 +169,8 @@ void UMeshVertexPaintingTool::Render(IToolsContextRenderAPI* RenderAPI)
 
 }
 
-void UMeshVertexPaintingTool::Tick(float DeltaTime)
+void UMeshVertexPaintingTool::OnTick(float DeltaTime)
 {
-	Super::Tick(DeltaTime);
-
 	if (UMeshToolManager* MeshToolManager = Cast<UMeshToolManager>(GetToolManager()))
 	{
 		if (MeshToolManager->bNeedsRecache)
diff --git a/Engine/Source/Editor/Experimental/EditorInteractiveToolsFramework/Private/EdModeInteractiveToolsContext.cpp b/Engine/Source/Editor/Experimental/EditorInteractiveToolsFramework/Private/EdModeInteractiveToolsContext.cpp
--- a/Engine/Source/Editor/Experimental/EditorInteractiveToolsFramework/Private/EdModeInteractiveToolsContext.cpp
+++ b/Engine/Source/Editor/Experimental/EditorInteractiveToolsFramework/Private/EdModeInteractiveToolsContext.cpp
@@ -557,7 +557,7 @@ void UEdModeInteractiveToolsContext::Tick(FEditorViewportClient* ViewportClient,
 
 	if (bInvalidationPending)
 	{
-		ViewportClient->Invalidate();
+		ViewportClient->Invalidate(false, false);
 		bInvalidationPending = false;
 	}
 
diff --git a/Engine/Source/Runtime/Experimental/InteractiveToolsFramework/Private/BaseTools/BaseBrushTool.cpp b/Engine/Source/Runtime/Experimental/InteractiveToolsFramework/Private/BaseTools/BaseBrushTool.cpp
--- a/Engine/Source/Runtime/Experimental/InteractiveToolsFramework/Private/BaseTools/BaseBrushTool.cpp
+++ b/Engine/Source/Runtime/Experimental/InteractiveToolsFramework/Private/BaseTools/BaseBrushTool.cpp
@@ -21,18 +21,6 @@ UBrushBaseProperties::UBrushBaseProperties()
 	BrushFalloffAmount = 1.0f;
 }
 
-void UBrushBaseProperties::SaveRestoreProperties(UInteractiveTool* Tool, bool bSaving)
-{
-	UBrushBaseProperties* PropertyCache = GetPropertyCache<UBrushBaseProperties>();
-	SaveRestoreProperty(PropertyCache->BrushSize, this->BrushSize, bSaving);
-	SaveRestoreProperty(PropertyCache->bSpecifyRadius, this->bSpecifyRadius, bSaving);
-	SaveRestoreProperty(PropertyCache->BrushRadius, this->BrushRadius, bSaving);
-	SaveRestoreProperty(PropertyCache->BrushFalloffAmount, this->BrushFalloffAmount, bSaving);
-	SaveRestoreProperty(PropertyCache->BrushStrength, this->BrushStrength, bSaving);
-}
-
-
-
 UBaseBrushTool::UBaseBrushTool()
 {
 	PropertyClass = UBrushBaseProperties::StaticClass();
@@ -239,14 +227,6 @@ void UBaseBrushTool::Render(IToolsContextRenderAPI* RenderAPI)
 	UpdateBrushStampIndicator();
 }
 
-void UBaseBrushTool::Tick(float DeltaTime)
-{
-	UMeshSurfacePointTool::Tick(DeltaTime);
-}
-
-
-
-
 const FString BaseBrushIndicatorGizmoType = TEXT("BrushIndicatorGizmoType");
 
 void UBaseBrushTool::SetupBrushStampIndicator()
diff --git a/Engine/Source/Runtime/Experimental/InteractiveToolsFramework/Private/InteractiveTool.cpp b/Engine/Source/Runtime/Experimental/InteractiveToolsFramework/Private/InteractiveTool.cpp
--- a/Engine/Source/Runtime/Experimental/InteractiveToolsFramework/Private/InteractiveTool.cpp
+++ b/Engine/Source/Runtime/Experimental/InteractiveToolsFramework/Private/InteractiveTool.cpp
@@ -8,6 +8,14 @@
 
 #define LOCTEXT_NAMESPACE "UInteractiveTool"
 
+#if WITH_EDITOR
+void UInteractiveToolPropertySet::PostEditChangeProperty(FPropertyChangedEvent& PropertyChangedEvent)
+{
+	Super::PostEditChangeProperty(PropertyChangedEvent);
+
+	OnModified.Broadcast(this, PropertyChangedEvent.Property);
+}
+#endif
 
 UInteractiveTool::UInteractiveTool()
 {
@@ -144,9 +152,40 @@ TArray<UObject*> UInteractiveTool::GetToolProperties(bool bEnabledOnly) const
 		}
 	}
 
-	return MoveTemp(Properties);
+	return Properties;
 }
 
+void
+UInteractiveToolPropertySet::SaveProperties(UInteractiveTool* SaveFromTool)
+{
+	SaveRestoreProperties(SaveFromTool, true);
+}
+
+void
+UInteractiveToolPropertySet::RestoreProperties(UInteractiveTool* RestoreToTool)
+{
+	SaveRestoreProperties(RestoreToTool, false);
+}
+
+void UInteractiveToolPropertySet::SaveRestoreProperties(UInteractiveTool* RestoreToTool, bool bSaving)
+{
+	UInteractiveToolPropertySet* PropertyCache = GetDynamicPropertyCache();
+	for ( FProperty* Prop : TFieldRange<FProperty>(GetClass()) )
+	{
+#if WITH_EDITOR
+		if (!Prop->HasMetaData(TEXT("TransientToolProperty")))
+#endif
+		{
+			void* DestValue = Prop->ContainerPtrToValuePtr<void>(this);
+			void* SrcValue = Prop->ContainerPtrToValuePtr<void>(PropertyCache);
+			if ( bSaving )
+			{
+				Swap(SrcValue, DestValue);
+			}
+			Prop->CopySingleValue(DestValue, SrcValue);
+		}
+	}
+}
 
 void UInteractiveTool::RegisterActions(FInteractiveToolActionSet& ActionSet)
 {
@@ -187,6 +226,22 @@ bool UInteractiveTool::CanAccept() const
 
 void UInteractiveTool::Tick(float DeltaTime)
 {
+	for (auto* Object : ToolPropertyObjects)
+	{
+		auto* Propset = Cast<UInteractiveToolPropertySet>(Object);
+		if ( Propset != nullptr )
+		{
+			if ( Propset->IsPropertySetEnabled() )
+			{
+				Propset->CheckAndUpdateWatched();
+			}
+			else
+			{
+				Propset->SilentUpdateWatched();
+			}
+		}
+	}
+	OnTick(DeltaTime);
 }
 
 UInteractiveToolManager* UInteractiveTool::GetToolManager() const
@@ -197,4 +252,4 @@ UInteractiveToolManager* UInteractiveTool::GetToolManager() const
 }
 
 
-#undef LOCTEXT_NAMESPACE
\ No newline at end of file
+#undef LOCTEXT_NAMESPACE
diff --git a/Engine/Source/Runtime/Experimental/InteractiveToolsFramework/Public/BaseTools/BaseBrushTool.h b/Engine/Source/Runtime/Experimental/InteractiveToolsFramework/Public/BaseTools/BaseBrushTool.h
--- a/Engine/Source/Runtime/Experimental/InteractiveToolsFramework/Public/BaseTools/BaseBrushTool.h
+++ b/Engine/Source/Runtime/Experimental/InteractiveToolsFramework/Public/BaseTools/BaseBrushTool.h
@@ -40,17 +40,12 @@ public:
 	float BrushFalloffAmount;
 
 	/** If false, then BrushStrength will not be shown in DetailsView panels (otherwise no effect) */
-	UPROPERTY()
+	UPROPERTY( meta = (TransientToolProperty))
 	bool bShowStrength = true;
 
 	/** If false, then BrushFalloffAmount will not be shown in DetailsView panels (otherwise no effect) */
-	UPROPERTY()
+	UPROPERTY(meta = (TransientToolProperty))
 	bool bShowFalloff = true;
-
-	//
-	// save/restore support
-	//
-	virtual void SaveRestoreProperties(UInteractiveTool* Tool, bool bSaving) override;
 };
 
 
@@ -98,7 +93,6 @@ public:
 	virtual void Setup() override;
 	virtual void Shutdown(EToolShutdownType ShutdownType) override;
 
-	virtual void Tick(float DeltaTime) override;
 	virtual void Render(IToolsContextRenderAPI* RenderAPI) override;
 
 	virtual void OnPropertyModified(UObject* PropertySet, FProperty* Property) override;
@@ -137,6 +131,8 @@ public:
 
 	virtual bool IsInBrushStroke() const { return bInBrushStroke; }
 
+	virtual double GetCurrentBrushRadius() const { return CurrentBrushRadius; }
+
 protected:
 
 	/**
diff --git a/Engine/Source/Runtime/Experimental/InteractiveToolsFramework/Public/InteractiveTool.h b/Engine/Source/Runtime/Experimental/InteractiveToolsFramework/Public/InteractiveTool.h
--- a/Engine/Source/Runtime/Experimental/InteractiveToolsFramework/Public/InteractiveTool.h
+++ b/Engine/Source/Runtime/Experimental/InteractiveToolsFramework/Public/InteractiveTool.h
@@ -5,7 +5,9 @@
 #include "CoreMinimal.h"
 #include "InputBehaviorSet.h"
 #include "InteractiveToolActionSet.h"
+#include "Shader.h"
 #include "ToolContextInterfaces.h"
+#include "UObject/UObjectGlobals.h"
 #include "InteractiveTool.generated.h"
 
 class UInteractiveToolManager;
@@ -22,9 +24,6 @@ enum class EToolShutdownType
 	Cancel = 2
 };
 
-
-
-
 /**
  * FInteractiveToolInfo provides information about a tool (name, tooltip, etc)
  */
@@ -35,6 +34,90 @@ struct INTERACTIVETOOLSFRAMEWORK_API FInteractiveToolInfo
 };
 
 
+class INTERACTIVETOOLSFRAMEWORK_API FWatchablePropertySet
+{
+public:
+	//
+	// Property watching infrastructure
+	//
+	class FPropertyWatcher
+	{
+	public:
+		virtual ~FPropertyWatcher() = default;
+		virtual void CheckAndUpdate() = 0;
+		virtual void SilentUpdate() = 0;
+	};
+
+	template <typename PropType>
+	class TPropertyWatcher : public FPropertyWatcher
+	{
+	public:
+		using FValueGetter = TFunction<PropType(void)>;
+		using FChangedCallback = TFunction<void(const PropType&)>;
+
+		TPropertyWatcher(const PropType& Property,
+						 FChangedCallback OnChangedIn)
+			: GetValue([&Property](){return Property;}),
+			  OnChanged(MoveTemp(OnChangedIn))
+		{}
+		TPropertyWatcher(FValueGetter GetValueIn,
+						 FChangedCallback OnChangedIn)
+			: GetValue(MoveTemp(GetValueIn)), OnChanged(MoveTemp(OnChangedIn))
+		{}
+		void CheckAndUpdate() final
+		{
+			PropType Value = GetValue();
+			if ((!Cached.IsSet()) || (Cached.GetValue() != Value))
+			{
+				Cached = Value;
+				OnChanged(Cached.GetValue());
+			}
+		}
+		void SilentUpdate() final
+		{
+			Cached = GetValue();
+		}
+	private:
+		TOptional<PropType> Cached;
+		FValueGetter GetValue;
+		FChangedCallback OnChanged;
+	};
+
+	FWatchablePropertySet() = default;
+	FWatchablePropertySet(const FWatchablePropertySet&) = delete;
+	FWatchablePropertySet& operator=(const FWatchablePropertySet&) = delete;
+
+	void CheckAndUpdateWatched()
+	{
+		for ( auto& PropWatcher : PropertyWatchers )
+		{
+			PropWatcher->CheckAndUpdate();
+		}
+	}
+	void SilentUpdateWatched()
+	{
+		for ( auto& PropWatcher : PropertyWatchers )
+		{
+			PropWatcher->SilentUpdate();
+		}
+	}
+
+	template <typename PropType>
+	void WatchProperty(const PropType& ValueIn,
+					   typename TPropertyWatcher<PropType>::FChangedCallback OnChangedIn)
+	{
+		PropertyWatchers.Emplace(MakeUnique<TPropertyWatcher<PropType>>(ValueIn, OnChangedIn));
+	}
+	template <typename PropType>
+	void WatchProperty(typename TPropertyWatcher<PropType>::FValueGetter GetValueIn,
+					   typename TPropertyWatcher<PropType>::FChangedCallback OnChangedIn)
+	{
+		PropertyWatchers.Emplace(MakeUnique<TPropertyWatcher<PropType>>(GetValueIn, OnChangedIn));
+	}
+private:
+	TArray<TUniquePtr<FPropertyWatcher>> PropertyWatchers;
+};
+
 /** This delegate is used by UInteractiveToolPropertySet */
 DECLARE_MULTICAST_DELEGATE_TwoParams(FInteractiveToolPropertySetModifiedSignature, UObject*, FProperty*);
 
@@ -42,73 +125,62 @@ DECLARE_MULTICAST_DELEGATE_TwoParams(FInteractiveToolPropertySetModifiedSignatur
 /**
  * A UInteractiveTool contains a set of UObjects that contain "properties" of the Tool, ie
  * the configuration flags, parameters, etc that control the Tool. Currently any UObject
- * can be added as a property set, however there is no automatic mechanism for those child 
+ * can be added as a property set, however there is no automatic mechanism for those child
  * UObjects to notify the Tool when a property changes.
- * 
+ *
  * If you make your property set UObjects subclasses of UInteractiveToolPropertySet, then
  * when the Tool Properties are changed *in the Editor*, the parent Tool will be automatically notified.
  * You can override UInteractiveTool::OnPropertyModified() to act on these notifications
  */
 UCLASS(Transient)
-class INTERACTIVETOOLSFRAMEWORK_API UInteractiveToolPropertySet : public UObject
+class INTERACTIVETOOLSFRAMEWORK_API UInteractiveToolPropertySet : public UObject, public FWatchablePropertySet
 {
 	GENERATED_BODY()
 
-protected:
-	FInteractiveToolPropertySetModifiedSignature OnModified;
-
 public:
-
 	/** @return the multicast delegate that is called when properties are modified */
 	FInteractiveToolPropertySetModifiedSignature& GetOnModified()
 	{
 		return OnModified;
 	}
 
-	/** 
-	  * Posts a message to the OnModified delegate with the modified FProperty 
-	  * @warning this function is currently only called in Editor (not at runtime)
-	  */
-	void PostEditChangeProperty(FPropertyChangedEvent& PropertyChangedEvent)
+	/** Return true if this property set is enabled. Enabled/Disable state is intended to be used to control things like visibility in UI/etc. */
+	bool IsPropertySetEnabled() const
 	{
-		OnModified.Broadcast(this, PropertyChangedEvent.Property);
+		return bIsPropertySetEnabled;
 	}
 
-
-
 	//
 	// Setting saving/serialization
 	//
-
-public:
-	/**
-	 * Save values of current Tool Properties. Implementing these functions is *optional*
-	 * and how it is implemented is up to the PropertySet implementation.
-	 * It is not necessary to save/restore all possible Properties (in many cases this would not make sense).
-	 * GetPropertyCache() can be used to return an instance of subclasses that is an easy
-	 * place to save/restore these properties
-	 */
-	virtual void SaveProperties(UInteractiveTool* SaveFromTool) 
-	{ 
-		SaveRestoreProperties(SaveFromTool, true); 
-	}
-
 	/**
-	 * Restore saved property values
+	 * Save and restore values of current Tool Properties between tool invocations
+	 *
+	 * The default behaviour of these functions is to Save or Restore every property in the property set.  It is not
+	 * necessary to save/restore all possible Properties (in many cases this would not make sense), so individual
+	 * properties may be skipped by adding the "TransientToolProperty" tag to their metadata on a property by property
+	 * basis.
+	 *
+	 * Property sets which need more exotic behaviour upon Save and Restore may override these routines
+	 *
+	 * GetPropertyCache() and GetDynamicPropertyCache() can be used to return an instance of either the static or the
+	 * dynamic type of the specified property set subclass which may be used as a place to save/restore these properties
+	 * by customized Save/Restore functions
 	 */
-	virtual void RestoreProperties(UInteractiveTool* RestoreToTool) 
-	{ 
-		SaveRestoreProperties(RestoreToTool, false);
-	}
+	virtual void SaveProperties(UInteractiveTool* SaveFromTool);
+	virtual void RestoreProperties(UInteractiveTool* RestoreToTool);
+private:
 
+	// Utility func used to implement the default Save/RestoreProperties funcs
+	void SaveRestoreProperties(UInteractiveTool* RestoreToTool, bool bSaving);
 protected:
 	/**
 	 * GetPropertyCache returns a class-internal object that subclasses can use to save/restore properties.
 	 * If the subclass is UMyPropertySet, this function should only ever be called as GetPropertyCache<UMyPropertySet>().
 	 */
 	template<typename ObjType>
-	ObjType* GetPropertyCache()
-	{ 
+	static ObjType* GetPropertyCache()
+	{
 		ObjType* CDO = GetMutableDefault<ObjType>();
 		if (CDO->CachedProperties == nullptr)
 		{
@@ -117,52 +189,37 @@ protected:
 		return CastChecked<ObjType>(CDO->CachedProperties);
 	}
 
-	/**
-	 * Subclass this to save and restore in a single function. Use the SaveRestoryProperty function for each member:
-	 *   UMyPropertySetSubclass* PropertyCache = GetPropertyCache<UMyPropertySetSubclass>();
-	 *	 SaveRestoreProperty(PropertyCache->PropSetMember, this->PropSetMember, bSaving);
-	 */
-	virtual void SaveRestoreProperties(UInteractiveTool* RestoreToTool, bool bSaving) {}
-
-	/**
-	 * Call this from SaveRestoreProperties to either save or restore a given property value
-	 */
-	template<typename T>
-	void SaveRestoreProperty(T& CacheValue, T& PropsValue, bool bSaving)
+	UInteractiveToolPropertySet* GetDynamicPropertyCache()
 	{
-		if (bSaving)
-		{
-			CacheValue = PropsValue;
-		}
-		else
+		UInteractiveToolPropertySet* CDO = GetMutableDefault<UInteractiveToolPropertySet>(GetClass());
+		if (CDO->CachedProperties == nullptr)
 		{
-			PropsValue = CacheValue;
+			CDO->CachedProperties = NewObject<UInteractiveToolPropertySet>((UObject*)GetTransientPackage(), GetClass());
 		}
+		return CDO->CachedProperties;
 	}
 
+public:
+#if WITH_EDITOR
+	/**
+	  * Posts a message to the OnModified delegate with the modified FProperty
+	  * @warning Please consider listening to OnModified instead of overriding this function
+	  * @warning this function is currently only called in Editor (not at runtime)
+	  */
+	virtual void PostEditChangeProperty(FPropertyChangedEvent& PropertyChangedEvent);
+#endif
 
-private:
+protected:
 	// CachedProperties should only ever be set to an instance of the subclass, ideally via GetPropertyCache().
 	UPROPERTY()
-	UObject* CachedProperties = nullptr;
-
-
-	//
-	// Visibility, enable/disable, etc
-	//
+	UInteractiveToolPropertySet* CachedProperties = nullptr;
 
-private:
 	UPROPERTY()
 	bool bIsPropertySetEnabled = true;
 
 	friend class UInteractiveTool;	// so that tool can enable/disable
 
-public:
-	/** Return true if this property set is enabled. Enabled/Disable state is intended to be used to control things like visibility in UI/etc. */
-	bool IsPropertySetEnabled() const
-	{
-		return bIsPropertySetEnabled;
-	}
+	FInteractiveToolPropertySetModifiedSignature OnModified;
 };
 
 
@@ -176,7 +233,7 @@ public:
  * "tool behavior", like a tool that responds to a mouse click, etc, that can be
  * extended to implement custom behaviors.
  *
- * In the framework, you do not create instances of UInteractiveTool yourself. 
+ * In the framework, you do not create instances of UInteractiveTool yourself.
  * You provide a UInteractiveToolBuilder implementation that can properly construct
  * an instance of your Tool, this is where for example default parameters would be set.
  * The ToolBuilder is registered with the ToolManager, and then UInteractiveToolManager::ActivateTool()
@@ -210,13 +267,12 @@ public:
 	 */
 	virtual void Render(IToolsContextRenderAPI* RenderAPI);
 
+
 	/**
-	 * Allow the Tool to do any necessary processing on Tick 
+	 * Non overrideable func which does processing and calls the tool's OnTick
 	 * @param DeltaTime the time delta since last tick
 	 */
-	virtual void Tick(float DeltaTime);
-
-
+	virtual void Tick(float DeltaTime) final;
 
 	/**
 	 * @return ToolManager that owns this Tool
@@ -224,18 +280,18 @@ public:
 	virtual UInteractiveToolManager* GetToolManager() const;
 
 
-	/** 
-	 * @return true if this Tool support being Cancelled, ie calling Shutdown(EToolShutdownType::Cancel)  
+	/**
+	 * @return true if this Tool support being Cancelled, ie calling Shutdown(EToolShutdownType::Cancel)
 	 */
 	virtual bool HasCancel() const;
 
-	/** 
-	 * @return true if this Tool support being Accepted, ie calling Shutdown(EToolShutdownType::Accept)  
+	/**
+	 * @return true if this Tool support being Accepted, ie calling Shutdown(EToolShutdownType::Accept)
 	 */
 	virtual bool HasAccept() const;
 
-	/** 
-	 * @return true if this Tool is currently in a state where it can be Accepted. This may be false if for example there was an error in the Tool. 
+	/**
+	 * @return true if this Tool is currently in a state where it can be Accepted. This may be false if for example there was an error in the Tool.
 	 */
 	virtual bool CanAccept() const;
 
@@ -331,12 +387,19 @@ protected:
 
 	//
 	// Action support/system
-	// 
+	//
 	// Your Tool subclass can register a set of "Actions" it can execute
 	// by overloading RegisterActions(). Then external systems can use GetActionSet() to
 	// find out what Actions your Tool supports, and ExecuteAction() to run those actions.
 	//
-	
+
+private:
+	/**
+	 * Allow the Tool to do any necessary processing on Tick
+	 * @param DeltaTime the time delta since last tick
+	 */
+	virtual void OnTick(float DeltaTime){};
+
 public:
 	/**
 	 * Get the internal Action Set for this Tool. The action set is created and registered on-demand.
@@ -361,8 +424,8 @@ protected:
 
 
 private:
-	/** 
-	 * Set of actions this Tool can execute. This variable is allocated on-demand. 
+	/**
+	 * Set of actions this Tool can execute. This variable is allocated on-demand.
 	 * Use GetActionSet() instead of accessing this pointer directly!
 	 */
 	FInteractiveToolActionSet* ToolActionSet = nullptr;
diff --git a/Engine/Source/Runtime/Experimental/InteractiveToolsFramework/Public/InteractiveToolChange.h b/Engine/Source/Runtime/Experimental/InteractiveToolsFramework/Public/InteractiveToolChange.h
--- a/Engine/Source/Runtime/Experimental/InteractiveToolsFramework/Public/InteractiveToolChange.h
+++ b/Engine/Source/Runtime/Experimental/InteractiveToolsFramework/Public/InteractiveToolChange.h
@@ -80,3 +80,53 @@ public:
 
 
 
+
+
+
+/**
+ * Holds another Change and forwards Apply/Revert to it, with
+ * calls to Before/After lambas, allowing client classes to
+ * respond to a change without having to intercept it explicitly.
+ * (Be very careful with these lambdas!)
+ */
+template<typename ChangeType>
+class TWrappedToolCommandChange : public FToolCommandChange
+{
+public:
+	TUniquePtr<ChangeType> WrappedChange;
+
+	TUniqueFunction<void(bool bRevert)> BeforeModify;
+	TUniqueFunction<void(bool bRevert)> AfterModify;
+
+	virtual void Apply(UObject* Object) override
+	{
+		if (BeforeModify)
+		{
+			BeforeModify(false);
+		}
+		WrappedChange->Apply(Object);
+		if (AfterModify)
+		{
+			AfterModify(false);
+		}
+	}
+
+	virtual void Revert(UObject* Object) override
+	{
+		if (BeforeModify)
+		{
+			BeforeModify(true);
+		}
+		WrappedChange->Revert(Object);
+		if (AfterModify)
+		{
+			AfterModify(true);
+		}
+	}
+
+	virtual FString ToString() const override
+	{
+		return WrappedChange->ToString();
+	}
+};
+
