Disable recusive union collisions.

diff --git a/Engine/Source/Runtime/Experimental/Chaos/Private/Chaos/Collision/CollisionDetector.cpp b/Engine/Source/Runtime/Experimental/Chaos/Private/Chaos/Collision/CollisionDetector.cpp
--- a/Engine/Source/Runtime/Experimental/Chaos/Private/Chaos/Collision/CollisionDetector.cpp
+++ b/Engine/Source/Runtime/Experimental/Chaos/Private/Chaos/Collision/CollisionDetector.cpp
@@ -8,6 +8,9 @@ namespace Chaos
 {
 	DEFINE_STAT(STAT_Collisions_Detect);
 	DEFINE_STAT(STAT_Collisions_BroadPhase);
+	DEFINE_STAT(STAT_Collisions_SpatialBroadPhase);
+	DEFINE_STAT(STAT_Collisions_Filtering);
+	DEFINE_STAT(STAT_Collisions_SpatialNarrowPhase);
 #if CHAOS_ENABLE_STAT_NARROWPHASE
 	DEFINE_STAT(STAT_Collisions_NarrowPhase);
 #endif
diff --git a/Engine/Source/Runtime/Experimental/Chaos/Private/Chaos/CollisionResolution.cpp b/Engine/Source/Runtime/Experimental/Chaos/Private/Chaos/CollisionResolution.cpp
--- a/Engine/Source/Runtime/Experimental/Chaos/Private/Chaos/CollisionResolution.cpp
+++ b/Engine/Source/Runtime/Experimental/Chaos/Private/Chaos/CollisionResolution.cpp
@@ -30,6 +30,8 @@ DECLARE_CYCLE_STAT(TEXT("Collisions::GJK"), STAT_Collisions_GJK, STATGROUP_Chaos
 #define SCOPE_CYCLE_COUNTER_GJK()
 #endif
 
+#define CHAOS_COLLIDE_CLUSTERED_UNIONS 0
+
 //PRAGMA_DISABLE_OPTIMIZATION
 
 float CCDEnableThresholdBoundsScale = 0.4f;
@@ -2605,6 +2607,7 @@ namespace Chaos
 				return;
 			}
 
+#if CHAOS_COLLIDE_CLUSTERED_UNIONS
 			if (Implicit0OuterType == FImplicitObjectUnionClustered::StaticType())
 			{
 				const FImplicitObjectUnionClustered* Union0 = Implicit0->template GetObject<FImplicitObjectUnionClustered>();
@@ -2614,6 +2617,7 @@ namespace Chaos
 				}
 				return;
 			}
+#endif
 
 			if (Implicit1OuterType == FImplicitObjectUnion::StaticType())
 			{
@@ -2625,7 +2629,7 @@ namespace Chaos
 				return;
 			}
 
-
+#if CHAOS_COLLIDE_CLUSTERED_UNIONS
 			if (Implicit1OuterType == FImplicitObjectUnionClustered::StaticType())
 			{
 				const FImplicitObjectUnionClustered* Union1 = Implicit1->template GetObject<FImplicitObjectUnionClustered>();
@@ -2635,6 +2639,7 @@ namespace Chaos
 				}
 				return;
 			}
+#endif
 
 			// Check shape pair filtering if enable
 			if (Context.bFilteringEnabled && !DoCollide(Implicit0Type, Particle0->GetImplicitShape(Implicit0), Implicit1Type, Particle1->GetImplicitShape(Implicit1)))
diff --git a/Engine/Source/Runtime/Experimental/Chaos/Private/Chaos/PBDRigidClustering.cpp b/Engine/Source/Runtime/Experimental/Chaos/Private/Chaos/PBDRigidClustering.cpp
--- a/Engine/Source/Runtime/Experimental/Chaos/Private/Chaos/PBDRigidClustering.cpp
+++ b/Engine/Source/Runtime/Experimental/Chaos/Private/Chaos/PBDRigidClustering.cpp
@@ -1443,7 +1443,7 @@ namespace Chaos
 			{
 				ensureMsgf(false, TEXT("Checking usage with no proxy and multiple ojects with levelsets"));
 
-				FImplicitObjectUnion UnionObject(MoveTemp(Objects));
+				FImplicitObjectUnionClustered UnionObject(MoveTemp(Objects));
 				TAABB<T, d> Bounds = UnionObject.BoundingBox();
 				const TVector<T, d> BoundsExtents = Bounds.Extents();
 				if (BoundsExtents.Min() >= MinLevelsetSize) //make sure the object is not too small
diff --git a/Engine/Source/Runtime/Experimental/Chaos/Public/Chaos/Collision/SpatialAccelerationBroadPhase.h b/Engine/Source/Runtime/Experimental/Chaos/Public/Chaos/Collision/SpatialAccelerationBroadPhase.h
--- a/Engine/Source/Runtime/Experimental/Chaos/Public/Chaos/Collision/SpatialAccelerationBroadPhase.h
+++ b/Engine/Source/Runtime/Experimental/Chaos/Public/Chaos/Collision/SpatialAccelerationBroadPhase.h
@@ -9,9 +9,14 @@
 #include "Chaos/ParticleHandle.h"
 #include "Chaos/PBDRigidsSOAs.h"
 #include "Chaos/Capsule.h"
+#include "ChaosStats.h"
 
 namespace Chaos
 {
+	DECLARE_CYCLE_STAT_EXTERN(TEXT("Collisions::BroadPhase"), STAT_Collisions_SpatialBroadPhase, STATGROUP_ChaosCollision, CHAOS_API);
+	DECLARE_CYCLE_STAT_EXTERN(TEXT("Collisions::NarrowPhase"), STAT_Collisions_SpatialNarrowPhase, STATGROUP_ChaosCollision, CHAOS_API);
+	DECLARE_CYCLE_STAT_EXTERN(TEXT("Collisions::Filtering"), STAT_Collisions_Filtering, STATGROUP_ChaosCollision, CHAOS_API);
+	
 	template <typename TPayloadType, typename T, int d>
 	class ISpatialAcceleration;
 
@@ -138,36 +143,40 @@ namespace Chaos
 
 			if (Particle1.CastToRigidParticle() && (Particle1.ObjectState() == EObjectStateType::Dynamic || Particle1.ObjectState() == EObjectStateType::Sleeping))
 			{
-				// @todo(ccaulfield): the spatial acceleration scheme needs to know the expanded bboxes for all particles, not just the one doing the test
 				const bool bBody1Bounded = HasBoundingBox(Particle1);
 				const FReal Box1Thickness = ComputeBoundsThickness(Particle1, Dt, BoundsThickness, BoundsThicknessVelocityInflation).Size();
-				if (bBody1Bounded)
 				{
-					// @todo(ccaulfield): COLLISION - see the NOTE below - fix it
+					SCOPE_CYCLE_COUNTER(STAT_Collisions_SpatialBroadPhase);
+					// @todo(ccaulfield): the spatial acceleration scheme needs to know the expanded bboxes for all particles, not just the one doing the test
+					if (bBody1Bounded)
+					{
+						// @todo(ccaulfield): COLLISION - see the NOTE below - fix it
 #if CHAOS_PARTICLEHANDLE_TODO
-					const TAABB<FReal, 3> Box1 = InSpatialAcceleration.GetWorldSpaceBoundingBox(Particle1);
+						const TAABB<FReal, 3> Box1 = InSpatialAcceleration.GetWorldSpaceBoundingBox(Particle1);
 #else
-					const TAABB<FReal, 3> Box1 = ComputeWorldSpaceBoundingBox(Particle1); // NOTE: this ignores the velocity expansion which is wrong
+						const TAABB<FReal, 3> Box1 = ComputeWorldSpaceBoundingBox(Particle1); // NOTE: this ignores the velocity expansion which is wrong
 #endif
 
-					CHAOS_COLLISION_STAT(StatData.RecordBoundsData(Box1));
-
-					FSimOverlapVisitor OverlapVisitor(PotentialIntersections);
-					InSpatialAcceleration.Overlap(Box1, OverlapVisitor);
-				}
-				else
-				{
-					const auto& GlobalElems = InSpatialAcceleration.GlobalObjects();
-					PotentialIntersections.Reserve(GlobalElems.Num());
+						CHAOS_COLLISION_STAT(StatData.RecordBoundsData(Box1));
 
-					for (auto& Elem : GlobalElems)
+						FSimOverlapVisitor OverlapVisitor(PotentialIntersections);
+						InSpatialAcceleration.Overlap(Box1, OverlapVisitor);
+					}
+					else
 					{
-						PotentialIntersections.Add(Elem.Payload);
+						const auto& GlobalElems = InSpatialAcceleration.GlobalObjects();
+						PotentialIntersections.Reserve(GlobalElems.Num());
+
+						for (auto& Elem : GlobalElems)
+						{
+							PotentialIntersections.Add(Elem.Payload);
+						}
 					}
-				}
 
-				CHAOS_COLLISION_STAT(StatData.RecordBroadphasePotentials(PotentialIntersections.Num()))
+					CHAOS_COLLISION_STAT(StatData.RecordBroadphasePotentials(PotentialIntersections.Num()))
+				}
 
+				SCOPE_CYCLE_COUNTER(STAT_Collisions_Filtering);
 				const int32 NumPotentials = PotentialIntersections.Num();
 				for (int32 i = 0; i < NumPotentials; ++i)
 				{
@@ -229,9 +238,12 @@ namespace Chaos
 					const FReal Box2Thickness = bIsParticle2Dynamic ? ComputeBoundsThickness(*Particle2.CastToRigidParticle(), Dt, BoundsThickness, BoundsThicknessVelocityInflation).Size()
 						: (bIsParticle2Kinematic ? ComputeBoundsThickness(*Particle2.CastToKinematicParticle(), Dt, BoundsThickness, BoundsThicknessVelocityInflation).Size() : (FReal)0);
 
-					FCollisionConstraintsArray NewConstraints;
-					NarrowPhase.GenerateCollisions(NewConstraints, Dt, Particle1.Handle(), Particle2.Handle(), FMath::Max(Box1Thickness, Box2Thickness), StatData);
-					Receiver.ReceiveCollisions(NewConstraints);
+					{
+						SCOPE_CYCLE_COUNTER(STAT_Collisions_SpatialNarrowPhase);
+						FCollisionConstraintsArray NewConstraints;
+						NarrowPhase.GenerateCollisions(NewConstraints, Dt, Particle1.Handle(), Particle2.Handle(), FMath::Max(Box1Thickness, Box2Thickness), StatData);
+						Receiver.ReceiveCollisions(NewConstraints);
+					}
 				}
 			}
 
