Hook up iterations from solver.

diff --git a/Engine/Source/Runtime/Experimental/ChaosSolverEngine/Private/Chaos/ChaosSolverActor.cpp b/Engine/Source/Runtime/Experimental/ChaosSolverEngine/Private/Chaos/ChaosSolverActor.cpp
--- a/Engine/Source/Runtime/Experimental/ChaosSolverEngine/Private/Chaos/ChaosSolverActor.cpp
+++ b/Engine/Source/Runtime/Experimental/ChaosSolverEngine/Private/Chaos/ChaosSolverActor.cpp
@@ -294,12 +294,12 @@ void AChaosSolverActor::BeginPlay()
 		{
 #if TODO_REIMPLEMENT_SOLVER_SETTINGS_ACCESSORS
 			InSolver->SetTimeStepMultiplier(InTimeStepMultiplier);
-			InSolver->SetIterations(InCollisionIterations);
-			InSolver->SetPushOutIterations(InPushOutIterations);
-			InSolver->SetPushOutPairIterations(InPushOutPairIterations);
 			InSolver->SetClusterConnectionFactor(InClusterConnectionFactor);
 			InSolver->SetClusterUnionConnectionType((Chaos::FClusterCreationParameters<float>::EConnectionMethod)InClusterUnionConnectionType);
 #endif
+			InSolver->SetIterations(InCollisionIterations);
+			InSolver->SetPushOutPairIterations(InPushOutPairIterations);
+			InSolver->SetPushOutIterations(InPushOutIterations);
 			InSolver->SetGenerateCollisionData(InDoGenerateCollisionData);
 			InSolver->SetGenerateBreakingData(InDoGenerateBreakingData);
 			InSolver->SetGenerateTrailingData(InDoGenerateTrailingData);
diff --git a/Engine/Source/Runtime/Experimental/ChaosSolvers/Public/PBDRigidsSolver.h b/Engine/Source/Runtime/Experimental/ChaosSolvers/Public/PBDRigidsSolver.h
--- a/Engine/Source/Runtime/Experimental/ChaosSolvers/Public/PBDRigidsSolver.h
+++ b/Engine/Source/Runtime/Experimental/ChaosSolvers/Public/PBDRigidsSolver.h
@@ -262,6 +262,11 @@ namespace Chaos
 		void SetMaxSubSteps(const int32 InMaxSubSteps) { MMaxSubSteps = InMaxSubSteps; }
 		int32 GetMaxSubSteps() const { return MMaxSubSteps; }
 
+		/**/
+		void SetIterations(const int32 InNumIterations) { GetEvolution()->SetNumIterations(InNumIterations); }
+		void SetPushOutIterations(const int32 InNumIterations) {  GetEvolution()->SetNumPushOutIterations(InNumIterations); }
+		void SetPushOutPairIterations(const int32 InNumIterations) {  GetEvolution()->GetCollisionConstraints().SetPushOutPairIterations(InNumIterations); }
+
 		/**/
 		void SetGenerateCollisionData(bool bDoGenerate) { GetEventFilters()->SetGenerateCollisionEvents(bDoGenerate); }
 		void SetGenerateBreakingData(bool bDoGenerate)
