Make narrow phase ifdefed.

diff --git a/Engine/Source/Runtime/Experimental/Chaos/Private/Chaos/Collision/CollisionDetector.cpp b/Engine/Source/Runtime/Experimental/Chaos/Private/Chaos/Collision/CollisionDetector.cpp
--- a/Engine/Source/Runtime/Experimental/Chaos/Private/Chaos/Collision/CollisionDetector.cpp
+++ b/Engine/Source/Runtime/Experimental/Chaos/Private/Chaos/Collision/CollisionDetector.cpp
@@ -10,7 +10,6 @@ namespace Chaos
 	DEFINE_STAT(STAT_Collisions_BroadPhase);
 	DEFINE_STAT(STAT_Collisions_SpatialBroadPhase);
 	DEFINE_STAT(STAT_Collisions_Filtering);
-	DEFINE_STAT(STAT_Collisions_SpatialNarrowPhase);
 #if CHAOS_ENABLE_STAT_NARROWPHASE
 	DEFINE_STAT(STAT_Collisions_NarrowPhase);
 #endif
diff --git a/Engine/Source/Runtime/Experimental/Chaos/Public/Chaos/Collision/NarrowPhase.h b/Engine/Source/Runtime/Experimental/Chaos/Public/Chaos/Collision/NarrowPhase.h
--- a/Engine/Source/Runtime/Experimental/Chaos/Public/Chaos/Collision/NarrowPhase.h
+++ b/Engine/Source/Runtime/Experimental/Chaos/Public/Chaos/Collision/NarrowPhase.h
@@ -15,7 +15,7 @@ namespace Chaos
 #define CHAOS_ENABLE_STAT_NARROWPHASE 0
 #if CHAOS_ENABLE_STAT_NARROWPHASE
 	DECLARE_CYCLE_STAT_EXTERN(TEXT("Collisions::NarrowPhase"), STAT_Collisions_NarrowPhase, STATGROUP_ChaosCollision, CHAOS_API);
-	#define SCOPE_CYCLE_COUNTER_NAROWPHASE() SCOPE_CYCLE_COUNTER(STAT_Collisions_GJK)
+	#define SCOPE_CYCLE_COUNTER_NAROWPHASE() SCOPE_CYCLE_COUNTER(STAT_Collisions_NarrowPhase)
 #else
 	#define SCOPE_CYCLE_COUNTER_NAROWPHASE()
 #endif
diff --git a/Engine/Source/Runtime/Experimental/Chaos/Public/Chaos/Collision/SpatialAccelerationBroadPhase.h b/Engine/Source/Runtime/Experimental/Chaos/Public/Chaos/Collision/SpatialAccelerationBroadPhase.h
--- a/Engine/Source/Runtime/Experimental/Chaos/Public/Chaos/Collision/SpatialAccelerationBroadPhase.h
+++ b/Engine/Source/Runtime/Experimental/Chaos/Public/Chaos/Collision/SpatialAccelerationBroadPhase.h
@@ -14,7 +14,6 @@
 namespace Chaos
 {
 	DECLARE_CYCLE_STAT_EXTERN(TEXT("Collisions::BroadPhase"), STAT_Collisions_SpatialBroadPhase, STATGROUP_ChaosCollision, CHAOS_API);
-	DECLARE_CYCLE_STAT_EXTERN(TEXT("Collisions::NarrowPhase"), STAT_Collisions_SpatialNarrowPhase, STATGROUP_ChaosCollision, CHAOS_API);
 	DECLARE_CYCLE_STAT_EXTERN(TEXT("Collisions::Filtering"), STAT_Collisions_Filtering, STATGROUP_ChaosCollision, CHAOS_API);
 	
 	template <typename TPayloadType, typename T, int d>
@@ -238,12 +237,9 @@ namespace Chaos
 					const FReal Box2Thickness = bIsParticle2Dynamic ? ComputeBoundsThickness(*Particle2.CastToRigidParticle(), Dt, BoundsThickness, BoundsThicknessVelocityInflation).Size()
 						: (bIsParticle2Kinematic ? ComputeBoundsThickness(*Particle2.CastToKinematicParticle(), Dt, BoundsThickness, BoundsThicknessVelocityInflation).Size() : (FReal)0);
 
-					{
-						SCOPE_CYCLE_COUNTER(STAT_Collisions_SpatialNarrowPhase);
-						FCollisionConstraintsArray NewConstraints;
-						NarrowPhase.GenerateCollisions(NewConstraints, Dt, Particle1.Handle(), Particle2.Handle(), FMath::Max(Box1Thickness, Box2Thickness), StatData);
-						Receiver.ReceiveCollisions(NewConstraints);
-					}
+					FCollisionConstraintsArray NewConstraints;
+					NarrowPhase.GenerateCollisions(NewConstraints, Dt, Particle1.Handle(), Particle2.Handle(), FMath::Max(Box1Thickness, Box2Thickness), StatData);
+					Receiver.ReceiveCollisions(NewConstraints);
 				}
 			}
 
