Fixed cluster unions being filtered out of collision detection due to internally created clusters not having valid colllision filter data

diff --git a/Engine/Source/Runtime/Experimental/Chaos/Private/Chaos/PBDRigidClustering.cpp b/Engine/Source/Runtime/Experimental/Chaos/Private/Chaos/PBDRigidClustering.cpp
--- a/Engine/Source/Runtime/Experimental/Chaos/Private/Chaos/PBDRigidClustering.cpp
+++ b/Engine/Source/Runtime/Experimental/Chaos/Private/Chaos/PBDRigidClustering.cpp
@@ -1339,6 +1339,32 @@ namespace Chaos
 		bool bUseParticleImplicit = false;
 		bool bUsingMultiChildProxy = false;
 
+		// Need to extract a filter off one of the cluster children 
+		FCollisionFilterData Filter;
+		for(TPBDRigidParticleHandle<T, d>* Child : Children)
+		{
+			bool bFilterValid = false;
+			for(const TUniquePtr<TPerShapeData<FReal, 3>>& Shape : Child->ShapesArray())
+			{
+				if(Shape)
+				{
+					Filter = Shape->SimData;
+					bFilterValid = Filter.Word0 != 0 || Filter.Word1 != 0 || Filter.Word2 != 0 || Filter.Word3 != 0;
+				}
+
+				if(bFilterValid)
+				{
+					break;
+				}
+			}
+
+			// Bail once we've found one
+			if(bFilterValid)
+			{
+				break;
+			}
+		}
+
 		{ // STAT_UpdateGeometry_GatherObjects
 			SCOPE_CYCLE_COUNTER(STAT_UpdateGeometry_GatherObjects);
 
@@ -1542,6 +1568,12 @@ namespace Chaos
 			const Chaos::TAABB<float, 3> TransformedBBox = LocalBounds.TransformedAABB(Xf);
 			Parent->SetWorldSpaceInflatedBounds(TransformedBBox);
 		}
+
+		// Set the captured filter to our new shapes
+		for(const TUniquePtr<TPerShapeData<FReal, 3>>& Shape : Parent->ShapesArray())
+		{
+			Shape->SimData = Filter;
+		}
 	}
 
 	float MinImpulseForStrainEval = 980 * 2 * 1.f / 30.f; //ignore impulses caused by just keeping object on ground. This is a total hack, we should not use accumulated impulse directly. Instead we need to look at delta v along constraint normal
