Filter cluster collisions.

diff --git a/Engine/Source/Runtime/Experimental/Chaos/Private/Chaos/CollisionResolution.cpp b/Engine/Source/Runtime/Experimental/Chaos/Private/Chaos/CollisionResolution.cpp
--- a/Engine/Source/Runtime/Experimental/Chaos/Private/Chaos/CollisionResolution.cpp
+++ b/Engine/Source/Runtime/Experimental/Chaos/Private/Chaos/CollisionResolution.cpp
@@ -30,7 +30,7 @@ DECLARE_CYCLE_STAT(TEXT("Collisions::GJK"), STAT_Collisions_GJK, STATGROUP_Chaos
 #define SCOPE_CYCLE_COUNTER_GJK()
 #endif
 
-#define CHAOS_COLLIDE_CLUSTERED_UNIONS 0
+#define CHAOS_COLLIDE_CLUSTERED_UNIONS 1
 
 //PRAGMA_DISABLE_OPTIMIZATION
 
@@ -2611,9 +2611,22 @@ namespace Chaos
 			if (Implicit0OuterType == FImplicitObjectUnionClustered::StaticType())
 			{
 				const FImplicitObjectUnionClustered* Union0 = Implicit0->template GetObject<FImplicitObjectUnionClustered>();
-				for (const auto& Child0 : Union0->GetObjects())
+				if (Implicit1->HasBoundingBox())
+				{
+					TArray<Pair<const FImplicitObject*, FRigidTransform3>> Children;
+					Union0->FindAllIntersectingObjects(Children, Implicit1->BoundingBox());
+					for (const auto& Child0 : Children)
+					{
+						TRigidTransform<T, d> TransformedChild0 = Child0.Second * LocalTransform0;
+						ConstructConstraints(Particle0, Particle1, Child0.First, Implicit1, TransformedChild0, LocalTransform1, CullDistance, Context, NewConstraints);
+					}
+				}
+				else
 				{
-					ConstructConstraints(Particle0, Particle1, Child0.Get(), Implicit1, LocalTransform0, LocalTransform1, CullDistance, Context, NewConstraints);
+					for (const auto& Child0 : Union0->GetObjects())
+					{
+						ConstructConstraints(Particle0, Particle1, Child0.Get(), Implicit1, LocalTransform0, LocalTransform1, CullDistance, Context, NewConstraints);
+					}
 				}
 				return;
 			}
@@ -2633,9 +2646,22 @@ namespace Chaos
 			if (Implicit1OuterType == FImplicitObjectUnionClustered::StaticType())
 			{
 				const FImplicitObjectUnionClustered* Union1 = Implicit1->template GetObject<FImplicitObjectUnionClustered>();
-				for (const auto& Child1 : Union1->GetObjects())
+				if (Implicit0->HasBoundingBox())
+				{
+					TArray<Pair<const FImplicitObject*, FRigidTransform3>> Children;
+					Union1->FindAllIntersectingObjects(Children, Implicit0->BoundingBox());
+					for (const auto& Child1 : Children)
+					{
+						TRigidTransform<T, d> TransformedChild1 = Child1.Second * LocalTransform1;
+						ConstructConstraints(Particle0, Particle1, Implicit0, Child1.First, LocalTransform0, TransformedChild1, CullDistance, Context, NewConstraints);
+					}
+				}
+				else
 				{
-					ConstructConstraints(Particle0, Particle1, Implicit0, Child1.Get(), LocalTransform0, LocalTransform1, CullDistance, Context, NewConstraints);
+					for (const auto& Child1 : Union1->GetObjects())
+					{
+						ConstructConstraints(Particle0, Particle1, Implicit0, Child1.Get(), LocalTransform0, LocalTransform1, CullDistance, Context, NewConstraints);
+					}
 				}
 				return;
 			}
