Reconnect ether drag properties to the simulator for geometry collections.  Move linear and angular ether drag properties from the GeometryCollectionComponent to UChaosPhysicalMaterial.  Add static (angular) friction propterty to UChaosPhysicalMaterial.

diff --git a/Engine/Source/Runtime/Engine/Classes/PhysicalMaterials/Experimental/ChaosPhysicalMaterial.h b/Engine/Source/Runtime/Engine/Classes/PhysicalMaterials/Experimental/ChaosPhysicalMaterial.h
--- a/Engine/Source/Runtime/Engine/Classes/PhysicalMaterials/Experimental/ChaosPhysicalMaterial.h
+++ b/Engine/Source/Runtime/Engine/Classes/PhysicalMaterials/Experimental/ChaosPhysicalMaterial.h
@@ -13,26 +13,44 @@
 #include "Physics/PhysicsInterfaceCore.h"
 #include "ChaosPhysicalMaterial.generated.h"
 
+class Chaos::FChaosPhysicsMaterial;
+
 /**
- * Physical materials are used to define the response of a physical object when interacting dynamically with the world.
+ * Physical materials are used to define the response of a physical object when 
+ * interacting dynamically with the world.
  */
 UCLASS(BlueprintType, Blueprintable, CollapseCategories, HideCategories = Object)
 class ENGINE_API UChaosPhysicalMaterial : public UObject
 {
 	GENERATED_UCLASS_BODY()
 
+	/** Update the values in \p Mat to reflect the values in this \c UObject. */
+	void CopyTo(Chaos::FChaosPhysicsMaterial& Mat) const;
+
 	//
 	// Surface properties.
 	//
 	
-	/** Friction value of surface, controls how easily things can slide on this surface (0 is frictionless, higher values increase the amount of friction) */
+	/** Friction value of a surface in motion, controls how easily things can slide on this surface (0 is frictionless, higher values increase the amount of friction) */
 	UPROPERTY(EditAnywhere, BlueprintReadOnly, Category = PhysicalMaterial, meta=(ClampMin=0))
 	float Friction;
-	
+
+	/** Friction value of surface at rest, controls how easily things can slide on this surface (0 is frictionless, higher values increase the amount of friction) */
+	UPROPERTY(EditAnywhere, BlueprintReadOnly, Category = PhysicalMaterial, meta = (ClampMin = 0))
+	float StaticFriction;
+
 	/** Restitution or 'bounciness' of this surface, between 0 (no bounce) and 1 (outgoing velocity is same as incoming). */
 	UPROPERTY(EditAnywhere, BlueprintReadOnly, Category = PhysicalMaterial, meta=(ClampMin=0))
 	float Restitution;
 
+	/** Uniform linear ether drag, the resistance a body experiences to its translation. */
+	UPROPERTY(EditAnywhere, BlueprintReadOnly, Category = "PhysicalMaterial", meta=(ClampMin = 0))
+	float LinearEtherDrag;
+
+	/** Uniform angular ether drag, the resistance a body experiences to its rotation. */
+	UPROPERTY(EditAnywhere, BlueprintReadOnly, Category = "PhysicalMaterial", meta = (ClampMin = 0))
+	float AngularEtherDrag;
+
 	/** How much to scale the damage threshold by on any destructible we are applied to */
 	UPROPERTY(EditAnywhere, BlueprintReadOnly, Category = PhysicalMaterial, meta=(ClampMin=0))
 	float SleepingLinearVelocityThreshold;
diff --git a/Engine/Source/Runtime/Engine/Private/PhysicsEngine/Experimental/ChaosPhysicalMaterial.cpp b/Engine/Source/Runtime/Engine/Private/PhysicsEngine/Experimental/ChaosPhysicalMaterial.cpp
--- a/Engine/Source/Runtime/Engine/Private/PhysicsEngine/Experimental/ChaosPhysicalMaterial.cpp
+++ b/Engine/Source/Runtime/Engine/Private/PhysicsEngine/Experimental/ChaosPhysicalMaterial.cpp
@@ -5,12 +5,30 @@
 =============================================================================*/ 
 
 #include "PhysicalMaterials/Experimental/ChaosPhysicalMaterial.h"
+#include "Chaos/Defines.h"
 
 UChaosPhysicalMaterial::UChaosPhysicalMaterial(const FObjectInitializer& ObjectInitializer)
 	: Super(ObjectInitializer)
 {
-	Friction = 0.5f;
-	Restitution = 0.1f;
-	SleepingAngularVelocityThreshold = 1.f;
-	SleepingLinearVelocityThreshold = 1.f;
+	// Initialize values from FChaosPhysicsMaterial
+	Chaos::FChaosPhysicsMaterial TmpMat;
+
+	Friction = TmpMat.Friction;
+	StaticFriction = TmpMat.StaticFriction;
+	Restitution = TmpMat.Restitution;
+	LinearEtherDrag = TmpMat.LinearEtherDrag;
+	AngularEtherDrag = TmpMat.AngularEtherDrag;
+	SleepingAngularVelocityThreshold = TmpMat.SleepingAngularThreshold;
+	SleepingLinearVelocityThreshold = TmpMat.SleepingLinearThreshold;
+}
+
+void UChaosPhysicalMaterial::CopyTo(Chaos::FChaosPhysicsMaterial& Mat) const
+{
+	Mat.Friction = Friction;
+	Mat.StaticFriction = StaticFriction;
+	Mat.Restitution = Restitution;
+	Mat.LinearEtherDrag = LinearEtherDrag;
+	Mat.AngularEtherDrag = AngularEtherDrag;
+	Mat.SleepingLinearThreshold = SleepingLinearVelocityThreshold;
+	Mat.SleepingAngularThreshold = SleepingAngularVelocityThreshold;
 }
diff --git a/Engine/Source/Runtime/Experimental/Chaos/Private/Chaos/PBDCollisionConstraints.cpp b/Engine/Source/Runtime/Experimental/Chaos/Private/Chaos/PBDCollisionConstraints.cpp
--- a/Engine/Source/Runtime/Experimental/Chaos/Private/Chaos/PBDCollisionConstraints.cpp
+++ b/Engine/Source/Runtime/Experimental/Chaos/Private/Chaos/PBDCollisionConstraints.cpp
@@ -99,7 +99,6 @@ namespace Chaos
 		, MApplyPushOutPairIterations(InApplyPushOutPairIterations)
 		, MCullDistance(CullDistance)
 		, MShapePadding(ShapePadding)
-		, MAngularFriction(0)
 		, bUseCCD(false)
 		, bEnableCollisions(true)
 		, bHandlesEnabled(true)
@@ -182,23 +181,26 @@ namespace Chaos
 
 			const FChaosPhysicsMaterial::ECombineMode FrictionCombineMode = FChaosPhysicsMaterial::ChooseCombineMode(PhysicsMaterial0->FrictionCombineMode,PhysicsMaterial1->FrictionCombineMode);
 			Contact.Friction = FChaosPhysicsMaterial::CombineHelper(PhysicsMaterial0->Friction,PhysicsMaterial1->Friction, FrictionCombineMode);
+			Contact.AngularFriction = FChaosPhysicsMaterial::CombineHelper(PhysicsMaterial0->StaticFriction, PhysicsMaterial1->StaticFriction, FrictionCombineMode);
 		}
 		else if (PhysicsMaterial0)
 		{
 			Contact.Restitution = PhysicsMaterial0->Restitution;
 			Contact.Friction = PhysicsMaterial0->Friction;
+			Contact.AngularFriction = PhysicsMaterial0->StaticFriction;
 		}
 		else if (PhysicsMaterial1)
 		{
 			Contact.Restitution = PhysicsMaterial1->Restitution;
 			Contact.Friction = PhysicsMaterial1->Friction;
+			Contact.AngularFriction = PhysicsMaterial1->StaticFriction;
 		}
 		else
 		{
 			Contact.Friction = DefaultCollisionFriction;
+			Contact.AngularFriction = 0;
 			Contact.Restitution = DefaultCollisionRestitution;
 		}
-		Contact.AngularFriction = MAngularFriction;
 
 		// Overrides for testing
 		if (CollisionFrictionOverride >= 0)
@@ -318,7 +320,6 @@ namespace Chaos
 		Handles.Reset();
 #endif
 
-		MAngularFriction = 0;
 		bUseCCD = false;
 	}
 
@@ -654,7 +655,7 @@ namespace Chaos
 					Collisions::ApplyPushOut(ConstraintHandle->GetContact(), IsTemporarilyStatic, IterationParameters, ParticleParameters);
 				}
 
-			}, bDisableCollisionParallelFor);
+			}, /*bDisableCollisionParallelFor*/true);
 		}
 
 		if (PostApplyPushOutCallback != nullptr)
diff --git a/Engine/Source/Runtime/Experimental/Chaos/Public/Chaos/Defines.h b/Engine/Source/Runtime/Experimental/Chaos/Public/Chaos/Defines.h
--- a/Engine/Source/Runtime/Experimental/Chaos/Public/Chaos/Defines.h
+++ b/Engine/Source/Runtime/Experimental/Chaos/Public/Chaos/Defines.h
@@ -37,7 +37,10 @@ namespace Chaos
 		};
 
 		FReal Friction;
+		FReal StaticFriction;
 		FReal Restitution;
+		FReal LinearEtherDrag;
+		FReal AngularEtherDrag;
 		FReal SleepingLinearThreshold;
 		FReal SleepingAngularThreshold;
 		FReal DisabledLinearThreshold;
@@ -45,13 +48,17 @@ namespace Chaos
 		int32 SleepCounterThreshold;
 		void* UserData;
 
+		/** Variable defaults are used for \c UChaosPhysicalMaterial \c UPROPERTY defaults. */
 		ECombineMode FrictionCombineMode;
 		ECombineMode RestitutionCombineMode;
 
 
 		FChaosPhysicsMaterial()
 			: Friction(0.5)
+			, StaticFriction(0.0)
 			, Restitution(0.1)
+			, LinearEtherDrag(0.0)
+			, AngularEtherDrag(0.0)
 			, SleepingLinearThreshold(1)
 			, SleepingAngularThreshold(1)
 			, DisabledLinearThreshold(0)
@@ -102,7 +109,7 @@ namespace Chaos
 		{
 			Ar.UsingCustomVersion(FExternalPhysicsCustomObjectVersion::GUID);
 
-			Ar << Friction << Restitution << SleepingLinearThreshold << SleepingAngularThreshold << DisabledLinearThreshold << DisabledAngularThreshold;
+			Ar << Friction << StaticFriction << Restitution << LinearEtherDrag << AngularEtherDrag << SleepingLinearThreshold << SleepingAngularThreshold << DisabledLinearThreshold << DisabledAngularThreshold;
 
 			if (Ar.CustomVer(FExternalPhysicsCustomObjectVersion::GUID) >= FExternalPhysicsCustomObjectVersion::PhysicsMaterialSleepCounterThreshold)
 			{
diff --git a/Engine/Source/Runtime/Experimental/Chaos/Public/Chaos/PBDCollisionConstraints.h b/Engine/Source/Runtime/Experimental/Chaos/Public/Chaos/PBDCollisionConstraints.h
--- a/Engine/Source/Runtime/Experimental/Chaos/Public/Chaos/PBDCollisionConstraints.h
+++ b/Engine/Source/Runtime/Experimental/Chaos/Public/Chaos/PBDCollisionConstraints.h
@@ -286,7 +286,6 @@ private:
 	int32 MApplyPushOutPairIterations;
 	FReal MCullDistance;
 	FReal MShapePadding;
-	FReal MAngularFriction;
 	bool bUseCCD;
 	bool bEnableCollisions;
 	bool bHandlesEnabled;
diff --git a/Engine/Source/Runtime/Experimental/ChaosSolvers/Private/PhysicsProxy/GeometryCollectionPhysicsProxy.cpp b/Engine/Source/Runtime/Experimental/ChaosSolvers/Private/PhysicsProxy/GeometryCollectionPhysicsProxy.cpp
--- a/Engine/Source/Runtime/Experimental/ChaosSolvers/Private/PhysicsProxy/GeometryCollectionPhysicsProxy.cpp
+++ b/Engine/Source/Runtime/Experimental/ChaosSolvers/Private/PhysicsProxy/GeometryCollectionPhysicsProxy.cpp
@@ -657,6 +657,15 @@ void FGeometryCollectionPhysicsProxy::InitializeBodiesPT(Chaos::FPBDRigidsSolver
 	const FGeometryCollection* RestCollection = Parameters.RestCollection;
 	const FGeometryDynamicCollection& DynamicCollection = PhysicsThreadCollection;
 
+	if (!Parameters.PhysicalMaterial)
+	{
+		// If PhysicalMaterial is non-null, then we have a physics material in the scene 
+		// that we're using.  If not, then we need to allocate a default FChaosPhysicsMaterial
+		// that we'll point PhysicalMaterial at.
+		Parameters.PhysicalMaterialOwner = TUniquePtr<Chaos::FChaosPhysicsMaterial>(new Chaos::FChaosPhysicsMaterial());
+		Parameters.PhysicalMaterial = Chaos::MakeSerializable(Parameters.PhysicalMaterialOwner);
+	}
+
 	if (Parameters.Simulating)
 	{
 		const TManagedArray<int32>& TransformIndex = RestCollection->TransformIndex;
@@ -767,7 +776,8 @@ void FGeometryCollectionPhysicsProxy::InitializeBodiesPT(Chaos::FPBDRigidsSolver
 					CollisionParticles->Resize(CollisionParticlesSize); // Truncates!
 				}
 
-
+				Handle->SetLinearEtherDrag(Parameters.PhysicalMaterial->LinearEtherDrag);
+				Handle->SetAngularEtherDrag(Parameters.PhysicalMaterial->AngularEtherDrag);
 				RigidsSolver->GetEvolution()->SetPhysicsMaterial(Handle, Parameters.PhysicalMaterial);
 			}
 		});
@@ -1088,6 +1098,11 @@ FGeometryCollectionPhysicsProxy::BuildClusters(
 
 	Parent->SetStrains(Damage);
 
+	// If no PhysicalMaterial was specified by the caller, then a default one was constructed 
+	// for us by InitializeBodiesPT().
+	check(Parameters.PhysicalMaterial);
+	Parent->SetLinearEtherDrag(Parameters.PhysicalMaterial->LinearEtherDrag);
+	Parent->SetAngularEtherDrag(Parameters.PhysicalMaterial->AngularEtherDrag);
 	GetSolver()->GetEvolution()->SetPhysicsMaterial(Parent, Parameters.PhysicalMaterial);
 
 	const FTransform ParentTransform = GeometryCollectionAlgo::GlobalMatrix(DynamicCollection.Transform, DynamicCollection.Parent, CollectionClusterIndex);
diff --git a/Engine/Source/Runtime/Experimental/GeometryCollectionEngine/Private/GeometryCollection/GeometryCollectionComponent.cpp b/Engine/Source/Runtime/Experimental/GeometryCollectionEngine/Private/GeometryCollection/GeometryCollectionComponent.cpp
--- a/Engine/Source/Runtime/Experimental/GeometryCollectionEngine/Private/GeometryCollection/GeometryCollectionComponent.cpp
+++ b/Engine/Source/Runtime/Experimental/GeometryCollectionEngine/Private/GeometryCollection/GeometryCollectionComponent.cpp
@@ -111,8 +111,6 @@ UGeometryCollectionComponent::UGeometryCollectionComponent(const FObjectInitiali
 	, ClusterConnectionType(EClusterConnectionTypeEnum::Chaos_PointImplicit)
 	, CollisionGroup(0)
 	, CollisionSampleFraction(1.0)
-	, LinearEtherDrag(0.0)
-	, AngularEtherDrag(0.0)
 	, InitialVelocityType(EInitialVelocityTypeEnum::Chaos_Initial_Velocity_User_Defined)
 	, InitialLinearVelocity(0.f, 0.f, 0.f)
 	, InitialAngularVelocity(0.f, 0.f, 0.f)
@@ -145,8 +143,6 @@ UGeometryCollectionComponent::UGeometryCollectionComponent(const FObjectInitiali
 	GlobalNavMeshInvalidationCounter += 3;
 	NavmeshInvalidationTimeSliceIndex = GlobalNavMeshInvalidationCounter;
 
-	ChaosMaterial = MakeUnique<Chaos::FChaosPhysicsMaterial>();
-
 	WorldBounds = FBoxSphereBounds(FBox(ForceInit));	
 
 	// default current cache time
@@ -1114,12 +1110,13 @@ void UGeometryCollectionComponent::OnCreatePhysicsState()
 		const bool bValidCollection = DynamicCollection && DynamicCollection->Transform.Num() > 0;
 		if (bValidWorld && bValidCollection)
 		{
+			if (!ChaosMaterial)
+			{
+				ChaosMaterial.Reset(new Chaos::FChaosPhysicsMaterial());
+			}
 			if (PhysicalMaterial)
 			{
-				ChaosMaterial->Friction = PhysicalMaterial->Friction;
-				ChaosMaterial->Restitution = PhysicalMaterial->Restitution;
-				ChaosMaterial->SleepingLinearThreshold = PhysicalMaterial->SleepingLinearVelocityThreshold;
-				ChaosMaterial->SleepingAngularThreshold = PhysicalMaterial->SleepingAngularVelocityThreshold;
+				PhysicalMaterial->CopyTo(*ChaosMaterial);
 			}
 
 			FPhysxUserData::Set<UPrimitiveComponent>(&PhysicsUserData, this);
@@ -1142,8 +1139,6 @@ void UGeometryCollectionComponent::OnCreatePhysicsState()
 				SimulationParameters.ClusterConnectionMethod = (Chaos::FClusterCreationParameters<float>::EConnectionMethod)ClusterConnectionType;
 				SimulationParameters.CollisionGroup = CollisionGroup;
 				SimulationParameters.CollisionSampleFraction = CollisionSampleFraction;
-				SimulationParameters.LinearEtherDrag = LinearEtherDrag;
-				SimulationParameters.AngularEtherDrag = AngularEtherDrag;
 				SimulationParameters.InitialVelocityType = InitialVelocityType;
 				SimulationParameters.InitialLinearVelocity = InitialLinearVelocity;
 				SimulationParameters.InitialAngularVelocity = InitialAngularVelocity;
@@ -1171,6 +1166,7 @@ void UGeometryCollectionComponent::OnCreatePhysicsState()
 				SimulationParameters.RemoveOnFractureEnabled = SimulationParameters.Shared.RemoveOnFractureIndices.Num() > 0;
 				SimulationParameters.WorldTransform = GetComponentToWorld();
 				SimulationParameters.UserData = static_cast<void*>(&PhysicsUserData);
+				SimulationParameters.PhysicalMaterial = Chaos::MakeSerializable(ChaosMaterial);
 			}
 
 
diff --git a/Engine/Source/Runtime/Experimental/GeometryCollectionEngine/Public/GeometryCollection/GeometryCollectionComponent.h b/Engine/Source/Runtime/Experimental/GeometryCollectionEngine/Public/GeometryCollection/GeometryCollectionComponent.h
--- a/Engine/Source/Runtime/Experimental/GeometryCollectionEngine/Public/GeometryCollection/GeometryCollectionComponent.h
+++ b/Engine/Source/Runtime/Experimental/GeometryCollectionEngine/Public/GeometryCollection/GeometryCollectionComponent.h
@@ -415,12 +415,12 @@ public:
 	float CollisionSampleFraction;
 
 	/** Uniform linear ether drag. */
-	UPROPERTY(EditAnywhere, BlueprintReadWrite, Category = "ChaosPhysics|Ether Drag")
-	float LinearEtherDrag;
+	UPROPERTY(meta = (DeprecatedProperty, DeprecationMessage = "Use PhysicalMaterial instead."))
+	float LinearEtherDrag_DEPRECATED;
 
 	/** Uniform angular ether drag. */
-	UPROPERTY(EditAnywhere, BlueprintReadWrite, Category = "ChaosPhysics|Ether Drag")
-	float AngularEtherDrag;
+	UPROPERTY(meta=(DeprecatedProperty, DeprecationMessage="Use PhysicalMaterial instead."))
+	float AngularEtherDrag_DEPRECATED;
 
 	/** Physical Properties */
 	UPROPERTY(EditAnywhere, BlueprintReadOnly, Category = "ChaosPhysics")
diff --git a/Engine/Source/Runtime/Experimental/GeometryCollectionSimulationCore/Public/GeometryCollection/GeometryCollectionSimulationCoreTypes.h b/Engine/Source/Runtime/Experimental/GeometryCollectionSimulationCore/Public/GeometryCollection/GeometryCollectionSimulationCoreTypes.h
--- a/Engine/Source/Runtime/Experimental/GeometryCollectionSimulationCore/Public/GeometryCollection/GeometryCollectionSimulationCoreTypes.h
+++ b/Engine/Source/Runtime/Experimental/GeometryCollectionSimulationCore/Public/GeometryCollection/GeometryCollectionSimulationCoreTypes.h
@@ -219,6 +219,8 @@ struct FTrailingDataSimulationParameters
 	float TrailingMinVolumeThreshold;
 };
 
+class FGeometryCollectionPhysicsProxy;
+
 struct FSimulationParameters
 {
 	FSimulationParameters()
@@ -235,8 +237,6 @@ struct FSimulationParameters
 		, ClusterConnectionMethod(Chaos::FClusterCreationParameters<float>::EConnectionMethod::PointImplicit)
 		, CollisionGroup(0)
 		, CollisionSampleFraction(1.0)
-		, LinearEtherDrag(0.0)
-		, AngularEtherDrag(0.0)
 		, InitialVelocityType(EInitialVelocityTypeEnum::Chaos_Initial_Velocity_None)
 		, InitialLinearVelocity(FVector(0))
 		, InitialAngularVelocity(FVector(0))
@@ -250,9 +250,9 @@ struct FSimulationParameters
 		, UserData(nullptr)
 	{}
 
-
 	FSimulationParameters(const FSimulationParameters& Other)
-		: RestCollection(Other.RestCollection)
+		: Name(Other.Name)
+		, RestCollection(Other.RestCollection)
 		, InitializationCommands(Other.InitializationCommands)
 		, RecordedTrack(Other.RecordedTrack)
 		, bOwnsTrack(false)
@@ -265,8 +265,6 @@ struct FSimulationParameters
 		, ClusterConnectionMethod(Other.ClusterConnectionMethod)
 		, CollisionGroup(Other.CollisionGroup)
 		, CollisionSampleFraction(Other.CollisionSampleFraction)
-		, LinearEtherDrag(Other.LinearEtherDrag)
-		, AngularEtherDrag(Other.AngularEtherDrag)
 		, InitialVelocityType(Other.InitialVelocityType)
 		, InitialLinearVelocity(Other.InitialLinearVelocity)
 		, InitialAngularVelocity(Other.InitialAngularVelocity)
@@ -283,7 +281,13 @@ struct FSimulationParameters
 		, SimulationFilterData(Other.SimulationFilterData)
 		, QueryFilterData(Other.QueryFilterData)
 		, UserData(Other.UserData)
-	{}
+	{
+		// Check to make sure we're not expecting the PhyicalMaterialOwner to be copied,
+		// which we can't because it's a TUniquePtr.  We'd need a non const version of this
+		// function to move the pointer.  However, the way the code is currently, this should
+		// never need to happen.
+		check(!Other.PhysicalMaterialOwner);
+	}
 
 	~FSimulationParameters()
 	{
@@ -315,9 +319,6 @@ struct FSimulationParameters
 	int32 CollisionGroup;
 	float CollisionSampleFraction;
 
-	float LinearEtherDrag;
-	float AngularEtherDrag;
-
 	EInitialVelocityTypeEnum InitialVelocityType;
 	FVector InitialLinearVelocity;
 	FVector InitialAngularVelocity;
@@ -329,7 +330,11 @@ struct FSimulationParameters
 
 	EObjectStateTypeEnum ObjectType;
 
-	Chaos::TSerializablePtr<Chaos::FChaosPhysicsMaterial> PhysicalMaterial;
+private:
+	friend class ::FGeometryCollectionPhysicsProxy;
+	TUniquePtr<Chaos::FChaosPhysicsMaterial> PhysicalMaterialOwner; // can be null
+public:
+	Chaos::TSerializablePtr<Chaos::FChaosPhysicsMaterial> PhysicalMaterial; // may or may not point to PhysicalMaterialOwner
 
 	FCollisionDataSimulationParameters CollisionData;
 	FBreakingDataSimulationParameters BreakingData;
