Re-enable parallel pushout for Chaos physics

diff --git a/Engine/Source/Runtime/Experimental/Chaos/Private/Chaos/PBDCollisionConstraints.cpp b/Engine/Source/Runtime/Experimental/Chaos/Private/Chaos/PBDCollisionConstraints.cpp
--- a/Engine/Source/Runtime/Experimental/Chaos/Private/Chaos/PBDCollisionConstraints.cpp
+++ b/Engine/Source/Runtime/Experimental/Chaos/Private/Chaos/PBDCollisionConstraints.cpp
@@ -655,7 +655,7 @@ namespace Chaos
 					Collisions::ApplyPushOut(ConstraintHandle->GetContact(), IsTemporarilyStatic, IterationParameters, ParticleParameters);
 				}
 
-			}, /*bDisableCollisionParallelFor*/true);
+			}, bDisableCollisionParallelFor);
 		}
 
 		if (PostApplyPushOutCallback != nullptr)
