Make using contact graph runtime ans sovler specified.

diff --git a/Engine/Source/Runtime/Experimental/Chaos/Private/Chaos/PBDConstraintColor.cpp b/Engine/Source/Runtime/Experimental/Chaos/Private/Chaos/PBDConstraintColor.cpp
--- a/Engine/Source/Runtime/Experimental/Chaos/Private/Chaos/PBDConstraintColor.cpp
+++ b/Engine/Source/Runtime/Experimental/Chaos/Private/Chaos/PBDConstraintColor.cpp
@@ -19,17 +19,18 @@
 
 using namespace Chaos;
 
+DECLARE_CYCLE_STAT(TEXT("FPBDConstraintColor::ComputeColors"), STAT_Constraint_ComputeColor, STATGROUP_Chaos);
+DECLARE_CYCLE_STAT(TEXT("FPBDConstraintColor::ComputeContactGraph"), STAT_Constraint_ComputeContactGraph, STATGROUP_Chaos);
+DECLARE_CYCLE_STAT(TEXT("FPBDConstraintColor::ComputeIslandColoring"), STAT_Constraint_ComputeIslandColoring, STATGROUP_Chaos);
+
 void FPBDConstraintColor::ComputeIslandColoring(const int32 Island, const FPBDConstraintGraph& ConstraintGraph, uint32 ContainerId)
 {
+	SCOPE_CYCLE_COUNTER(STAT_Constraint_ComputeIslandColoring);
 	const TArray<TGeometryParticleHandle<FReal, 3>*>& IslandParticles = ConstraintGraph.GetIslandParticles(Island);
 	FLevelToColorToConstraintListMap& LevelToColorToConstraintListMap = IslandData[Island].LevelToColorToConstraintListMap;
 	int32& MaxColor = IslandData[Island].MaxColor;
 
-#ifdef USE_CONTACT_LEVELS
 	const int32 MaxLevel = IslandData[Island].MaxLevel;
-#else
-	const int32 MaxLevel = 0;
-#endif
 	
 	LevelToColorToConstraintListMap.Reset();
 	LevelToColorToConstraintListMap.SetNum(MaxLevel + 1);
@@ -126,11 +127,7 @@ void FPBDConstraintColor::ComputeIslandColoring(const int32 Island, const FPBDCo
 					ColorNode.NextColor++;
 				}
 
-#ifdef USE_CONTACT_LEVELS
 				int32 Level = ColorEdge.Level;
-#else
-				int32 Level = 0;
-#endif
 
 				if ((Level < 0) || (Level >= LevelToColorToConstraintListMap.Num()))
 				{
@@ -173,10 +170,9 @@ void FPBDConstraintColor::ComputeIslandColoring(const int32 Island, const FPBDCo
 	}
 }
 
-#ifdef USE_CONTACT_LEVELS
-
 void FPBDConstraintColor::ComputeContactGraph(const int32 Island, const FPBDConstraintGraph& ConstraintGraph, uint32 ContainerId)
 {
+	SCOPE_CYCLE_COUNTER(STAT_Constraint_ComputeContactGraph);
 	const TArray<int32>& ConstraintDataIndices = ConstraintGraph.GetIslandConstraintData(Island);
 
 	IslandData[Island].MaxLevel = ConstraintDataIndices.Num() ? 0 : -1;
@@ -261,8 +257,6 @@ void FPBDConstraintColor::ComputeContactGraph(const int32 Island, const FPBDCons
 
 	check(IslandData[Island].MaxLevel >= 0 || !ConstraintDataIndices.Num());
 }
-#endif
-
 
 void FPBDConstraintColor::InitializeColor(const FPBDConstraintGraph& ConstraintGraph)
 {
@@ -292,9 +286,19 @@ void FPBDConstraintColor::InitializeColor(const FPBDConstraintGraph& ConstraintG
 
 void FPBDConstraintColor::ComputeColor(const int32 Island, const FPBDConstraintGraph& ConstraintGraph, uint32 ContainerId)
 {
-#ifdef USE_CONTACT_LEVELS
-	ComputeContactGraph(Island, ConstraintGraph, ContainerId);
-#endif
+	SCOPE_CYCLE_COUNTER(STAT_Constraint_ComputeColor);
+	if (bUseContactGraph)
+	{
+		ComputeContactGraph(Island, ConstraintGraph, ContainerId);
+	}
+	else
+	{
+		for (FGraphEdgeColor& Edge : Edges)
+		{
+			Edge.Level = 0;
+		}
+		IslandData[Island].MaxLevel = 0;
+	}
 	ComputeIslandColoring(Island, ConstraintGraph, ContainerId);
 }
 
@@ -320,11 +324,7 @@ int FPBDConstraintColor::GetIslandMaxLevel(int32 Island) const
 {
 	if (Island < IslandData.Num())
 	{
-#ifdef USE_CONTACT_LEVELS
 		return IslandData[Island].MaxLevel;
-#else
-		return 0;
-#endif
 	}
 	return -1;
 }
diff --git a/Engine/Source/Runtime/Experimental/Chaos/Public/Chaos/PBDConstraintColor.h b/Engine/Source/Runtime/Experimental/Chaos/Public/Chaos/PBDConstraintColor.h
--- a/Engine/Source/Runtime/Experimental/Chaos/Public/Chaos/PBDConstraintColor.h
+++ b/Engine/Source/Runtime/Experimental/Chaos/Public/Chaos/PBDConstraintColor.h
@@ -3,9 +3,6 @@
 
 #include "Chaos/ParticleHandle.h"
 
-// @todo(ccaulfield): move level/shock system out of coloring
-#define USE_CONTACT_LEVELS 1
-
 namespace Chaos
 {
 	class FPBDConstraintGraph;
@@ -23,6 +20,10 @@ namespace Chaos
 		typedef TMap<int32, FConstraintList> FColorToConstraintListMap;
 		typedef TArray<FColorToConstraintListMap> FLevelToColorToConstraintListMap;
 
+		FPBDConstraintColor()
+			: bUseContactGraph(true)
+		{}
+
 		/**
 		 * Initialize the color structures based on the connectivity graph (i.e., reset all color-related node, edge and island data).
 		 */
@@ -48,6 +49,11 @@ namespace Chaos
 		 */
 		int GetIslandMaxLevel(int32 Island) const;
 
+		void SetUseContactGraph(bool bInUseContactGraph)
+		{
+			bUseContactGraph = bInUseContactGraph;
+		}
+
 	private:
 		void ComputeContactGraph(const int32 Island, const FPBDConstraintGraph& ConstraintGraph, uint32 ContainerId);
 		void ComputeIslandColoring(const int32 Island, const FPBDConstraintGraph& ConstraintGraph, uint32 ContainerId);
@@ -66,30 +72,22 @@ namespace Chaos
 		{
 			FGraphEdgeColor()
 				: Color(INDEX_NONE)
-	#ifdef USE_CONTACT_LEVELS
 				, Level(INDEX_NONE)
-	#endif
 			{
 			}
 			int32 Color;
-	#ifdef USE_CONTACT_LEVELS
 			int32 Level;
-	#endif
 		};
 
 		struct FIslandColor
 		{
 			FIslandColor()
 				: MaxColor(0)
-	#ifdef USE_CONTACT_LEVELS
 				, MaxLevel(0)
-	#endif
 			{
 			}
 			int32 MaxColor;
-	#ifdef USE_CONTACT_LEVELS
 			int32 MaxLevel;
-	#endif
 			FLevelToColorToConstraintListMap LevelToColorToConstraintListMap;
 		};
 
@@ -98,6 +96,7 @@ namespace Chaos
 		TArray<FIslandColor> IslandData;
 		FLevelToColorToConstraintListMap EmptyLevelToColorToConstraintListMap;
 		TArray<int32> UpdatedNodes;
+		bool bUseContactGraph;
 	};
 
 }
diff --git a/Engine/Source/Runtime/Experimental/Chaos/Public/Chaos/PBDConstraintGraph.h b/Engine/Source/Runtime/Experimental/Chaos/Public/Chaos/PBDConstraintGraph.h
--- a/Engine/Source/Runtime/Experimental/Chaos/Public/Chaos/PBDConstraintGraph.h
+++ b/Engine/Source/Runtime/Experimental/Chaos/Public/Chaos/PBDConstraintGraph.h
@@ -3,10 +3,6 @@
 
 #include "ParticleHandle.h"
 
-
-// @todo(ccaulfield): can we get rid of this now?
-#define USE_CONTACT_LEVELS 1
-
 namespace Chaos
 {
 	template<typename T, int d>
diff --git a/Engine/Source/Runtime/Experimental/Chaos/Public/Chaos/PBDConstraintRule.h b/Engine/Source/Runtime/Experimental/Chaos/Public/Chaos/PBDConstraintRule.h
--- a/Engine/Source/Runtime/Experimental/Chaos/Public/Chaos/PBDConstraintRule.h
+++ b/Engine/Source/Runtime/Experimental/Chaos/Public/Chaos/PBDConstraintRule.h
@@ -150,6 +150,8 @@ namespace Chaos
 		/** Set up the perf-acceleration structures for the specified island. May be called in parallel for islands */
 		virtual void UpdateAccelerationStructures(const int32 Island) {}
 
+		virtual void SetUseContactGraph(const bool InUseContactGraph) {}
+
 		/** Remove all constraints associated with the specified particles */
 		// @todo(ccaulfield): remove uint version
 		virtual void RemoveConstraints(const TSet<TGeometryParticleHandle<FReal, 3>*>& RemovedParticles) { }
@@ -406,6 +408,11 @@ namespace Chaos
 			GraphColor.ComputeColor(Island, *ConstraintGraph, ContainerId);
 		}
 
+		virtual void SetUseContactGraph(const bool bInUseContactGraph) override
+		{
+			GraphColor.SetUseContactGraph(bInUseContactGraph);
+		}
+
 		template<typename TVisitor>
 		void VisitIslandConstraints(const int32 Island, const TVisitor& Visitor) const
 		{
diff --git a/Engine/Source/Runtime/Experimental/ChaosSolverEngine/Private/Chaos/ChaosSolverActor.cpp b/Engine/Source/Runtime/Experimental/ChaosSolverEngine/Private/Chaos/ChaosSolverActor.cpp
--- a/Engine/Source/Runtime/Experimental/ChaosSolverEngine/Private/Chaos/ChaosSolverActor.cpp
+++ b/Engine/Source/Runtime/Experimental/ChaosSolverEngine/Private/Chaos/ChaosSolverActor.cpp
@@ -201,6 +201,7 @@ AChaosSolverActor::AChaosSolverActor(const FObjectInitializer& ObjectInitializer
 	, bHasFloor(true)
 	, FloorHeight(0.f)
 	, MassScale(1.f)
+	, bGenerateContactGraph(true)
 	, ChaosDebugSubstepControl()
 {
 	// @question(Benn) : Does this need to be created on the Physics thread using a queued command?
@@ -289,7 +290,8 @@ void AChaosSolverActor::BeginPlay()
 			, InTrailingFilterSettings = TrailingFilterSettings
 			, InHasFloor = bHasFloor
 			, InFloorHeight = FloorHeight
-			, InMassScale = MassScale]
+			, InMassScale = MassScale
+			, InGenerateContactGraph = bGenerateContactGraph]
 		(Chaos::FPhysicsSolver* InSolver)
 		{
 #if TODO_REIMPLEMENT_SOLVER_SETTINGS_ACCESSORS
@@ -306,6 +308,7 @@ void AChaosSolverActor::BeginPlay()
 			InSolver->SetCollisionFilterSettings(InCollisionFilterSettings);
 			InSolver->SetBreakingFilterSettings(InBreakingFilterSettings);
 			InSolver->SetTrailingFilterSettings(InTrailingFilterSettings);
+			InSolver->SetUseContactGraph(InGenerateContactGraph);
 
 #if TODO_REIMPLEMENT_SOLVER_SETTINGS_ACCESSORS
 			InSolver->SetMassScale(InMassScale);
@@ -526,6 +529,14 @@ void AChaosSolverActor::PostEditChangeProperty(struct FPropertyChangedEvent& Pro
 				{
 				});
 			}
+			else if (PropertyChangedEvent.Property->GetFName() == GET_MEMBER_NAME_CHECKED(AChaosSolverActor, bGenerateContactGraph))
+			{
+				PhysDispatcher->EnqueueCommandImmediate(Solver, [InGenerateContactGraph = bGenerateContactGraph]
+				(Chaos::FPhysicsSolver* InSolver)
+				{
+				});
+			}
+
 #if TODO_REIMPLEMENT_TIMESTEP_MULTIPLIER
 			else if (PropertyChangedEvent.Property->GetFName() == GET_MEMBER_NAME_CHECKED(AChaosSolverActor, MassScale))
 			{
diff --git a/Engine/Source/Runtime/Experimental/ChaosSolverEngine/Public/Chaos/ChaosSolverActor.h b/Engine/Source/Runtime/Experimental/ChaosSolverEngine/Public/Chaos/ChaosSolverActor.h
--- a/Engine/Source/Runtime/Experimental/ChaosSolverEngine/Public/Chaos/ChaosSolverActor.h
+++ b/Engine/Source/Runtime/Experimental/ChaosSolverEngine/Public/Chaos/ChaosSolverActor.h
@@ -159,6 +159,12 @@ public:
 	UPROPERTY(EditAnywhere, BlueprintReadWrite, Category = "ChaosPhysics", meta = (DisplayName = "Mass Scale"))
 	float MassScale;
 
+	/*
+	* 
+	*/
+	UPROPERTY(EditAnywhere, BlueprintReadWrite, Category = "ChaosPhysics", meta = (DisplayName = "Generate Contact Graph"))
+	bool bGenerateContactGraph;
+
 	/*
 	* Control to pause/step/substep the solver to the next synchronization point.
 	*/
diff --git a/Engine/Source/Runtime/Experimental/ChaosSolvers/Public/PBDRigidsSolver.h b/Engine/Source/Runtime/Experimental/ChaosSolvers/Public/PBDRigidsSolver.h
--- a/Engine/Source/Runtime/Experimental/ChaosSolvers/Public/PBDRigidsSolver.h
+++ b/Engine/Source/Runtime/Experimental/ChaosSolvers/Public/PBDRigidsSolver.h
@@ -266,6 +266,7 @@ namespace Chaos
 		void SetIterations(const int32 InNumIterations) { GetEvolution()->SetNumIterations(InNumIterations); }
 		void SetPushOutIterations(const int32 InNumIterations) {  GetEvolution()->SetNumPushOutIterations(InNumIterations); }
 		void SetPushOutPairIterations(const int32 InNumIterations) {  GetEvolution()->GetCollisionConstraints().SetPushOutPairIterations(InNumIterations); }
+		void SetUseContactGraph(const bool bInUseContactGraph) { GetEvolution()->GetCollisionConstraintsRule().SetUseContactGraph(bInUseContactGraph); }
 
 		/**/
 		void SetGenerateCollisionData(bool bDoGenerate) { GetEventFilters()->SetGenerateCollisionEvents(bDoGenerate); }
