Refactor evolution to allow type traits

diff --git a/Engine/Source/Programs/HeadlessChaos/Private/HeadlessChaosApp.cpp b/Engine/Source/Programs/HeadlessChaos/Private/HeadlessChaosApp.cpp
--- a/Engine/Source/Programs/HeadlessChaos/Private/HeadlessChaosApp.cpp
+++ b/Engine/Source/Programs/HeadlessChaos/Private/HeadlessChaosApp.cpp
@@ -94,7 +94,7 @@ TEST(CollisionTests, PGS) {
 }
 
 TEST(ForceTests, Forces) {
-	ChaosTest::Gravity<Chaos::FPBDRigidsEvolutionGBF, float>();
+	ChaosTest::Gravity<Chaos::FPBDRigidsEvolution, float>();
 	SUCCEED();
 }
 
diff --git a/Engine/Source/Programs/HeadlessChaos/Private/HeadlessChaosTestBP.cpp b/Engine/Source/Programs/HeadlessChaos/Private/HeadlessChaosTestBP.cpp
--- a/Engine/Source/Programs/HeadlessChaos/Private/HeadlessChaosTestBP.cpp
+++ b/Engine/Source/Programs/HeadlessChaos/Private/HeadlessChaosTestBP.cpp
@@ -25,7 +25,7 @@ namespace ChaosTest
 			FChaosArchive Ar(*File);
 
 			TPBDRigidParticles<T, 3> Particles;
-			FPBDRigidsEvolutionGBF Evolution(MoveTemp(Particles));
+			FPBDRigidsEvolution Evolution(MoveTemp(Particles));
 			Evolution.SerializeForPerfTest(Ar);
 			Evolution.AdvanceOneTimeStep(Dt);
 			Evolution.EndFrame(Dt);
diff --git a/Engine/Source/Programs/HeadlessChaos/Private/HeadlessChaosTestBroadphase.cpp b/Engine/Source/Programs/HeadlessChaos/Private/HeadlessChaosTestBroadphase.cpp
--- a/Engine/Source/Programs/HeadlessChaos/Private/HeadlessChaosTestBroadphase.cpp
+++ b/Engine/Source/Programs/HeadlessChaos/Private/HeadlessChaosTestBroadphase.cpp
@@ -702,7 +702,7 @@ namespace ChaosTest
 		Particle->X() = TVector<FReal, 3>(0, 0, 0);
 
 		THandleArray<FChaosPhysicsMaterial> PhysicalMaterials;
-		FPBDRigidsEvolutionGBF Evolution(Particles, PhysicalMaterials);
+		FPBDRigidsEvolution Evolution(Particles, PhysicalMaterials);
 
 		// Flush spatial acceleration structures to put particle into structure.
 		Evolution.FlushSpatialAcceleration();
diff --git a/Engine/Source/Programs/HeadlessChaos/Private/HeadlessChaosTestClustering.cpp b/Engine/Source/Programs/HeadlessChaos/Private/HeadlessChaosTestClustering.cpp
--- a/Engine/Source/Programs/HeadlessChaos/Private/HeadlessChaosTestClustering.cpp
+++ b/Engine/Source/Programs/HeadlessChaos/Private/HeadlessChaosTestClustering.cpp
@@ -26,7 +26,7 @@ namespace ChaosTest {
 	{
 		TPBDRigidsSOAs<T, 3> Particles;
 		THandleArray<FChaosPhysicsMaterial> PhysicalMaterials;
-		FPBDRigidsEvolutionGBF Evolution(Particles, PhysicalMaterials);
+		FPBDRigidsEvolution Evolution(Particles, PhysicalMaterials);
 		TPBDRigidClusteredParticles<T, 3>& ClusteredParticles = Particles.GetClusteredParticles();
 		
 		uint32 FirstId = ClusteredParticles.Size();
@@ -73,7 +73,7 @@ namespace ChaosTest {
 	{
 		TPBDRigidsSOAs<T, 3> Particles;
 		THandleArray<FChaosPhysicsMaterial> PhysicalMaterials;
-		FPBDRigidsEvolutionGBF Evolution(Particles, PhysicalMaterials);
+		FPBDRigidsEvolution Evolution(Particles, PhysicalMaterials);
 		auto& ClusteredParticles = Particles.GetClusteredParticles();
 
 		//create a long row of boxes - the depth 0 cluster is the entire row, the depth 1 clusters 4 boxes each, the depth 2 clusters are 1 box each
@@ -199,7 +199,7 @@ namespace ChaosTest {
 	{
 		TPBDRigidsSOAs<T, 3> Particles;
 		THandleArray<FChaosPhysicsMaterial> PhysicalMaterials;
-		FPBDRigidsEvolutionGBF Evolution(Particles, PhysicalMaterials);
+		FPBDRigidsEvolution Evolution(Particles, PhysicalMaterials);
 		auto& ClusteredParticles = Particles.GetClusteredParticles();
 
 		//create a long row of boxes - the depth 0 cluster is the entire row, the depth 1 clusters 4 boxes each, the depth 2 clusters are 1 box each
diff --git a/Engine/Source/Programs/HeadlessChaos/Private/HeadlessChaosTestConstraintHandle.cpp b/Engine/Source/Programs/HeadlessChaos/Private/HeadlessChaosTestConstraintHandle.cpp
--- a/Engine/Source/Programs/HeadlessChaos/Private/HeadlessChaosTestConstraintHandle.cpp
+++ b/Engine/Source/Programs/HeadlessChaos/Private/HeadlessChaosTestConstraintHandle.cpp
@@ -190,35 +190,35 @@ namespace ChaosTest {
 
 	TEST(ConstraintHandleTests, DISABLED_CollisionConstraintHandle)
 	{
-		CollisionConstraintHandles<Chaos::FPBDRigidsEvolutionGBF>();
+		CollisionConstraintHandles<Chaos::FPBDRigidsEvolution>();
 
 		SUCCEED();
 	}
 
 	TEST(ConstraintHandleTests, JointConstraintHandle)
 	{
-		JointConstraintHandles<Chaos::FPBDRigidsEvolutionGBF>();
+		JointConstraintHandles<Chaos::FPBDRigidsEvolution>();
 
 		SUCCEED();
 	}
 
 	TEST(ConstraintHandleTests, PositionConstraintHandles)
 	{
-		PositionConstraintHandles<Chaos::FPBDRigidsEvolutionGBF>();
+		PositionConstraintHandles<Chaos::FPBDRigidsEvolution>();
 
 		SUCCEED();
 	}
 
 	TEST(ConstraintHandleTests, RigidSpringConstraintHandles)
 	{
-		RigidSpringConstraintHandles<Chaos::FPBDRigidsEvolutionGBF>();
+		RigidSpringConstraintHandles<Chaos::FPBDRigidsEvolution>();
 
 		SUCCEED();
 	}
 
 	TEST(ConstraintHandleTests, RigidDynamicSpringConstraintHandles)
 	{
-		RigidDynamicSpringConstraintHandles<Chaos::FPBDRigidsEvolutionGBF>();
+		RigidDynamicSpringConstraintHandles<Chaos::FPBDRigidsEvolution>();
 
 		SUCCEED();
 	}
diff --git a/Engine/Source/Programs/HeadlessChaos/Private/HeadlessChaosTestConstraints.cpp b/Engine/Source/Programs/HeadlessChaos/Private/HeadlessChaosTestConstraints.cpp
--- a/Engine/Source/Programs/HeadlessChaos/Private/HeadlessChaosTestConstraints.cpp
+++ b/Engine/Source/Programs/HeadlessChaos/Private/HeadlessChaosTestConstraints.cpp
@@ -45,7 +45,7 @@ namespace ChaosTest {
 		{
 			TPBDRigidsSOAs<FReal, 3> Particles;
 			THandleArray<FChaosPhysicsMaterial> PhysicalMaterials;
-			FPBDRigidsEvolutionGBF Evolution(Particles, PhysicalMaterials);
+			FPBDRigidsEvolution Evolution(Particles, PhysicalMaterials);
 			TArray<TPBDRigidParticleHandle<FReal, 3>*> Dynamics = Evolution.CreateDynamicParticles(1);
 			Dynamics[0]->SetGravityEnabled(false);
 
@@ -119,8 +119,8 @@ namespace ChaosTest {
 
 	TEST(ConstraintTests, Constraints)
 	{
-		ChaosTest::Position<Chaos::FPBDRigidsEvolutionGBF>();
-		ChaosTest::PositionAndJoint<Chaos::FPBDRigidsEvolutionGBF>();
+		ChaosTest::Position<Chaos::FPBDRigidsEvolution>();
+		ChaosTest::PositionAndJoint<Chaos::FPBDRigidsEvolution>();
 
 		SUCCEED();
 	}
diff --git a/Engine/Source/Programs/HeadlessChaos/Private/HeadlessChaosTestForces.cpp b/Engine/Source/Programs/HeadlessChaos/Private/HeadlessChaosTestForces.cpp
--- a/Engine/Source/Programs/HeadlessChaos/Private/HeadlessChaosTestForces.cpp
+++ b/Engine/Source/Programs/HeadlessChaos/Private/HeadlessChaosTestForces.cpp
@@ -24,7 +24,7 @@ namespace ChaosTest {
 		Evolution.AdvanceOneTimeStep(0.1);
 		EXPECT_LT(Dynamics[0]->X()[2], 0);
 	}
-	template void Gravity<Chaos::FPBDRigidsEvolutionGBF, float>();
+	template void Gravity<Chaos::FPBDRigidsEvolution, float>();
 
 
 }
\ No newline at end of file
diff --git a/Engine/Source/Programs/HeadlessChaos/Private/HeadlessChaosTestJoints.cpp b/Engine/Source/Programs/HeadlessChaos/Private/HeadlessChaosTestJoints.cpp
--- a/Engine/Source/Programs/HeadlessChaos/Private/HeadlessChaosTestJoints.cpp
+++ b/Engine/Source/Programs/HeadlessChaos/Private/HeadlessChaosTestJoints.cpp
@@ -477,7 +477,7 @@ namespace ChaosTest {
 		Box2->X() = FVec3((FReal)500, (FReal)0, (FReal)1000);
 		Box2->P() = Box2->X();
 		THandleArray<FChaosPhysicsMaterial> PhysicalMaterials;
-		FPBDRigidsEvolutionGBF Evolution(Particles, PhysicalMaterials);
+		FPBDRigidsEvolution Evolution(Particles, PhysicalMaterials);
 		TVector<TGeometryParticleHandle<FReal, 3>*,2> ConstrainedParticles = TVector<TGeometryParticleHandle<FReal,3>*, 2>(StaticBox, Box2);
 		TVector<FVec3, 2> Points = { FVec3((FReal)100, (FReal)0, (FReal)1000), FVec3((FReal)400, (FReal)0, (FReal)1000) };
 
@@ -519,7 +519,7 @@ namespace ChaosTest {
 			Box2.P() = Box2.X();
 
 			THandleArray<FChaosPhysicsMaterial> PhysicalMaterials;
-			FPBDRigidsEvolutionGBF Evolution(Particles, PhysicalMaterials);
+			FPBDRigidsEvolution Evolution(Particles, PhysicalMaterials);
 			TArray<TVector<TGeometryParticleHandle<FReal,3>*, 2>> Constraints = { TVector<TGeometryParticleHandle<FReal,3>*, 2>(&StaticBox, &Box2) };
 
 			Evolution.SetPhysicsMaterial(&StaticBox, MakeSerializable(PhysicalMaterial));
@@ -549,7 +549,7 @@ namespace ChaosTest {
 			Box2.P() = Box2.X();
 
 			THandleArray<FChaosPhysicsMaterial> PhysicalMaterials;
-			FPBDRigidsEvolutionGBF Evolution(Particles, PhysicalMaterials);
+			FPBDRigidsEvolution Evolution(Particles, PhysicalMaterials);
 			TArray<TVector<TGeometryParticleHandle<FReal,3>*, 2>> Constraints = { TVector<TGeometryParticleHandle<FReal,3>*, 2>(&StaticBox, &Box2) };
 
 			Evolution.SetPhysicsMaterial(&StaticBox, MakeSerializable(PhysicalMaterial));
diff --git a/Engine/Source/Programs/HeadlessChaos/Private/HeadlessChaosTestSerialization.cpp b/Engine/Source/Programs/HeadlessChaos/Private/HeadlessChaosTestSerialization.cpp
--- a/Engine/Source/Programs/HeadlessChaos/Private/HeadlessChaosTestSerialization.cpp
+++ b/Engine/Source/Programs/HeadlessChaos/Private/HeadlessChaosTestSerialization.cpp
@@ -426,7 +426,7 @@ namespace ChaosTest
 				TPBDRigidsSOAs<float, 3> Particles;
 
 				THandleArray<FChaosPhysicsMaterial> PhysicalMaterials;
-				FPBDRigidsEvolutionGBF Evolution(Particles, PhysicalMaterials);
+				FPBDRigidsEvolution Evolution(Particles, PhysicalMaterials);
 
 				Evolution.Serialize(ChaosAr);
 				Evolution.AdvanceOneTimeStep(1 / 60.f);
diff --git a/Engine/Source/Programs/HeadlessChaos/Private/HeadlessChaosTestSim.cpp b/Engine/Source/Programs/HeadlessChaos/Private/HeadlessChaosTestSim.cpp
--- a/Engine/Source/Programs/HeadlessChaos/Private/HeadlessChaosTestSim.cpp
+++ b/Engine/Source/Programs/HeadlessChaos/Private/HeadlessChaosTestSim.cpp
@@ -17,7 +17,7 @@ namespace ChaosTest {
 	{
 		TPBDRigidsSOAs<FReal, 3> Particles;
 		THandleArray<FChaosPhysicsMaterial> PhysicalMaterials;
-		FPBDRigidsEvolutionGBF Evolution(Particles, PhysicalMaterials);
+		FPBDRigidsEvolution Evolution(Particles, PhysicalMaterials);
 		auto Static = Evolution.CreateStaticParticles(1)[0];
 		auto Dynamic = Evolution.CreateDynamicParticles(1)[0];
 
@@ -52,7 +52,7 @@ namespace ChaosTest {
 	{
 		TPBDRigidsSOAs<FReal, 3> Particles;
 		THandleArray<FChaosPhysicsMaterial> PhysicalMaterials;
-		FPBDRigidsEvolutionGBF Evolution(Particles, PhysicalMaterials);
+		FPBDRigidsEvolution Evolution(Particles, PhysicalMaterials);
 		auto Static = Evolution.CreateStaticParticles(1)[0];
 		auto Dynamic = Evolution.CreateDynamicParticles(1)[0];
 
@@ -91,7 +91,7 @@ namespace ChaosTest {
 	{
 		TPBDRigidsSOAs<FReal, 3> Particles;
 		THandleArray<FChaosPhysicsMaterial> PhysicalMaterials;
-		FPBDRigidsEvolutionGBF Evolution(Particles, PhysicalMaterials);
+		FPBDRigidsEvolution Evolution(Particles, PhysicalMaterials);
 		auto Static = Evolution.CreateStaticParticles(1)[0];
 		auto Dynamic = Evolution.CreateDynamicParticles(1)[0];
 
@@ -120,7 +120,7 @@ namespace ChaosTest {
 	{
 		TPBDRigidsSOAs<FReal, 3> Particles;
 		THandleArray<FChaosPhysicsMaterial> PhysicalMaterials;
-		FPBDRigidsEvolutionGBF Evolution(Particles, PhysicalMaterials);
+		FPBDRigidsEvolution Evolution(Particles, PhysicalMaterials);
 		auto Static = Evolution.CreateStaticParticles(1)[0];
 		auto Dynamic1 = Evolution.CreateDynamicParticles(1)[0];
 		auto Dynamic2 = Evolution.CreateDynamicParticles(1)[0];
diff --git a/Engine/Source/Programs/HeadlessChaos/Private/HeadlessChaosTestUtility.cpp b/Engine/Source/Programs/HeadlessChaos/Private/HeadlessChaosTestUtility.cpp
--- a/Engine/Source/Programs/HeadlessChaos/Private/HeadlessChaosTestUtility.cpp
+++ b/Engine/Source/Programs/HeadlessChaos/Private/HeadlessChaosTestUtility.cpp
@@ -45,7 +45,7 @@ namespace ChaosTest {
 	template int32 AppendAnalyticSphere(TPBDRigidParticles<float, 3> &, float Scale);
 
 	template<class T>
-	TPBDRigidParticleHandle<T, 3>* AppendAnalyticSphere2(FPBDRigidsEvolutionGBF& Evolution, T Scale)
+	TPBDRigidParticleHandle<T, 3>* AppendAnalyticSphere2(FPBDRigidsEvolution& Evolution, T Scale)
 	{
 		TArray<TPBDRigidParticleHandle<T, 3>*> Particles = Evolution.CreateDynamicParticles(1);
 		TPBDRigidParticleHandle<T, 3>* Particle = Particles[0];
@@ -65,7 +65,7 @@ namespace ChaosTest {
 
 		return Particle;
 	}
-	template TPBDRigidParticleHandle<float, 3>* AppendAnalyticSphere2(FPBDRigidsEvolutionGBF& Evolution, float Scale);
+	template TPBDRigidParticleHandle<float, 3>* AppendAnalyticSphere2(FPBDRigidsEvolution& Evolution, float Scale);
 
 
 	template<class T>
@@ -116,22 +116,22 @@ namespace ChaosTest {
 	}
 
 	template<class T>
-	TKinematicGeometryParticleHandle<T, 3>* AppendKinematicAnalyticBox2(FPBDRigidsEvolutionGBF& Evolution, TVector<T, 3> Scale)
+	TKinematicGeometryParticleHandle<T, 3>* AppendKinematicAnalyticBox2(FPBDRigidsEvolution& Evolution, TVector<T, 3> Scale)
 	{
 		TArray<TKinematicGeometryParticleHandle<T, 3>*> Particles = Evolution.CreateKinematicParticles(1);
 		InitAnalyticBox2(Particles[0], Scale);
 		return Particles[0];
 	}
-	template TKinematicGeometryParticleHandle<float, 3>* AppendKinematicAnalyticBox2(FPBDRigidsEvolutionGBF& Evolution, TVector<float, 3> Scale);
+	template TKinematicGeometryParticleHandle<float, 3>* AppendKinematicAnalyticBox2(FPBDRigidsEvolution& Evolution, TVector<float, 3> Scale);
 
 	template<class T>
-	TPBDRigidParticleHandle<T, 3>* AppendDynamicAnalyticBox2(FPBDRigidsEvolutionGBF& Evolution, TVector<T, 3> Scale)
+	TPBDRigidParticleHandle<T, 3>* AppendDynamicAnalyticBox2(FPBDRigidsEvolution& Evolution, TVector<T, 3> Scale)
 	{
 		TArray<TPBDRigidParticleHandle<T, 3>*> Particles = Evolution.CreateDynamicParticles(1);
 		InitAnalyticBox2(Particles[0], Scale);
 		return Particles[0];
 	}
-	template TPBDRigidParticleHandle<float, 3>* AppendDynamicAnalyticBox2(FPBDRigidsEvolutionGBF& Evolution, TVector<float, 3> Scale);
+	template TPBDRigidParticleHandle<float, 3>* AppendDynamicAnalyticBox2(FPBDRigidsEvolution& Evolution, TVector<float, 3> Scale);
 
 
 	template<class T>
@@ -412,13 +412,13 @@ namespace ChaosTest {
 
 
 	template<class T>
-	TPBDRigidParticleHandle<T, 3>* AppendDynamicParticleBox(FPBDRigidsEvolutionGBF& Evolution, const TVector<T, 3>& Scale, TArray<TVector<int32, 3>>* OutElements)
+	TPBDRigidParticleHandle<T, 3>* AppendDynamicParticleBox(FPBDRigidsEvolution& Evolution, const TVector<T, 3>& Scale, TArray<TVector<int32, 3>>* OutElements)
 	{
 		TArray<TPBDRigidParticleHandle<T, 3>*> Particles = Evolution.CreateDynamicParticles(1);
 		InitDynamicParticleBox2(Particles[0], Scale, OutElements);
 		return Particles[0];
 	}
-	template TPBDRigidParticleHandle<float, 3>* AppendDynamicParticleBox(FPBDRigidsEvolutionGBF& Evolution, const TVector<float, 3>& Scale, TArray<Chaos::TVector<int32, 3>>* OutElements);
+	template TPBDRigidParticleHandle<float, 3>* AppendDynamicParticleBox(FPBDRigidsEvolution& Evolution, const TVector<float, 3>& Scale, TArray<Chaos::TVector<int32, 3>>* OutElements);
 
 	template<class T>
 	TPBDRigidParticleHandle<T, 3>* AppendDynamicParticleBox(TPBDRigidsSOAs<T, 3>& SOAs, const TVector<T, 3>& Scale, TArray<TVector<int32, 3>>* OutElements)
@@ -457,13 +457,13 @@ namespace ChaosTest {
 	template TPBDRigidParticleHandle<float, 3>* AppendDynamicParticleTaperedCylinder(TPBDRigidsSOAs<float, 3>& Evolution, const TVector<float, 3>& Scale, TArray<Chaos::TVector<int32, 3>>* OutElements);
 
 	template<class T>
-	TPBDRigidParticleHandle<T, 3>* AppendClusteredParticleBox(FPBDRigidsEvolutionGBF& Evolution, const TVector<T, 3>& Scale, TArray<TVector<int32, 3>>* OutElements)
+	TPBDRigidParticleHandle<T, 3>* AppendClusteredParticleBox(FPBDRigidsEvolution& Evolution, const TVector<T, 3>& Scale, TArray<TVector<int32, 3>>* OutElements)
 	{
 		auto Particles = Evolution.CreateClusteredParticles(1);
 		InitDynamicParticleBox2(Particles[0], Scale, OutElements);
 		return Particles[0];
 	}
-	template TPBDRigidParticleHandle<float, 3>* AppendClusteredParticleBox(FPBDRigidsEvolutionGBF& Evolution, const TVector<float, 3>& Scale, TArray<Chaos::TVector<int32, 3>>* OutElements);
+	template TPBDRigidParticleHandle<float, 3>* AppendClusteredParticleBox(FPBDRigidsEvolution& Evolution, const TVector<float, 3>& Scale, TArray<Chaos::TVector<int32, 3>>* OutElements);
 
 	template<class T>
 	TPBDRigidParticleHandle<T, 3>* AppendClusteredParticleBox(TPBDRigidsSOAs<T, 3>& SOAs, const TVector<T, 3>& Scale, TArray<TVector<int32, 3>>* OutElements)
@@ -555,7 +555,7 @@ namespace ChaosTest {
 	template int AppendStaticAnalyticFloor(TPBDRigidParticles<float, 3> &);
 
 	template<class T>
-	TKinematicGeometryParticleHandle<T, 3>* AppendStaticAnalyticFloor(FPBDRigidsEvolutionGBF& Evolution)
+	TKinematicGeometryParticleHandle<T, 3>* AppendStaticAnalyticFloor(FPBDRigidsEvolution& Evolution)
 	{
 		TArray<TKinematicGeometryParticleHandle<T, 3>*> Particles = Evolution.CreateKinematicParticles(1);
 		TKinematicGeometryParticleHandle<T, 3>* Particle = Particles[0];
@@ -570,7 +570,7 @@ namespace ChaosTest {
 
 		return Particle;
 	}
-	template TKinematicGeometryParticleHandle<float, 3>* AppendStaticAnalyticFloor(FPBDRigidsEvolutionGBF& Evolution);
+	template TKinematicGeometryParticleHandle<float, 3>* AppendStaticAnalyticFloor(FPBDRigidsEvolution& Evolution);
 
 	template<class T>
 	TKinematicGeometryParticleHandle<T, 3>* AppendStaticAnalyticFloor(TPBDRigidsSOAs<T, 3>& SOAs)
diff --git a/Engine/Source/Programs/HeadlessChaos/Public/HeadlessChaosTestConstraints.h b/Engine/Source/Programs/HeadlessChaos/Public/HeadlessChaosTestConstraints.h
--- a/Engine/Source/Programs/HeadlessChaos/Public/HeadlessChaosTestConstraints.h
+++ b/Engine/Source/Programs/HeadlessChaos/Public/HeadlessChaosTestConstraints.h
@@ -61,7 +61,7 @@ namespace ChaosTest
 
 		// Solver state
 		TPBDRigidsSOAs<FReal, 3> SOAs;
-		FPBDRigidsEvolutionGBF Evolution;
+		FPBDRigidsEvolution Evolution;
 		TUniquePtr<FChaosPhysicsMaterial> PhysicalMaterial;
 		THandleArray<FChaosPhysicsMaterial> PhysicalMaterials;
 
diff --git a/Engine/Source/Programs/HeadlessChaos/Public/HeadlessChaosTestUtility.h b/Engine/Source/Programs/HeadlessChaos/Public/HeadlessChaosTestUtility.h
--- a/Engine/Source/Programs/HeadlessChaos/Public/HeadlessChaosTestUtility.h
+++ b/Engine/Source/Programs/HeadlessChaos/Public/HeadlessChaosTestUtility.h
@@ -16,10 +16,10 @@
 #include "Chaos/UniformGrid.h"
 #include "Chaos/Utilities.h"
 #include "Chaos/ParticleHandleFwd.h"
+#include "Chaos/PBDRigidsEvolutionFwd.h"
 
 namespace Chaos
 {
-	class FPBDRigidsEvolutionGBF;
 	template <typename T, int d> class TPBDRigidsSOAs;
 }
 
@@ -45,21 +45,21 @@ namespace ChaosTest {
 	template<class T>
 	int32 AppendAnalyticSphere(TPBDRigidParticles<T, 3> & InParticles, T Scale = (T)1);
 	template<class T>
-	TPBDRigidParticleHandle<T, 3>* AppendAnalyticSphere2(FPBDRigidsEvolutionGBF& Evolution, T Scale = (T)1);
+	TPBDRigidParticleHandle<T, 3>* AppendAnalyticSphere2(FPBDRigidsEvolution& Evolution, T Scale = (T)1);
 
 	/**/
 	template<class T>
 	int32 AppendAnalyticBox(TPBDRigidParticles<T, 3>& InParticles, TVector<T, 3> Scale = TVector<T, 3>(1));
 	template<class T>
-	TKinematicGeometryParticleHandle<T, 3>* AppendKinematicAnalyticBox2(FPBDRigidsEvolutionGBF& Evolution, TVector<T, 3> Scale = TVector<T, 3>(1));
+	TKinematicGeometryParticleHandle<T, 3>* AppendKinematicAnalyticBox2(FPBDRigidsEvolution& Evolution, TVector<T, 3> Scale = TVector<T, 3>(1));
 	template<class T>
-	TPBDRigidParticleHandle<T, 3>* AppendDynamicAnalyticBox2(FPBDRigidsEvolutionGBF& Evolution, TVector<T, 3> Scale = TVector<T, 3>(1));
+	TPBDRigidParticleHandle<T, 3>* AppendDynamicAnalyticBox2(FPBDRigidsEvolution& Evolution, TVector<T, 3> Scale = TVector<T, 3>(1));
 
 	/**/
 	template<class T>
 	int32 AppendParticleBox(TPBDRigidParticles<T, 3>& InParticles, TVector<T, 3> Scale = TVector<T, 3>(1), TArray<TVector<int32, 3>>* OutElements = nullptr);
 	template<class T>
-	TPBDRigidParticleHandle<T, 3>* AppendDynamicParticleBox(FPBDRigidsEvolutionGBF& Evolution, const TVector<T, 3>& Scale = TVector<T, 3>(1), TArray<TVector<int32, 3>>* OutElements = nullptr);
+	TPBDRigidParticleHandle<T, 3>* AppendDynamicParticleBox(FPBDRigidsEvolution& Evolution, const TVector<T, 3>& Scale = TVector<T, 3>(1), TArray<TVector<int32, 3>>* OutElements = nullptr);
 
 	template<class T>
 	TPBDRigidParticleHandle<T, 3>* AppendDynamicParticleBox(TPBDRigidsSOAs<T, 3>& SOAs, const TVector<T, 3>& Scale = TVector<T, 3>(1), TArray<TVector<int32, 3>>* OutElements = nullptr);
@@ -77,7 +77,7 @@ namespace ChaosTest {
 	TGeometryParticleHandle<T, 3>* AppendStaticParticleBox(TPBDRigidsSOAs<T, 3>& SOAs, const TVector<T, 3>& Scale = TVector<T, 3>(1), TArray<TVector<int32, 3>>* OutElements = nullptr);
 
 	template<class T>
-	TPBDRigidParticleHandle<T, 3>* AppendClusteredParticleBox(FPBDRigidsEvolutionGBF& Evolution, const TVector<T, 3>& Scale = TVector<T, 3>(1), TArray<TVector<int32, 3>>* OutElements = nullptr);
+	TPBDRigidParticleHandle<T, 3>* AppendClusteredParticleBox(FPBDRigidsEvolution& Evolution, const TVector<T, 3>& Scale = TVector<T, 3>(1), TArray<TVector<int32, 3>>* OutElements = nullptr);
 
 	template<class T>
 	TPBDRigidParticleHandle<T, 3>* AppendClusteredParticleBox(TPBDRigidsSOAs<T, 3>& SOAs, const TVector<T, 3>& Scale = TVector<T, 3>(1), TArray<TVector<int32, 3>>* OutElements = nullptr);
@@ -86,7 +86,7 @@ namespace ChaosTest {
 	template<class T>
 	int32 AppendStaticAnalyticFloor(TPBDRigidParticles<T, 3>& InParticles);
 	template<class T>
-	TKinematicGeometryParticleHandle<T, 3>* AppendStaticAnalyticFloor(FPBDRigidsEvolutionGBF& Evolution);
+	TKinematicGeometryParticleHandle<T, 3>* AppendStaticAnalyticFloor(FPBDRigidsEvolution& Evolution);
 
 	template<class T>
 	TKinematicGeometryParticleHandle<T, 3>* AppendStaticAnalyticFloor(TPBDRigidsSOAs<T, 3>& SOAs);
diff --git a/Engine/Source/Runtime/Engine/Private/PhysicsEngine/PhysUtils.cpp b/Engine/Source/Runtime/Engine/Private/PhysicsEngine/PhysUtils.cpp
--- a/Engine/Source/Runtime/Engine/Private/PhysicsEngine/PhysUtils.cpp
+++ b/Engine/Source/Runtime/Engine/Private/PhysicsEngine/PhysUtils.cpp
@@ -504,7 +504,7 @@ bool FPhysicsInterface::ExecPhysCommands(const TCHAR* Cmd, FOutputDevice* Output
 		Chaos::FChaosArchive Ar(BaseAr);
 		FPhysScene* PhysScene = InWorld->GetPhysicsScene();
 		Chaos::FPhysicsSolver* Solver = PhysScene->GetSolver();
-		Chaos::FPBDRigidsEvolutionGBF* Evolution = Solver->GetEvolution();
+		auto* Evolution = Solver->GetEvolution();
 		Evolution->Serialize(Ar);
 		TUniquePtr<Chaos::FChaosArchiveContext> ArchiveContext = Ar.StealContext();
 
diff --git a/Engine/Source/Runtime/Engine/Public/Physics/ImmediatePhysics/ImmediatePhysicsChaos/ImmediatePhysicsCore_Chaos.h b/Engine/Source/Runtime/Engine/Public/Physics/ImmediatePhysics/ImmediatePhysicsChaos/ImmediatePhysicsCore_Chaos.h
--- a/Engine/Source/Runtime/Engine/Public/Physics/ImmediatePhysics/ImmediatePhysicsChaos/ImmediatePhysicsCore_Chaos.h
+++ b/Engine/Source/Runtime/Engine/Public/Physics/ImmediatePhysics/ImmediatePhysicsChaos/ImmediatePhysicsCore_Chaos.h
@@ -7,6 +7,7 @@
 
 #include "Chaos/Core.h"
 #include "Chaos/ParticleHandleFwd.h"
+#include "Chaos/PBDRigidsEvolutionFwd.h"
 
 namespace Chaos
 {
@@ -16,7 +17,6 @@ namespace Chaos
 	class FParticlePairCollisionDetector;
 	class FPBDJointConstraintHandle;
 	class FPBDJointConstraints;
-	class FPBDRigidsEvolutionGBF;
 	class FPerShapeData;
 	template<class T> class TArrayCollectionArray;
 	template<typename T, int D> struct TKinematicGeometryParticleParameters;
diff --git a/Engine/Source/Runtime/Experimental/Chaos/Private/Chaos/PBDRigidClustering.cpp b/Engine/Source/Runtime/Experimental/Chaos/Private/Chaos/PBDRigidClustering.cpp
--- a/Engine/Source/Runtime/Experimental/Chaos/Private/Chaos/PBDRigidClustering.cpp
+++ b/Engine/Source/Runtime/Experimental/Chaos/Private/Chaos/PBDRigidClustering.cpp
@@ -21,6 +21,7 @@
 #include "Chaos/PerParticleEulerStepVelocity.h"
 #include "Chaos/PerParticleEtherDrag.h"
 #include "Chaos/PerParticlePBDEulerStep.h"
+#include "Chaos/EvolutionTraits.h"
 #include "CoreMinimal.h"
 
 namespace Chaos
@@ -113,9 +114,9 @@ namespace Chaos
 	}
 
 	DECLARE_CYCLE_STAT(TEXT("TPBDRigidClustering<>::RewindAndEvolve<BGF>()"), STAT_RewindAndEvolve_BGF, STATGROUP_Chaos);
-	template<class T, int d>
+	template<typename Traits, typename T, int d>
 	void RewindAndEvolve(
-		FPBDRigidsEvolutionGBF& Evolution, 
+		TPBDRigidsEvolutionGBF<Traits>& Evolution, 
 		TPBDRigidClusteredParticles<T, d>& InParticles, 
 		const TSet<int32>& IslandsToRecollide, 
 		const TSet<TPBDRigidParticleHandle<T, d>*> AllActivatedChildren,
@@ -2070,7 +2071,10 @@ namespace Chaos
 } // namespace Chaos
 
 using namespace Chaos;
-template class CHAOS_API Chaos::TPBDRigidClustering<FPBDRigidsEvolutionGBF, FPBDCollisionConstraints, float, 3>;
+
+#define EVOLUTION_TRAIT(Trait) template class CHAOS_API Chaos::TPBDRigidClustering<Chaos::TPBDRigidsEvolutionGBF<Trait>, FPBDCollisionConstraints, float, 3>;
+#include "Chaos/EvolutionTraits.inl"
+#undef EVOLUTION_TRAIT
 
 
 template CHAOS_API void Chaos::UpdateClusterMassProperties<float, 3>(
diff --git a/Engine/Source/Runtime/Experimental/Chaos/Private/Chaos/PBDRigidsEvolution.cpp b/Engine/Source/Runtime/Experimental/Chaos/Private/Chaos/PBDRigidsEvolution.cpp
--- a/Engine/Source/Runtime/Experimental/Chaos/Private/Chaos/PBDRigidsEvolution.cpp
+++ b/Engine/Source/Runtime/Experimental/Chaos/Private/Chaos/PBDRigidsEvolution.cpp
@@ -169,7 +169,8 @@ namespace Chaos
 		}
 	};
 
-	FPBDRigidsEvolutionBase::FPBDRigidsEvolutionBase(TPBDRigidsSOAs<FReal, 3>& InParticles, THandleArray<FChaosPhysicsMaterial>& InSolverPhysicsMaterials, int32 InNumIterations, int32 InNumPushOutIterations, bool InIsSingleThreaded)
+	template <typename Traits>
+	TPBDRigidsEvolutionBase<Traits>::TPBDRigidsEvolutionBase(TPBDRigidsSOAs<FReal, 3>& InParticles, THandleArray<FChaosPhysicsMaterial>& InSolverPhysicsMaterials, int32 InNumIterations, int32 InNumPushOutIterations, bool InIsSingleThreaded)
 	    : Particles(InParticles)
 		, SolverPhysicsMaterials(InSolverPhysicsMaterials)
 		, bExternalReady(false)
@@ -191,7 +192,8 @@ namespace Chaos
 		ComputeIntermediateSpatialAcceleration();
 	}
 
-	FPBDRigidsEvolutionBase::~FPBDRigidsEvolutionBase()
+	template <typename Traits>
+	TPBDRigidsEvolutionBase<Traits>::~TPBDRigidsEvolutionBase()
 	{
 		Particles.GetParticleHandles().RemoveArray(&PhysicsMaterials);
 		Particles.GetParticleHandles().RemoveArray(&PerParticlePhysicsMaterials);
@@ -209,8 +211,8 @@ namespace Chaos
 	DECLARE_CYCLE_STAT(TEXT("CreateInitialAccelerationStructure"), STAT_CreateInitialAccelerationStructure, STATGROUP_Chaos);
 	DECLARE_CYCLE_STAT(TEXT("CreateNonSlicedStructures"), STAT_CreateNonSlicedStructures, STATGROUP_Chaos);
 
-
-	FPBDRigidsEvolutionBase::FChaosAccelerationStructureTask::FChaosAccelerationStructureTask(
+	template <typename Traits>
+	TPBDRigidsEvolutionBase<Traits>::FChaosAccelerationStructureTask::FChaosAccelerationStructureTask(
 		ISpatialAccelerationCollectionFactory& InSpatialCollectionFactory
 		, const TMap<FSpatialAccelerationIdx, TUniquePtr<FSpatialAccelerationCache>>& InSpatialAccelerationCache
 		, TUniquePtr<FAccelerationStructure>& InAccelerationStructure
@@ -226,17 +228,20 @@ namespace Chaos
 	{
 	}
 
-	TStatId FPBDRigidsEvolutionBase::FChaosAccelerationStructureTask::GetStatId()
+	template <typename Traits>
+	TStatId TPBDRigidsEvolutionBase<Traits>::FChaosAccelerationStructureTask::GetStatId()
 	{
 		RETURN_QUICK_DECLARE_CYCLE_STAT(FChaosAccelerationStructureTask, STATGROUP_Chaos);
 	}
 
-	ENamedThreads::Type FPBDRigidsEvolutionBase::FChaosAccelerationStructureTask::GetDesiredThread()
+	template <typename Traits>
+	ENamedThreads::Type TPBDRigidsEvolutionBase<Traits>::FChaosAccelerationStructureTask::GetDesiredThread()
 	{
 		return ENamedThreads::AnyBackgroundThreadNormalTask;
 	}
 
-	ESubsequentsMode::Type FPBDRigidsEvolutionBase::FChaosAccelerationStructureTask::GetSubsequentsMode()
+	template <typename Traits>
+	ESubsequentsMode::Type TPBDRigidsEvolutionBase<Traits>::FChaosAccelerationStructureTask::GetSubsequentsMode()
 	{
 		return ESubsequentsMode::TrackSubsequents;
 	}
@@ -267,7 +272,8 @@ namespace Chaos
 		}
 	}
 
-	void FPBDRigidsEvolutionBase::FChaosAccelerationStructureTask::DoTask(ENamedThreads::Type CurrentThread, const FGraphEventRef& MyCompletionGraphEvent)
+	template <typename Traits>
+	void TPBDRigidsEvolutionBase<Traits>::FChaosAccelerationStructureTask::DoTask(ENamedThreads::Type CurrentThread, const FGraphEventRef& MyCompletionGraphEvent)
 	{
 		LLM_SCOPE(ELLMTag::ChaosAcceleration);
 
@@ -360,7 +366,8 @@ namespace Chaos
 		}
 	}
 
-	void FPBDRigidsEvolutionBase::ApplyParticlePendingData(const FPendingSpatialData& SpatialData, FAccelerationStructure& AccelerationStructure, bool bUpdateCache)
+	template <typename Traits>
+	void TPBDRigidsEvolutionBase<Traits>::ApplyParticlePendingData(const FPendingSpatialData& SpatialData, FAccelerationStructure& AccelerationStructure, bool bUpdateCache)
 	{
 		if (SpatialData.bDelete)
 		{
@@ -420,7 +427,8 @@ namespace Chaos
 		}
 	}
 
-	void FPBDRigidsEvolutionBase::FlushInternalAccelerationQueue()
+	template <typename Traits>
+	void TPBDRigidsEvolutionBase<Traits>::FlushInternalAccelerationQueue()
 	{
 		for (const FPendingSpatialData& PendingData : InternalAccelerationQueue.PendingData)
 		{
@@ -429,7 +437,8 @@ namespace Chaos
 		InternalAccelerationQueue.Reset();
 	}
 
-	void FPBDRigidsEvolutionBase::FlushAsyncAccelerationQueue()
+	template <typename Traits>
+	void TPBDRigidsEvolutionBase<Traits>::FlushAsyncAccelerationQueue()
 	{
 		for (const FPendingSpatialData& PendingData : AsyncAccelerationQueue.PendingData)
 		{
@@ -456,7 +465,8 @@ namespace Chaos
 		ExternalAccelerationQueue.Reset();
 	}
 
-	void FPBDRigidsEvolutionBase::FlushExternalAccelerationQueue(FAccelerationStructure& Acceleration)
+	template <typename Traits>
+	void TPBDRigidsEvolutionBase<Traits>::FlushExternalAccelerationQueue(FAccelerationStructure& Acceleration)
 	{
 		for (const FPendingSpatialData& PendingData : ExternalAccelerationQueue.PendingData)
 		{
@@ -465,7 +475,8 @@ namespace Chaos
 		ExternalAccelerationQueue.Reset();
 	}
 
-	void FPBDRigidsEvolutionBase::WaitOnAccelerationStructure()
+	template <typename Traits>
+	void TPBDRigidsEvolutionBase<Traits>::WaitOnAccelerationStructure()
 	{
 		if (AccelerationStructureTaskComplete.GetReference())
 		{
@@ -476,7 +487,8 @@ namespace Chaos
 		}
 	}
 
-	void FPBDRigidsEvolutionBase::ComputeIntermediateSpatialAcceleration(bool bBlock)
+	template <typename Traits>
+	void TPBDRigidsEvolutionBase<Traits>::ComputeIntermediateSpatialAcceleration(bool bBlock)
 	{
 		LLM_SCOPE(ELLMTag::ChaosAcceleration);
 		SCOPE_CYCLE_COUNTER(STAT_ComputeIntermediateSpatialAcceleration);
@@ -546,7 +558,8 @@ namespace Chaos
 		}
 	}
 
-	void FPBDRigidsEvolutionBase::UpdateExternalAccelerationStructure(TUniquePtr<ISpatialAccelerationCollection<TAccelerationStructureHandle<FReal, 3>, FReal, 3>>& StructToUpdate)
+	template <typename Traits>
+	void TPBDRigidsEvolutionBase<Traits>::UpdateExternalAccelerationStructure(TUniquePtr<ISpatialAccelerationCollection<TAccelerationStructureHandle<FReal, 3>, FReal, 3>>& StructToUpdate)
 	{
 		DECLARE_SCOPE_CYCLE_COUNTER(TEXT("CreateExternalAccelerationStructure"), STAT_CreateExternalAccelerationStructure, STATGROUP_Physics);
 		LLM_SCOPE(ELLMTag::ChaosAcceleration);
@@ -564,7 +577,8 @@ namespace Chaos
 		}
 	}
 
-	void FPBDRigidsEvolutionBase::FlushSpatialAcceleration()
+	template <typename Traits>
+	void TPBDRigidsEvolutionBase<Traits>::FlushSpatialAcceleration()
 	{
 		//force build acceleration structure with latest data
 		ComputeIntermediateSpatialAcceleration(true);
@@ -572,7 +586,8 @@ namespace Chaos
 		ComputeIntermediateSpatialAcceleration(true);
 	}
 
-	void FPBDRigidsEvolutionBase::RebuildSpatialAccelerationForPerfTest()
+	template <typename Traits>
+	void TPBDRigidsEvolutionBase<Traits>::RebuildSpatialAccelerationForPerfTest()
 	{
 		WaitOnAccelerationStructure();
 
@@ -591,7 +606,8 @@ namespace Chaos
 		FlushSpatialAcceleration();
 	}
 
-	void FPBDRigidsEvolutionBase::Serialize(FChaosArchive& Ar)
+	template <typename Traits>
+	void TPBDRigidsEvolutionBase<Traits>::Serialize(FChaosArchive& Ar)
 	{
 		ensure(false);	//disabled transient data serialization. Need to rethink
 		int32 DefaultBroadphaseType = ConfigSettings.BroadphaseType;
@@ -660,4 +676,8 @@ namespace Chaos
 
 		ConfigSettings.BroadphaseType = DefaultBroadphaseType;
 	}
+
+#define EVOLUTION_TRAIT(Trait) template class TPBDRigidsEvolutionBase<Trait>;
+#include "Chaos/EvolutionTraits.inl"
+#undef EVOLUTION_TRAIT
 }
diff --git a/Engine/Source/Runtime/Experimental/Chaos/Private/Chaos/PBDRigidsEvolutionGBF.cpp b/Engine/Source/Runtime/Experimental/Chaos/Private/Chaos/PBDRigidsEvolutionGBF.cpp
--- a/Engine/Source/Runtime/Experimental/Chaos/Private/Chaos/PBDRigidsEvolutionGBF.cpp
+++ b/Engine/Source/Runtime/Experimental/Chaos/Private/Chaos/PBDRigidsEvolutionGBF.cpp
@@ -59,20 +59,20 @@ float HackCCD_DepthThreshold = 0.05f;
 FAutoConsoleVariableRef CVarHackCCDDepthThreshold(TEXT("p.Chaos.CCD.DepthThreshold"), HackCCD_DepthThreshold, TEXT("When returning to TOI, leave this much contact depth (as a fraction of MinBounds)"));
 
 
-DECLARE_CYCLE_STAT(TEXT("FPBDRigidsEvolutionGBF::AdvanceOneTimeStep"), STAT_Evolution_AdvanceOneTimeStep, STATGROUP_Chaos);
-DECLARE_CYCLE_STAT(TEXT("FPBDRigidsEvolutionGBF::Integrate"), STAT_Evolution_Integrate, STATGROUP_Chaos);
-DECLARE_CYCLE_STAT(TEXT("FPBDRigidsEvolutionGBF::KinematicTargets"), STAT_Evolution_KinematicTargets, STATGROUP_Chaos);
-DECLARE_CYCLE_STAT(TEXT("FPBDRigidsEvolutionGBF::PrepareConstraints"), STAT_Evolution_PrepareConstraints, STATGROUP_Chaos);
-DECLARE_CYCLE_STAT(TEXT("FPBDRigidsEvolutionGBF::UnprepareConstraints"), STAT_Evolution_UnprepareConstraints, STATGROUP_Chaos);
-DECLARE_CYCLE_STAT(TEXT("FPBDRigidsEvolutionGBF::ApplyConstraints"), STAT_Evolution_ApplyConstraints, STATGROUP_Chaos);
-DECLARE_CYCLE_STAT(TEXT("FPBDRigidsEvolutionGBF::UpdateVelocities"), STAT_Evolution_UpdateVelocites, STATGROUP_Chaos);
-DECLARE_CYCLE_STAT(TEXT("FPBDRigidsEvolutionGBF::ApplyPushOut"), STAT_Evolution_ApplyPushOut, STATGROUP_Chaos);
-DECLARE_CYCLE_STAT(TEXT("FPBDRigidsEvolutionGBF::DetectCollisions"), STAT_Evolution_DetectCollisions, STATGROUP_Chaos);
-DECLARE_CYCLE_STAT(TEXT("FPBDRigidsEvolutionGBF::UpdateConstraintPositionBasedState"), STAT_Evolution_UpdateConstraintPositionBasedState, STATGROUP_Chaos);
-DECLARE_CYCLE_STAT(TEXT("FPBDRigidsEvolutionGBF::CreateConstraintGraph"), STAT_Evolution_CreateConstraintGraph, STATGROUP_Chaos);
-DECLARE_CYCLE_STAT(TEXT("FPBDRigidsEvolutionGBF::CreateIslands"), STAT_Evolution_CreateIslands, STATGROUP_Chaos);
-DECLARE_CYCLE_STAT(TEXT("FPBDRigidsEvolutionGBF::ParallelSolve"), STAT_Evolution_ParallelSolve, STATGROUP_Chaos);
-DECLARE_CYCLE_STAT(TEXT("FPBDRigidsEvolutionGBF::DeactivateSleep"), STAT_Evolution_DeactivateSleep, STATGROUP_Chaos);
+DECLARE_CYCLE_STAT(TEXT("TPBDRigidsEvolutionGBF::AdvanceOneTimeStep"), STAT_Evolution_AdvanceOneTimeStep, STATGROUP_Chaos);
+DECLARE_CYCLE_STAT(TEXT("TPBDRigidsEvolutionGBF::Integrate"), STAT_Evolution_Integrate, STATGROUP_Chaos);
+DECLARE_CYCLE_STAT(TEXT("TPBDRigidsEvolutionGBF::KinematicTargets"), STAT_Evolution_KinematicTargets, STATGROUP_Chaos);
+DECLARE_CYCLE_STAT(TEXT("TPBDRigidsEvolutionGBF::PrepareConstraints"), STAT_Evolution_PrepareConstraints, STATGROUP_Chaos);
+DECLARE_CYCLE_STAT(TEXT("TPBDRigidsEvolutionGBF::UnprepareConstraints"), STAT_Evolution_UnprepareConstraints, STATGROUP_Chaos);
+DECLARE_CYCLE_STAT(TEXT("TPBDRigidsEvolutionGBF::ApplyConstraints"), STAT_Evolution_ApplyConstraints, STATGROUP_Chaos);
+DECLARE_CYCLE_STAT(TEXT("TPBDRigidsEvolutionGBF::UpdateVelocities"), STAT_Evolution_UpdateVelocites, STATGROUP_Chaos);
+DECLARE_CYCLE_STAT(TEXT("TPBDRigidsEvolutionGBF::ApplyPushOut"), STAT_Evolution_ApplyPushOut, STATGROUP_Chaos);
+DECLARE_CYCLE_STAT(TEXT("TPBDRigidsEvolutionGBF::DetectCollisions"), STAT_Evolution_DetectCollisions, STATGROUP_Chaos);
+DECLARE_CYCLE_STAT(TEXT("TPBDRigidsEvolutionGBF::UpdateConstraintPositionBasedState"), STAT_Evolution_UpdateConstraintPositionBasedState, STATGROUP_Chaos);
+DECLARE_CYCLE_STAT(TEXT("TPBDRigidsEvolutionGBF::CreateConstraintGraph"), STAT_Evolution_CreateConstraintGraph, STATGROUP_Chaos);
+DECLARE_CYCLE_STAT(TEXT("TPBDRigidsEvolutionGBF::CreateIslands"), STAT_Evolution_CreateIslands, STATGROUP_Chaos);
+DECLARE_CYCLE_STAT(TEXT("TPBDRigidsEvolutionGBF::ParallelSolve"), STAT_Evolution_ParallelSolve, STATGROUP_Chaos);
+DECLARE_CYCLE_STAT(TEXT("TPBDRigidsEvolutionGBF::DeactivateSleep"), STAT_Evolution_DeactivateSleep, STATGROUP_Chaos);
 
 int32 SerializeEvolution = 0;
 FAutoConsoleVariableRef CVarSerializeEvolution(TEXT("p.SerializeEvolution"), SerializeEvolution, TEXT(""));
@@ -109,7 +109,8 @@ void SerializeToDisk(TEvolution& Evolution)
 }
 #endif
 
-void FPBDRigidsEvolutionGBF::Advance(const FReal Dt, const FReal MaxStepDt, const int32 MaxSteps)
+template <typename Traits>
+void TPBDRigidsEvolutionGBF<Traits>::Advance(const FReal Dt,const FReal MaxStepDt,const int32 MaxSteps)
 {
 	// Determine how many steps we would like to take
 	int32 NumSteps = FMath::CeilToInt(Dt / MaxStepDt);
@@ -249,7 +250,8 @@ void CCDHack(const FReal Dt, TParticleView<TPBDRigidParticles<FReal, 3>>& Partic
 //
 //
 
-void FPBDRigidsEvolutionGBF::AdvanceOneTimeStep(const FReal Dt, const FReal StepFraction)
+template <typename Traits>
+void TPBDRigidsEvolutionGBF<Traits>::AdvanceOneTimeStep(const FReal Dt,const FReal StepFraction)
 {
 	PrepareTick();
 
@@ -258,7 +260,8 @@ void FPBDRigidsEvolutionGBF::AdvanceOneTimeStep(const FReal Dt, const FReal Step
 	UnprepareTick();
 }
 
-void FPBDRigidsEvolutionGBF::AdvanceOneTimeStepImpl(const FReal Dt, const FReal StepFraction)
+template <typename Traits>
+void TPBDRigidsEvolutionGBF<Traits>::AdvanceOneTimeStepImpl(const FReal Dt,const FReal StepFraction)
 {
 	SCOPE_CYCLE_COUNTER(STAT_Evolution_AdvanceOneTimeStep);
 
@@ -442,7 +445,8 @@ void FPBDRigidsEvolutionGBF::AdvanceOneTimeStepImpl(const FReal Dt, const FReal
 	ParticleUpdatePosition(Particles.GetActiveParticlesView(), Dt);
 }
 
-FPBDRigidsEvolutionGBF::FPBDRigidsEvolutionGBF(TPBDRigidsSOAs<FReal, 3>& InParticles, THandleArray<FChaosPhysicsMaterial>& SolverPhysicsMaterials, int32 InNumIterations, int32 InNumPushoutIterations, bool InIsSingleThreaded)
+template <typename Traits>
+TPBDRigidsEvolutionGBF<Traits>::TPBDRigidsEvolutionGBF(TPBDRigidsSOAs<FReal,3>& InParticles,THandleArray<FChaosPhysicsMaterial>& SolverPhysicsMaterials,int32 InNumIterations,int32 InNumPushoutIterations,bool InIsSingleThreaded)
 	: Base(InParticles, SolverPhysicsMaterials, InNumIterations, InNumPushoutIterations, InIsSingleThreaded)
 	, Clustering(*this, Particles.GetClusteredParticles())
 	, CollisionConstraints(InParticles, Collided, PhysicsMaterials, DefaultNumPairIterations, DefaultNumPushOutPairIterations)
@@ -486,10 +490,14 @@ FPBDRigidsEvolutionGBF::FPBDRigidsEvolutionGBF(TPBDRigidsSOAs<FReal, 3>& InParti
 	NarrowPhase.GetContext().bAllowManifolds = false;
 }
 
-void FPBDRigidsEvolutionGBF::Serialize(FChaosArchive& Ar)
+template <typename Traits>
+void TPBDRigidsEvolutionGBF<Traits>::Serialize(FChaosArchive& Ar)
 {
 	Base::Serialize(Ar);
 }
 
+#define EVOLUTION_TRAIT(Trait) template class TPBDRigidsEvolutionGBF<Trait>;
+#include "Chaos/EvolutionTraits.inl"
+#undef EVOLUTION_TRAIT
 }
 
diff --git a/Engine/Source/Runtime/Experimental/Chaos/Public/Chaos/EvolutionTraits.h b/Engine/Source/Runtime/Experimental/Chaos/Public/Chaos/EvolutionTraits.h
new file mode 100644
--- /dev/null
+++ b/Engine/Source/Runtime/Experimental/Chaos/Public/Chaos/EvolutionTraits.h
@@ -0,0 +1,16 @@
+// Copyright Epic Games, Inc. All Rights Reserved.
+#pragma once
+
+namespace Chaos
+{
+
+struct FRewindableEvolutionTraits
+{
+	static constexpr bool IsRewindable(){ return true; }
+};
+
+struct FNonRewindableEvolutionTraits
+{
+	static constexpr bool IsRewindable(){ return false; }
+};
+}
diff --git a/Engine/Source/Runtime/Experimental/Chaos/Public/Chaos/EvolutionTraits.inl b/Engine/Source/Runtime/Experimental/Chaos/Public/Chaos/EvolutionTraits.inl
new file mode 100644
--- /dev/null
+++ b/Engine/Source/Runtime/Experimental/Chaos/Public/Chaos/EvolutionTraits.inl
@@ -0,0 +1,3 @@
+// Copyright Epic Games, Inc. All Rights Reserved.
+EVOLUTION_TRAIT(FNonRewindableEvolutionTraits)
+EVOLUTION_TRAIT(FRewindableEvolutionTraits)
\ No newline at end of file
diff --git a/Engine/Source/Runtime/Experimental/Chaos/Public/Chaos/PBDRigidsEvolution.h b/Engine/Source/Runtime/Experimental/Chaos/Public/Chaos/PBDRigidsEvolution.h
--- a/Engine/Source/Runtime/Experimental/Chaos/Public/Chaos/PBDRigidsEvolution.h
+++ b/Engine/Source/Runtime/Experimental/Chaos/Public/Chaos/PBDRigidsEvolution.h
@@ -13,6 +13,8 @@
 #include "HAL/Event.h"
 #include "Chaos/PBDRigidsSOAs.h"
 #include "Chaos/SpatialAccelerationCollection.h"
+#include "Chaos/EvolutionTraits.h"
+#include "Chaos/PBDRigidsEvolutionFwd.h"
 
 
 extern int32 ChaosRigidsEvolutionApplyAllowEarlyOutCVar;
@@ -220,7 +222,8 @@ struct CHAOS_API ISpatialAccelerationCollectionFactory
 	virtual ~ISpatialAccelerationCollectionFactory() = default;
 };
 
-class FPBDRigidsEvolutionBase
+template <typename Traits>
+class TPBDRigidsEvolutionBase
 {
   public:
 	typedef TFunction<void(TTransientPBDRigidParticleHandle<FReal, 3>& Particle, const FReal)> FForceRule;
@@ -231,8 +234,8 @@ class FPBDRigidsEvolutionBase
 
 	friend void ChaosTest::TestPendingSpatialDataHandlePointerConflict();
 
-	CHAOS_API FPBDRigidsEvolutionBase(TPBDRigidsSOAs<FReal, 3>& InParticles, THandleArray<FChaosPhysicsMaterial>& InSolverPhysicsMaterials, int32 InNumIterations = 1, int32 InNumPushOutIterations = 1, bool InIsSingleThreaded = false);
-	CHAOS_API virtual ~FPBDRigidsEvolutionBase();
+	CHAOS_API TPBDRigidsEvolutionBase(TPBDRigidsSOAs<FReal, 3>& InParticles, THandleArray<FChaosPhysicsMaterial>& InSolverPhysicsMaterials, int32 InNumIterations = 1, int32 InNumPushOutIterations = 1, bool InIsSingleThreaded = false);
+	CHAOS_API virtual ~TPBDRigidsEvolutionBase();
 
 	CHAOS_API TArray<TGeometryParticleHandle<FReal, 3>*> CreateStaticParticles(int32 NumParticles, const FUniqueIdx* ExistingIndices = nullptr, const TGeometryParticleParameters<FReal, 3>& Params = TGeometryParticleParameters<FReal, 3>())
 	{
@@ -909,4 +912,8 @@ protected:
 	TUniquePtr<ISpatialAccelerationCollectionFactory> SpatialCollectionFactory;
 };
 
+#define EVOLUTION_TRAIT(Trait) extern template TPBDRigidsEvolutionBase<Trait>;
+#include "Chaos/EvolutionTraits.inl"
+#undef EVOLUTION_TRAIT
+
 }
diff --git a/Engine/Source/Runtime/Experimental/Chaos/Public/Chaos/PBDRigidsEvolutionFwd.h b/Engine/Source/Runtime/Experimental/Chaos/Public/Chaos/PBDRigidsEvolutionFwd.h
new file mode 100644
--- /dev/null
+++ b/Engine/Source/Runtime/Experimental/Chaos/Public/Chaos/PBDRigidsEvolutionFwd.h
@@ -0,0 +1,14 @@
+// Copyright Epic Games, Inc. All Rights Reserved.
+#pragma once
+
+namespace Chaos
+{
+	template <typename Traits>
+	class TPBDRigidsEvolutionBase;
+
+	template <typename Traits>
+	class TPBDRigidsEvolutionGBF;
+
+	//The default evolution used by unreal
+	using FPBDRigidsEvolution = TPBDRigidsEvolutionGBF<struct FNonRewindableEvolutionTraits>;
+}
\ No newline at end of file
diff --git a/Engine/Source/Runtime/Experimental/Chaos/Public/Chaos/PBDRigidsEvolutionGBF.h b/Engine/Source/Runtime/Experimental/Chaos/Public/Chaos/PBDRigidsEvolutionGBF.h
--- a/Engine/Source/Runtime/Experimental/Chaos/Public/Chaos/PBDRigidsEvolutionGBF.h
+++ b/Engine/Source/Runtime/Experimental/Chaos/Public/Chaos/PBDRigidsEvolutionGBF.h
@@ -25,8 +25,6 @@ namespace Chaos
 	CHAOS_API extern float HackLinearDrag;
 	CHAOS_API extern float HackAngularDrag;
 
-	class FPBDRigidsEvolutionGBF;
-
 	using FPBDRigidsEvolutionCallback = TFunction<void()>;
 
 	using FPBDRigidsEvolutionIslandCallback = TFunction<void(int32 Island)>;
@@ -34,11 +32,12 @@ namespace Chaos
 	using FPBDRigidsEvolutionInternalHandleCallback = TFunction<void(
 		const TGeometryParticleHandle<float, 3> * OldParticle,
 		const TGeometryParticleHandle<float, 3> * NewParticle)>;
-	
-	class FPBDRigidsEvolutionGBF : public FPBDRigidsEvolutionBase
+
+	template <typename Traits>
+	class TPBDRigidsEvolutionGBF : public TPBDRigidsEvolutionBase<Traits>
 	{
 	public:
-		using Base = FPBDRigidsEvolutionBase;
+		using Base = TPBDRigidsEvolutionBase<Traits>;
 
 		using FGravityForces = TPerParticleGravity<FReal, 3>;
 		using FCollisionConstraints = FPBDCollisionConstraints;
@@ -54,8 +53,8 @@ namespace Chaos
 		// @todo(chaos): Required by clustering - clean up
 		using Base::ApplyPushOut;
 
-		CHAOS_API FPBDRigidsEvolutionGBF(TPBDRigidsSOAs<FReal, 3>& InParticles, THandleArray<FChaosPhysicsMaterial>& SolverPhysicsMaterials, int32 InNumIterations = DefaultNumIterations, int32 InNumPushoutIterations = DefaultNumPushOutIterations, bool InIsSingleThreaded = false);
-		CHAOS_API ~FPBDRigidsEvolutionGBF() {}
+		CHAOS_API TPBDRigidsEvolutionGBF(TPBDRigidsSOAs<FReal, 3>& InParticles, THandleArray<FChaosPhysicsMaterial>& SolverPhysicsMaterials, int32 InNumIterations = DefaultNumIterations, int32 InNumPushoutIterations = DefaultNumPushOutIterations, bool InIsSingleThreaded = false);
+		CHAOS_API ~TPBDRigidsEvolutionGBF() {}
 
 		void SetPostIntegrateCallback(const FPBDRigidsEvolutionCallback& Cb)
 		{
@@ -197,7 +196,7 @@ namespace Chaos
 
 		CHAOS_API void AdvanceOneTimeStepImpl(const FReal dt, const FReal StepFraction);
 
-		TPBDRigidClustering<FPBDRigidsEvolutionGBF, FPBDCollisionConstraints, FReal, 3> Clustering;
+		TPBDRigidClustering<TPBDRigidsEvolutionGBF<Traits>, FPBDCollisionConstraints, FReal, 3> Clustering;
 
 		FGravityForces GravityForces;
 		FCollisionConstraints CollisionConstraints;
@@ -214,4 +213,8 @@ namespace Chaos
 		FPBDRigidsEvolutionIslandCallback PostApplyPushOutCallback;
 		FPBDRigidsEvolutionInternalHandleCallback InternalParticleInitilization;
 	};
+
+#define EVOLUTION_TRAIT(Trait) extern template TPBDRigidsEvolutionGBF<Trait>;
+#include "Chaos/EvolutionTraits.inl"
+#undef EVOLUTION_TRAIT
 }
diff --git a/Engine/Source/Runtime/Experimental/ChaosSolvers/Public/PBDRigidsSolver.h b/Engine/Source/Runtime/Experimental/ChaosSolvers/Public/PBDRigidsSolver.h
--- a/Engine/Source/Runtime/Experimental/ChaosSolvers/Public/PBDRigidsSolver.h
+++ b/Engine/Source/Runtime/Experimental/ChaosSolvers/Public/PBDRigidsSolver.h
@@ -96,7 +96,7 @@ namespace Chaos
 
 		typedef Chaos::TGeometryParticle<float, 3> FParticle;
 		typedef Chaos::TGeometryParticleHandle<float, 3> FHandle;
-		typedef Chaos::FPBDRigidsEvolutionGBF FPBDRigidsEvolution;
+		typedef Chaos::FPBDRigidsEvolution FPBDRigidsEvolution;
 
 		typedef TPBDRigidDynamicSpringConstraints<float, 3> FRigidDynamicSpringConstraints;
 		typedef TPBDPositionConstraints<float, 3> FPositionConstraints;
diff --git a/Engine/Source/Runtime/PhysicsCore/Private/PhysTestSerializer.cpp b/Engine/Source/Runtime/PhysicsCore/Private/PhysTestSerializer.cpp
--- a/Engine/Source/Runtime/PhysicsCore/Private/PhysTestSerializer.cpp
+++ b/Engine/Source/Runtime/PhysicsCore/Private/PhysTestSerializer.cpp
@@ -119,7 +119,7 @@ void FPhysTestSerializer::Serialize(Chaos::FChaosArchive& Ar)
 	ChaosContext = Ar.StealContext();
 }
 
-void FPhysTestSerializer::SetPhysicsData(Chaos::FPBDRigidsEvolutionGBF& Evolution)
+void FPhysTestSerializer::SetPhysicsData(Chaos::FPBDRigidsEvolution& Evolution)
 {
 	bDiskDataIsChaos = true;
 	Data.Empty();
@@ -336,11 +336,11 @@ void FPhysTestSerializer::CreateChaosData()
 			++Idx;
 		}
 
-		ChaosEvolution = MakeUnique<FPBDRigidsEvolutionGBF>(Particles, PhysicalMaterials);
+		ChaosEvolution = MakeUnique<FPBDRigidsEvolution>(Particles, PhysicalMaterials);
 	}
 	else
 	{
-		ChaosEvolution = MakeUnique<FPBDRigidsEvolutionGBF>(Particles, PhysicalMaterials);
+		ChaosEvolution = MakeUnique<FPBDRigidsEvolution>(Particles, PhysicalMaterials);
 
 		FMemoryReader Ar(Data);
 		FChaosArchive ChaosAr(Ar);
diff --git a/Engine/Source/Runtime/PhysicsCore/Private/SQCapture.cpp b/Engine/Source/Runtime/PhysicsCore/Private/SQCapture.cpp
--- a/Engine/Source/Runtime/PhysicsCore/Private/SQCapture.cpp
+++ b/Engine/Source/Runtime/PhysicsCore/Private/SQCapture.cpp
@@ -356,7 +356,7 @@ void FSQCapture::Serialize(Chaos::FChaosArchive& Ar)
 	}
 }
 
-void FSQCapture::StartCaptureChaosSweep(const Chaos::FPBDRigidsEvolutionGBF& Evolution, const Chaos::FImplicitObject& InQueryGeom, const FTransform& InStartTM, const FVector& InDir, float InDeltaMag, FHitFlags InOutputFlags, const FQueryFilterData& QueryFilter, const FCollisionFilterData& FilterData, ICollisionQueryFilterCallbackBase& Callback)
+void FSQCapture::StartCaptureChaosSweep(const Chaos::FPBDRigidsEvolution& Evolution, const Chaos::FImplicitObject& InQueryGeom, const FTransform& InStartTM, const FVector& InDir, float InDeltaMag, FHitFlags InOutputFlags, const FQueryFilterData& QueryFilter, const FCollisionFilterData& FilterData, ICollisionQueryFilterCallbackBase& Callback)
 {
 	if (IsInGameThread())
 	{
@@ -388,7 +388,7 @@ void FSQCapture::EndCaptureChaosSweep(const ChaosInterface::FSQHitBuffer<ChaosIn
 #endif
 }
 
-void FSQCapture::StartCaptureChaosRaycast(const Chaos::FPBDRigidsEvolutionGBF& Evolution, const FVector& InStartPoint, const FVector& InDir, float InDeltaMag, FHitFlags InOutputFlags, const FQueryFilterData& QueryFilter, const FCollisionFilterData& FilterData, ICollisionQueryFilterCallbackBase& Callback)
+void FSQCapture::StartCaptureChaosRaycast(const Chaos::FPBDRigidsEvolution& Evolution, const FVector& InStartPoint, const FVector& InDir, float InDeltaMag, FHitFlags InOutputFlags, const FQueryFilterData& QueryFilter, const FCollisionFilterData& FilterData, ICollisionQueryFilterCallbackBase& Callback)
 {
 	if (IsInGameThread())
 	{
@@ -418,7 +418,7 @@ void FSQCapture::EndCaptureChaosRaycast(const ChaosInterface::FSQHitBuffer<Chaos
 #endif
 }
 
-void FSQCapture::StartCaptureChaosOverlap(const Chaos::FPBDRigidsEvolutionGBF& Evolution, const Chaos::FImplicitObject& InQueryGeom, const FTransform& InStartTM, const FQueryFilterData& QueryFilter, const FCollisionFilterData& FilterData, ICollisionQueryFilterCallbackBase& Callback)
+void FSQCapture::StartCaptureChaosOverlap(const Chaos::FPBDRigidsEvolution& Evolution, const Chaos::FImplicitObject& InQueryGeom, const FTransform& InStartTM, const FQueryFilterData& QueryFilter, const FCollisionFilterData& FilterData, ICollisionQueryFilterCallbackBase& Callback)
 {
 	if (IsInGameThread())
 	{
@@ -447,7 +447,7 @@ void FSQCapture::EndCaptureChaosOverlap(const ChaosInterface::FSQHitBuffer<Chaos
 #endif
 }
 
-void FSQCapture::CaptureChaosFilterResults(const Chaos::FPBDRigidsEvolutionGBF& TransientEvolution, const FCollisionFilterData& FilterData, ICollisionQueryFilterCallbackBase& Callback)
+void FSQCapture::CaptureChaosFilterResults(const Chaos::FPBDRigidsEvolution& TransientEvolution, const FCollisionFilterData& FilterData, ICollisionQueryFilterCallbackBase& Callback)
 {
 	using namespace Chaos;
 	const TPBDRigidsSOAs<float, 3>& Particles = TransientEvolution.GetParticles();
diff --git a/Engine/Source/Runtime/PhysicsCore/Public/PhysTestSerializer.h b/Engine/Source/Runtime/PhysicsCore/Public/PhysTestSerializer.h
--- a/Engine/Source/Runtime/PhysicsCore/Public/PhysTestSerializer.h
+++ b/Engine/Source/Runtime/PhysicsCore/Public/PhysTestSerializer.h
@@ -35,7 +35,6 @@ namespace physx
 
 namespace Chaos
 {
-	class FPBDRigidsEvolutionGBF;
 
 	class FChaosArchive;
 }
@@ -53,7 +52,7 @@ public:
 	void Serialize(const TCHAR* FilePrefix);
 
 	//Set the data from an external source. This will obliterate any existing data. Make sure you are not holding on to old internal data as it will go away
-	void SetPhysicsData(Chaos::FPBDRigidsEvolutionGBF& ChaosEvolution);
+	void SetPhysicsData(Chaos::FPBDRigidsEvolution& ChaosEvolution);
 
 #if PHYSICS_INTERFACE_PHYSX
 	void SetPhysicsData(physx::PxScene& Scene);
@@ -89,7 +88,7 @@ public:
 		return SQCapture.Get();
 	}
 
-	Chaos::FPBDRigidsEvolutionGBF* GetChaosData()
+	Chaos::FPBDRigidsEvolution* GetChaosData()
 	{
 #if 0
 		if (!bChaosDataReady)
@@ -140,7 +139,7 @@ private:
 
 	TUniquePtr<FSQCapture> SQCapture;
 
-	TUniquePtr<Chaos::FPBDRigidsEvolutionGBF> ChaosEvolution;
+	TUniquePtr<Chaos::FPBDRigidsEvolution> ChaosEvolution;
 	Chaos::TPBDRigidsSOAs<float, 3> Particles;
 	Chaos::THandleArray<Chaos::FChaosPhysicsMaterial> PhysicalMaterials;
 	TArray <TUniquePtr<Chaos::TGeometryParticle<float, 3>>> GTParticles;
diff --git a/Engine/Source/Runtime/PhysicsCore/Public/SQCapture.h b/Engine/Source/Runtime/PhysicsCore/Public/SQCapture.h
--- a/Engine/Source/Runtime/PhysicsCore/Public/SQCapture.h
+++ b/Engine/Source/Runtime/PhysicsCore/Public/SQCapture.h
@@ -11,7 +11,7 @@
 #include "PhysXInterfaceWrapperCore.h"
 #include "CollisionQueryFilterCallbackCore.h"
 #include "ChaosInterfaceWrapperCore.h"
-
+#include "Chaos/PBDrigidsEvolutionFwd.h"
 
 class FPhysTestSerializer;
 
@@ -24,8 +24,6 @@ namespace ChaosInterface
 
 namespace Chaos
 {
-	class FPBDRigidsEvolutionGBF;
-
 	class FChaosArchive;
 }
 
@@ -54,13 +52,13 @@ struct PHYSICSCORE_API FSQCapture
 	void EndCapturePhysXOverlap(const PxHitCallback<PxOverlapHit>& Results);
 #endif
 
-	void StartCaptureChaosSweep(const Chaos::FPBDRigidsEvolutionGBF& Evolution, const Chaos::FImplicitObject& InQueryGeom, const FTransform& InStartTM, const FVector& InDir, float InDeltaMag, FHitFlags InOutputFlags, const FQueryFilterData& QueryFilter, const FCollisionFilterData& FilterData, ICollisionQueryFilterCallbackBase& Callback);
+	void StartCaptureChaosSweep(const Chaos::FPBDRigidsEvolution& Evolution, const Chaos::FImplicitObject& InQueryGeom, const FTransform& InStartTM, const FVector& InDir, float InDeltaMag, FHitFlags InOutputFlags, const FQueryFilterData& QueryFilter, const FCollisionFilterData& FilterData, ICollisionQueryFilterCallbackBase& Callback);
 	void EndCaptureChaosSweep(const ChaosInterface::FSQHitBuffer<ChaosInterface::FSweepHit>& Results);
 
-	void StartCaptureChaosRaycast(const Chaos::FPBDRigidsEvolutionGBF& Evolution, const FVector& InStartPoint, const FVector& InDir, float InDeltaMag, FHitFlags InOutputFlags, const FQueryFilterData& QueryFilter, const FCollisionFilterData& FilterData, ICollisionQueryFilterCallbackBase& Callback);
+	void StartCaptureChaosRaycast(const Chaos::FPBDRigidsEvolution& Evolution, const FVector& InStartPoint, const FVector& InDir, float InDeltaMag, FHitFlags InOutputFlags, const FQueryFilterData& QueryFilter, const FCollisionFilterData& FilterData, ICollisionQueryFilterCallbackBase& Callback);
 	void EndCaptureChaosRaycast(const ChaosInterface::FSQHitBuffer<ChaosInterface::FRaycastHit>& Results);
 
-	void StartCaptureChaosOverlap(const Chaos::FPBDRigidsEvolutionGBF& Evolution, const Chaos::FImplicitObject& InQueryGeom, const FTransform& InStartTM, const FQueryFilterData& QueryFilter, const FCollisionFilterData& FilterData, ICollisionQueryFilterCallbackBase& Callback);
+	void StartCaptureChaosOverlap(const Chaos::FPBDRigidsEvolution& Evolution, const Chaos::FImplicitObject& InQueryGeom, const FTransform& InStartTM, const FQueryFilterData& QueryFilter, const FCollisionFilterData& FilterData, ICollisionQueryFilterCallbackBase& Callback);
 	void EndCaptureChaosOverlap(const ChaosInterface::FSQHitBuffer<ChaosInterface::FOverlapHit>& Results);
 
 	ECollisionQueryHitType GetFilterResult(const Chaos::FPerShapeData* Shape, const Chaos::TGeometryParticle<float,3>* Actor) const;
@@ -159,7 +157,7 @@ private:
 
 	FPhysTestSerializer& PhysSerializer;
 
-	void CaptureChaosFilterResults(const Chaos::FPBDRigidsEvolutionGBF& Evolution, const FCollisionFilterData& FilterData, ICollisionQueryFilterCallbackBase& Callback);
+	void CaptureChaosFilterResults(const Chaos::FPBDRigidsEvolution& Evolution, const FCollisionFilterData& FilterData, ICollisionQueryFilterCallbackBase& Callback);
 
 	TMap<Chaos::TGeometryParticle<float,3>*, TArray<TPair<Chaos::FPerShapeData*, ECollisionQueryHitType>>> ChaosActorToShapeHitsArray;
 
