Fix CIS: extern template class instead of extern template

diff --git a/Engine/Source/Runtime/Experimental/Chaos/Public/Chaos/PBDRigidsEvolution.h b/Engine/Source/Runtime/Experimental/Chaos/Public/Chaos/PBDRigidsEvolution.h
--- a/Engine/Source/Runtime/Experimental/Chaos/Public/Chaos/PBDRigidsEvolution.h
+++ b/Engine/Source/Runtime/Experimental/Chaos/Public/Chaos/PBDRigidsEvolution.h
@@ -912,7 +912,7 @@ protected:
 	TUniquePtr<ISpatialAccelerationCollectionFactory> SpatialCollectionFactory;
 };
 
-#define EVOLUTION_TRAIT(Trait) extern template TPBDRigidsEvolutionBase<Trait>;
+#define EVOLUTION_TRAIT(Trait) extern template class TPBDRigidsEvolutionBase<Trait>;
 #include "Chaos/EvolutionTraits.inl"
 #undef EVOLUTION_TRAIT
 
diff --git a/Engine/Source/Runtime/Experimental/Chaos/Public/Chaos/PBDRigidsEvolutionGBF.h b/Engine/Source/Runtime/Experimental/Chaos/Public/Chaos/PBDRigidsEvolutionGBF.h
--- a/Engine/Source/Runtime/Experimental/Chaos/Public/Chaos/PBDRigidsEvolutionGBF.h
+++ b/Engine/Source/Runtime/Experimental/Chaos/Public/Chaos/PBDRigidsEvolutionGBF.h
@@ -214,7 +214,7 @@ namespace Chaos
 		FPBDRigidsEvolutionInternalHandleCallback InternalParticleInitilization;
 	};
 
-#define EVOLUTION_TRAIT(Trait) extern template TPBDRigidsEvolutionGBF<Trait>;
+#define EVOLUTION_TRAIT(Trait) extern template class TPBDRigidsEvolutionGBF<Trait>;
 #include "Chaos/EvolutionTraits.inl"
 #undef EVOLUTION_TRAIT
 }
