Fixing CIS: dependent names

diff --git a/Engine/Source/Runtime/Experimental/Chaos/Public/Chaos/PBDRigidsEvolutionGBF.h b/Engine/Source/Runtime/Experimental/Chaos/Public/Chaos/PBDRigidsEvolutionGBF.h
--- a/Engine/Source/Runtime/Experimental/Chaos/Public/Chaos/PBDRigidsEvolutionGBF.h
+++ b/Engine/Source/Runtime/Experimental/Chaos/Public/Chaos/PBDRigidsEvolutionGBF.h
@@ -38,6 +38,15 @@ namespace Chaos
 	{
 	public:
 		using Base = TPBDRigidsEvolutionBase<Traits>;
+		using Base::Particles;
+		using Base::FForceRule;
+		using Base::PrepareTick;
+		using Base::UnprepareTick;
+		using Base::ApplyKinematicTargets;
+		using Base::UpdateConstraintPositionBasedState;
+		using Base::InternalAcceleration;
+		using Base::CreateConstraintGraph;
+		using Base::CreateIslands;
 
 		using FGravityForces = TPerParticleGravity<FReal, 3>;
 		using FCollisionConstraints = FPBDCollisionConstraints;
