Fix CIS: dependent names part two

diff --git a/Engine/Source/Runtime/Experimental/Chaos/Public/Chaos/PBDRigidsEvolutionGBF.h b/Engine/Source/Runtime/Experimental/Chaos/Public/Chaos/PBDRigidsEvolutionGBF.h
--- a/Engine/Source/Runtime/Experimental/Chaos/Public/Chaos/PBDRigidsEvolutionGBF.h
+++ b/Engine/Source/Runtime/Experimental/Chaos/Public/Chaos/PBDRigidsEvolutionGBF.h
@@ -40,6 +40,7 @@ namespace Chaos
 		using Base = TPBDRigidsEvolutionBase<Traits>;
 		using Base::Particles;
 		using Base::FForceRule;
+		using Base::ForceRules;
 		using Base::PrepareTick;
 		using Base::UnprepareTick;
 		using Base::ApplyKinematicTargets;
