Fix CIS: missing typename

diff --git a/Engine/Source/Runtime/Experimental/Chaos/Public/Chaos/PBDRigidsEvolutionGBF.h b/Engine/Source/Runtime/Experimental/Chaos/Public/Chaos/PBDRigidsEvolutionGBF.h
--- a/Engine/Source/Runtime/Experimental/Chaos/Public/Chaos/PBDRigidsEvolutionGBF.h
+++ b/Engine/Source/Runtime/Experimental/Chaos/Public/Chaos/PBDRigidsEvolutionGBF.h
@@ -39,7 +39,7 @@ namespace Chaos
 	public:
 		using Base = TPBDRigidsEvolutionBase<Traits>;
 		using Base::Particles;
-		using Base::FForceRule;
+		using typename Base::FForceRule;
 		using Base::ForceRules;
 		using Base::PrepareTick;
 		using Base::UnprepareTick;
@@ -48,6 +48,8 @@ namespace Chaos
 		using Base::InternalAcceleration;
 		using Base::CreateConstraintGraph;
 		using Base::CreateIslands;
+		using Base::GetParticles;
+		using Base::DirtyParticle;
 
 		using FGravityForces = TPerParticleGravity<FReal, 3>;
 		using FCollisionConstraints = FPBDCollisionConstraints;
