CIS Fix: more dependent names

diff --git a/Engine/Source/Runtime/Experimental/Chaos/Public/Chaos/PBDRigidsEvolutionGBF.h b/Engine/Source/Runtime/Experimental/Chaos/Public/Chaos/PBDRigidsEvolutionGBF.h
--- a/Engine/Source/Runtime/Experimental/Chaos/Public/Chaos/PBDRigidsEvolutionGBF.h
+++ b/Engine/Source/Runtime/Experimental/Chaos/Public/Chaos/PBDRigidsEvolutionGBF.h
@@ -50,6 +50,19 @@ namespace Chaos
 		using Base::CreateIslands;
 		using Base::GetParticles;
 		using Base::DirtyParticle;
+		using Base::SetPhysicsMaterial;
+		using Base::SetPerParticlePhysicsMaterial;
+		using Base::GetPerParticlePhysicsMaterial;
+		using Base::CreateParticle;
+		using Base::GenerateUniqueIdx;
+		using Base::DestroyParticle;
+		using Base::CreateClusteredParticles;
+		using Base::EnableParticle;
+		using Base::DisableParticles;
+		using Base::GetActiveClusteredArray;
+		using Base::NumIslands;
+		using Base::GetNonDisabledClusteredArray;
+		using Base::DisableParticle;
 
 		using FGravityForces = TPerParticleGravity<FReal, 3>;
 		using FCollisionConstraints = FPBDCollisionConstraints;
