Fix CIS: Even more dependent names

diff --git a/Engine/Source/Runtime/Experimental/Chaos/Public/Chaos/PBDRigidsEvolutionGBF.h b/Engine/Source/Runtime/Experimental/Chaos/Public/Chaos/PBDRigidsEvolutionGBF.h
--- a/Engine/Source/Runtime/Experimental/Chaos/Public/Chaos/PBDRigidsEvolutionGBF.h
+++ b/Engine/Source/Runtime/Experimental/Chaos/Public/Chaos/PBDRigidsEvolutionGBF.h
@@ -63,6 +63,14 @@ namespace Chaos
 		using Base::NumIslands;
 		using Base::GetNonDisabledClusteredArray;
 		using Base::DisableParticle;
+		using Base::PrepareIteration;
+		using Base::GetConstraintGraph;
+		using Base::ApplyConstraints;
+		using Base::UpdateVelocities;
+		using Base::PhysicsMaterials;
+		using Base::ParticleDisableCount;
+		using Base::SolverPhysicsMaterials;
+		using Base::UnprepareIteration;
 
 		using FGravityForces = TPerParticleGravity<FReal, 3>;
 		using FCollisionConstraints = FPBDCollisionConstraints;
diff --git a/Engine/Source/Runtime/PhysicsCore/Public/SQCapture.h b/Engine/Source/Runtime/PhysicsCore/Public/SQCapture.h
--- a/Engine/Source/Runtime/PhysicsCore/Public/SQCapture.h
+++ b/Engine/Source/Runtime/PhysicsCore/Public/SQCapture.h
@@ -11,7 +11,7 @@
 #include "PhysXInterfaceWrapperCore.h"
 #include "CollisionQueryFilterCallbackCore.h"
 #include "ChaosInterfaceWrapperCore.h"
-#include "Chaos/PBDrigidsEvolutionFwd.h"
+#include "Chaos/PBDRigidsEvolutionFwd.h"
 
 class FPhysTestSerializer;
 
