Fix CIS: last dependent name

diff --git a/Engine/Source/Runtime/Experimental/Chaos/Public/Chaos/PBDRigidsEvolutionGBF.h b/Engine/Source/Runtime/Experimental/Chaos/Public/Chaos/PBDRigidsEvolutionGBF.h
--- a/Engine/Source/Runtime/Experimental/Chaos/Public/Chaos/PBDRigidsEvolutionGBF.h
+++ b/Engine/Source/Runtime/Experimental/Chaos/Public/Chaos/PBDRigidsEvolutionGBF.h
@@ -71,6 +71,13 @@ namespace Chaos
 		using Base::ParticleDisableCount;
 		using Base::SolverPhysicsMaterials;
 		using Base::UnprepareIteration;
+		using Base::CaptureRewindData;
+		using Base::Collided;
+		using Base::SetParticleUpdateVelocityFunction;
+		using Base::SetParticleUpdatePositionFunction;
+		using Base::AddForceFunction;
+		using Base::AddConstraintRule;
+		using Base::ParticleUpdatePosition;
 
 		using FGravityForces = TPerParticleGravity<FReal, 3>;
 		using FCollisionConstraints = FPBDCollisionConstraints;
