CIS - missing include for stat definition

diff --git a/Engine/Source/Runtime/Experimental/Chaos/Private/Chaos/Collision/CollisionDetector.cpp b/Engine/Source/Runtime/Experimental/Chaos/Private/Chaos/Collision/CollisionDetector.cpp
--- a/Engine/Source/Runtime/Experimental/Chaos/Private/Chaos/Collision/CollisionDetector.cpp
+++ b/Engine/Source/Runtime/Experimental/Chaos/Private/Chaos/Collision/CollisionDetector.cpp
@@ -1,6 +1,7 @@
 // Copyright Epic Games, Inc. All Rights Reserved.
 #include "Chaos/Collision/CollisionDetector.h"
 #include "Chaos/Collision/ParticlePairBroadPhase.h"
+#include "Chaos/Collision/SpatialAccelerationBroadPhase.h"
 #include "ChaosLog.h"
 #include "ChaosStats.h"
 
