CIS - fixes to Chaos changes coming down from 4.25 conflicting with update API in main.

diff --git a/Engine/Source/Runtime/Experimental/Chaos/Private/Chaos/CollisionResolution.cpp b/Engine/Source/Runtime/Experimental/Chaos/Private/Chaos/CollisionResolution.cpp
--- a/Engine/Source/Runtime/Experimental/Chaos/Private/Chaos/CollisionResolution.cpp
+++ b/Engine/Source/Runtime/Experimental/Chaos/Private/Chaos/CollisionResolution.cpp
@@ -2617,15 +2617,15 @@ namespace Chaos
 					Union0->FindAllIntersectingObjects(Children, Implicit1->BoundingBox());
 					for (const auto& Child0 : Children)
 					{
-						TRigidTransform<T, d> TransformedChild0 = Child0.Second * LocalTransform0;
-						ConstructConstraints(Particle0, Particle1, Child0.First, Implicit1, TransformedChild0, LocalTransform1, CullDistance, Context, NewConstraints);
+						TRigidTransform<FReal, 3> TransformedChild0 = Child0.Second * LocalTransform0;
+						ConstructConstraints<T_TRAITS>(Particle0, Particle1, Child0.First, Implicit1, TransformedChild0, LocalTransform1, CullDistance, Context, NewConstraints);
 					}
 				}
 				else
 				{
 					for (const auto& Child0 : Union0->GetObjects())
 					{
-						ConstructConstraints(Particle0, Particle1, Child0.Get(), Implicit1, LocalTransform0, LocalTransform1, CullDistance, Context, NewConstraints);
+						ConstructConstraints<T_TRAITS>(Particle0, Particle1, Child0.Get(), Implicit1, LocalTransform0, LocalTransform1, CullDistance, Context, NewConstraints);
 					}
 				}
 				return;
@@ -2652,15 +2652,15 @@ namespace Chaos
 					Union1->FindAllIntersectingObjects(Children, Implicit0->BoundingBox());
 					for (const auto& Child1 : Children)
 					{
-						TRigidTransform<T, d> TransformedChild1 = Child1.Second * LocalTransform1;
-						ConstructConstraints(Particle0, Particle1, Implicit0, Child1.First, LocalTransform0, TransformedChild1, CullDistance, Context, NewConstraints);
+						TRigidTransform<FReal, 3> TransformedChild1 = Child1.Second * LocalTransform1;
+						ConstructConstraints<T_TRAITS>(Particle0, Particle1, Implicit0, Child1.First, LocalTransform0, TransformedChild1, CullDistance, Context, NewConstraints);
 					}
 				}
 				else
 				{
 					for (const auto& Child1 : Union1->GetObjects())
 					{
-						ConstructConstraints(Particle0, Particle1, Implicit0, Child1.Get(), LocalTransform0, LocalTransform1, CullDistance, Context, NewConstraints);
+						ConstructConstraints<T_TRAITS>(Particle0, Particle1, Implicit0, Child1.Get(), LocalTransform0, LocalTransform1, CullDistance, Context, NewConstraints);
 					}
 				}
 				return;
diff --git a/Engine/Source/Runtime/Experimental/Chaos/Private/Chaos/PBDRigidClustering.cpp b/Engine/Source/Runtime/Experimental/Chaos/Private/Chaos/PBDRigidClustering.cpp
--- a/Engine/Source/Runtime/Experimental/Chaos/Private/Chaos/PBDRigidClustering.cpp
+++ b/Engine/Source/Runtime/Experimental/Chaos/Private/Chaos/PBDRigidClustering.cpp
@@ -1345,11 +1345,11 @@ namespace Chaos
 		for(TPBDRigidParticleHandle<T, d>* Child : Children)
 		{
 			bool bFilterValid = false;
-			for(const TUniquePtr<TPerShapeData<FReal, 3>>& Shape : Child->ShapesArray())
+			for(const TUniquePtr<FPerShapeData>& Shape : Child->ShapesArray())
 			{
 				if(Shape)
 				{
-					Filter = Shape->SimData;
+					Filter = Shape->GetSimData();
 					bFilterValid = Filter.Word0 != 0 || Filter.Word1 != 0 || Filter.Word2 != 0 || Filter.Word3 != 0;
 				}
 
@@ -1571,9 +1571,9 @@ namespace Chaos
 		}
 
 		// Set the captured filter to our new shapes
-		for(const TUniquePtr<TPerShapeData<FReal, 3>>& Shape : Parent->ShapesArray())
+		for(const TUniquePtr<FPerShapeData>& Shape : Parent->ShapesArray())
 		{
-			Shape->SimData = Filter;
+			Shape->SetSimData(Filter);
 		}
 	}
 
