CIS - Clang builds, add export for linking error

diff --git a/Engine/Source/Runtime/Experimental/Chaos/Public/Chaos/PBDRigidsEvolution.h b/Engine/Source/Runtime/Experimental/Chaos/Public/Chaos/PBDRigidsEvolution.h
--- a/Engine/Source/Runtime/Experimental/Chaos/Public/Chaos/PBDRigidsEvolution.h
+++ b/Engine/Source/Runtime/Experimental/Chaos/Public/Chaos/PBDRigidsEvolution.h
@@ -390,7 +390,7 @@ class TPBDRigidsEvolutionBase
 		ExternalSpatialData = SpatialData;
 	}
 
-	void DestroyParticle(TGeometryParticleHandle<FReal, 3>* Particle)
+	CHAOS_API void DestroyParticle(TGeometryParticleHandle<FReal, 3>* Particle)
 	{
 		RemoveParticleFromAccelerationStructure(*Particle);
 		ConstraintGraph.RemoveParticle(Particle);
