#UE4 CIS fix for Mac

diff --git a/Engine/Source/Runtime/Experimental/Chaos/Public/Chaos/PBDRigidsEvolution.h b/Engine/Source/Runtime/Experimental/Chaos/Public/Chaos/PBDRigidsEvolution.h
--- a/Engine/Source/Runtime/Experimental/Chaos/Public/Chaos/PBDRigidsEvolution.h
+++ b/Engine/Source/Runtime/Experimental/Chaos/Public/Chaos/PBDRigidsEvolution.h
@@ -912,7 +912,11 @@ protected:
 	TUniquePtr<ISpatialAccelerationCollectionFactory> SpatialCollectionFactory;
 };
 
+#if PLATFORM_MAC
+#define EVOLUTION_TRAIT(Trait) extern template class CHAOS_API TPBDRigidsEvolutionBase<Trait>;
+#else
 #define EVOLUTION_TRAIT(Trait) extern template class TPBDRigidsEvolutionBase<Trait>;
+#endif
 #include "Chaos/EvolutionTraits.inl"
 #undef EVOLUTION_TRAIT
 
diff --git a/Engine/Source/Runtime/Experimental/Chaos/Public/Chaos/PBDRigidsEvolutionGBF.h b/Engine/Source/Runtime/Experimental/Chaos/Public/Chaos/PBDRigidsEvolutionGBF.h
--- a/Engine/Source/Runtime/Experimental/Chaos/Public/Chaos/PBDRigidsEvolutionGBF.h
+++ b/Engine/Source/Runtime/Experimental/Chaos/Public/Chaos/PBDRigidsEvolutionGBF.h
@@ -254,7 +254,11 @@ namespace Chaos
 		FPBDRigidsEvolutionInternalHandleCallback InternalParticleInitilization;
 	};
 
-#define EVOLUTION_TRAIT(Trait) extern template class TPBDRigidsEvolutionGBF<Trait>;
+#if PLATFORM_MAC
+#define EVOLUTION_TRAIT(Trait) extern template class CHAOS_API TPBDRigidsEvolutionBase<Trait>;
+#else
+#define EVOLUTION_TRAIT(Trait) extern template class TPBDRigidsEvolutionBase<Trait>;
+#endif
 #include "Chaos/EvolutionTraits.inl"
 #undef EVOLUTION_TRAIT
 }
