Chaos - fix accidental override of iteration counts

diff --git a/Engine/Source/Runtime/Experimental/Chaos/Private/Chaos/PBDRigidsEvolution.cpp b/Engine/Source/Runtime/Experimental/Chaos/Private/Chaos/PBDRigidsEvolution.cpp
--- a/Engine/Source/Runtime/Experimental/Chaos/Private/Chaos/PBDRigidsEvolution.cpp
+++ b/Engine/Source/Runtime/Experimental/Chaos/Private/Chaos/PBDRigidsEvolution.cpp
@@ -16,7 +16,7 @@ FAutoConsoleVariableRef CVarChaosRigidsEvolutionApplyPushoutAllowEarlyOut(TEXT("
 int32 ChaosNumPushOutIterationsOverride = -1;
 FAutoConsoleVariableRef CVarChaosNumPushOutIterationsOverride(TEXT("p.ChaosNumPushOutIterationsOverride"), ChaosNumPushOutIterationsOverride, TEXT("Override for num push out iterations if >= 0 [def:-1]"));
 
-int32 ChaosNumContactIterationsOverride = 1;
+int32 ChaosNumContactIterationsOverride = -1;
 FAutoConsoleVariableRef CVarChaosNumContactIterationsOverride(TEXT("p.ChaosNumContactIterationsOverride"), ChaosNumContactIterationsOverride, TEXT("Override for num contact iterations if >= 0. [def:-1]"));
 
 namespace Chaos
