Fix typo to hopefully fix mac cis

diff --git a/Engine/Source/Runtime/Experimental/Chaos/Public/Chaos/PBDRigidsEvolutionGBF.h b/Engine/Source/Runtime/Experimental/Chaos/Public/Chaos/PBDRigidsEvolutionGBF.h
--- a/Engine/Source/Runtime/Experimental/Chaos/Public/Chaos/PBDRigidsEvolutionGBF.h
+++ b/Engine/Source/Runtime/Experimental/Chaos/Public/Chaos/PBDRigidsEvolutionGBF.h
@@ -255,9 +255,9 @@ namespace Chaos
 	};
 
 #if PLATFORM_MAC
-#define EVOLUTION_TRAIT(Trait) extern template class CHAOS_API TPBDRigidsEvolutionBase<Trait>;
+#define EVOLUTION_TRAIT(Trait) extern template class CHAOS_API TPBDRigidsEvolutionGBF<Trait>;
 #else
-#define EVOLUTION_TRAIT(Trait) extern template class TPBDRigidsEvolutionBase<Trait>;
+#define EVOLUTION_TRAIT(Trait) extern template class TPBDRigidsEvolutionGBF<Trait>;
 #endif
 #include "Chaos/EvolutionTraits.inl"
 #undef EVOLUTION_TRAIT
