Fix linux editor

diff --git a/Engine/Source/Runtime/Experimental/Chaos/Public/Chaos/PBDRigidsEvolution.h b/Engine/Source/Runtime/Experimental/Chaos/Public/Chaos/PBDRigidsEvolution.h
--- a/Engine/Source/Runtime/Experimental/Chaos/Public/Chaos/PBDRigidsEvolution.h
+++ b/Engine/Source/Runtime/Experimental/Chaos/Public/Chaos/PBDRigidsEvolution.h
@@ -912,7 +912,7 @@ protected:
 	TUniquePtr<ISpatialAccelerationCollectionFactory> SpatialCollectionFactory;
 };
 
-#if PLATFORM_MAC
+#if PLATFORM_MAC || PLATFORM_LINUX
 #define EVOLUTION_TRAIT(Trait) extern template class CHAOS_API TPBDRigidsEvolutionBase<Trait>;
 #else
 #define EVOLUTION_TRAIT(Trait) extern template class TPBDRigidsEvolutionBase<Trait>;
diff --git a/Engine/Source/Runtime/Experimental/Chaos/Public/Chaos/PBDRigidsEvolutionGBF.h b/Engine/Source/Runtime/Experimental/Chaos/Public/Chaos/PBDRigidsEvolutionGBF.h
--- a/Engine/Source/Runtime/Experimental/Chaos/Public/Chaos/PBDRigidsEvolutionGBF.h
+++ b/Engine/Source/Runtime/Experimental/Chaos/Public/Chaos/PBDRigidsEvolutionGBF.h
@@ -254,7 +254,7 @@ namespace Chaos
 		FPBDRigidsEvolutionInternalHandleCallback InternalParticleInitilization;
 	};
 
-#if PLATFORM_MAC
+#if PLATFORM_MAC || PLATFORM_LINUX
 #define EVOLUTION_TRAIT(Trait) extern template class CHAOS_API TPBDRigidsEvolutionGBF<Trait>;
 #else
 #define EVOLUTION_TRAIT(Trait) extern template class TPBDRigidsEvolutionGBF<Trait>;
