Fix collisions of unions

diff --git a/Engine/Source/Runtime/Experimental/Chaos/Private/Chaos/CollisionResolution.cpp b/Engine/Source/Runtime/Experimental/Chaos/Private/Chaos/CollisionResolution.cpp
--- a/Engine/Source/Runtime/Experimental/Chaos/Private/Chaos/CollisionResolution.cpp
+++ b/Engine/Source/Runtime/Experimental/Chaos/Private/Chaos/CollisionResolution.cpp
@@ -1975,7 +1975,7 @@ namespace Chaos
 			}
 
 			TGenericParticleHandle<FReal, 3> Particle1 = Constraint.Particle[1];
-			FRigidTransform3 LevelsetTM = FRigidTransform3(Particle1->P(), Particle1->Q());
+			FRigidTransform3 LevelsetTM = Constraint.ImplicitTransform[1] * FRigidTransform3(Particle1->P(), Particle1->Q());
 			if (!(ensure(!FMath::IsNaN(LevelsetTM.GetTranslation().X)) && ensure(!FMath::IsNaN(LevelsetTM.GetTranslation().Y)) && ensure(!FMath::IsNaN(LevelsetTM.GetTranslation().Z))))
 			{
 				return;
@@ -2008,6 +2008,8 @@ namespace Chaos
 			{
 				Constraint.Particle[0] = Particle1;
 				Constraint.Particle[1] = Particle0;
+				Constraint.ImplicitTransform[0] = ParticleImplicit1TM;
+				Constraint.ImplicitTransform[1] = ParticleImplicit0TM;
 				Constraint.SetManifold(Implicit1, Implicit0);
 			}
 			else
@@ -2614,7 +2616,7 @@ namespace Chaos
 				if (Implicit1->HasBoundingBox())
 				{
 					TArray<Pair<const FImplicitObject*, FRigidTransform3>> Children;
-					Union0->FindAllIntersectingObjects(Children, Implicit1->BoundingBox());
+					Union0->FindAllIntersectingObjects(Children, Implicit1->BoundingBox().TransformedAABB(LocalTransform1.GetRelativeTransform(LocalTransform0)));
 					for (const auto& Child0 : Children)
 					{
 						TRigidTransform<FReal, 3> TransformedChild0 = Child0.Second * LocalTransform0;
@@ -2649,7 +2651,7 @@ namespace Chaos
 				if (Implicit0->HasBoundingBox())
 				{
 					TArray<Pair<const FImplicitObject*, FRigidTransform3>> Children;
-					Union1->FindAllIntersectingObjects(Children, Implicit0->BoundingBox());
+					Union1->FindAllIntersectingObjects(Children, Implicit0->BoundingBox().TransformedAABB(LocalTransform0.GetRelativeTransform(LocalTransform1)));
 					for (const auto& Child1 : Children)
 					{
 						TRigidTransform<FReal, 3> TransformedChild1 = Child1.Second * LocalTransform1;
