Better fix for TOI CCD. Disabled globally but can be enabled where required
#rnx

diff --git a/Engine/Source/Runtime/Experimental/Chaos/Private/Chaos/PBDRigidsEvolutionGBF.cpp b/Engine/Source/Runtime/Experimental/Chaos/Private/Chaos/PBDRigidsEvolutionGBF.cpp
--- a/Engine/Source/Runtime/Experimental/Chaos/Private/Chaos/PBDRigidsEvolutionGBF.cpp
+++ b/Engine/Source/Runtime/Experimental/Chaos/Private/Chaos/PBDRigidsEvolutionGBF.cpp
@@ -52,7 +52,7 @@ float BoundsThicknessVelocityMultiplier = 2.0f;
 FAutoConsoleVariableRef CVarBoundsThickness(TEXT("p.CollisionBoundsThickness"), BoundsThickness, TEXT("Collision inflation for speculative contact generation.[def:0.0]"));
 FAutoConsoleVariableRef CVarBoundsThicknessVelocityMultiplier(TEXT("p.CollisionBoundsVelocityInflation"), BoundsThicknessVelocityMultiplier, TEXT("Collision velocity inflation for speculatibe contact generation.[def:2.0]"));
 
-float HackCCD_EnableThreshold = -1.f;//0.3f;
+float HackCCD_EnableThreshold = -1.f;
 FAutoConsoleVariableRef CVarHackCCDVelThreshold(TEXT("p.Chaos.CCD.EnableThreshold"), HackCCD_EnableThreshold, TEXT("If distance moved is greater than this times the minimum object dimension, use CCD"));
 
 float HackCCD_DepthThreshold = 0.05f;
