Skip calls to RemoveConstraint in Chaos if we're never rewinding state and resimulating - the constraints are rebuild on the next frame instead
#rnx

diff --git a/Engine/Source/Runtime/Experimental/Chaos/Private/Chaos/PBDRigidClustering.cpp b/Engine/Source/Runtime/Experimental/Chaos/Private/Chaos/PBDRigidClustering.cpp
--- a/Engine/Source/Runtime/Experimental/Chaos/Private/Chaos/PBDRigidClustering.cpp
+++ b/Engine/Source/Runtime/Experimental/Chaos/Private/Chaos/PBDRigidClustering.cpp
@@ -35,8 +35,8 @@ namespace Chaos
 	int32 UseConnectivity = 1;
 	FAutoConsoleVariableRef CVarUseConnectivity(TEXT("p.UseConnectivity"), UseConnectivity, TEXT("Whether to use connectivity graph when breaking up clusters"));
 
-	float ChildrenInheritVelocity = 1.f;
-	FAutoConsoleVariableRef CVarChildrenInheritVelocity(TEXT("p.ChildrenInheritVelocity"), ChildrenInheritVelocity, TEXT("Whether children inherit parent collision velocity when declustering. 0 has no impact velocity like glass, 1 has full impact velocity like brick"));
+	float ChaosClusteringChildrenInheritVelocity = 1.f;
+	FAutoConsoleVariableRef CVarChildrenInheritVelocity(TEXT("p.ChildrenInheritVelocity"), ChaosClusteringChildrenInheritVelocity, TEXT("Whether children inherit parent collision velocity when declustering. 0 has no impact velocity like glass, 1 has full impact velocity like brick"));
 
 	int32 ComputeClusterCollisionStrains = 1;
 	FAutoConsoleVariableRef CVarComputeClusterCollisionStrains(TEXT("p.ComputeClusterCollisionStrains"), ComputeClusterCollisionStrains, TEXT("Whether to use collision constraints when processing clustering."));
@@ -169,7 +169,7 @@ namespace Chaos
 			}
 		}
 
-		const bool bRewindOnDeclusterSolve = ChildrenInheritVelocity < 1.f;
+		const bool bRewindOnDeclusterSolve = ChaosClusteringChildrenInheritVelocity < 1.f;
 		if (bRewindOnDeclusterSolve)
 		{
 			// @todo(mlentine): We can precompute internal constraints which can filter some from the narrow phase tests but may not help much
@@ -678,7 +678,7 @@ namespace Chaos
 		TArray<TPBDRigidParticleHandle<T, d>*>& Children = MChildren[ClusteredParticle];
 
 		bool bChildrenChanged = false;
-		const bool bRewindOnDecluster = ChildrenInheritVelocity < 1;
+		const bool bRewindOnDecluster = ChaosClusteringChildrenInheritVelocity < 1;
 		const TRigidTransform<T, d> PreSolveTM = 
 			bRewindOnDecluster ? 
 			TRigidTransform<T, d>(ClusteredParticle->X(), ClusteredParticle->R()) : 
@@ -1043,7 +1043,7 @@ namespace Chaos
 					AllActivatedChildren.Append(ActivatedChildren);
 				}
 
-				const bool bRewindOnDecluster = ChildrenInheritVelocity < 1.f;
+				const bool bRewindOnDecluster = ChaosClusteringChildrenInheritVelocity < 1.f;
 				if (bRewindOnDecluster && AllActivatedChildren.Num())
 				{
 					SCOPE_CYCLE_COUNTER(STAT_ClusterRewind);
@@ -1053,7 +1053,7 @@ namespace Chaos
 						RewindAndEvolve(MEvolution, MParticles, IslandsToRecollide, AllActivatedChildren, Dt, CollisionRule);
 					}
 
-					if (ChildrenInheritVelocity > 0.f)
+					if (ChaosClusteringChildrenInheritVelocity > 0.f)
 					{
 						for (auto Itr : ClusterToActivatedChildren)
 						{
@@ -1062,9 +1062,9 @@ namespace Chaos
 							for (TPBDRigidParticleHandle<T, d>* ActiveChild : ActivatedChildren)
 							{
 								ActiveChild->SetV(
-									ActiveChild->V() * (1.f - ChildrenInheritVelocity) + ClusteredParticle->V() * ChildrenInheritVelocity);
+									ActiveChild->V() * (1.f - ChaosClusteringChildrenInheritVelocity) + ClusteredParticle->V() * ChaosClusteringChildrenInheritVelocity);
 								ActiveChild->SetW(
-									ActiveChild->W() * (1.f - ChildrenInheritVelocity) + ClusteredParticle->W() * ChildrenInheritVelocity);
+									ActiveChild->W() * (1.f - ChaosClusteringChildrenInheritVelocity) + ClusteredParticle->W() * ChaosClusteringChildrenInheritVelocity);
 							}
 						}
 					}
diff --git a/Engine/Source/Runtime/Experimental/Chaos/Public/Chaos/PBDRigidClustering.h b/Engine/Source/Runtime/Experimental/Chaos/Public/Chaos/PBDRigidClustering.h
--- a/Engine/Source/Runtime/Experimental/Chaos/Public/Chaos/PBDRigidClustering.h
+++ b/Engine/Source/Runtime/Experimental/Chaos/Public/Chaos/PBDRigidClustering.h
@@ -12,6 +12,11 @@
 
 #define TODO_CONVERT_GEOMETRY_COLLECTION_PARTICLE_INDICES_TO_PARTICLE_POINTERS 0
 
+namespace Chaos
+{
+	extern float ChaosClusteringChildrenInheritVelocity;
+}
+
 namespace Chaos
 {
 
diff --git a/Engine/Source/Runtime/Experimental/Chaos/Public/Chaos/PBDRigidsEvolution.h b/Engine/Source/Runtime/Experimental/Chaos/Public/Chaos/PBDRigidsEvolution.h
--- a/Engine/Source/Runtime/Experimental/Chaos/Public/Chaos/PBDRigidsEvolution.h
+++ b/Engine/Source/Runtime/Experimental/Chaos/Public/Chaos/PBDRigidsEvolution.h
@@ -333,7 +333,6 @@ class TPBDRigidsEvolutionBase
 		RemoveParticleFromAccelerationStructure(*Particle);
 		Particles.DisableParticle(Particle);
 		ConstraintGraph.DisableParticle(Particle);
-
 		RemoveConstraints(TSet<TGeometryParticleHandle<FReal, 3>*>({ Particle }));
 	}
 
@@ -437,7 +436,6 @@ class TPBDRigidsEvolutionBase
 		}
 
 		ConstraintGraph.DisableParticles(InParticles);
-
 		RemoveConstraints(InParticles);
 	}
 
@@ -454,9 +452,13 @@ class TPBDRigidsEvolutionBase
 	// @todo(ccaulfield): Remove the uint version
 	CHAOS_API void RemoveConstraints(const TSet<TGeometryParticleHandle<FReal, 3>*>& RemovedParticles)
 	{
-		for (FPBDConstraintGraphRule* ConstraintRule : ConstraintRules)
+		// Only remove constraints if we have the possibility of rewinding state. Otherwise they will be rebuilt next frame
+		if(ChaosClusteringChildrenInheritVelocity < 1.0f)
 		{
-			ConstraintRule->RemoveConstraints(RemovedParticles);
+			for(FPBDConstraintGraphRule* ConstraintRule : ConstraintRules)
+			{
+				ConstraintRule->RemoveConstraints(RemovedParticles);
+			}
 		}
 	}
 
