Export ChaosClusteringChildrenInheritVelocity as its used in a header and my be used in a different module

diff --git a/Engine/Source/Runtime/Experimental/Chaos/Private/Chaos/PBDRigidClustering.cpp b/Engine/Source/Runtime/Experimental/Chaos/Private/Chaos/PBDRigidClustering.cpp
--- a/Engine/Source/Runtime/Experimental/Chaos/Private/Chaos/PBDRigidClustering.cpp
+++ b/Engine/Source/Runtime/Experimental/Chaos/Private/Chaos/PBDRigidClustering.cpp
@@ -35,7 +35,7 @@ namespace Chaos
 	int32 UseConnectivity = 1;
 	FAutoConsoleVariableRef CVarUseConnectivity(TEXT("p.UseConnectivity"), UseConnectivity, TEXT("Whether to use connectivity graph when breaking up clusters"));
 
-	float ChaosClusteringChildrenInheritVelocity = 1.f;
+	CHAOS_API float ChaosClusteringChildrenInheritVelocity = 1.f;
 	FAutoConsoleVariableRef CVarChildrenInheritVelocity(TEXT("p.ChildrenInheritVelocity"), ChaosClusteringChildrenInheritVelocity, TEXT("Whether children inherit parent collision velocity when declustering. 0 has no impact velocity like glass, 1 has full impact velocity like brick"));
 
 	int32 ComputeClusterCollisionStrains = 1;
