CIS - Windows symbol export fix
#rnx

diff --git a/Engine/Source/Runtime/Experimental/Chaos/Public/Chaos/PBDRigidClustering.h b/Engine/Source/Runtime/Experimental/Chaos/Public/Chaos/PBDRigidClustering.h
--- a/Engine/Source/Runtime/Experimental/Chaos/Public/Chaos/PBDRigidClustering.h
+++ b/Engine/Source/Runtime/Experimental/Chaos/Public/Chaos/PBDRigidClustering.h
@@ -14,7 +14,7 @@
 
 namespace Chaos
 {
-	extern float ChaosClusteringChildrenInheritVelocity;
+	extern CHAOS_API float ChaosClusteringChildrenInheritVelocity;
 }
 
 namespace Chaos
