Chaos - Joint fix swing drives when both axes are not locked.

diff --git a/Engine/Source/Runtime/Experimental/Chaos/Private/Chaos/Joint/PBDJointSolverGaussSeidel.cpp b/Engine/Source/Runtime/Experimental/Chaos/Private/Chaos/Joint/PBDJointSolverGaussSeidel.cpp
--- a/Engine/Source/Runtime/Experimental/Chaos/Private/Chaos/Joint/PBDJointSolverGaussSeidel.cpp
+++ b/Engine/Source/Runtime/Experimental/Chaos/Private/Chaos/Joint/PBDJointSolverGaussSeidel.cpp
@@ -115,7 +115,8 @@ namespace Chaos
 		TwistSoftLambda = (FReal)0;
 		SwingSoftLambda = (FReal)0;
 		TwistDriveLambda = (FReal)0;
-		SwingDriveLambda = (FReal)0;
+		Swing1DriveLambda = (FReal)0;
+		Swing2DriveLambda = (FReal)0;
 
 		PositionTolerance = SolverSettings.PositionTolerance;
 		AngleTolerance = SolverSettings.AngleTolerance;
@@ -396,29 +397,20 @@ namespace Chaos
 			bool bSwing2Locked = Swing2Motion == EJointMotionType::Locked;
 
 			// No SLerp drive if we have a locked rotation (it will be grayed out in the editor in this case, but could still have been set before the rotation was locked)
+			// @todo(ccaulfield): setting should be cleaned up before being passed to the solver
 			if ((JointSettings.bAngularSLerpPositionDriveEnabled || JointSettings.bAngularSLerpVelocityDriveEnabled) && !bTwistLocked && !bSwing1Locked && !bSwing2Locked)
 			{
 				NumActive += ApplySLerpDrive(Dt, SolverSettings, JointSettings);
 			}
 			else
 			{
-				if ((JointSettings.bAngularTwistPositionDriveEnabled || JointSettings.bAngularTwistVelocityDriveEnabled) && !bTwistLocked)
-				{
-					NumActive += ApplyTwistDrive(Dt, SolverSettings, JointSettings);
-				}
-
+				const bool bTwistDriveEnabled = ((JointSettings.bAngularTwistPositionDriveEnabled || JointSettings.bAngularTwistVelocityDriveEnabled) && !bTwistLocked);
 				const bool bSwingDriveEnabled = (JointSettings.bAngularSwingPositionDriveEnabled || JointSettings.bAngularSwingVelocityDriveEnabled);
-				if (bSwingDriveEnabled && !bSwing1Locked && !bSwing2Locked)
-				{
-					NumActive += ApplyConeDrive(Dt, SolverSettings, JointSettings);
-				}
-				else if (bSwingDriveEnabled && !bSwing1Locked)
+				const bool bSwing1DriveEnabled = bSwingDriveEnabled && !bSwing1Locked;
+				const bool bSwing2DriveEnabled = bSwingDriveEnabled && !bSwing2Locked;
+				if (bTwistDriveEnabled || bSwing1DriveEnabled || bSwing2DriveEnabled)
 				{
-					NumActive += ApplySwingDrive(Dt, SolverSettings, JointSettings, EJointAngularConstraintIndex::Swing1);
-				}
-				else if (bSwingDriveEnabled && !bSwing2Locked)
-				{
-					NumActive += ApplySwingDrive(Dt, SolverSettings, JointSettings, EJointAngularConstraintIndex::Swing2);
+					ApplySwingTwistDrives(Dt, SolverSettings, JointSettings, bTwistDriveEnabled, bSwing1DriveEnabled, bSwing2DriveEnabled);
 				}
 			}
 		}
@@ -1045,31 +1037,6 @@ namespace Chaos
 	}
 
 
-	int32 FJointSolverGaussSeidel::ApplyTwistDrive(
-		const FReal Dt,
-		const FPBDJointSolverSettings& SolverSettings,
-		const FPBDJointSettings& JointSettings)
-	{
-		FVec3 TwistAxis;
-		FReal TwistAngle;
-		FPBDJointUtilities::GetTwistAxisAngle(Rs[0], Rs[1], TwistAxis, TwistAngle);
-
-		const FReal TwistAngleTarget = JointSettings.AngularDriveTargetAngles[(int32)EJointAngularConstraintIndex::Twist];
-		const FReal DTwistAngle = TwistAngle - TwistAngleTarget;
-
-		// Apply twist correction
-		if (FMath::Abs(DTwistAngle) > AngleTolerance)
-		{
-			const FReal AngularDriveStiffness = FPBDJointUtilities::GetAngularTwistDriveStiffness(SolverSettings, JointSettings);
-			const FReal AngularDriveDamping = FPBDJointUtilities::GetAngularTwistDriveDamping(SolverSettings, JointSettings);
-			const bool bAccelerationMode = FPBDJointUtilities::GetDriveAccelerationMode(SolverSettings, JointSettings);
-			ApplyRotationConstraintSoft(Dt, AngularDriveStiffness, AngularDriveDamping, bAccelerationMode, TwistAxis, DTwistAngle, TwistDriveLambda);
-			return 1;
-		}
-		return 0;
-	}
-
-
 	int32 FJointSolverGaussSeidel::ApplyTwistProjection(
 		const FReal Dt,
 		const FPBDJointSolverSettings& SolverSettings,
@@ -1140,37 +1107,6 @@ namespace Chaos
 	}
 
 
-	int32 FJointSolverGaussSeidel::ApplyConeDrive(
-		const FReal Dt,
-		const FPBDJointSolverSettings& SolverSettings,
-		const FPBDJointSettings& JointSettings)
-	{
-		// Calculate swing angle and axis
-		FReal SwingAngle;
-		FVec3 SwingAxisLocal;
-		FPBDJointUtilities::GetConeAxisAngleLocal(Rs[0], Rs[1], SolverSettings.SwingTwistAngleTolerance, SwingAxisLocal, SwingAngle);
-		const FVec3 SwingAxis = Rs[0] * SwingAxisLocal;
-
-		// Circular swing target (max of Swing1, Swing2 targets)
-		// @todo(ccaulfield): what should cone target really do?
-		const FReal Swing1Target = JointSettings.AngularDriveTargetAngles[(int32)EJointAngularConstraintIndex::Swing1];
-		const FReal Swing2Target = JointSettings.AngularDriveTargetAngles[(int32)EJointAngularConstraintIndex::Swing2];
-		const FReal SwingAngleTarget = FMath::Max(Swing1Target, Swing2Target);
-		const FReal DSwingAngle = SwingAngle - SwingAngleTarget;
-
-		// Apply drive forces to each body
-		if (FMath::Abs(DSwingAngle) > AngleTolerance)
-		{
-			const FReal AngularDriveStiffness = FPBDJointUtilities::GetAngularSwingDriveStiffness(SolverSettings, JointSettings);
-			const FReal AngularDriveDamping = FPBDJointUtilities::GetAngularSwingDriveDamping(SolverSettings, JointSettings);
-			const bool bAccelerationMode = FPBDJointUtilities::GetDriveAccelerationMode(SolverSettings, JointSettings);
-			ApplyRotationConstraintSoft(Dt, AngularDriveStiffness, AngularDriveDamping, bAccelerationMode, SwingAxis, DSwingAngle, SwingDriveLambda);
-			return 1;
-		}
-		return 0;
-	}
-
-
 	int32 FJointSolverGaussSeidel::ApplyConeProjection(
 		const FReal Dt,
 		const FPBDJointSolverSettings& SolverSettings,
@@ -1369,40 +1305,64 @@ namespace Chaos
 	}
 
 
-	int32 FJointSolverGaussSeidel::ApplySwingDrive(
+	int32 FJointSolverGaussSeidel::ApplySwingProjection(
 		const FReal Dt,
 		const FPBDJointSolverSettings& SolverSettings,
 		const FPBDJointSettings& JointSettings,
 		const EJointAngularConstraintIndex SwingConstraintIndex)
 	{
-		FVec3 SwingAxis;
-		FReal SwingAngle;
-		FPBDJointUtilities::GetSwingAxisAngle(Rs[0], Rs[1], SolverSettings.SwingTwistAngleTolerance, SwingConstraintIndex, SwingAxis, SwingAngle);
-
-		const FReal SwingAngleTarget = JointSettings.AngularDriveTargetAngles[(int32)SwingConstraintIndex];
-		const FReal DSwingAngle = SwingAngle - SwingAngleTarget;
-
-		// Apply drive forces to each body
-		if (FMath::Abs(DSwingAngle) > AngleTolerance)
-		{
-			const FReal AngularDriveStiffness = FPBDJointUtilities::GetAngularSwingDriveStiffness(SolverSettings, JointSettings);
-			const FReal AngularDriveDamping = FPBDJointUtilities::GetAngularSwingDriveDamping(SolverSettings, JointSettings);
-			const bool bAccelerationMode = FPBDJointUtilities::GetDriveAccelerationMode(SolverSettings, JointSettings);
-			ApplyRotationConstraintSoft(Dt, AngularDriveStiffness, AngularDriveDamping, bAccelerationMode, SwingAxis, DSwingAngle, SwingDriveLambda);
-			return 1;
-		}
+		// @todo(ccaulfield): implement swing projection
 		return 0;
 	}
 
 
-	int32 FJointSolverGaussSeidel::ApplySwingProjection(
+	int32 FJointSolverGaussSeidel::ApplySwingTwistDrives(
 		const FReal Dt,
 		const FPBDJointSolverSettings& SolverSettings,
 		const FPBDJointSettings& JointSettings,
-		const EJointAngularConstraintIndex SwingConstraintIndex)
+		const bool bTwistDriveEnabled,
+		const bool bSwing1DriveEnabled,
+		const bool bSwing2DriveEnabled)
 	{
-		// @todo(ccaulfield): implement swing projection
-		return 0;
+		FRotation3 R1Target = Rs[0] * JointSettings.AngularDrivePositionTarget;
+		R1Target.EnforceShortestArcWith(Rs[1]);
+		FRotation3 R1Error = R1Target.Inverse() * Rs[1];
+		FVec3 R1TwistAxisError = R1Error * FJointConstants::TwistAxis();
+
+		// Angle approximation Angle ~= Sin(Angle) for small angles, underestimates for large angles
+		const FReal DTwistAngle = 2.0f * R1Error.X;
+		const FReal DSwing1Angle = R1TwistAxisError.Y;
+		const FReal DSwing2Angle = -R1TwistAxisError.Z;
+
+		const FReal AngularTwistDriveStiffness = FPBDJointUtilities::GetAngularTwistDriveStiffness(SolverSettings, JointSettings);
+		const FReal AngularTwistDriveDamping = FPBDJointUtilities::GetAngularTwistDriveDamping(SolverSettings, JointSettings);
+		const FReal AngularSwingDriveStiffness = FPBDJointUtilities::GetAngularSwingDriveStiffness(SolverSettings, JointSettings);
+		const FReal AngularSwingDriveDamping = FPBDJointUtilities::GetAngularSwingDriveDamping(SolverSettings, JointSettings);
+		const bool bAccelerationMode = FPBDJointUtilities::GetDriveAccelerationMode(SolverSettings, JointSettings);
+		int32 NumActive = 0;
+
+		if (bTwistDriveEnabled && (FMath::Abs(DTwistAngle) > AngleTolerance))
+		{
+			const FVec3 TwistAxis = Rs[1] * FJointConstants::TwistAxis();
+			ApplyRotationConstraintSoft(Dt, AngularTwistDriveStiffness, AngularTwistDriveDamping, bAccelerationMode, TwistAxis, DTwistAngle, TwistDriveLambda);
+			++NumActive;
+		}
+
+		if (bSwing1DriveEnabled && (FMath::Abs(DSwing1Angle) > AngleTolerance))
+		{
+			const FVec3 Swing1Axis = Rs[1] * FJointConstants::Swing1Axis();
+			ApplyRotationConstraintSoft(Dt, AngularSwingDriveStiffness, AngularSwingDriveDamping, bAccelerationMode, Swing1Axis, DSwing1Angle, Swing1DriveLambda);
+			++NumActive;
+		}
+
+		if (bSwing2DriveEnabled && (FMath::Abs(DSwing2Angle) > AngleTolerance))
+		{
+			const FVec3 Swing2Axis = Rs[1] * FJointConstants::Swing2Axis();
+			ApplyRotationConstraintSoft(Dt, AngularSwingDriveStiffness, AngularSwingDriveDamping, bAccelerationMode, Swing2Axis, DSwing2Angle, Swing2DriveLambda);
+			++NumActive;
+		}
+
+		return NumActive;
 	}
 
 
@@ -1430,7 +1390,7 @@ namespace Chaos
 				const FReal AngularDriveStiffness = FPBDJointUtilities::GetAngularSLerpDriveStiffness(SolverSettings, JointSettings);
 				const FReal AngularDriveDamping = FPBDJointUtilities::GetAngularSLerpDriveDamping(SolverSettings, JointSettings);
 				const bool bAccelerationMode = FPBDJointUtilities::GetDriveAccelerationMode(SolverSettings, JointSettings);
-				ApplyRotationConstraintSoft(Dt, AngularDriveStiffness, AngularDriveDamping, bAccelerationMode, SLerpAxis, -SLerpAngle, SwingDriveLambda);
+				ApplyRotationConstraintSoft(Dt, AngularDriveStiffness, AngularDriveDamping, bAccelerationMode, SLerpAxis, -SLerpAngle, Swing1DriveLambda);
 				return 1;
 			}
 		}
diff --git a/Engine/Source/Runtime/Experimental/Chaos/Private/Chaos/Joint/PBDJointSolverGaussSeidel.ispc b/Engine/Source/Runtime/Experimental/Chaos/Private/Chaos/Joint/PBDJointSolverGaussSeidel.ispc
--- a/Engine/Source/Runtime/Experimental/Chaos/Private/Chaos/Joint/PBDJointSolverGaussSeidel.ispc
+++ b/Engine/Source/Runtime/Experimental/Chaos/Private/Chaos/Joint/PBDJointSolverGaussSeidel.ispc
@@ -40,7 +40,8 @@ struct FJointSolverGaussSeidel
 	// XPBD Previous iteration world-space body state
 	FVector PrevPs[MaxConstrainedBodies];		// World-space particle CoM positions
 	float TwistDriveLambda;						// XPBD constraint multipliers (net applied constraint-space deltas)
-	float SwingDriveLambda;						// XPBD constraint multipliers (net applied constraint-space deltas)
+	float Swing1DriveLambda;					// XPBD constraint multipliers (net applied constraint-space deltas)
+	float Swing2DriveLambda;					// XPBD constraint multipliers (net applied constraint-space deltas)
 	FVector4 PrevQs[MaxConstrainedBodies];		// World-space particle CoM rotations
 	FVector PrevXs[MaxConstrainedBodies];		// World-space joint connector positions
 	
diff --git a/Engine/Source/Runtime/Experimental/Chaos/Public/Chaos/Joint/PBDJointSolverGaussSeidel.h b/Engine/Source/Runtime/Experimental/Chaos/Public/Chaos/Joint/PBDJointSolverGaussSeidel.h
--- a/Engine/Source/Runtime/Experimental/Chaos/Public/Chaos/Joint/PBDJointSolverGaussSeidel.h
+++ b/Engine/Source/Runtime/Experimental/Chaos/Public/Chaos/Joint/PBDJointSolverGaussSeidel.h
@@ -243,11 +243,6 @@ namespace Chaos
 			const FPBDJointSettings& JointSettings,
 			const bool bUseSoftLimit);
 
-		int32 ApplyTwistDrive(
-			const FReal Dt,
-			const FPBDJointSolverSettings& SolverSettings,
-			const FPBDJointSettings& JointSettings);
-
 		int32 ApplyTwistProjection(
 			const FReal Dt,
 			const FPBDJointSolverSettings& SolverSettings,
@@ -259,11 +254,6 @@ namespace Chaos
 			const FPBDJointSettings& JointSettings,
 			const bool bUseSoftLimit);
 
-		int32 ApplyConeDrive(
-			const FReal Dt,
-			const FPBDJointSolverSettings& SolverSettings,
-			const FPBDJointSettings& JointSettings);
-
 		int32 ApplyConeProjection(
 			const FReal Dt,
 			const FPBDJointSolverSettings& SolverSettings,
@@ -293,17 +283,19 @@ namespace Chaos
 			const EJointAngularConstraintIndex SwingConstraintIndex,
 			const bool bUseSoftLimit);
 
-		int32 ApplySwingDrive(
+		int32 ApplySwingProjection(
 			const FReal Dt,
 			const FPBDJointSolverSettings& SolverSettings,
 			const FPBDJointSettings& JointSettings,
 			const EJointAngularConstraintIndex SwingConstraintIndex);
 
-		int32 ApplySwingProjection(
+		int32 ApplySwingTwistDrives(
 			const FReal Dt,
 			const FPBDJointSolverSettings& SolverSettings,
 			const FPBDJointSettings& JointSettings,
-			const EJointAngularConstraintIndex SwingConstraintIndex);
+			const bool bTwistDriveEnabled,
+			const bool bSwing1DriveEnabled,
+			const bool bSwing2DriveEnabled);
 
 		int32 ApplySLerpDrive(
 			const FReal Dt,
@@ -393,7 +385,8 @@ namespace Chaos
 		// XPBD Previous iteration world-space body state
 		FVec3 PrevPs[MaxConstrainedBodies];			// World-space particle CoM positions
 		FReal TwistDriveLambda;						// XPBD constraint multipliers (net applied constraint-space deltas)
-		FReal SwingDriveLambda;						// XPBD constraint multipliers (net applied constraint-space deltas)
+		FReal Swing1DriveLambda;					// XPBD constraint multipliers (net applied constraint-space deltas)
+		FReal Swing2DriveLambda;					// XPBD constraint multipliers (net applied constraint-space deltas)
 		FRotation3 PrevQs[MaxConstrainedBodies];	// World-space particle CoM rotations
 		FVec3 PrevXs[MaxConstrainedBodies];			// World-space joint connector positions
 
