Compile fix.

diff --git a/Engine/Source/Runtime/Experimental/Chaos/Private/Chaos/CollisionResolution.cpp b/Engine/Source/Runtime/Experimental/Chaos/Private/Chaos/CollisionResolution.cpp
--- a/Engine/Source/Runtime/Experimental/Chaos/Private/Chaos/CollisionResolution.cpp
+++ b/Engine/Source/Runtime/Experimental/Chaos/Private/Chaos/CollisionResolution.cpp
@@ -2001,7 +2001,9 @@ namespace Chaos
 		template<typename T_TRAITS>
 		void ConstructLevelsetLevelsetConstraints(TGeometryParticleHandle<FReal, 3>* Particle0, TGeometryParticleHandle<FReal, 3>* Particle1, const FImplicitObject* Implicit0, const FImplicitObject* Implicit1, const FRigidTransform3& LocalTransform0, const FRigidTransform3& LocalTransform1, const FReal CullDistance, FCollisionConstraintsArray& NewConstraints)
 		{
-			FRigidBodyPointContactConstraint Constraint = FRigidBodyPointContactConstraint(Particle0, Implicit0, LocalTransform0, Particle1, Implicit1, LocalTransform1, EContactShapesType::LevelSetLevelSet);
+			FRigidTransform3 ParticleImplicit0TM = LocalTransform0.GetRelativeTransform(Collisions::GetTransform(Particle0));
+			FRigidTransform3 ParticleImplicit1TM = LocalTransform1.GetRelativeTransform(Collisions::GetTransform(Particle1));
+			FRigidBodyPointContactConstraint Constraint = FRigidBodyPointContactConstraint(Particle0, Implicit0, ParticleImplicit0TM, Particle1, Implicit1, ParticleImplicit1TM, EContactShapesType::LevelSetLevelSet);
 
 			bool bIsParticleDynamic0 = Particle0->CastToRigidParticle() && Particle0->ObjectState() == EObjectStateType::Dynamic;
 			if (!Particle1->Geometry() || (bIsParticleDynamic0 && !Particle0->CastToRigidParticle()->CollisionParticlesSize() && Particle0->Geometry() && !Particle0->Geometry()->IsUnderlyingUnion()))
