Chaos - RBAN suport "Enable Gravity" in PhAT

diff --git a/Engine/Source/Runtime/Engine/Private/PhysicsEngine/ImmediatePhysics/ImmediatePhysicsChaos/ImmediatePhysicsActorHandle_Chaos.cpp b/Engine/Source/Runtime/Engine/Private/PhysicsEngine/ImmediatePhysics/ImmediatePhysicsChaos/ImmediatePhysicsActorHandle_Chaos.cpp
--- a/Engine/Source/Runtime/Engine/Private/PhysicsEngine/ImmediatePhysics/ImmediatePhysicsChaos/ImmediatePhysicsActorHandle_Chaos.cpp
+++ b/Engine/Source/Runtime/Engine/Private/PhysicsEngine/ImmediatePhysics/ImmediatePhysicsChaos/ImmediatePhysicsActorHandle_Chaos.cpp
@@ -344,6 +344,7 @@ namespace ImmediatePhysics_Chaos
 					{
 						Dynamic->SetLinearEtherDrag(BodyInstance->LinearDamping);
 						Dynamic->SetAngularEtherDrag(BodyInstance->AngularDamping);
+						Dynamic->SetGravityEnabled(BodyInstance->bEnableGravity);
 					}
 					Dynamic->Disabled() = true;
 				}
diff --git a/Engine/Source/Runtime/Experimental/Chaos/Private/Chaos/Evolution/PBDMinEvolution.cpp b/Engine/Source/Runtime/Experimental/Chaos/Private/Chaos/Evolution/PBDMinEvolution.cpp
--- a/Engine/Source/Runtime/Experimental/Chaos/Private/Chaos/Evolution/PBDMinEvolution.cpp
+++ b/Engine/Source/Runtime/Experimental/Chaos/Private/Chaos/Evolution/PBDMinEvolution.cpp
@@ -74,14 +74,6 @@ namespace Chaos
 			Rewind(StepDt, RewindDt);
 		}
 
-		for (TTransientPBDRigidParticleHandle<FReal, 3>& Particle : Particles.GetActiveParticlesView())
-		{
-			if (Particle.ObjectState() == EObjectStateType::Dynamic)
-			{
-				Particle.F() += Particle.M() * Gravity;
-			}
-		}
-
 		for (int32 Step = 0; Step < NumSteps; ++Step)
 		{
 			// StepFraction: how much of the remaining time this step represents, used to interpolate kinematic targets
@@ -231,7 +223,8 @@ namespace Chaos
 
 				// Calculate new velocities from forces, torques and drag
 				const FMatrix33 WorldInvI = Utilities::ComputeWorldSpaceInertia(RCoM, Particle.InvI());
-				const FVec3 DV = Particle.InvM() * (Particle.F() * Dt + Particle.LinearImpulse());
+				const FVec3 G = (Particle.GravityEnabled()) ? Gravity : FVec3(0);
+				const FVec3 DV = Particle.InvM() * ((G + Particle.F()) * Dt + Particle.LinearImpulse());
 				const FVec3 DW = WorldInvI * (Particle.Torque() * Dt + Particle.AngularImpulse());
 				const FReal LinearDrag = FMath::Max(FReal(0), FReal(1) - (Particle.LinearEtherDrag() * Dt));
 				const FReal AngularDrag = FMath::Max(FReal(0), FReal(1) - (Particle.AngularEtherDrag() * Dt));
