Fix acceleration structure update for elements with big bounds.

diff --git a/Engine/Source/Runtime/Experimental/Chaos/Public/Chaos/BoundingVolume.h b/Engine/Source/Runtime/Experimental/Chaos/Public/Chaos/BoundingVolume.h
--- a/Engine/Source/Runtime/Experimental/Chaos/Public/Chaos/BoundingVolume.h
+++ b/Engine/Source/Runtime/Experimental/Chaos/Public/Chaos/BoundingVolume.h
@@ -214,13 +214,7 @@ public:
 		{
 			if (PayloadInfo->GlobalPayloadIdx != INDEX_NONE)
 			{
-				ensure(PayloadInfo->DirtyPayloadIdx == INDEX_NONE);
-				auto LastGlobalPayload = MGlobalPayloads.Last().Payload;
-				if (!(LastGlobalPayload == Payload))
-				{
-					MPayloadInfo.FindChecked(LastGlobalPayload).GlobalPayloadIdx = PayloadInfo->GlobalPayloadIdx;
-				}
-				MGlobalPayloads.RemoveAtSwap(PayloadInfo->GlobalPayloadIdx);
+				RemoveGlobalElement(Payload, *PayloadInfo);
 			}
 			else
 			{
@@ -242,6 +236,11 @@ public:
 				RemoveElementFromExistingGrid(Payload, *PayloadInfo);
 				AddElementToExistingGrid(Payload, *PayloadInfo, NewBounds, bHasBounds);
 			}
+			else if (bHasBounds)
+			{
+				RemoveGlobalElement(Payload, *PayloadInfo);
+				AddElementToExistingGrid(Payload, *PayloadInfo, NewBounds, bHasBounds);
+			}
 		}
 		else
 		{
@@ -393,6 +392,17 @@ private:
 	using FCellElement = TBVCellElement<TPayloadType, T, d>;
 	using FPayloadInfo = TBVPayloadInfo<T, d>;
 
+	void RemoveGlobalElement(const TPayloadType& Payload, const FPayloadInfo& PayloadInfo)
+	{
+		ensure(PayloadInfo.DirtyPayloadIdx == INDEX_NONE);
+		auto LastGlobalPayload = MGlobalPayloads.Last().Payload;
+		if (!(LastGlobalPayload == Payload))
+		{
+			MPayloadInfo.FindChecked(LastGlobalPayload).GlobalPayloadIdx = PayloadInfo.GlobalPayloadIdx;
+		}
+		MGlobalPayloads.RemoveAtSwap(PayloadInfo.GlobalPayloadIdx);
+	}
+
 	template <typename SQVisitor, bool bPruneDuplicates = true>
 	bool RaycastImp(const TVector<T, d>& Start, FQueryFastData& CurData, SQVisitor& Visitor) const
 	{
