Fix incorrect intersection time for sweeps againsts scaled trimesh.

diff --git a/Engine/Source/Runtime/Experimental/Chaos/Private/Chaos/TriangleMeshImplicitObject.cpp b/Engine/Source/Runtime/Experimental/Chaos/Private/Chaos/TriangleMeshImplicitObject.cpp
--- a/Engine/Source/Runtime/Experimental/Chaos/Private/Chaos/TriangleMeshImplicitObject.cpp
+++ b/Engine/Source/Runtime/Experimental/Chaos/Private/Chaos/TriangleMeshImplicitObject.cpp
@@ -677,7 +677,8 @@ struct FTriangleMeshSweepVisitor
 
 		if(GJKRaycast2<FReal>(Tri, WorldScaleQueryGeom, ScaledStartTM, ScaledDirNormalized, LengthScale * CurData.CurrentLength, Time, HitPosition, HitNormal, Thickness, bComputeMTD))
 		{
-			if(Time < OutTime)
+			// Time is world scale, OutTime is local scale.
+			if(Time < LengthScale * OutTime)
 			{
 				TransformSweepOutputsHelper(TriMeshScale, HitNormal, HitPosition, LengthScale, Time, OutNormal, OutPosition, OutTime);
 
@@ -691,7 +692,7 @@ struct FTriangleMeshSweepVisitor
 					return false;
 				}
 
-				CurData.SetLength(Time);
+				CurData.SetLength(OutTime);
 			}
 		}
 
