Chaos Physics - Accumulated impulse clipping changes:
Added energy clamping to avoid exploding geometry collections
Small refactor to improve code readability

#[fyi] Ryan.Kautzman, Kriss.Gossart, Chris.Caulfield, Michael.Lentine
#[review]-12947532 [at] Michael.Lentine

diff --git a/Engine/Source/Runtime/Experimental/Chaos/Private/Chaos/PBDCollisionConstraintsContact.cpp b/Engine/Source/Runtime/Experimental/Chaos/Private/Chaos/PBDCollisionConstraintsContact.cpp
--- a/Engine/Source/Runtime/Experimental/Chaos/Private/Chaos/PBDCollisionConstraintsContact.cpp
+++ b/Engine/Source/Runtime/Experimental/Chaos/Private/Chaos/PBDCollisionConstraintsContact.cpp
@@ -150,7 +150,6 @@ namespace Chaos
 			const FMatrix33 Factor =
 				(bIsRigidDynamic0 ? ComputeFactorMatrix3(VectorToPoint1, WorldSpaceInvI1, PBDRigid0->InvM()) : FMatrix33(0)) +
 				(bIsRigidDynamic1 ? ComputeFactorMatrix3(VectorToPoint2, WorldSpaceInvI2, PBDRigid1->InvM()) : FMatrix33(0));
-			FVec3 Impulse;
 			FVec3 AngularImpulse(0);
 
 			// Resting contact if very close to the surface
@@ -162,31 +161,29 @@ namespace Chaos
 			const FVec3 VelocityTarget = (-Restitution*FVec3::DotProduct(RelativeVelocity, Contact.Normal)) * Contact.Normal;
 			const FVec3 VelocityChange = VelocityTarget - RelativeVelocity;
 			const FMatrix33 FactorInverse = Factor.Inverse();
-			Impulse = FactorInverse * VelocityChange;  // Delta Impulse = (J.M^(-1).J^(T))^(-1).(ContactVelocityError)
+			FVec3 Impulse = FactorInverse * VelocityChange;  // Delta Impulse = (J.M^(-1).J^(T))^(-1).(ContactVelocityError)
 
 			// clip the impulse so that the accumulated impulse is not in the wrong direction and in the friction cone
 			// Clipping the accumulated impulse instead of the delta impulse (for the current iteration) is very important for jitter
+			const FVec3 NewUnclippedAccumulatedImpulse = AccumulatedImpulse + Impulse;
+			FVec3 ClippedAccumulatedImpulse(0.0f);  // To be calculated
 			{
-				const FVec3 OldAccumulatedImpulse = AccumulatedImpulse;
-				const FVec3 NewAccumulatedImpulse = AccumulatedImpulse + Impulse;
-				FVec3 ClippedAccumulatedImpulse(0.0f);
 				//Project to normal
-				const float NewAccImpNormalSize = FVec3::DotProduct(NewAccumulatedImpulse, Contact.Normal);
+				const float NewAccImpNormalSize = FVec3::DotProduct(NewUnclippedAccumulatedImpulse, Contact.Normal);
 				if (NewAccImpNormalSize > 0) // Clipping impulses to be positive
 				{
 					if (Friction <= 0)
 					{
 						ClippedAccumulatedImpulse = NewAccImpNormalSize * Contact.Normal;
-						Impulse = ClippedAccumulatedImpulse - OldAccumulatedImpulse;
 					}
 					else
 					{
-						const FVec3 ImpulseTangential = NewAccumulatedImpulse - NewAccImpNormalSize * Contact.Normal;
+						const FVec3 ImpulseTangential = NewUnclippedAccumulatedImpulse - NewAccImpNormalSize * Contact.Normal;
 						const FReal ImpulseTangentialSize = ImpulseTangential.Size();
 						const FReal MaximumImpulseTangential = Friction * NewAccImpNormalSize;
 						if (ImpulseTangentialSize <= MaximumImpulseTangential)
 						{
-							//Static friction case. We are done Impulse is correct, unless we need to add Angular friction:
+							//Static friction case.
 							if (AngularFriction)
 							{
 								ApplyAngularFriction(
@@ -207,27 +204,36 @@ namespace Chaos
 									WorldSpaceInvI1,
 									WorldSpaceInvI2);
 							}
+							ClippedAccumulatedImpulse = Impulse + AccumulatedImpulse;
 						}
 						else
 						{
 							ClippedAccumulatedImpulse = (MaximumImpulseTangential / ImpulseTangentialSize) * ImpulseTangential + NewAccImpNormalSize * Contact.Normal;
-							Impulse = ClippedAccumulatedImpulse - OldAccumulatedImpulse;
 						}
 					}
 				}
 				else
 				{
 					// Clipped Impulse is 0
-					Impulse = -OldAccumulatedImpulse; // Important case: the delta impulse (Impulse) should remove the accumulated impulse
+					ClippedAccumulatedImpulse = FVec3(0); // Important case: the delta impulse should remove the accumulated impulse
 				}
 			}
-				
-			AccumulatedImpulse += Impulse;
+
+			FVec3 OutDeltaImpulse = ClippedAccumulatedImpulse - AccumulatedImpulse;
+			if (Chaos_Collision_EnergyClampEnabled != 0)
+			{
+				// Clamp the delta impulse to make sure we don't gain kinetic energy (ignore potential energy)
+				// Todo: Investigate Energy clamping to work with accumulated impulses
+				// else this might introduce a small amount of jitter, since we are clamping delta impulses instead of accumulated ones
+				OutDeltaImpulse = GetEnergyClampedImpulse(Particle0->CastToRigidParticle(), Particle1->CastToRigidParticle(), OutDeltaImpulse, VectorToPoint1, VectorToPoint2, Body1Velocity, Body2Velocity);
+			}
+
+			AccumulatedImpulse += OutDeltaImpulse; // Now update the accumulated Impulse
 			if (bIsRigidDynamic0)
 			{
 				// Velocity update for next step
-				const FVec3 NetAngularImpulse = FVec3::CrossProduct(VectorToPoint1, Impulse) + AngularImpulse;
-				const FVec3 DV = PBDRigid0->InvM() * Impulse;
+				const FVec3 NetAngularImpulse = FVec3::CrossProduct(VectorToPoint1, OutDeltaImpulse) + AngularImpulse;
+				const FVec3 DV = PBDRigid0->InvM() * OutDeltaImpulse;
 				const FVec3 DW = WorldSpaceInvI1 * NetAngularImpulse;
 				PBDRigid0->V() += DV;
 				PBDRigid0->W() += DW;
@@ -240,8 +246,8 @@ namespace Chaos
 			if (bIsRigidDynamic1)
 			{
 				// Velocity update for next step
-				const FVec3 NetAngularImpulse = FVec3::CrossProduct(VectorToPoint2, -Impulse) - AngularImpulse;
-				const FVec3 DV = -PBDRigid1->InvM() * Impulse;
+				const FVec3 NetAngularImpulse = FVec3::CrossProduct(VectorToPoint2, -OutDeltaImpulse) - AngularImpulse;
+				const FVec3 DV = -PBDRigid1->InvM() * OutDeltaImpulse;
 				const FVec3 DW = WorldSpaceInvI2 * NetAngularImpulse;
 				PBDRigid1->V() += DV;
 				PBDRigid1->W() += DW;
