Added traits to chaos solver.

diff --git a/Engine/Source/Runtime/Engine/Private/PhysicsEngine/Experimental/PhysInterface_Chaos.cpp b/Engine/Source/Runtime/Engine/Private/PhysicsEngine/Experimental/PhysInterface_Chaos.cpp
--- a/Engine/Source/Runtime/Engine/Private/PhysicsEngine/Experimental/PhysInterface_Chaos.cpp
+++ b/Engine/Source/Runtime/Engine/Private/PhysicsEngine/Experimental/PhysInterface_Chaos.cpp
@@ -94,7 +94,7 @@ Chaos::FChaosPhysicsMaterial* GetMaterialFromInternalFaceIndex(const FPhysicsSha
 	const auto& Materials = Shape.GetMaterials();
 	if(Materials.Num() > 0 && Actor.GetProxy())
 	{
-		Chaos::FPBDRigidsSolver* Solver = Actor.GetProxy()->GetSolver();
+		Chaos::FPBDRigidsSolver* Solver = Actor.GetProxy()->GetSolver<Chaos::FPBDRigidsSolver>();
 
 		if(ensure(Solver))
 		{
@@ -124,7 +124,7 @@ Chaos::FChaosPhysicsMaterial* GetMaterialFromInternalFaceIndexAndHitLocation(con
 
 		if (Shape.GetMaterials().Num() > 0 && Actor.GetProxy())
 		{
-			Chaos::FPBDRigidsSolver* Solver = Actor.GetProxy()->GetSolver();
+			Chaos::FPBDRigidsSolver* Solver = Actor.GetProxy()->GetSolver<Chaos::FPBDRigidsSolver>();
 
 			if (ensure(Solver))
 			{
@@ -477,7 +477,7 @@ FPhysScene* FPhysInterface_Chaos::GetCurrentScene(const FPhysicsActorHandle& InH
 
 	if (IPhysicsProxyBase* Proxy = InHandle->GetProxy())
 	{
-		Chaos::FPhysicsSolver* Solver = Proxy->GetSolver();
+		Chaos::FPBDRigidsSolver* Solver = Proxy->GetSolver<Chaos::FPBDRigidsSolver>();
 		return static_cast<FPhysScene*>(Solver ? Solver->PhysSceneHack : nullptr);
 	}
 	return nullptr;
diff --git a/Engine/Source/Runtime/Engine/Private/PhysicsEngine/Experimental/PhysScene_Chaos.cpp b/Engine/Source/Runtime/Engine/Private/PhysicsEngine/Experimental/PhysScene_Chaos.cpp
--- a/Engine/Source/Runtime/Engine/Private/PhysicsEngine/Experimental/PhysScene_Chaos.cpp
+++ b/Engine/Source/Runtime/Engine/Private/PhysicsEngine/Experimental/PhysScene_Chaos.cpp
@@ -717,10 +717,10 @@ void RemovePhysicsProxy(ObjectType* InObject, Chaos::FPhysicsSolver* InSolver, F
 	const bool bDedicatedThread = PhysDispatcher->GetMode() == Chaos::EThreadingMode::DedicatedThread;
 
 	// Remove the object from the solver
-	PhysDispatcher->EnqueueCommandImmediate(InSolver, [InObject, bDedicatedThread](Chaos::FPBDRigidsSolver* InnerSolver)
+	InSolver->EnqueueCommandImmediate([InSolver, InObject, bDedicatedThread]()
 	{
 #if CHAOS_PARTICLEHANDLE_TODO
-		InnerSolver->UnregisterObject(InObject);
+		InSolver->UnregisterObject(InObject);
 #endif
 		InObject->OnRemoveFromScene();
 
@@ -791,7 +791,7 @@ void FPhysScene_Chaos::RemoveObject(FGeometryParticlePhysicsProxy* InObject)
 
 void FPhysScene_Chaos::RemoveObject(FGeometryCollectionPhysicsProxy* InObject)
 {
-	Chaos::FPhysicsSolver* Solver = InObject->GetSolver();
+	Chaos::FPhysicsSolver* Solver = InObject->GetSolver<Chaos::FPhysicsSolver>();
 	if(Solver && !Solver->UnregisterObject(InObject))
 	{
 		UE_LOG(LogChaos, Warning, TEXT("Attempted to remove an object that wasn't found in its solver's gamethread storage - it's likely the solver has been mistakenly changed."));
@@ -802,7 +802,8 @@ void FPhysScene_Chaos::RemoveObject(FGeometryCollectionPhysicsProxy* InObject)
 
 void FPhysScene_Chaos::RemoveObject(FFieldSystemPhysicsProxy* InObject)
 {
-	Chaos::FPhysicsSolver* CurrSceneSolver = InObject->GetSolver();
+	//Does it make sense to remove field form just one solver since it affects multiple solvers?
+	Chaos::FPhysicsSolver* CurrSceneSolver = InObject->GetSolver<Chaos::FPhysicsSolver>();
 	if(CurrSceneSolver)
 	{
 		if(!CurrSceneSolver->UnregisterObject(InObject))
@@ -1347,11 +1348,11 @@ void FPhysScene_ChaosInterface::Flush_AssumesLocked()
 		//Make sure any dirty proxy data is pushed
 		Solver->PushPhysicsState(Dispatcher);
 
-		TQueue<TFunction<void(Chaos::FPhysicsSolver*)>, EQueueMode::Mpsc>& Queue = Solver->GetCommandQueue();
-		TFunction<void(Chaos::FPhysicsSolver*)> Command;
+		TQueue<TFunction<void()>, EQueueMode::Mpsc>& Queue = Solver->GetCommandQueue();
+		TFunction<void()> Command;
 		while(Queue.Dequeue(Command))
 		{
-			Command(Solver);
+			Command();
 		}
 
 		// Populate the spacial acceleration
@@ -1803,7 +1804,7 @@ void FPhysScene_ChaosInterface::UpdateKinematicsOnDeferredSkelMeshes()
 	{
 		// Assumes all particles have the same solver, safe for now, maybe not in the future.
 		IPhysicsProxyBase* Proxy = ProxiesToDirty[0];
-		Chaos::FPhysicsSolverBase* Solver = Proxy->GetSolver();
+		auto* Solver = Proxy->GetSolver<Chaos::FPhysicsSolverBase>();
 		Solver->AddDirtyProxiesUnsafe(ProxiesToDirty);
 	}
 
@@ -2097,11 +2098,11 @@ void FPhysScene_ChaosInterface::EndFrame(ULineBatchComponent* InLineBatcher)
 		// flush solver queues
 		for (FPhysicsSolver* Solver : SolverList)
 		{
-			TQueue<TFunction<void(Chaos::FPhysicsSolver*)>, EQueueMode::Mpsc>& Queue = Solver->GetCommandQueue();
-			TFunction<void(Chaos::FPhysicsSolver*)> Command;
+			TQueue<TFunction<void()>, EQueueMode::Mpsc>& Queue = Solver->GetCommandQueue();
+			TFunction<void()> Command;
 			while (Queue.Dequeue(Command))
 			{
-				Command(Solver);
+				Command();
 			}
 		}
 
diff --git a/Engine/Source/Runtime/Engine/Public/Physics/Experimental/PhysScene_Chaos.h b/Engine/Source/Runtime/Engine/Public/Physics/Experimental/PhysScene_Chaos.h
--- a/Engine/Source/Runtime/Engine/Public/Physics/Experimental/PhysScene_Chaos.h
+++ b/Engine/Source/Runtime/Engine/Public/Physics/Experimental/PhysScene_Chaos.h
@@ -18,6 +18,7 @@
 #include "Framework/Threading.h"
 #include "Chaos/Core.h"
 #include "Chaos/CollisionResolutionTypes.h"
+#include "Chaos/PBDRigidsEvolutionFwd.h"
 
 #ifndef CHAOS_WITH_PAUSABLE_SOLVER
 #define CHAOS_WITH_PAUSABLE_SOLVER 1
@@ -37,7 +38,6 @@ struct FSolverStateStorage;
 
 class FSkeletalMeshPhysicsProxy;
 class FStaticMeshPhysicsProxy;
-class FGeometryCollectionPhysicsProxy;
 class FFieldSystemPhysicsProxy;
 
 class IPhysicsProxyBase;
diff --git a/Engine/Source/Runtime/Experimental/Chaos/Public/Chaos/Declares.h b/Engine/Source/Runtime/Experimental/Chaos/Public/Chaos/Declares.h
--- a/Engine/Source/Runtime/Experimental/Chaos/Public/Chaos/Declares.h
+++ b/Engine/Source/Runtime/Experimental/Chaos/Public/Chaos/Declares.h
@@ -3,11 +3,10 @@
 #pragma once
 
 #include "Chaos/ParticleHandleFwd.h"
+#include "Chaos/PBDRigidsEvolutionFwd.h"
 
 namespace Chaos
 {
-	class FPBDRigidsSolver;
-
 	typedef FPBDRigidsSolver FPhysicsSolver;
 
 	class FImplicitObject;
diff --git a/Engine/Source/Runtime/Experimental/Chaos/Public/Chaos/Defines.h b/Engine/Source/Runtime/Experimental/Chaos/Public/Chaos/Defines.h
--- a/Engine/Source/Runtime/Experimental/Chaos/Public/Chaos/Defines.h
+++ b/Engine/Source/Runtime/Experimental/Chaos/Public/Chaos/Defines.h
@@ -1,6 +1,22 @@
 // Copyright Epic Games, Inc. All Rights Reserved.
 #pragma once
 
+#ifndef CHAOS_TEMPLATE_API
+#if PLATFORM_MAC || PLATFORM_LINUX
+#define CHAOS_TEMPLATE_API CHAOS_API
+#else
+#define CHAOS_TEMPLATE_API
+#endif
+#endif
+
+#ifndef CHAOSSOLVERS_TEMPLATE_API
+#if PLATFORM_MAC || PLATFORM_LINUX
+#define CHAOSSOLVERS_TEMPLATE_API CHAOSSOLVERS_API
+#else
+#define CHAOSSOLVERS_TEMPLATE_API
+#endif
+#endif
+
 #if COMPILE_WITHOUT_UNREAL_SUPPORT
 	#include <stdint.h>
 #else
diff --git a/Engine/Source/Runtime/Experimental/Chaos/Public/Chaos/Framework/PhysicsProxy.h b/Engine/Source/Runtime/Experimental/Chaos/Public/Chaos/Framework/PhysicsProxy.h
--- a/Engine/Source/Runtime/Experimental/Chaos/Public/Chaos/Framework/PhysicsProxy.h
+++ b/Engine/Source/Runtime/Experimental/Chaos/Public/Chaos/Framework/PhysicsProxy.h
@@ -72,13 +72,16 @@ public:
 	void ParameterUpdateCallback(FParticlesType& InParticles, const float InTime) { static_cast<Concrete*>(this)->ParameterUpdateCallback(InParticles, InTime); }
 	void DisableCollisionsCallback(TSet<TTuple<int32, int32>>& InPairs) { static_cast<Concrete*>(this)->DisableCollisionsCallback(InPairs); }
 	void AddForceCallback(FParticlesType& InParticles, const float InDt, const int32 InIndex) { static_cast<Concrete*>(this)->AddForceCallback(InParticles, InDt, InIndex); }
-	void FieldForcesUpdateCallback(Chaos::FPhysicsSolver* InSolver, FParticlesType& Particles, Chaos::TArrayCollectionArray<FVector> & Force, Chaos::TArrayCollectionArray<FVector> & Torque, const float Time) { static_cast<Concrete*>(this)->FieldForcesUpdateCallback(InSolver, Particles, Force, Torque, Time); }
+
+	template <typename Traits>
+	void FieldForcesUpdateCallback(Chaos::TPBDRigidsSolver<Traits>* InSolver, FParticlesType& Particles, Chaos::TArrayCollectionArray<FVector> & Force, Chaos::TArrayCollectionArray<FVector> & Torque, const float Time) { static_cast<Concrete*>(this)->FieldForcesUpdateCallback(InSolver, Particles, Force, Torque, Time); }
 
 	/** The Particle Binding creates a connection between the particles in the simulation and the solver objects dataset. */
 	void BindParticleCallbackMapping(Chaos::TArrayCollectionArray<PhysicsProxyWrapper> & PhysicsProxyReverseMap, Chaos::TArrayCollectionArray<int32> & ParticleIDReverseMap) {static_cast<Concrete*>(this)->BindParticleCallbackMapping(PhysicsProxyReverseMap, ParticleIDReverseMap);}
 
 	/** Called to buffer a command to be processed at the next available safe opportunity */
-	void BufferCommand(Chaos::FPhysicsSolver* InSolver, const FFieldSystemCommand& InCommand) { static_cast<Concrete*>(this)->BufferCommand(InSolver, InCommand); }
+	template <typename Traits>
+	void BufferCommand(Chaos::TPBDRigidsSolver<Traits>* InSolver, const FFieldSystemCommand& InCommand) { static_cast<Concrete*>(this)->BufferCommand(InSolver, InCommand); }
 
 	/** Returns the concrete type of the derived class*/
 	EPhysicsProxyType ConcreteType() { return static_cast<Concrete*>(this)->ConcreteType(); }
diff --git a/Engine/Source/Runtime/Experimental/Chaos/Public/Chaos/Framework/PhysicsProxyBase.h b/Engine/Source/Runtime/Experimental/Chaos/Public/Chaos/Framework/PhysicsProxyBase.h
--- a/Engine/Source/Runtime/Experimental/Chaos/Public/Chaos/Framework/PhysicsProxyBase.h
+++ b/Engine/Source/Runtime/Experimental/Chaos/Public/Chaos/Framework/PhysicsProxyBase.h
@@ -33,7 +33,7 @@ public:
 
 	virtual UObject* GetOwner() const = 0;
 
-	template< class SOLVER_TYPE = Chaos::FPhysicsSolver>
+	template< class SOLVER_TYPE>
 	SOLVER_TYPE* GetSolver() const { return static_cast<SOLVER_TYPE*>(Solver); }
 
 	//Should this be in the public API? probably not
diff --git a/Engine/Source/Runtime/Experimental/Chaos/Public/Chaos/Framework/PhysicsSolverBase.h b/Engine/Source/Runtime/Experimental/Chaos/Public/Chaos/Framework/PhysicsSolverBase.h
--- a/Engine/Source/Runtime/Experimental/Chaos/Public/Chaos/Framework/PhysicsSolverBase.h
+++ b/Engine/Source/Runtime/Experimental/Chaos/Public/Chaos/Framework/PhysicsSolverBase.h
@@ -9,8 +9,15 @@
 #include "Chaos/ParticleDirtyFlags.h"
 #include "Async/ParallelFor.h"
 
+class FChaosSolversModule;
+class FPhysicsSolverAdvanceTask;
+
 namespace Chaos
 {
+	class FPersistentPhysicsTask;
+
+	enum class ELockType: uint8;
+
 	struct FDirtyProxy
 	{
 		IPhysicsProxyBase* Proxy;
@@ -178,6 +185,9 @@ namespace Chaos
 
 		void ChangeBufferMode(EMultiBufferMode InBufferMode);
 
+		TQueue<TFunction<void()>,EQueueMode::Mpsc>& GetCommandQueue() { return CommandQueue; }
+		bool HasPendingCommands() const { return !CommandQueue.IsEmpty(); }
+
 		void AddDirtyProxy(IPhysicsProxyBase * ProxyBaseIn)
 		{
 			DirtyProxiesDataBuffer.AccessProducerBuffer()->Add(ProxyBaseIn);
@@ -253,6 +263,12 @@ namespace Chaos
 		FName DebugName;
 #endif
 
+	//
+	// Commands
+	//
+	TQueue<TFunction<void()>,EQueueMode::Mpsc> CommandQueue;
+
+
 	private:
 
 		/** 
@@ -262,5 +278,13 @@ namespace Chaos
 		 * @see FChaosSolversModule::CreateSolver
 		 */
 		const UObject* Owner = nullptr;
+		FRWLock QueryMaterialLock;
+
+		friend FChaosSolversModule;
+		friend FPersistentPhysicsTask;
+		friend FPhysicsSolverAdvanceTask;
+
+		template<ELockType>
+		friend struct TSolverQueryMaterialScope;
 	};
 }
diff --git a/Engine/Source/Runtime/Experimental/Chaos/Public/Chaos/PBDRigidsEvolution.h b/Engine/Source/Runtime/Experimental/Chaos/Public/Chaos/PBDRigidsEvolution.h
--- a/Engine/Source/Runtime/Experimental/Chaos/Public/Chaos/PBDRigidsEvolution.h
+++ b/Engine/Source/Runtime/Experimental/Chaos/Public/Chaos/PBDRigidsEvolution.h
@@ -15,6 +15,7 @@
 #include "Chaos/SpatialAccelerationCollection.h"
 #include "Chaos/EvolutionTraits.h"
 #include "Chaos/PBDRigidsEvolutionFwd.h"
+#include "Chaos/Defines.h"
 
 
 extern int32 ChaosRigidsEvolutionApplyAllowEarlyOutCVar;
@@ -919,11 +920,7 @@ protected:
 	TUniquePtr<ISpatialAccelerationCollectionFactory> SpatialCollectionFactory;
 };
 
-#if PLATFORM_MAC || PLATFORM_LINUX
-#define EVOLUTION_TRAIT(Trait) extern template class CHAOS_API TPBDRigidsEvolutionBase<Trait>;
-#else
-#define EVOLUTION_TRAIT(Trait) extern template class TPBDRigidsEvolutionBase<Trait>;
-#endif
+#define EVOLUTION_TRAIT(Trait) extern template class CHAOS_TEMPLATE_API TPBDRigidsEvolutionBase<Trait>;
 #include "Chaos/EvolutionTraits.inl"
 #undef EVOLUTION_TRAIT
 
diff --git a/Engine/Source/Runtime/Experimental/Chaos/Public/Chaos/PBDRigidsEvolutionFwd.h b/Engine/Source/Runtime/Experimental/Chaos/Public/Chaos/PBDRigidsEvolutionFwd.h
--- a/Engine/Source/Runtime/Experimental/Chaos/Public/Chaos/PBDRigidsEvolutionFwd.h
+++ b/Engine/Source/Runtime/Experimental/Chaos/Public/Chaos/PBDRigidsEvolutionFwd.h
@@ -9,6 +9,23 @@ namespace Chaos
 	template <typename Traits>
 	class TPBDRigidsEvolutionGBF;
 
+	using FDefaultTraits = struct FNonRewindableEvolutionTraits;
+
 	//The default evolution used by unreal
-	using FPBDRigidsEvolution = TPBDRigidsEvolutionGBF<struct FNonRewindableEvolutionTraits>;
-}
\ No newline at end of file
+	using FPBDRigidsEvolution = TPBDRigidsEvolutionGBF<FDefaultTraits>;
+
+	template <typename Traits>
+	class TPBDRigidsSolver;
+
+	using FPBDRigidsSolver = TPBDRigidsSolver<FDefaultTraits>;
+
+	template <typename Traits>
+	class TEventManager;
+
+	using FEventManager = TEventManager<FDefaultTraits>;
+}
+
+template <typename Traits>
+class TGeometryCollectionPhysicsProxy;
+
+using FGeometryCollectionPhysicsProxy = TGeometryCollectionPhysicsProxy<Chaos::FDefaultTraits>;
\ No newline at end of file
diff --git a/Engine/Source/Runtime/Experimental/Chaos/Public/Chaos/PBDRigidsEvolutionGBF.h b/Engine/Source/Runtime/Experimental/Chaos/Public/Chaos/PBDRigidsEvolutionGBF.h
--- a/Engine/Source/Runtime/Experimental/Chaos/Public/Chaos/PBDRigidsEvolutionGBF.h
+++ b/Engine/Source/Runtime/Experimental/Chaos/Public/Chaos/PBDRigidsEvolutionGBF.h
@@ -254,11 +254,7 @@ namespace Chaos
 		FPBDRigidsEvolutionInternalHandleCallback InternalParticleInitilization;
 	};
 
-#if PLATFORM_MAC || PLATFORM_LINUX
-#define EVOLUTION_TRAIT(Trait) extern template class CHAOS_API TPBDRigidsEvolutionGBF<Trait>;
-#else
-#define EVOLUTION_TRAIT(Trait) extern template class TPBDRigidsEvolutionGBF<Trait>;
-#endif
+#define EVOLUTION_TRAIT(Trait) extern template class CHAOS_TEMPLATE_API TPBDRigidsEvolutionGBF<Trait>;
 #include "Chaos/EvolutionTraits.inl"
 #undef EVOLUTION_TRAIT
 }
diff --git a/Engine/Source/Runtime/Experimental/ChaosSolverEngine/Private/Chaos/ChaosSolverActor.cpp b/Engine/Source/Runtime/Experimental/ChaosSolverEngine/Private/Chaos/ChaosSolverActor.cpp
--- a/Engine/Source/Runtime/Experimental/ChaosSolverEngine/Private/Chaos/ChaosSolverActor.cpp
+++ b/Engine/Source/Runtime/Experimental/ChaosSolverEngine/Private/Chaos/ChaosSolverActor.cpp
@@ -272,11 +272,9 @@ void AChaosSolverActor::BeginPlay()
 {
 	Super::BeginPlay();
 
-	Chaos::IDispatcher* PhysDispatcher = PhysScene->GetDispatcher();
-	if (PhysDispatcher)
-	{
-		PhysDispatcher->EnqueueCommandImmediate(Solver,
-			[InTimeStepMultiplier = TimeStepMultiplier
+	Solver->EnqueueCommandImmediate(
+		[InSolver = Solver
+		, InTimeStepMultiplier = TimeStepMultiplier
 			, InCollisionIterations = CollisionIterations
 			, InPushOutIterations = PushOutIterations
 			, InPushOutPairIterations = PushOutPairIterations
@@ -292,6 +290,7 @@ void AChaosSolverActor::BeginPlay()
 			, InFloorHeight = FloorHeight
 			, InMassScale = MassScale
 			, InGenerateContactGraph = bGenerateContactGraph]
+		]
 		(Chaos::FPhysicsSolver* InSolver)
 		{
 #if TODO_REIMPLEMENT_SOLVER_SETTINGS_ACCESSORS
@@ -334,7 +333,6 @@ void AChaosSolverActor::BeginPlay()
 				InSolver->RegisterObject(Particle);
 			}
 		});
-	}
 
 #if TODO_REIMPLEMENT_DEBUG_SUBSTEP
 #if CHAOS_DEBUG_SUBSTEP
@@ -355,17 +353,13 @@ void AChaosSolverActor::BeginPlay()
 
 void AChaosSolverActor::EndPlay(const EEndPlayReason::Type ReasonEnd)
 {
-	Chaos::IDispatcher* PhysDispatcher = PhysScene->GetDispatcher();
-	if (PhysDispatcher)
-	{
-		PhysDispatcher->EnqueueCommandImmediate(Solver, [](Chaos::FPhysicsSolver* InSolver)
+	Solver->EnqueueCommandImmediate([InSolver=Solver]()
 		{
 			// #TODO BG - We should really reset the solver here but the current reset function
 			// is really heavy handed and clears out absolutely everything. Ideally we want to keep
 			// all of the solver physics proxies and revert to a state before the very first tick
 			InSolver->SetEnabled(false);
 		});
-	}
 #if TODO_REIMPLEMENT_DEBUG_SUBSTEP
 #if CHAOS_DEBUG_SUBSTEP
 	ChaosSolverActorConsoleObjects->RemoveSolver(GetName());
@@ -410,16 +404,11 @@ void AChaosSolverActor::SetSolverActive(bool bActive)
 {
 	if(Solver && PhysScene)
 	{
-		Chaos::IDispatcher* Dispatcher = PhysScene->GetDispatcher();
-
-		if(Dispatcher)
-		{
-			Dispatcher->EnqueueCommandImmediate(Solver, [bShouldBeActive = bActive](Chaos::FPhysicsSolver* InSolver)
+		Solver->EnqueueCommandImmediate([InSolver=Solver, bShouldBeActive = bActive]()
 			{
 				InSolver->SetEnabled(bShouldBeActive);
 			});
 		}
-	}
 }
 
 #if WITH_EDITOR
@@ -429,38 +418,35 @@ void AChaosSolverActor::PostEditChangeProperty(struct FPropertyChangedEvent& Pro
 
 	if (Solver && PropertyChangedEvent.Property)
 	{
-		Chaos::IDispatcher* PhysDispatcher = PhysScene->GetDispatcher();
-		if (PhysDispatcher)
-		{
 #if TODO_REIMPLEMENT_TIMESTEP_MULTIPLIER
 			if (PropertyChangedEvent.Property->GetFName() == GET_MEMBER_NAME_CHECKED(AChaosSolverActor, TimeStepMultiplier))
 			{
-				PhysDispatcher->EnqueueCommandImmediate(Solver, [InTimeStepMultiplier = TimeStepMultiplier]
-				(Chaos::FPhysicsSolver* InSolver)
+			Solver->EnqueueCommandImmediate([InSolver=Solver, InTimeStepMultiplier = TimeStepMultiplier]
+			()
 				{
 					InSolver->SetTimeStepMultiplier(InTimeStepMultiplier);
 				});
 			}
 			else if (PropertyChangedEvent.Property->GetFName() == GET_MEMBER_NAME_CHECKED(AChaosSolverActor, CollisionIterations))
 			{
-				PhysDispatcher->EnqueueCommandImmediate(Solver, [InCollisionIterations = CollisionIterations]
-				(Chaos::FPhysicsSolver* InSolver)
+			Solver->EnqueueCommandImmediate([InSolver=Solver, InCollisionIterations = CollisionIterations]
+			()
 				{
 					InSolver->SetIterations(InCollisionIterations);
 				});
 			}
 			else if (PropertyChangedEvent.Property->GetFName() == GET_MEMBER_NAME_CHECKED(AChaosSolverActor, PushOutIterations))
 			{
-				PhysDispatcher->EnqueueCommandImmediate(Solver, [InPushOutIterations = PushOutIterations]
-				(Chaos::FPhysicsSolver* InSolver)
+			Solver->EnqueueCommandImmediate([InSolver=Solver, InPushOutIterations = PushOutIterations]
+			()
 				{
 					InSolver->SetPushOutIterations(InPushOutIterations);
 				});
 			}
 			else if (PropertyChangedEvent.Property->GetFName() == GET_MEMBER_NAME_CHECKED(AChaosSolverActor, PushOutPairIterations))
 			{
-				PhysDispatcher->EnqueueCommandImmediate(Solver, [InPushOutPairIterations = PushOutPairIterations]
-				(Chaos::FPhysicsSolver* InSolver)
+			Solver->EnqueueCommandImmediate([InSolver=Solver, InPushOutPairIterations = PushOutPairIterations]
+			()
 				{
 					InSolver->SetPushOutPairIterations(InPushOutPairIterations);
 				});
@@ -469,65 +455,58 @@ void AChaosSolverActor::PostEditChangeProperty(struct FPropertyChangedEvent& Pro
 #endif
 			if (PropertyChangedEvent.Property->GetFName() == GET_MEMBER_NAME_CHECKED(AChaosSolverActor, DoGenerateCollisionData))
 			{
-				PhysDispatcher->EnqueueCommandImmediate(Solver, [InDoGenerateCollisionData = DoGenerateCollisionData]
-				(Chaos::FPhysicsSolver* InSolver)
+			Solver->EnqueueCommandImmediate([InSolver=Solver, InDoGenerateCollisionData = DoGenerateCollisionData]
+			()
 				{
 					InSolver->SetGenerateCollisionData(InDoGenerateCollisionData);
 				});
 			}
 			else if (PropertyChangedEvent.MemberProperty->GetFName() == GET_MEMBER_NAME_CHECKED(AChaosSolverActor, CollisionFilterSettings))
 			{
-				PhysDispatcher->EnqueueCommandImmediate(Solver, [InCollisionFilterSettings = CollisionFilterSettings]
-				(Chaos::FPhysicsSolver* InSolver)
+			Solver->EnqueueCommandImmediate([InSolver=Solver, InCollisionFilterSettings = CollisionFilterSettings]
+			()
 				{
 					InSolver->SetCollisionFilterSettings(InCollisionFilterSettings);
 				});
 			}
 			else if (PropertyChangedEvent.Property->GetFName() == GET_MEMBER_NAME_CHECKED(AChaosSolverActor, DoGenerateBreakingData))
 			{
-				PhysDispatcher->EnqueueCommandImmediate(Solver, [InDoGenerateBreakingData = DoGenerateBreakingData]
-				(Chaos::FPhysicsSolver* InSolver)
+			Solver->EnqueueCommandImmediate([InSolver=Solver, InDoGenerateBreakingData = DoGenerateBreakingData]
+			()
 				{
 					InSolver->SetGenerateBreakingData(InDoGenerateBreakingData);
 				});
 			}
 			else if (PropertyChangedEvent.MemberProperty->GetFName() == GET_MEMBER_NAME_CHECKED(AChaosSolverActor, BreakingFilterSettings))
 			{
-				PhysDispatcher->EnqueueCommandImmediate(Solver, [InBreakingFilterSettings = BreakingFilterSettings]
-				(Chaos::FPhysicsSolver* InSolver)
+			Solver->EnqueueCommandImmediate([InSolver=Solver, InBreakingFilterSettings = BreakingFilterSettings]
+			()
 				{
 					InSolver->SetBreakingFilterSettings(InBreakingFilterSettings);
 				});
 			}
 			else if (PropertyChangedEvent.Property->GetFName() == GET_MEMBER_NAME_CHECKED(AChaosSolverActor, DoGenerateTrailingData))
 			{
-				PhysDispatcher->EnqueueCommandImmediate(Solver, [InDoGenerateTrailingData = DoGenerateTrailingData]
-				(Chaos::FPhysicsSolver* InSolver)
+			Solver->EnqueueCommandImmediate([InSolver=Solver, InDoGenerateTrailingData = DoGenerateTrailingData]
+			()
 				{
 					InSolver->SetGenerateTrailingData(InDoGenerateTrailingData);
 				});
 			}
 			else if (PropertyChangedEvent.MemberProperty->GetFName() == GET_MEMBER_NAME_CHECKED(AChaosSolverActor, TrailingFilterSettings))
 			{
-				PhysDispatcher->EnqueueCommandImmediate(Solver, [InTrailingFilterSettings = TrailingFilterSettings]
-				(Chaos::FPhysicsSolver* InSolver)
+			Solver->EnqueueCommandImmediate([InSolver=Solver, InTrailingFilterSettings = TrailingFilterSettings]
+			()
 				{
 					InSolver->SetTrailingFilterSettings(InTrailingFilterSettings);
 				});
 			}
 			else if (PropertyChangedEvent.Property->GetFName() == GET_MEMBER_NAME_CHECKED(AChaosSolverActor, bHasFloor))
 			{
-				PhysDispatcher->EnqueueCommandImmediate(Solver, [InHasFloor = bHasFloor]
-				(Chaos::FPhysicsSolver* InSolver)
-				{
-				});
+
 			}
 			else if (PropertyChangedEvent.Property->GetFName() == GET_MEMBER_NAME_CHECKED(AChaosSolverActor, FloorHeight))
 			{
-				PhysDispatcher->EnqueueCommandImmediate(Solver, [InFloorHeight = FloorHeight]
-				(Chaos::FPhysicsSolver* InSolver)
-				{
-				});
 			}
 			else if (PropertyChangedEvent.Property->GetFName() == GET_MEMBER_NAME_CHECKED(AChaosSolverActor, bGenerateContactGraph))
 			{
@@ -540,15 +519,14 @@ void AChaosSolverActor::PostEditChangeProperty(struct FPropertyChangedEvent& Pro
 #if TODO_REIMPLEMENT_TIMESTEP_MULTIPLIER
 			else if (PropertyChangedEvent.Property->GetFName() == GET_MEMBER_NAME_CHECKED(AChaosSolverActor, MassScale))
 			{
-				PhysDispatcher->EnqueueCommandImmediate(Solver, [InMassScale = MassScale]
-				(Chaos::FPhysicsSolver* InSolver)
+			Solver->EnqueueCommandImmediate([InSolver=Solver, InMassScale = MassScale]
+			()
 				{
 					InSolver->SetMassScale(InMassScale);
 				});
 			}
 #endif
 		}
-	}
 
 #if CHAOS_DEBUG_SUBSTEP
 #if TODO_REIMPLEMENT_DEBUG_SUBSTEP
@@ -593,7 +571,7 @@ void SerializeForPerfTest(const TArray< FString >&, UWorld* World, FOutputDevice
 	for (TActorIterator<AChaosSolverActor> Itr(World); Itr; ++Itr)
 	{
 		Chaos::FPhysicsSolver* Solver = Itr->GetSolver();
-		FChaosSolversModule::GetModule()->GetDispatcher()->EnqueueCommandImmediate(Solver, [FileName](Chaos::FPhysicsSolver* SolverIn) { SolverIn->SerializeForPerfTest(FileName); });
+		Solver->EnqueueCommandImmediate([FileName, InSolver=Solver]() { InSolver->SerializeForPerfTest(FileName); });
 	}
 }
 
diff --git a/Engine/Source/Runtime/Experimental/ChaosSolvers/Private/ChaosSolversModule.cpp b/Engine/Source/Runtime/Experimental/ChaosSolvers/Private/ChaosSolversModule.cpp
--- a/Engine/Source/Runtime/Experimental/ChaosSolvers/Private/ChaosSolversModule.cpp
+++ b/Engine/Source/Runtime/Experimental/ChaosSolvers/Private/ChaosSolversModule.cpp
@@ -306,7 +306,7 @@ void FChaosSolversModule::ShutdownThreadingMode()
 				IDispatcher::FSolverCommand Command;
 				while(Solver->CommandQueue.Dequeue(Command))
 				{
-					Command(Solver);
+					Command();
 				}
 			}
 
@@ -953,10 +953,10 @@ void FChaosSolversModule::ChangeBufferMode(Chaos::EMultiBufferMode BufferMode)
 	{
 		if (Dispatcher)
 		{
-			Dispatcher->EnqueueCommandImmediate(Solver, [InBufferMode = BufferMode]
-			(Chaos::FPhysicsSolver* InSolver)
+			Dispatcher->EnqueueCommandImmediate(Solver, [InBufferMode = BufferMode, Solver]
+			()
 			{
-				InSolver->ChangeBufferMode(InBufferMode);
+				Solver->ChangeBufferMode(InBufferMode);
 			});
 		}
 		else
@@ -998,9 +998,9 @@ void FChaosSolversModule::OnUpdateMaterial(Chaos::FMaterialHandle InHandle)
 		for(Chaos::FPhysicsSolver* Solver : AllSolvers)
 		{
 			// Send a copy of the material to each solver
-			Dispatcher->EnqueueCommandImmediate(Solver, [InHandle, MaterialCopy = *Material](Chaos::FPhysicsSolver* InSolver)
+			Dispatcher->EnqueueCommandImmediate(Solver, [InHandle, MaterialCopy = *Material, Solver]()
 			{
-				InSolver->UpdateMaterial(InHandle, MaterialCopy);
+				Solver->UpdateMaterial(InHandle, MaterialCopy);
 			});
 		}
 	}
@@ -1018,9 +1018,9 @@ void FChaosSolversModule::OnCreateMaterial(Chaos::FMaterialHandle InHandle)
 		for(Chaos::FPhysicsSolver* Solver : AllSolvers)
 		{
 			// Send a copy of the material to each solver
-			Dispatcher->EnqueueCommandImmediate(Solver, [InHandle, MaterialCopy = *Material](Chaos::FPhysicsSolver* InSolver)
+			Dispatcher->EnqueueCommandImmediate(Solver, [InHandle, MaterialCopy = *Material, Solver]()
 			{
-				InSolver->CreateMaterial(InHandle, MaterialCopy);
+				Solver->CreateMaterial(InHandle, MaterialCopy);
 			});
 		}
 	}
@@ -1041,9 +1041,9 @@ void FChaosSolversModule::OnDestroyMaterial(Chaos::FMaterialHandle InHandle)
 		for(Chaos::FPhysicsSolver* Solver : AllSolvers)
 		{
 			// Notify each solver
-			Dispatcher->EnqueueCommandImmediate(Solver, [InHandle](Chaos::FPhysicsSolver* InSolver)
+			Dispatcher->EnqueueCommandImmediate(Solver, [InHandle, Solver]()
 			{
-				InSolver->DestroyMaterial(InHandle);
+				Solver->DestroyMaterial(InHandle);
 			});
 		}
 	}
@@ -1061,9 +1061,9 @@ void FChaosSolversModule::OnUpdateMaterialMask(Chaos::FMaterialMaskHandle InHand
 		for (Chaos::FPhysicsSolver* Solver : AllSolvers)
 		{
 			// Send a copy of the material to each solver
-			Dispatcher->EnqueueCommandImmediate(Solver, [InHandle, MaterialMaskCopy = *MaterialMask](Chaos::FPhysicsSolver* InSolver)
+			Dispatcher->EnqueueCommandImmediate(Solver, [InHandle, MaterialMaskCopy = *MaterialMask, Solver]()
 			{
-				InSolver->UpdateMaterialMask(InHandle, MaterialMaskCopy);
+				Solver->UpdateMaterialMask(InHandle, MaterialMaskCopy);
 			});
 		}
 	}
@@ -1081,9 +1081,9 @@ void FChaosSolversModule::OnCreateMaterialMask(Chaos::FMaterialMaskHandle InHand
 		for (Chaos::FPhysicsSolver* Solver : AllSolvers)
 		{
 			// Send a copy of the material to each solver
-			Dispatcher->EnqueueCommandImmediate(Solver, [InHandle, MaterialMaskCopy = *MaterialMask](Chaos::FPhysicsSolver* InSolver)
+			Dispatcher->EnqueueCommandImmediate(Solver, [InHandle, MaterialMaskCopy = *MaterialMask, Solver]()
 			{
-				InSolver->CreateMaterialMask(InHandle, MaterialMaskCopy);
+				Solver->CreateMaterialMask(InHandle, MaterialMaskCopy);
 			});
 		}
 	}
@@ -1104,9 +1104,9 @@ void FChaosSolversModule::OnDestroyMaterialMask(Chaos::FMaterialMaskHandle InHan
 		for (Chaos::FPhysicsSolver* Solver : AllSolvers)
 		{
 			// Notify each solver
-			Dispatcher->EnqueueCommandImmediate(Solver, [InHandle](Chaos::FPhysicsSolver* InSolver)
+			Dispatcher->EnqueueCommandImmediate(Solver, [InHandle, Solver]()
 			{
-				InSolver->DestroyMaterialMask(InHandle);
+				Solver->DestroyMaterialMask(InHandle);
 			});
 		}
 	}
diff --git a/Engine/Source/Runtime/Experimental/ChaosSolvers/Private/EventDefaults.cpp b/Engine/Source/Runtime/Experimental/ChaosSolvers/Private/EventDefaults.cpp
--- a/Engine/Source/Runtime/Experimental/ChaosSolvers/Private/EventDefaults.cpp
+++ b/Engine/Source/Runtime/Experimental/ChaosSolvers/Private/EventDefaults.cpp
@@ -16,7 +16,8 @@
 namespace Chaos
 {
 
-	void FEventDefaults::RegisterSystemEvents(FEventManager& EventManager)
+	template <typename Traits>
+	void TEventDefaults<Traits>::RegisterSystemEvents(TEventManager<Traits>& EventManager)
 	{
 		RegisterCollisionEvent(EventManager);
 		RegisterBreakingEvent(EventManager);
@@ -24,10 +25,11 @@ namespace Chaos
 		RegisterSleepingEvent(EventManager);
 	}
 
-	void FEventDefaults::RegisterCollisionEvent(FEventManager& EventManager)
+	template <typename Traits>
+	void TEventDefaults<Traits>::RegisterCollisionEvent(TEventManager<Traits>& EventManager)
 	{
-		EventManager.RegisterEvent<FCollisionEventData>(EEventType::Collision, []
-		(const Chaos::FPBDRigidsSolver* Solver, FCollisionEventData& CollisionEventData)
+		EventManager.template RegisterEvent<FCollisionEventData>(EEventType::Collision, []
+		(const Chaos::TPBDRigidsSolver<Traits>* Solver, FCollisionEventData& CollisionEventData)
 		{
 			check(Solver);
 			SCOPE_CYCLE_COUNTER(STAT_GatherCollisionEvent);
@@ -45,7 +47,7 @@ namespace Chaos
 			CollisionEventData.CollisionData.TimeCreated = Solver->MTime;
 			CollisionEventData.PhysicsProxyToCollisionIndices.TimeCreated = Solver->MTime;
 
-			const FPBDRigidsSolver::FPBDRigidsEvolution* Evolution = Solver->GetEvolution();
+			const auto* Evolution = Solver->GetEvolution();
 
 			const FPBDCollisionConstraints& CollisionRule = Evolution->GetCollisionConstraints();
 
@@ -53,9 +55,9 @@ namespace Chaos
 			const TPBDRigidParticles<float, 3>& Particles = Evolution->GetParticles().GetDynamicParticles();
 			const TArrayCollectionArray<ClusterId>& ClusterIdsArray = Evolution->GetRigidClustering().GetClusterIdsArray();
 #if TODO_REIMPLEMENT_RIGID_CLUSTERING
-			const Chaos::FPBDRigidsSolver::FClusteringType::FClusterMap& ParentToChildrenMap = Evolution->GetRigidClustering().GetChildrenMap();
+			const Chaos::TPBDRigidsSolver<Traits>::FClusteringType::FClusterMap& ParentToChildrenMap = Evolution->GetRigidClustering().GetChildrenMap();
 #endif
-			const Chaos::TPBDRigidClustering<FPhysicsSolver::FPBDRigidsEvolution, FPBDCollisionConstraints, float, 3>::FClusterMap& ParentToChildrenMap = Evolution->GetRigidClustering().GetChildrenMap();
+			const typename Chaos::TPBDRigidClustering<typename TPBDRigidsSolver<Traits>::FPBDRigidsEvolution, FPBDCollisionConstraints, float, 3>::FClusterMap& ParentToChildrenMap = Evolution->GetRigidClustering().GetChildrenMap();
 
 			if(CollisionRule.NumConstraints() > 0)
 			{
@@ -78,33 +80,33 @@ namespace Chaos
 							{
 								if (NumValidCollisions >= CollisionRule.NumConstraints()) break;
 
-								if (Proxy != nullptr)
+						if (Proxy != nullptr)
+						{
+							if (ensure(!Constraint.AccumulatedImpulse.ContainsNaN() && FMath::IsFinite(Constraint.GetPhi())))
+							{
+								TGeometryParticleHandle<float, 3>* Particle0 = Constraint.Particle[0];
+								TGeometryParticleHandle<float, 3>* Particle1 = Constraint.Particle[1];
+								TKinematicGeometryParticleHandle<float, 3>* Body0 = Particle0->CastToKinematicParticle();
+
+								// presently when a rigidbody or kinematic hits static geometry then Body1 is null
+								TKinematicGeometryParticleHandle<float, 3>* Body1 = Particle1->CastToKinematicParticle();
+
+								if (!Constraint.AccumulatedImpulse.IsZero() && Body0)
 								{
-									if (ensure(!Constraint.AccumulatedImpulse.ContainsNaN() && FMath::IsFinite(Constraint.GetPhi())))
+									if (ensure(!Constraint.GetLocation().ContainsNaN() &&
+										!Constraint.GetNormal().ContainsNaN()) &&
+										!Body0->V().ContainsNaN() &&
+										!Body0->W().ContainsNaN() &&
+										(Body1 == nullptr || ((!Body1->V().ContainsNaN()) && !Body1->W().ContainsNaN())))
 									{
-										TGeometryParticleHandle<float, 3>* Particle0 = Constraint.Particle[0];
-										TGeometryParticleHandle<float, 3>* Particle1 = Constraint.Particle[1];
-										TKinematicGeometryParticleHandle<float, 3>* Body0 = Particle0->CastToKinematicParticle();
-
-										// presently when a rigidbody or kinematic hits static geometry then Body1 is null
-										TKinematicGeometryParticleHandle<float, 3>* Body1 = Particle1->CastToKinematicParticle();
-
-										if (!Constraint.AccumulatedImpulse.IsZero() && Body0)
-										{
-											if (ensure(!Constraint.GetLocation().ContainsNaN() &&
-												!Constraint.GetNormal().ContainsNaN()) &&
-												!Body0->V().ContainsNaN() &&
-												!Body0->W().ContainsNaN() &&
-												(Body1 == nullptr || ((!Body1->V().ContainsNaN()) && !Body1->W().ContainsNaN())))
-											{
-												ValidCollisionHandles[NumValidCollisions] = ContactHandle;
-												NumValidCollisions++;
-											}
-										}
+										ValidCollisionHandles[NumValidCollisions] = ContactHandle;
+										NumValidCollisions++;
 									}
 								}
 							}
 						}
+					}
+						}
 					}
 					else if (ContactHandle->GetType() == FCollisionConstraintBase::FType::MultiPoint)
 					{
@@ -118,35 +120,35 @@ namespace Chaos
 							{
 								if (NumValidCollisions >= CollisionRule.NumConstraints()) break;
 
-								if (Proxy != nullptr)
+						if (Proxy != nullptr)
+						{
+							if (ensure(!Constraint.AccumulatedImpulse.ContainsNaN() && FMath::IsFinite(Constraint.GetPhi())))
+							{
+								TGeometryParticleHandle<float, 3>* Particle0 = Constraint.Particle[0];
+								TGeometryParticleHandle<float, 3>* Particle1 = Constraint.Particle[1];
+								TKinematicGeometryParticleHandle<float, 3>* Body0 = Particle0->CastToKinematicParticle();
+
+								// presently when a rigidbody or kinematic hits static geometry then Body1 is null
+								TKinematicGeometryParticleHandle<float, 3>* Body1 = Particle1->CastToKinematicParticle();
+
+								if (!Constraint.AccumulatedImpulse.IsZero() && Body0)
 								{
-									if (ensure(!Constraint.AccumulatedImpulse.ContainsNaN() && FMath::IsFinite(Constraint.GetPhi())))
+									if (ensure(!Constraint.GetLocation().ContainsNaN() &&
+										!Constraint.GetNormal().ContainsNaN()) &&
+										!Body0->V().ContainsNaN() &&
+										!Body0->W().ContainsNaN() &&
+										(Body1 == nullptr || ((!Body1->V().ContainsNaN()) && !Body1->W().ContainsNaN())))
 									{
-										TGeometryParticleHandle<float, 3>* Particle0 = Constraint.Particle[0];
-										TGeometryParticleHandle<float, 3>* Particle1 = Constraint.Particle[1];
-										TKinematicGeometryParticleHandle<float, 3>* Body0 = Particle0->CastToKinematicParticle();
-
-										// presently when a rigidbody or kinematic hits static geometry then Body1 is null
-										TKinematicGeometryParticleHandle<float, 3>* Body1 = Particle1->CastToKinematicParticle();
-
-										if (!Constraint.AccumulatedImpulse.IsZero() && Body0)
-										{
-											if (ensure(!Constraint.GetLocation().ContainsNaN() &&
-												!Constraint.GetNormal().ContainsNaN()) &&
-												!Body0->V().ContainsNaN() &&
-												!Body0->W().ContainsNaN() &&
-												(Body1 == nullptr || ((!Body1->V().ContainsNaN()) && !Body1->W().ContainsNaN())))
-											{
-												ValidCollisionHandles[NumValidCollisions] = ContactHandle;
-												NumValidCollisions++;
-											}
-										}
+										ValidCollisionHandles[NumValidCollisions] = ContactHandle;
+										NumValidCollisions++;
 									}
 								}
 							}
 						}
 					}
 				}
+					}
+				}
 
 				ValidCollisionHandles.SetNum(NumValidCollisions);
 
@@ -221,11 +223,11 @@ namespace Chaos
 #endif
 
 								// Add to AllCollisionsIndicesByPhysicsProxy
-								AllCollisionsIndicesByPhysicsProxy.FindOrAdd(PhysicsProxy).Add(FEventManager::EncodeCollisionIndex(NewIdx, false));
+								AllCollisionsIndicesByPhysicsProxy.FindOrAdd(PhysicsProxy).Add(TEventManager<Traits>::EncodeCollisionIndex(NewIdx, false));
 
 								if (OtherPhysicsProxy && OtherPhysicsProxy != PhysicsProxy)
 								{
-									AllCollisionsIndicesByPhysicsProxy.FindOrAdd(OtherPhysicsProxy).Add(FEventManager::EncodeCollisionIndex(NewIdx, true));
+									AllCollisionsIndicesByPhysicsProxy.FindOrAdd(OtherPhysicsProxy).Add(TEventManager<Traits>::EncodeCollisionIndex(NewIdx, true));
 								}
 							}
 						}
@@ -297,11 +299,11 @@ namespace Chaos
 #endif
 
 								// Add to AllCollisionsIndicesByPhysicsProxy
-								AllCollisionsIndicesByPhysicsProxy.FindOrAdd(PhysicsProxy).Add(FEventManager::EncodeCollisionIndex(NewIdx, false));
+								AllCollisionsIndicesByPhysicsProxy.FindOrAdd(PhysicsProxy).Add(TEventManager<Traits>::EncodeCollisionIndex(NewIdx, false));
 
 								if (OtherPhysicsProxy && OtherPhysicsProxy != PhysicsProxy)
 								{
-									AllCollisionsIndicesByPhysicsProxy.FindOrAdd(OtherPhysicsProxy).Add(FEventManager::EncodeCollisionIndex(NewIdx, true));
+									AllCollisionsIndicesByPhysicsProxy.FindOrAdd(OtherPhysicsProxy).Add(TEventManager<Traits>::EncodeCollisionIndex(NewIdx, true));
 								}
 							}
 						}
@@ -311,10 +313,11 @@ namespace Chaos
 		});
 	}
 
-	void FEventDefaults::RegisterBreakingEvent(FEventManager& EventManager)
+	template <typename Traits>
+	void TEventDefaults<Traits>::RegisterBreakingEvent(TEventManager<Traits>& EventManager)
 	{
-		EventManager.RegisterEvent<FBreakingEventData>(EEventType::Breaking, []
-		(const Chaos::FPBDRigidsSolver* Solver, FBreakingEventData& BreakingEventData)
+		EventManager.template RegisterEvent<FBreakingEventData>(EEventType::Breaking, []
+		(const Chaos::TPBDRigidsSolver<Traits>* Solver, FBreakingEventData& BreakingEventData)
 		{
 			check(Solver);
 			SCOPE_CYCLE_COUNTER(STAT_GatherBreakingEvent);
@@ -329,7 +332,7 @@ namespace Chaos
 
 			BreakingEventData.BreakingData.TimeCreated = Solver->MTime;
 
-			const FPBDRigidsSolver::FPBDRigidsEvolution* Evolution = Solver->GetEvolution();
+			const auto* Evolution = Solver->GetEvolution();
 			const TPBDRigidParticles<float, 3>& Particles = Evolution->GetParticles().GetDynamicParticles();
 			const TArray<TBreakingData<float, 3>>& AllBreakingsArray = Evolution->GetRigidClustering().GetAllClusterBreakings();
 			const TArrayCollectionArray<ClusterId>& ClusterIdsArray = Evolution->GetRigidClustering().GetClusterIdsArray();
@@ -381,7 +384,7 @@ namespace Chaos
 									BreakingDataArrayItem.ParticleIndexMesh = ParticleIndexMesh;
 								}
 #endif
-							}
+								}
 						}
 					}
 				}
@@ -389,10 +392,11 @@ namespace Chaos
 		});
 	}
 
-	void FEventDefaults::RegisterTrailingEvent(FEventManager& EventManager)
+	template <typename Traits>
+	void TEventDefaults<Traits>::RegisterTrailingEvent(TEventManager<Traits>& EventManager)
 	{
-		EventManager.RegisterEvent<FTrailingEventData>(EEventType::Trailing, []
-		(const Chaos::FPBDRigidsSolver* Solver, FTrailingEventData& TrailingEventData)
+		EventManager.template RegisterEvent<FTrailingEventData>(EEventType::Trailing, []
+		(const Chaos::TPBDRigidsSolver<Traits>* Solver, FTrailingEventData& TrailingEventData)
 		{
 			check(Solver);
 
@@ -400,7 +404,7 @@ namespace Chaos
 			if (!Solver->GetEventFilters()->IsTrailingEventEnabled())
 				return;
 
-			const FPBDRigidsSolver::FPBDRigidsEvolution* Evolution = Solver->GetEvolution();
+			const auto* Evolution = Solver->GetEvolution();
 
 			const TArrayCollectionArray<ClusterId>& ClusterIdsArray = Evolution->GetRigidClustering().GetClusterIdsArray();
 #if TODO_REIMPLEMENT_RIGID_CLUSTERING
@@ -461,20 +465,21 @@ namespace Chaos
 		});
 	}
 
-	void FEventDefaults::RegisterSleepingEvent(FEventManager& EventManager)
+	template <typename Traits>
+	void TEventDefaults<Traits>::RegisterSleepingEvent(TEventManager<Traits>& EventManager)
 	{
-		EventManager.RegisterEvent<FSleepingEventData>(EEventType::Sleeping, []
-		(const Chaos::FPBDRigidsSolver* Solver, FSleepingEventData& SleepingEventData)
+		EventManager.template RegisterEvent<FSleepingEventData>(EEventType::Sleeping, []
+		(const Chaos::TPBDRigidsSolver<Traits>* Solver, FSleepingEventData& SleepingEventData)
 		{
 			check(Solver);
 			SCOPE_CYCLE_COUNTER(STAT_GatherSleepingEvent);
 
-			const FPBDRigidsSolver::FPBDRigidsEvolution* Evolution = Solver->GetEvolution();
+			const auto* Evolution = Solver->GetEvolution();
 
 			FSleepingDataArray& EventSleepDataArray = SleepingEventData.SleepingData;
 			EventSleepDataArray.Reset();
 
-			Chaos::FPBDRigidsSolver* NonConstSolver = (Chaos::FPBDRigidsSolver*)(Solver);
+			Chaos::TPBDRigidsSolver<Traits>* NonConstSolver = const_cast<Chaos::TPBDRigidsSolver<Traits>*>(Solver);
 
 			NonConstSolver->Particles.GetDynamicParticles().GetSleepDataLock().ReadLock();
 			auto& SolverSleepingData = NonConstSolver->Particles.GetDynamicParticles().GetSleepData();
@@ -486,13 +491,13 @@ namespace Chaos
 					{
 						for (IPhysicsProxyBase* Proxy : *Proxies)
 						{
-							TGeometryParticle<float, 3>* Particle = SleepData.Particle->GTGeometryParticle();
+					TGeometryParticle<float, 3>* Particle = SleepData.Particle->GTGeometryParticle();
 							if (Particle != nullptr && Proxy != nullptr)
-							{
-								int32 NewIdx = EventSleepDataArray.Add(TSleepingData<float, 3>());
-								TSleepingData<float, 3>& SleepingDataArrayItem = EventSleepDataArray[NewIdx];
-								SleepingDataArrayItem.Particle = Particle;
-								SleepingDataArrayItem.Sleeping = SleepData.Sleeping;
+					{
+						int32 NewIdx = EventSleepDataArray.Add(TSleepingData<float, 3>());
+						TSleepingData<float, 3>& SleepingDataArrayItem = EventSleepDataArray[NewIdx];
+						SleepingDataArrayItem.Particle = Particle;
+						SleepingDataArrayItem.Sleeping = SleepData.Sleeping;
 							}
 						}
 					}
@@ -505,4 +510,8 @@ namespace Chaos
 
 		});
 	}
-}
+
+#define EVOLUTION_TRAIT(Trait) template class TEventDefaults<Trait>;
+#include "Chaos/EvolutionTraits.inl"
+#undef EVOLUTION_TRAIT
+}
\ No newline at end of file
diff --git a/Engine/Source/Runtime/Experimental/ChaosSolvers/Private/EventManager.cpp b/Engine/Source/Runtime/Experimental/ChaosSolvers/Private/EventManager.cpp
--- a/Engine/Source/Runtime/Experimental/ChaosSolvers/Private/EventManager.cpp
+++ b/Engine/Source/Runtime/Experimental/ChaosSolvers/Private/EventManager.cpp
@@ -4,10 +4,11 @@
 
 namespace Chaos
 {
-	void FEventManager::Reset()
+	template <typename Traits>
+	void TEventManager<Traits>::Reset()
 	{
 		ContainerLock.WriteLock();
-		for (FEventContainerPtr Container : EventContainers)
+		for (TEventContainerBasePtr<Traits> Container : EventContainers)
 		{
 			delete Container;
 			Container = nullptr;
@@ -16,7 +17,8 @@ namespace Chaos
 		ContainerLock.WriteUnlock();
 	}
 
-	void FEventManager::UnregisterEvent(const FEventID& EventID)
+	template <typename Traits>
+	void TEventManager<Traits>::UnregisterEvent(const FEventID& EventID)
 	{
 		ContainerLock.WriteLock();
 		if (EventID < EventContainers.Num())
@@ -27,7 +29,8 @@ namespace Chaos
 		ContainerLock.WriteUnlock();
 	}
 
-	void FEventManager::UnregisterHandler(const FEventID& EventID, const void* InHandler)
+	template <typename Traits>
+	void TEventManager<Traits>::UnregisterHandler(const FEventID& EventID, const void* InHandler)
 	{
 		ContainerLock.WriteLock();
 		checkf(EventID < EventContainers.Num(), TEXT("Unregistering event Handler for an event ID that does not exist"));
@@ -35,10 +38,11 @@ namespace Chaos
 		ContainerLock.WriteUnlock();
 	}
 
-	void FEventManager::FillProducerData(const Chaos::FPBDRigidsSolver* Solver)
+	template <typename Traits>
+	void TEventManager<Traits>::FillProducerData(const Chaos::TPBDRigidsSolver<Traits>* Solver)
 	{
 		ContainerLock.ReadLock();
-		for (FEventContainerPtr EventContainer : EventContainers)
+		for (TEventContainerBasePtr<Traits> EventContainer : EventContainers)
 		{
 			if (EventContainer)
 			{
@@ -48,7 +52,8 @@ namespace Chaos
 		ContainerLock.ReadUnlock();
 	}
 
-	void FEventManager::FlipBuffersIfRequired()
+	template <typename Traits>
+	void TEventManager<Traits>::FlipBuffersIfRequired()
 	{
 		if (BufferMode == EMultiBufferMode::Double)
 		{
@@ -56,7 +61,7 @@ namespace Chaos
 		}
 
 		ContainerLock.ReadLock();
-		for (FEventContainerPtr EventContainer : EventContainers)
+		for (TEventContainerBasePtr<Traits> EventContainer : EventContainers)
 		{
 			if (EventContainer)
 			{
@@ -71,7 +76,8 @@ namespace Chaos
 		}
 	}
 
-	void FEventManager::DispatchEvents()
+	template <typename Traits>
+	void TEventManager<Traits>::DispatchEvents()
 	{
 		if (BufferMode == EMultiBufferMode::Double)
 		{
@@ -79,7 +85,7 @@ namespace Chaos
 		}
 
 		ContainerLock.ReadLock();
-		for (FEventContainerPtr EventContainer : EventContainers)
+		for (TEventContainerBasePtr<Traits> EventContainer : EventContainers)
 		{
 			if (EventContainer)
 			{
@@ -95,7 +101,8 @@ namespace Chaos
 
 	}
 
-	void FEventManager::InternalRegisterInjector(const FEventID& EventID, const FEventContainerPtr& Container)
+	template <typename Traits>
+	void TEventManager<Traits>::InternalRegisterInjector(const FEventID& EventID, const TEventContainerBasePtr<Traits>& Container)
 	{
 		if (EventID > EventContainers.Num())
 		{
@@ -108,16 +115,20 @@ namespace Chaos
 		EventContainers.EmplaceAt(EventID, Container);
 	}
 
-	int32 FEventManager::EncodeCollisionIndex(int32 ActualCollisionIndex, bool bSwapOrder)
+	template <typename Traits>
+	int32 TEventManager<Traits>::EncodeCollisionIndex(int32 ActualCollisionIndex, bool bSwapOrder)
 	{
 		return bSwapOrder ? (ActualCollisionIndex | (1 << 31)) : ActualCollisionIndex;
 	}
 
-	int32 FEventManager::DecodeCollisionIndex(int32 EncodedCollisionIdx, bool& bSwapOrder)
+	template <typename Traits>
+	int32 TEventManager<Traits>::DecodeCollisionIndex(int32 EncodedCollisionIdx, bool& bSwapOrder)
 	{
 		bSwapOrder = EncodedCollisionIdx & (1 << 31);
 		return EncodedCollisionIdx & ~(1 << 31);
 	}
 
-
+#define EVOLUTION_TRAIT(Trait) template class CHAOSSOLVERS_API TEventManager<Trait>;
+#include "Chaos/EvolutionTraits.inl"
+#undef EVOLUTION_TRAIT
 }
diff --git a/Engine/Source/Runtime/Experimental/ChaosSolvers/Private/Framework/Dispatcher.cpp b/Engine/Source/Runtime/Experimental/ChaosSolvers/Private/Framework/Dispatcher.cpp
--- a/Engine/Source/Runtime/Experimental/ChaosSolvers/Private/Framework/Dispatcher.cpp
+++ b/Engine/Source/Runtime/Experimental/ChaosSolvers/Private/Framework/Dispatcher.cpp
@@ -12,7 +12,7 @@
 namespace Chaos
 {
 	template<>
-	void FDispatcher<EThreadingMode::DedicatedThread>::EnqueueCommandImmediate(FPhysicsSolver* InSolver, FSolverCommand InCommand)
+	void FDispatcher<EThreadingMode::DedicatedThread>::EnqueueCommandImmediate(FPhysicsSolverBase* InSolver, FSolverCommand InCommand)
 	{
 		check(InSolver && InCommand);
 		InSolver->GetCommandQueue().Enqueue(InCommand);
@@ -89,10 +89,10 @@ namespace Chaos
 	//////////////////////////////////////////////////////////////////////////
 
 	template<>
-	void FDispatcher<EThreadingMode::SingleThread>::EnqueueCommandImmediate(FPhysicsSolver* InSolver, FSolverCommand InCommand)
+	void FDispatcher<EThreadingMode::SingleThread>::EnqueueCommandImmediate(FPhysicsSolverBase* InSolver, FSolverCommand InCommand)
 	{
 		check(InSolver && InCommand);
-		InCommand(InSolver);
+		InCommand();
 	}
 
 	template<>
@@ -141,7 +141,7 @@ namespace Chaos
 	//////////////////////////////////////////////////////////////////////////
 
 	template<>
-	void FDispatcher<EThreadingMode::TaskGraph>::EnqueueCommandImmediate(FPhysicsSolver* InSolver, FSolverCommand InCommand)
+	void FDispatcher<EThreadingMode::TaskGraph>::EnqueueCommandImmediate(FPhysicsSolverBase* InSolver, FSolverCommand InCommand)
 	{
 		check(InSolver && InCommand);
 		InSolver->GetCommandQueue().Enqueue(InCommand);
diff --git a/Engine/Source/Runtime/Experimental/ChaosSolvers/Private/Framework/DispatcherImpl.h b/Engine/Source/Runtime/Experimental/ChaosSolvers/Private/Framework/DispatcherImpl.h
--- a/Engine/Source/Runtime/Experimental/ChaosSolvers/Private/Framework/DispatcherImpl.h
+++ b/Engine/Source/Runtime/Experimental/ChaosSolvers/Private/Framework/DispatcherImpl.h
@@ -41,7 +41,7 @@ namespace Chaos
 
 		virtual void EnqueueCommandImmediate(FGlobalCommand InCommand) final override;
 		virtual void EnqueueCommandImmediate(FTaskCommand InCommand) final override;
-		virtual void EnqueueCommandImmediate(FPhysicsSolver* InSolver, FSolverCommand InCommand) final override;
+		virtual void EnqueueCommandImmediate(FPhysicsSolverBase* InSolver, FSolverCommand InCommand) final override;
 
 		virtual EThreadingMode GetMode() const final override { return Mode; }
 
@@ -63,7 +63,7 @@ namespace Chaos
 	////////////////////////////////////////////////////////////////////////////
 
 	template<>
-	void FDispatcher<EThreadingMode::DedicatedThread>::EnqueueCommandImmediate(FPhysicsSolver* InSolver, FSolverCommand InCommand);
+	void FDispatcher<EThreadingMode::DedicatedThread>::EnqueueCommandImmediate(FPhysicsSolverBase* InSolver, FSolverCommand InCommand);
 
 	template<>
 	void FDispatcher<EThreadingMode::DedicatedThread>::EnqueueCommandImmediate(FTaskCommand InCommand);
@@ -80,7 +80,7 @@ namespace Chaos
 	////////////////////////////////////////////////////////////////////////////
 
 	template<>
-	void FDispatcher<EThreadingMode::SingleThread>::EnqueueCommandImmediate(FPhysicsSolver* InSolver, FSolverCommand InCommand);
+	void FDispatcher<EThreadingMode::SingleThread>::EnqueueCommandImmediate(FPhysicsSolverBase* InSolver, FSolverCommand InCommand);
 
 	template<>
 	void FDispatcher<EThreadingMode::SingleThread>::EnqueueCommandImmediate(FTaskCommand InCommand);
@@ -97,7 +97,7 @@ namespace Chaos
 	////////////////////////////////////////////////////////////////////////////
 
 	template<>
-	void FDispatcher<EThreadingMode::TaskGraph>::EnqueueCommandImmediate(FPhysicsSolver* InSolver, FSolverCommand InCommand);
+	void FDispatcher<EThreadingMode::TaskGraph>::EnqueueCommandImmediate(FPhysicsSolverBase* InSolver, FSolverCommand InCommand);
 
 	template<>
 	void FDispatcher<EThreadingMode::TaskGraph>::EnqueueCommandImmediate(FTaskCommand InCommand);
diff --git a/Engine/Source/Runtime/Experimental/ChaosSolvers/Private/Framework/PersistentTask.cpp b/Engine/Source/Runtime/Experimental/ChaosSolvers/Private/Framework/PersistentTask.cpp
--- a/Engine/Source/Runtime/Experimental/ChaosSolvers/Private/Framework/PersistentTask.cpp
+++ b/Engine/Source/Runtime/Experimental/ChaosSolvers/Private/Framework/PersistentTask.cpp
@@ -383,11 +383,11 @@ namespace Chaos
 		SCOPE_CYCLE_COUNTER(STAT_HandleSolverCommands);
 
 		check(InSolver);
-		TQueue<TFunction<void(FPhysicsSolver*)>, EQueueMode::Mpsc>& Queue = InSolver->CommandQueue;
-		TFunction<void(FPhysicsSolver*)> Command;
+		TQueue<TFunction<void()>, EQueueMode::Mpsc>& Queue = InSolver->CommandQueue;
+		TFunction<void()> Command;
 		while(Queue.Dequeue(Command))
 		{
-			Command(InSolver);
+			Command();
 		}
 	}
 
diff --git a/Engine/Source/Runtime/Experimental/ChaosSolvers/Private/Framework/PhysicsTickTask.cpp b/Engine/Source/Runtime/Experimental/ChaosSolvers/Private/Framework/PhysicsTickTask.cpp
--- a/Engine/Source/Runtime/Experimental/ChaosSolvers/Private/Framework/PhysicsTickTask.cpp
+++ b/Engine/Source/Runtime/Experimental/ChaosSolvers/Private/Framework/PhysicsTickTask.cpp
@@ -201,11 +201,11 @@ void FPhysicsSolverAdvanceTask::StepSolver(Chaos::FPhysicsSolver* InSolver, floa
 	{
 		SCOPE_CYCLE_COUNTER(STAT_HandleSolverCommands);
 
-		TQueue<TFunction<void(FPhysicsSolver*)>, EQueueMode::Mpsc>& Queue = InSolver->CommandQueue;
-		TFunction<void(FPhysicsSolver*)> Command;
+		TQueue<TFunction<void()>, EQueueMode::Mpsc>& Queue = InSolver->CommandQueue;
+		TFunction<void()> Command;
 		while(Queue.Dequeue(Command))
 		{
-			Command(InSolver);
+			Command();
 		}
 	}
 
diff --git a/Engine/Source/Runtime/Experimental/ChaosSolvers/Private/PBDRigidsSolver.cpp b/Engine/Source/Runtime/Experimental/ChaosSolvers/Private/PBDRigidsSolver.cpp
--- a/Engine/Source/Runtime/Experimental/ChaosSolvers/Private/PBDRigidsSolver.cpp
+++ b/Engine/Source/Runtime/Experimental/ChaosSolvers/Private/PBDRigidsSolver.cpp
@@ -57,12 +57,13 @@ FAutoConsoleVariableRef CVarChaosSolverCollisionUseManifolds(TEXT("p.Chaos.Solve
 namespace Chaos
 {
 
+	template <typename Traits>
 	class AdvanceOneTimeStepTask : public FNonAbandonableTask
 	{
 		friend class FAutoDeleteAsyncTask<AdvanceOneTimeStepTask>;
 	public:
 		AdvanceOneTimeStepTask(
-			FPBDRigidsSolver* Scene,
+			TPBDRigidsSolver<Traits>* Scene,
 			const float DeltaTime)
 			: MSolver(Scene)
 			, MDeltaTime(DeltaTime)
@@ -93,7 +94,7 @@ namespace Chaos
 						PositionTarget, PositionTargetedParticles, /*AnimatedPositions,*/ MSolver->GetSolverTime());
 				}
 
-				for (FGeometryCollectionPhysicsProxy* Obj : MSolver->GetGeometryCollectionPhysicsProxies())
+				for (TGeometryCollectionPhysicsProxy<Traits>* Obj : MSolver->GetGeometryCollectionPhysicsProxies())
 				{
 					Obj->ParameterUpdateCallback(MSolver->GetEvolution()->GetParticles().GetGeometryCollectionParticles(), MSolver->GetSolverTime());
 				}
@@ -141,7 +142,7 @@ namespace Chaos
 						Obj->FieldForcesUpdateCallback(MSolver, ClusteredParticles, Forces, Torques, MSolver->GetSolverTime());
 					}
 
-					for (FGeometryCollectionPhysicsProxy* Obj : MSolver->GetGeometryCollectionPhysicsProxies())
+					for (TGeometryCollectionPhysicsProxy<Traits>* Obj : MSolver->GetGeometryCollectionPhysicsProxies())
 					{
 						Obj->ParameterUpdateCallback(MSolver->GetEvolution()->GetParticles().GetGeometryCollectionParticles(), MSolver->GetSolverTime());
 					}
@@ -194,13 +195,14 @@ namespace Chaos
 			RETURN_QUICK_DECLARE_CYCLE_STAT(AdvanceOneTimeStepTask, STATGROUP_ThreadPoolAsyncTasks);
 		}
 
-		FPBDRigidsSolver* MSolver;
+		TPBDRigidsSolver<Traits>* MSolver;
 		float MDeltaTime;
 		TSharedPtr<FCriticalSection> PrevLock, CurrentLock;
 		TSharedPtr<FEvent> PrevEvent, CurrentEvent;
 	};
 
-	FPBDRigidsSolver::FPBDRigidsSolver(const EMultiBufferMode BufferingModeIn, UObject* InOwner)
+	template <typename Traits>
+	TPBDRigidsSolver<Traits>::TPBDRigidsSolver(const EMultiBufferMode BufferingModeIn, UObject* InOwner)
 		: Super(BufferingModeIn, InOwner)
 		, CurrentFrame(0)
 		, MTime(0.0)
@@ -213,7 +215,7 @@ namespace Chaos
 		, bIsFloorAnalytic(false)
 		, FloorHeight(0.f)
 		, MEvolution(new FPBDRigidsEvolution(Particles, SimMaterials, ChaosSolverCollisionDefaultIterationsCVar, ChaosSolverCollisionDefaultPushoutIterationsCVar, BufferingModeIn == Chaos::EMultiBufferMode::Single))
-		, MEventManager(new FEventManager(BufferingModeIn))
+		, MEventManager(new TEventManager<Traits>(BufferingModeIn))
 		, MSolverEventFilters(new FSolverEventFilters())
 		, MActiveParticlesBuffer(new FActiveParticlesBuffer(BufferingModeIn, BufferingModeIn == Chaos::EMultiBufferMode::Single))
 		, MCurrentLock(new FCriticalSection())
@@ -233,14 +235,15 @@ namespace Chaos
 			});
 	}
 
-	FPBDRigidsSolver::~FPBDRigidsSolver()
+	template <typename Traits>
+	TPBDRigidsSolver<Traits>::~TPBDRigidsSolver()
 	{
 		if(ChaosSolverCleanupCommandsOnDestruction != 0)
 		{
-			TFunction<void(FPhysicsSolver*)> Command;
+			TFunction<void()> Command;
 			while(CommandQueue.Dequeue(Command))
 			{
-				Command(this);
+				Command();
 			}
 		}
 	}
@@ -253,11 +256,12 @@ namespace Chaos
 		TEXT(""),
 		ECVF_Default);
 
-	void FPBDRigidsSolver::RegisterObject(TGeometryParticle<float, 3>* GTParticle)
+	template <typename Traits>
+	void TPBDRigidsSolver<Traits>::RegisterObject(TGeometryParticle<float, 3>* GTParticle)
 	{
 		LLM_SCOPE(ELLMTag::Chaos);
 
-		UE_LOG(LogPBDRigidsSolver, Verbose, TEXT("FPBDRigidsSolver::RegisterObject()"));
+		UE_LOG(LogPBDRigidsSolver, Verbose, TEXT("TPBDRigidsSolver::RegisterObject()"));
 
 		// Make sure this particle doesn't already have a proxy
 		checkSlow(GTParticle->GetProxy() == nullptr);
@@ -335,9 +339,10 @@ namespace Chaos
 		AddDirtyProxy(ProxyBase);
 	}
 
-	void FPBDRigidsSolver::UnregisterObject(TGeometryParticle<float, 3>* GTParticle)
+	template <typename Traits>
+	void TPBDRigidsSolver<Traits>::UnregisterObject(TGeometryParticle<float, 3>* GTParticle)
 	{
-		UE_LOG(LogPBDRigidsSolver, Verbose, TEXT("FPBDRigidsSolver::UnregisterObject()"));
+		UE_LOG(LogPBDRigidsSolver, Verbose, TEXT("TPBDRigidsSolver::UnregisterObject()"));
 
 		// Get the proxy associated with this particle
 		IPhysicsProxyBase* InProxy = GTParticle->GetProxy();
@@ -371,12 +376,12 @@ namespace Chaos
 		}
 
 		// Enqueue a command to remove the particle and delete the proxy
-		FChaosSolversModule::GetModule()->GetDispatcher()->EnqueueCommandImmediate(this, [InProxy, InParticleType](FPBDRigidsSolver* Solver)
+		EnqueueCommandImmediate([InProxy, InParticleType, this]()
 		{
-			UE_LOG(LogPBDRigidsSolver, Verbose, TEXT("FPBDRigidsSolver::UnregisterObject() ~ Dequeue"));
+			UE_LOG(LogPBDRigidsSolver, Verbose, TEXT("TPBDRigidsSolver::UnregisterObject() ~ Dequeue"));
 
 				// Generally need to remove stale events for particles that no longer exist
-				Solver->GetEventManager()->ClearEvents<FCollisionEventData>(EEventType::Collision, [InProxy]
+				Solver->GetEventManager()->template ClearEvents<FCollisionEventData>(EEventType::Collision, [InProxy]
 				(FCollisionEventData& EventDataInOut)
 				{
 					Chaos::FCollisionDataArray const& CollisionData = EventDataInOut.CollisionData.AllCollisionsArray;
@@ -389,7 +394,7 @@ namespace Chaos
 							for (int32 EncodedCollisionIdx : *CollisionIndices)
 							{
 								bool bSwapOrder;
-								int32 CollisionIdx = Chaos::FEventManager::DecodeCollisionIndex(EncodedCollisionIdx, bSwapOrder);
+								int32 CollisionIdx = Chaos::TEventManager<Traits>::DecodeCollisionIndex(EncodedCollisionIdx, bSwapOrder);
 
 								// invalidate but don't delete from array, as this would mean we'd need to reindex PhysicsProxyToIndicesMap to maintain the other collisions lookup
 								Chaos::TCollisionData<float, 3>& CollisionDataItem = EventDataInOut.CollisionData.AllCollisionsArray[CollisionIdx];
@@ -434,28 +439,29 @@ namespace Chaos
 			if(Handle)
 			{
 				// Remove from rewind data
-				if(FRewindData* RewindData = Solver->GetRewindData())
+				if(FRewindData* RewindData = GetRewindData())
 				{
 					RewindData->RemoveParticle(Handle->UniqueIdx());
 				}
 
-			  // Remove game thread particle from ActiveGameThreadParticles so we won't crash when pulling physics state
-			  // if this particle was deleted after buffering results. 
-			  Solver->GetActiveParticlesBuffer()->RemoveActiveParticleFromConsumerBuffer(Handle->GTGeometryParticle());
+				// Remove game thread particle from ActiveGameThreadParticles so we won't crash when pulling physics state
+				// if this particle was deleted after buffering results. 
+				GetActiveParticlesBuffer()->RemoveActiveParticleFromConsumerBuffer(Handle->GTGeometryParticle());
   
-			  Solver->MParticleToProxy.Remove(Handle);
+				MParticleToProxy.Remove(Handle);
   
-			  // Use the handle to destroy the particle data
-			  Solver->GetEvolution()->DestroyParticle(Handle);
+				// Use the handle to destroy the particle data
+				GetEvolution()->DestroyParticle(Handle);
 			}
 
 		});
 
 	}
 
-	void FPBDRigidsSolver::RegisterObject(FGeometryCollectionPhysicsProxy* InProxy)
+	template <typename Traits>
+	void TPBDRigidsSolver<Traits>::RegisterObject(TGeometryCollectionPhysicsProxy<Traits>* InProxy)
 	{
-		UE_LOG(LogPBDRigidsSolver, Verbose, TEXT("FPBDRigidsSolver::RegisterObject(FGeometryCollectionPhysicsProxy*)"));
+		UE_LOG(LogPBDRigidsSolver, Verbose, TEXT("TPBDRigidsSolver::RegisterObject(TGeometryCollectionPhysicsProxy*)"));
 		GeometryCollectionPhysicsProxies.AddUnique(InProxy);
 		InProxy->SetSolver(this);
 		InProxy->Initialize();
@@ -463,26 +469,27 @@ namespace Chaos
 		FParticlesType* InParticles = &GetParticles();
 
 		// Finish registration on the physics thread...
-		FChaosSolversModule::GetModule()->GetDispatcher()->EnqueueCommandImmediate(this,
-			[InParticles, InProxy](FPBDRigidsSolver* Solver)
+		EnqueueCommandImmediate([InParticles, InProxy, this]()
 		{
 			UE_LOG(LogPBDRigidsSolver, Verbose, 
-				TEXT("FPBDRigidsSolver::RegisterObject(FGeometryCollectionPhysicsProxy*)"));
+				TEXT("TPBDRigidsSolver::RegisterObject(TGeometryCollectionPhysicsProxy*)"));
 			check(InParticles);
-			InProxy->InitializeBodiesPT(Solver, *InParticles);
+			InProxy->InitializeBodiesPT(this, *InParticles);
 		});
 	}
 
-	bool FPBDRigidsSolver::UnregisterObject(FGeometryCollectionPhysicsProxy* InProxy)
+	template <typename Traits>
+	bool TPBDRigidsSolver<Traits>::UnregisterObject(TGeometryCollectionPhysicsProxy<Traits>* InProxy)
 	{
 		InProxy->OnRemoveFromSolver(this);
-		InProxy->SetSolver(static_cast<FPBDRigidsSolver*>(nullptr));
+		InProxy->SetSolver(static_cast<TPBDRigidsSolver<Traits>*>(nullptr));
 		return GeometryCollectionPhysicsProxies.Remove(InProxy) != 0;
 	}
 
-	void FPBDRigidsSolver::RegisterObject(FFieldSystemPhysicsProxy* InProxy)
+	template <typename Traits>
+	void TPBDRigidsSolver<Traits>::RegisterObject(FFieldSystemPhysicsProxy* InProxy)
 	{
-		UE_LOG(LogPBDRigidsSolver, Verbose, TEXT("FPBDRigidsSolver::RegisterObject(FFieldSystemPhysicsProxy*)"));
+		UE_LOG(LogPBDRigidsSolver, Verbose, TEXT("TBDRigidsSolver::RegisterObject(FFieldSystemPhysicsProxy*)"));
 		FieldSystemPhysicsProxies.AddUnique(InProxy);
 		InProxy->SetSolver(this);
 		InProxy->Initialize();
@@ -490,22 +497,24 @@ namespace Chaos
 
 		FChaosSolversModule::GetModule()->GetDispatcher()->EnqueueCommandImmediate(
 			this,
-			[InProxy, ProxyData](FPBDRigidsSolver* Solver)
+			[InProxy, ProxyData]()
 			{
 				UE_LOG(LogPBDRigidsSolver, Verbose,
-					TEXT("FPBDRigidsSolver::RegisterObject(FFieldSystemPhysicsProxy*)"));
+					TEXT("TPBDRigidsSolver::RegisterObject(FFieldSystemPhysicsProxy*)"));
 				//InProxy->ActivateBodies();
 				InProxy->PushToPhysicsState(ProxyData);
 			});
 	}
 
-	bool FPBDRigidsSolver::UnregisterObject(FFieldSystemPhysicsProxy* InProxy)
+	template <typename Traits>
+	bool TPBDRigidsSolver<Traits>::UnregisterObject(FFieldSystemPhysicsProxy* InProxy)
 	{
-		InProxy->SetSolver(static_cast<FPBDRigidsSolver*>(nullptr));
+		InProxy->SetSolver(static_cast<TPBDRigidsSolver<Traits>*>(nullptr));
 		return FieldSystemPhysicsProxies.Remove(InProxy) != 0;
 	}
 
-	bool FPBDRigidsSolver::IsSimulating() const
+	template <typename Traits>
+	bool TPBDRigidsSolver<Traits>::IsSimulating() const
 	{
 		for (FGeometryParticlePhysicsProxy* Obj : GeometryParticlePhysicsProxies)
 			if (Obj->IsSimulating())
@@ -522,7 +531,7 @@ namespace Chaos
 		for (FStaticMeshPhysicsProxy* Obj : StaticMeshPhysicsProxies)
 			if (Obj->IsSimulating())
 				return true;
-		for (FGeometryCollectionPhysicsProxy* Obj : GeometryCollectionPhysicsProxies)
+		for (TGeometryCollectionPhysicsProxy<Traits>* Obj : GeometryCollectionPhysicsProxies)
 			if (Obj->IsSimulating())
 				return true;
 		for (FFieldSystemPhysicsProxy* Obj : FieldSystemPhysicsProxies)
@@ -534,7 +543,8 @@ namespace Chaos
 	int32 RewindCaptureNumFrames = -1;
 	FAutoConsoleVariableRef CVarRewindCaptureNumFrames(TEXT("p.RewindCaptureNumFrames"),RewindCaptureNumFrames,TEXT("The number of frames to capture rewind for. Requires restart of solver"));
 	
-	void FPBDRigidsSolver::Reset()
+	template <typename Traits>
+	void TPBDRigidsSolver<Traits>::Reset()
 	{
 		UE_LOG(LogPBDRigidsSolver, Verbose, TEXT("PBDRigidsSolver::Reset()"));
 
@@ -559,16 +569,18 @@ namespace Chaos
 			FinalizeRewindData(ActiveParticles);
 		});
 
-		FEventDefaults::RegisterSystemEvents(*GetEventManager());
+		TEventDefaults<Traits>::RegisterSystemEvents(*GetEventManager());
 	}
 
-	void FPBDRigidsSolver::ChangeBufferMode(Chaos::EMultiBufferMode InBufferMode)
+	template <typename Traits>
+	void TPBDRigidsSolver<Traits>::ChangeBufferMode(Chaos::EMultiBufferMode InBufferMode)
 	{
 		// This seems unused inside the solver? #BH
 		BufferMode = InBufferMode;
 	}
 
-	void FPBDRigidsSolver::AdvanceSolverBy(float DeltaTime)
+	template <typename Traits>
+	void TPBDRigidsSolver<Traits>::AdvanceSolverBy(float DeltaTime)
 	{
 		MEvolution->GetCollisionDetector().GetNarrowPhase().GetContext().bDeferUpdate = (ChaosSolverCollisionDeferNarrowPhase != 0);
 		MEvolution->GetCollisionDetector().GetNarrowPhase().GetContext().bAllowManifolds = (ChaosSolverCollisionUseManifolds != 0);
@@ -577,11 +589,12 @@ namespace Chaos
 		if (bEnabled)
 		{
 			MLastDt = DeltaTime;
-			AdvanceOneTimeStepTask(this, DeltaTime).DoWork();
+			AdvanceOneTimeStepTask<Traits>(this, DeltaTime).DoWork();
 		}
 	}
 
-	void FPBDRigidsSolver::SyncEvents_GameThread()
+	template <typename Traits>
+	void TPBDRigidsSolver<Traits>::SyncEvents_GameThread()
 	{
 		GetEventManager()->DispatchEvents();
 	}
@@ -674,104 +687,13 @@ namespace Chaos
 		Solver->RemoveDirtyProxy(Proxy);
 	}
 
-	void PushPhysicsStateExec(FPBDRigidsSolver* Solver, FGeometryCollectionPhysicsProxy* Proxy, Chaos::IDispatcher* Dispatcher)
+	template <typename Traits>
+	void TPBDRigidsSolver<Traits>::PushPhysicsStatePooled(IDispatcher* Dispatcher)
 	{
-		Proxy->NewData();
-
-		auto Cmd = [Proxy, Solver](Chaos::FPersistentPhysicsTask* PhysThread)
-		{
-			FPBDRigidsSolver::FPBDRigidsEvolution* Evolution = Solver->GetEvolution();
-			
-			TArray<FGeometryCollectionPhysicsProxy::FClusterHandle*>& Handles = Proxy->GetSolverParticleHandles();
-			for(FGeometryCollectionPhysicsProxy::FClusterHandle* Handle : Handles)
-			{
-				if(Handle && !Handle->Disabled())
-				{
-					Evolution->DirtyParticle(*Handle);
-				}
-			}
-
-			Proxy->PushToPhysicsState(nullptr);
-		};
-
-		if(Dispatcher)
-		{
-			Dispatcher->EnqueueCommandImmediate(Cmd);
-		}
-		else
-		{
-			Cmd(nullptr);
-		}
-
-		Proxy->ClearAccumulatedData();
-		Solver->RemoveDirtyProxy(Proxy);
-	}
-
-	void PushPhysicsStateExec(FPBDRigidsSolver* Solver, FFieldSystemPhysicsProxy* Proxy, Chaos::IDispatcher* Dispatcher)
-	{
-		Chaos::FParticleData* ProxyData = Proxy->NewData();
-
-		auto Cmd = [Proxy, Solver, ProxyData](Chaos::FPersistentPhysicsTask* PhysThread)
-		{
-			Proxy->PushToPhysicsState(nullptr);
-		};
-
-		if(Dispatcher)
-		{
-			Dispatcher->EnqueueCommandImmediate(Cmd);
-		}
-		else
-		{
-			Cmd(nullptr);
-		}
-
-		Proxy->ClearAccumulatedData();
-		Solver->RemoveDirtyProxy(Proxy);
-	}
-
-	void FPBDRigidsSolver::PushPhysicsStatePooled(IDispatcher* Dispatcher)
-	{
-#if 0
-		ensure(Dispatcher != nullptr);
-
-		// reset the per frame pool
-		RigidParticlePool.ResetPool();
-		KinematicGeometryParticlePool.ResetPool();
-		GeometryParticlePool.ResetPool();
-
-		TArray< IPhysicsProxyBase*> DirtyProxiesArray = DirtyProxiesSet.Array();
-		for (auto& Proxy : DirtyProxiesArray)
-		{
-			switch (Proxy->GetType())
-			{
-				//case EPhysicsProxyType::NoneType: // 0
-				//case EPhysicsProxyType::StaticMeshType: // 1
-			case EPhysicsProxyType::GeometryCollectionType: // 2
-				PushPhysicsStateExec(this, static_cast<FGeometryCollectionPhysicsProxy*>(Proxy), Dispatcher); // non pool api
-				break;
-			case EPhysicsProxyType::SingleRigidParticleType: // 7
-				PushPhysicsStateExec(this, static_cast<FSingleParticlePhysicsProxy< Chaos::TPBDRigidParticle<float, 3> >*>(Proxy), RigidParticlePool, Dispatcher);
-				break;
-				// case EPhysicsProxyType::FieldType: // 3
-				// case EPhysicsProxyType::SkeletalMeshType: // 4
-			case EPhysicsProxyType::SingleKinematicParticleType: // 6
-				PushPhysicsStateExec(this, static_cast<FSingleParticlePhysicsProxy< Chaos::TKinematicGeometryParticle<float, 3> >*>(Proxy), KinematicGeometryParticlePool, Dispatcher);
-				break;
-			case EPhysicsProxyType::SingleGeometryParticleType: // 5
-				PushPhysicsStateExec(this, static_cast<FSingleParticlePhysicsProxy< Chaos::TGeometryParticle<float, 3> >*>(Proxy), GeometryParticlePool, Dispatcher);
-				break;
-			default:
-				ensure("Unknown proxy type in physics solver.");
-			}
-		}
-
-		FlushExec(RigidParticlePool, Dispatcher, this);
-		FlushExec(KinematicGeometryParticlePool, Dispatcher, this);
-		FlushExec(GeometryParticlePool, Dispatcher, this);
-#endif
 	}
 
-	void FPBDRigidsSolver::PushPhysicsState(IDispatcher* Dispatcher)
+	template <typename Traits>
+	void TPBDRigidsSolver<Traits>::PushPhysicsState(IDispatcher* Dispatcher)
 	{
 		QUICK_SCOPE_CYCLE_COUNTER(STAT_PushPhysicsState);
 		FDirtySet* DirtyProxiesData = DirtyProxiesDataBuffer.AccessProducerBuffer();
@@ -821,10 +743,10 @@ namespace Chaos
 		DirtyProxiesDataBuffer.FlipProducer();
 		const bool bIsSingleThreaded = FChaosSolversModule::GetModule()->GetDispatcher()->GetMode() == EThreadingMode::SingleThread;
 
-		FChaosSolversModule::GetModule()->GetDispatcher()->EnqueueCommandImmediate(this,[bIsSingleThreaded, Manager,DirtyProxiesData,ShapeDirtyData](FPBDRigidsSolver* Solver)
+		EnqueueCommandImmediate([bIsSingleThreaded, Manager,DirtyProxiesData,ShapeDirtyData, this]()
 		{
-			FRewindData* RewindData = Solver->GetRewindData();
-			auto ProcessProxyPT = [bIsSingleThreaded, Solver,Manager,DirtyProxiesData,ShapeDirtyData, RewindData](auto& Proxy,int32 DataIdx,FDirtyProxy& Dirty,const auto& CreateHandleFunc)
+			FRewindData* RewindData = GetRewindData();
+			auto ProcessProxyPT = [bIsSingleThreaded, Manager,DirtyProxiesData,ShapeDirtyData, RewindData, this](auto& Proxy,int32 DataIdx,FDirtyProxy& Dirty,const auto& CreateHandleFunc)
 			{
 				const bool bIsNew = !Proxy->IsInitialized();
 				//single threaded version already created particle, but didn't initialize it
@@ -848,14 +770,14 @@ namespace Chaos
 					}
 				}
 
-				Proxy->PushToPhysicsState(*Manager, DataIdx, Dirty, ShapeDirtyData);
+				Proxy->PushToPhysicsState(*Manager, DataIdx, Dirty, ShapeDirtyData, *GetEvolution());
 
 				if(bIsNew)
 				{
 					auto Handle = Proxy->GetHandle();
 					Handle->GTGeometryParticle() = Proxy->GetParticle();
-					Solver->AddParticleToProxy(Handle,Proxy);
-					Solver->GetEvolution()->CreateParticle(Handle);
+					AddParticleToProxy(Handle,Proxy);;
+					GetEvolution()->CreateParticle(Handle);
 					Proxy->SetInitialized(true);
 				}
 
@@ -868,26 +790,26 @@ namespace Chaos
 			}
 
 			//need to create new particle handles
-			DirtyProxiesData->ForEachProxy([Solver, &ProcessProxyPT](int32 DataIdx,FDirtyProxy& Dirty)
+			DirtyProxiesData->ForEachProxy([this, &ProcessProxyPT](int32 DataIdx,FDirtyProxy& Dirty)
 			{
 				switch(Dirty.Proxy->GetType())
 				{
 				case EPhysicsProxyType::SingleRigidParticleType:
 				{
 					auto Proxy = static_cast<FRigidParticlePhysicsProxy*>(Dirty.Proxy);
-					ProcessProxyPT(Proxy,DataIdx,Dirty,[Solver](const FUniqueIdx* UniqueIdx){ return Solver->Particles.CreateDynamicParticles(1,UniqueIdx)[0]; });
+					ProcessProxyPT(Proxy,DataIdx,Dirty,[this](const FUniqueIdx* UniqueIdx){ return Particles.CreateDynamicParticles(1,UniqueIdx)[0]; });
 					break;
 				}
 				case EPhysicsProxyType::SingleKinematicParticleType:
 				{
 					auto Proxy = static_cast<FKinematicGeometryParticlePhysicsProxy*>(Dirty.Proxy);
-					ProcessProxyPT(Proxy,DataIdx,Dirty,[Solver](const FUniqueIdx* UniqueIdx){ return Solver->Particles.CreateKinematicParticles(1,UniqueIdx)[0]; });
+					ProcessProxyPT(Proxy,DataIdx,Dirty,[this](const FUniqueIdx* UniqueIdx){ return Particles.CreateKinematicParticles(1,UniqueIdx)[0]; });
 					break;
 				}
 				case EPhysicsProxyType::SingleGeometryParticleType:
 				{
 					auto Proxy = static_cast<FGeometryParticlePhysicsProxy*>(Dirty.Proxy);
-					ProcessProxyPT(Proxy,DataIdx,Dirty,[Solver](const FUniqueIdx* UniqueIdx){ return Solver->Particles.CreateStaticParticles(1,UniqueIdx)[0]; });
+					ProcessProxyPT(Proxy,DataIdx,Dirty,[this](const FUniqueIdx* UniqueIdx){ return Particles.CreateStaticParticles(1,UniqueIdx)[0]; });
 					break;
 				}
 				default:
@@ -899,10 +821,11 @@ namespace Chaos
 		});
 	}
 
-	void FPBDRigidsSolver::BufferPhysicsResults()
+	template <typename Traits>
+	void TPBDRigidsSolver<Traits>::BufferPhysicsResults()
 	{
 		//ensure(IsInPhysicsThread());
-		TArray<FGeometryCollectionPhysicsProxy*> ActiveGC;
+		TArray<TGeometryCollectionPhysicsProxy<Traits>*> ActiveGC;
 		ActiveGC.Reserve(GeometryCollectionPhysicsProxies.Num());
 
 		TParticleView<TPBDRigidParticles<float, 3>>& ActiveParticles = GetParticles().GetActiveParticlesView();
@@ -913,31 +836,31 @@ namespace Chaos
 				for (IPhysicsProxyBase* Proxy : *Proxies)
 				{
 					if (Proxy != nullptr)
-					{
-						switch (ActiveObject.GetParticleType())
-						{
-						case Chaos::EParticleType::Rigid:
-							((FRigidParticlePhysicsProxy*)(Proxy))->BufferPhysicsResults();
-							break;
-						case Chaos::EParticleType::Kinematic:
-							((FKinematicGeometryParticlePhysicsProxy*)(Proxy))->BufferPhysicsResults();
-							break;
-						case Chaos::EParticleType::Static:
-							((FGeometryParticlePhysicsProxy*)(Proxy))->BufferPhysicsResults();
-							break;
-						case Chaos::EParticleType::GeometryCollection:
-							ActiveGC.AddUnique((FGeometryCollectionPhysicsProxy*)(Proxy));
-							break;
-						case Chaos::EParticleType::Clustered:
-							ActiveGC.AddUnique((FGeometryCollectionPhysicsProxy*)(Proxy));
-							break;
-						default:
-							check(false);
-						}
-					}
+			{
+				switch (ActiveObject.GetParticleType())
+				{
+				case Chaos::EParticleType::Rigid:
+					((FRigidParticlePhysicsProxy*)(Proxy))->BufferPhysicsResults();
+					break;
+				case Chaos::EParticleType::Kinematic:
+					((FKinematicGeometryParticlePhysicsProxy*)(Proxy))->BufferPhysicsResults();
+					break;
+				case Chaos::EParticleType::Static:
+					((FGeometryParticlePhysicsProxy*)(Proxy))->BufferPhysicsResults();
+					break;
+				case Chaos::EParticleType::GeometryCollection:
+					ActiveGC.AddUnique((TGeometryCollectionPhysicsProxy<Traits>*)(Proxy));
+					break;
+				case Chaos::EParticleType::Clustered:
+					ActiveGC.AddUnique((TGeometryCollectionPhysicsProxy<Traits>*)(Proxy));
+					break;
+				default:
+					check(false);
 				}
 			}
 		}
+			}
+		}
 
 		for (auto* GCProxy : ActiveGC)
 		{
@@ -945,10 +868,11 @@ namespace Chaos
 		}
 	}
 
-	void FPBDRigidsSolver::FlipBuffers()
+	template <typename Traits>
+	void TPBDRigidsSolver<Traits>::FlipBuffers()
 	{
 		//ensure(IsInPhysicsThread());
-		TArray<FGeometryCollectionPhysicsProxy*> ActiveGC;
+		TArray<TGeometryCollectionPhysicsProxy<Traits>*> ActiveGC;
 		ActiveGC.Reserve(GeometryCollectionPhysicsProxies.Num());
 
 		TParticleView<TPBDRigidParticles<float, 3>>& ActiveParticles = GetParticles().GetActiveParticlesView();
@@ -959,31 +883,31 @@ namespace Chaos
 				for (IPhysicsProxyBase* Proxy : *Proxies)
 				{
 					if (Proxy != nullptr)
-					{
-						switch (ActiveObject.GetParticleType())
-						{
-						case Chaos::EParticleType::Rigid:
-							((FRigidParticlePhysicsProxy*)(Proxy))->FlipBuffer();
-							break;
-						case Chaos::EParticleType::Kinematic:
-							((FKinematicGeometryParticlePhysicsProxy*)(Proxy))->FlipBuffer();
-							break;
-						case Chaos::EParticleType::Static:
-							((FGeometryParticlePhysicsProxy*)(Proxy))->FlipBuffer();
-							break;
-						case Chaos::EParticleType::GeometryCollection:
-							ActiveGC.AddUnique((FGeometryCollectionPhysicsProxy*)(Proxy));
-							break;
-						case Chaos::EParticleType::Clustered:
-							ActiveGC.AddUnique((FGeometryCollectionPhysicsProxy*)(Proxy));
-							break;
-						default:
-							check(false);
-						}
-					}
+			{
+				switch (ActiveObject.GetParticleType())
+				{
+				case Chaos::EParticleType::Rigid:
+					((FRigidParticlePhysicsProxy*)(Proxy))->FlipBuffer();
+					break;
+				case Chaos::EParticleType::Kinematic:
+					((FKinematicGeometryParticlePhysicsProxy*)(Proxy))->FlipBuffer();
+					break;
+				case Chaos::EParticleType::Static:
+					((FGeometryParticlePhysicsProxy*)(Proxy))->FlipBuffer();
+					break;
+				case Chaos::EParticleType::GeometryCollection:
+					ActiveGC.AddUnique((TGeometryCollectionPhysicsProxy<Traits>*)(Proxy));
+					break;
+				case Chaos::EParticleType::Clustered:
+					ActiveGC.AddUnique((TGeometryCollectionPhysicsProxy<Traits>*)(Proxy));
+					break;
+				default:
+					check(false);
 				}
 			}
 		}
+			}
+		}
 
 		for (auto* GCProxy : ActiveGC)
 		{
@@ -996,10 +920,11 @@ namespace Chaos
 	// FPhysScene_ChaosInterface::SyncBodies() instead, and then immediately afterwards 
 	// calls FPBDRigidsSovler::SyncEvents_GameThread().  This function is used by tests,
 	// however.
-	void FPBDRigidsSolver::UpdateGameThreadStructures()
+	template <typename Traits>
+	void TPBDRigidsSolver<Traits>::UpdateGameThreadStructures()
 	{
 		//ensure(IsInGameThread());
-		TArray<FGeometryCollectionPhysicsProxy*> ActiveGC;
+		TArray<TGeometryCollectionPhysicsProxy<Traits>*> ActiveGC;
 		ActiveGC.Reserve(GeometryCollectionPhysicsProxies.Num());
 
 		TParticleView<TPBDRigidParticles<float, 3>>& ActiveParticles = GetParticles().GetActiveParticlesView();
@@ -1010,31 +935,31 @@ namespace Chaos
 				for (IPhysicsProxyBase* Proxy : *Proxies)
 				{
 					if (Proxy != nullptr)
-					{
-						switch (ActiveObject.GetParticleType())
-						{
-						case Chaos::EParticleType::Rigid:
-							((FRigidParticlePhysicsProxy*)(Proxy))->PullFromPhysicsState();
-							break;
-						case Chaos::EParticleType::Kinematic:
-							((FKinematicGeometryParticlePhysicsProxy*)(Proxy))->PullFromPhysicsState();
-							break;
-						case Chaos::EParticleType::Static:
-							((FGeometryParticlePhysicsProxy*)(Proxy))->PullFromPhysicsState();
-							break;
-						case Chaos::EParticleType::GeometryCollection:
-							ActiveGC.AddUnique((FGeometryCollectionPhysicsProxy*)(Proxy));
-							break;
-						case Chaos::EParticleType::Clustered:
-							ActiveGC.AddUnique((FGeometryCollectionPhysicsProxy*)(Proxy));
-							break;
-						default:
-							check(false);
-						}
-					}
+			{
+				switch (ActiveObject.GetParticleType())
+				{
+				case Chaos::EParticleType::Rigid:
+					((FRigidParticlePhysicsProxy*)(Proxy))->PullFromPhysicsState();
+					break;
+				case Chaos::EParticleType::Kinematic:
+					((FKinematicGeometryParticlePhysicsProxy*)(Proxy))->PullFromPhysicsState();
+					break;
+				case Chaos::EParticleType::Static:
+					((FGeometryParticlePhysicsProxy*)(Proxy))->PullFromPhysicsState();
+					break;
+				case Chaos::EParticleType::GeometryCollection:
+					ActiveGC.AddUnique((TGeometryCollectionPhysicsProxy<Traits>*)(Proxy));
+					break;
+				case Chaos::EParticleType::Clustered:
+					ActiveGC.AddUnique((TGeometryCollectionPhysicsProxy<Traits>*)(Proxy));
+					break;
+				default:
+					check(false);
 				}
 			}
 		}
+			}
+		}
 
 		for (auto* GCProxy : ActiveGC)
 		{
@@ -1042,8 +967,8 @@ namespace Chaos
 		}
 	}
 
-
-	void FPBDRigidsSolver::PostTickDebugDraw() const
+	template <typename Traits>
+	void TPBDRigidsSolver<Traits>::PostTickDebugDraw() const
 	{
 #if CHAOS_DEBUG_DRAW
 		if (ChaosSolverDrawCollisions == 1) {
@@ -1052,50 +977,58 @@ namespace Chaos
 #endif
 	}
 
-
-	void FPBDRigidsSolver::UpdateMaterial(Chaos::FMaterialHandle InHandle, const Chaos::FChaosPhysicsMaterial& InNewData)
+	template <typename Traits>
+	void TPBDRigidsSolver<Traits>::UpdateMaterial(Chaos::FMaterialHandle InHandle, const Chaos::FChaosPhysicsMaterial& InNewData)
 	{
 		*SimMaterials.Get(InHandle.InnerHandle) = InNewData;
 	}
 
-	void FPBDRigidsSolver::CreateMaterial(Chaos::FMaterialHandle InHandle, const Chaos::FChaosPhysicsMaterial& InNewData)
+	template <typename Traits>
+	void TPBDRigidsSolver<Traits>::CreateMaterial(Chaos::FMaterialHandle InHandle, const Chaos::FChaosPhysicsMaterial& InNewData)
 	{
 		ensure(SimMaterials.Create(InNewData) == InHandle.InnerHandle);
 	}
 
-	void FPBDRigidsSolver::DestroyMaterial(Chaos::FMaterialHandle InHandle)
+	template <typename Traits>
+	void TPBDRigidsSolver<Traits>::DestroyMaterial(Chaos::FMaterialHandle InHandle)
 	{
 		SimMaterials.Destroy(InHandle.InnerHandle);
 	}
 
-	void FPBDRigidsSolver::UpdateMaterialMask(Chaos::FMaterialMaskHandle InHandle, const Chaos::FChaosPhysicsMaterialMask& InNewData)
+	template <typename Traits>
+	void TPBDRigidsSolver<Traits>::UpdateMaterialMask(Chaos::FMaterialMaskHandle InHandle, const Chaos::FChaosPhysicsMaterialMask& InNewData)
 	{
 		*SimMaterialMasks.Get(InHandle.InnerHandle) = InNewData;
 	}
 
-	void FPBDRigidsSolver::CreateMaterialMask(Chaos::FMaterialMaskHandle InHandle, const Chaos::FChaosPhysicsMaterialMask& InNewData)
+	template <typename Traits>
+	void TPBDRigidsSolver<Traits>::CreateMaterialMask(Chaos::FMaterialMaskHandle InHandle, const Chaos::FChaosPhysicsMaterialMask& InNewData)
 	{
 		ensure(SimMaterialMasks.Create(InNewData) == InHandle.InnerHandle);
 	}
 
-	void FPBDRigidsSolver::DestroyMaterialMask(Chaos::FMaterialMaskHandle InHandle)
+	template <typename Traits>
+	void TPBDRigidsSolver<Traits>::DestroyMaterialMask(Chaos::FMaterialMaskHandle InHandle)
 	{
 		SimMaterialMasks.Destroy(InHandle.InnerHandle);
 	}
 
-	void FPBDRigidsSolver::SyncQueryMaterials()
+	template <typename Traits>
+	void TPBDRigidsSolver<Traits>::SyncQueryMaterials()
 	{
 		TSolverQueryMaterialScope<ELockType::Write> Scope(this);
 		QueryMaterials = SimMaterials;
 		QueryMaterialMasks = SimMaterialMasks;
 	}
 
-	void FPBDRigidsSolver::EnableRewindCapture(int32 NumFrames)
+	template <typename Traits>
+	void TPBDRigidsSolver<Traits>::EnableRewindCapture(int32 NumFrames)
 	{
 		MRewindData = MakeUnique<FRewindData>(NumFrames);
 	}
 
-	void FPBDRigidsSolver::FinalizeRewindData(const TParticleView<TPBDRigidParticles<FReal,3>>& ActiveParticles)
+	template <typename Traits>
+	void TPBDRigidsSolver<Traits>::FinalizeRewindData(const TParticleView<TPBDRigidParticles<FReal,3>>& ActiveParticles)
 	{
 		using namespace Chaos;
 		//Simulated objects must have their properties captured for rewind
@@ -1114,14 +1047,14 @@ namespace Chaos
 					{
 
 						if (ActiveObject.GetParticleType() == EParticleType::Rigid)
-						{
-							//may want to remove branch using templates outside loop
+					{
+						//may want to remove branch using templates outside loop
 							if (MRewindData->IsResim())
-							{
+						{
 								MRewindData->PushPTDirtyData<true>(*static_cast<const FRigidParticlePhysicsProxy*>(Proxy)->GetHandle(), DataIdx++);
-							}
-							else
-							{
+						}
+						else
+						{
 								MRewindData->PushPTDirtyData<false>(*static_cast<const FRigidParticlePhysicsProxy*>(Proxy)->GetHandle(), DataIdx++);
 							}
 						}
@@ -1131,4 +1064,8 @@ namespace Chaos
 		}
 	}
 
+#define EVOLUTION_TRAIT(Trait) template class CHAOSSOLVERS_API TPBDRigidsSolver<Trait>;
+#include "Chaos/EvolutionTraits.inl"
+#undef EVOLUTION_TRAIT
+
 }; // namespace Chaos
diff --git a/Engine/Source/Runtime/Experimental/ChaosSolvers/Private/PhysicsProxy/FieldSystemPhysicsProxy.cpp b/Engine/Source/Runtime/Experimental/ChaosSolvers/Private/PhysicsProxy/FieldSystemPhysicsProxy.cpp
--- a/Engine/Source/Runtime/Experimental/ChaosSolvers/Private/PhysicsProxy/FieldSystemPhysicsProxy.cpp
+++ b/Engine/Source/Runtime/Experimental/ChaosSolvers/Private/PhysicsProxy/FieldSystemPhysicsProxy.cpp
@@ -32,7 +32,7 @@ FFieldSystemPhysicsProxy::~FFieldSystemPhysicsProxy()
 {
 	// Need to delete any command lists we have hanging around
 	FScopeLock Lock(&CommandLock);
-	for(TTuple<Chaos::FPhysicsSolver*, TArray<FFieldSystemCommand>*>& Pair : Commands)
+	for(TTuple<Chaos::FPhysicsSolverBase*, TArray<FFieldSystemCommand>*>& Pair : Commands)
 	{
 		delete Pair.Get<1>();
 	}
@@ -47,8 +47,9 @@ bool FFieldSystemPhysicsProxy::IsSimulating() const
 	return true; // #todo Actually start gating this?
 }
 
+template <typename Traits>
 void FFieldSystemPhysicsProxy::FieldParameterUpdateCallback(
-	Chaos::FPhysicsSolver* InSolver, 
+	Chaos::TPBDRigidsSolver<Traits>* InSolver, 
 	FParticlesType& Particles, 
 	Chaos::TArrayCollectionArray<float>& Strains, 
 	Chaos::TPBDPositionConstraints<float, 3>& PositionTarget, 
@@ -59,7 +60,7 @@ void FFieldSystemPhysicsProxy::FieldParameterUpdateCallback(
 	using namespace Chaos;
 	SCOPE_CYCLE_COUNTER(STAT_ParamUpdateField_Object);
 	
-	Chaos::FPhysicsSolver* CurrentSolver = InSolver;
+	Chaos::TPBDRigidsSolver<Traits>* CurrentSolver = InSolver;
 
 	if (Commands.Num() && InSolver)
 	{
@@ -893,8 +894,9 @@ void FFieldSystemPhysicsProxy::FieldParameterUpdateCallback(
 	}
 }
 
+template <typename Traits>
 void FFieldSystemPhysicsProxy::FieldForcesUpdateCallback(
-	Chaos::FPhysicsSolver* InSolver, 
+	Chaos::TPBDRigidsSolver<Traits>* InSolver, 
 	FParticlesType& Particles, 
 	Chaos::TArrayCollectionArray<FVector> & Force, 
 	Chaos::TArrayCollectionArray<FVector> & Torque, 
@@ -902,7 +904,7 @@ void FFieldSystemPhysicsProxy::FieldForcesUpdateCallback(
 {
 	if (Commands.Num() && InSolver)
 	{
-		Chaos::FPhysicsSolver* CurrentSolver = InSolver;
+		Chaos::TPBDRigidsSolver<Traits>* CurrentSolver = InSolver;
 		TArray<ContextIndex> IndicesArray;
 
 		TArray<FFieldSystemCommand>* CommandListPtr = GetSolverCommandList(InSolver);
@@ -1055,7 +1057,7 @@ void FFieldSystemPhysicsProxy::FieldForcesUpdateCallback(
 	}
 }
 
-void FFieldSystemPhysicsProxy::BufferCommand(Chaos::FPhysicsSolver* InSolver, const FFieldSystemCommand& InCommand)
+void FFieldSystemPhysicsProxy::BufferCommand(Chaos::FPhysicsSolverBase* InSolver, const FFieldSystemCommand& InCommand)
 {
 	// TODO: Consider inspecting InCommand and bucketing according to which evaluation 
 	// path it requires; FieldParameterUpdateCallback() or FieldForcesUpdateCallback().
@@ -1123,9 +1125,10 @@ void FFieldSystemPhysicsProxy::BufferCommand(Chaos::FPhysicsSolver* InSolver, co
 //	}
 //}
 
+template <typename Traits>
 void FFieldSystemPhysicsProxy::GetParticleHandles(
 	TArray<Chaos::TGeometryParticleHandle<float, 3>*>& Handles,
-	const Chaos::FPhysicsSolver* RigidSolver,
+	const Chaos::TPBDRigidsSolver<Traits>* RigidSolver,
 	const EFieldResolutionType ResolutionType, 
 	const bool bForce)
 {
@@ -1191,7 +1194,7 @@ void FFieldSystemPhysicsProxy::GetParticleHandles(
 	}
 }
 
-TArray<FFieldSystemCommand>* FFieldSystemPhysicsProxy::GetSolverCommandList(const Chaos::FPhysicsSolver* InSolver)
+TArray<FFieldSystemCommand>* FFieldSystemPhysicsProxy::GetSolverCommandList(const Chaos::FPhysicsSolverBase* InSolver)
 {
 	FScopeLock Lock(&CommandLock);
 	TArray<FFieldSystemCommand>** ExistingList = Commands.Find(InSolver);
@@ -1201,3 +1204,28 @@ TArray<FFieldSystemCommand>* FFieldSystemPhysicsProxy::GetSolverCommandList(cons
 void FFieldSystemPhysicsProxy::OnRemoveFromScene()
 {
 }
+
+#define EVOLUTION_TRAIT(Traits)\
+template void FFieldSystemPhysicsProxy::FieldParameterUpdateCallback(\
+		Chaos::TPBDRigidsSolver<Chaos::Traits>* InSolver, \
+		FParticlesType& InParticles, \
+		Chaos::TArrayCollectionArray<float>& Strains, \
+		Chaos::TPBDPositionConstraints<float, 3>& PositionTarget, \
+		TMap<int32, int32>& PositionTargetedParticles, \
+		const float InTime);\
+\
+template void FFieldSystemPhysicsProxy::FieldForcesUpdateCallback(\
+		Chaos::TPBDRigidsSolver<Chaos::Traits>* InSolver, \
+		FParticlesType& Particles, \
+		Chaos::TArrayCollectionArray<FVector> & Force, \
+		Chaos::TArrayCollectionArray<FVector> & Torque, \
+		const float Time);\
+\
+template void FFieldSystemPhysicsProxy::GetParticleHandles(\
+		TArray<Chaos::TGeometryParticleHandle<float,3>*>& Handles,\
+		const Chaos::TPBDRigidsSolver<Chaos::Traits>* RigidSolver,\
+		const EFieldResolutionType ResolutionType,\
+		const bool bForce);\
+
+#include "Chaos/EvolutionTraits.inl"
+#undef EVOLUTION_TRAIT
\ No newline at end of file
diff --git a/Engine/Source/Runtime/Experimental/ChaosSolvers/Private/PhysicsProxy/GeometryCollectionPhysicsProxy.cpp b/Engine/Source/Runtime/Experimental/ChaosSolvers/Private/PhysicsProxy/GeometryCollectionPhysicsProxy.cpp
--- a/Engine/Source/Runtime/Experimental/ChaosSolvers/Private/PhysicsProxy/GeometryCollectionPhysicsProxy.cpp
+++ b/Engine/Source/Runtime/Experimental/ChaosSolvers/Private/PhysicsProxy/GeometryCollectionPhysicsProxy.cpp
@@ -89,7 +89,7 @@ void FGeometryCollectionResults::Reset()
 }
 
 //==============================================================================
-// FGeometryCollectionPhysicsProxy helper functions
+// TGeometryCollectionPhysicsProxy helper functions
 //==============================================================================
 
 bool IsMultithreaded()
@@ -209,7 +209,7 @@ TArray<int32> ComputeRecursiveOrder(const FGeometryCollection& Collection)
 	return TransformOrder;
 }
 
-DECLARE_CYCLE_STAT(TEXT("FGeometryCollectionPhysicsProxy::PopulateSimulatedParticle"), STAT_PopulateSimulatedParticle, STATGROUP_Chaos);
+DECLARE_CYCLE_STAT(TEXT("TGeometryCollectionPhysicsProxy<Traits>::PopulateSimulatedParticle"), STAT_PopulateSimulatedParticle, STATGROUP_Chaos);
 void PopulateSimulatedParticle(
 	Chaos::TPBDRigidParticleHandle<float,3>* Handle,
 	const FSharedSimulationParameters& SharedParams,
@@ -366,10 +366,11 @@ void PopulateSimulatedParticle(
 }
 
 //==============================================================================
-// FGeometryCollectionPhysicsProxy
+// TGeometryCollectionPhysicsProxy
 //==============================================================================
 
-FGeometryCollectionPhysicsProxy::FGeometryCollectionPhysicsProxy(
+template <typename Traits>
+TGeometryCollectionPhysicsProxy<Traits>::TGeometryCollectionPhysicsProxy(
 	UObject* InOwner,
 	FGeometryDynamicCollection& GameThreadCollectionIn,
 	const FSimulationParameters& SimulationParameters,
@@ -403,13 +404,15 @@ FGeometryCollectionPhysicsProxy::FGeometryCollectionPhysicsProxy(
 	check(BufferMode == Chaos::EMultiBufferMode::TripleGuarded);
 }
 
-FGeometryCollectionPhysicsProxy::~FGeometryCollectionPhysicsProxy()
+template <typename Traits>
+TGeometryCollectionPhysicsProxy<Traits>::~TGeometryCollectionPhysicsProxy()
 {}
 
 float ReportHighParticleFraction = -1.f;
 FAutoConsoleVariableRef CVarReportHighParticleFraction(TEXT("p.gc.ReportHighParticleFraction"), ReportHighParticleFraction, TEXT("Report any objects with particle fraction above this threshold"));
 
-void FGeometryCollectionPhysicsProxy::Initialize()
+template <typename Traits>
+void TGeometryCollectionPhysicsProxy<Traits>::Initialize()
 {
 	check(IsInGameThread());
 
@@ -582,7 +585,8 @@ void FGeometryCollectionPhysicsProxy::Initialize()
 	}
 }
 
-void FGeometryCollectionPhysicsProxy::InitializeDynamicCollection(FGeometryDynamicCollection& DynamicCollection, const FGeometryCollection& RestCollection, const FSimulationParameters& Params)
+template <typename Traits>
+void TGeometryCollectionPhysicsProxy<Traits>::InitializeDynamicCollection(FGeometryDynamicCollection& DynamicCollection, const FGeometryCollection& RestCollection, const FSimulationParameters& Params)
 {
 	// 
 	// This function will use the rest collection to populate the dynamic collection. 
@@ -652,7 +656,8 @@ void FGeometryCollectionPhysicsProxy::InitializeDynamicCollection(FGeometryDynam
 int32 ReportTooManyChildrenNum = -1;
 FAutoConsoleVariableRef CVarReportTooManyChildrenNum(TEXT("p.ReportTooManyChildrenNum"), ReportTooManyChildrenNum, TEXT("Issue warning if more than this many children exist in a single cluster"));
 
-void FGeometryCollectionPhysicsProxy::InitializeBodiesPT(Chaos::FPBDRigidsSolver* RigidsSolver, Chaos::FPBDRigidsSolver::FParticlesType& Particles)
+template <typename Traits>
+void TGeometryCollectionPhysicsProxy<Traits>::InitializeBodiesPT(Chaos::TPBDRigidsSolver<Traits>* RigidsSolver, typename Chaos::TPBDRigidsSolver<Traits>::FParticlesType& Particles)
 {
 	const FGeometryCollection* RestCollection = Parameters.RestCollection;
 	const FGeometryDynamicCollection& DynamicCollection = PhysicsThreadCollection;
@@ -798,7 +803,7 @@ void FGeometryCollectionPhysicsProxy::InitializeBodiesPT(Chaos::FPBDRigidsSolver
 			Commands.Add(Cmd);
 		}
 		Parameters.InitializationCommands.Empty();
-		ProcessCommands(Particles.GetGeometryCollectionParticles(), GetSolver()->GetSolverTime());
+		ProcessCommands(Particles.GetGeometryCollectionParticles(), GetSolver<FSolver>()->GetSolverTime());
 
 		if (Parameters.InitialVelocityType == EInitialVelocityTypeEnum::Chaos_Initial_Velocity_User_Defined)
 		{
@@ -944,7 +949,7 @@ void FGeometryCollectionPhysicsProxy::InitializeBodiesPT(Chaos::FPBDRigidsSolver
 						// where some of these parameters are defined (ie: CoillisionThicknessPercent)
 						ClusterParams.ConnectionMethod = Parameters.ClusterConnectionMethod;
 						
-						GetSolver()->GetEvolution()->GetRigidClustering().GenerateConnectionGraph(SolverClusterHandles[TransformGroupIndex], ClusterParams);
+						GetSolver<FSolver>()->GetEvolution()->GetRigidClustering().GenerateConnectionGraph(SolverClusterHandles[TransformGroupIndex], ClusterParams);
 					}
 				}
 			}
@@ -957,7 +962,7 @@ void FGeometryCollectionPhysicsProxy::InitializeBodiesPT(Chaos::FPBDRigidsSolver
 		{
 			if (UpdateRecordedStateCallback)
 			{
-				UpdateRecordedStateCallback(0.0f, RigidBodyID, Particles, GetSolver()->GetCollisionConstraints());
+				UpdateRecordedStateCallback(0.0f, RigidBodyID, Particles, GetSolver<FSolver>()->GetCollisionConstraints());
 			}
 		}
 #endif // TODO_REIMPLEMENT_RIGID_CACHING
@@ -972,7 +977,7 @@ void FGeometryCollectionPhysicsProxy::InitializeBodiesPT(Chaos::FPBDRigidsSolver
 			// back on once it's off, but even if we didn't enclose this in an if(),
 			// this function won't be called again unless something dirties the proxy.
 
-			Chaos::TPerParticleGravity<float, 3>& GravityForces = GetSolver()->GetEvolution()->GetGravityForces();
+			Chaos::TPerParticleGravity<float, 3>& GravityForces = GetSolver<FSolver>()->GetEvolution()->GetGravityForces();
 			for (int32 HandleIdx = 0; HandleIdx < SolverParticleHandles.Num(); ++HandleIdx)
 			{
 				if (Chaos::TPBDRigidParticleHandle<float, 3>* Handle = SolverParticleHandles[HandleIdx])
@@ -989,11 +994,12 @@ void FGeometryCollectionPhysicsProxy::InitializeBodiesPT(Chaos::FPBDRigidsSolver
 int32 ReportNoLevelsetCluster = 0;
 FAutoConsoleVariableRef CVarReportNoLevelsetCluster(TEXT("p.gc.ReportNoLevelsetCluster"), ReportNoLevelsetCluster, TEXT("Report any cluster objects without levelsets"));
 
-DECLARE_CYCLE_STAT(TEXT("FGeometryCollectionPhysicsProxy::BuildClusters"), STAT_BuildClusters, STATGROUP_Chaos);
-DECLARE_CYCLE_STAT(TEXT("FGeometryCollectionPhysicsProxy::BuildClusters:GlobalMatrices"), STAT_BuildClustersGlobalMatrices, STATGROUP_Chaos);
+DECLARE_CYCLE_STAT(TEXT("TGeometryCollectionPhysicsProxy::BuildClusters"), STAT_BuildClusters, STATGROUP_Chaos);
+DECLARE_CYCLE_STAT(TEXT("TGeometryCollectionPhysicsProxy::BuildClusters:GlobalMatrices"), STAT_BuildClustersGlobalMatrices, STATGROUP_Chaos);
 
+template <typename Traits>
 Chaos::TPBDRigidClusteredParticleHandle<float, 3>* 
-FGeometryCollectionPhysicsProxy::BuildClusters(
+TGeometryCollectionPhysicsProxy<Traits>::BuildClusters(
 	const uint32 CollectionClusterIndex, // TransformGroupIndex
 	TArray<Chaos::TPBDRigidParticleHandle<float,3>*>& ChildHandles,
 	const TArray<int32>& ChildTransformGroupIndices,
@@ -1038,7 +1044,7 @@ FGeometryCollectionPhysicsProxy::BuildClusters(
 
 	// Construct an active cluster particle, disable children, derive M and I from children:
 	Chaos::TPBDRigidClusteredParticleHandle<float, 3>* Parent = 
-		GetSolver()->GetEvolution()->GetRigidClustering().CreateClusterParticle(
+		GetSolver<FSolver>()->GetEvolution()->GetRigidClustering().CreateClusterParticle(
 			Parameters.ClusterGroupIndex, 
 			MoveTemp(ChildHandlesCopy),
 			ClusterCreationParameters,
@@ -1103,7 +1109,7 @@ FGeometryCollectionPhysicsProxy::BuildClusters(
 	check(Parameters.PhysicalMaterial);
 	Parent->SetLinearEtherDrag(Parameters.PhysicalMaterial->LinearEtherDrag);
 	Parent->SetAngularEtherDrag(Parameters.PhysicalMaterial->AngularEtherDrag);
-	GetSolver()->GetEvolution()->SetPhysicsMaterial(Parent, Parameters.PhysicalMaterial);
+	GetSolver<FSolver>()->GetEvolution()->SetPhysicsMaterial(Parent, Parameters.PhysicalMaterial);
 
 	const FTransform ParentTransform = GeometryCollectionAlgo::GlobalMatrix(DynamicCollection.Transform, DynamicCollection.Parent, CollectionClusterIndex);
 
@@ -1120,12 +1126,12 @@ FGeometryCollectionPhysicsProxy::BuildClusters(
 		SolverClusterHandles[ChildTransformGroupIndex] = Parent;
 
 		MinCollisionGroup = FMath::Min(Child->CollisionGroup(), MinCollisionGroup);
-	}
+		}
 	Parent->SetCollisionGroup(MinCollisionGroup);
 
 	// Populate bounds as we didn't pass a shared implicit to PopulateSimulatedParticle this will have been skipped, now that we have the full cluster we can build it
 	if(Parent->Geometry() && Parent->Geometry()->HasBoundingBox())
-	{
+		{
 		Parent->SetHasBounds(true);
 		Parent->SetLocalBounds(Parent->Geometry()->BoundingBox());
 		const Chaos::TAABB<float, 3>& LocalBounds = Parent->LocalBounds();
@@ -1134,19 +1140,19 @@ FGeometryCollectionPhysicsProxy::BuildClusters(
 		Parent->SetWorldSpaceInflatedBounds(TransformedBBox);
 
 		GetSolver()->GetEvolution()->DirtyParticle(*Parent);
-	}
+		}
 
 	return Parent;
 }
 
-
-void FGeometryCollectionPhysicsProxy::GetRelevantHandles(
+template <typename Traits>
+void TGeometryCollectionPhysicsProxy<Traits>::GetRelevantHandles(
 	TArray<Chaos::TGeometryParticleHandle<float, 3>*>& Handles,
 	TArray<FVector>& Samples,
 	TArray<ContextIndex>& SampleIndices,
-	const Chaos::FPhysicsSolver* RigidSolver, 
+	const Chaos::TPBDRigidsSolver<Traits>* RigidSolver, 
 	EFieldResolutionType ResolutionType, 
-	bool bForce = true)
+	bool bForce)
 {
 	if(bForce)
 	{
@@ -1240,7 +1246,8 @@ void FGeometryCollectionPhysicsProxy::GetRelevantHandles(
 	}
 }
 
-void FGeometryCollectionPhysicsProxy::PushKinematicStateToSolver(FParticlesType& Particles)
+template <typename Traits>
+void TGeometryCollectionPhysicsProxy<Traits>::PushKinematicStateToSolver(FParticlesType& Particles)
 {
 	FGeometryDynamicCollection& Collection = GameThreadCollection;
 	if (Collection.Transform.Num())
@@ -1293,7 +1300,8 @@ void FGeometryCollectionPhysicsProxy::PushKinematicStateToSolver(FParticlesType&
 	}
 }
 
-int32 FGeometryCollectionPhysicsProxy::CalculateHierarchyLevel(const FGeometryDynamicCollection& GeometryCollection, int32 TransformIndex) const
+template <typename Traits>
+int32 TGeometryCollectionPhysicsProxy<Traits>::CalculateHierarchyLevel(const FGeometryDynamicCollection& GeometryCollection, int32 TransformIndex) const
 {
 	int32 Level = 0;
 	while (GeometryCollection.Parent[TransformIndex] != -1)
@@ -1304,8 +1312,8 @@ int32 FGeometryCollectionPhysicsProxy::CalculateHierarchyLevel(const FGeometryDy
 	return Level;
 }
 
-
-void FGeometryCollectionPhysicsProxy::InitializeRemoveOnFracture(FParticlesType& Particles, const TManagedArray<int32>& DynamicState)
+template <typename Traits>
+void TGeometryCollectionPhysicsProxy<Traits>::InitializeRemoveOnFracture(FParticlesType& Particles, const TManagedArray<int32>& DynamicState)
 {
 	/*
 	@todo break everything
@@ -1329,7 +1337,8 @@ void FGeometryCollectionPhysicsProxy::InitializeRemoveOnFracture(FParticlesType&
 	*/
 }
 
-void FGeometryCollectionPhysicsProxy::OnRemoveFromSolver(Chaos::FPBDRigidsSolver *RBDSolver)
+template <typename Traits>
+void TGeometryCollectionPhysicsProxy<Traits>::OnRemoveFromSolver(Chaos::TPBDRigidsSolver<Traits> *RBDSolver)
 {
 	const FGeometryDynamicCollection& DynamicCollection = PhysicsThreadCollection;
 
@@ -1339,14 +1348,15 @@ void FGeometryCollectionPhysicsProxy::OnRemoveFromSolver(Chaos::FPBDRigidsSolver
 	}
 }
 
-void FGeometryCollectionPhysicsProxy::OnRemoveFromScene()
+template <typename Traits>
+void TGeometryCollectionPhysicsProxy<Traits>::OnRemoveFromScene()
 {
 #if TODO_REIMPLEMENT_GET_RIGID_PARTICLES
 	// #BG TODO This isn't great - we currently cannot handle things being removed from the solver.
 	// need to refactor how we handle this and actually remove the particles instead of just constantly
 	// growing the array. Currently everything is just tracked by index though so the solver will have
 	// to notify all the proxies that a chunk of data was removed - or use a sparse array (undesireable)
-	Chaos::FPhysicsSolver::FParticlesType& Particles = GetSolver()->GetRigidParticles();
+	Chaos::FPhysicsSolver::FParticlesType& Particles = GetSolver<FSolver>()->GetRigidParticles();
 
 	// #BG TODO Special case here because right now we reset/realloc the evolution per geom component
 	// in endplay which clears this out. That needs to not happen and be based on world shutdown
@@ -1362,14 +1372,15 @@ void FGeometryCollectionPhysicsProxy::OnRemoveFromScene()
 	{
 		for (int32 ParticleIndex = 0; ParticleIndex < Count; ++ParticleIndex)
 		{
-			GetSolver()->GetEvolution()->DisableParticle(Begin + ParticleIndex);
-			GetSolver()->GetRigidClustering().GetTopLevelClusterParents().Remove(Begin + ParticleIndex);
+			GetSolver<FSolver>()->GetEvolution()->DisableParticle(Begin + ParticleIndex);
+			GetSolver<FSolver>()->GetRigidClustering().GetTopLevelClusterParents().Remove(Begin + ParticleIndex);
 		}
 	}
 #endif
 }
 
-void FGeometryCollectionPhysicsProxy::SyncBeforeDestroy()
+template <typename Traits>
+void TGeometryCollectionPhysicsProxy<Traits>::SyncBeforeDestroy()
 {
 	if(FinalSyncFunc)
 	{
@@ -1379,14 +1390,16 @@ void FGeometryCollectionPhysicsProxy::SyncBeforeDestroy()
 	}
 }
 
-void FGeometryCollectionPhysicsProxy::BufferGameState() 
+template <typename Traits>
+void TGeometryCollectionPhysicsProxy<Traits>::BufferGameState() 
 {
 	//
 	// There is currently no per advance updates to the GeometryCollection
 	//
 }
 
-void FGeometryCollectionPhysicsProxy::PushToPhysicsState(const Chaos::FParticleData *InData)
+template <typename Traits>
+void TGeometryCollectionPhysicsProxy<Traits>::PushToPhysicsState(const Chaos::FParticleData *InData)
 {
 	/*CONTEXT: GAMETHREAD->to->PHYSICSTHREAD
 	* Called on the game thread when the solver is about to advance forward.This
@@ -1408,7 +1421,8 @@ void FGeometryCollectionPhysicsProxy::PushToPhysicsState(const Chaos::FParticleD
 	}
 }
 
-void FGeometryCollectionPhysicsProxy::BufferPhysicsResults()
+template <typename Traits>
+void TGeometryCollectionPhysicsProxy<Traits>::BufferPhysicsResults()
 {
 	/**
 	 * CONTEXT: PHYSICSTHREAD
@@ -1533,10 +1547,10 @@ void FGeometryCollectionPhysicsProxy::BufferPhysicsResults()
 
 						// syncronize parents if it has changed.
 						if (SolverClusterID[TransformGroupIndex] != ClusterParent)
-						{
-							// Force all driven rigid bodies out of the transform hierarchy
+				{
+					// Force all driven rigid bodies out of the transform hierarchy
 							if (Parent[TransformGroupIndex] != INDEX_NONE)
-							{
+					{
 								// If the parent of this NON DISABLED body is set to anything other than INDEX_NONE,
 								// then it was just unparented, likely either by rigid clustering or by fields.  We
 								// need to force all such enabled rigid bodies out of the transform hierarchy.
@@ -1549,12 +1563,12 @@ void FGeometryCollectionPhysicsProxy::BufferPhysicsResults()
 								// Indicate that this object needs to be updated and the proxy is active.
 								TargetResults.DisabledStates[TransformGroupIndex] = false;
 								IsObjectDynamic = true;
-							}
+					}
 							SolverClusterID[TransformGroupIndex] = Handle->ClusterIds().Id;
-						}
+				}
 
 						if (ClusterParent->InternalCluster() )
-						{
+	{
 							Chaos::TPBDRigidClusteredParticleHandle<float, 3>* ProxyElementHandle = SolverParticleHandles[TransformGroupIndex];
 
 							FTransform& ParticleToWorld = TargetResults.ParticleToWorldTransforms[TransformGroupIndex];
@@ -1575,8 +1589,8 @@ void FGeometryCollectionPhysicsProxy::BufferPhysicsResults()
 							ProxyElementHandle->R() = ParticleToWorld.GetRotation();
 							GetSolver()->GetEvolution()->DirtyParticle(*ProxyElementHandle);
 						}
-					}
-				}
+		}
+	}
 			}// end if
 
 			PhysicsThreadCollection.Active[TransformGroupIndex] = !TargetResults.DisabledStates[TransformGroupIndex];
@@ -1586,7 +1600,7 @@ void FGeometryCollectionPhysicsProxy::BufferPhysicsResults()
 
 
 	// If object is dynamic, compute global matrices	
- 	if (IsObjectDynamic || TargetResults.GlobalTransforms.Num() == 0)
+	if (IsObjectDynamic || TargetResults.GlobalTransforms.Num() == 0)
 	{
 		SCOPE_CYCLE_COUNTER(STAT_CalcGlobalGCMatrices);
 		check(TargetResults.Transforms.Num() == TargetResults.Parent.Num());
@@ -1610,7 +1624,8 @@ void FGeometryCollectionPhysicsProxy::BufferPhysicsResults()
 	TargetResults.IsObjectLoading = IsObjectLoading;
 }
 
-void FGeometryCollectionPhysicsProxy::FlipBuffer()
+template <typename Traits>
+void TGeometryCollectionPhysicsProxy<Traits>::FlipBuffer()
 {
 	/**
 	 * CONTEXT: PHYSICSTHREAD (Write Locked)
@@ -1623,7 +1638,8 @@ void FGeometryCollectionPhysicsProxy::FlipBuffer()
 }
 
 // Called from FPhysScene_ChaosInterface::SyncBodies(), NOT the solver.
-void FGeometryCollectionPhysicsProxy::PullFromPhysicsState()
+template <typename Traits>
+void TGeometryCollectionPhysicsProxy<Traits>::PullFromPhysicsState()
 {
 	/**
 	 * CONTEXT: GAMETHREAD (Read Locked)
@@ -1706,7 +1722,8 @@ int32 FindSizeSpecificIdx(const TArray<FSharedSimulationSizeSpecificData>& SizeS
 	NOTE - Making any changes to data stored on the rest collection below MUST be accompanied
 	by a rotation of the DDC key in FDerivedDataGeometryCollectionCooker::GetVersionString
 */
-void FGeometryCollectionPhysicsProxy::InitializeSharedCollisionStructures(
+template <typename Traits>
+void TGeometryCollectionPhysicsProxy<Traits>::InitializeSharedCollisionStructures(
 	Chaos::FErrorReporter& ErrorReporter,
 	FGeometryCollection& RestCollection,
 	const FSharedSimulationParameters& SharedParams)
@@ -2323,7 +2340,8 @@ void FGeometryCollectionPhysicsProxy::InitializeSharedCollisionStructures(
 	}
 }
 
-void FGeometryCollectionPhysicsProxy::InitRemoveOnFracture(FGeometryCollection& RestCollection, const FSharedSimulationParameters& SharedParams)
+template <typename Traits>
+void TGeometryCollectionPhysicsProxy<Traits>::InitRemoveOnFracture(FGeometryCollection& RestCollection, const FSharedSimulationParameters& SharedParams)
 {
 	if (SharedParams.RemoveOnFractureIndices.Num() == 0)
 	{
@@ -2444,8 +2462,8 @@ void BuildSimulationData(Chaos::FErrorReporter& ErrorReporter, FGeometryCollecti
 // FIELDS
 //==============================================================================
 
-
-void FGeometryCollectionPhysicsProxy::ProcessCommands(FParticlesType& Particles, const float Time)
+template <typename Traits>
+void TGeometryCollectionPhysicsProxy<Traits>::ProcessCommands(FParticlesType& Particles, const float Time)
 {
 	FGeometryDynamicCollection& Collection = GameThreadCollection;
 
@@ -2457,7 +2475,7 @@ void FGeometryCollectionPhysicsProxy::ProcessCommands(FParticlesType& Particles,
 		TArray<ContextIndex> SampleIndices;
 
 		TArray<ContextIndex> IndicesArray;
-		Chaos::FPhysicsSolver* CurrentSolver = GetSolver();
+		FSolver* CurrentSolver = GetSolver<FSolver>();
 
 		for (int32 CommandIndex = Commands.Num()-1; 0<=CommandIndex; CommandIndex--)
 		{
@@ -2478,7 +2496,7 @@ void FGeometryCollectionPhysicsProxy::ProcessCommands(FParticlesType& Particles,
 					TEXT("Field based evaluation of the simulations 'DynamicState' parameter expects int32 field inputs.")))
 				{
 
-					FGeometryCollectionPhysicsProxy::GetRelevantHandles(Handles, Samples, SampleIndices, CurrentSolver, ResolutionType, IndicesArray.Num() != Particles.Size());
+					TGeometryCollectionPhysicsProxy<Traits>::GetRelevantHandles(Handles, Samples, SampleIndices, CurrentSolver, ResolutionType, IndicesArray.Num() != Particles.Size());
 					if (Handles.Num())
 					{
 						TArrayView<FVector> SamplesView(&(Samples[0]), Samples.Num());
@@ -2537,7 +2555,7 @@ void FGeometryCollectionPhysicsProxy::ProcessCommands(FParticlesType& Particles,
 					if (ensureMsgf(Command.RootNode->Type() == FFieldNode<FVector>::StaticType(),
 						TEXT("Field based evaluation of the simulations 'InitialLinearVelocity' parameter expects FVector field inputs.")))
 					{
-						FGeometryCollectionPhysicsProxy::GetRelevantHandles(Handles, Samples, SampleIndices, CurrentSolver, ResolutionType, IndicesArray.Num() != Particles.Size());
+						TGeometryCollectionPhysicsProxy<Traits>::GetRelevantHandles(Handles, Samples, SampleIndices, CurrentSolver, ResolutionType, IndicesArray.Num() != Particles.Size());
 						if (Handles.Num())
 						{
 							TArrayView<FVector> SamplesView(&(Samples[0]), Samples.Num());
@@ -2562,7 +2580,7 @@ void FGeometryCollectionPhysicsProxy::ProcessCommands(FParticlesType& Particles,
 					if (ensureMsgf(Command.RootNode->Type() == FFieldNode<FVector>::StaticType(),
 						TEXT("Field based evaluation of the simulations 'InitialAngularVelocity' parameter expects FVector field inputs.")))
 					{
-						FGeometryCollectionPhysicsProxy::GetRelevantHandles(Handles, Samples, SampleIndices, CurrentSolver, ResolutionType, IndicesArray.Num() != Particles.Size());
+						TGeometryCollectionPhysicsProxy<Traits>::GetRelevantHandles(Handles, Samples, SampleIndices, CurrentSolver, ResolutionType, IndicesArray.Num() != Particles.Size());
 						if (Handles.Num())
 						{
 							TArrayView<FVector> SamplesView(&(Samples[0]), Samples.Num());
@@ -2587,13 +2605,13 @@ void FGeometryCollectionPhysicsProxy::ProcessCommands(FParticlesType& Particles,
 	// Process Particle-Particle commands
 	if (Commands.Num())
 	{
-		Chaos::FPhysicsSolver* CurrentSolver = GetSolver();
+		FSolver* CurrentSolver = GetSolver<FSolver>();
 
 		TArray<Chaos::TGeometryParticleHandle<float, 3>*> Handles;
 		TArray<FVector> Samples;
 		TArray<ContextIndex> SampleIndices;
 
-		FGeometryCollectionPhysicsProxy::GetRelevantHandles(Handles, Samples, SampleIndices, CurrentSolver, EFieldResolutionType::Field_Resolution_Maximum, true);
+		TGeometryCollectionPhysicsProxy<Traits>::GetRelevantHandles(Handles, Samples, SampleIndices, CurrentSolver, EFieldResolutionType::Field_Resolution_Maximum, true);
 
 		for (int32 CommandIndex = Commands.Num() - 1; 0 <= CommandIndex; CommandIndex--)
 		{
@@ -2658,8 +2676,8 @@ void FGeometryCollectionPhysicsProxy::ProcessCommands(FParticlesType& Particles,
 	}
 }
 
-
-void FGeometryCollectionPhysicsProxy::FieldForcesUpdateCallback(Chaos::FPhysicsSolver* InSolver, FParticlesType& Particles, Chaos::TArrayCollectionArray<FVector> & Force, Chaos::TArrayCollectionArray<FVector> & Torque, const float Time)
+template <typename Traits>
+void TGeometryCollectionPhysicsProxy<Traits>::FieldForcesUpdateCallback(Chaos::TPBDRigidsSolver<Traits>* InSolver, FParticlesType& Particles, Chaos::TArrayCollectionArray<FVector> & Force, Chaos::TArrayCollectionArray<FVector> & Torque, const float Time)
 {
 	if (Commands.Num())
 		{
@@ -2667,7 +2685,7 @@ void FGeometryCollectionPhysicsProxy::FieldForcesUpdateCallback(Chaos::FPhysicsS
 		TArray<FVector> Samples;
 		TArray<ContextIndex> SampleIndices;
 
-		FGeometryCollectionPhysicsProxy::GetRelevantHandles(Handles, Samples, SampleIndices, InSolver, EFieldResolutionType::Field_Resolution_Maximum, true);
+		TGeometryCollectionPhysicsProxy<Traits>::GetRelevantHandles(Handles, Samples, SampleIndices, InSolver, EFieldResolutionType::Field_Resolution_Maximum, true);
 
 			for (int32 CommandIndex = 0; CommandIndex < Commands.Num(); CommandIndex++)
 			{
@@ -2712,7 +2730,8 @@ void FGeometryCollectionPhysicsProxy::FieldForcesUpdateCallback(Chaos::FPhysicsS
 		}
 }
 
-void FGeometryCollectionPhysicsProxy::ParameterUpdateCallback(FParticlesType& Particles, const float Time)
+template <typename Traits>
+void TGeometryCollectionPhysicsProxy<Traits>::ParameterUpdateCallback(FParticlesType& Particles, const float Time)
 {
 	if (GameThreadCollection.Transform.Num())
 	{
@@ -2720,3 +2739,6 @@ void FGeometryCollectionPhysicsProxy::ParameterUpdateCallback(FParticlesType& Pa
 	}
 }
 
+#define EVOLUTION_TRAIT(Traits) template class CHAOSSOLVERS_API TGeometryCollectionPhysicsProxy<Chaos::Traits>;
+#include "Chaos/EvolutionTraits.inl"
+#undef EVOLUTION_TRAIT
diff --git a/Engine/Source/Runtime/Experimental/ChaosSolvers/Private/PhysicsProxy/SingleParticlePhysicsProxy.cpp b/Engine/Source/Runtime/Experimental/ChaosSolvers/Private/PhysicsProxy/SingleParticlePhysicsProxy.cpp
--- a/Engine/Source/Runtime/Experimental/ChaosSolvers/Private/PhysicsProxy/SingleParticlePhysicsProxy.cpp
+++ b/Engine/Source/Runtime/Experimental/ChaosSolvers/Private/PhysicsProxy/SingleParticlePhysicsProxy.cpp
@@ -36,8 +36,8 @@ FSingleParticlePhysicsProxy<PARTICLE_TYPE>::GetInitialState() const
 	return InitialState; 
 }
 
-template <Chaos::EParticleType ParticleType, typename TParticleHandle>
-void PushToPhysicsStateImp(const Chaos::FDirtyPropertiesManager& Manager, TParticleHandle* Handle, int32 DataIdx, const Chaos::FDirtyProxy& Dirty, Chaos::FShapeDirtyData* ShapesData, Chaos::FPhysicsSolver* Solver, const bool bInitialized)
+template <Chaos::EParticleType ParticleType, typename TParticleHandle, typename TEvolution>
+void PushToPhysicsStateImp(const Chaos::FDirtyPropertiesManager& Manager, TParticleHandle* Handle, int32 DataIdx, const Chaos::FDirtyProxy& Dirty, Chaos::FShapeDirtyData* ShapesData, TEvolution& Evolution, const bool bInitialized)
 {
 	using namespace Chaos;
 	constexpr bool bHasKinematicData = ParticleType != EParticleType::Static;
@@ -81,7 +81,7 @@ void PushToPhysicsStateImp(const Chaos::FDirtyPropertiesManager& Manager, TParti
 				Handle->SetWorldSpaceInflatedBounds(WorldSpaceBounds);
 			}
 
-			Solver->GetEvolution()->DirtyParticle(*Handle);
+			Evolution.DirtyParticle(*Handle);
 		}
 
 		if(bHasDynamicData)
@@ -98,7 +98,7 @@ void PushToPhysicsStateImp(const Chaos::FDirtyPropertiesManager& Manager, TParti
 
 			if(auto NewData = ParticleData.FindDynamicMisc(Manager,DataIdx))
 			{
-				Solver->GetEvolution()->SetParticleObjectState(RigidHandle,NewData->ObjectState());
+				Evolution.SetParticleObjectState(RigidHandle,NewData->ObjectState());
 
 				RigidHandle->SetDynamicMisc(*NewData);
 			}
@@ -126,34 +126,10 @@ void PushToPhysicsStateImp(const Chaos::FDirtyPropertiesManager& Manager, TParti
 //
 
 template< >
-void FSingleParticlePhysicsProxy<Chaos::TGeometryParticle<float, 3>>::PushToPhysicsState(const Chaos::FDirtyPropertiesManager& Manager, int32 DataIdx, const Chaos::FDirtyProxy& Dirty, Chaos::FShapeDirtyData* ShapesData)
+template <typename Traits>
+void FSingleParticlePhysicsProxy<Chaos::TGeometryParticle<float, 3>>::PushToPhysicsState(const Chaos::FDirtyPropertiesManager& Manager, int32 DataIdx, const Chaos::FDirtyProxy& Dirty, Chaos::FShapeDirtyData* ShapesData, Chaos::TPBDRigidsEvolutionGBF<Traits>& Evolution)
 {
-	PushToPhysicsStateImp<Chaos::EParticleType::Static>(Manager,Handle,DataIdx,Dirty,ShapesData,GetSolver(),bInitialized);
-#if 0
-	// move the copied game thread data into the handle
-	if (auto* RigidHandle = static_cast<Chaos::TGeometryParticleHandle<float, 3>*>(Handle))
-	{
-		auto NewX = Data->FindX();
-		auto NewR = Data->FindR();
-		auto NewGeometry = Data->FindGeometry();
-
-		if(NewX){ RigidHandle->SetX(*NewX); }
-		if(NewR){ RigidHandle->SetR(*NewR); }
-		if(NewGeometry){ RigidHandle->SetSharedGeometry(*NewGeometry); }
-
-		if(NewX || NewR || NewGeometry)
-		{
-			const auto& Geometry = RigidHandle->Geometry();
-			if(Geometry && Geometry->HasBoundingBox())
-			{
-				RigidHandle->SetHasBounds(true);
-				RigidHandle->SetLocalBounds(Geometry->BoundingBox());
-				RigidHandle->SetWorldSpaceInflatedBounds(Geometry->BoundingBox().TransformedAABB(Chaos::TRigidTransform<float,3>(RigidHandle->X(),RigidHandle->R())));
-			}
-		}
-
-	}
-#endif
+	PushToPhysicsStateImp<Chaos::EParticleType::Static>(Manager,Handle,DataIdx,Dirty,ShapesData,Evolution,bInitialized);
 }
 
 
@@ -205,130 +181,11 @@ void FSingleParticlePhysicsProxy<Chaos::TGeometryParticle<float, 3>>::ClearEvent
 // TKinematicGeometryParticle template specialization 
 //
 
-template< >
-void FSingleParticlePhysicsProxy<Chaos::TKinematicGeometryParticle<float, 3>>::PushToPhysicsState(const Chaos::FDirtyPropertiesManager& Manager, int32 DataIdx, const Chaos::FDirtyProxy& Dirty, Chaos::FShapeDirtyData* ShapesData)
+template <>
+template<typename Traits>
+void FSingleParticlePhysicsProxy<Chaos::TKinematicGeometryParticle<float, 3>>::PushToPhysicsState(const Chaos::FDirtyPropertiesManager& Manager, int32 DataIdx, const Chaos::FDirtyProxy& Dirty, Chaos::FShapeDirtyData* ShapesData, Chaos::TPBDRigidsEvolutionGBF<Traits>& Evolution)
 {
-	PushToPhysicsStateImp<Chaos::EParticleType::Kinematic>(Manager,Handle->CastToKinematicParticle(),DataIdx,Dirty,ShapesData,GetSolver(),bInitialized);
-#if 0
-	// move the copied game thread data into the handle
-	if(auto* RigidHandle = static_cast<Chaos::TKinematicGeometryParticleHandle<float,3>*>(Handle))
-	{
-		auto NewX = Data->FindX();
-		auto NewR = Data->FindR();
-		auto NewGeometry = Data->FindGeometry();
-
-		if(NewX){ RigidHandle->SetX(*NewX); }
-		if(NewR){ RigidHandle->SetR(*NewR); }
-		if(NewGeometry){ RigidHandle->SetSharedGeometry(*NewGeometry); }
-
-		if(NewX || NewR || NewGeometry)
-		{
-			const auto& Geometry = RigidHandle->Geometry();
-			if(Geometry && Geometry->HasBoundingBox())
-			{
-				RigidHandle->SetHasBounds(true);
-				RigidHandle->SetLocalBounds(Geometry->BoundingBox());
-				RigidHandle->SetWorldSpaceInflatedBounds(Geometry->BoundingBox().TransformedAABB(Chaos::TRigidTransform<float,3>(RigidHandle->X(),RigidHandle->R())));
-			}
-		}
-
-		if(auto NewSpatialIdx = Data->FindSpatialIdx()){ RigidHandle->SetSpatialIdx(*NewSpatialIdx); }
-		if(auto NewUniqueIdx = Data->FindUniqueIdx()){ RigidHandle->SetUniqueIdx(*NewUniqueIdx); }
-		if(auto NewUserData = Data->FindUserData()){ RigidHandle->SetUserData(*NewUserData); }
-#if CHAOS_CHECKED
-		if(auto NewDebugName = Data->FindDebugName()){ RigidHandle->SetDebugName(*NewDebugName); }
-#endif
-
-		int32 ShapeIdx = 0;
-		const auto& Shapes = RigidHandle->ShapesArray();
-		for(Chaos::FShapePropertiesData* ShapeData : ShapesData->GetRemoteDatas())
-		{
-			//todo: avoid iterating over every shape regardless of if it's dirty or not
-			if(auto NewData = ShapeData->FindQueryData()){ Shapes[ShapeIdx]->SetQueryData(*NewData); }
-			if(auto NewData = ShapeData->FindSimData()){ Shapes[ShapeIdx]->SetSimData(*NewData); }
-			if(auto NewData = ShapeData->FindUserData()){ Shapes[ShapeIdx]->SetUserData(*NewData); }
-			if(auto NewData = ShapeData->FindGeometry()){ Shapes[ShapeIdx]->SetGeometry(*NewData); }
-			//todo: fully marshal materials
-#if 0
-			if(auto NewData = ShapeData->FindMaterials()){ Shapes[ShapeIdx]->SetMaterials(*NewData); }
-			if(auto NewData = ShapeData->FindMaterialMasks()){ Shapes[ShapeIdx]->SetMaterialMasks(*NewData); }
-			if(auto NewData = ShapeData->FindMaterialMaskMaps()){ Shapes[ShapeIdx]->SetMaterialMaskMaps(*NewData); }
-			if(auto NewData = ShapeData->FindMaterialMaskMapMaterials()){ Shapes[ShapeIdx]->SetMaterialMaskMapMaterials(*NewData); }
-#endif
-			if(auto NewData = ShapeData->FindDisable()){ Shapes[ShapeIdx]->SetDisable(*NewData); }
-			if(auto NewData = ShapeData->FindSimulate()){ Shapes[ShapeIdx]->SetSimulate(*NewData); }
-			if(auto NewData = ShapeData->FindCollisionTraceType()){ Shapes[ShapeIdx]->SetCollisionTraceType(*NewData); }
-		}
-	}
-
-#if 0
-	// move the copied game thread data into the handle
-	if (auto* RigidHandle = static_cast<Chaos::TKinematicGeometryParticleHandle<float, 3>*>(Handle))
-	{
-		typedef Chaos::TKinematicGeometryParticle<float, 3> PARTICLE_TYPE;
-		const auto* Data = static_cast<const typename PARTICLE_TYPE::FData*>(InData);
-
-		RigidHandle->SetX(Data->X);
-		RigidHandle->SetR(Data->R);
-		RigidHandle->SetSharedGeometry(Data->Geometry);
-		RigidHandle->SetV(Data->MV);
-		RigidHandle->SetW(Data->MW);
-		RigidHandle->SetCenterOfMass(Data->MCenterOfMass);
-		RigidHandle->SetRotationOfMass(Data->MRotationOfMass);
-		RigidHandle->SetSpatialIdx(Data->SpatialIdx);	//todo: this needs to only happen once during initialization
-		RigidHandle->SetUniqueIdx(Data->UniqueIdx);
-		RigidHandle->SetUserData(Data->UserData);
-#if CHAOS_CHECKED
-		RigidHandle->SetDebugName(Data->DebugName);
-#endif
-		
-		if (Data->Geometry && Data->Geometry->HasBoundingBox())
-		{
-			RigidHandle->SetHasBounds(true);
-			RigidHandle->SetLocalBounds(Data->Geometry->BoundingBox());
-			Chaos::TAABB<float, 3> WorldSpaceBox = Data->Geometry->BoundingBox().TransformedAABB(Chaos::TRigidTransform<float, 3>(Data->X, Data->R));
-			WorldSpaceBox.ThickenSymmetrically(Data->MV);
-			RigidHandle->SetWorldSpaceInflatedBounds(WorldSpaceBox);
-		}
-
-		if (Data->DirtyFlags.IsDirty(Chaos::EParticleFlags::ShapeDisableCollision))
-		{
-			int32 CurrShape = 0;
-			for (const TUniquePtr<Chaos::TPerShapeData<Chaos::FReal, 3>>& Shape : RigidHandle->ShapesArray())
-			{
-				Shape->bDisable = Data->ShapeCollisionDisableFlags[CurrShape++];
-			}
-		}
-		if (Data->DirtyFlags.IsDirty(Chaos::EParticleFlags::CollisionTraceType))
-		{
-			int32 CurrShape = 0;
-			for (const TUniquePtr<Chaos::TPerShapeData<Chaos::FReal, 3>>& Shape : RigidHandle->ShapesArray())
-			{
-				Shape->CollisionTraceType = Data->CollisionTraceType[CurrShape++];
-			}
-		}
-		//if (Data->DirtyFlags.IsDirty(Chaos::EParticleFlags::ShapeSimData))
-		{
-			int32 CurrShape = 0;
-			for (const TUniquePtr<Chaos::TPerShapeData<Chaos::FReal, 3>>& Shape : RigidHandle->ShapesArray())
-			{
-				Shape->SimData = Data->ShapeSimData[CurrShape];
-				Shape->QueryData = Data->ShapeQueryData[CurrShape++];
-			}
-		}
-
-		{
-			int32 CurrShape = 0;
-			for (const TUniquePtr<Chaos::TPerShapeData<Chaos::FReal, 3>>& Shape : RigidHandle->ShapesArray())
-			{
-				Shape->Materials = Data->ShapeMaterials[CurrShape++];
-			}
-		}
-
-
-	}
-#endif
-#endif
+	PushToPhysicsStateImp<Chaos::EParticleType::Kinematic>(Manager,Handle->CastToKinematicParticle(),DataIdx,Dirty,ShapesData,Evolution,bInitialized);
 }
 
 template< >
@@ -377,151 +234,11 @@ void FSingleParticlePhysicsProxy<Chaos::TKinematicGeometryParticle<float, 3>>::C
 // TPBDRigidParticle template specialization 
 //
 
-template< >
-void FSingleParticlePhysicsProxy<Chaos::TPBDRigidParticle<float, 3>>::PushToPhysicsState(const Chaos::FDirtyPropertiesManager& Manager, int32 DataIdx, const Chaos::FDirtyProxy& Dirty, Chaos::FShapeDirtyData* ShapesData)
+template<>
+template<typename Traits>
+void FSingleParticlePhysicsProxy<Chaos::TPBDRigidParticle<float, 3>>::PushToPhysicsState(const Chaos::FDirtyPropertiesManager& Manager, int32 DataIdx, const Chaos::FDirtyProxy& Dirty, Chaos::FShapeDirtyData* ShapesData, Chaos::TPBDRigidsEvolutionGBF<Traits>& Evolution)
 {
-	PushToPhysicsStateImp<Chaos::EParticleType::Rigid>(Manager,Handle->CastToRigidParticle(),DataIdx,Dirty,ShapesData,GetSolver(),bInitialized);
-#if 0
-	// move the copied game thread data into the handle
-	if (auto* RigidHandle = static_cast<Chaos::TPBDRigidParticleHandle<float, 3>*>(Handle))
-	{
-		typedef Chaos::TPBDRigidParticle<float, 3> PARTICLE_TYPE;
-		const auto* Data = static_cast<const typename PARTICLE_TYPE::FData*>(InData);
-
-		//
-		// TODO: Split all these setters into separate functions so that
-		//       we don't have to check for each dirty type. We don't
-		//       *actually* have to check each dirty type, but some of
-		//       these set operations require extra actions. This avoids
-		//       performing them for every set.
-		//
-
-		// TODO: Since this can't change, it doesn't have a dirty flag.
-		// This should be moved to initialization.
-		RigidHandle->SetSpatialIdx(Data->SpatialIdx);	//todo: this needs to only happen once during initialization
-		RigidHandle->SetUniqueIdx(Data->UniqueIdx);	//todo: this needs to only happen once during initialization
-		RigidHandle->SetUserData(Data->UserData);
-#if CHAOS_CHECKED
-		RigidHandle->SetDebugName(Data->DebugName);
-#endif
-
-		if (Data->DirtyFlags.IsDirty(Chaos::EParticleFlags::X))
-		{
-			RigidHandle->SetX(Data->X);
-			RigidHandle->SetP(Data->X);
-		}
-		if (Data->DirtyFlags.IsDirty(Chaos::EParticleFlags::R))
-		{
-			RigidHandle->SetR(Data->R);
-			RigidHandle->SetQ(Data->R);
-		}
-		if (Data->DirtyFlags.IsDirty(Chaos::EParticleFlags::Geometry))
-		{
-			RigidHandle->SetSharedGeometry(Data->Geometry);
-		}
-		if (Data->DirtyFlags.IsDirty(Chaos::EParticleFlags::V))
-		{
-			RigidHandle->SetV(Data->MV);
-		}
-		if (Data->DirtyFlags.IsDirty(Chaos::EParticleFlags::W))
-		{
-			RigidHandle->SetW(Data->MW);
-		}
-
-		RigidHandle->SetCenterOfMass(Data->MCenterOfMass);
-		RigidHandle->SetRotationOfMass(Data->MRotationOfMass);
-		RigidHandle->SetM(Data->MM);
-		RigidHandle->SetInvM(Data->MInvM);
-		RigidHandle->SetI(Data->MI);
-		RigidHandle->SetInvI(Data->MInvI);
-		RigidHandle->SetLinearEtherDrag(Data->MLinearEtherDrag);
-		RigidHandle->SetAngularEtherDrag(Data->MAngularEtherDrag);
-
-		if (Data->DirtyFlags.IsDirty(Chaos::EParticleFlags::F))
-		{
-			RigidHandle->SetF(Data->MF);
-		}
-		if (Data->DirtyFlags.IsDirty(Chaos::EParticleFlags::Torque))
-		{
-			RigidHandle->SetTorque(Data->MTorque);
-		}
-
-		if (Data->DirtyFlags.IsDirty(Chaos::EParticleFlags::LinearImpulse))
-		{
-			RigidHandle->SetLinearImpulse(Data->MLinearImpulse);
-		}
-
-		if (Data->DirtyFlags.IsDirty(Chaos::EParticleFlags::AngularImpulse))
-		{
-			RigidHandle->SetAngularImpulse(Data->MAngularImpulse);
-		}
-
-		if (Data->DirtyFlags.IsDirty(Chaos::EParticleFlags::ObjectState))
-		{
-			GetSolver()->GetEvolution()->SetParticleObjectState(RigidHandle, Data->MObjectState);
-		}
-		if (Data->DirtyFlags.IsDirty(Chaos::EParticleFlags::GravityEnabled))
-		{
-			GetSolver()->GetEvolution()->GetGravityForces().SetEnabled(*RigidHandle, Data->MGravityEnabled);
-		}
-
-		if (Data->DirtyFlags.IsDirty((int32)Chaos::EParticleFlags::X | (int32)Chaos::EParticleFlags::R | (int32)Chaos::EParticleFlags::V | (int32)Chaos::EParticleFlags::Geometry))
-		{
-			if (Data->Geometry && Data->Geometry->HasBoundingBox())
-			{
-				RigidHandle->SetHasBounds(true);
-				RigidHandle->SetLocalBounds(Data->Geometry->BoundingBox());
-				Chaos::TAABB<float, 3> WorldSpaceBox = Data->Geometry->BoundingBox().TransformedAABB(Chaos::TRigidTransform<float, 3>(Data->X, Data->R));
-				WorldSpaceBox.ThickenSymmetrically(Data->MV);
-				RigidHandle->SetWorldSpaceInflatedBounds(WorldSpaceBox);
-			}
-			else
-			{
-				//todo: compute bounds based on sample points
-			}
-		}
-
-		if (Data->DirtyFlags.IsDirty(Chaos::EParticleFlags::ShapeDisableCollision))
-		{
-			int32 CurrShape = 0;
-			for (const TUniquePtr<Chaos::TPerShapeData<Chaos::FReal, 3>>& Shape : RigidHandle->ShapesArray())
-			{
-				Shape->bDisable = Data->ShapeCollisionDisableFlags[CurrShape++];
-			}
-		}
-		if (Data->DirtyFlags.IsDirty(Chaos::EParticleFlags::CollisionTraceType))
-		{
-			int32 CurrShape = 0;
-			for (const TUniquePtr<Chaos::TPerShapeData<Chaos::FReal, 3>>& Shape : RigidHandle->ShapesArray())
-			{
-				Shape->CollisionTraceType = Data->CollisionTraceType[CurrShape++];
-			}
-		}
-		//if (Data->DirtyFlags.IsDirty(Chaos::EParticleFlags::ShapeSimData))
-		{
-			int32 CurrShape = 0;
-			for (const TUniquePtr<Chaos::TPerShapeData<Chaos::FReal, 3>>& Shape : RigidHandle->ShapesArray())
-			{
-				Shape->SimData = Data->ShapeSimData[CurrShape];
-				Shape->QueryData = Data->ShapeQueryData[CurrShape++];
-			}
-		}
-
-		{
-			int32 CurrShape = 0;
-			for (const TUniquePtr<Chaos::TPerShapeData<Chaos::FReal, 3>>& Shape : RigidHandle->ShapesArray())
-			{
-				Shape->Materials = Data->ShapeMaterials[CurrShape++];
-			}
-		}
-		
-		if(Data->MInitialized)
-		{
-			// wait for the first pass with nothing updated to claim its initialized
-			bInitialized = true;
-		}
-	}
-#endif
+	PushToPhysicsStateImp<Chaos::EParticleType::Rigid>(Manager,Handle->CastToRigidParticle(),DataIdx,Dirty,ShapesData,Evolution,bInitialized);
 }
 
 template< >
@@ -597,3 +314,16 @@ EPhysicsProxyType  FSingleParticlePhysicsProxy<Chaos::TPBDRigidParticle<float, 3
 template class CHAOSSOLVERS_API FSingleParticlePhysicsProxy< Chaos::TGeometryParticle<float, 3> >;
 template class CHAOSSOLVERS_API FSingleParticlePhysicsProxy< Chaos::TKinematicGeometryParticle<float, 3> >;
 template class CHAOSSOLVERS_API FSingleParticlePhysicsProxy< Chaos::TPBDRigidParticle<float,3> >;
+
+
+#define EVOLUTION_TRAIT(Traits)\
+template void FSingleParticlePhysicsProxy<Chaos::TGeometryParticle<Chaos::FReal,3>>::PushToPhysicsState(const Chaos::FDirtyPropertiesManager& Manager,\
+	int32 DataIdx,const Chaos::FDirtyProxy& Dirty,Chaos::FShapeDirtyData* ShapesData, Chaos::TPBDRigidsEvolutionGBF<Chaos::Traits>& Evolution);\
+\
+template void FSingleParticlePhysicsProxy<Chaos::TKinematicGeometryParticle<Chaos::FReal,3>>::PushToPhysicsState(const Chaos::FDirtyPropertiesManager& Manager,\
+	int32 DataIdx,const Chaos::FDirtyProxy& Dirty,Chaos::FShapeDirtyData* ShapesData,Chaos::TPBDRigidsEvolutionGBF<Chaos::Traits>& Evolution);\
+\
+template void FSingleParticlePhysicsProxy<Chaos::TPBDRigidParticle<Chaos::FReal,3>>::PushToPhysicsState(const Chaos::FDirtyPropertiesManager& Manager,\
+	int32 DataIdx,const Chaos::FDirtyProxy& Dirty,Chaos::FShapeDirtyData* ShapesData, Chaos::TPBDRigidsEvolutionGBF<Chaos::Traits>& Evolution);
+#include "Chaos/EvolutionTraits.inl"
+#undef EVOLUTION_TRAIT
\ No newline at end of file
diff --git a/Engine/Source/Runtime/Experimental/ChaosSolvers/Private/PhysicsProxy/SkeletalMeshPhysicsProxy.cpp b/Engine/Source/Runtime/Experimental/ChaosSolvers/Private/PhysicsProxy/SkeletalMeshPhysicsProxy.cpp
--- a/Engine/Source/Runtime/Experimental/ChaosSolvers/Private/PhysicsProxy/SkeletalMeshPhysicsProxy.cpp
+++ b/Engine/Source/Runtime/Experimental/ChaosSolvers/Private/PhysicsProxy/SkeletalMeshPhysicsProxy.cpp
@@ -5,6 +5,10 @@
 #include "ChaosStats.h"
 #include "PhysicsSolver.h"
 
+#ifndef ADD_TRAITS_TO_SKELTAL_MESH_PROXY
+#define ADD_TRAITS_TO_SKELTAL_MESH_PROXY 0
+#endif
+
 FSkeletalMeshPhysicsProxy::FSkeletalMeshPhysicsProxy(UObject* InOwner, const FInitFunc& InInitFunc)
 	: Base(InOwner)
 	, JointConstraints(Chaos::FPBDJointSolverSettings())
@@ -340,6 +344,8 @@ void FSkeletalMeshPhysicsProxy::CreateRigidBodyCallback(FParticlesType& Particle
 				BBox.Extents().Product());
 		}
 
+	ensure(false);	//GetSolver needs a trait, but we don't know it - this class changing anyway
+#if ADD_TRAITS_TO_SKELTAL_MESH_PROXY
 		// Add joints to the scene
 		if (JointConstraints.NumConstraints() > 0)
 		{
@@ -347,6 +353,7 @@ void FSkeletalMeshPhysicsProxy::CreateRigidBodyCallback(FParticlesType& Particle
 		}
 
 		bInitializedState = true;
+#endif
 	}
 }
 
@@ -379,6 +386,8 @@ void FSkeletalMeshPhysicsProxy::SyncBeforeDestroy()
 
 void FSkeletalMeshPhysicsProxy::OnRemoveFromScene()
 {
+	ensure(false);	//GetSolver needs a trait, but we don't know it - this class changing anyway
+#if ADD_TRAITS_TO_SKELTAL_MESH_PROXY
 	// Disable the particle we added
 	Chaos::FPhysicsSolver* CurrSolver = GetSolver();
 	if(CurrSolver && RigidBodyIds.Num())
@@ -402,6 +411,7 @@ void FSkeletalMeshPhysicsProxy::OnRemoveFromScene()
 		}
 #endif
 	}
+#endif
 }
 
 void FSkeletalMeshPhysicsProxy::BufferPhysicsResults()
@@ -409,6 +419,8 @@ void FSkeletalMeshPhysicsProxy::BufferPhysicsResults()
 	using namespace Chaos;
 	// PHYSICS THREAD
 
+	ensure(false);	//GetSolver needs a trait, but we don't know it - this class changing anyway
+#if ADD_TRAITS_TO_SKELTAL_MESH_PROXY
 	// @todo(ccaulfield)
 	if (GetSolver() != nullptr)
 	{
@@ -439,6 +451,7 @@ void FSkeletalMeshPhysicsProxy::BufferPhysicsResults()
 		}
 #endif
 	}
+#endif
 }
 
 void FSkeletalMeshPhysicsProxy::FlipBuffer()
diff --git a/Engine/Source/Runtime/Experimental/ChaosSolvers/Private/PhysicsProxy/StaticMeshPhysicsObject.cpp b/Engine/Source/Runtime/Experimental/ChaosSolvers/Private/PhysicsProxy/StaticMeshPhysicsObject.cpp
--- a/Engine/Source/Runtime/Experimental/ChaosSolvers/Private/PhysicsProxy/StaticMeshPhysicsObject.cpp
+++ b/Engine/Source/Runtime/Experimental/ChaosSolvers/Private/PhysicsProxy/StaticMeshPhysicsObject.cpp
@@ -10,6 +10,10 @@
 #include "Chaos/Serializable.h"
 #include "ChaosStats.h"
 
+#ifndef ADD_TRAITS_TO_STATIC_MESH_PROXY
+#define ADD_TRAITS_TO_STATIC_MESH_PROXY 0
+#endif
+
 FStaticMeshPhysicsProxy::FStaticMeshPhysicsProxy(UObject* InOwner, FCallbackInitFunc InInitFunc, FSyncDynamicFunc InSyncFunc)
 	: Base(InOwner)
 	, bInitializedState(false)
@@ -325,6 +329,8 @@ void FStaticMeshPhysicsProxy::AddForceCallback(FParticlesType& InParticles, cons
 
 void FStaticMeshPhysicsProxy::OnRemoveFromScene()
 {
+	ensure(false);	//GetSolver needs a trait, but we don't know it - this class changing anyway
+#if ADD_TRAITS_TO_STATIC_MESH_PROXY
 	// Disable the particle we added
 	Chaos::FPhysicsSolver* CurrSolver = GetSolver();
 
@@ -342,6 +348,7 @@ void FStaticMeshPhysicsProxy::OnRemoveFromScene()
 		CurrSolver->GetRigidClustering().GetTopLevelClusterParents().Remove(RigidBodyId);
 #endif
 	}
+#endif
 }
 
 void FStaticMeshPhysicsProxy::BufferPhysicsResults()
diff --git a/Engine/Source/Runtime/Experimental/ChaosSolvers/Public/EventDefaults.h b/Engine/Source/Runtime/Experimental/ChaosSolvers/Public/EventDefaults.h
--- a/Engine/Source/Runtime/Experimental/ChaosSolvers/Public/EventDefaults.h
+++ b/Engine/Source/Runtime/Experimental/ChaosSolvers/Public/EventDefaults.h
@@ -6,34 +6,39 @@
 
 namespace Chaos
 {
-	class FEventDefaults
+	template <typename Traits>
+	class TEventDefaults
 	{
 	public:
 
 		/**
 		 * Register default event types
 		 */
-		static void RegisterSystemEvents(FEventManager& EventManager);
+		static void RegisterSystemEvents(TEventManager<Traits>& EventManager);
 
 	private:
 
 		/**
 		 * Register collision event gathering function & data type
 		 */
-		static void RegisterCollisionEvent(FEventManager& EventManager);
+		static void RegisterCollisionEvent(TEventManager<Traits>& EventManager);
 
 		/**
 		 * Register breaking event gathering function & data type
 		 */
-		static void RegisterBreakingEvent(FEventManager& EventManager);
+		static void RegisterBreakingEvent(TEventManager<Traits>& EventManager);
 
 		/**
 		 * Register trailing event gathering function & data type
 		 */
-		static void RegisterTrailingEvent(FEventManager& EventManager);
+		static void RegisterTrailingEvent(TEventManager<Traits>& EventManager);
 
 
-		static void RegisterSleepingEvent(FEventManager& EventManager);
+		static void RegisterSleepingEvent(TEventManager<Traits>& EventManager);
 
 	};
+
+#define EVOLUTION_TRAIT(Trait) extern template class CHAOSSOLVERS_TEMPLATE_API TEventDefaults<Trait>;
+#include "Chaos/EvolutionTraits.inl"
+#undef EVOLUTION_TRAIT
 }
\ No newline at end of file
diff --git a/Engine/Source/Runtime/Experimental/ChaosSolvers/Public/EventManager.h b/Engine/Source/Runtime/Experimental/ChaosSolvers/Public/EventManager.h
--- a/Engine/Source/Runtime/Experimental/ChaosSolvers/Public/EventManager.h
+++ b/Engine/Source/Runtime/Experimental/ChaosSolvers/Public/EventManager.h
@@ -5,6 +5,9 @@
 #include "CoreMinimal.h"
 #include "Chaos/Framework/MultiBufferResource.h"
 #include "Chaos/Framework/PhysicsProxyBase.h"
+#include "Chaos/PBDRigidsEvolutionFwd.h"
+#include "Chaos/EvolutionTraits.h"
+#include "Chaos/Defines.h"
 
 namespace Chaos
 {
@@ -14,8 +17,6 @@ namespace Chaos
 	//#todo : warning if trying to add same event ID twice
 	//#todo : need sparse array for EventID -> EventContainer?
 
-	class FPBDRigidsSolver;
-
 	/**
 	 * Predefined System Event Types
 	 */
@@ -35,7 +36,7 @@ namespace Chaos
 	 */
 	class IEventHandler
 	{
-		template<typename PayloadType>
+		template<typename PayloadType, typename Traits>
 		friend class TEventContainer;
 
 	public:
@@ -86,10 +87,11 @@ namespace Chaos
 	/**
 	 * Interface for the injected producer function and associated data buffer
 	 */
-	class IEventContainer
+	template <typename Traits>
+	class TEventContainerBase
 	{
 	public:
-		virtual ~IEventContainer() {}
+		virtual ~TEventContainerBase() {}
 		/**
 		 * Register the delegate function that will handle the events on the game thread
 		 */
@@ -103,7 +105,7 @@ namespace Chaos
 		/*
 		 * Inject data from the physics solver into the producer side of the buffer
 		 */
-		virtual void InjectProducerData(const FPBDRigidsSolver* Solver) = 0;
+		virtual void InjectProducerData(const TPBDRigidsSolver<Traits>* Solver) = 0;
 
 		/**
 		 * Flips the buffer if the buffer type is double or triple
@@ -119,14 +121,14 @@ namespace Chaos
 	/**
 	 * Class that owns the injected producer function and its associated data buffer
 	 */
-	template<typename PayloadType>
-	class TEventContainer : public IEventContainer
+	template<typename PayloadType, typename Traits>
+	class TEventContainer : public TEventContainerBase<Traits>
 	{
 	public:
 		/**
 		 * Regular constructor
 		 */
-		TEventContainer(const Chaos::EMultiBufferMode& BufferMode, TFunction<void(const FPBDRigidsSolver* Solver, PayloadType& EventDataInOut)> InFunction)
+		TEventContainer(const Chaos::EMultiBufferMode& BufferMode, TFunction<void(const TPBDRigidsSolver<Traits>* Solver, PayloadType& EventDataInOut)> InFunction)
 			: InjectedFunction(InFunction)
 			, EventBuffer(Chaos::FMultiBufferFactory<PayloadType>::CreateBuffer(BufferMode))
 		{
@@ -135,7 +137,7 @@ namespace Chaos
 		/**
 		 * Copy constructor
 		 */
-		TEventContainer(TEventContainer& Other)
+		TEventContainer(TEventContainer<PayloadType, Traits>& Other)
 		{
 			InjectedFunction = Other.InjectedFunction;
 			EventBuffer = MoveTemp(Other.EventBuffer);
@@ -180,7 +182,7 @@ namespace Chaos
 		/*
 		 * Inject data from the physics solver into the producer side of the buffer
 		 */
-		virtual void InjectProducerData(const FPBDRigidsSolver* Solver)
+		virtual void InjectProducerData(const TPBDRigidsSolver<Traits>* Solver)
 		{
 			InjectedFunction(Solver, *EventBuffer->AccessProducerBuffer());
 		}
@@ -220,7 +222,7 @@ namespace Chaos
 		/**
 		 * The function that handles filling the event data buffer
 		 */
-		TFunction<void(const FPBDRigidsSolver* Solver, PayloadType& EventData)> InjectedFunction;
+		TFunction<void(const TPBDRigidsSolver<Traits>* Solver, PayloadType& EventData)> InjectedFunction;
 
 		/**
 		 * The data buffer that is filled by the producer and read by the consumer
@@ -236,18 +238,20 @@ namespace Chaos
 	/**
 	 * Pointer to event data buffer & injector functionality
 	 */
-	typedef IEventContainer* FEventContainerPtr;
-
+	template <typename Traits>
+	using TEventContainerBasePtr = TEventContainerBase<Traits>*;
 
-	class CHAOSSOLVERS_API FEventManager
+	template <typename Traits>
+	class TEventManager
 	{
-		friend class FPBDRigidsSolver;
+		template <typename Traits2>
+		friend class TPBDRigidsSolver;
 
 	public:
 
-		FEventManager(const Chaos::EMultiBufferMode& BufferModeIn) : BufferMode(BufferModeIn) {}
+		TEventManager(const Chaos::EMultiBufferMode& BufferModeIn) : BufferMode(BufferModeIn) {}
 
-		~FEventManager()
+		~TEventManager()
 		{
 			Reset();
 		}
@@ -269,10 +273,10 @@ namespace Chaos
 		 * Register a new event into the system, providing the function that will fill the producer side of the event buffer
 		 */
 		template<typename PayloadType>
-		void RegisterEvent(const FEventID& EventID, TFunction<void(const Chaos::FPBDRigidsSolver* Solver, PayloadType& EventData)> InFunction)
+		void RegisterEvent(const FEventID& EventID, TFunction<void(const Chaos::TPBDRigidsSolver<Traits>* Solver, PayloadType& EventData)> InFunction)
 		{
 			ContainerLock.WriteLock();
-			InternalRegisterInjector(EventID, new TEventContainer<PayloadType>(BufferMode, InFunction));
+			InternalRegisterInjector(EventID, new TEventContainer<PayloadType, Traits>(BufferMode, InFunction));
 			ContainerLock.WriteUnlock();
 		}
 
@@ -284,7 +288,7 @@ namespace Chaos
 		{
 			ContainerLock.ReadLock();
 
-			((TEventContainer<PayloadType>*)(EventContainers[EventID]))->DestroyStaleEvents(InFunction);
+			((TEventContainer<PayloadType, Traits>*)(EventContainers[EventID]))->DestroyStaleEvents(InFunction);
 			ContainerLock.ReadUnlock();
 		}
 
@@ -313,7 +317,7 @@ namespace Chaos
 		/**
 		 * Called by the solver to invoke the functions that fill the producer side of all the event data buffers
 		 */
-		void FillProducerData(const Chaos::FPBDRigidsSolver* Solver);
+		void FillProducerData(const Chaos::TPBDRigidsSolver<Traits>* Solver);
 
 		/**
 		 * Flips the event data buffer if it is of double or triple buffer type
@@ -332,13 +336,16 @@ namespace Chaos
 
 	private:
 
-		void InternalRegisterInjector(const FEventID& EventID, const FEventContainerPtr& Container);
+		void InternalRegisterInjector(const FEventID& EventID, const TEventContainerBasePtr<Traits>& Container);
 
 		Chaos::EMultiBufferMode BufferMode;			// specifies the buffer type to be constructed, single, double, triple
-		TArray<FEventContainerPtr> EventContainers;	// Array of event types
+		TArray<TEventContainerBasePtr<Traits>> EventContainers;	// Array of event types
 		FRWLock ResourceLock;
 		FRWLock ContainerLock;
 
 	};
 
-}
+#define EVOLUTION_TRAIT(Trait) extern template class CHAOSSOLVERS_TEMPLATE_API TEventManager<Trait>;
+#include "Chaos/EvolutionTraits.inl"
+#undef EVOLUTION_TRAIT
+}
\ No newline at end of file
diff --git a/Engine/Source/Runtime/Experimental/ChaosSolvers/Public/Framework/Dispatcher.h b/Engine/Source/Runtime/Experimental/ChaosSolvers/Public/Framework/Dispatcher.h
--- a/Engine/Source/Runtime/Experimental/ChaosSolvers/Public/Framework/Dispatcher.h
+++ b/Engine/Source/Runtime/Experimental/ChaosSolvers/Public/Framework/Dispatcher.h
@@ -15,6 +15,8 @@ namespace Chaos
 	class FPersistentPhysicsTask;
 	class FCommandList;
 	class FCommandListData;
+
+	class FPhysicsSolverBase;
 }
 
 namespace Chaos
@@ -56,7 +58,7 @@ namespace Chaos
 
 		using FGlobalCommand = TFunction<void()>;
 		using FTaskCommand = TFunction<void(FPersistentPhysicsTask*)>;
-		using FSolverCommand = TFunction<void(FPhysicsSolver*)>;
+		using FSolverCommand = TFunction<void()>;
 
 		/**
 		 * Immediate commands:
@@ -67,7 +69,7 @@ namespace Chaos
 		 */
 		virtual void EnqueueCommandImmediate(FGlobalCommand InCommand) = 0;
 		virtual void EnqueueCommandImmediate(FTaskCommand InCommand) = 0;
-		virtual void EnqueueCommandImmediate(FPhysicsSolver* InSolver, FSolverCommand InCommand) = 0;
+		virtual void EnqueueCommandImmediate(FPhysicsSolverBase* InSolver, FSolverCommand InCommand) = 0;
 
 		/**
 		 * Get the current threading mode for this dispatcher
diff --git a/Engine/Source/Runtime/Experimental/ChaosSolvers/Public/PBDRigidActiveParticlesBuffer.h b/Engine/Source/Runtime/Experimental/ChaosSolvers/Public/PBDRigidActiveParticlesBuffer.h
--- a/Engine/Source/Runtime/Experimental/ChaosSolvers/Public/PBDRigidActiveParticlesBuffer.h
+++ b/Engine/Source/Runtime/Experimental/ChaosSolvers/Public/PBDRigidActiveParticlesBuffer.h
@@ -4,11 +4,10 @@
 
 #include "Chaos/Framework/MultiBufferResource.h"
 #include "Chaos/ParticleHandle.h"
+#include "Chaos/PBDRigidsEvolutionFwd.h"
 
 namespace Chaos
 {
-	class FPBDRigidsSolver;
-
 	/**
 	 * Solver specific data buffered for use on Game thread
 	 */
diff --git a/Engine/Source/Runtime/Experimental/ChaosSolvers/Public/PBDRigidsSolver.h b/Engine/Source/Runtime/Experimental/ChaosSolvers/Public/PBDRigidsSolver.h
--- a/Engine/Source/Runtime/Experimental/ChaosSolvers/Public/PBDRigidsSolver.h
+++ b/Engine/Source/Runtime/Experimental/ChaosSolvers/Public/PBDRigidsSolver.h
@@ -24,13 +24,15 @@
 #include "PBDRigidActiveParticlesBuffer.h"
 #include "PhysicsProxy/SingleParticlePhysicsProxyFwd.h"
 #include "SolverEventFilters.h"
+#include "Chaos/EvolutionTraits.h"
+#include "Chaos/PBDRigidsEvolutionFwd.h"
+#include "ChaosSolversModule.h"
 
 class FPhysicsSolverAdvanceTask;
 class FPhysInterface_Chaos;
 
 class FSkeletalMeshPhysicsProxy;
 class FStaticMeshPhysicsProxy;
-class FGeometryCollectionPhysicsProxy;
 class FFieldSystemPhysicsProxy;
 
 #define PBDRIGID_PREALLOC_COUNT 1024
@@ -44,10 +46,8 @@ extern int32 ChaosSolverParticlePoolNumFrameUntilShrink;
 */
 namespace Chaos
 {
-	class AdvanceOneTimeStepTask;
 	class FPersistentPhysicsTask;
 	class FChaosArchive;
-	class FPBDRigidsSolver;
 	class FRewindData;
 
 	enum class ELockType : uint8
@@ -65,15 +65,16 @@ namespace Chaos
 	/**
 	*
 	*/
-	class CHAOSSOLVERS_API FPBDRigidsSolver : public FPhysicsSolverBase
+	template <typename Traits>
+	class TPBDRigidsSolver : public FPhysicsSolverBase
 	{
 
-		FPBDRigidsSolver(const EMultiBufferMode BufferingModeIn, UObject* InOwner);
+		TPBDRigidsSolver(const EMultiBufferMode BufferingModeIn, UObject* InOwner);
 
 	public:
 
 		// #BGTODO ensure no external callers directly deleting, make private and push everything through DestroySolver.
-		~FPBDRigidsSolver();
+		~TPBDRigidsSolver();
 
 		typedef FPhysicsSolverBase Super;
 
@@ -83,7 +84,9 @@ namespace Chaos
 
 		template<EThreadingMode Mode>
 		friend class FDispatcher;
-		friend class FEventDefaults;
+		
+		template <typename Traits2>
+		friend class TEventDefaults;
 
 		friend class FPhysInterface_Chaos;
 		friend class FPhysScene_ChaosInterface;
@@ -96,7 +99,7 @@ namespace Chaos
 
 		typedef Chaos::TGeometryParticle<float, 3> FParticle;
 		typedef Chaos::TGeometryParticleHandle<float, 3> FHandle;
-		typedef Chaos::FPBDRigidsEvolution FPBDRigidsEvolution;
+		typedef Chaos::TPBDRigidsEvolutionGBF<Traits> FPBDRigidsEvolution;
 
 		typedef TPBDRigidDynamicSpringConstraints<float, 3> FRigidDynamicSpringConstraints;
 		typedef TPBDPositionConstraints<float, 3> FPositionConstraints;
@@ -110,7 +113,12 @@ namespace Chaos
 		//
 
 		void ChangeBufferMode(Chaos::EMultiBufferMode InBufferMode);
-		TQueue<TFunction<void(FPBDRigidsSolver*)>, EQueueMode::Mpsc>& GetCommandQueue() { return CommandQueue; }
+
+		template <typename Lambda>
+		void EnqueueCommandImmediate(const Lambda& Func)
+		{
+			FChaosSolversModule::GetModule()->GetDispatcher()->EnqueueCommandImmediate(this, Func);
+		}
 
 
 		//
@@ -121,8 +129,8 @@ namespace Chaos
 		void UnregisterObject(Chaos::TGeometryParticle<float, 3>* GTParticle);
 
 		// TODO: Set up an interface for registering fields and geometry collections
-		void RegisterObject(FGeometryCollectionPhysicsProxy* InProxy);
-		bool UnregisterObject(FGeometryCollectionPhysicsProxy* InProxy);
+		void RegisterObject(TGeometryCollectionPhysicsProxy<Traits>* InProxy);
+		bool UnregisterObject(TGeometryCollectionPhysicsProxy<Traits>* InProxy);
 		void RegisterObject(FFieldSystemPhysicsProxy* InProxy);
 		bool UnregisterObject(FFieldSystemPhysicsProxy* InProxy);
 
@@ -157,7 +165,7 @@ namespace Chaos
 			{
 				InCallable(Obj);
 			}
-			for (FGeometryCollectionPhysicsProxy* Obj : GeometryCollectionPhysicsProxies)
+			for (TGeometryCollectionPhysicsProxy<Traits>* Obj : GeometryCollectionPhysicsProxies)
 			{
 				InCallable(Obj);
 			}
@@ -197,7 +205,7 @@ namespace Chaos
 			});
 			Chaos::PhysicsParallelFor(GeometryCollectionPhysicsProxies.Num(), [this, &InCallable](const int32 Index)
 			{
-				FGeometryCollectionPhysicsProxy* Obj = GeometryCollectionPhysicsProxies[Index];
+				TGeometryCollectionPhysicsProxy<Traits>* Obj = GeometryCollectionPhysicsProxies[Index];
 				InCallable(Obj);
 			});
 			Chaos::PhysicsParallelFor(FieldSystemPhysicsProxies.Num(), [this, &InCallable](const int32 Index)
@@ -221,7 +229,6 @@ namespace Chaos
 		bool Enabled() const { if (bEnabled) return this->IsSimulating(); return false; }
 		void SetEnabled(bool bEnabledIn) { bEnabled = bEnabledIn; }
 		bool HasActiveParticles() const { return !!GetNumPhysicsProxies(); }
-		bool HasPendingCommands() const { return !CommandQueue.IsEmpty(); }
 		FActiveParticlesBuffer* GetActiveParticlesBuffer() const { return MActiveParticlesBuffer.Get(); }
 
 		/**/
@@ -306,7 +313,7 @@ namespace Chaos
 		}
 
 		/**/
-		FEventManager* GetEventManager() { return MEventManager.Get(); }
+		TEventManager<Traits>* GetEventManager() { return MEventManager.Get(); }
 
 		/**/
 		FSolverEventFilters* GetEventFilters() { return MSolverEventFilters.Get(); }
@@ -323,7 +330,7 @@ namespace Chaos
 			return FieldSystemPhysicsProxies;
 		}
 
-		TArray<FGeometryCollectionPhysicsProxy*>& GetGeometryCollectionPhysicsProxies()
+		TArray<TGeometryCollectionPhysicsProxy<Traits>*>& GetGeometryCollectionPhysicsProxies()
 		{
 			return GeometryCollectionPhysicsProxies;
 		}
@@ -401,17 +408,12 @@ namespace Chaos
 
 		FParticlesType Particles;
 		TUniquePtr<FPBDRigidsEvolution> MEvolution;
-		TUniquePtr<FEventManager> MEventManager;
+		TUniquePtr<TEventManager<Traits>> MEventManager;
 		TUniquePtr<FSolverEventFilters> MSolverEventFilters;
 		TUniquePtr<FActiveParticlesBuffer> MActiveParticlesBuffer;
 		TMap<const Chaos::TGeometryParticleHandle<float, 3>*, TSet<IPhysicsProxyBase*> > MParticleToProxy;
 		TUniquePtr<FRewindData> MRewindData;
 
-		//
-		// Commands
-		//
-		TQueue<TFunction<void(FPBDRigidsSolver*)>, EQueueMode::Mpsc> CommandQueue;
-
 		//
 		// Proxies
 		//
@@ -421,7 +423,7 @@ namespace Chaos
 		TArray< FRigidParticlePhysicsProxy* > RigidParticlePhysicsProxies;
 		TArray< FSkeletalMeshPhysicsProxy* > SkeletalMeshPhysicsProxies; // dep
 		TArray< FStaticMeshPhysicsProxy* > StaticMeshPhysicsProxies; // dep
-		TArray< FGeometryCollectionPhysicsProxy* > GeometryCollectionPhysicsProxies;
+		TArray< TGeometryCollectionPhysicsProxy<Traits>* > GeometryCollectionPhysicsProxies;
 		TArray< FFieldSystemPhysicsProxy* > FieldSystemPhysicsProxies;
 
 		// Physics material mirrors for the solver. These should generally stay in sync with the global material list from
@@ -430,15 +432,12 @@ namespace Chaos
 		//
 		// There are two copies here to enable SQ to lock only the solvers that it needs to handle the material access during a query
 		// instead of having to lock the entire physics state of the runtime.
-		FRWLock QueryMaterialLock;
+		
 		THandleArray<FChaosPhysicsMaterial> QueryMaterials;
 		THandleArray<FChaosPhysicsMaterialMask> QueryMaterialMasks;
 		THandleArray<FChaosPhysicsMaterial> SimMaterials;
 		THandleArray<FChaosPhysicsMaterialMask> SimMaterialMasks;
 
-		template<ELockType>
-		friend struct TSolverQueryMaterialScope;
-
 		TUniquePtr<IBufferResource<FDirtyPropertiesManager>> DirtyPropertiesManager;
 	public:
 
@@ -528,7 +527,7 @@ namespace Chaos
 		TSolverQueryMaterialScope() = delete;
 
 
-		explicit TSolverQueryMaterialScope(FPBDRigidsSolver* InSolver)
+		explicit TSolverQueryMaterialScope(FPhysicsSolverBase* InSolver)
 			: Solver(InSolver)
 		{
 			check(Solver);
@@ -541,7 +540,7 @@ namespace Chaos
 		}
 
 	private:
-		FPBDRigidsSolver* Solver;
+		FPhysicsSolverBase* Solver;
 	};
 
 	template<>
@@ -549,7 +548,7 @@ namespace Chaos
 	{
 		TSolverQueryMaterialScope() = delete;
 
-		explicit TSolverQueryMaterialScope(FPBDRigidsSolver* InSolver)
+		explicit TSolverQueryMaterialScope(FPhysicsSolverBase* InSolver)
 			: Solver(InSolver)
 		{
 			check(Solver);
@@ -562,7 +561,11 @@ namespace Chaos
 		}
 
 	private:
-		FPBDRigidsSolver* Solver;
+		FPhysicsSolverBase* Solver;
 	};
 
+#define EVOLUTION_TRAIT(Trait) extern template class CHAOSSOLVERS_TEMPLATE_API TPBDRigidsSolver<Trait>;
+#include "Chaos/EvolutionTraits.inl"
+#undef EVOLUTION_TRAIT
+
 }; // namespace Chaos
diff --git a/Engine/Source/Runtime/Experimental/ChaosSolvers/Public/PhysicsProxy/FieldSystemPhysicsProxy.h b/Engine/Source/Runtime/Experimental/ChaosSolvers/Public/PhysicsProxy/FieldSystemPhysicsProxy.h
--- a/Engine/Source/Runtime/Experimental/ChaosSolvers/Public/PhysicsProxy/FieldSystemPhysicsProxy.h
+++ b/Engine/Source/Runtime/Experimental/ChaosSolvers/Public/PhysicsProxy/FieldSystemPhysicsProxy.h
@@ -5,6 +5,8 @@
 #include "Chaos/Framework/PhysicsProxy.h"
 #include "Field/FieldSystem.h"
 #include "Chaos/PBDPositionConstraints.h"
+#include "Chaos/Defines.h"
+#include "Chaos/EvolutionTraits.h"
 
 struct FKinematicProxy;
 
@@ -45,8 +47,9 @@ public:
 	 *	* EFieldPhysicsType::Field_PositionAnimated
 	 *	* EFieldPhysicsType::Field_DynamicConstraint
 	 */
+	template <typename Traits>
 	void FieldParameterUpdateCallback(
-		Chaos::FPhysicsSolver* InSolver, 
+		Chaos::TPBDRigidsSolver<Traits>* InSolver, 
 		FParticlesType& InParticles, 
 		Chaos::TArrayCollectionArray<float>& Strains, 
 		Chaos::TPBDPositionConstraints<float, 3>& PositionTarget, 
@@ -61,14 +64,15 @@ public:
 	 *	* EFieldPhysicsType::Field_LinearForce
 	 *	* EFieldPhysicsType::Field_AngularTorque
 	 */
+	template <typename Traits>
 	void FieldForcesUpdateCallback(
-		Chaos::FPhysicsSolver* InSolver, 
+		Chaos::TPBDRigidsSolver<Traits>* InSolver, 
 		FParticlesType& Particles, 
 		Chaos::TArrayCollectionArray<FVector> & Force, 
 		Chaos::TArrayCollectionArray<FVector> & Torque, 
 		const float Time);
 
-	void BufferCommand(Chaos::FPhysicsSolver* InSolver, const FFieldSystemCommand& InCommand);
+	void BufferCommand(Chaos::FPhysicsSolverBase* InSolver, const FFieldSystemCommand& InCommand);
 
 	// Inactive Callbacks
 	void ParameterUpdateCallback(FParticlesType& InParticles, const float InTime){}
@@ -106,16 +110,42 @@ public:
 	 * the active clusters are set in the \p IndicesArray.
 	 */
 
+	template <typename Traits>
 	static void GetParticleHandles(
 		TArray<Chaos::TGeometryParticleHandle<float,3>*>& Handles,
-		const Chaos::FPhysicsSolver* RigidSolver,
+		const Chaos::TPBDRigidsSolver<Traits>* RigidSolver,
 		const EFieldResolutionType ResolutionType,
 		const bool bForce = true);
 
 private:
 
-	TArray<FFieldSystemCommand>* GetSolverCommandList(const Chaos::FPhysicsSolver* InSolver);
+	TArray<FFieldSystemCommand>* GetSolverCommandList(const Chaos::FPhysicsSolverBase* InSolver);
 
 	FCriticalSection CommandLock;
-	TMap<Chaos::FPhysicsSolver*, TArray<FFieldSystemCommand>*> Commands;
+	TMap<Chaos::FPhysicsSolverBase*, TArray<FFieldSystemCommand>*> Commands;
 };
+
+#define EVOLUTION_TRAIT(Traits)\
+extern template CHAOSSOLVERS_API void FFieldSystemPhysicsProxy::FieldParameterUpdateCallback(\
+		Chaos::TPBDRigidsSolver<Chaos::Traits>* InSolver, \
+		FParticlesType& InParticles, \
+		Chaos::TArrayCollectionArray<float>& Strains, \
+		Chaos::TPBDPositionConstraints<float, 3>& PositionTarget, \
+		TMap<int32, int32>& PositionTargetedParticles, \
+		const float InTime);\
+\
+extern template CHAOSSOLVERS_API void FFieldSystemPhysicsProxy::FieldForcesUpdateCallback(\
+		Chaos::TPBDRigidsSolver<Chaos::Traits>* InSolver, \
+		FParticlesType& Particles, \
+		Chaos::TArrayCollectionArray<FVector> & Force, \
+		Chaos::TArrayCollectionArray<FVector> & Torque, \
+		const float Time);\
+\
+extern template CHAOSSOLVERS_API void FFieldSystemPhysicsProxy::GetParticleHandles(\
+		TArray<Chaos::TGeometryParticleHandle<float,3>*>& Handles,\
+		const Chaos::TPBDRigidsSolver<Chaos::Traits>* RigidSolver,\
+		const EFieldResolutionType ResolutionType,\
+		const bool bForce);\
+
+#include "Chaos/EvolutionTraits.inl"
+#undef EVOLUTION_TRAIT
diff --git a/Engine/Source/Runtime/Experimental/ChaosSolvers/Public/PhysicsProxy/GeometryCollectionPhysicsProxy.h b/Engine/Source/Runtime/Experimental/ChaosSolvers/Public/PhysicsProxy/GeometryCollectionPhysicsProxy.h
--- a/Engine/Source/Runtime/Experimental/ChaosSolvers/Public/PhysicsProxy/GeometryCollectionPhysicsProxy.h
+++ b/Engine/Source/Runtime/Experimental/ChaosSolvers/Public/PhysicsProxy/GeometryCollectionPhysicsProxy.h
@@ -15,6 +15,8 @@
 #include "Chaos/ParticleHandleFwd.h"
 #include "Containers/Array.h"
 #include "PBDRigidsSolver.h"
+#include "Chaos/EvolutionTraits.h"
+#include "Chaos/Defines.h"
 
 namespace Chaos
 {
@@ -50,10 +52,11 @@ private:
  * Class to manage sharing data between the game thread and the simulation thread 
  * (which may not be different than the game thread) for a \c FGeometryDynamicCollection.
  */
-class CHAOSSOLVERS_API FGeometryCollectionPhysicsProxy : public TPhysicsProxy<FGeometryCollectionPhysicsProxy, FStubGeometryCollectionData>
+template <typename Traits>
+class TGeometryCollectionPhysicsProxy : public TPhysicsProxy<TGeometryCollectionPhysicsProxy<Traits>, FStubGeometryCollectionData>
 {
 public:
-	typedef TPhysicsProxy<FGeometryCollectionPhysicsProxy, FStubGeometryCollectionData> Base;
+	typedef TPhysicsProxy<TGeometryCollectionPhysicsProxy<Traits>, FStubGeometryCollectionData> Base;
 	typedef FCollisionStructureManager::FSimplicial FSimplicial;
 	typedef Chaos::TPBDRigidParticleHandle<float, 3> FParticleHandle;
 	typedef Chaos::TPBDRigidClusteredParticleHandle<float, 3> FClusterHandle;
@@ -62,8 +65,9 @@ public:
 	using FInitFunc = TFunction<void(FSimulationParameters&)>;
 	using FCacheSyncFunc = TFunction<void(const FGeometryCollectionResults&)>;
 	using FFinalSyncFunc = TFunction<void(const FRecordedTransformTrack&)>;
+	using IPhysicsProxyBase::GetSolver;
 
-	FGeometryCollectionPhysicsProxy() = delete;
+	TGeometryCollectionPhysicsProxy() = delete;
 	/**
 	 * \p InOwner
 	 * \p InDynamicCollection game thread owned geometry collection.
@@ -71,7 +75,7 @@ public:
 	 * \p InCacheSyncFunc callback invoked from \c PullFromPhysicsState().
 	 * \p InFinalSyncFunc callback invoked from \c SyncBeforeDestory().
 	 */
-	FGeometryCollectionPhysicsProxy(
+	TGeometryCollectionPhysicsProxy(
 		UObject* InOwner, 
 		FGeometryDynamicCollection& GameThreadCollection,
 		const FSimulationParameters& SimulationParameters,
@@ -81,7 +85,7 @@ public:
 		FCacheSyncFunc InCacheSyncFunc, 
 		FFinalSyncFunc InFinalSyncFunc  ,
 		const Chaos::EMultiBufferMode BufferMode=Chaos::EMultiBufferMode::TripleGuarded);
-	virtual ~FGeometryCollectionPhysicsProxy();
+	virtual ~TGeometryCollectionPhysicsProxy();
 
 	/**
 	 * Construct \c PTDynamicCollection, copying attributes from the game thread, 
@@ -96,8 +100,8 @@ public:
 	 * Called by solver command registered by \c FPBDRigidsSolver::RegisterObject().
 	 */
 	void InitializeBodiesPT(
-		Chaos::FPBDRigidsSolver* RigidsSolver,
-		Chaos::FPBDRigidsSolver::FParticlesType& Particles);
+		Chaos::TPBDRigidsSolver<Traits>* RigidsSolver,
+		typename Chaos::TPBDRigidsSolver<Traits>::FParticlesType& Particles);
 
 	/** */
 	static void InitializeDynamicCollection(FGeometryDynamicCollection& DynamicCollection, const FGeometryCollection& RestCollection, const FSimulationParameters& Params);
@@ -148,7 +152,7 @@ public:
 	EPhysicsProxyType ConcreteType() { return EPhysicsProxyType::GeometryCollectionType; }
 
 	void SyncBeforeDestroy();
-	void OnRemoveFromSolver(Chaos::FPBDRigidsSolver *RBDSolver);
+	void OnRemoveFromSolver(Chaos::TPBDRigidsSolver<Traits> *RBDSolver);
 	void OnRemoveFromScene();
 
 	void SetCollisionParticlesPerObjectFraction(float CollisionParticlesPerObjectFractionIn) 
@@ -170,7 +174,7 @@ public:
 	static void InitializeSharedCollisionStructures(Chaos::FErrorReporter& ErrorReporter, FGeometryCollection& RestCollection, const FSharedSimulationParameters& SharedParams);
 	static void InitRemoveOnFracture(FGeometryCollection& RestCollection, const FSharedSimulationParameters& SharedParams);
 
-	void FieldForcesUpdateCallback(Chaos::FPhysicsSolver* InSolver, FParticlesType& Particles, Chaos::TArrayCollectionArray<FVector>& Force, Chaos::TArrayCollectionArray<FVector>& Torque, const float Time);
+	void FieldForcesUpdateCallback(Chaos::TPBDRigidsSolver<Traits>* InSolver, FParticlesType& Particles, Chaos::TArrayCollectionArray<FVector>& Force, Chaos::TArrayCollectionArray<FVector>& Torque, const float Time);
 	void ParameterUpdateCallback(FParticlesType& InParticles, const float InTime);
 
 	void UpdateKinematicBodiesCallback(const FParticlesType& InParticles, const float InDt, const float InTime, FKinematicProxy& InKinematicProxy) {}
@@ -203,7 +207,7 @@ protected:
 		TArray<Chaos::TGeometryParticleHandle<float, 3>*>& Handles,
 		TArray<FVector>& Samples,
 		TArray<ContextIndex>& SampleIndices,
-		const Chaos::FPhysicsSolver* RigidSolver, 
+		const Chaos::TPBDRigidsSolver<Traits>* RigidSolver, 
 		EFieldResolutionType ResolutionType, 
 		bool bForce);
 
@@ -220,6 +224,7 @@ protected:
 
 private:
 
+	using FSolver = Chaos::TPBDRigidsSolver<Traits>;
 	FSimulationParameters Parameters;
 	TArray<FFieldSystemCommand> Commands;
 
@@ -296,5 +301,9 @@ private:
 
 };
 
+#define EVOLUTION_TRAIT(Traits) extern template CHAOSSOLVERS_TEMPLATE_API class TGeometryCollectionPhysicsProxy<Chaos::Traits>;
+#include "Chaos/EvolutionTraits.inl"
+#undef EVOLUTION_TRAIT
+
 CHAOSSOLVERS_API Chaos::TTriangleMesh<float>* CreateTriangleMesh(const int32 FaceStart,const int32 FaceCount,const TManagedArray<bool>& Visible,const TManagedArray<FIntVector>& Indices);
 CHAOSSOLVERS_API void BuildSimulationData(Chaos::FErrorReporter& ErrorReporter, FGeometryCollection& GeometryCollection, const FSharedSimulationParameters& SharedParams);
diff --git a/Engine/Source/Runtime/Experimental/ChaosSolvers/Public/PhysicsProxy/SingleParticlePhysicsProxy.h b/Engine/Source/Runtime/Experimental/ChaosSolvers/Public/PhysicsProxy/SingleParticlePhysicsProxy.h
--- a/Engine/Source/Runtime/Experimental/ChaosSolvers/Public/PhysicsProxy/SingleParticlePhysicsProxy.h
+++ b/Engine/Source/Runtime/Experimental/ChaosSolvers/Public/PhysicsProxy/SingleParticlePhysicsProxy.h
@@ -8,11 +8,16 @@
 #include "Chaos/PBDPositionConstraints.h"
 #include "Chaos/ParticleHandle.h"
 #include "PhysicsCoreTypes.h"
+#include "Chaos/Defines.h"
+#include "Chaos/EvolutionTraits.h"
 
 namespace Chaos
 {
 	template<typename T, int d>
 	class TGeometryParticle;
+
+	template <typename Traits>
+	class TPBDRigidsEvolutionGBF;
 }
 
 class FInitialState
@@ -123,7 +128,8 @@ public:
 	/**/
 	void FlipBuffer() { BufferedData->FlipProducer(); }
 
-	void PushToPhysicsState(const Chaos::FDirtyPropertiesManager& Manager,int32 DataIdx,const Chaos::FDirtyProxy& Dirty,Chaos::FShapeDirtyData* ShapesData);
+	template <typename Traits>
+	void PushToPhysicsState(const Chaos::FDirtyPropertiesManager& Manager,int32 DataIdx,const Chaos::FDirtyProxy& Dirty,Chaos::FShapeDirtyData* ShapesData, Chaos::TPBDRigidsEvolutionGBF<Traits>& Evolution);
 
 	/**/
 	void ClearAccumulatedData();
@@ -184,8 +190,9 @@ CHAOSSOLVERS_API void FSingleParticlePhysicsProxy<Chaos::TGeometryParticle<float
 template< >
 bool FSingleParticlePhysicsProxy<Chaos::TGeometryParticle<float, 3>>::IsDirty();
 
-template< >
-void FSingleParticlePhysicsProxy<Chaos::TGeometryParticle<float,3>>::PushToPhysicsState(const Chaos::FDirtyPropertiesManager& Manager,int32 DataIdx,const Chaos::FDirtyProxy& Dirty,Chaos::FShapeDirtyData* ShapesData);
+template<>
+template<typename Traits>
+void FSingleParticlePhysicsProxy<Chaos::TGeometryParticle<float,3>>::PushToPhysicsState(const Chaos::FDirtyPropertiesManager& Manager,int32 DataIdx,const Chaos::FDirtyProxy& Dirty,Chaos::FShapeDirtyData* ShapesData, Chaos::TPBDRigidsEvolutionGBF<Traits>& Evolution);
 
 template< >
 EPhysicsProxyType FSingleParticlePhysicsProxy<Chaos::TGeometryParticle<float, 3>>::ConcreteType();
@@ -198,8 +205,9 @@ CHAOSSOLVERS_API void FSingleParticlePhysicsProxy<Chaos::TGeometryParticle<float
 
 // TKinematicGeometryParticle specialization prototypes
 
-template< >
-void FSingleParticlePhysicsProxy<Chaos::TKinematicGeometryParticle<float,3>>::PushToPhysicsState(const Chaos::FDirtyPropertiesManager& Manager,int32 DataIdx,const Chaos::FDirtyProxy& Dirty,Chaos::FShapeDirtyData* ShapesData);
+template<>
+template<typename Traits>
+void FSingleParticlePhysicsProxy<Chaos::TKinematicGeometryParticle<float,3>>::PushToPhysicsState(const Chaos::FDirtyPropertiesManager& Manager,int32 DataIdx,const Chaos::FDirtyProxy& Dirty,Chaos::FShapeDirtyData* ShapesData, Chaos::TPBDRigidsEvolutionGBF<Traits>& Evolution);
 
 template< >
 void FSingleParticlePhysicsProxy<Chaos::TKinematicGeometryParticle<float, 3>>::ClearAccumulatedData();
@@ -236,8 +244,9 @@ CHAOSSOLVERS_API void FSingleParticlePhysicsProxy<Chaos::TPBDRigidParticle<float
 template< >
 bool FSingleParticlePhysicsProxy<Chaos::TPBDRigidParticle<float, 3>>::IsDirty();
 
-template< >
-void FSingleParticlePhysicsProxy<Chaos::TPBDRigidParticle<float,3>>::PushToPhysicsState(const Chaos::FDirtyPropertiesManager& Manager,int32 DataIdx,const Chaos::FDirtyProxy& Dirty,Chaos::FShapeDirtyData* ShapesData);
+template<>
+template<typename Traits>
+void FSingleParticlePhysicsProxy<Chaos::TPBDRigidParticle<float,3>>::PushToPhysicsState(const Chaos::FDirtyPropertiesManager& Manager,int32 DataIdx,const Chaos::FDirtyProxy& Dirty,Chaos::FShapeDirtyData* ShapesData, Chaos::TPBDRigidsEvolutionGBF<Traits>& Evolution);
 
 template< >
 EPhysicsProxyType FSingleParticlePhysicsProxy<Chaos::TPBDRigidParticle<float, 3>>::ConcreteType();
@@ -252,3 +261,16 @@ CHAOSSOLVERS_API void FSingleParticlePhysicsProxy<Chaos::TPBDRigidParticle<float
 extern template class FSingleParticlePhysicsProxy< Chaos::TGeometryParticle<float, 3> >;
 extern template class FSingleParticlePhysicsProxy< Chaos::TKinematicGeometryParticle<float, 3> >;
 extern template class FSingleParticlePhysicsProxy< Chaos::TPBDRigidParticle<float,3> >;
+
+
+#define EVOLUTION_TRAIT(Traits)\
+extern template CHAOSSOLVERS_API void FSingleParticlePhysicsProxy<Chaos::TGeometryParticle<Chaos::FReal,3>>::PushToPhysicsState(const Chaos::FDirtyPropertiesManager& Manager,\
+	int32 DataIdx,const Chaos::FDirtyProxy& Dirty,Chaos::FShapeDirtyData* ShapesData, Chaos::TPBDRigidsEvolutionGBF<Chaos::Traits>& Evolution);\
+\
+extern template CHAOSSOLVERS_API void FSingleParticlePhysicsProxy<Chaos::TKinematicGeometryParticle<Chaos::FReal,3>>::PushToPhysicsState(const Chaos::FDirtyPropertiesManager& Manager,\
+	int32 DataIdx,const Chaos::FDirtyProxy& Dirty,Chaos::FShapeDirtyData* ShapesData,Chaos::TPBDRigidsEvolutionGBF<Chaos::Traits>& Evolution);\
+\
+extern template CHAOSSOLVERS_API void FSingleParticlePhysicsProxy<Chaos::TPBDRigidParticle<Chaos::FReal,3>>::PushToPhysicsState(const Chaos::FDirtyPropertiesManager& Manager,\
+	int32 DataIdx,const Chaos::FDirtyProxy& Dirty,Chaos::FShapeDirtyData* ShapesData, Chaos::TPBDRigidsEvolutionGBF<Chaos::Traits>& Evolution);
+#include "Chaos/EvolutionTraits.inl"
+#undef EVOLUTION_TRAIT
diff --git a/Engine/Source/Runtime/Experimental/GeometryCollectionEngine/Private/GeometryCollection/GeometryCollectionDebugDrawComponent.cpp b/Engine/Source/Runtime/Experimental/GeometryCollectionEngine/Private/GeometryCollection/GeometryCollectionDebugDrawComponent.cpp
--- a/Engine/Source/Runtime/Experimental/GeometryCollectionEngine/Private/GeometryCollection/GeometryCollectionDebugDrawComponent.cpp
+++ b/Engine/Source/Runtime/Experimental/GeometryCollectionEngine/Private/GeometryCollection/GeometryCollectionDebugDrawComponent.cpp
@@ -737,7 +737,7 @@ void UGeometryCollectionDebugDrawComponent::DebugDrawChaosTick()
 	const FGeometryCollectionPhysicsProxy* const PhysicsProxy = GeometryCollectionComponent->GetPhysicsProxy();
 	if (PhysicsProxy)
 	{
-		ParticlesData.Sync(PhysicsProxy->GetSolver(), RigidBodyIds);
+		ParticlesData.Sync(PhysicsProxy->GetSolver<Chaos::FPhysicsSolver>(), RigidBodyIds);
 	}
 
 	// Visualize single rigid body
diff --git a/Engine/Source/Runtime/Experimental/GeometryCollectionEngine/Private/GeometryCollection/GeometryCollectionParticlesData.cpp b/Engine/Source/Runtime/Experimental/GeometryCollectionEngine/Private/GeometryCollection/GeometryCollectionParticlesData.cpp
--- a/Engine/Source/Runtime/Experimental/GeometryCollectionEngine/Private/GeometryCollection/GeometryCollectionParticlesData.cpp
+++ b/Engine/Source/Runtime/Experimental/GeometryCollectionEngine/Private/GeometryCollection/GeometryCollectionParticlesData.cpp
@@ -54,7 +54,7 @@ void TGeometryCollectionParticlesData<T, d>::Sync(Chaos::FPhysicsSolver* Solver,
 			Chaos::IDispatcher* const PhysicsDispatcher = ChaosModule->GetDispatcher();
 			check(PhysicsDispatcher);
 
-			PhysicsDispatcher->EnqueueCommandImmediate(Solver, [this, &RigidBodyIds](Chaos::FPhysicsSolver* InSolver)
+			Solver->EnqueueCommandImmediate([this, &RigidBodyIds, InSolver=Solver]()
 			{
 				// Iterate through all data
 				FData& Data = BufferedData.GetPhysicsDataForWrite();
diff --git a/Engine/Source/Runtime/Experimental/GeometryCollectionEngine/Public/GeometryCollection/GeometryCollectionComponent.h b/Engine/Source/Runtime/Experimental/GeometryCollectionEngine/Public/GeometryCollection/GeometryCollectionComponent.h
--- a/Engine/Source/Runtime/Experimental/GeometryCollectionEngine/Public/GeometryCollection/GeometryCollectionComponent.h
+++ b/Engine/Source/Runtime/Experimental/GeometryCollectionEngine/Public/GeometryCollection/GeometryCollectionComponent.h
@@ -21,12 +21,12 @@
 #include "Chaos/ChaosGameplayEventDispatcher.h"
 #include "Chaos/ChaosNotifyHandlerInterface.h"
 #include "Chaos/ChaosSolverComponentTypes.h"
+#include "Chaos/PBDRigidsEvolutionFwd.h"
 
 #include "GeometryCollectionComponent.generated.h"
 
 struct FGeometryCollectionConstantData;
 struct FGeometryCollectionDynamicData;
-class FGeometryCollectionPhysicsProxy;
 class UGeometryCollectionComponent;
 class UBoxComponent;
 class UGeometryCollectionCache;
