Fix CIS: missing include

diff --git a/Engine/Source/Runtime/Experimental/Chaos/Public/Chaos/Framework/PhysicsSolverBase.h b/Engine/Source/Runtime/Experimental/Chaos/Public/Chaos/Framework/PhysicsSolverBase.h
--- a/Engine/Source/Runtime/Experimental/Chaos/Public/Chaos/Framework/PhysicsSolverBase.h
+++ b/Engine/Source/Runtime/Experimental/Chaos/Public/Chaos/Framework/PhysicsSolverBase.h
@@ -8,6 +8,7 @@
 #include "Chaos/Framework/PhysicsProxyBase.h"
 #include "Chaos/ParticleDirtyFlags.h"
 #include "Async/ParallelFor.h"
+#include "Containers/Queue.h"
 
 class FChaosSolversModule;
 class FPhysicsSolverAdvanceTask;
