Fixing up collisions between solver traits changes and previous changes made in 4.25
#rnx

diff --git a/Engine/Source/Runtime/Experimental/ChaosSolverEngine/Private/Chaos/ChaosSolverActor.cpp b/Engine/Source/Runtime/Experimental/ChaosSolverEngine/Private/Chaos/ChaosSolverActor.cpp
--- a/Engine/Source/Runtime/Experimental/ChaosSolverEngine/Private/Chaos/ChaosSolverActor.cpp
+++ b/Engine/Source/Runtime/Experimental/ChaosSolverEngine/Private/Chaos/ChaosSolverActor.cpp
@@ -273,25 +273,24 @@ void AChaosSolverActor::BeginPlay()
 	Super::BeginPlay();
 
 	Solver->EnqueueCommandImmediate(
-		[InSolver = Solver
+		[ InSolver = Solver
 		, InTimeStepMultiplier = TimeStepMultiplier
-			, InCollisionIterations = CollisionIterations
-			, InPushOutIterations = PushOutIterations
-			, InPushOutPairIterations = PushOutPairIterations
-			, InClusterConnectionFactor = ClusterConnectionFactor
-			, InClusterUnionConnectionType = ClusterUnionConnectionType
-			, InDoGenerateCollisionData = DoGenerateCollisionData
-			, InDoGenerateBreakingData = DoGenerateBreakingData
-			, InDoGenerateTrailingData = DoGenerateTrailingData
-			, InCollisionFilterSettings = CollisionFilterSettings
-			, InBreakingFilterSettings = BreakingFilterSettings
-			, InTrailingFilterSettings = TrailingFilterSettings
-			, InHasFloor = bHasFloor
-			, InFloorHeight = FloorHeight
-			, InMassScale = MassScale
-			, InGenerateContactGraph = bGenerateContactGraph]
-		]
-		(Chaos::FPhysicsSolver* InSolver)
+		, InCollisionIterations = CollisionIterations
+		, InPushOutIterations = PushOutIterations
+		, InPushOutPairIterations = PushOutPairIterations
+		, InClusterConnectionFactor = ClusterConnectionFactor
+		, InClusterUnionConnectionType = ClusterUnionConnectionType
+		, InDoGenerateCollisionData = DoGenerateCollisionData
+		, InDoGenerateBreakingData = DoGenerateBreakingData
+		, InDoGenerateTrailingData = DoGenerateTrailingData
+		, InCollisionFilterSettings = CollisionFilterSettings
+		, InBreakingFilterSettings = BreakingFilterSettings
+		, InTrailingFilterSettings = TrailingFilterSettings
+		, InHasFloor = bHasFloor
+		, InFloorHeight = FloorHeight
+		, InMassScale = MassScale
+		, InGenerateContactGraph = bGenerateContactGraph]
+		()
 		{
 #if TODO_REIMPLEMENT_SOLVER_SETTINGS_ACCESSORS
 			InSolver->SetTimeStepMultiplier(InTimeStepMultiplier);
@@ -421,32 +420,32 @@ void AChaosSolverActor::PostEditChangeProperty(struct FPropertyChangedEvent& Pro
 #if TODO_REIMPLEMENT_TIMESTEP_MULTIPLIER
 			if (PropertyChangedEvent.Property->GetFName() == GET_MEMBER_NAME_CHECKED(AChaosSolverActor, TimeStepMultiplier))
 			{
-			Solver->EnqueueCommandImmediate([InSolver=Solver, InTimeStepMultiplier = TimeStepMultiplier]
-			()
+				Solver->EnqueueCommandImmediate([InSolver=Solver, InTimeStepMultiplier = TimeStepMultiplier]
+				()
 				{
 					InSolver->SetTimeStepMultiplier(InTimeStepMultiplier);
 				});
 			}
 			else if (PropertyChangedEvent.Property->GetFName() == GET_MEMBER_NAME_CHECKED(AChaosSolverActor, CollisionIterations))
 			{
-			Solver->EnqueueCommandImmediate([InSolver=Solver, InCollisionIterations = CollisionIterations]
-			()
+				Solver->EnqueueCommandImmediate([InSolver=Solver, InCollisionIterations = CollisionIterations]
+				()
 				{
 					InSolver->SetIterations(InCollisionIterations);
 				});
 			}
 			else if (PropertyChangedEvent.Property->GetFName() == GET_MEMBER_NAME_CHECKED(AChaosSolverActor, PushOutIterations))
 			{
-			Solver->EnqueueCommandImmediate([InSolver=Solver, InPushOutIterations = PushOutIterations]
-			()
+				Solver->EnqueueCommandImmediate([InSolver=Solver, InPushOutIterations = PushOutIterations]
+				()
 				{
 					InSolver->SetPushOutIterations(InPushOutIterations);
 				});
 			}
 			else if (PropertyChangedEvent.Property->GetFName() == GET_MEMBER_NAME_CHECKED(AChaosSolverActor, PushOutPairIterations))
 			{
-			Solver->EnqueueCommandImmediate([InSolver=Solver, InPushOutPairIterations = PushOutPairIterations]
-			()
+				Solver->EnqueueCommandImmediate([InSolver=Solver, InPushOutPairIterations = PushOutPairIterations]
+				()
 				{
 					InSolver->SetPushOutPairIterations(InPushOutPairIterations);
 				});
@@ -455,63 +454,61 @@ void AChaosSolverActor::PostEditChangeProperty(struct FPropertyChangedEvent& Pro
 #endif
 			if (PropertyChangedEvent.Property->GetFName() == GET_MEMBER_NAME_CHECKED(AChaosSolverActor, DoGenerateCollisionData))
 			{
-			Solver->EnqueueCommandImmediate([InSolver=Solver, InDoGenerateCollisionData = DoGenerateCollisionData]
-			()
+				Solver->EnqueueCommandImmediate([InSolver = Solver, InDoGenerateCollisionData = DoGenerateCollisionData]
+				()
 				{
 					InSolver->SetGenerateCollisionData(InDoGenerateCollisionData);
 				});
 			}
 			else if (PropertyChangedEvent.MemberProperty->GetFName() == GET_MEMBER_NAME_CHECKED(AChaosSolverActor, CollisionFilterSettings))
 			{
-			Solver->EnqueueCommandImmediate([InSolver=Solver, InCollisionFilterSettings = CollisionFilterSettings]
-			()
+				Solver->EnqueueCommandImmediate([InSolver=Solver, InCollisionFilterSettings = CollisionFilterSettings]
+				()
 				{
 					InSolver->SetCollisionFilterSettings(InCollisionFilterSettings);
 				});
 			}
 			else if (PropertyChangedEvent.Property->GetFName() == GET_MEMBER_NAME_CHECKED(AChaosSolverActor, DoGenerateBreakingData))
 			{
-			Solver->EnqueueCommandImmediate([InSolver=Solver, InDoGenerateBreakingData = DoGenerateBreakingData]
-			()
+				Solver->EnqueueCommandImmediate([InSolver=Solver, InDoGenerateBreakingData = DoGenerateBreakingData]
+				()
 				{
 					InSolver->SetGenerateBreakingData(InDoGenerateBreakingData);
 				});
 			}
 			else if (PropertyChangedEvent.MemberProperty->GetFName() == GET_MEMBER_NAME_CHECKED(AChaosSolverActor, BreakingFilterSettings))
 			{
-			Solver->EnqueueCommandImmediate([InSolver=Solver, InBreakingFilterSettings = BreakingFilterSettings]
-			()
+				Solver->EnqueueCommandImmediate([InSolver=Solver, InBreakingFilterSettings = BreakingFilterSettings]
+				()
 				{
 					InSolver->SetBreakingFilterSettings(InBreakingFilterSettings);
 				});
 			}
 			else if (PropertyChangedEvent.Property->GetFName() == GET_MEMBER_NAME_CHECKED(AChaosSolverActor, DoGenerateTrailingData))
 			{
-			Solver->EnqueueCommandImmediate([InSolver=Solver, InDoGenerateTrailingData = DoGenerateTrailingData]
-			()
+				Solver->EnqueueCommandImmediate([InSolver=Solver, InDoGenerateTrailingData = DoGenerateTrailingData]
+				()
 				{
 					InSolver->SetGenerateTrailingData(InDoGenerateTrailingData);
 				});
 			}
 			else if (PropertyChangedEvent.MemberProperty->GetFName() == GET_MEMBER_NAME_CHECKED(AChaosSolverActor, TrailingFilterSettings))
 			{
-			Solver->EnqueueCommandImmediate([InSolver=Solver, InTrailingFilterSettings = TrailingFilterSettings]
-			()
+				Solver->EnqueueCommandImmediate([InSolver=Solver, InTrailingFilterSettings = TrailingFilterSettings]
+				()
 				{
 					InSolver->SetTrailingFilterSettings(InTrailingFilterSettings);
 				});
 			}
 			else if (PropertyChangedEvent.Property->GetFName() == GET_MEMBER_NAME_CHECKED(AChaosSolverActor, bHasFloor))
 			{
-
 			}
 			else if (PropertyChangedEvent.Property->GetFName() == GET_MEMBER_NAME_CHECKED(AChaosSolverActor, FloorHeight))
 			{
 			}
 			else if (PropertyChangedEvent.Property->GetFName() == GET_MEMBER_NAME_CHECKED(AChaosSolverActor, bGenerateContactGraph))
 			{
-				PhysDispatcher->EnqueueCommandImmediate(Solver, [InGenerateContactGraph = bGenerateContactGraph]
-				(Chaos::FPhysicsSolver* InSolver)
+				Solver->EnqueueCommandImmediate([InSolver = Solver, InGenerateContactGraph = bGenerateContactGraph]()
 				{
 				});
 			}
@@ -519,8 +516,7 @@ void AChaosSolverActor::PostEditChangeProperty(struct FPropertyChangedEvent& Pro
 #if TODO_REIMPLEMENT_TIMESTEP_MULTIPLIER
 			else if (PropertyChangedEvent.Property->GetFName() == GET_MEMBER_NAME_CHECKED(AChaosSolverActor, MassScale))
 			{
-			Solver->EnqueueCommandImmediate([InSolver=Solver, InMassScale = MassScale]
-			()
+				Solver->EnqueueCommandImmediate([InSolver=Solver, InMassScale = MassScale]()
 				{
 					InSolver->SetMassScale(InMassScale);
 				});
diff --git a/Engine/Source/Runtime/Experimental/ChaosSolvers/Private/PBDRigidsSolver.cpp b/Engine/Source/Runtime/Experimental/ChaosSolvers/Private/PBDRigidsSolver.cpp
--- a/Engine/Source/Runtime/Experimental/ChaosSolvers/Private/PBDRigidsSolver.cpp
+++ b/Engine/Source/Runtime/Experimental/ChaosSolvers/Private/PBDRigidsSolver.cpp
@@ -381,7 +381,7 @@ namespace Chaos
 			UE_LOG(LogPBDRigidsSolver, Verbose, TEXT("TPBDRigidsSolver::UnregisterObject() ~ Dequeue"));
 
 				// Generally need to remove stale events for particles that no longer exist
-				Solver->GetEventManager()->template ClearEvents<FCollisionEventData>(EEventType::Collision, [InProxy]
+				GetEventManager()->template ClearEvents<FCollisionEventData>(EEventType::Collision, [InProxy]
 				(FCollisionEventData& EventDataInOut)
 				{
 					Chaos::FCollisionDataArray const& CollisionData = EventDataInOut.CollisionData.AllCollisionsArray;
diff --git a/Engine/Source/Runtime/Experimental/GeometryCollectionSimulationCore/Public/GeometryCollection/GeometryCollectionSimulationCoreTypes.h b/Engine/Source/Runtime/Experimental/GeometryCollectionSimulationCore/Public/GeometryCollection/GeometryCollectionSimulationCoreTypes.h
--- a/Engine/Source/Runtime/Experimental/GeometryCollectionSimulationCore/Public/GeometryCollection/GeometryCollectionSimulationCoreTypes.h
+++ b/Engine/Source/Runtime/Experimental/GeometryCollectionSimulationCore/Public/GeometryCollection/GeometryCollectionSimulationCoreTypes.h
@@ -8,6 +8,7 @@
 #include "GeometryCollectionSimulationTypes.h"
 #include "Chaos/ClusterCreationParameters.h"
 #include "Chaos/CollisionFilterData.h"
+#include "Chaos/PBDRigidsEvolutionFwd.h"
 
 class FGeometryCollection;
 class FGeometryDynamicCollection;
@@ -219,8 +220,6 @@ struct FTrailingDataSimulationParameters
 	float TrailingMinVolumeThreshold;
 };
 
-class FGeometryCollectionPhysicsProxy;
-
 struct FSimulationParameters
 {
 	FSimulationParameters()
@@ -331,7 +330,8 @@ struct FSimulationParameters
 	EObjectStateTypeEnum ObjectType;
 
 private:
-	friend class ::FGeometryCollectionPhysicsProxy;
+	template<typename>
+	friend class TGeometryCollectionPhysicsProxy;
 	TUniquePtr<Chaos::FChaosPhysicsMaterial> PhysicalMaterialOwner; // can be null
 public:
 	Chaos::TSerializablePtr<Chaos::FChaosPhysicsMaterial> PhysicalMaterial; // may or may not point to PhysicalMaterialOwner
