Fix CIS: missing include in non-unity

diff --git a/Engine/Source/Runtime/Experimental/Chaos/Public/Chaos/Defines.h b/Engine/Source/Runtime/Experimental/Chaos/Public/Chaos/Defines.h
--- a/Engine/Source/Runtime/Experimental/Chaos/Public/Chaos/Defines.h
+++ b/Engine/Source/Runtime/Experimental/Chaos/Public/Chaos/Defines.h
@@ -1,6 +1,8 @@
 // Copyright Epic Games, Inc. All Rights Reserved.
 #pragma once
 
+#include "HAL/Platform.h"
+
 #ifndef CHAOS_TEMPLATE_API
 #if PLATFORM_MAC || PLATFORM_LINUX
 #define CHAOS_TEMPLATE_API CHAOS_API
