Missed file from solver traits fixes
#rnx

diff --git a/Engine/Source/Runtime/Experimental/ChaosSolvers/Private/PhysicsProxy/GeometryCollectionPhysicsProxy.cpp b/Engine/Source/Runtime/Experimental/ChaosSolvers/Private/PhysicsProxy/GeometryCollectionPhysicsProxy.cpp
--- a/Engine/Source/Runtime/Experimental/ChaosSolvers/Private/PhysicsProxy/GeometryCollectionPhysicsProxy.cpp
+++ b/Engine/Source/Runtime/Experimental/ChaosSolvers/Private/PhysicsProxy/GeometryCollectionPhysicsProxy.cpp
@@ -1131,7 +1131,7 @@ TGeometryCollectionPhysicsProxy<Traits>::BuildClusters(
 
 	// Populate bounds as we didn't pass a shared implicit to PopulateSimulatedParticle this will have been skipped, now that we have the full cluster we can build it
 	if(Parent->Geometry() && Parent->Geometry()->HasBoundingBox())
-		{
+	{
 		Parent->SetHasBounds(true);
 		Parent->SetLocalBounds(Parent->Geometry()->BoundingBox());
 		const Chaos::TAABB<float, 3>& LocalBounds = Parent->LocalBounds();
@@ -1139,8 +1139,8 @@ TGeometryCollectionPhysicsProxy<Traits>::BuildClusters(
 		const Chaos::TAABB<float, 3> TransformedBBox = LocalBounds.TransformedAABB(Xf);
 		Parent->SetWorldSpaceInflatedBounds(TransformedBBox);
 
-		GetSolver()->GetEvolution()->DirtyParticle(*Parent);
-		}
+		GetSolver<FSolver>()->GetEvolution()->DirtyParticle(*Parent);
+	}
 
 	return Parent;
 }
@@ -1587,7 +1587,7 @@ void TGeometryCollectionPhysicsProxy<Traits>::BufferPhysicsResults()
 
 							ProxyElementHandle->X() = ParticleToWorld.GetTranslation();
 							ProxyElementHandle->R() = ParticleToWorld.GetRotation();
-							GetSolver()->GetEvolution()->DirtyParticle(*ProxyElementHandle);
+							GetSolver<FSolver>()->GetEvolution()->DirtyParticle(*ProxyElementHandle);
 						}
 		}
 	}
