Fix CIS: move dll export after class

diff --git a/Engine/Source/Runtime/Experimental/ChaosSolvers/Public/PhysicsProxy/GeometryCollectionPhysicsProxy.h b/Engine/Source/Runtime/Experimental/ChaosSolvers/Public/PhysicsProxy/GeometryCollectionPhysicsProxy.h
--- a/Engine/Source/Runtime/Experimental/ChaosSolvers/Public/PhysicsProxy/GeometryCollectionPhysicsProxy.h
+++ b/Engine/Source/Runtime/Experimental/ChaosSolvers/Public/PhysicsProxy/GeometryCollectionPhysicsProxy.h
@@ -301,7 +301,7 @@ private:
 
 };
 
-#define EVOLUTION_TRAIT(Traits) extern template CHAOSSOLVERS_TEMPLATE_API class TGeometryCollectionPhysicsProxy<Chaos::Traits>;
+#define EVOLUTION_TRAIT(Traits) extern template class CHAOSSOLVERS_TEMPLATE_API TGeometryCollectionPhysicsProxy<Chaos::Traits>;
 #include "Chaos/EvolutionTraits.inl"
 #undef EVOLUTION_TRAIT
 
