Make it so you can create different solvers, unit tests automatically use all solvers

Benn.Gallagher

diff --git a/Engine/Source/Programs/HeadlessChaos/Private/HeadlessChaosApp.cpp b/Engine/Source/Programs/HeadlessChaos/Private/HeadlessChaosApp.cpp
--- a/Engine/Source/Programs/HeadlessChaos/Private/HeadlessChaosApp.cpp
+++ b/Engine/Source/Programs/HeadlessChaos/Private/HeadlessChaosApp.cpp
@@ -213,18 +213,6 @@ TEST(Perf, PerfTests)
 	SUCCEED();
 }
 
-
-TEST(ChaosSolver, ChaosSolverTests)
-{
-	ChaosTest::SingleParticleProxySingleThreadTest<float>();
-	ChaosTest::CommandListTest<float>();
-	ChaosTest::SingleParticleProxyWakeEventPropergationTest<float>();
-}
-TEST(ChaosSolver, DISABLED_ChaosSolverTests)
-{
-	ChaosTest::SingleParticleProxyTaskGraphTest<float>();
-}
-
 TEST(Handles, FrameworkTests)
 {
 	ChaosTest::Handles::HandleArrayTest<float>();
diff --git a/Engine/Source/Programs/HeadlessChaos/Private/HeadlessChaosTestRewind.cpp b/Engine/Source/Programs/HeadlessChaos/Private/HeadlessChaosTestRewind.cpp
--- a/Engine/Source/Programs/HeadlessChaos/Private/HeadlessChaosTestRewind.cpp
+++ b/Engine/Source/Programs/HeadlessChaos/Private/HeadlessChaosTestRewind.cpp
@@ -18,7 +18,8 @@ namespace ChaosTest {
 
     using namespace Chaos;
 
-	void TickSolverHelper(FChaosSolversModule* Module, FPhysicsSolver* Solver, FReal Dt = 1.0)
+	template <typename TSolver>
+	void TickSolverHelper(FChaosSolversModule* Module, TSolver* Solver, FReal Dt = 1.0)
 	{
 		Solver->PushPhysicsState(Module->GetDispatcher());
 		FPhysicsSolverAdvanceTask AdvanceTask(Solver,Dt);
@@ -28,7 +29,7 @@ namespace ChaosTest {
 		Solver->UpdateGameThreadStructures();
 	}
 
-	GTEST_TEST(RewindTest,MovingGeomChange)
+	TYPED_TEST(AllTraits, RewindTest_MovingGeomChange)
 	{
 		auto Sphere = TSharedPtr<FImplicitObject,ESPMode::ThreadSafe>(new TSphere<float,3>(TVector<float,3>(0),10));
 		auto Box = TSharedPtr<FImplicitObject,ESPMode::ThreadSafe>(new TBox<float,3>(FVec3(0),FVec3(1)));
@@ -38,7 +39,7 @@ namespace ChaosTest {
 		Module->ChangeThreadingMode(EChaosThreadingMode::SingleThread);
 
 		// Make a solver
-		FPhysicsSolver* Solver = Module->CreateSolver(nullptr,ESolverFlags::Standalone);
+		auto* Solver = Module->CreateSolver<TypeParam>(nullptr,ESolverFlags::Standalone);
 		Solver->SetEnabled(true);
 
 		Solver->EnableRewindCapture(20);
@@ -112,7 +113,7 @@ namespace ChaosTest {
 	}
 
 
-	GTEST_TEST(RewindTest,AddForce)
+	TYPED_TEST(AllTraits, RewindTest_AddForce)
 	{
 		auto Sphere = TSharedPtr<FImplicitObject,ESPMode::ThreadSafe>(new TSphere<float,3>(TVector<float,3>(0),10));
 
@@ -120,7 +121,7 @@ namespace ChaosTest {
 		Module->ChangeThreadingMode(EChaosThreadingMode::SingleThread);
 
 		// Make a solver
-		FPhysicsSolver* Solver = Module->CreateSolver(nullptr,ESolverFlags::Standalone);
+		auto* Solver = Module->CreateSolver<TypeParam>(nullptr,ESolverFlags::Standalone);
 		Solver->SetEnabled(true);
 
 		Solver->EnableRewindCapture(20);
@@ -156,7 +157,7 @@ namespace ChaosTest {
 		Module->DestroySolver(Solver);
 	}
 
-	GTEST_TEST(RewindTest,IntermittentForce)
+	TYPED_TEST(AllTraits, RewindTest_IntermittentForce)
 	{
 		auto Sphere = TSharedPtr<FImplicitObject,ESPMode::ThreadSafe>(new TSphere<float,3>(TVector<float,3>(0),10));
 
@@ -164,7 +165,7 @@ namespace ChaosTest {
 		Module->ChangeThreadingMode(EChaosThreadingMode::SingleThread);
 
 		// Make a solver
-		FPhysicsSolver* Solver = Module->CreateSolver(nullptr,ESolverFlags::Standalone);
+		auto* Solver = Module->CreateSolver<TypeParam>(nullptr,ESolverFlags::Standalone);
 		Solver->SetEnabled(true);
 
 		Solver->EnableRewindCapture(20);
@@ -219,7 +220,7 @@ namespace ChaosTest {
 		Module->DestroySolver(Solver);
 	}
 
-	GTEST_TEST(RewindTest,IntermittentGeomChange)
+	TYPED_TEST(AllTraits, RewindTest_IntermittentGeomChange)
 	{
 		auto Sphere = TSharedPtr<FImplicitObject,ESPMode::ThreadSafe>(new TSphere<float,3>(TVector<float,3>(0),10));
 		auto Box = TSharedPtr<FImplicitObject,ESPMode::ThreadSafe>(new TBox<float,3>(FVec3(0),FVec3(1)));
@@ -229,7 +230,7 @@ namespace ChaosTest {
 		Module->ChangeThreadingMode(EChaosThreadingMode::SingleThread);
 
 		// Make a solver
-		FPhysicsSolver* Solver = Module->CreateSolver(nullptr,ESolverFlags::Standalone);
+		auto* Solver = Module->CreateSolver<TypeParam>(nullptr,ESolverFlags::Standalone);
 		Solver->SetEnabled(true);
 
 		Solver->EnableRewindCapture(20);
@@ -293,7 +294,7 @@ namespace ChaosTest {
 		Module->DestroySolver(Solver);
 	}
 
-	GTEST_TEST(RewindTest,FallingObjectWithTeleport)
+	TYPED_TEST(AllTraits, RewindTest_FallingObjectWithTeleport)
 	{
 		auto Sphere = TSharedPtr<FImplicitObject,ESPMode::ThreadSafe>(new TSphere<float,3>(TVector<float,3>(0),10));
 
@@ -301,7 +302,7 @@ namespace ChaosTest {
 		Module->ChangeThreadingMode(EChaosThreadingMode::SingleThread);
 
 		// Make a solver
-		FPhysicsSolver* Solver = Module->CreateSolver(nullptr,ESolverFlags::Standalone);
+		auto* Solver = Module->CreateSolver<TypeParam>(nullptr,ESolverFlags::Standalone);
 		Solver->SetEnabled(true);
 
 		Solver->EnableRewindCapture(20);
@@ -349,7 +350,7 @@ namespace ChaosTest {
 		Module->DestroySolver(Solver);
 	}
 
-	GTEST_TEST(RewindTest,ApplyRewind)
+	TYPED_TEST(AllTraits, RewindTest_ApplyRewind)
 	{
 		auto Sphere = TSharedPtr<FImplicitObject,ESPMode::ThreadSafe>(new TSphere<float,3>(TVector<float,3>(0),10));
 
@@ -357,7 +358,7 @@ namespace ChaosTest {
 		Module->ChangeThreadingMode(EChaosThreadingMode::SingleThread);
 
 		// Make a solver
-		FPhysicsSolver* Solver = Module->CreateSolver(nullptr,ESolverFlags::Standalone);
+		auto* Solver = Module->CreateSolver<TypeParam>(nullptr,ESolverFlags::Standalone);
 		Solver->SetEnabled(true);
 
 		Solver->EnableRewindCapture(20);
@@ -420,7 +421,7 @@ namespace ChaosTest {
 		Module->DestroySolver(Solver);
 	}
 
-	GTEST_TEST(RewindTest,Remove)
+	TYPED_TEST(AllTraits, RewindTest_Remove)
 	{
 		auto Sphere = TSharedPtr<FImplicitObject,ESPMode::ThreadSafe>(new TSphere<float,3>(TVector<float,3>(0),10));
 
@@ -428,7 +429,7 @@ namespace ChaosTest {
 		Module->ChangeThreadingMode(EChaosThreadingMode::SingleThread);
 
 		// Make a solver
-		FPhysicsSolver* Solver = Module->CreateSolver(nullptr,ESolverFlags::Standalone);
+		auto* Solver = Module->CreateSolver<TypeParam>(nullptr,ESolverFlags::Standalone);
 		Solver->SetEnabled(true);
 
 		Solver->EnableRewindCapture(20);
@@ -471,7 +472,7 @@ namespace ChaosTest {
 		Module->DestroySolver(Solver);
 	}
 
-	GTEST_TEST(RewindTest,BufferLimit)
+	TYPED_TEST(AllTraits, RewindTest_BufferLimit)
 	{
 		auto Sphere = TSharedPtr<FImplicitObject,ESPMode::ThreadSafe>(new TSphere<float,3>(TVector<float,3>(0),10));
 
@@ -479,7 +480,7 @@ namespace ChaosTest {
 		Module->ChangeThreadingMode(EChaosThreadingMode::SingleThread);
 
 		// Make a solver
-		FPhysicsSolver* Solver = Module->CreateSolver(nullptr,ESolverFlags::Standalone);
+		auto* Solver = Module->CreateSolver<TypeParam>(nullptr,ESolverFlags::Standalone);
 		Solver->SetEnabled(true);
 
 		Solver->EnableRewindCapture(5);
@@ -535,7 +536,7 @@ namespace ChaosTest {
 		Module->DestroySolver(Solver);
 	}
 
-	GTEST_TEST(RewindTest,NumDirty)
+	TYPED_TEST(AllTraits, RewindTest_NumDirty)
 	{
 		auto Sphere = TSharedPtr<FImplicitObject,ESPMode::ThreadSafe>(new TSphere<float,3>(TVector<float,3>(0),10));
 
@@ -543,7 +544,7 @@ namespace ChaosTest {
 		Module->ChangeThreadingMode(EChaosThreadingMode::SingleThread);
 
 		// Make a solver
-		FPhysicsSolver* Solver = Module->CreateSolver(nullptr,ESolverFlags::Standalone);
+		auto* Solver = Module->CreateSolver<TypeParam>(nullptr,ESolverFlags::Standalone);
 		Solver->SetEnabled(true);
 
 		Solver->EnableRewindCapture(5);
@@ -594,7 +595,7 @@ namespace ChaosTest {
 		Module->DestroySolver(Solver);
 	}
 
-	GTEST_TEST(RewindTest,Resim)
+	TYPED_TEST(AllTraits, RewindTest_Resim)
 	{
 		auto Sphere = TSharedPtr<FImplicitObject,ESPMode::ThreadSafe>(new TSphere<float,3>(TVector<float,3>(0),10));
 
@@ -602,7 +603,7 @@ namespace ChaosTest {
 		Module->ChangeThreadingMode(EChaosThreadingMode::SingleThread);
 
 		// Make a solver
-		FPhysicsSolver* Solver = Module->CreateSolver(nullptr,ESolverFlags::Standalone);
+		auto* Solver = Module->CreateSolver<TypeParam>(nullptr,ESolverFlags::Standalone);
 		Solver->SetEnabled(true);
 
 		Solver->EnableRewindCapture(5);
@@ -704,7 +705,7 @@ namespace ChaosTest {
 		Module->DestroySolver(Solver);
 	}
 
-	GTEST_TEST(RewindTest,ResimDesyncAfterMissingTeleport)
+	TYPED_TEST(AllTraits, RewindTest_ResimDesyncAfterMissingTeleport)
 	{
 		auto Sphere = TSharedPtr<FImplicitObject,ESPMode::ThreadSafe>(new TSphere<float,3>(TVector<float,3>(0),10));
 
@@ -712,7 +713,7 @@ namespace ChaosTest {
 		Module->ChangeThreadingMode(EChaosThreadingMode::SingleThread);
 
 		// Make a solver
-		FPhysicsSolver* Solver = Module->CreateSolver(nullptr,ESolverFlags::Standalone);
+		auto* Solver = Module->CreateSolver<TypeParam>(nullptr,ESolverFlags::Standalone);
 		Solver->SetEnabled(true);
 
 		Solver->EnableRewindCapture(7);
@@ -779,7 +780,7 @@ namespace ChaosTest {
 		Module->DestroySolver(Solver);
 	}
 
-	GTEST_TEST(RewindTest,ResimDesyncAfterChangingMass)
+	TYPED_TEST(AllTraits, RewindTest_ResimDesyncAfterChangingMass)
 	{
 		auto Sphere = TSharedPtr<FImplicitObject,ESPMode::ThreadSafe>(new TSphere<float,3>(TVector<float,3>(0),10));
 
@@ -787,7 +788,7 @@ namespace ChaosTest {
 		Module->ChangeThreadingMode(EChaosThreadingMode::SingleThread);
 
 		// Make a solver
-		FPhysicsSolver* Solver = Module->CreateSolver(nullptr,ESolverFlags::Standalone);
+		auto* Solver = Module->CreateSolver<TypeParam>(nullptr,ESolverFlags::Standalone);
 		Solver->SetEnabled(true);
 
 		Solver->EnableRewindCapture(7);
@@ -847,7 +848,7 @@ namespace ChaosTest {
 		Module->DestroySolver(Solver);
 	}
 
-	GTEST_TEST(RewindTest,DISABLED_DesyncFromPT)
+	TYPED_TEST(AllTraits, DISABLED_RewindTest_DesyncFromPT)
 	{
 		//We want to detect when sim results change
 		//Detecting output of position and velocity is expensive and hard to track
@@ -859,7 +860,7 @@ namespace ChaosTest {
 		Module->ChangeThreadingMode(EChaosThreadingMode::SingleThread);
 
 		// Make a solver
-		FPhysicsSolver* Solver = Module->CreateSolver(nullptr,ESolverFlags::Standalone);
+		auto* Solver = Module->CreateSolver<TypeParam>(nullptr,ESolverFlags::Standalone);
 		Solver->SetEnabled(true);
 		Solver->EnableRewindCapture(7);
 
@@ -913,7 +914,7 @@ namespace ChaosTest {
 		Module->DestroySolver(Solver);
 	}
 
-	GTEST_TEST(RewindTest,DeltaTimeRecord)
+	TYPED_TEST(AllTraits, RewindTest_DeltaTimeRecord)
 	{
 		auto Sphere = TSharedPtr<FImplicitObject,ESPMode::ThreadSafe>(new TSphere<float,3>(TVector<float,3>(0),10));
 
@@ -921,7 +922,7 @@ namespace ChaosTest {
 		Module->ChangeThreadingMode(EChaosThreadingMode::SingleThread);
 
 		// Make a solver
-		FPhysicsSolver* Solver = Module->CreateSolver(nullptr,ESolverFlags::Standalone);
+		auto* Solver = Module->CreateSolver<TypeParam>(nullptr,ESolverFlags::Standalone);
 		Solver->SetEnabled(true);
 
 		Solver->EnableRewindCapture(7);
@@ -959,7 +960,7 @@ namespace ChaosTest {
 		Module->DestroySolver(Solver);
 	}
 
-	GTEST_TEST(RewindTest,ResimDesyncFromChangeForce)
+	TYPED_TEST(AllTraits, RewindTest_ResimDesyncFromChangeForce)
 	{
 		auto Sphere = TSharedPtr<FImplicitObject,ESPMode::ThreadSafe>(new TSphere<float,3>(TVector<float,3>(0),10));
 
@@ -967,7 +968,7 @@ namespace ChaosTest {
 		Module->ChangeThreadingMode(EChaosThreadingMode::SingleThread);
 
 		// Make a solver
-		FPhysicsSolver* Solver = Module->CreateSolver(nullptr,ESolverFlags::Standalone);
+		auto* Solver = Module->CreateSolver<TypeParam>(nullptr,ESolverFlags::Standalone);
 		Solver->SetEnabled(true);
 
 		Solver->EnableRewindCapture(7);
diff --git a/Engine/Source/Programs/HeadlessChaos/Private/HeadlessChaosTestSolverCommandList.cpp b/Engine/Source/Programs/HeadlessChaos/Private/HeadlessChaosTestSolverCommandList.cpp
--- a/Engine/Source/Programs/HeadlessChaos/Private/HeadlessChaosTestSolverCommandList.cpp
+++ b/Engine/Source/Programs/HeadlessChaos/Private/HeadlessChaosTestSolverCommandList.cpp
@@ -14,13 +14,14 @@
 #include "Modules/ModuleManager.h"
 #include "Framework/PhysicsTickTask.h"
 #include "Framework/CommandBuffer.h"
+#include "HeadlessChaosTestUtility.h"
 
 
 namespace ChaosTest {
 
     using namespace Chaos;
 
-	template<class T>
+	template<typename Traits, class T>
 	void CommandListTest()
 	{
 		auto Sphere = TSharedPtr<FImplicitObject, ESPMode::ThreadSafe>(new TSphere<float, 3>(TVector<float, 3>(0), 10));
@@ -31,7 +32,7 @@ namespace ChaosTest {
 		Module->ChangeThreadingMode(EChaosThreadingMode::DedicatedThread);
 
 		// Make a solver
-		Chaos::FPhysicsSolver* Solver = FChaosSolversModule::GetModule()->CreateSolver(nullptr, ESolverFlags::Standalone);
+		auto* Solver = FChaosSolversModule::GetModule()->CreateSolver<Traits>(nullptr, ESolverFlags::Standalone);
 		Solver->SetEnabled(true);
 
 		// Make a particle
@@ -73,5 +74,9 @@ namespace ChaosTest {
 		FChaosSolversModule::GetModule()->DestroySolver(Solver);
 	}
 
-	template void CommandListTest<float>();
+	
+	TYPED_TEST(AllTraits,CommandListTests)
+	{
+		ChaosTest::CommandListTest<TypeParam,float>();
+	}
 }
diff --git a/Engine/Source/Programs/HeadlessChaos/Private/HeadlessChaosTestSolverProxies.cpp b/Engine/Source/Programs/HeadlessChaos/Private/HeadlessChaosTestSolverProxies.cpp
--- a/Engine/Source/Programs/HeadlessChaos/Private/HeadlessChaosTestSolverProxies.cpp
+++ b/Engine/Source/Programs/HeadlessChaos/Private/HeadlessChaosTestSolverProxies.cpp
@@ -21,7 +21,7 @@ namespace ChaosTest {
 
     using namespace Chaos;
 
-	template<class T>
+	template<typename Traits, class T>
 	void SingleParticleProxySingleThreadTest()
 	{
 		auto Sphere = TSharedPtr<FImplicitObject, ESPMode::ThreadSafe>(new TSphere<float, 3>(TVector<float, 3>(0), 10));
@@ -30,7 +30,7 @@ namespace ChaosTest {
 		Module->ChangeThreadingMode(EChaosThreadingMode::SingleThread);
 
 		// Make a solver
-		Chaos::FPhysicsSolver* Solver = Module->CreateSolver(nullptr, ESolverFlags::Standalone);
+		auto* Solver = Module->CreateSolver<Traits>(nullptr, ESolverFlags::Standalone);
 		Solver->SetEnabled(true);
 
 		// Make a particle
@@ -73,7 +73,7 @@ namespace ChaosTest {
 	template void SingleParticleProxySingleThreadTest<float>();
 
 	
-	template<class T>
+	template<typename Traits, class T>
 	void SingleParticleProxyTaskGraphTest()
 	{
 		//
@@ -89,7 +89,7 @@ namespace ChaosTest {
 		Module->ChangeThreadingMode(EChaosThreadingMode::DedicatedThread);
 
 		// Make a solver
-		Chaos::FPhysicsSolver* Solver = Module->CreateSolver(nullptr, ESolverFlags::Standalone);
+		auto* Solver = Module->CreateSolver<Traits>(nullptr, ESolverFlags::Standalone);
 		Solver->SetEnabled(true);
 
 		// Make a particle
@@ -133,10 +133,8 @@ namespace ChaosTest {
 		Module->DestroySolver(Solver);
 	}
 
-	template void SingleParticleProxyTaskGraphTest<float>();
 
-
-	template<class T>
+	template<typename Traits, class T>
 	void SingleParticleProxyWakeEventPropergationTest()
 	{
 		auto Sphere = TSharedPtr<FImplicitObject, ESPMode::ThreadSafe>(new TSphere<float, 3>(TVector<float, 3>(0), 10));
@@ -145,7 +143,7 @@ namespace ChaosTest {
 		Module->ChangeThreadingMode(EChaosThreadingMode::SingleThread);
 
 		// Make a solver
-		Chaos::FPhysicsSolver* Solver = Module->CreateSolver(nullptr, ESolverFlags::Standalone);
+		auto* Solver = Module->CreateSolver<Traits>(nullptr, ESolverFlags::Standalone);
 		Solver->SetEnabled(true);
 
 		// Make a particle
@@ -198,8 +196,15 @@ namespace ChaosTest {
 		Module->DestroySolver(Solver);
 	}
 
-	template void SingleParticleProxyWakeEventPropergationTest<float>();
-
+	TYPED_TEST(AllTraits, SingleParticleProxyTests)
+	{
+		ChaosTest::SingleParticleProxySingleThreadTest<TypeParam,float>();
+		ChaosTest::SingleParticleProxyWakeEventPropergationTest<TypeParam,float>();
+	}
 
+	TYPED_TEST(AllTraits,DISABLED_SingleParticleProxyTests)
+	{
+		ChaosTest::SingleParticleProxyTaskGraphTest<TypeParam,float>();
+	}
 
 }
diff --git a/Engine/Source/Programs/HeadlessChaos/Public/HeadlessChaosTestUtility.h b/Engine/Source/Programs/HeadlessChaos/Public/HeadlessChaosTestUtility.h
--- a/Engine/Source/Programs/HeadlessChaos/Public/HeadlessChaosTestUtility.h
+++ b/Engine/Source/Programs/HeadlessChaos/Public/HeadlessChaosTestUtility.h
@@ -53,6 +53,19 @@ namespace ChaosTest {
 #undef EVOLUTION_TRAIT
 	TYPED_TEST_SUITE(AllEvolutions,AllEvolutionTypes);
 
+	//should be TAllTraits, but used in macros and logs and this makes it more readable
+	template <typename T>
+	class AllTraits : public ::testing::Test {};
+
+#define EVOLUTION_TRAIT(Trait) ,Trait
+	using AllTraitsTypesTmp = TTypesWithoutVoid <
+		void
+#include "Chaos/EvolutionTraits.inl"
+	>;
+	using AllTraitsTypes = AllTraitsTypesTmp::Types;
+#undef EVOLUTION_TRAIT
+	TYPED_TEST_SUITE(AllTraits,AllTraitsTypes);
+
 	MATCHER_P2(VectorNear, V, Tolerance, "") { return arg.Equals(V, Tolerance); }
 
 	// Expects each component of the vector is within T of its corresponding component in A. 
diff --git a/Engine/Source/Runtime/Engine/Private/PhysicsEngine/Experimental/PhysScene_Chaos.cpp b/Engine/Source/Runtime/Engine/Private/PhysicsEngine/Experimental/PhysScene_Chaos.cpp
--- a/Engine/Source/Runtime/Engine/Private/PhysicsEngine/Experimental/PhysScene_Chaos.cpp
+++ b/Engine/Source/Runtime/Engine/Private/PhysicsEngine/Experimental/PhysScene_Chaos.cpp
@@ -215,9 +215,9 @@ private:
 
 		FSpacialDebugDraw DrawInterface(WorldPtr);
 
-		const TArray<FPhysicsSolver*>& Solvers = ChaosModule->GetAllSolvers();
+		const TArray<FPhysicsSolverBase*>& Solvers = ChaosModule->GetAllSolvers();
 
-		for(FPhysicsSolver* Solver : Solvers)
+		for(FPhysicsSolverBase* Solver : Solvers)
 		{
 			if(bDrawHier)
 			{
@@ -396,7 +396,7 @@ FPhysScene_Chaos::FPhysScene_Chaos(AActor* InSolverActor
 
 	UWorld* WorldPtr = SolverActor.Get() ? SolverActor->GetWorld() : nullptr;
 
-	SceneSolver = ChaosModule->CreateSolver(WorldPtr
+	SceneSolver = ChaosModule->CreateSolver<Chaos::FDefaultTraits>(WorldPtr
 #if CHAOS_CHECKED
 	, DebugName
 #endif
@@ -603,23 +603,27 @@ void FPhysScene_Chaos::AddObject(UPrimitiveComponent* Component, FFieldSystemPhy
 
 	if (Chaos::IDispatcher* Dispatcher = GetDispatcher())
 	{
-		TArray<Chaos::FPhysicsSolver*> WorldSolverList = ChaosModule->GetAllSolvers();
+		TArray<Chaos::FPhysicsSolverBase*> WorldSolverList = ChaosModule->GetAllSolvers();
 
-		for(Chaos::FPhysicsSolver* Solver : WorldSolverList)
+		for(Chaos::FPhysicsSolverBase* Solver : WorldSolverList)
 		{
-			if(true || Solver->HasActiveParticles())
+			Solver->CastHelper([Dispatcher, InObject](auto& Concrete)
 			{
-				Solver->RegisterObject(InObject);
-
-				if(/*bDedicatedThread && */Dispatcher)
+				if(true || Concrete.HasActiveParticles())
 				{
-					// Pass the proxy off to the physics thread
-					Dispatcher->EnqueueCommandImmediate([InObject, InSolver = Solver](Chaos::FPersistentPhysicsTask* PhysThread)
+					Concrete.RegisterObject(InObject);
+
+					if(/*bDedicatedThread && */Dispatcher)
+					{
+						// Pass the proxy off to the physics thread
+						Dispatcher->EnqueueCommandImmediate([InObject,&Concrete](Chaos::FPersistentPhysicsTask* PhysThread)
 						{
-							InSolver->RegisterObject(InObject);
+							Concrete.RegisterObject(InObject);
 						});
+					}
 				}
-			}
+			});
+			
 		}
 	}
 }
@@ -814,23 +818,26 @@ void FPhysScene_Chaos::RemoveObject(FFieldSystemPhysicsProxy* InObject)
 
 		if(Chaos::IDispatcher* Dispatcher = GetDispatcher())
 		{
-			TArray<Chaos::FPhysicsSolver*> SolverList = ChaosModule->GetAllSolvers();
+			TArray<Chaos::FPhysicsSolverBase*> SolverList = ChaosModule->GetAllSolvers();
 
-			for(Chaos::FPhysicsSolver* Solver : SolverList)
+			for(Chaos::FPhysicsSolverBase* Solver : SolverList)
 			{
-				if(true || Solver->HasActiveParticles())
+				Solver->CastHelper([Dispatcher, InObject](auto& Concrete)
 				{
-					Solver->UnregisterObject(InObject);
-
-					if(/*bDedicatedThread && */Dispatcher)
+					if(true || Concrete.HasActiveParticles())
 					{
-						// Pass the proxy off to the physics thread
-						Dispatcher->EnqueueCommandImmediate([InObject, InSolver = Solver](Chaos::FPersistentPhysicsTask* PhysThread)
+						Concrete.UnregisterObject(InObject);
+
+						if(/*bDedicatedThread && */Dispatcher)
+						{
+							// Pass the proxy off to the physics thread
+							Dispatcher->EnqueueCommandImmediate([InObject,&Concrete](Chaos::FPersistentPhysicsTask* PhysThread)
 							{
-								InSolver->UnregisterObject(InObject);
+								Concrete.UnregisterObject(InObject);
 							});
+						}
 					}
-				}
+				});
 			}
 		}
 	}
@@ -1974,7 +1981,7 @@ void FPhysScene_ChaosInterface::StartFrame()
 			OnPhysScenePreTick.Broadcast(this, Dt);
 			OnPhysSceneStep.Broadcast(this, Dt);
 
-			TArray<Chaos::FPhysicsSolver*> SolverList;
+			TArray<Chaos::FPhysicsSolverBase*> SolverList;
 			SolverModule->GetSolversMutable(GetOwningWorld(), SolverList);
 			
 			if (FPhysicsSolver* Solver = GetSolver())
@@ -2086,7 +2093,7 @@ void FPhysScene_ChaosInterface::EndFrame(ULineBatchComponent* InLineBatcher)
 
 		// Make a list of solvers to process. This is a list of all solvers registered to our world
 		// And our internal base scene solver.
-		TArray<Chaos::FPhysicsSolver*> SolverList;
+		TArray<Chaos::FPhysicsSolverBase*> SolverList;
 		SolverModule->GetSolversMutable(GetOwningWorld(), SolverList);
 
 		if(FPhysicsSolver* Solver = GetSolver())
@@ -2096,7 +2103,7 @@ void FPhysScene_ChaosInterface::EndFrame(ULineBatchComponent* InLineBatcher)
 		}
 
 		// flush solver queues
-		for (FPhysicsSolver* Solver : SolverList)
+		for (FPhysicsSolverBase* Solver : SolverList)
 		{
 			TQueue<TFunction<void()>, EQueueMode::Mpsc>& Queue = Solver->GetCommandQueue();
 			TFunction<void()> Command;
@@ -2114,29 +2121,36 @@ void FPhysScene_ChaosInterface::EndFrame(ULineBatchComponent* InLineBatcher)
 			//for now just copy the whole thing, stomping any changes that came from GT
 			Scene.CopySolverAccelerationStructure();
 
-			TArray<FPhysicsSolver*> ActiveSolvers;
+			TArray<FPhysicsSolverBase*> ActiveSolvers;
 			ActiveSolvers.Reserve(SolverList.Num());
 
 			// #BG calculate active solver list once as we dispatch our first task
-			for(FPhysicsSolver* Solver : SolverList)
+			for(FPhysicsSolverBase* Solver : SolverList)
 			{
-				if(Solver->HasActiveParticles())
+				Solver->CastHelper([&ActiveSolvers](auto& Concrete)
 				{
-					ActiveSolvers.Add(Solver);
-				}
+					if(Concrete.HasActiveParticles())
+					{
+						ActiveSolvers.Add(&Concrete);
+					}
+				});
+				
 			}
 
 			const int32 NumActiveSolvers = ActiveSolvers.Num();
 
-			for (FPhysicsSolver* Solver : ActiveSolvers)
+			for (FPhysicsSolverBase* Solver : ActiveSolvers)
 			{
-				SyncBodies(Solver);
-				Solver->SyncEvents_GameThread();
-
+				Solver->CastHelper([&ActiveSolvers,this](auto& Concrete)
 				{
-					SCOPE_CYCLE_COUNTER(STAT_SqUpdateMaterials);
-					Solver->SyncQueryMaterials();
-				}
+					SyncBodies(&Concrete);
+					Concrete.SyncEvents_GameThread();
+
+					{
+						SCOPE_CYCLE_COUNTER(STAT_SqUpdateMaterials);
+						Concrete.SyncQueryMaterials();
+					}
+				});
 			}
 		}
 
@@ -2218,7 +2232,8 @@ void FPhysScene_ChaosInterface::KillVisualDebugger()
 
 }
 
-void FPhysScene_ChaosInterface::SyncBodies(Chaos::FPhysicsSolver* Solver)
+template <typename TSolver>
+void FPhysScene_ChaosInterface::SyncBodies(TSolver* Solver)
 {
 	DECLARE_SCOPE_CYCLE_COUNTER(TEXT("SyncBodies"), STAT_SyncBodies, STATGROUP_Physics);
 	TArray<FPhysScenePendingComponentTransform_Chaos> PendingTransforms;
@@ -2346,22 +2361,25 @@ void FPhysScene_ChaosInterface::CompleteSceneSimulation(ENamedThreads::Type Curr
 
 		check(Module);
 
-		TArray<FPhysicsSolver*> SolverList = Module->GetSolversMutable(GetOwningWorld());
+		TArray<FPhysicsSolverBase*> SolverList = Module->GetSolversMutable(GetOwningWorld());
 		FPhysicsSolver* SceneSolver = GetSolver();
 
-		TArray<FPhysicsSolver*> ActiveSolvers;
+		TArray<FPhysicsSolverBase*> ActiveSolvers;
 
 		if(SolverList.Num() > 0)
 		{
 			ActiveSolvers.Reserve(SolverList.Num());
 
 			// #BG calculate active solver list once as we dispatch our first task
-			for(FPhysicsSolver* Solver : SolverList)
+			for(FPhysicsSolverBase* Solver : SolverList)
 			{
-				if(Solver->HasActiveParticles())
+				Solver->CastHelper([&ActiveSolvers](auto& Concrete)
 				{
-					ActiveSolvers.Add(Solver);
-				}
+					if(Concrete.HasActiveParticles())
+					{
+						ActiveSolvers.Add(&Concrete);
+					}
+				});
 			}
 		}
 
@@ -2373,13 +2391,14 @@ void FPhysScene_ChaosInterface::CompleteSceneSimulation(ENamedThreads::Type Curr
 		const int32 NumActiveSolvers = ActiveSolvers.Num();
 
 		PhysicsParallelFor(NumActiveSolvers, [&](int32 Index)
-			{
-				FPhysicsSolver* Solver = ActiveSolvers[Index];
-
-				Solver->GetActiveParticlesBuffer()->CaptureSolverData(Solver);
-				Solver->BufferPhysicsResults();
-				Solver->FlipBuffers();
-			});
+		{
+			//TODO: support any type not just default traits
+			FPhysicsSolverBase* Solver = ActiveSolvers[Index];
+			auto& Concrete = Solver->CastChecked<Chaos::FDefaultTraits>();
+			Concrete.GetActiveParticlesBuffer()->CaptureSolverData(&Concrete);
+			Concrete.BufferPhysicsResults();
+			Concrete.FlipBuffers();
+		});
 	}
 }
 
diff --git a/Engine/Source/Runtime/Engine/Public/Physics/Experimental/PhysScene_Chaos.h b/Engine/Source/Runtime/Engine/Public/Physics/Experimental/PhysScene_Chaos.h
--- a/Engine/Source/Runtime/Engine/Public/Physics/Experimental/PhysScene_Chaos.h
+++ b/Engine/Source/Runtime/Engine/Public/Physics/Experimental/PhysScene_Chaos.h
@@ -453,7 +453,8 @@ private:
 	bool IsOwningWorldEditor() const;
 #endif
 
-	void SyncBodies(Chaos::FPhysicsSolver* Solver);
+	template <typename TSolver>
+	void SyncBodies(TSolver* Solver);
 
 	void SetKinematicTransform(FPhysicsActorHandle& InActorReference, const Chaos::TRigidTransform<float, 3>& NewTransform)
 	{
diff --git a/Engine/Source/Runtime/Experimental/Chaos/Public/Chaos/EvolutionTraits.h b/Engine/Source/Runtime/Experimental/Chaos/Public/Chaos/EvolutionTraits.h
--- a/Engine/Source/Runtime/Experimental/Chaos/Public/Chaos/EvolutionTraits.h
+++ b/Engine/Source/Runtime/Experimental/Chaos/Public/Chaos/EvolutionTraits.h
@@ -13,4 +13,18 @@ struct FNonRewindableEvolutionTraits
 {
 	static constexpr bool IsRewindable(){ return false; }
 };
+
+#define EVOLUTION_TRAIT(Trait) Trait,
+enum class ETraits
+{
+#include "Chaos/EvolutionTraits.inl"
+	NumTraits
+};
+#undef EVOLUTION_TRAIT
+
+template <typename T>
+inline constexpr ETraits TraitToIdx(){ return ETraits::NumTraits; }
+#define EVOLUTION_TRAIT(Trait) template <> inline constexpr ETraits TraitToIdx<Trait>(){ return ETraits::Trait; }
+#include "Chaos/EvolutionTraits.inl"
+#undef EVOLUTION_TRAIT
 }
diff --git a/Engine/Source/Runtime/Experimental/Chaos/Public/Chaos/Framework/PhysicsSolverBase.h b/Engine/Source/Runtime/Experimental/Chaos/Public/Chaos/Framework/PhysicsSolverBase.h
--- a/Engine/Source/Runtime/Experimental/Chaos/Public/Chaos/Framework/PhysicsSolverBase.h
+++ b/Engine/Source/Runtime/Experimental/Chaos/Public/Chaos/Framework/PhysicsSolverBase.h
@@ -9,6 +9,7 @@
 #include "Chaos/ParticleDirtyFlags.h"
 #include "Async/ParallelFor.h"
 #include "Containers/Queue.h"
+#include "Chaos/EvolutionTraits.h"
 
 class FChaosSolversModule;
 class FPhysicsSolverAdvanceTask;
@@ -184,6 +185,24 @@ namespace Chaos
 	{
 	public:
 
+#define EVOLUTION_TRAIT(Trait) case ETraits::Trait: Func((TPBDRigidsSolver<Trait>&)(*this)); return;
+		template <typename Lambda>
+		void CastHelper(const Lambda& Func)
+		{
+			switch(TraitIdx)
+			{
+#include "Chaos/EvolutionTraits.inl"
+			}
+		}
+#undef EVOLUTION_TRAIT
+
+		template <typename Traits>
+		TPBDRigidsSolver<Traits>& CastChecked()
+		{
+			check(TraitIdx == TraitToIdx<Traits>());
+			return (TPBDRigidsSolver<Traits>&)(*this);
+		}
+
 		void ChangeBufferMode(EMultiBufferMode InBufferMode);
 
 		TQueue<TFunction<void()>,EQueueMode::Mpsc>& GetCommandQueue() { return CommandQueue; }
@@ -243,9 +262,10 @@ namespace Chaos
 		EMultiBufferMode BufferMode;
 
 		/** Protected construction so callers still have to go through the module to create new instances */
-		FPhysicsSolverBase(const EMultiBufferMode BufferingModeIn, UObject* InOwner)
+		FPhysicsSolverBase(const EMultiBufferMode BufferingModeIn, UObject* InOwner, ETraits InTraitIdx)
 			: BufferMode(BufferingModeIn)
 			, Owner(InOwner)
+			, TraitIdx(InTraitIdx)
 		{}
 
 		/** Only allow construction with valid parameters as well as restricting to module construction */
@@ -287,5 +307,7 @@ namespace Chaos
 
 		template<ELockType>
 		friend struct TSolverQueryMaterialScope;
+
+		ETraits TraitIdx;
 	};
 }
diff --git a/Engine/Source/Runtime/Experimental/ChaosSolvers/Private/ChaosSolversModule.cpp b/Engine/Source/Runtime/Experimental/ChaosSolvers/Private/ChaosSolversModule.cpp
--- a/Engine/Source/Runtime/Experimental/ChaosSolvers/Private/ChaosSolversModule.cpp
+++ b/Engine/Source/Runtime/Experimental/ChaosSolvers/Private/ChaosSolversModule.cpp
@@ -301,7 +301,7 @@ void FChaosSolversModule::ShutdownThreadingMode()
 			// need to handle this for task graph dispatchers
 			Dispatcher->Execute();
 
-			for(FPBDRigidsSolver* Solver : AllSolvers)
+			for(FPhysicsSolverBase* Solver : AllSolvers)
 			{
 				IDispatcher::FSolverCommand Command;
 				while(Solver->CommandQueue.Dequeue(Command))
@@ -463,20 +463,8 @@ void FChaosSolversModule::SyncTask(bool bForceBlockingSync /*= false*/)
 	UpdateStats();
 }
 
-Chaos::FPhysicsSolver* FChaosSolversModule::CreateSolver(UObject* InOwner
-#if CHAOS_CHECKED
-	, const FName& DebugName
-#endif
-)
-{
-	return CreateSolver(InOwner, ESolverFlags::None
-#if CHAOS_CHECKED
-		, DebugName
-#endif
-	);
-}
-
-Chaos::FPhysicsSolver* FChaosSolversModule::CreateSolver(UObject* InOwner, ESolverFlags InFlags
+template <typename Traits>
+Chaos::TPBDRigidsSolver<Traits>* FChaosSolversModule::CreateSolver(UObject* InOwner, ESolverFlags InFlags
 #if CHAOS_CHECKED
 	, const FName& DebugName
 #endif
@@ -492,11 +480,11 @@ Chaos::FPhysicsSolver* FChaosSolversModule::CreateSolver(UObject* InOwner, ESolv
 		SolverBufferMode = GetDesiredBufferingMode();
 	}
 
-	AllSolvers.Add(new Chaos::FPhysicsSolver(SolverBufferMode, InOwner));
-	Chaos::FPhysicsSolver* NewSolver = AllSolvers.Last();
+	auto* NewSolver = new Chaos::TPBDRigidsSolver<Traits>(SolverBufferMode,InOwner);
+	AllSolvers.Add(NewSolver);
 
 	// Add The solver to the owner list
-	TArray<Chaos::FPhysicsSolver*>& OwnerSolverList = SolverMap.FindOrAdd(InOwner);
+	TArray<Chaos::FPhysicsSolverBase*>& OwnerSolverList = SolverMap.FindOrAdd(InOwner);
 	OwnerSolverList.Add(NewSolver);
 
 #if CHAOS_CHECKED
@@ -529,7 +517,7 @@ Chaos::FPhysicsSolver* FChaosSolversModule::CreateSolver(UObject* InOwner, ESolv
 	return NewSolver;
 }
 
-void FChaosSolversModule::MigrateSolver(Chaos::FPhysicsSolver* InSolver, const UObject* InNewOwner)
+void FChaosSolversModule::MigrateSolver(Chaos::FPhysicsSolverBase* InSolver, const UObject* InNewOwner)
 {
 	checkSlow(IsInGameThread());
 
@@ -541,7 +529,7 @@ void FChaosSolversModule::MigrateSolver(Chaos::FPhysicsSolver* InSolver, const U
 
 	if(const UObject* CurrentOwner = InSolver->GetOwner())
 	{
-		TArray<Chaos::FPhysicsSolver*>* OwnerList = SolverMap.Find(CurrentOwner);
+		TArray<Chaos::FPhysicsSolverBase*>* OwnerList = SolverMap.Find(CurrentOwner);
 		ensure(OwnerList->Remove(InSolver));
 	}
 
@@ -549,7 +537,7 @@ void FChaosSolversModule::MigrateSolver(Chaos::FPhysicsSolver* InSolver, const U
 
 	if(InNewOwner)
 	{
-		TArray<Chaos::FPhysicsSolver*>& OwnerList = SolverMap.FindOrAdd(InNewOwner);
+		TArray<Chaos::FPhysicsSolverBase*>& OwnerList = SolverMap.FindOrAdd(InNewOwner);
 		OwnerList.Add(InSolver);
 	}
 }
@@ -578,7 +566,7 @@ void FChaosSolversModule::SetDedicatedThreadTickMode(EChaosSolverTickMode InTick
 	});
 }
 
-void FChaosSolversModule::DestroySolver(Chaos::FPhysicsSolver* InSolver)
+void FChaosSolversModule::DestroySolver(Chaos::FPhysicsSolverBase* InSolver)
 {
 	LLM_SCOPE(ELLMTag::Chaos);
 
@@ -587,7 +575,7 @@ void FChaosSolversModule::DestroySolver(Chaos::FPhysicsSolver* InSolver)
 	if(AllSolvers.Remove(InSolver) > 0)
 	{
 		// Also remove from the owner list
-		TArray<Chaos::FPhysicsSolver*> OwnerList = GetSolversMutable(InSolver->GetOwner());
+		TArray<Chaos::FPhysicsSolverBase*> OwnerList = GetSolversMutable(InSolver->GetOwner());
 		ensureMsgf(OwnerList.Remove(InSolver), TEXT("Removed a solver from the global list but not an owner list."));
 
 		if(Dispatcher)
@@ -613,40 +601,40 @@ void FChaosSolversModule::DestroySolver(Chaos::FPhysicsSolver* InSolver)
 	}
 }
 
-const TArray<Chaos::FPhysicsSolver*>& FChaosSolversModule::GetAllSolvers() const
+const TArray<Chaos::FPhysicsSolverBase*>& FChaosSolversModule::GetAllSolvers() const
 {
 	return AllSolvers;
 }
 
-TArray<const Chaos::FPhysicsSolver*> FChaosSolversModule::GetSolvers(const UObject* InOwner) const
+TArray<const Chaos::FPhysicsSolverBase*> FChaosSolversModule::GetSolvers(const UObject* InOwner) const
 {
 	// Can't just return the ptr from TMap::Find here as it's likely these solver lists will be sent to
 	// the physics thread. In which case the solver pointers are stable but the container pointers are not
-	TArray<const Chaos::FPhysicsSolver*> Solvers;
+	TArray<const Chaos::FPhysicsSolverBase*> Solvers;
 	GetSolvers(InOwner, Solvers);
 	return Solvers;
 }
 
-TArray<Chaos::FPhysicsSolver*> FChaosSolversModule::GetSolversMutable(const UObject* InOwner)
+TArray<Chaos::FPhysicsSolverBase*> FChaosSolversModule::GetSolversMutable(const UObject* InOwner)
 {
 	// Can't just return the ptr from TMap::Find here as it's likely these solver lists will be sent to
 	// the physics thread. In which case the solver pointers are stable but the container pointers are not
-	TArray<Chaos::FPhysicsSolver*> Solvers;
+	TArray<Chaos::FPhysicsSolverBase*> Solvers;
 	GetSolversMutable(InOwner, Solvers);
 	return Solvers;
 }
 
-void FChaosSolversModule::GetSolvers(const UObject* InOwner, TArray<const Chaos::FPhysicsSolver*>& OutSolvers) const
+void FChaosSolversModule::GetSolvers(const UObject* InOwner, TArray<const Chaos::FPhysicsSolverBase*>& OutSolvers) const
 {
-	if(const TArray<Chaos::FPhysicsSolver*>* Solvers = SolverMap.Find(InOwner))
+	if(const TArray<Chaos::FPhysicsSolverBase*>* Solvers = SolverMap.Find(InOwner))
 	{
 		OutSolvers.Append(*Solvers);
 	}
 }
 
-void FChaosSolversModule::GetSolversMutable(const UObject* InOwner, TArray<Chaos::FPhysicsSolver*>& OutSolvers)
+void FChaosSolversModule::GetSolversMutable(const UObject* InOwner, TArray<Chaos::FPhysicsSolverBase*>& OutSolvers)
 {
-	if(TArray<Chaos::FPhysicsSolver*>* Solvers = SolverMap.Find(InOwner))
+	if(TArray<Chaos::FPhysicsSolverBase*>* Solvers = SolverMap.Find(InOwner))
 	{
 		OutSolvers.Append(*Solvers);
 	}
@@ -680,7 +668,7 @@ void FChaosSolversModule::DumpHierarchyStats(int32* OutOptMaxCellElements)
 	const int32 NumSolvers = AllSolvers.Num();
 	for(int32 SolverIndex = 0; SolverIndex < NumSolvers; ++SolverIndex)
 	{
-		Chaos::FPhysicsSolver* Solver = AllSolvers[SolverIndex];
+		Chaos::FPhysicsSolverBase* Solver = AllSolvers[SolverIndex];
 #if TODO_REIMPLEMENT_SPATIAL_ACCELERATION_ACCESS
 		if(const Chaos::ISpatialAcceleration<float,3>* SpatialAcceleration = Solver->GetSpatialAcceleration())
 		{
@@ -949,7 +937,7 @@ bool FChaosSolversModule::ShouldStepSolver(int32& InOutSingleStepCounter) const
 
 void FChaosSolversModule::ChangeBufferMode(Chaos::EMultiBufferMode BufferMode)
 {
-	for (Chaos::FPhysicsSolver* Solver : AllSolvers)
+	for (Chaos::FPhysicsSolverBase* Solver : AllSolvers)
 	{
 		if (Dispatcher)
 		{
@@ -995,12 +983,15 @@ void FChaosSolversModule::OnUpdateMaterial(Chaos::FMaterialHandle InHandle)
 
 	if(ensure(Material))
 	{
-		for(Chaos::FPhysicsSolver* Solver : AllSolvers)
+		for(Chaos::FPhysicsSolverBase* Solver : AllSolvers)
 		{
 			// Send a copy of the material to each solver
 			Dispatcher->EnqueueCommandImmediate(Solver, [InHandle, MaterialCopy = *Material, Solver]()
 			{
-				Solver->UpdateMaterial(InHandle, MaterialCopy);
+				Solver->CastHelper([InHandle, &MaterialCopy](auto& Concrete)
+				{
+					Concrete.UpdateMaterial(InHandle,MaterialCopy);
+				});
 			});
 		}
 	}
@@ -1015,12 +1006,15 @@ void FChaosSolversModule::OnCreateMaterial(Chaos::FMaterialHandle InHandle)
 
 	if(ensure(Material))
 	{
-		for(Chaos::FPhysicsSolver* Solver : AllSolvers)
+		for(Chaos::FPhysicsSolverBase* Solver : AllSolvers)
 		{
 			// Send a copy of the material to each solver
 			Dispatcher->EnqueueCommandImmediate(Solver, [InHandle, MaterialCopy = *Material, Solver]()
 			{
-				Solver->CreateMaterial(InHandle, MaterialCopy);
+				Solver->CastHelper([InHandle, &MaterialCopy](auto& Concrete)
+				{
+					Concrete.CreateMaterial(InHandle,MaterialCopy);
+				});
 			});
 		}
 	}
@@ -1038,12 +1032,15 @@ void FChaosSolversModule::OnDestroyMaterial(Chaos::FMaterialHandle InHandle)
 
 	if(ensure(Material))
 	{
-		for(Chaos::FPhysicsSolver* Solver : AllSolvers)
+		for(Chaos::FPhysicsSolverBase* Solver : AllSolvers)
 		{
 			// Notify each solver
 			Dispatcher->EnqueueCommandImmediate(Solver, [InHandle, Solver]()
 			{
-				Solver->DestroyMaterial(InHandle);
+				Solver->CastHelper([InHandle](auto& Concrete)
+				{
+					Concrete.DestroyMaterial(InHandle);
+				});
 			});
 		}
 	}
@@ -1058,12 +1055,15 @@ void FChaosSolversModule::OnUpdateMaterialMask(Chaos::FMaterialMaskHandle InHand
 
 	if (ensure(MaterialMask))
 	{
-		for (Chaos::FPhysicsSolver* Solver : AllSolvers)
+		for (Chaos::FPhysicsSolverBase* Solver : AllSolvers)
 		{
 			// Send a copy of the material to each solver
 			Dispatcher->EnqueueCommandImmediate(Solver, [InHandle, MaterialMaskCopy = *MaterialMask, Solver]()
 			{
-				Solver->UpdateMaterialMask(InHandle, MaterialMaskCopy);
+				Solver->CastHelper([InHandle,&MaterialMaskCopy](auto& Concrete)
+				{
+					Concrete.UpdateMaterialMask(InHandle,MaterialMaskCopy);
+				});
 			});
 		}
 	}
@@ -1078,12 +1078,15 @@ void FChaosSolversModule::OnCreateMaterialMask(Chaos::FMaterialMaskHandle InHand
 
 	if (ensure(MaterialMask))
 	{
-		for (Chaos::FPhysicsSolver* Solver : AllSolvers)
+		for (Chaos::FPhysicsSolverBase* Solver : AllSolvers)
 		{
 			// Send a copy of the material to each solver
 			Dispatcher->EnqueueCommandImmediate(Solver, [InHandle, MaterialMaskCopy = *MaterialMask, Solver]()
 			{
-				Solver->CreateMaterialMask(InHandle, MaterialMaskCopy);
+				Solver->CastHelper([InHandle,&MaterialMaskCopy](auto& Concrete)
+				{
+					Concrete.CreateMaterialMask(InHandle,MaterialMaskCopy);
+				});
 			});
 		}
 	}
@@ -1101,12 +1104,15 @@ void FChaosSolversModule::OnDestroyMaterialMask(Chaos::FMaterialMaskHandle InHan
 
 	if (ensure(MaterialMask))
 	{
-		for (Chaos::FPhysicsSolver* Solver : AllSolvers)
+		for (Chaos::FPhysicsSolverBase* Solver : AllSolvers)
 		{
 			// Notify each solver
 			Dispatcher->EnqueueCommandImmediate(Solver, [InHandle, Solver]()
 			{
-				Solver->DestroyMaterialMask(InHandle);
+				Solver->CastHelper([InHandle](auto& Concrete)
+				{
+					Concrete.DestroyMaterialMask(InHandle);
+				});
 			});
 		}
 	}
@@ -1222,4 +1228,15 @@ bool FChaosScopedPhysicsThreadLock::DidGetLock() const
 	return bGotLock;
 }
 
+#if CHAOS_CHECKED
+#define EVOLUTION_TRAIT(Traits)\
+template Chaos::TPBDRigidsSolver<Chaos::Traits>* FChaosSolversModule::CreateSolver(UObject* InOwner, ESolverFlags InFlags, const FName& DebugName);
+#else
+#define EVOLUTION_TRAIT(Traits)\
+template Chaos::TPBDRigidsSolver<Chaos::Traits>* FChaosSolversModule::CreateSolver(UObject* InOwner,ESolverFlags InFlags);
+#endif
+
+#include "Chaos/EvolutionTraits.inl"
+#undef EVOLUTION_TRAIT
+
 IMPLEMENT_MODULE(FChaosSolversModule, ChaosSolvers);
diff --git a/Engine/Source/Runtime/Experimental/ChaosSolvers/Private/Framework/DebugSolverTasks.cpp b/Engine/Source/Runtime/Experimental/ChaosSolvers/Private/Framework/DebugSolverTasks.cpp
--- a/Engine/Source/Runtime/Experimental/ChaosSolvers/Private/Framework/DebugSolverTasks.cpp
+++ b/Engine/Source/Runtime/Experimental/ChaosSolvers/Private/Framework/DebugSolverTasks.cpp
@@ -36,7 +36,7 @@ namespace Chaos
 		DebugSubstep.Add(true, TEXT("Debug thread exit"));
 	}
 
-	void FDebugSolverTasks::DebugStep(FPhysicsSolver* Solver, TFunction<void()> StepFunction)
+	void FDebugSolverTasks::DebugStep(FPhysicsSolverBase* Solver, TFunction<void()> StepFunction)
 	{
 		// Retrieve debug thread pointer for this solver.
 		// It must have been already added unless the single thread mode has just been switched on.
@@ -81,7 +81,7 @@ namespace Chaos
 #endif
 	}
 
-	void FDebugSolverTasks::Add(FPhysicsSolver* Solver)
+	void FDebugSolverTasks::Add(FPhysicsSolverBase* Solver)
 	{
 		// Add solver to task map
 		SolverToTaskMap.Add(Solver, nullptr);
@@ -93,7 +93,7 @@ namespace Chaos
 #endif
 	}
 
-	void FDebugSolverTasks::Remove(FPhysicsSolver* Solver)
+	void FDebugSolverTasks::Remove(FPhysicsSolverBase* Solver)
 	{
 #if TODO_REIMPLEMENT_DEBUG_SUBSTEP
 		// Remove the debug advance task for this solver, if any was created
diff --git a/Engine/Source/Runtime/Experimental/ChaosSolvers/Private/Framework/PersistentTask.cpp b/Engine/Source/Runtime/Experimental/ChaosSolvers/Private/Framework/PersistentTask.cpp
--- a/Engine/Source/Runtime/Experimental/ChaosSolvers/Private/Framework/PersistentTask.cpp
+++ b/Engine/Source/Runtime/Experimental/ChaosSolvers/Private/Framework/PersistentTask.cpp
@@ -60,7 +60,7 @@ namespace Chaos
 
 #if CHAOS_DEBUG_SUBSTEP
 		// Prepare the debug substepping tasks for all existing solvers
-		for (FPhysicsSolver* Solver : Solvers)
+		for (FPhysicsSolverBase* Solver : Solvers)
 		{
 			DebugSolverTasks.Add(Solver);
 		}
@@ -75,7 +75,7 @@ namespace Chaos
 		Timestep->Reset();
 
 		// Scratch array for active solver list
-		TArray<Chaos::FPhysicsSolver*> ActiveSolverList;
+		TArray<Chaos::FPhysicsSolverBase*> ActiveSolverList;
 
 		while(bRunning)
 		{
@@ -92,12 +92,16 @@ namespace Chaos
 			{
 				// Go wide if possible on the solvers
 				ActiveSolverList.Reset(Solvers.Num());
-				for(Chaos::FPhysicsSolver* Solver : Solvers)
+				for(Chaos::FPhysicsSolverBase* Solver : Solvers)
 				{
-					if(Solver->HasActiveParticles())
+					Solver->CastHelper([&ActiveSolverList](auto& Concrete)
 					{
-						ActiveSolverList.Add(Solver);
-					}
+						if(Concrete.HasActiveParticles())
+						{
+							ActiveSolverList.Add(&Concrete);
+						}
+					});
+					
 				}
 
 				const int32 NumSolverEntries = ActiveSolverList.Num();
@@ -106,7 +110,7 @@ namespace Chaos
 				{
 					SCOPE_CYCLE_COUNTER(STAT_SolverAdvance);
 				
-					FPhysicsSolver* Solver = ActiveSolverList[Index];
+					FPhysicsSolverBase* Solver = ActiveSolverList[Index];
 
 					// Execute Step function either on this thread or in a pausable side debug thread
 					DebugSolverTasks.DebugStep(Solver, [this, Solver, Dt]()
@@ -144,12 +148,13 @@ namespace Chaos
 
 				for(int32 SolverIndex = 0; SolverIndex < NumSolvers; ++SolverIndex)
 				{
-					FPhysicsSolver* Solver = Solvers[SolverIndex];
+					FPhysicsSolverBase* Solver = Solvers[SolverIndex];
 
 					FPersistentPhysicsTaskStatistics::FPerSolverStatistics& SolverStat = CurrStats.SolverStats[SolverIndex];
-					FPhysicsSolver::FPBDRigidsEvolution& Evolution = *Solver->MEvolution;
 
 #if TODO_REIMPLEMENT_SOLVER_ENABLING
+					FPhysicsSolver::FPBDRigidsEvolution& Evolution = *Solver->MEvolution;
+
 					if(Solver->Enabled())
 					{
 #if TODO_REIMPLMEENT_EVOLUTION_ACCESSORS
@@ -177,29 +182,30 @@ namespace Chaos
 		ShutdownEvent->Trigger();
 	}
 
-	void FPersistentPhysicsTask::StepSolver(FPhysicsSolver* InSolver, float Dt)
+	void FPersistentPhysicsTask::StepSolver(FPhysicsSolverBase* InSolver, float Dt)
 	{
 		HandleSolverCommands(InSolver);
 
 		// Check whether this solver is paused (changes in pause states usually happens during HandleSolverCommands)
 #if CHAOS_WITH_PAUSABLE_SOLVER
 #if TODO_REIMPLEMENT_SOLVER_PAUSING
-		if (!InSolver->Paused())
+		if(!InSolver->Paused())
 #endif
 #endif
+			InSolver->CastHelper([Dt, this](auto& Concrete)
 		{
-			if (InSolver->bEnabled)
+			if(Concrete.bEnabled)
 			{
 				// Only process if we have something to actually simulate
-				if(InSolver->HasActiveParticles())
+				if(Concrete.HasActiveParticles())
 				{
-					AdvanceSolver(InSolver, Dt);
+					AdvanceSolver(&Concrete,Dt);
 
 					{
 						SCOPE_CYCLE_COUNTER(STAT_BufferPhysicsResults);
 						CacheLock.ReadLock();
 
-						InSolver->ForEachPhysicsProxyParallel([](auto* Object)
+						Concrete.ForEachPhysicsProxyParallel([](auto* Object)
 						{
 							Object->BufferPhysicsResults();
 						});
@@ -211,7 +217,7 @@ namespace Chaos
 						SCOPE_CYCLE_COUNTER(STAT_FlipResults);
 						CacheLock.WriteLock();
 
-						InSolver->ForEachPhysicsProxy([](auto* Object)
+						Concrete.ForEachPhysicsProxy([](auto* Object)
 						{
 							Object->FlipBuffer();
 						});
@@ -220,16 +226,16 @@ namespace Chaos
 					}
 				}
 			}
-		}
+		});
 	}
 
-	void FPersistentPhysicsTask::AddSolver(FPhysicsSolver* InSolver)
+	void FPersistentPhysicsTask::AddSolver(FPhysicsSolverBase* InSolver)
 	{
 		Solvers.Add(InSolver);
 		DebugSolverTasks.Add(InSolver);
 	}
 
-	void FPersistentPhysicsTask::RemoveSolver(FPhysicsSolver* InSolver)
+	void FPersistentPhysicsTask::RemoveSolver(FPhysicsSolverBase* InSolver)
 	{
 		DebugSolverTasks.Remove(InSolver);
 		Solvers.Remove(InSolver);
@@ -249,11 +255,14 @@ namespace Chaos
 		{
 			TArray<FFieldSystemPhysicsProxy*> FieldsToDelete;
 
-			for(FPhysicsSolver* Solver : Solvers)
+			for(FPhysicsSolverBase* Solver : Solvers)
 			{
-				Solver->ForEachPhysicsProxy([](auto* Object)
+				Solver->CastHelper([](auto& Concrete)
 				{
-					Object->PullFromPhysicsState();
+					Concrete.ForEachPhysicsProxy([](auto* Object)
+					{
+						Object->PullFromPhysicsState();
+					});
 				});
 
 #if TODO_REIMPLEMENT_REMOVED_PROXY_STORAGE
@@ -286,24 +295,32 @@ namespace Chaos
 				delete FieldObj;
 			}
 
-			for (FPhysicsSolver* Solver : Solvers)
+			for (FPhysicsSolverBase* Solver : Solvers)
 			{
-				Solver->SyncEvents_GameThread();
+				Solver->CastHelper([](auto& Concrete)
+				{
+					Concrete.SyncEvents_GameThread();
+				});
 			}
 		}
 		else
 		{ 
-			for(FPhysicsSolver* Solver : Solvers)
+			for(FPhysicsSolverBase* Solver : Solvers)
 			{
-				Solver->ForEachPhysicsProxy([](auto* Object)
+				Solver->CastHelper([](auto& Concrete)
 				{
-					Object->PullFromPhysicsState();
+					Concrete.ForEachPhysicsProxy([](auto* Object)
+					{
+						Object->PullFromPhysicsState();
+					});
 				});
 			}
-
-			for (FPhysicsSolver* Solver : Solvers)
+			for(FPhysicsSolverBase* Solver : Solvers)
 			{
-				Solver->SyncEvents_GameThread();
+				Solver->CastHelper([](auto& Concrete)
+				{
+					Concrete.SyncEvents_GameThread();
+				});
 			}
 		}
 
@@ -378,7 +395,7 @@ namespace Chaos
 		return Stats.GetGameDataForRead();
 	}
 
-	void FPersistentPhysicsTask::HandleSolverCommands(FPhysicsSolver* InSolver)
+	void FPersistentPhysicsTask::HandleSolverCommands(FPhysicsSolverBase* InSolver)
 	{
 		SCOPE_CYCLE_COUNTER(STAT_HandleSolverCommands);
 
@@ -391,7 +408,8 @@ namespace Chaos
 		}
 	}
 
-	void FPersistentPhysicsTask::AdvanceSolver(FPhysicsSolver* InSolver, float InDt)
+	template <typename TSolver>
+	void FPersistentPhysicsTask::AdvanceSolver(TSolver* InSolver, float InDt)
 	{
 		SCOPE_CYCLE_COUNTER(STAT_IntegrateSolver);
 
diff --git a/Engine/Source/Runtime/Experimental/ChaosSolvers/Private/Framework/PhysicsTickTask.cpp b/Engine/Source/Runtime/Experimental/ChaosSolvers/Private/Framework/PhysicsTickTask.cpp
--- a/Engine/Source/Runtime/Experimental/ChaosSolvers/Private/Framework/PhysicsTickTask.cpp
+++ b/Engine/Source/Runtime/Experimental/ChaosSolvers/Private/Framework/PhysicsTickTask.cpp
@@ -42,7 +42,7 @@ FPhysicsTickTask::FPhysicsTickTask(FGraphEventRef& InCompletionEvent, Chaos::FPh
 	Module->GetSolverUpdatePrerequisites(SolverTaskPrerequisites);
 }
 
-FPhysicsTickTask::FPhysicsTickTask(FGraphEventRef& InCompletionEvent, const TArray<Chaos::FPhysicsSolver*>& InSolverList, float InDt)
+FPhysicsTickTask::FPhysicsTickTask(FGraphEventRef& InCompletionEvent, const TArray<Chaos::FPhysicsSolverBase*>& InSolverList, float InDt)
 	: CompletionEvent(InCompletionEvent)
 	, Module(nullptr)
 	, Dt(InDt)
@@ -85,15 +85,18 @@ void FPhysicsTickTask::DoTask(ENamedThreads::Type CurrentThread, const FGraphEve
 
 
 	// List of active solvers (assume all are active for single alloc)
-	TArray<FPhysicsSolver*> ActiveSolvers;
+	TArray<FPhysicsSolverBase*> ActiveSolvers;
 	ActiveSolvers.Reserve(SolverList.Num());
 
-	for(FPhysicsSolver* Solver : SolverList)
+	for(FPhysicsSolverBase* Solver : SolverList)
 	{
-		if(Solver->HasActiveParticles() || Solver->HasPendingCommands())
+		Solver->CastHelper([&ActiveSolvers](auto& ConcreteSolver)
 		{
-			ActiveSolvers.Add(Solver);
-		}
+			if(ConcreteSolver.HasActiveParticles() || ConcreteSolver.HasPendingCommands())
+			{
+				ActiveSolvers.Add(&ConcreteSolver);
+			}
+		});
 	}
 
 	const int32 NumActiveSolvers = ActiveSolvers.Num();
@@ -103,7 +106,7 @@ void FPhysicsTickTask::DoTask(ENamedThreads::Type CurrentThread, const FGraphEve
 
 	// For each solver spawn a new solver advance task (which will run the per-solver command buffer and then advance the solver)
 	// Record the task reference as a prerequisite for the completion
-	for(FPhysicsSolver* Solver : ActiveSolvers)
+	for(FPhysicsSolverBase* Solver : ActiveSolvers)
 	{
 		FGraphEventRef SolverTaskRef = TGraphTask<FPhysicsSolverAdvanceTask>::CreateTask(&SolverTaskPrerequisites).ConstructAndDispatchWhenReady(Solver, Dt);
 		CompletionTaskPrerequisites.Add(SolverTaskRef);
@@ -159,7 +162,7 @@ void FPhysicsCommandsTask::DoTask(ENamedThreads::Type CurrentThread, const FGrap
 
 //////////////////////////////////////////////////////////////////////////
 
-FPhysicsSolverAdvanceTask::FPhysicsSolverAdvanceTask(Chaos::FPhysicsSolver* InSolver, float InDt)
+FPhysicsSolverAdvanceTask::FPhysicsSolverAdvanceTask(Chaos::FPhysicsSolverBase* InSolver, float InDt)
 	: Solver(InSolver)
 	, Dt(InDt)
 {
@@ -188,10 +191,11 @@ void FPhysicsSolverAdvanceTask::DoTask(ENamedThreads::Type CurrentThread, const
 	SCOPE_CYCLE_COUNTER(STAT_ChaosTick);
 	CSV_SCOPED_TIMING_STAT_EXCLUSIVE(Physics);
 
-	StepSolver(Solver, Dt);
+	Solver->CastHelper([this](auto& Concrete){StepSolver(&Concrete,Dt);});
 }
 
-void FPhysicsSolverAdvanceTask::StepSolver(Chaos::FPhysicsSolver* InSolver, float InDt)
+template <typename TSolver>
+void FPhysicsSolverAdvanceTask::StepSolver(TSolver* InSolver, float InDt)
 {
 	using namespace Chaos;
 
@@ -212,7 +216,7 @@ void FPhysicsSolverAdvanceTask::StepSolver(Chaos::FPhysicsSolver* InSolver, floa
 	if(InSolver->bEnabled)
 	{
 		// Only process if we have something to actually simulate
-		if(Solver->HasActiveParticles())
+		if(InSolver->HasActiveParticles())
 		{
 			InSolver->AdvanceSolverBy(InDt);
 		}
diff --git a/Engine/Source/Runtime/Experimental/ChaosSolvers/Private/PBDRigidsSolver.cpp b/Engine/Source/Runtime/Experimental/ChaosSolvers/Private/PBDRigidsSolver.cpp
--- a/Engine/Source/Runtime/Experimental/ChaosSolvers/Private/PBDRigidsSolver.cpp
+++ b/Engine/Source/Runtime/Experimental/ChaosSolvers/Private/PBDRigidsSolver.cpp
@@ -203,7 +203,7 @@ namespace Chaos
 
 	template <typename Traits>
 	TPBDRigidsSolver<Traits>::TPBDRigidsSolver(const EMultiBufferMode BufferingModeIn, UObject* InOwner)
-		: Super(BufferingModeIn, InOwner)
+		: Super(BufferingModeIn, InOwner, TraitToIdx<Traits>())
 		, CurrentFrame(0)
 		, MTime(0.0)
 		, MLastDt(0.0)
diff --git a/Engine/Source/Runtime/Experimental/ChaosSolvers/Public/ChaosSolversModule.h b/Engine/Source/Runtime/Experimental/ChaosSolvers/Public/ChaosSolversModule.h
--- a/Engine/Source/Runtime/Experimental/ChaosSolvers/Public/ChaosSolversModule.h
+++ b/Engine/Source/Runtime/Experimental/ChaosSolvers/Public/ChaosSolversModule.h
@@ -16,6 +16,8 @@
 #include "Chaos/Framework/MultiBufferResource.h"
 #include "Chaos/Declares.h"
 #include "Chaos/PhysicalMaterials.h"
+#include "Chaos/Defines.h"
+#include "Chaos/EvolutionTraits.h"
 
 /** Classes that want to set the solver actor class can implement this. */
 class IChaosSolverActorClassProvider
@@ -159,18 +161,27 @@ public:
 	 * @param InOwner Ptr to some owning UObject. The module doesn't use this but allows calling code to organize solver ownership
 	 * @param bStandalone Whether the solver is standalone (not sent to physics thread - updating left to caller)
 	 */
-	Chaos::FPhysicsSolver* CreateSolver(UObject* InOwner
+	template <typename Traits>
+	Chaos::TPBDRigidsSolver<Traits>* CreateSolver(UObject* InOwner
 #if CHAOS_CHECKED
 		, const FName& DebugName = NAME_None
 #endif
-	);
-	Chaos::FPhysicsSolver* CreateSolver(UObject* InOwner, ESolverFlags InFlags
+	)
+	{
+		return CreateSolver<Traits>(InOwner,ESolverFlags::None
+#if CHAOS_CHECKED
+			,DebugName
+#endif
+		);
+	}
+	template <typename Traits>
+	Chaos::TPBDRigidsSolver<Traits>* CreateSolver(UObject* InOwner, ESolverFlags InFlags
 #if CHAOS_CHECKED
 		, const FName& DebugName = NAME_None
 #endif
 	);
 
-	void MigrateSolver(Chaos::FPhysicsSolver* InSolver, const UObject* InNewOwner);
+	void MigrateSolver(Chaos::FPhysicsSolverBase* InSolver, const UObject* InNewOwner);
 
 	Chaos::EMultiBufferMode GetBufferModeFromThreadingModel(Chaos::EThreadingMode ThreadingMode) const
 	{
@@ -235,21 +246,21 @@ public:
 	 *
 	 * Should be called on whichever thread currently owns the solver state
 	 */
-	void DestroySolver(Chaos::FPhysicsSolver* InState);
+	void DestroySolver(Chaos::FPhysicsSolverBase* InState);
 
 	/** Retrieve the master list of solvers. This contains all owned, unowned and standalone solvers */
-	const TArray<Chaos::FPhysicsSolver*>& GetAllSolvers() const;
+	const TArray<Chaos::FPhysicsSolverBase*>& GetAllSolvers() const;
 
 	/**
 	 * Read access to the current solver-state objects, be aware which thread owns this data when
 	 * attempting to use this. Physics thread will query when spinning up to get current world state
 	 */
-	TArray<const Chaos::FPhysicsSolver*> GetSolvers(const UObject* InOwner) const;
-	void GetSolvers(const UObject* InOwner, TArray<const Chaos::FPhysicsSolver*>& OutSolvers) const;
+	TArray<const Chaos::FPhysicsSolverBase*> GetSolvers(const UObject* InOwner) const;
+	void GetSolvers(const UObject* InOwner, TArray<const Chaos::FPhysicsSolverBase*>& OutSolvers) const;
 
 	/** Non-const access to the solver array is private - only the module should ever modify the storage lists */
-	TArray<Chaos::FPhysicsSolver*> GetSolversMutable(const UObject* InOwner);
-	void GetSolversMutable(const UObject* InOwner, TArray<Chaos::FPhysicsSolver*>& OutSolvers);
+	TArray<Chaos::FPhysicsSolverBase*> GetSolversMutable(const UObject* InOwner);
+	void GetSolversMutable(const UObject* InOwner, TArray<Chaos::FPhysicsSolverBase*>& OutSolvers);
 
 	/**
 	 * Outputs statistics for the solver hierarchies. Currently engine calls into this
@@ -352,11 +363,11 @@ private:
 	FDelegateHandle PreExitHandle;
 
 	// Flat list of every solver the module has created.
-	TArray<Chaos::FPhysicsSolver*> AllSolvers;
+	TArray<Chaos::FPhysicsSolverBase*> AllSolvers;
 
 	// Map of solver owners to the solvers they own. Nullptr is a valid 'owner' in that it's a map
 	// key for unowned, standalone solvers
-	TMap<const UObject*, TArray<Chaos::FPhysicsSolver*>> SolverMap;
+	TMap<const UObject*, TArray<Chaos::FPhysicsSolverBase*>> SolverMap;
 
 	// Lock for the above list to ensure we don't delete solvers out from underneath other threads
 	// or mess up the solvers array during use.
@@ -450,3 +461,14 @@ struct CHAOSSOLVERS_API FChaosScopeSolverLock
 		FChaosSolversModule::GetModule()->UnlockSolvers();
 	}
 };
+
+#if CHAOS_CHECKED
+#define EVOLUTION_TRAIT(Traits)\
+extern template CHAOSSOLVERS_API Chaos::TPBDRigidsSolver<Chaos::Traits>* FChaosSolversModule::CreateSolver(UObject* InOwner, ESolverFlags InFlags, const FName& DebugName);
+#else
+#define EVOLUTION_TRAIT(Traits)\
+extern template CHAOSSOLVERS_API Chaos::TPBDRigidsSolver<Chaos::Traits>* FChaosSolversModule::CreateSolver(UObject* InOwner,ESolverFlags InFlags);
+#endif
+
+#include "Chaos/EvolutionTraits.inl"
+#undef EVOLUTION_TRAIT
diff --git a/Engine/Source/Runtime/Experimental/ChaosSolvers/Public/Framework/DebugSolverTasks.h b/Engine/Source/Runtime/Experimental/ChaosSolvers/Public/Framework/DebugSolverTasks.h
--- a/Engine/Source/Runtime/Experimental/ChaosSolvers/Public/Framework/DebugSolverTasks.h
+++ b/Engine/Source/Runtime/Experimental/ChaosSolvers/Public/Framework/DebugSolverTasks.h
@@ -11,6 +11,7 @@
 
 namespace Chaos
 {
+class FPhysicsSolverBase;
 	/**
 	 * Async task to run the solver advance in its own debug thread, substep by substep.
 	 */
@@ -49,19 +50,19 @@ namespace Chaos
 		~FDebugSolverTasks() { Shutdown(); }
 
 		/** Add debug task entry for the specified solver. */
-		void Add(FPhysicsSolver* Solver);
+		void Add(FPhysicsSolverBase* Solver);
 
 		/** Remove the debug task entry for the specified solver, and delete its task if any was created. */
-		void Remove(FPhysicsSolver* Solver);
+		void Remove(FPhysicsSolverBase* Solver);
 
 		/** Run the specified step function in one go within the current thread, or in a debug thread substep by substep depending on the Solver's DebugSustep status. */
-		void DebugStep(FPhysicsSolver* Solver, TFunction<void()> StepFunction);
+		void DebugStep(FPhysicsSolverBase* Solver, TFunction<void()> StepFunction);
 
 		/** Shutdown all debug threads. */
 		void Shutdown();
 
 	private:
-		TMap<FPhysicsSolver*, FAsyncTask<FDebugSolverTask>*> SolverToTaskMap;
+		TMap<FPhysicsSolverBase*, FAsyncTask<FDebugSolverTask>*> SolverToTaskMap;
 	};
 }  // namespace Chaos
 
@@ -78,10 +79,10 @@ namespace Chaos
 		FDebugSolverTasks() {}
 		~FDebugSolverTasks() {}
 
-		void Add(FPhysicsSolver* /*Solver*/) {}
-		void Remove(FPhysicsSolver* /*Solver*/) {}
+		void Add(FPhysicsSolverBase* /*Solver*/) {}
+		void Remove(FPhysicsSolverBase* /*Solver*/) {}
 
-		FORCEINLINE void DebugStep(FPhysicsSolver* /*Solver*/, TFunction<void()> StepFunction) { StepFunction(); }
+		FORCEINLINE void DebugStep(FPhysicsSolverBase* /*Solver*/, TFunction<void()> StepFunction) { StepFunction(); }
 
 		void Shutdown() {}
 	};
diff --git a/Engine/Source/Runtime/Experimental/ChaosSolvers/Public/Framework/PersistentTask.h b/Engine/Source/Runtime/Experimental/ChaosSolvers/Public/Framework/PersistentTask.h
--- a/Engine/Source/Runtime/Experimental/ChaosSolvers/Public/Framework/PersistentTask.h
+++ b/Engine/Source/Runtime/Experimental/ChaosSolvers/Public/Framework/PersistentTask.h
@@ -127,12 +127,12 @@ namespace Chaos
 		 * Once the solver has been added to this task the game thread should never
 		 * touch the internal state again unless performing a sync of the data
 		 */
-		void AddSolver(FPhysicsSolver* InSolver);
+		void AddSolver(FPhysicsSolverBase* InSolver);
 
 		/**
 		 * Removes a solver from the internal list of solvers to run on the async task
 		 */
-		void RemoveSolver(FPhysicsSolver* InSolver);
+		void RemoveSolver(FPhysicsSolverBase* InSolver);
 
 		/**
 		 * Synchronize proxies to their most recent gamethread readable results
@@ -199,12 +199,14 @@ namespace Chaos
 			RETURN_QUICK_DECLARE_CYCLE_STAT(FPersistentPhysicsTask, STATGROUP_ThreadPoolAsyncTasks);
 		}
 
-		void StepSolver(FPhysicsSolver* InSolver, float InDt);
-		void HandleSolverCommands(FPhysicsSolver* InSolver);
-		void AdvanceSolver(FPhysicsSolver* InSolver, float InDt);
+		void StepSolver(FPhysicsSolverBase* InSolver, float InDt);
+		void HandleSolverCommands(FPhysicsSolverBase* InSolver);
+
+		template <typename TSolver>
+		void AdvanceSolver(TSolver* InSolver, float InDt);
 
 		// List of solvers we'll advance in this task
-		TArray<FPhysicsSolver*> Solvers;
+		TArray<FPhysicsSolverBase*> Solvers;
 
 		// Debug threads used to debug substep solver advance.
 		FDebugSolverTasks DebugSolverTasks;
diff --git a/Engine/Source/Runtime/Experimental/ChaosSolvers/Public/Framework/PhysicsTickTask.h b/Engine/Source/Runtime/Experimental/ChaosSolvers/Public/Framework/PhysicsTickTask.h
--- a/Engine/Source/Runtime/Experimental/ChaosSolvers/Public/Framework/PhysicsTickTask.h
+++ b/Engine/Source/Runtime/Experimental/ChaosSolvers/Public/Framework/PhysicsTickTask.h
@@ -11,6 +11,7 @@ class FChaosSolversModule;
 namespace Chaos
 {
 	class IDispatcher;
+	class FPhysicsSolverBase;
 }
 
 /**
@@ -27,7 +28,7 @@ public:
 	/** Construct a task that will tick the provided solver (or all solvers in the module if nullptr passed) */
 	FPhysicsTickTask(FGraphEventRef& InCompletionEvent, Chaos::FPhysicsSolver* InPhysicsSolver, float InDt);
 	/** Construct a task to tick the provided list of solvers */
-	FPhysicsTickTask(FGraphEventRef& InCompletionEvent, const TArray<Chaos::FPhysicsSolver*>& InSolverList, float InDt);
+	FPhysicsTickTask(FGraphEventRef& InCompletionEvent, const TArray<Chaos::FPhysicsSolverBase*>& InSolverList, float InDt);
 
 	/** Task API */
 	TStatId GetStatId() const;
@@ -48,7 +49,7 @@ private:
 	FChaosSolversModule* Module;
 
 	/** The solvers this task will tick */
-	TArray<Chaos::FPhysicsSolver*> SolverList;
+	TArray<Chaos::FPhysicsSolverBase*> SolverList;
 
 	/** Delta time for the solver tick */
 	float Dt;
@@ -83,7 +84,7 @@ class CHAOSSOLVERS_API FPhysicsSolverAdvanceTask
 {
 public:
 
-	FPhysicsSolverAdvanceTask(Chaos::FPhysicsSolver* InSolver, float InDt);
+	FPhysicsSolverAdvanceTask(Chaos::FPhysicsSolverBase* InSolver, float InDt);
 
 	TStatId GetStatId() const;
 	static ENamedThreads::Type GetDesiredThread();
@@ -92,9 +93,10 @@ public:
 
 private:
 
-	void StepSolver(Chaos::FPhysicsSolver* InSolver, float InDt);
+	template <typename TSolver>
+	static void StepSolver(TSolver* InSolver, float InDt);
 
-	Chaos::FPhysicsSolver* Solver;
+	Chaos::FPhysicsSolverBase* Solver;
 	float Dt;
 };
 
diff --git a/Engine/Source/Runtime/Experimental/FieldSystem/Source/FieldSystemEngine/Private/Field/FieldSystemComponent.cpp b/Engine/Source/Runtime/Experimental/FieldSystem/Source/FieldSystemEngine/Private/Field/FieldSystemComponent.cpp
--- a/Engine/Source/Runtime/Experimental/FieldSystem/Source/FieldSystemEngine/Private/Field/FieldSystemComponent.cpp
+++ b/Engine/Source/Runtime/Experimental/FieldSystem/Source/FieldSystemEngine/Private/Field/FieldSystemComponent.cpp
@@ -140,7 +140,7 @@ void UFieldSystemComponent::DispatchCommand(const FFieldSystemCommand& InCommand
 		checkSlow(PhysicsDispatcher); // Should always have one of these
 
 		// Assemble a list of compatible solvers
-		TArray<Chaos::FPhysicsSolver*> SupportedSolverList;
+		TArray<Chaos::FPhysicsSolverBase*> SupportedSolverList;
 		if(SupportedSolvers.Num() > 0)
 		{
 			for(TSoftObjectPtr<AChaosSolverActor>& SolverActorPtr : SupportedSolvers)
@@ -152,7 +152,7 @@ void UFieldSystemComponent::DispatchCommand(const FFieldSystemCommand& InCommand
 			}
 		}
 
-		TArray<Chaos::FPhysicsSolver*> WorldSolverList = ChaosModule->GetAllSolvers();
+		TArray<Chaos::FPhysicsSolverBase*> WorldSolverList = ChaosModule->GetAllSolvers();
 
 		// #BGTODO Currently all commands will end up actually executing a frame late. That's because this command has to be logged as a global command
 		// so we don't end up with multiple solver threads writing to the proxy. We need a better way to buffer up multi-solver commands so they can be
@@ -163,14 +163,17 @@ void UFieldSystemComponent::DispatchCommand(const FFieldSystemCommand& InCommand
 			{
 				const int32 NumFilterSolvers = SupportedSolverList.Num();
 
-				for(Chaos::FPhysicsSolver* Solver : WorldSolvers)
+				for(Chaos::FPhysicsSolverBase* Solver : WorldSolvers)
 				{
-					const bool bSolverValid = NumFilterSolvers == 0 || SupportedSolverList.Contains(Solver);
-
-					if(Solver->Enabled() && Solver->HasActiveParticles() && bSolverValid)
+					Solver->CastHelper([NumFilterSolvers, &SupportedSolverList, PhysicsProxy, &NewCommand](auto& Concrete)
 					{
-						PhysicsProxy->BufferCommand(Solver, NewCommand);
-					}
+						const bool bSolverValid = NumFilterSolvers == 0 || SupportedSolverList.Contains(&Concrete);
+
+						if(Concrete.Enabled() && Concrete.HasActiveParticles() && bSolverValid)
+						{
+							PhysicsProxy->BufferCommand(&Concrete,NewCommand);
+						}
+					});
 				}
 			});
 		}
