Destruction: Fixed push state not being called on custom solvers so external additions get processed correctly (like the solver floor). Fixed bad transforms in cluster union constraint construction causing cluster unions to fall through each other.
#rnx

diff --git a/Engine/Source/Runtime/Engine/Private/PhysicsEngine/Experimental/PhysScene_Chaos.cpp b/Engine/Source/Runtime/Engine/Private/PhysicsEngine/Experimental/PhysScene_Chaos.cpp
--- a/Engine/Source/Runtime/Engine/Private/PhysicsEngine/Experimental/PhysScene_Chaos.cpp
+++ b/Engine/Source/Runtime/Engine/Private/PhysicsEngine/Experimental/PhysScene_Chaos.cpp
@@ -1986,12 +1986,18 @@ void FPhysScene_ChaosInterface::StartFrame()
 			
 			if (FPhysicsSolver* Solver = GetSolver())
 			{
-				Solver->PushPhysicsState(Dispatcher);
-
 				// Make sure our solver is in the list
 				SolverList.AddUnique(Solver);
 			}
 
+			for(FPhysicsSolverBase* Solver : SolverList)
+			{
+				Solver->CastHelper([Dispatcher](auto& InSolver)
+					{
+						InSolver.PushPhysicsState(Dispatcher);
+					});
+			}
+
 			FGraphEventRef SimulationCompleteEvent = FGraphEvent::CreateGraphEvent();
 
 			// Need to fire off a parallel task to handle running physics commands and
diff --git a/Engine/Source/Runtime/Experimental/Chaos/Private/Chaos/CollisionResolution.cpp b/Engine/Source/Runtime/Experimental/Chaos/Private/Chaos/CollisionResolution.cpp
--- a/Engine/Source/Runtime/Experimental/Chaos/Private/Chaos/CollisionResolution.cpp
+++ b/Engine/Source/Runtime/Experimental/Chaos/Private/Chaos/CollisionResolution.cpp
@@ -2610,22 +2610,29 @@ namespace Chaos
 			}
 
 #if CHAOS_COLLIDE_CLUSTERED_UNIONS
-			if (Implicit0OuterType == FImplicitObjectUnionClustered::StaticType())
+			if(Implicit0OuterType == FImplicitObjectUnionClustered::StaticType())
 			{
 				const FImplicitObjectUnionClustered* Union0 = Implicit0->template GetObject<FImplicitObjectUnionClustered>();
-				if (Implicit1->HasBoundingBox())
+				if(Implicit1->HasBoundingBox())
 				{
 					TArray<Pair<const FImplicitObject*, FRigidTransform3>> Children;
-					Union0->FindAllIntersectingObjects(Children, Implicit1->BoundingBox().TransformedAABB(LocalTransform1.GetRelativeTransform(LocalTransform0)));
-					for (const auto& Child0 : Children)
+
+					// Need to get transformed bounds of 1 in the space of 0
+					FRigidTransform3 TM0 = LocalTransform0 * Collisions::GetTransform(Particle0);
+					FRigidTransform3 TM1 = LocalTransform1 * Collisions::GetTransform(Particle1);
+					FRigidTransform3 TM1ToTM0 = TM1.GetRelativeTransform(TM0);
+					FAABB3 QueryBounds = Implicit1->BoundingBox().TransformedAABB(TM1ToTM0);
+
+					Union0->FindAllIntersectingObjects(Children, QueryBounds);
+
+					for(const Pair<const FImplicitObject*, FRigidTransform3>& Child0 : Children)
 					{
-						TRigidTransform<FReal, 3> TransformedChild0 = Child0.Second * LocalTransform0;
-						ConstructConstraints<T_TRAITS>(Particle0, Particle1, Child0.First, Implicit1, TransformedChild0, LocalTransform1, CullDistance, Context, NewConstraints);
+						ConstructConstraints<T_TRAITS>(Particle0, Particle1, Child0.First, Implicit1, Child0.Second * LocalTransform0, LocalTransform1, CullDistance, Context, NewConstraints);
 					}
 				}
 				else
 				{
-					for (const auto& Child0 : Union0->GetObjects())
+					for(const TUniquePtr<FImplicitObject>& Child0 : Union0->GetObjects())
 					{
 						ConstructConstraints<T_TRAITS>(Particle0, Particle1, Child0.Get(), Implicit1, LocalTransform0, LocalTransform1, CullDistance, Context, NewConstraints);
 					}
@@ -2645,22 +2652,29 @@ namespace Chaos
 			}
 
 #if CHAOS_COLLIDE_CLUSTERED_UNIONS
-			if (Implicit1OuterType == FImplicitObjectUnionClustered::StaticType())
+			if(Implicit1OuterType == FImplicitObjectUnionClustered::StaticType())
 			{
 				const FImplicitObjectUnionClustered* Union1 = Implicit1->template GetObject<FImplicitObjectUnionClustered>();
-				if (Implicit0->HasBoundingBox())
+				if(Implicit0->HasBoundingBox())
 				{
 					TArray<Pair<const FImplicitObject*, FRigidTransform3>> Children;
-					Union1->FindAllIntersectingObjects(Children, Implicit0->BoundingBox().TransformedAABB(LocalTransform0.GetRelativeTransform(LocalTransform1)));
-					for (const auto& Child1 : Children)
+					
+					// Need to get transformed bounds of 0 in the space of 1
+					FRigidTransform3 TM0 = LocalTransform0 * Collisions::GetTransform(Particle0);
+					FRigidTransform3 TM1 = LocalTransform1 * Collisions::GetTransform(Particle1);
+					FRigidTransform3 TM0ToTM1 = TM0.GetRelativeTransform(TM1);
+					FAABB3 QueryBounds = Implicit0->BoundingBox().TransformedAABB(TM0ToTM1);
+
+					Union1->FindAllIntersectingObjects(Children, QueryBounds);
+
+					for(const Pair<const FImplicitObject*, FRigidTransform3>& Child1 : Children)
 					{
-						TRigidTransform<FReal, 3> TransformedChild1 = Child1.Second * LocalTransform1;
-						ConstructConstraints<T_TRAITS>(Particle0, Particle1, Implicit0, Child1.First, LocalTransform0, TransformedChild1, CullDistance, Context, NewConstraints);
+						ConstructConstraints<T_TRAITS>(Particle0, Particle1, Implicit0, Child1.First, LocalTransform0, Child1.Second * LocalTransform1, CullDistance, Context, NewConstraints);
 					}
 				}
 				else
 				{
-					for (const auto& Child1 : Union1->GetObjects())
+					for(const TUniquePtr<FImplicitObject>& Child1 : Union1->GetObjects())
 					{
 						ConstructConstraints<T_TRAITS>(Particle0, Particle1, Implicit0, Child1.Get(), LocalTransform0, LocalTransform1, CullDistance, Context, NewConstraints);
 					}
