* Retrieve physical materials specified on instanced shapes.
* Retrieve physical materials in the PerParticleMaterials auxilary array

diff --git a/Engine/Source/Runtime/Engine/Private/PhysicsEngine/ImmediatePhysics/ImmediatePhysicsChaos/ImmediatePhysicsSimulation_Chaos.cpp b/Engine/Source/Runtime/Engine/Private/PhysicsEngine/ImmediatePhysics/ImmediatePhysicsChaos/ImmediatePhysicsSimulation_Chaos.cpp
--- a/Engine/Source/Runtime/Engine/Private/PhysicsEngine/ImmediatePhysics/ImmediatePhysicsChaos/ImmediatePhysicsSimulation_Chaos.cpp
+++ b/Engine/Source/Runtime/Engine/Private/PhysicsEngine/ImmediatePhysics/ImmediatePhysicsChaos/ImmediatePhysicsSimulation_Chaos.cpp
@@ -182,7 +182,7 @@ namespace ImmediatePhysics_Chaos
 		FImplementation()
 			: Particles()
 			, Joints()
-			, Collisions(Particles, CollidedParticles, ParticleMaterials, 0, 0, ChaosImmediate_Collision_CullDistance, ChaosImmediate_Collision_ShapePadding)
+			, Collisions(Particles, CollidedParticles, ParticleMaterials, PerParticleMaterials, 0, 0, ChaosImmediate_Collision_CullDistance, ChaosImmediate_Collision_ShapePadding)
 			, BroadPhase(ActivePotentiallyCollidingPairs, ChaosImmediate_Collision_CullDistance)
 			, NarrowPhase()
 			, CollisionDetector(BroadPhase, NarrowPhase, Collisions)
diff --git a/Engine/Source/Runtime/Experimental/Chaos/Private/Chaos/PBDCollisionConstraints.cpp b/Engine/Source/Runtime/Experimental/Chaos/Private/Chaos/PBDCollisionConstraints.cpp
--- a/Engine/Source/Runtime/Experimental/Chaos/Private/Chaos/PBDCollisionConstraints.cpp
+++ b/Engine/Source/Runtime/Experimental/Chaos/Private/Chaos/PBDCollisionConstraints.cpp
@@ -20,6 +20,7 @@
 #include "Chaos/PBDRigidsSOAs.h"
 #include "Chaos/Sphere.h"
 #include "Chaos/Transform.h"
+#include "Chaos/CastingUtilities.h"
 #include "ChaosLog.h"
 #include "ChaosStats.h"
 #include "Containers/Queue.h"
@@ -84,7 +85,8 @@ namespace Chaos
 	FPBDCollisionConstraints::FPBDCollisionConstraints(
 		const TPBDRigidsSOAs<FReal, 3>& InParticles,
 		TArrayCollectionArray<bool>& Collided,
-		const TArrayCollectionArray<TSerializablePtr<FChaosPhysicsMaterial>>& InPerParticleMaterials,
+		const TArrayCollectionArray<TSerializablePtr<FChaosPhysicsMaterial>>& InPhysicsMaterials,
+		const TArrayCollectionArray<TUniquePtr<FChaosPhysicsMaterial>>& InPerParticlePhysicsMaterials,
 		const int32 InApplyPairIterations /*= 1*/,
 		const int32 InApplyPushOutPairIterations /*= 1*/,
 		const FReal CullDistance /*= (FReal)0*/,
@@ -94,7 +96,8 @@ namespace Chaos
 		, NumActiveSweptPointConstraints(0)
 		, NumActiveIterativeConstraints(0)
 		, MCollided(Collided)
-		, MPhysicsMaterials(InPerParticleMaterials)
+		, MPhysicsMaterials(InPhysicsMaterials)
+		, MPerParticlePhysicsMaterials(InPerParticlePhysicsMaterials)
 		, MApplyPairIterations(InApplyPairIterations)
 		, MApplyPushOutPairIterations(InApplyPushOutPairIterations)
 		, MCullDistance(CullDistance)
@@ -142,36 +145,44 @@ namespace Chaos
 		PostApplyPushOutCallback = nullptr;
 	}
 
-	const FChaosPhysicsMaterial* GetPhysicsMaterial(const TGeometryParticleHandle<FReal, 3>* Particle, const FImplicitObject* Geom, const TArrayCollectionArray<TSerializablePtr<FChaosPhysicsMaterial>>& PhysicsMaterials)
+	const FChaosPhysicsMaterial* GetPhysicsMaterial(const TGeometryParticleHandle<FReal, 3>* Particle, const FImplicitObject* Geom, const TArrayCollectionArray<TSerializablePtr<FChaosPhysicsMaterial>>& PhysicsMaterials, const TArrayCollectionArray<TUniquePtr<FChaosPhysicsMaterial>>& PerParticlePhysicsMaterials)
 	{
 		// Use the per-particle material if it exists
+		const FChaosPhysicsMaterial* UniquePhysicsMaterial = Particle->AuxilaryValue(PerParticlePhysicsMaterials).Get();
+		if (UniquePhysicsMaterial != nullptr)
+		{
+			return UniquePhysicsMaterial;
+		}
 		const FChaosPhysicsMaterial* PhysicsMaterial = Particle->AuxilaryValue(PhysicsMaterials).Get();
+		if (PhysicsMaterial != nullptr)
+		{
+			return PhysicsMaterial;
+		}
 
 		// If no particle material, see if the shape has one
-		if (PhysicsMaterial == nullptr)
+		// @todo(chaos): handle materials for meshes etc
+		for (const TUniquePtr<FPerShapeData>& ShapeData : Particle->ShapesArray())
 		{
-			// @todo(chaos): handle materials for meshes etc
-		
-			for (const TUniquePtr<FPerShapeData>& ShapeData : Particle->ShapesArray())
+			const FImplicitObject* OuterShapeGeom = ShapeData->GetGeometry().Get();
+			const FImplicitObject* InnerShapeGeom = Utilities::ImplicitChildHelper(OuterShapeGeom);
+			if (Geom == OuterShapeGeom || Geom == InnerShapeGeom)
 			{
-				if (Geom == ShapeData->GetGeometry().Get())
+				if (ShapeData->GetMaterials().Num() > 0)
 				{
-					if (ShapeData->GetMaterials().Num() > 0)
-					{
-						PhysicsMaterial = ShapeData->GetMaterials()[0].Get();
-					}
-					break;
+					return ShapeData->GetMaterials()[0].Get();
 				}
 			}
 		}
 
-		return PhysicsMaterial;
+		check(false); // The geometry used for this particle does not belong to the particle..
+
+		return nullptr;
 	}
 
 	void FPBDCollisionConstraints::UpdateConstraintMaterialProperties(FCollisionConstraintBase& Constraint)
 	{
-		const FChaosPhysicsMaterial* PhysicsMaterial0 = GetPhysicsMaterial(Constraint.Particle[0], Constraint.Manifold.Implicit[0], MPhysicsMaterials);
-		const FChaosPhysicsMaterial* PhysicsMaterial1 = GetPhysicsMaterial(Constraint.Particle[1], Constraint.Manifold.Implicit[1], MPhysicsMaterials);
+		const FChaosPhysicsMaterial* PhysicsMaterial0 = GetPhysicsMaterial(Constraint.Particle[0], Constraint.Manifold.Implicit[0], MPhysicsMaterials, MPerParticlePhysicsMaterials);
+		const FChaosPhysicsMaterial* PhysicsMaterial1 = GetPhysicsMaterial(Constraint.Particle[1], Constraint.Manifold.Implicit[1], MPhysicsMaterials, MPerParticlePhysicsMaterials);
 
 		FCollisionContact& Contact = Constraint.Manifold;
 		if (PhysicsMaterial0 && PhysicsMaterial1)
diff --git a/Engine/Source/Runtime/Experimental/Chaos/Private/Chaos/PBDRigidsEvolutionGBF.cpp b/Engine/Source/Runtime/Experimental/Chaos/Private/Chaos/PBDRigidsEvolutionGBF.cpp
--- a/Engine/Source/Runtime/Experimental/Chaos/Private/Chaos/PBDRigidsEvolutionGBF.cpp
+++ b/Engine/Source/Runtime/Experimental/Chaos/Private/Chaos/PBDRigidsEvolutionGBF.cpp
@@ -449,7 +449,7 @@ template <typename Traits>
 TPBDRigidsEvolutionGBF<Traits>::TPBDRigidsEvolutionGBF(TPBDRigidsSOAs<FReal,3>& InParticles,THandleArray<FChaosPhysicsMaterial>& SolverPhysicsMaterials,int32 InNumIterations,int32 InNumPushoutIterations,bool InIsSingleThreaded)
 	: Base(InParticles, SolverPhysicsMaterials, InNumIterations, InNumPushoutIterations, InIsSingleThreaded)
 	, Clustering(*this, Particles.GetClusteredParticles())
-	, CollisionConstraints(InParticles, Collided, PhysicsMaterials, DefaultNumPairIterations, DefaultNumPushOutPairIterations)
+	, CollisionConstraints(InParticles, Collided, PhysicsMaterials, PerParticlePhysicsMaterials, DefaultNumPairIterations, DefaultNumPushOutPairIterations)
 	, CollisionRule(CollisionConstraints)
 	, BroadPhase(InParticles, BoundsThickness, BoundsThicknessVelocityMultiplier)
 	, NarrowPhase()
diff --git a/Engine/Source/Runtime/Experimental/Chaos/Public/Chaos/PBDCollisionConstraints.h b/Engine/Source/Runtime/Experimental/Chaos/Public/Chaos/PBDCollisionConstraints.h
--- a/Engine/Source/Runtime/Experimental/Chaos/Public/Chaos/PBDCollisionConstraints.h
+++ b/Engine/Source/Runtime/Experimental/Chaos/Public/Chaos/PBDCollisionConstraints.h
@@ -51,7 +51,8 @@ public:
 
 	FPBDCollisionConstraints(const TPBDRigidsSOAs<FReal, 3>& InParticles, 
 		TArrayCollectionArray<bool>& Collided, 
-		const TArrayCollectionArray<TSerializablePtr<FChaosPhysicsMaterial>>& PerParticleMaterials, 
+		const TArrayCollectionArray<TSerializablePtr<FChaosPhysicsMaterial>>& PhysicsMaterials, 
+		const TArrayCollectionArray<TUniquePtr<FChaosPhysicsMaterial>>& PerParticlePhysicsMaterials, 
 		const int32 ApplyPairIterations = 1, const int32 ApplyPushOutPairIterations = 1, const FReal CullDistance = (FReal)0, const FReal ShapePadding = (FReal)0);
 
 	virtual ~FPBDCollisionConstraints() {}
@@ -282,6 +283,7 @@ private:
 
 	TArrayCollectionArray<bool>& MCollided;
 	const TArrayCollectionArray<TSerializablePtr<FChaosPhysicsMaterial>>& MPhysicsMaterials;
+	const TArrayCollectionArray<TUniquePtr<FChaosPhysicsMaterial>>& MPerParticlePhysicsMaterials;
 	int32 MApplyPairIterations;
 	int32 MApplyPushOutPairIterations;
 	FReal MCullDistance;
diff --git a/Engine/Source/Runtime/Experimental/Chaos/Public/Chaos/PBDRigidsEvolution.h b/Engine/Source/Runtime/Experimental/Chaos/Public/Chaos/PBDRigidsEvolution.h
--- a/Engine/Source/Runtime/Experimental/Chaos/Public/Chaos/PBDRigidsEvolution.h
+++ b/Engine/Source/Runtime/Experimental/Chaos/Public/Chaos/PBDRigidsEvolution.h
@@ -473,7 +473,6 @@ class TPBDRigidsEvolutionBase
 
 	CHAOS_API void SetPerParticlePhysicsMaterial(TGeometryParticleHandle<FReal, 3>* Particle, TUniquePtr<FChaosPhysicsMaterial> &InMaterial)
 	{
-		check(!Particle->AuxilaryValue(PerParticlePhysicsMaterials)); //shouldn't be setting non unique material if a unique one already exists
 		Particle->AuxilaryValue(PerParticlePhysicsMaterials) = MoveTemp(InMaterial);
 	}
 
