Fix CIS Linux

diff --git a/Engine/Source/Runtime/Experimental/Chaos/Public/Chaos/PBDRigidsEvolutionGBF.h b/Engine/Source/Runtime/Experimental/Chaos/Public/Chaos/PBDRigidsEvolutionGBF.h
--- a/Engine/Source/Runtime/Experimental/Chaos/Public/Chaos/PBDRigidsEvolutionGBF.h
+++ b/Engine/Source/Runtime/Experimental/Chaos/Public/Chaos/PBDRigidsEvolutionGBF.h
@@ -68,6 +68,7 @@ namespace Chaos
 		using Base::ApplyConstraints;
 		using Base::UpdateVelocities;
 		using Base::PhysicsMaterials;
+		using Base::PerParticlePhysicsMaterials;
 		using Base::ParticleDisableCount;
 		using Base::SolverPhysicsMaterials;
 		using Base::UnprepareIteration;
