Added resim type. Made it so simulated objects don't invalidate over and over. Added simplest version of resim as slave, missing pre forces and various edge cases

diff --git a/Engine/Source/Programs/HeadlessChaos/Private/HeadlessChaosTestRewind.cpp b/Engine/Source/Programs/HeadlessChaos/Private/HeadlessChaosTestRewind.cpp
--- a/Engine/Source/Programs/HeadlessChaos/Private/HeadlessChaosTestRewind.cpp
+++ b/Engine/Source/Programs/HeadlessChaos/Private/HeadlessChaosTestRewind.cpp
@@ -1032,4 +1032,71 @@ namespace ChaosTest {
 
 		Module->DestroySolver(Solver);
 	}
+
+	TYPED_TEST(AllTraits,RewindTest_ResimAsSlave)
+	{
+		auto Sphere = TSharedPtr<FImplicitObject,ESPMode::ThreadSafe>(new TSphere<FReal,3>(TVector<float,3>(0),10));
+		auto Box = TSharedPtr<FImplicitObject,ESPMode::ThreadSafe>(new TBox<FReal,3>(FVec3(-100,-100,-100),FVec3(100,100,0)));
+
+		FChaosSolversModule* Module = FChaosSolversModule::GetModule();
+		Module->ChangeThreadingMode(EChaosThreadingMode::SingleThread);
+
+		// Make a solver
+		auto* Solver = Module->CreateSolver<TypeParam>(nullptr,ESolverFlags::Standalone);
+		Solver->SetEnabled(true);
+		Solver->EnableRewindCapture(7);
+
+		// Make particles
+		auto Dynamic = TPBDRigidParticle<float,3>::CreateParticle();
+		auto Kinematic = TKinematicGeometryParticle<float,3>::CreateParticle();
+
+		Dynamic->SetGeometry(Sphere);
+		Dynamic->SetGravityEnabled(true);
+		Solver->RegisterObject(Dynamic.Get());
+
+		Kinematic->SetGeometry(Box);
+		Solver->RegisterObject(Kinematic.Get());
+
+		Dynamic->SetX(FVec3(0,0,17));
+		Dynamic->SetGravityEnabled(false);
+		Dynamic->SetV(FVec3(0,0,-1));
+		Dynamic->SetObjectState(EObjectStateType::Dynamic);
+		Dynamic->SetResimType(EResimType::SimAsSlave);
+
+		Kinematic->SetX(FVec3(0,0,0));
+
+		ChaosTest::SetParticleSimDataToCollide({Dynamic.Get(),Kinematic.Get()});
+
+		const int32 LastStep = 11;
+
+		TArray<FVec3> Xs;
+
+		for(int Step = 0; Step <= LastStep; ++Step)
+		{
+			TickSolverHelper(Module,Solver);
+			Xs.Add(Dynamic->X());
+		}
+
+		EXPECT_FLOAT_EQ(Dynamic->X()[2],10);
+
+		const int RewindStep = 5;
+
+		FRewindData* RewindData = Solver->GetRewindData();
+		EXPECT_TRUE(RewindData->RewindToFrame(RewindStep));
+
+		//make avoid collision
+		Kinematic->SetX(FVec3(0,0,100000));
+
+		for(int Step = RewindStep; Step <= LastStep; ++Step)
+		{
+			//Resim but dynamic will take old path since it's marked as ResimAsSlave
+			TickSolverHelper(Module,Solver);
+
+			EXPECT_VECTOR_FLOAT_EQ(Dynamic->X(),Xs[Step]);
+		}
+
+		EXPECT_FLOAT_EQ(Dynamic->X()[2],10);
+
+		Module->DestroySolver(Solver);
+	}
 }
diff --git a/Engine/Source/Runtime/Experimental/Chaos/Public/Chaos/ParticleDirtyFlags.h b/Engine/Source/Runtime/Experimental/Chaos/Public/Chaos/ParticleDirtyFlags.h
--- a/Engine/Source/Runtime/Experimental/Chaos/Public/Chaos/ParticleDirtyFlags.h
+++ b/Engine/Source/Runtime/Experimental/Chaos/Public/Chaos/ParticleDirtyFlags.h
@@ -17,6 +17,8 @@ class FName;
 namespace Chaos
 {
 
+enum class EResimType: uint8;
+
 class FParticlePositionRotation
 {
 public:
@@ -258,6 +260,7 @@ public:
 		SetObjectState(Other.ObjectState());
 		SetGravityEnabled(Other.GravityEnabled());
 		SetCollisionGroup(Other.CollisionGroup());
+		SetResimType(Other.ResimType());
 	}
 
 	template <typename TOther>
@@ -267,7 +270,8 @@ public:
 			&& LinearEtherDrag() == Other.LinearEtherDrag()
 			&& AngularEtherDrag() == Other.AngularEtherDrag()
 			&& GravityEnabled() == Other.GravityEnabled()
-			&& CollisionGroup() == Other.CollisionGroup();
+			&& CollisionGroup() == Other.CollisionGroup()
+			&& ResimType() == Other.ResimType();
 	}
 
 	bool operator==(const FParticleDynamicMisc& Other) const
@@ -290,6 +294,8 @@ public:
 	int32 CollisionGroup() const { return MCollisionGroup; }
 	void SetCollisionGroup(int32 InGroup){ MCollisionGroup = InGroup; }
 
+	EResimType ResimType() const { return MResimType; }
+	void SetResimType(EResimType Type) { MResimType = Type; }
 
 private:
 	FReal MLinearEtherDrag;
@@ -297,6 +303,8 @@ private:
 	int32 MCollisionGroup;
 
 	EObjectStateType MObjectState;
+	EResimType MResimType;
+
 	bool MGravityEnabled;
 };
 
diff --git a/Engine/Source/Runtime/Experimental/Chaos/Public/Chaos/ParticleHandle.h b/Engine/Source/Runtime/Experimental/Chaos/Public/Chaos/ParticleHandle.h
--- a/Engine/Source/Runtime/Experimental/Chaos/Public/Chaos/ParticleHandle.h
+++ b/Engine/Source/Runtime/Experimental/Chaos/Public/Chaos/ParticleHandle.h
@@ -85,6 +85,7 @@ void PBDRigidParticleHandleImpDefaultConstruct(TPBDRigidParticleHandleImp<T, d,
 	Concrete.SetAngularEtherDrag(0.f);
 	Concrete.SetObjectStateLowLevel(Params.bStartSleeping ? EObjectStateType::Sleeping : EObjectStateType::Dynamic);
 	Concrete.SetGravityEnabled(Params.bGravityEnabled);
+	Concrete.SetResimType(EResimType::FullSim);
 }
 
 template <typename T, int d>
@@ -109,6 +110,7 @@ void PBDRigidParticleDefaultConstruct(TPBDRigidParticle<T,d>& Concrete, const TP
 	Concrete.SetGravityEnabled(Params.bGravityEnabled);
 	Concrete.ClearEvents();
 	Concrete.SetInitialized(false);
+	Concrete.SetResimType(EResimType::FullSim);
 }
 
 
@@ -757,6 +759,7 @@ public:
 		SetAngularEtherDrag(DynamicMisc.AngularEtherDrag());
 		SetCollisionGroup(DynamicMisc.CollisionGroup());
 		SetGravityEnabled(DynamicMisc.GravityEnabled());
+		SetResimType(DynamicMisc.ResimType());
 	}
 
 	const PMatrix<T, d, d>& I() const { return PBDRigidParticles->I(ParticleIdx); }
@@ -813,6 +816,12 @@ public:
 
 	void SetGravityEnabled(bool bEnabled){ PBDRigidParticles->GravityEnabled(ParticleIdx) = bEnabled; }
 
+	EResimType ResimType() const { return PBDRigidParticles->ResimType(ParticleIdx);}
+
+	void SetResimType(EResimType ResimType){ PBDRigidParticles->ResimType(ParticleIdx) = ResimType; }
+
+
+
 	static constexpr EParticleType StaticType() { return EParticleType::Rigid; }
 };
 
@@ -1594,7 +1603,7 @@ public:
 	const FShapesArray& ShapesArray() const { return MShapesArray; }
 
 	EObjectStateType ObjectState() const;
-	void SetObjectState(const EObjectStateType InState, bool bAllowEvents = false);
+	void SetObjectState(const EObjectStateType InState, bool bAllowEvents = false, bool bInvalidate=true);
 
 	EParticleType ObjectType() const
 	{
@@ -2046,6 +2055,16 @@ public:
 		this->MInitialized = InInitialized;
 	}
 
+	void SetResimType(EResimType ResimType)
+	{
+		MMiscData.Modify(true,MDirtyFlags,Proxy,[ResimType](auto& Data){ Data.SetResimType(ResimType);});
+	}
+
+	EResimType ResimType() const
+	{
+		return MMiscData.Read().ResimType();
+	}
+
 	const TVector<T, d>& F() const { return MDynamics.Read().F(); }
 	void SetF(const TVector<T, d>& InF, bool bInvalidate = true)
 	{
@@ -2164,7 +2183,7 @@ public:
 	}
 
 	EObjectStateType ObjectState() const { return MMiscData.Read().ObjectState(); }
-	void SetObjectState(const EObjectStateType InState, bool bAllowEvents = false)
+	void SetObjectState(const EObjectStateType InState, bool bAllowEvents = false, bool bInvalidate=true)
 	{
 		if (bAllowEvents && ObjectState() != EObjectStateType::Dynamic && InState == EObjectStateType::Dynamic)
 		{
@@ -2175,8 +2194,8 @@ public:
 		{
 			// When an object is forced into a sleep state, the velocities must be zeroed and buffered,
 			// in case the velocity is queried during sleep, or in case the object is woken up again.
-			this->SetV(FVec3(0.f), true);
-			this->SetW(FVec3(0.f), true);
+			this->SetV(FVec3(0.f), bInvalidate);
+			this->SetW(FVec3(0.f), bInvalidate);
 
 			// Dynamic particle properties must be marked clean in order not to actually apply forces which
 			// have been buffered. If another force is added after the object is put to sleep, the old forces
@@ -2185,7 +2204,7 @@ public:
 			MDirtyFlags.MarkClean(ParticlePropToFlag(EParticleProperty::Dynamics));
 		}
 
-		MMiscData.Modify(true,MDirtyFlags,Proxy,[&InState](auto& Data){ Data.SetObjectState(InState);});
+		MMiscData.Modify(bInvalidate,MDirtyFlags,Proxy,[&InState](auto& Data){ Data.SetObjectState(InState);});
 
 	}
 
@@ -2409,12 +2428,12 @@ const TPBDRigidParticle<T, d>* TGeometryParticle<T, d>::CastToRigidParticle()  c
 }
 
 template <typename T, int d>
-void TGeometryParticle<T, d>::SetObjectState(const EObjectStateType InState, bool bAllowEvents)
+void TGeometryParticle<T, d>::SetObjectState(const EObjectStateType InState, bool bAllowEvents, bool bInvalidate)
 {
 	TPBDRigidParticle<T, d>* Dyn = CastToRigidParticle();
 	if (Dyn)
 	{
-		Dyn->SetObjectState(InState, bAllowEvents);
+		Dyn->SetObjectState(InState, bAllowEvents, bInvalidate);
 	}
 }
 
diff --git a/Engine/Source/Runtime/Experimental/Chaos/Public/Chaos/RigidParticles.h b/Engine/Source/Runtime/Experimental/Chaos/Public/Chaos/RigidParticles.h
--- a/Engine/Source/Runtime/Experimental/Chaos/Public/Chaos/RigidParticles.h
+++ b/Engine/Source/Runtime/Experimental/Chaos/Public/Chaos/RigidParticles.h
@@ -15,6 +15,14 @@
 namespace Chaos
 {
 
+enum class EResimType : uint8
+{
+	FullSim,	//fully re-run simulation and keep results (any forces must be applied again)
+	//SimWithPrevForces, //use previous forces and keep results (UNIMPLEMENTED)
+	SimAsSlave //use previous forces and snap to previous results regardless of variation - used to push other objects away
+	//SimAsKinematic //treat as kinematic (UNIMPLEMENTED)
+};
+
 template<class T, int d>
 struct TSleepData
 {
@@ -79,6 +87,7 @@ class TRigidParticles : public TKinematicGeometryParticles<T, d>
 		TArrayCollection::AddArray(&MIsland);
 		TArrayCollection::AddArray(&MToBeRemovedOnFracture);
 		TArrayCollection::AddArray(&MGravityEnabled);
+		TArrayCollection::AddArray(&MResimType);
 	}
 	TRigidParticles(const TRigidParticles<T, d>& Other) = delete;
 	CHAOS_API TRigidParticles(TRigidParticles<T, d>&& Other)
@@ -96,6 +105,7 @@ class TRigidParticles : public TKinematicGeometryParticles<T, d>
 		, MCollisionGroup(MoveTemp(Other.MCollisionGroup))
 		, MObjectState(MoveTemp(Other.MObjectState))
 		, MGravityEnabled(MoveTemp(Other.MGravityEnabled))
+		, MResimType(MoveTemp(Other.MResimType))
 	{
 		TArrayCollection::AddArray(&MF);
 		TArrayCollection::AddArray(&MT);
@@ -116,6 +126,7 @@ class TRigidParticles : public TKinematicGeometryParticles<T, d>
 		TArrayCollection::AddArray(&MIsland);
 		TArrayCollection::AddArray(&MToBeRemovedOnFracture);
 		TArrayCollection::AddArray(&MGravityEnabled);
+		TArrayCollection::AddArray(&MResimType);
 	}
 
 	CHAOS_API virtual ~TRigidParticles()
@@ -182,6 +193,9 @@ class TRigidParticles : public TKinematicGeometryParticles<T, d>
 	FORCEINLINE const bool& GravityEnabled(const int32 Index) const { return MGravityEnabled[Index]; }
 	FORCEINLINE bool& GravityEnabled(const int32 Index) { return MGravityEnabled[Index]; }
 
+	FORCEINLINE EResimType ResimType(const int32 Index) const { return MResimType[Index]; }
+	FORCEINLINE EResimType& ResimType(const int32 Index) { return MResimType[Index]; }
+
 	FORCEINLINE TArray<TSleepData<T, d>>& GetSleepData() { return MSleepData; }
 	FORCEINLINE	void AddSleepData(TGeometryParticleHandle<T, d>* Particle, bool Sleeping)
 	{ 
@@ -265,6 +279,7 @@ class TRigidParticles : public TKinematicGeometryParticles<T, d>
 	TArrayCollectionArray<bool> MToBeRemovedOnFracture;
 	TArrayCollectionArray<EObjectStateType> MObjectState;
 	TArrayCollectionArray<bool> MGravityEnabled;
+	TArrayCollectionArray<EResimType> MResimType;
 
 	TArray<TSleepData<T, d>> MSleepData;
 	FRWLock SleepDataLock;
diff --git a/Engine/Source/Runtime/Experimental/ChaosSolvers/Private/PBDRigidsSolver.cpp b/Engine/Source/Runtime/Experimental/ChaosSolvers/Private/PBDRigidsSolver.cpp
--- a/Engine/Source/Runtime/Experimental/ChaosSolvers/Private/PBDRigidsSolver.cpp
+++ b/Engine/Source/Runtime/Experimental/ChaosSolvers/Private/PBDRigidsSolver.cpp
@@ -183,6 +183,11 @@ namespace Chaos
 				MSolver->GetEvolution()->EndFrame(MDeltaTime);
 			}
 
+			if(FRewindData* RewindData = MSolver->GetRewindData())
+			{
+				RewindData->FinishFrame();
+			}
+
 			MSolver->GetSolverTime() += MDeltaTime;
 			MSolver->GetCurrentFrame()++;
 			MSolver->PostTickDebugDraw();
@@ -749,12 +754,18 @@ namespace Chaos
 			auto ProcessProxyPT = [bIsSingleThreaded, Manager,DirtyProxiesData,ShapeDirtyData, RewindData, this](auto& Proxy,int32 DataIdx,FDirtyProxy& Dirty,const auto& CreateHandleFunc)
 			{
 				const bool bIsNew = !Proxy->IsInitialized();
-				//single threaded version already created particle, but didn't initialize it
-				if(!bIsSingleThreaded && bIsNew)
+				if(bIsNew)
 				{
-					const auto* NonFrequentData = Dirty.ParticleData.FindNonFrequentData(*Manager,DataIdx);
-					const FUniqueIdx* UniqueIdx = NonFrequentData ? &NonFrequentData->UniqueIdx() : nullptr;
-					Proxy->SetHandle(CreateHandleFunc(UniqueIdx));
+					//single threaded version already created particle, but didn't initialize it
+					if(!bIsSingleThreaded)
+					{
+						const auto* NonFrequentData = Dirty.ParticleData.FindNonFrequentData(*Manager,DataIdx);
+						const FUniqueIdx* UniqueIdx = NonFrequentData ? &NonFrequentData->UniqueIdx() : nullptr;
+						Proxy->SetHandle(CreateHandleFunc(UniqueIdx));
+					}
+
+					auto Handle = Proxy->GetHandle();
+					Handle->GTGeometryParticle() = Proxy->GetParticle();
 				}
 
 				if(RewindData)
@@ -775,8 +786,7 @@ namespace Chaos
 				if(bIsNew)
 				{
 					auto Handle = Proxy->GetHandle();
-					Handle->GTGeometryParticle() = Proxy->GetParticle();
-					AddParticleToProxy(Handle,Proxy);;
+					AddParticleToProxy(Handle,Proxy);
 					GetEvolution()->CreateParticle(Handle);
 					Proxy->SetInitialized(true);
 				}
@@ -1047,15 +1057,15 @@ namespace Chaos
 					{
 
 						if (ActiveObject.GetParticleType() == EParticleType::Rigid)
-					{
-						//may want to remove branch using templates outside loop
-							if (MRewindData->IsResim())
-						{
-								MRewindData->PushPTDirtyData<true>(*static_cast<const FRigidParticlePhysicsProxy*>(Proxy)->GetHandle(), DataIdx++);
-						}
-						else
 						{
-								MRewindData->PushPTDirtyData<false>(*static_cast<const FRigidParticlePhysicsProxy*>(Proxy)->GetHandle(), DataIdx++);
+							//may want to remove branch using templates outside loop
+							if (MRewindData->IsResim())
+							{
+									MRewindData->PushPTDirtyData<true>(*static_cast<FRigidParticlePhysicsProxy*>(Proxy)->GetHandle(), DataIdx++);
+							}
+							else
+							{
+								MRewindData->PushPTDirtyData<false>(*static_cast<FRigidParticlePhysicsProxy*>(Proxy)->GetHandle(), DataIdx++);
 							}
 						}
 					}
diff --git a/Engine/Source/Runtime/Experimental/ChaosSolvers/Private/PhysicsProxy/SingleParticlePhysicsProxy.cpp b/Engine/Source/Runtime/Experimental/ChaosSolvers/Private/PhysicsProxy/SingleParticlePhysicsProxy.cpp
--- a/Engine/Source/Runtime/Experimental/ChaosSolvers/Private/PhysicsProxy/SingleParticlePhysicsProxy.cpp
+++ b/Engine/Source/Runtime/Experimental/ChaosSolvers/Private/PhysicsProxy/SingleParticlePhysicsProxy.cpp
@@ -282,7 +282,7 @@ void FSingleParticlePhysicsProxy<Chaos::TPBDRigidParticle<float, 3>>::PullFromPh
 		//question: is it ok to call this when it was one of the other properties that changed?
 		if (!Particle->IsDirty(Chaos::EParticleFlags::DynamicMisc))
 		{
-			Particle->SetObjectState(Buffer->MObjectState, true);
+			Particle->SetObjectState(Buffer->MObjectState, true, /*bInvalidate=*/false);
 		}
 	}
 }
diff --git a/Engine/Source/Runtime/Experimental/ChaosSolvers/Public/RewindData.h b/Engine/Source/Runtime/Experimental/ChaosSolvers/Public/RewindData.h
--- a/Engine/Source/Runtime/Experimental/ChaosSolvers/Public/RewindData.h
+++ b/Engine/Source/Runtime/Experimental/ChaosSolvers/Public/RewindData.h
@@ -324,7 +324,8 @@ public:
 		});
 	}
 
-	void SyncToParticle(TGeometryParticle<FReal,3>& Particle) const
+	template <typename TParticle>
+	void SyncToParticle(TParticle& Particle) const
 	{
 		ParticlePositionRotation.SyncToParticle([&Particle](const auto& Data)
 		{
@@ -769,17 +770,17 @@ public:
 				if(const FGeometryParticleStateBase* RewindState = GetStateAtFrameImp(DirtyParticleInfo, Frame))
 				{
 
-					LatestState.SyncIfDirty(FDirtyPropData(DestManager,DataIdx++),*DirtyParticleInfo.Particle,*RewindState);
+					LatestState.SyncIfDirty(FDirtyPropData(DestManager,DataIdx++),*DirtyParticleInfo.GetGTParticle(),*RewindState);
 					CoalesceBack(DirtyParticleInfo.Frames, CurFrame);
 
-					RewindState->SyncToParticle(*DirtyParticleInfo.Particle);
+					RewindState->SyncToParticle(*DirtyParticleInfo.GetGTParticle());
 				}
 			}
 			else
 			{
 				if(const FGeometryParticleStateBase* RewindState = GetStateAtFrameImp(DirtyParticleInfo,Frame))
 				{
-					RewindState->SyncToParticle(*DirtyParticleInfo.Particle);
+					RewindState->SyncToParticle(*DirtyParticleInfo.GetGTParticle());
 				}
 			}
 		}
@@ -833,6 +834,7 @@ public:
 	EFutureQueryResult GetFutureStateAtFrame(FGeometryParticleState& OutState,int32 Frame) const
 	{
 		ensure(IsResim());
+		ensure(IsInGameThread());
 		const TGeometryParticle<FReal,3>& Particle = OutState.GetParticle();
 
 		if(const FDirtyParticleInfo* Info = FindParticle(Particle.UniqueIdx()))
@@ -878,8 +880,43 @@ public:
 				{
 					Info.Desync(CurFrame-1, LatestFrame);
 				}
+
+				if(auto* Rigid = Info.GetPTParticle()->CastToRigidParticle())
+				{
+					if(Rigid->ResimType() == EResimType::SimAsSlave)
+					{
+						//TODO: Need to reapply forces automatically
+					}
+				}
+			}
+		}
+	}
+
+	void FinishFrame()
+	{
+		QUICK_SCOPE_CYCLE_COUNTER(RewindDataFinishFrame);
+		if(IsResim())
+		{
+			ensure(IsInGameThread());
+			//snap particles forward that are not desynced or  do not have resim enabled
+			//TODO: handle desync case
+			for(FDirtyParticleInfo& Info : AllDirtyParticles)
+			{
+				if(auto* Rigid = Info.GetPTParticle()->CastToRigidParticle())
+				{
+					if(Rigid->ResimType() == EResimType::SimAsSlave)
+					{
+						ensure(!Info.bDesync);
+						const FGeometryParticleStateBase* State = GetStateAtFrameImp(Info,CurFrame);
+						if(ensure(State != nullptr))
+						{
+							State->SyncToParticle(*Rigid);
+						}
+					}
+				}
 			}
 		}
+		
 	}
 
 	bool IsResim() const
@@ -934,7 +971,7 @@ public:
 		auto ProcessProxy = [this,&SrcManagerWrapper, SrcDataIdx, Dirty, &DestManagerWrapper](const auto Proxy)
 		{
 			const auto PTParticle = Proxy->GetHandle();
-			FDirtyParticleInfo& Info = FindOrAddParticle(*Proxy->GetParticle(),PTParticle->UniqueIdx());
+			FDirtyParticleInfo& Info = FindOrAddParticle(*PTParticle);
 			Info.LastDirtyFrame = CurFrame;
 			Info.GTDirtyOnFrame[CurFrame].SetWave(CurFrame,CurWave);
 
@@ -1002,25 +1039,32 @@ public:
 	}
 
 	template <bool bResim>
-	void PushPTDirtyData(const TPBDRigidParticleHandle<FReal,3>& Rigid, const int32 SrcDataIdx)
+	void PushPTDirtyData(TPBDRigidParticleHandle<FReal,3>& Rigid, const int32 SrcDataIdx)
 	{
 		const int32 DestDataIdx = SrcDataIdx + DataIdxOffset;
 
-		//todo: is this check needed? why do we pass sleeping rigids into this function?
-		if(SimWritablePropsMayChange(Rigid))
+		if(bResim && Rigid.ResimType() == EResimType::SimAsSlave)
+		{
+			//resim not allowed for this particle so don't modify
+		}
+		else
 		{
-			FDirtyParticleInfo& Info = FindOrAddParticle(*Rigid.GTGeometryParticle(), Rigid.UniqueIdx());
-			Info.LastDirtyFrame = CurFrame-1;
+			//todo: is this check needed? why do we pass sleeping rigids into this function?
+			if(SimWritablePropsMayChange(Rigid))
+			{
+				FDirtyParticleInfo& Info = FindOrAddParticle(Rigid);
+				Info.LastDirtyFrame = CurFrame-1;
 
-			//User called PrepareManagerForFrame (or PrepareFrameForPTDirty) for the previous frame, so use it
-			FDirtyPropData DestManagerWrapper(Managers[CurFrame-1].Manager.Get(), DestDataIdx);
+				//User called PrepareManagerForFrame (or PrepareFrameForPTDirty) for the previous frame, so use it
+				FDirtyPropData DestManagerWrapper(Managers[CurFrame-1].Manager.Get(), DestDataIdx);
 
-			//sim-writable properties changed at head, so we must write down what they were
-			FGeometryParticleStateBase& LatestState = Info.AddFrame(CurFrame-1);
-			LatestState.SyncSimWritablePropsFromSim(DestManagerWrapper,Rigid);
+				//sim-writable properties changed at head, so we must write down what they were
+				FGeometryParticleStateBase& LatestState = Info.AddFrame(CurFrame-1);
+				LatestState.SyncSimWritablePropsFromSim(DestManagerWrapper,Rigid);
 
-			//update any previous frames that were pointing at head
-			CoalesceBack(Info.Frames, CurFrame-1);
+				//update any previous frames that were pointing at head
+				CoalesceBack(Info.Frames, CurFrame-1);
+			}
 		}
 	}
 
@@ -1126,15 +1170,19 @@ private:
 	{
 		TCircularBuffer<FFrameInfo> Frames;
 		TCircularBuffer<FDirtyFrameInfo> GTDirtyOnFrame;
-		TGeometryParticle<FReal,3>* Particle;
+	private:
+		TGeometryParticle<FReal,3>* GTParticle;
+		TGeometryParticleHandle<FReal,3>* PTParticle;
+	public:
 		FUniqueIdx CachedUniqueIdx;	//Needed when manipulating on physics thread and Particle data cannot be read
 		int32 LastDirtyFrame;	//Track how recently this was made dirty
 		bool bDesync;
 
-		FDirtyParticleInfo(TGeometryParticle<FReal,3>& UnsafeGTParticle,const FUniqueIdx UniqueIdx,const int32 CurFrame,const int32 NumFrames)
+		FDirtyParticleInfo(TGeometryParticle<FReal,3>& UnsafeGTParticle, TGeometryParticleHandle<FReal,3>& InPTParticle, const FUniqueIdx UniqueIdx,const int32 CurFrame,const int32 NumFrames)
 		: Frames(NumFrames)
 		, GTDirtyOnFrame(NumFrames)
-		, Particle(&UnsafeGTParticle)
+		, GTParticle(&UnsafeGTParticle)
+		, PTParticle(&InPTParticle)
 		, CachedUniqueIdx(UniqueIdx)
 		, LastDirtyFrame(CurFrame)
 		, bDesync(true)
@@ -1142,6 +1190,17 @@ private:
 
 		}
 
+		TGeometryParticle<FReal,3>* GetGTParticle() const
+		{
+			ensure(IsInGameThread());
+			return GTParticle;
+		}
+
+		TGeometryParticleHandle<FReal,3>* GetPTParticle() const
+		{
+			return PTParticle;
+		}
+
 		FGeometryParticleStateBase& AddFrame(int32 FrameIdx)
 		{
 			FFrameInfo& Info = Frames[FrameIdx];
@@ -1218,14 +1277,17 @@ private:
 		return nullptr;
 	}
 
-	FDirtyParticleInfo& FindOrAddParticle(TGeometryParticle<FReal,3>& UnsafeGTParticle,const FUniqueIdx UniqueIdx)
+	FDirtyParticleInfo& FindOrAddParticle(TGeometryParticleHandle<FReal,3>& PTParticle)
 	{
+		const FUniqueIdx UniqueIdx = PTParticle.UniqueIdx();
 		if(FDirtyParticleInfo* Info = FindParticle(UniqueIdx))
 		{
 			return *Info;
 		}
 
-		const int32 DirtyIdx = AllDirtyParticles.Add(FDirtyParticleInfo(UnsafeGTParticle,UniqueIdx,CurFrame,Managers.Capacity()));
+		auto* GTUnsafeParticle = PTParticle.GTGeometryParticle();
+		ensure(GTUnsafeParticle != nullptr);
+		const int32 DirtyIdx = AllDirtyParticles.Add(FDirtyParticleInfo(*GTUnsafeParticle, PTParticle,UniqueIdx,CurFrame,Managers.Capacity()));
 		ParticleToAllDirtyIdx.Add(UniqueIdx,DirtyIdx);
 
 		return AllDirtyParticles[DirtyIdx];
