Chaos Physics: Heuristic: Only use accumulated impulse clipping if the contacts did not move much (local space) during previous iterations.

#[fyi] Michael.Lentine, Benn.Gallagher, Chris.Caulfield, Brice.Criswell, Jaco.vandyk

diff --git a/Engine/Source/Runtime/Experimental/Chaos/Private/Chaos/CollisionResolution.cpp b/Engine/Source/Runtime/Experimental/Chaos/Private/Chaos/CollisionResolution.cpp
--- a/Engine/Source/Runtime/Experimental/Chaos/Private/Chaos/CollisionResolution.cpp
+++ b/Engine/Source/Runtime/Experimental/Chaos/Private/Chaos/CollisionResolution.cpp
@@ -2070,6 +2070,9 @@ namespace Chaos
 			const FImplicitObject& Implicit0 = *Constraint.Manifold.Implicit[0];
 			const FImplicitObject& Implicit1 = *Constraint.Manifold.Implicit[1];
 
+			const FVec3 OriginalContactPositionLocal0 = WorldTransform0.InverseTransformPosition(Constraint.Manifold.Location);
+			const FVec3 OriginalContactPositionLocal1 = WorldTransform1.InverseTransformPosition(Constraint.Manifold.Location);
+
 			switch (Constraint.Manifold.ShapesType)
 			{
 			case EContactShapesType::SphereSphere:
@@ -2183,6 +2186,10 @@ namespace Chaos
 				ensure(false);
 				break;
 			}
+
+			const FVec3 NewContactPositionLocal0 = WorldTransform0.InverseTransformPosition(Constraint.Manifold.Location);
+			const FVec3 NewContactPositionLocal1 = WorldTransform1.InverseTransformPosition(Constraint.Manifold.Location);
+			Constraint.Manifold.ContactMoveSQRDistance = FMath::Max((NewContactPositionLocal0 - OriginalContactPositionLocal0).SizeSquared(), (NewContactPositionLocal1 - OriginalContactPositionLocal1).SizeSquared());
 		}
 
 		template<typename T_TRAITS>
@@ -2407,10 +2414,12 @@ namespace Chaos
 			// Get the plane and point transforms (depends which body owns the plane)
 			const FRigidTransform3& PlaneTransform = (Constraint.GetManifoldPlaneOwnerIndex() == 0) ? WorldTransform0 : WorldTransform1;
 			const FRigidTransform3& PointsTransform = (Constraint.GetManifoldPlaneOwnerIndex() == 0) ? WorldTransform1 : WorldTransform0;
+			const FVec3 OriginalContactPositionLocal = PointsTransform.InverseTransformPosition(Constraint.Manifold.Location);
 
 			// World-space manifold plane
 			FVec3 PlaneNormal = PlaneTransform.TransformVectorNoScale(Constraint.GetManifoldPlaneNormal());
 			FVec3 PlanePos = PlaneTransform.TransformPosition(Constraint.GetManifoldPlanePosition());
+			FReal ContactMoveSQRDistance = 0;
 
 			// Select the best manifold point
 			for (int32 PointIndex = 0; PointIndex < NumPoints; ++PointIndex)
@@ -2428,8 +2437,10 @@ namespace Chaos
 					Constraint.Manifold.Phi = PointDistance;
 					Constraint.Manifold.Location = ContactPos;
 					Constraint.Manifold.Normal = ContactNormal;
+					ContactMoveSQRDistance = (OriginalContactPositionLocal - Constraint.GetManifoldPoint(PointIndex)).SizeSquared();
 				}
 			}
+			Constraint.Manifold.ContactMoveSQRDistance = ContactMoveSQRDistance;
 		}
 
 
diff --git a/Engine/Source/Runtime/Experimental/Chaos/Private/Chaos/PBDCollisionConstraintsContact.cpp b/Engine/Source/Runtime/Experimental/Chaos/Private/Chaos/PBDCollisionConstraintsContact.cpp
--- a/Engine/Source/Runtime/Experimental/Chaos/Private/Chaos/PBDCollisionConstraintsContact.cpp
+++ b/Engine/Source/Runtime/Experimental/Chaos/Private/Chaos/PBDCollisionConstraintsContact.cpp
@@ -29,6 +29,10 @@ namespace Chaos
 		int32 Chaos_Collision_ForceApplyType = 0;
 		FAutoConsoleVariableRef CVarChaosCollisionAlternativeApply(TEXT("p.Chaos.Collision.ForceApplyType"), Chaos_Collision_ForceApplyType, TEXT("Force Apply step to use Velocity(1) or Position(2) modes"));
 
+		float Chaos_Collision_ContactMovementAllowance = 0.05f;
+		FAutoConsoleVariableRef CVarChaosCollisionContactMovementAllowance(TEXT("p.Chaos.Collision.AntiJitterContactMovementAllowance"), Chaos_Collision_ContactMovementAllowance, 
+			TEXT("If a contact is close to where it was during a previous iteration, we will assume it is the same contact that moved (to reduce jitter). Expressed as the fraction of movement distance and Centre of Mass distance to the contact point"));
+
 		extern void UpdateManifold(FRigidBodyMultiPointContactConstraint& Constraint, const FReal CullDistance, const FCollisionContext& Context)
 		{
 			const FRigidTransform3 Transform0 = GetTransform(Constraint.Particle[0]);
@@ -118,12 +122,12 @@ namespace Chaos
 			}
 		}
 		
-		FVec3 ApplyContactAccImpulse(FCollisionContact& Contact,
+		void ApplyContactAccImpulse(FCollisionContact& Contact,
 			TGenericParticleHandle<FReal, 3> Particle0,
 			TGenericParticleHandle<FReal, 3> Particle1,
 			const FContactIterationParameters& IterationParameters,
 			const FContactParticleParameters& ParticleParameters,
-			FVec3& AccumulatedImpulse)
+			FVec3& AccumulatedImpulse)  // InOut
 		{
 
 			TPBDRigidParticleHandle<FReal, 3>* PBDRigid0 = Particle0->CastToRigidParticle();
@@ -165,7 +169,13 @@ namespace Chaos
 
 			// clip the impulse so that the accumulated impulse is not in the wrong direction and in the friction cone
 			// Clipping the accumulated impulse instead of the delta impulse (for the current iteration) is very important for jitter
-			const FVec3 NewUnclippedAccumulatedImpulse = AccumulatedImpulse + Impulse;
+			// Use a zero vector for the accumulated impulse if we are not allowed to use it (when the contact moves significantly between calls).
+			const FReal RelativeScale = FMath::Min(VectorToPoint1.SizeSquared(), VectorToPoint2.SizeSquared()); // Use this as a scale to avoid absolute distances
+			bool bUseAccumalatedImpulseInSolve = Contact.bUseAccumalatedImpulseInSolve && 
+				((Contact.ContactMoveSQRDistance == (FReal)0) || (RelativeScale > (FReal)0 && Contact.ContactMoveSQRDistance / RelativeScale < Chaos_Collision_ContactMovementAllowance * Chaos_Collision_ContactMovementAllowance));
+			Contact.bUseAccumalatedImpulseInSolve = bUseAccumalatedImpulseInSolve; // Once disabled, disabled until the end of the frame
+			const FVec3 AccImp = Contact.bUseAccumalatedImpulseInSolve ? AccumulatedImpulse: FVec3((FReal)0);
+			const FVec3 NewUnclippedAccumulatedImpulse = AccImp + Impulse;
 			FVec3 ClippedAccumulatedImpulse(0.0f);  // To be calculated
 			{
 				//Project to normal
@@ -204,7 +214,7 @@ namespace Chaos
 									WorldSpaceInvI1,
 									WorldSpaceInvI2);
 							}
-							ClippedAccumulatedImpulse = Impulse + AccumulatedImpulse;
+							ClippedAccumulatedImpulse = Impulse + AccImp;
 						}
 						else
 						{
@@ -219,7 +229,7 @@ namespace Chaos
 				}
 			}
 
-			FVec3 OutDeltaImpulse = ClippedAccumulatedImpulse - AccumulatedImpulse;
+			FVec3 OutDeltaImpulse = ClippedAccumulatedImpulse - AccImp;
 			if (Chaos_Collision_EnergyClampEnabled != 0)
 			{
 				// Clamp the delta impulse to make sure we don't gain kinetic energy (ignore potential energy)
@@ -228,7 +238,6 @@ namespace Chaos
 				OutDeltaImpulse = GetEnergyClampedImpulse(Particle0->CastToRigidParticle(), Particle1->CastToRigidParticle(), OutDeltaImpulse, VectorToPoint1, VectorToPoint2, Body1Velocity, Body2Velocity);
 			}
 
-			AccumulatedImpulse += OutDeltaImpulse; // Now update the accumulated Impulse
 			if (bIsRigidDynamic0)
 			{
 				// Velocity update for next step
@@ -257,8 +266,9 @@ namespace Chaos
 				Q1.Normalize();
 				FParticleUtilities::SetCoMWorldTransform(PBDRigid1, P1, Q1);
 			}
+
+			AccumulatedImpulse += OutDeltaImpulse; // Now update the accumulated Impulse (we need to keep track of this even if it is not used by the contact solver)
 			
-			return AccumulatedImpulse;
 		}
 
 		// Apply contacts, impulse clipping is done on delta impulses as apposed to Accumulated impulses
diff --git a/Engine/Source/Runtime/Experimental/Chaos/Public/Chaos/Collision/PBDCollisionConstraint.h b/Engine/Source/Runtime/Experimental/Chaos/Public/Chaos/Collision/PBDCollisionConstraint.h
--- a/Engine/Source/Runtime/Experimental/Chaos/Public/Chaos/Collision/PBDCollisionConstraint.h
+++ b/Engine/Source/Runtime/Experimental/Chaos/Public/Chaos/Collision/PBDCollisionConstraint.h
@@ -20,13 +20,14 @@ namespace Chaos
 	{
 	public:
 		FCollisionContact(const FImplicitObject* InImplicit0 = nullptr, const FImplicitObject* InImplicit1 = nullptr)
-			: bDisabled(false), Normal(0), Location(0), Phi(FLT_MAX), Friction(0), AngularFriction(0), Restitution(0), InvInertiaScale0(1.f), InvInertiaScale1(1.f), ShapesType(EContactShapesType::Unknown)
+			: bDisabled(false), bUseAccumalatedImpulseInSolve(true), Normal(0), Location(0), Phi(FLT_MAX), Friction(0), AngularFriction(0), Restitution(0), InvInertiaScale0(1.f), InvInertiaScale1(1.f), ContactMoveSQRDistance(0), ShapesType(EContactShapesType::Unknown)
 		{
 			Implicit[0] = InImplicit0;
 			Implicit[1] = InImplicit1;
 		}
 
 		bool bDisabled;
+		bool bUseAccumalatedImpulseInSolve; // Accumulated Impulse will only be used when the contact did not move significantly during iterations (This will reduce jitter)
 		FVec3 Normal;
 		FVec3 Location;
 		FReal Phi;
@@ -36,6 +37,7 @@ namespace Chaos
 		FReal Restitution;
 		FReal InvInertiaScale0;
 		FReal InvInertiaScale1;
+		FReal ContactMoveSQRDistance; // How much the contact position moved after the last update
 		EContactShapesType ShapesType;
 
 
