Chaos - Query analytic ignore flag when building geometry collection simplicials so if analytics are disabled we still pick up collision particles.
#rnx

diff --git a/Engine/Source/Runtime/Experimental/ChaosSolvers/Private/PhysicsProxy/GeometryCollectionPhysicsProxy.cpp b/Engine/Source/Runtime/Experimental/ChaosSolvers/Private/PhysicsProxy/GeometryCollectionPhysicsProxy.cpp
--- a/Engine/Source/Runtime/Experimental/ChaosSolvers/Private/PhysicsProxy/GeometryCollectionPhysicsProxy.cpp
+++ b/Engine/Source/Runtime/Experimental/ChaosSolvers/Private/PhysicsProxy/GeometryCollectionPhysicsProxy.cpp
@@ -611,8 +611,13 @@ void TGeometryCollectionPhysicsProxy<Traits>::InitializeDynamicCollection(FGeome
 
 	// process simplicials
 	{
+		// CVar defined in BodyInstance but pertinent here as we will need to copy simplicials in the case that this is set.
+		// Original CVar is read-only so taking a static ptr here is fine as the value cannot be changed
+		static IConsoleVariable* AnalyticDisableCVar = IConsoleManager::Get().FindConsoleVariable(TEXT("p.IgnoreAnalyticCollisionsOverride"));
+		static const bool bAnalyticsDisabled = (AnalyticDisableCVar && AnalyticDisableCVar->GetBool());
+
 		if (RestCollection.HasAttribute(DynamicCollection.SimplicialsAttribute, FTransformCollection::TransformGroup)
-			&& Params.Shared.SizeSpecificData[0].CollisionType == ECollisionTypeEnum::Chaos_Surface_Volumetric)
+			&& (Params.Shared.SizeSpecificData[0].CollisionType == ECollisionTypeEnum::Chaos_Surface_Volumetric || bAnalyticsDisabled))
 		{
 			const auto& RestSimplicials = RestCollection.GetAttribute<TUniquePtr<FSimplicial>>(
 				DynamicCollection.SimplicialsAttribute, FTransformCollection::TransformGroup);
