Make it so non FullResim particles automatically have their dynamics reapplied

diff --git a/Engine/Source/Programs/HeadlessChaos/Private/HeadlessChaosTestRewind.cpp b/Engine/Source/Programs/HeadlessChaos/Private/HeadlessChaosTestRewind.cpp
--- a/Engine/Source/Programs/HeadlessChaos/Private/HeadlessChaosTestRewind.cpp
+++ b/Engine/Source/Programs/HeadlessChaos/Private/HeadlessChaosTestRewind.cpp
@@ -1061,7 +1061,7 @@ namespace ChaosTest {
 		Dynamic->SetGravityEnabled(false);
 		Dynamic->SetV(FVec3(0,0,-1));
 		Dynamic->SetObjectState(EObjectStateType::Dynamic);
-		Dynamic->SetResimType(EResimType::SimAsSlave);
+		Dynamic->SetResimType(EResimType::ResimAsSlave);
 
 		Kinematic->SetX(FVec3(0,0,0));
 
@@ -1099,4 +1099,67 @@ namespace ChaosTest {
 
 		Module->DestroySolver(Solver);
 	}
+
+	TYPED_TEST(AllTraits,RewindTest_ResimAsSlaveWithForces)
+	{
+		auto Box = TSharedPtr<FImplicitObject,ESPMode::ThreadSafe>(new TBox<FReal,3>(FVec3(-10,-10,-10),FVec3(10,10,10)));
+
+		FChaosSolversModule* Module = FChaosSolversModule::GetModule();
+		Module->ChangeThreadingMode(EChaosThreadingMode::SingleThread);
+
+		// Make a solver
+		auto* Solver = Module->CreateSolver<TypeParam>(nullptr,ESolverFlags::Standalone);
+		Solver->SetEnabled(true);
+		Solver->EnableRewindCapture(7);
+
+		// Make particles
+		auto FullSim = TPBDRigidParticle<float,3>::CreateParticle();
+		auto SlaveSim = TPBDRigidParticle<float,3>::CreateParticle();
+
+		FullSim->SetGeometry(Box);
+		FullSim->SetGravityEnabled(false);
+		Solver->RegisterObject(FullSim.Get());
+
+		SlaveSim->SetGeometry(Box);
+		SlaveSim->SetGravityEnabled(false);
+		Solver->RegisterObject(SlaveSim.Get());
+
+		FullSim->SetX(FVec3(0,0,20));
+		FullSim->SetObjectState(EObjectStateType::Dynamic);
+		FullSim->SetM(1);
+		FullSim->SetInvM(1);
+
+		SlaveSim->SetX(FVec3(0,0,0));
+		SlaveSim->SetResimType(EResimType::ResimAsSlave);
+		SlaveSim->SetM(1);
+		SlaveSim->SetInvM(1);
+
+		ChaosTest::SetParticleSimDataToCollide({FullSim.Get(),SlaveSim.Get()});
+
+		const int32 LastStep = 11;
+
+		TArray<FVec3> Xs;
+
+		for(int Step = 0; Step <= LastStep; ++Step)
+		{
+			SlaveSim->SetLinearImpulse(FVec3(0,0,0.5));
+			TickSolverHelper(Module,Solver);
+			Xs.Add(FullSim->X());
+		}
+
+		const int RewindStep = 5;
+
+		FRewindData* RewindData = Solver->GetRewindData();
+		EXPECT_TRUE(RewindData->RewindToFrame(RewindStep));
+
+		for(int Step = RewindStep; Step <= LastStep; ++Step)
+		{
+			//resim - slave sim should have its impulses automatically added thus moving FullSim in the exact same way
+			TickSolverHelper(Module,Solver);
+
+			EXPECT_VECTOR_FLOAT_EQ(FullSim->X(),Xs[Step]);
+		}
+
+		Module->DestroySolver(Solver);
+	}
 }
diff --git a/Engine/Source/Runtime/Experimental/Chaos/Public/Chaos/ParticleHandle.h b/Engine/Source/Runtime/Experimental/Chaos/Public/Chaos/ParticleHandle.h
--- a/Engine/Source/Runtime/Experimental/Chaos/Public/Chaos/ParticleHandle.h
+++ b/Engine/Source/Runtime/Experimental/Chaos/Public/Chaos/ParticleHandle.h
@@ -85,7 +85,7 @@ void PBDRigidParticleHandleImpDefaultConstruct(TPBDRigidParticleHandleImp<T, d,
 	Concrete.SetAngularEtherDrag(0.f);
 	Concrete.SetObjectStateLowLevel(Params.bStartSleeping ? EObjectStateType::Sleeping : EObjectStateType::Dynamic);
 	Concrete.SetGravityEnabled(Params.bGravityEnabled);
-	Concrete.SetResimType(EResimType::FullSim);
+	Concrete.SetResimType(EResimType::FullResim);
 }
 
 template <typename T, int d>
@@ -110,7 +110,7 @@ void PBDRigidParticleDefaultConstruct(TPBDRigidParticle<T,d>& Concrete, const TP
 	Concrete.SetGravityEnabled(Params.bGravityEnabled);
 	Concrete.ClearEvents();
 	Concrete.SetInitialized(false);
-	Concrete.SetResimType(EResimType::FullSim);
+	Concrete.SetResimType(EResimType::FullResim);
 }
 
 
diff --git a/Engine/Source/Runtime/Experimental/Chaos/Public/Chaos/RigidParticles.h b/Engine/Source/Runtime/Experimental/Chaos/Public/Chaos/RigidParticles.h
--- a/Engine/Source/Runtime/Experimental/Chaos/Public/Chaos/RigidParticles.h
+++ b/Engine/Source/Runtime/Experimental/Chaos/Public/Chaos/RigidParticles.h
@@ -17,10 +17,10 @@ namespace Chaos
 
 enum class EResimType : uint8
 {
-	FullSim,	//fully re-run simulation and keep results (any forces must be applied again)
-	//SimWithPrevForces, //use previous forces and keep results (UNIMPLEMENTED)
-	SimAsSlave //use previous forces and snap to previous results regardless of variation - used to push other objects away
-	//SimAsKinematic //treat as kinematic (UNIMPLEMENTED)
+	FullResim,	//fully re-run simulation and keep results (any forces must be applied again)
+	//ResimWithPrevForces, //use previous forces and keep results (UNIMPLEMENTED)
+	ResimAsSlave //use previous forces and snap to previous results regardless of variation - used to push other objects away
+	//ResimAsKinematic //treat as kinematic (UNIMPLEMENTED)
 };
 
 template<class T, int d>
diff --git a/Engine/Source/Runtime/Experimental/ChaosSolvers/Public/RewindData.h b/Engine/Source/Runtime/Experimental/ChaosSolvers/Public/RewindData.h
--- a/Engine/Source/Runtime/Experimental/ChaosSolvers/Public/RewindData.h
+++ b/Engine/Source/Runtime/Experimental/ChaosSolvers/Public/RewindData.h
@@ -356,6 +356,15 @@ public:
 			{
 				Rigid->SetMassProps(Data);
 			});
+
+			if(Rigid->ResimType() != EResimType::FullResim)
+			{
+				//Not full resim so apply dynamics automatically
+				Dynamics.SyncToParticle([Rigid](const auto& Data)
+				{
+					Rigid->SetDynamics(Data);
+				});
+			}
 		}
 	}
 
@@ -874,20 +883,20 @@ public:
 			}
 			else if(IsResim())
 			{
+				EResimType ResimType = EResimType::FullResim;
+				if(auto* Rigid = Info.GetPTParticle()->CastToRigidParticle())
+				{
+					ResimType = Rigid->ResimType();
+				}
+
 				//During a resim it's possible the user will not dirty a particle that was previously dirty.
 				//If this happens we need to mark the particle as desynced
-				if(!Info.bDesync && Info.GTDirtyOnFrame[CurFrame-1].MissingWrite(CurFrame-1, CurWave))
+				if(ResimType == EResimType::FullResim && !Info.bDesync && Info.GTDirtyOnFrame[CurFrame-1].MissingWrite(CurFrame-1, CurWave))
 				{
 					Info.Desync(CurFrame-1, LatestFrame);
 				}
 
-				if(auto* Rigid = Info.GetPTParticle()->CastToRigidParticle())
-				{
-					if(Rigid->ResimType() == EResimType::SimAsSlave)
-					{
-						//TODO: Need to reapply forces automatically
-					}
-				}
+				
 			}
 		}
 	}
@@ -904,7 +913,7 @@ public:
 			{
 				if(auto* Rigid = Info.GetPTParticle()->CastToRigidParticle())
 				{
-					if(Rigid->ResimType() == EResimType::SimAsSlave)
+					if(Rigid->ResimType() == EResimType::ResimAsSlave)
 					{
 						ensure(!Info.bDesync);
 						const FGeometryParticleStateBase* State = GetStateAtFrameImp(Info,CurFrame);
@@ -978,18 +987,29 @@ public:
 			//check if particle has desynced
 			if(bResim)
 			{
-				FGeometryParticleState FutureState(*Proxy->GetParticle());
-				if(GetFutureStateAtFrame(FutureState,CurFrame) == EFutureQueryResult::Ok)
+				EResimType ResimType = EResimType::FullResim;
+				if(const auto* Rigid = PTParticle->CastToRigidParticle())
+				{
+					ResimType = Rigid->ResimType();
+				}
+				
+				//Only desync if full resim - might be nice to give a log warning for other cases
+				if(ResimType == EResimType::FullResim)
 				{
-					if(FutureState.IsDesynced(SrcManagerWrapper, *PTParticle, Dirty.ParticleData.GetFlags()))
+					//TODO: should not be passing GTParticle in here, it's not used so ok but not safe if someone decides to use it by mistake
+					FGeometryParticleState FutureState(*Proxy->GetParticle());
+					if(GetFutureStateAtFrame(FutureState,CurFrame) == EFutureQueryResult::Ok)
+					{
+						if(FutureState.IsDesynced(SrcManagerWrapper, *PTParticle, Dirty.ParticleData.GetFlags()))
+						{
+							Info.Desync(CurFrame-1, LatestFrame);
+						}
+					}
+					else if(!Info.bDesync)
 					{
 						Info.Desync(CurFrame-1, LatestFrame);
 					}
 				}
-				else if(!Info.bDesync)
-				{
-					Info.Desync(CurFrame-1, LatestFrame);
-				}
 			}
 
 			//Most properties are always a frame behind
@@ -1043,7 +1063,7 @@ public:
 	{
 		const int32 DestDataIdx = SrcDataIdx + DataIdxOffset;
 
-		if(bResim && Rigid.ResimType() == EResimType::SimAsSlave)
+		if(bResim && Rigid.ResimType() == EResimType::ResimAsSlave)
 		{
 			//resim not allowed for this particle so don't modify
 		}
