Chaos
- Fix joint double-locked swing solver
- Cone Angular Projection WIP

diff --git a/Engine/Source/Runtime/Engine/Private/PhysicsEngine/ImmediatePhysics/ImmediatePhysicsChaos/ImmediatePhysicsSimulation_Chaos.cpp b/Engine/Source/Runtime/Engine/Private/PhysicsEngine/ImmediatePhysics/ImmediatePhysicsChaos/ImmediatePhysicsSimulation_Chaos.cpp
--- a/Engine/Source/Runtime/Engine/Private/PhysicsEngine/ImmediatePhysics/ImmediatePhysicsChaos/ImmediatePhysicsSimulation_Chaos.cpp
+++ b/Engine/Source/Runtime/Engine/Private/PhysicsEngine/ImmediatePhysics/ImmediatePhysicsChaos/ImmediatePhysicsSimulation_Chaos.cpp
@@ -95,7 +95,7 @@ float ChaosImmediate_Joint_MinParentMassRatio = 0.2f;
 float ChaosImmediate_Joint_MaxInertiaRatio = 5.0f;
 float ChaosImmediate_Joint_AngularPositionCorrection = 0.0f;
 float ChaosImmediate_Joint_ProjectionInvMassScale = 0.0f;
-float ChaosImmediate_Joint_VelProjectionInvMassScale = 1.0f;
+float ChaosImmediate_Joint_VelProjectionInvMassScale = 0.0f;
 FAutoConsoleVariableRef CVarChaosImmPhysJointPairIterations(TEXT("p.Chaos.ImmPhys.Joint.PairIterations"), ChaosImmediate_Joint_PairIterations, TEXT("Override joint pair iterations (if >= 0)"));
 FAutoConsoleVariableRef CVarChaosImmPhysJointPushOutPairIterations(TEXT("p.Chaos.ImmPhys.Joint.PushOutPairIterations"), ChaosImmediate_Joint_PushOutPairIterations, TEXT("Override joint push-out pair iterations (if >= 0)"));
 FAutoConsoleVariableRef CVarChaosImmPhysJointSwingTwistAngleTolerance(TEXT("p.Chaos.ImmPhys.Joint.SwingTwistAngleTolerance"), ChaosImmediate_Joint_SwingTwistAngleTolerance, TEXT("SwingTwistAngleTolerance."));
diff --git a/Engine/Source/Runtime/Experimental/Chaos/Private/Chaos/Joint/PBDJointSolverGaussSeidel.cpp b/Engine/Source/Runtime/Experimental/Chaos/Private/Chaos/Joint/PBDJointSolverGaussSeidel.cpp
--- a/Engine/Source/Runtime/Experimental/Chaos/Private/Chaos/Joint/PBDJointSolverGaussSeidel.cpp
+++ b/Engine/Source/Runtime/Experimental/Chaos/Private/Chaos/Joint/PBDJointSolverGaussSeidel.cpp
@@ -779,7 +779,9 @@ namespace Chaos
 		const FVec3& Axis,
 		const FReal Angle)
 	{
-		const FVec3 DR1 = -Angle * Axis;
+		// NOTE: May be called with non-normalized axis (and similarly scaled angle), hence the divide by 
+		// length squared which is already handled by the joint mass calc in ApplyRotationConstraintDD
+		const FVec3 DR1 = (-Angle / Axis.SizeSquared()) * Axis;
 		ApplyRotationDelta(DIndex, Stiffness, DR1);
 	}
 
@@ -789,11 +791,9 @@ namespace Chaos
 		const FVec3& Axis,
 		const FReal Angle)
 	{
-		// World-space inverse mass
+		// Joint-space inverse mass
 		const FVec3 IA0 = Utilities::Multiply(InvIs[0], Axis);
 		const FVec3 IA1 = Utilities::Multiply(InvIs[1], Axis);
-
-		// Joint-space inverse mass
 		const FReal II0 = FVec3::DotProduct(Axis, IA0);
 		const FReal II1 = FVec3::DotProduct(Axis, IA1);
 
@@ -1793,9 +1793,12 @@ namespace Chaos
 		int32 NumApplied = 0;
 		bool bConeActive = false;
 
+		const FReal ProjectionPositionTolerance = 0.0f;//PositionTolerance;
+		const FReal ProjectionAngleTolerance = 0.0f;//AngleTolerance;
+
 		// Correct the position error
 		const FVec3 CX = Xs[1] - Xs[0];
-		if (CX.Size() > PositionTolerance)
+		if (CX.Size() > ProjectionPositionTolerance)
 		{
 			const FReal InvM0 = ProjectionInvMassScale * InvMs[0];
 			const FMatrix33 InvI0 = ProjectionInvMassScale * InvIs[0];
@@ -1828,7 +1831,7 @@ namespace Chaos
 			++NumApplied;
 		}
 
-		// Correct the cone error, maintaining position constraint
+		// Correct the cone error while maintaining position constraint by rotating about the swing axis through the joint position
 		const FReal Swing1Limit = JointSettings.AngularLimits[(int32)EJointAngularConstraintIndex::Swing1];
 		const FReal Swing2Limit = JointSettings.AngularLimits[(int32)EJointAngularConstraintIndex::Swing2];
 		FVec3 SwingAxis = FVec3(0);
@@ -1839,7 +1842,7 @@ namespace Chaos
 			FReal DSwingAngle = 0.0f;
 			FPBDJointUtilities::GetEllipticalConeAxisErrorLocal(Rs[0], Rs[1], Swing2Limit, Swing1Limit, SwingAxisLocal, DSwingAngle);
 			SwingAxis = Rs[0] * SwingAxisLocal;
-			if (DSwingAngle > AngleTolerance)
+			if (DSwingAngle > ProjectionAngleTolerance)
 			{
 				if (InvM0 > SMALL_NUMBER)
 				{
@@ -1859,155 +1862,131 @@ namespace Chaos
 			}
 		}
 
-		static bool b4x4 = false;
+		// Final position correction. Following this, both the position constraint and swing
+		// constraint should be exactly resolved.
+		const FVec3 CX2 = Xs[1] - Xs[0];
+		if (CX2.Size() > ProjectionPositionTolerance)
+		{
+			ApplyPositionDelta(1, LinearProjection, -CX2);
+		}
 
 		// Velocity correction
-		// NON-ITERATIVE VERSION
-		if (b4x4)
 		{
-			const FReal InvM0 = VelProjectionInvMassScale * InvMs[0];
-			const FMatrix33 InvI0 = VelProjectionInvMassScale * InvIs[0];
-
-			if (InvM0 > SMALL_NUMBER)
+			// Recalculate swing axis. THis should not be necessary if we did not rotate the parent body.
+			// @todo(ccaulfield): see if we really need this.
+			//if (InvM0 > SMALL_NUMBER)
 			{
-				// If the parent body moved, we need to recalculate the swing axis
 				FVec3 SwingAxisLocal;
 				FReal DSwingAngle = 0.0f;
 				FPBDJointUtilities::GetEllipticalConeAxisErrorLocal(Rs[0], Rs[1], Swing2Limit, Swing1Limit, SwingAxisLocal, DSwingAngle);
 				SwingAxis = Rs[0] * SwingAxisLocal;
 			}
 
-			// Calculate velocity and angular velocity error at joint
-			const FVec3 V0 = Vs[0] + FVec3::CrossProduct(Ws[0], Xs[0] - Ps[0]);
-			const FVec3 V1 = Vs[1] + FVec3::CrossProduct(Ws[1], Xs[1] - Ps[1]);
-			const FReal W0 = FVec3::DotProduct(Ws[0], SwingAxis);
-			const FReal W1 = FVec3::DotProduct(Ws[1], SwingAxis);
-			const FVec3 CV = V1 - V0;
-			const FReal CW = FMath::Max(0.0f, W1 - W0);
-			const TVector<FReal, 4> C = TVector<FReal, 4>(CV.X, CV.Y, CV.Z, CW);
-
-			// Calculate joint-space mass from the Jacobian
-			FMatrix33 J00 = FMatrix33(0, 0, 0);
-			FVec3 J01 = FVec3(0);
-			FReal J11 = 1.0f;
-			if (InvM0 > SMALL_NUMBER)
-			{
-				J00 += Utilities::ComputeJointFactorMatrix(Xs[0] - Ps[0], InvI0, InvM0);
-				if (bConeActive && (CW > 0.0f))
-				{
-					J01 = FVec3::CrossProduct(Xs[0] - Ps[0], Utilities::Multiply(InvI0, SwingAxis));
-					J11 = FVec3::DotProduct(SwingAxis, Utilities::Multiply(InvI0, SwingAxis));
-				}
-			}
-			if (InvMs[1] > SMALL_NUMBER)
-			{
-				J00 += Utilities::ComputeJointFactorMatrix(Xs[1] - Ps[1], InvIs[1], InvMs[1]);
-				if (bConeActive && (CW > 0.0f))
-				{
-					J01 = FVec3::CrossProduct(Xs[1] - Ps[1], Utilities::Multiply(InvIs[1], SwingAxis));
-					J11 = FVec3::DotProduct(SwingAxis, Utilities::Multiply(InvIs[1], SwingAxis));
-				}
-			}
-			PMatrix<FReal, 4, 4>& J = reinterpret_cast<PMatrix<FReal, 4, 4>&>(J00);
-			J.M[3][0] = J01.X;
-			J.M[3][1] = J01.Y;
-			J.M[3][2] = J01.Z;
-			J.M[0][3] = -J01.X;
-			J.M[1][3] = -J01.Y;
-			J.M[2][3] = -J01.Z;
-			J.M[3][3] = J11;
-			PMatrix<FReal, 4, 4> IJ = J.Inverse();
-
-			// Calculate the joint-space velocity and angular velocity correction
-			const TVector<FReal, 4> D = IJ.TransformFVector4(C);
-			const FVec3 DV = FVec3(D.X, D.Y, D.Z);
-			const FReal DW = D.W;
-
-			// Apply the world-space velocity and angular velocity correction
-			if (InvM0 > SMALL_NUMBER)
-			{
-				FVec3 DV0 = InvM0 * DV;
-				FVec3 DW0 = FVec3::CrossProduct(Xs[0] - Ps[0], Utilities::Multiply(InvI0, DV)) - Utilities::Multiply(InvI0, SwingAxis) * DW;
-				ApplyVelocityDelta(0, 1.0f, DV0, DW0);
-			}
-			if (InvMs[1] > SMALL_NUMBER)
-			{
-				FVec3 DV1 = -InvMs[1] * DV;
-				FVec3 DW1 = FVec3::CrossProduct(Xs[1] - Ps[1], Utilities::Multiply(InvIs[1], DV)) - Utilities::Multiply(InvIs[1], SwingAxis) * DW;
-				ApplyVelocityDelta(1, 1.0f, DV1, DW1);
-			}
-		}
-
-		// Velocity correction
-		// ITERATIVE VERSION
-		if (!b4x4)
-		{
-			const FReal InvM0 = VelProjectionInvMassScale * InvMs[0];
-			const FMatrix33 InvI0 = VelProjectionInvMassScale * InvIs[0];
-
-			FMatrix33 J = FMatrix33(0, 0, 0);
-			if (InvM0 > 0)
-			{
-				J = J + Utilities::ComputeJointFactorMatrix(Xs[0] - Ps[0], InvI0, InvM0);
-			}
-			if (InvMs[1] > 0)
-			{
-				J = J + Utilities::ComputeJointFactorMatrix(Xs[1] - Ps[1], InvIs[1], InvMs[1]);
-			}
-			const FMatrix33 IJ = J.Inverse();
-
-			FVec3 IA0 = FVec3(0);
-			FVec3 IA1 = FVec3(0);
-			FReal II0 = 0.0f;
-			FReal II1 = 0.0f;
-			if (bConeActive)
-			{
-				if (InvM0 > 0.0f)
-				{
-					II0 = FVec3::DotProduct(SwingAxis, Utilities::Multiply(InvI0, SwingAxis));
-				}
-				if (InvMs[1] > 0.0f)
-				{
-					II1 = FVec3::DotProduct(SwingAxis, Utilities::Multiply(InvIs[1], SwingAxis));
-				}
-			}
-
-			static int32 NumVelIts = 10;
-			for (int32 VelIt = 0; VelIt < NumVelIts; ++VelIt)
+			// We may iterate here if the swing constraint is active but the angular velocity does not required
+			// fixing (because it is not moving against the constraint), if when fixing the position constraint
+			// we introduce an angular velocity that would violate the swing constraint.
+			for (int32 ProjIts = 0; ProjIts < 2; ++ProjIts)
 			{
+				const FReal InvM0 = VelProjectionInvMassScale * InvMs[0];
+				const FMatrix33 InvI0 = VelProjectionInvMassScale * InvIs[0];
+				
+				// Calculate velocity and angular velocity error at joint
 				const FVec3 V0 = Vs[0] + FVec3::CrossProduct(Ws[0], Xs[0] - Ps[0]);
 				const FVec3 V1 = Vs[1] + FVec3::CrossProduct(Ws[1], Xs[1] - Ps[1]);
 				const FReal W0 = FVec3::DotProduct(Ws[0], SwingAxis);
 				const FReal W1 = FVec3::DotProduct(Ws[1], SwingAxis);
-				const FVec3 DV = Utilities::Multiply(IJ, V1 - V0);
-				const FVec3 DW = (W1 - W0) * SwingAxis;
-
-				if (InvM0 > 0)
+				const FVec3 CV = V1 - V0;
+				const FReal CW = FMath::Max(0.0f, W1 - W0);
+
+				// Calculate the joint-space velocity and angular velocity correction
+				// If the swing constraint is active, and we need to fix angular velocity about the axis,
+				// we have 4 simultaneous equations to solve (velocity at joint, and angular velocity about axis).
+				// If the swing constraint is inactive, or the angular velocity about the axis is negative,
+				// we have only 3 equations to solve (velocity at joint).
+				FVec3 DV = FVec3(0);
+				FReal DW = 0;
+				if (bConeActive && (CW > 0.0f))
 				{
-					FVec3 DV0 = InvM0 * DV;
-					FVec3 DW0 = Utilities::Multiply(InvI0, FVec3::CrossProduct(Xs[0] - Ps[0], DV));
-
-					if (bConeActive && ((W1 - W0) > 0.0f))
+					// Calculate joint-space mass from the Jacobian
+					FMatrix33 J00 = FMatrix33(0, 0, 0);
+					FVec3 J01 = FVec3(0);
+					FReal J11 = 0;
+					if (InvM0 > SMALL_NUMBER)
 					{
-						DW0 += (II0 / (II0 + II1)) * DW;
-						DV0 += FVec3(0);// -FVec3::CrossProduct(DW0, Xs[0] - Ps[0]);
+						J00 += Utilities::ComputeJointFactorMatrix(Xs[0] - Ps[0], InvI0, InvM0);
+						J01 += FVec3::CrossProduct(Xs[0] - Ps[0], Utilities::Multiply(InvI0, SwingAxis));
+						J11 += FVec3::DotProduct(SwingAxis, Utilities::Multiply(InvI0, SwingAxis));
 					}
-
-					ApplyVelocityDelta(0, 1.0f, DV0, DW0);
+					if (InvMs[1] > SMALL_NUMBER)
+					{
+						J00 += Utilities::ComputeJointFactorMatrix(Xs[1] - Ps[1], InvIs[1], InvMs[1]);
+						J01 += FVec3::CrossProduct(Xs[1] - Ps[1], Utilities::Multiply(InvIs[1], SwingAxis));
+						J11 += FVec3::DotProduct(SwingAxis, Utilities::Multiply(InvIs[1], SwingAxis));
+					}
+					PMatrix<FReal, 4, 4>& J = reinterpret_cast<PMatrix<FReal, 4, 4>&>(J00);
+					J.M[3][0] = -J01.X;
+					J.M[3][1] = -J01.Y;
+					J.M[3][2] = -J01.Z;
+					J.M[0][3] = -J01.X;
+					J.M[1][3] = -J01.Y;
+					J.M[2][3] = -J01.Z;
+					J.M[3][3] = J11;
+					PMatrix<FReal, 4, 4> IJ = J.Inverse();
+
+					const TVector<FReal, 4> C = TVector<FReal, 4>(CV.X, CV.Y, CV.Z, CW);
+					const TVector<FReal, 4> D = IJ.TransformFVector4(C);
+					DV = FVec3(D.X, D.Y, D.Z);
+					DW = D.W;
 				}
-				if (InvMs[1] > 0)
+				else
 				{
-					FVec3 DV1 = -InvMs[1] * DV;
-					FVec3 DW1 = FVec3::CrossProduct(Xs[1] - Ps[1], Utilities::Multiply(InvIs[1], -DV));
-
-					if (bConeActive && ((W1 - W0) > 0.0f))
+					// Calculate joint-space mass from the Jacobian
+					FMatrix33 J00 = FMatrix33(0, 0, 0);
+					if (InvM0 > SMALL_NUMBER)
+					{
+						J00 += Utilities::ComputeJointFactorMatrix(Xs[0] - Ps[0], InvI0, InvM0);
+					}
+					if (InvMs[1] > SMALL_NUMBER)
 					{
-						DW1 += -(II1 / (II0 + II1)) * DW;
-						DV1 += FVec3(0);// -FVec3::CrossProduct(DW0, Xs[0] - Ps[0]);
+						J00 += Utilities::ComputeJointFactorMatrix(Xs[1] - Ps[1], InvIs[1], InvMs[1]);
 					}
+					FMatrix33 IJ = J00.Inverse();
+
+					DV = Utilities::Multiply(IJ, CV);
+				}
 
+				// Apply the world-space velocity and angular velocity correction
+				FVec3 DV0 = InvM0 * DV;
+				FVec3 DW0 = Utilities::Multiply(InvI0, FVec3::CrossProduct(Xs[0] - Ps[0], DV)) + Utilities::Multiply(InvI0, SwingAxis) * DW;
+				if (InvM0 > SMALL_NUMBER)
+				{
+					ApplyVelocityDelta(0, 1.0f, DV0, DW0);
+				}
+				FVec3 DV1 = -InvMs[1] * DV;
+				FVec3 DWV1 = Utilities::Multiply(InvIs[1], FVec3::CrossProduct(Xs[1] - Ps[1], DV));
+				FVec3 DWW1 = Utilities::Multiply(InvIs[1], SwingAxis) * DW;
+				FVec3 DW1 = -(DWV1 + DWW1);
+				if (InvMs[1] > SMALL_NUMBER)
+				{
 					ApplyVelocityDelta(1, 1.0f, DV1, DW1);
 				}
+
+				const FVec3 V02 = Vs[0] + FVec3::CrossProduct(Ws[0], Xs[0] - Ps[0]);
+				const FVec3 V12 = Vs[1] + FVec3::CrossProduct(Ws[1], Xs[1] - Ps[1]);
+				const FReal W02 = FVec3::DotProduct(Ws[0], SwingAxis);
+				const FReal W12 = FVec3::DotProduct(Ws[1], SwingAxis);
+				const FVec3 CV2 = V12 - V02;
+				const FReal CW2 = FMath::Max(0.0f, W12 - W02);
+				const FVec3 DWOffAxis0 = DW0 - FVec3::DotProduct(DW0, SwingAxis) * SwingAxis;
+				const FVec3 DWOffAxis1 = DW1 - FVec3::DotProduct(DW1, SwingAxis) * SwingAxis;
+				const FVec3 DWOffAxisW1 = DWW1 - FVec3::DotProduct(DWW1, SwingAxis) * SwingAxis;
+				const FVec3 DWOffAxisV1 = DWV1 - FVec3::DotProduct(DWV1, SwingAxis) * SwingAxis;
+
+				// If we resolved the angvel constraint or it was ignored and not reactivated (by the vel constraint), we are done
+				if ((bConeActive && (CW > 0.0f)) || (CW2 <= 0.0f))
+				{
+					break;
+				}
 			}
 		}
 
