Bumping Chaos DDC GUID to fix a crash caused by serialization issue (change by [REDACTED_PERSON])

diff --git a/Engine/Source/Runtime/Experimental/ChaosCore/Public/Chaos/Core.h b/Engine/Source/Runtime/Experimental/ChaosCore/Public/Chaos/Core.h
--- a/Engine/Source/Runtime/Experimental/ChaosCore/Public/Chaos/Core.h
+++ b/Engine/Source/Runtime/Experimental/ChaosCore/Public/Chaos/Core.h
@@ -30,5 +30,5 @@ namespace Chaos
 	using TVec3 = TVector<T, 3>;
 
 	/** Core version string for Chaos data. Any DDC builder dependent on Chaos for serialization should depend on this version */
-	static const FString ChaosVersionString = TEXT("030758A4893440588B23DD9BCDF339B1");
+	static const FString ChaosVersionString = TEXT("F7CC5FD0E5174BB9B87A3B340A878EE3");
 }
\ No newline at end of file
