Change sleep counter default to 20 instead of 0. Make it so that if an island is deemed awake we reset counter

diff --git a/Engine/Source/Runtime/Experimental/Chaos/Private/Chaos/PBDConstraintGraph.cpp b/Engine/Source/Runtime/Experimental/Chaos/Private/Chaos/PBDConstraintGraph.cpp
--- a/Engine/Source/Runtime/Experimental/Chaos/Private/Chaos/PBDConstraintGraph.cpp
+++ b/Engine/Source/Runtime/Experimental/Chaos/Private/Chaos/PBDConstraintGraph.cpp
@@ -16,8 +16,8 @@ using namespace Chaos;
 bool ChaosSolverCollisionDefaultUseMaterialSleepThresholdsCVar = true;
 FAutoConsoleVariableRef CVarChaosSolverCollisionDefaultUseMaterialSleepThresholds(TEXT("p.ChaosSolverCollisionDefaultUseMaterialSleepThresholds"), ChaosSolverCollisionDefaultUseMaterialSleepThresholdsCVar, TEXT("Enable material support for sleeping thresholds[def:true]"));
 
-int32 ChaosSolverCollisionDefaultSleepCounterThresholdCVar = 0; 
-FAutoConsoleVariableRef CVarChaosSolverCollisionDefaultSleepCounterThreshold(TEXT("p.ChaosSolverCollisionDefaultSleepCounterThreshold"), ChaosSolverCollisionDefaultSleepCounterThresholdCVar, TEXT("Default counter threshold for sleeping.[def:0]"));
+int32 ChaosSolverCollisionDefaultSleepCounterThresholdCVar = 20; 
+FAutoConsoleVariableRef CVarChaosSolverCollisionDefaultSleepCounterThreshold(TEXT("p.ChaosSolverCollisionDefaultSleepCounterThreshold"), ChaosSolverCollisionDefaultSleepCounterThresholdCVar, TEXT("Default counter threshold for sleeping.[def:20]"));
 
 float ChaosSolverCollisionDefaultLinearSleepThresholdCVar = 0.001f; // .001 unit mass cm
 FAutoConsoleVariableRef CVarChaosSolverCollisionDefaultLinearSleepThreshold(TEXT("p.ChaosSolverCollisionDefaultLinearSleepThreshold"), ChaosSolverCollisionDefaultLinearSleepThresholdCVar, TEXT("Default linear threshold for sleeping.[def:0.001]"));
@@ -740,6 +740,11 @@ bool FPBDConstraintGraph::SleepInactive(const int32 Island, const TArrayCollecti
 			IslandSleepCount++;
 		}
 	}
+	else
+	{
+		//Reset sleep count since island is awake
+		IslandSleepCount = 0;
+	}
 
 	return false;
 }
