Chaos Physics: Default number of collision iterations changed from 1 to 4

diff --git a/Engine/Source/Runtime/Experimental/ChaosSolvers/Private/PBDRigidsSolver.cpp b/Engine/Source/Runtime/Experimental/ChaosSolvers/Private/PBDRigidsSolver.cpp
--- a/Engine/Source/Runtime/Experimental/ChaosSolvers/Private/PBDRigidsSolver.cpp
+++ b/Engine/Source/Runtime/Experimental/ChaosSolvers/Private/PBDRigidsSolver.cpp
@@ -37,7 +37,7 @@ FAutoConsoleVariableRef CVarChaosSolverUseParticlePool(TEXT("p.Chaos.Solver.UseP
 int32 ChaosSolverParticlePoolNumFrameUntilShrink = 30;
 FAutoConsoleVariableRef CVarChaosSolverParticlePoolNumFrameUntilShrink(TEXT("p.Chaos.Solver.ParticlePoolNumFrameUntilShrink"), ChaosSolverParticlePoolNumFrameUntilShrink, TEXT("Num Frame until we can potentially shrink the pool"));
 
-int32 ChaosSolverCollisionDefaultIterationsCVar = 1;
+int32 ChaosSolverCollisionDefaultIterationsCVar = 4;
 FAutoConsoleVariableRef CVarChaosSolverCollisionDefaultIterations(TEXT("p.ChaosSolverCollisionDefaultIterations"), ChaosSolverCollisionDefaultIterationsCVar, TEXT("Default collision iterations for the solver.[def:1]"));
 
 int32 ChaosSolverCollisionDefaultPushoutIterationsCVar = 3;
