Chaos
- linear breakable constraints (WIP)
- angular projection for two locked swing axes (disabled)
- removed some unused vars and structs

diff --git a/Engine/Source/Runtime/Engine/Private/PhysicsEngine/ImmediatePhysics/ImmediatePhysicsChaos/ImmediatePhysicsJointHandle_Chaos.cpp b/Engine/Source/Runtime/Engine/Private/PhysicsEngine/ImmediatePhysics/ImmediatePhysicsChaos/ImmediatePhysicsJointHandle_Chaos.cpp
--- a/Engine/Source/Runtime/Engine/Private/PhysicsEngine/ImmediatePhysics/ImmediatePhysicsChaos/ImmediatePhysicsJointHandle_Chaos.cpp
+++ b/Engine/Source/Runtime/Engine/Private/PhysicsEngine/ImmediatePhysics/ImmediatePhysicsChaos/ImmediatePhysicsJointHandle_Chaos.cpp
@@ -40,6 +40,9 @@ namespace ImmediatePhysics_Chaos
 	float ChaosImmediate_JointMinAngularProjection = 0.0f;
 	float ChaosImmediate_JointMaxAngularProjection = 0.0f;
 
+	float ChaosImmediate_LinearBreakScale = 1.0f;
+	float ChaosImmediate_AngularBreakScale = 1.0f;
+
 	FAutoConsoleVariableRef CVarJointStiffness(TEXT("p.Chaos.ImmPhys.JointStiffness"), ChaosImmediate_JointStiffness, TEXT("Hard-joint solver stiffness."));
 	FAutoConsoleVariableRef CVarLinearDriveStiffnessScale(TEXT("p.Chaos.ImmPhys.LinearDriveStiffnessScale"), ChaosImmediate_LinearDriveStiffnessScale, TEXT("Conversion factor for Linear drive stiffness."));
 	FAutoConsoleVariableRef CVarLinearDriveDampingScale(TEXT("p.Chaos.ImmPhys.LinaearDriveDampingScale"), ChaosImmediate_LinearDriveDampingScale, TEXT("Conversion factor for Linear drive damping."));
@@ -55,6 +58,8 @@ namespace ImmediatePhysics_Chaos
 	FAutoConsoleVariableRef CVarJointMaxLinearProjection(TEXT("p.Chaos.ImmPhys.JointMaxLinearProjection"), ChaosImmediate_JointMaxLinearProjection, TEXT("Joint max projection (for joints with projection enabled)."));
 	FAutoConsoleVariableRef CVarJointMinAngularProjection(TEXT("p.Chaos.ImmPhys.JointMinAngularProjection"), ChaosImmediate_JointMinAngularProjection, TEXT("Joint min projection (for joints with projection disabled)."));
 	FAutoConsoleVariableRef CVarJointMaxAngularProjection(TEXT("p.Chaos.ImmPhys.JointMaxAngularProjection"), ChaosImmediate_JointMaxAngularProjection, TEXT("Joint max projection (for joints with projection enabled)."));
+	FAutoConsoleVariableRef CVarJointLinearBreakScale(TEXT("p.Chaos.ImmPhys.LinearBreakScale"), ChaosImmediate_LinearBreakScale, TEXT("Conversion factory for Linear Break Theshold."));
+	FAutoConsoleVariableRef CVarJointAngularBreakScale(TEXT("p.Chaos.ImmPhys.AngularBreakScale"), ChaosImmediate_AngularBreakScale, TEXT("Conversion factory for Angular Break Theshold."));
 
 	void TransferJointSettings(FConstraintInstance* ConstraintInstance, Chaos::FPBDJointSettings& ConstraintSettings)
 	{
@@ -108,9 +113,6 @@ namespace ImmediatePhysics_Chaos
 
 		ConstraintSettings.AngularDrivePositionTarget = FQuat(Profile.AngularDrive.OrientationTarget);
 		ConstraintSettings.AngularDriveVelocityTarget = Profile.AngularDrive.AngularVelocityTarget;
-		ConstraintSettings.AngularDriveTargetAngles[(int32)EJointAngularConstraintIndex::Twist] = FMath::DegreesToRadians(Profile.AngularDrive.OrientationTarget.Roll);
-		ConstraintSettings.AngularDriveTargetAngles[(int32)EJointAngularConstraintIndex::Swing1] = FMath::DegreesToRadians(Profile.AngularDrive.OrientationTarget.Yaw);
-		ConstraintSettings.AngularDriveTargetAngles[(int32)EJointAngularConstraintIndex::Swing2] = FMath::DegreesToRadians(Profile.AngularDrive.OrientationTarget.Pitch);
 
 		if (Profile.AngularDrive.AngularDriveMode == EAngularDriveMode::SLERP)
 		{
@@ -128,6 +130,9 @@ namespace ImmediatePhysics_Chaos
 		ConstraintSettings.AngularDriveDamping = ChaosImmediate_AngularDriveDampingScale * Profile.AngularDrive.TwistDrive.Damping;
 		ConstraintSettings.AngularDriveForceMode = EJointForceMode::Acceleration;
 
+		ConstraintSettings.LinearBreakForce = (Profile.bLinearBreakable) ? ChaosImmediate_LinearBreakScale * Profile.LinearBreakThreshold : 0.0f;
+		ConstraintSettings.AngularBreakTorque = (Profile.bAngularBreakable) ? ChaosImmediate_AngularBreakScale * Profile.AngularBreakThreshold : 0.0f;
+
 		// UE Disables Soft Limits when the Limit is less than some threshold. This is not necessary in Chaos but for now we also do it for parity's sake (See FLinearConstraint::UpdateLinearLimit_AssumesLocked).
 		if (ConstraintSettings.LinearLimit < RB_MinSizeToLockDOF)
 		{
diff --git a/Engine/Source/Runtime/Engine/Private/PhysicsEngine/ImmediatePhysics/ImmediatePhysicsChaos/ImmediatePhysicsSimulation_Chaos.cpp b/Engine/Source/Runtime/Engine/Private/PhysicsEngine/ImmediatePhysics/ImmediatePhysicsChaos/ImmediatePhysicsSimulation_Chaos.cpp
--- a/Engine/Source/Runtime/Engine/Private/PhysicsEngine/ImmediatePhysics/ImmediatePhysicsChaos/ImmediatePhysicsSimulation_Chaos.cpp
+++ b/Engine/Source/Runtime/Engine/Private/PhysicsEngine/ImmediatePhysics/ImmediatePhysicsChaos/ImmediatePhysicsSimulation_Chaos.cpp
@@ -56,7 +56,7 @@ FAutoConsoleVariableRef CVarChaosImmPhysFixedStepTolerance(TEXT("p.Chaos.ImmPhys
 
 int32 ChaosImmediate_Collision_Enabled = 1;
 int32 ChaosImmediate_Collision_PairIterations = -1;
-int32 ChaosImmediate_Collision_PushOutPairIterations = -1;
+int32 ChaosImmediate_Collision_PushOutPairIterations = 0;	// Force disabled - not compatible with ECollisionApplyType::Position
 int32 ChaosImmediate_Collision_Priority = 1;
 float ChaosImmediate_Collision_CullDistance = 1.0f;
 float ChaosImmediate_Collision_ShapePadding = 0;
diff --git a/Engine/Source/Runtime/Experimental/Chaos/Private/Chaos/Joint/PBDJointSolverGaussSeidel.cpp b/Engine/Source/Runtime/Experimental/Chaos/Private/Chaos/Joint/PBDJointSolverGaussSeidel.cpp
--- a/Engine/Source/Runtime/Experimental/Chaos/Private/Chaos/Joint/PBDJointSolverGaussSeidel.cpp
+++ b/Engine/Source/Runtime/Experimental/Chaos/Private/Chaos/Joint/PBDJointSolverGaussSeidel.cpp
@@ -110,6 +110,9 @@ namespace Chaos
 		PrevXs[0] = PrevP0 + PrevQ0 * XL0.GetTranslation();
 		PrevXs[1] = PrevP1 + PrevQ1 * XL1.GetTranslation();
 
+		NetLinearImpulse = FVec3(0);
+		NetAngularImpulse = FVec3(0);
+
 		LinearSoftLambda = (FReal)0;
 		LinearDriveLambda = (FReal)0;
 		TwistSoftLambda = (FReal)0;
@@ -252,12 +255,22 @@ namespace Chaos
 			&& !bSwingSoft 
 			&& (Swing1Motion == EJointMotionType::Limited) 
 			&& (Swing2Motion == EJointMotionType::Limited);
+		bool bTwoSwingLockProject =
+			SolverSettings.bEnableSwingLimits
+			&& (AngularProjection > 0.0f)
+			&& (Swing1Motion == EJointMotionType::Locked)
+			&& (Swing2Motion == EJointMotionType::Locked);
 
 		if (bPointProject && bConeProject)
 		{
 			// Fixed point constraint with cone limits
 			NumActive += ApplyPointConeProjection(Dt, SolverSettings, JointSettings);
 		}
+		else if (bPointProject && bTwoSwingLockProject)
+		{
+			// Fixed point constraint with two locked swing axes
+			NumActive += ApplyPointDoubleSwingLockProjection(Dt, SolverSettings, JointSettings);
+		}
 		else if (bPointProject)
 		{
 			// Fixed point constraint with no angular limits
@@ -717,6 +730,8 @@ namespace Chaos
 
 		ApplyPositionDelta(Stiffness, DP0, DP1);
 		ApplyRotationDelta(Stiffness, DR0, DR1);
+
+		NetLinearImpulse += Stiffness * DX;
 	}
 
 
@@ -765,9 +780,11 @@ namespace Chaos
 			const FVec3 DR0 = DLambda * Utilities::Multiply(InvIs[0], AngularAxis0);
 			const FVec3 DR1 = -DLambda * Utilities::Multiply(InvIs[1], AngularAxis1);
 	
-			Lambda += DLambda;
 			ApplyPositionDelta((FReal)1, DP0, DP1);
 			ApplyRotationDelta((FReal)1, DR0, DR1);
+
+			Lambda += DLambda;
+			NetLinearImpulse += DLambda * Axis;
 		}
 	}
 	
@@ -783,6 +800,8 @@ namespace Chaos
 		// length squared which is already handled by the joint mass calc in ApplyRotationConstraintDD
 		const FVec3 DR1 = (-Angle / Axis.SizeSquared()) * Axis;
 		ApplyRotationDelta(DIndex, Stiffness, DR1);
+
+		NetAngularImpulse += (-Stiffness / InvMs[DIndex]) * DR1;
 	}
 
 
@@ -803,6 +822,8 @@ namespace Chaos
 		//const FVec3 DR1 = Axis * -(Angle * II1 / (II0 + II1));
 
 		ApplyRotationDelta(Stiffness, DR0, DR1);
+
+		NetAngularImpulse += Axis * (Stiffness * Angle / (II0 + II1));
 	}
 
 
@@ -872,8 +893,10 @@ namespace Chaos
 			//const FVec3 DR1 = IA1 * -DLambda;
 			const FVec3 DR1 = Axis * -(DLambda * II1);
 
-			Lambda += DLambda;
 			ApplyRotationDelta(DIndex, 1.0f, DR1);
+	
+			Lambda += DLambda;
+			NetAngularImpulse += DLambda * Axis;
 		}
 	}
 
@@ -927,8 +950,10 @@ namespace Chaos
 			const FVec3 DR0 = Axis * (DLambda * II0);
 			const FVec3 DR1 = Axis * -(DLambda * II1);
 
-			Lambda += DLambda;
 			ApplyRotationDelta(1.0f, DR0, DR1);
+
+			Lambda += DLambda;
+			NetAngularImpulse += DLambda * Axis;
 		}
 	}
 
@@ -1433,6 +1458,8 @@ namespace Chaos
 				const FVec3 DR1 = Utilities::Multiply(InvIs[DIndex], FVec3::CrossProduct(Xs[DIndex] - Ps[DIndex], -DX));
 
 				ApplyDelta(DIndex, LinearStiffness, DP1, DR1);
+
+				NetLinearImpulse += LinearStiffness * DX;
 			}
 			return 1;
 		}
@@ -1478,6 +1505,8 @@ namespace Chaos
 
 				ApplyPositionDelta(LinearStiffness, DP0, DP1);
 				ApplyRotationDelta(LinearStiffness, DR0, DR1);
+
+				NetLinearImpulse += LinearStiffness * DX;
 			}
 			return 1;
 		}
@@ -1696,10 +1725,14 @@ namespace Chaos
 		const FPBDJointSettings& JointSettings)
 	{
 		FReal LinearProjection = FPBDJointUtilities::GetLinearProjection(SolverSettings, JointSettings);
+		const FReal ProjectionPositionTolerance = 0.0f;//PositionTolerance;
+
 		const FVec3 CX = Xs[1] - Xs[0];
-		bool bApplyProjection = (CX.Size() > KINDA_SMALL_NUMBER); // + PositionTolerance);
+		bool bApplyProjection = true;//(CX.Size() > KINDA_SMALL_NUMBER); // + ProjectionPositionTolerance);
 		if (bApplyProjection)
 		{
+			UE_LOG(LogChaosJoint, Verbose, TEXT("Projection error %f"), CX.Size());
+
 			if (bChaos_Joint_ISPC_Enabled)
 			{
 #if INTEL_ISPC
@@ -1708,36 +1741,21 @@ namespace Chaos
 			}
 			else
 			{
-				// Apply a position correction
+				// Apply a correction for the error in the joint connectors (apply position and rotation)
 				{
-					const FReal InvM0 = ProjectionInvMassScale * InvMs[0];
-					const FMatrix33 InvI0 = ProjectionInvMassScale * InvIs[0];
-
-					FMatrix33 J = FMatrix33(0, 0, 0);
-					if (InvM0 > SMALL_NUMBER)
-					{
-						J += Utilities::ComputeJointFactorMatrix(Xs[0] - Ps[0], InvI0, InvM0);
-					}
-					if (InvMs[1] > SMALL_NUMBER)
-					{
-						J += Utilities::ComputeJointFactorMatrix(Xs[1] - Ps[1], InvIs[1], InvMs[1]);
-					}
+					FMatrix33 J = Utilities::ComputeJointFactorMatrix(Xs[1] - Ps[1], InvIs[1], InvMs[1]);
 					const FMatrix33 IJ = J.Inverse();
 					const FVec3 DX = Utilities::Multiply(IJ, CX);
+					const FVec3 DP1 = -InvMs[1] * DX;
+					const FVec3 DR1 = Utilities::Multiply(InvIs[1], FVec3::CrossProduct(Xs[1] - Ps[1], -DX));
+					UE_LOG(LogChaosJoint, Verbose, TEXT("Projection pos/rot correction %f %f"), DP1.Size(), DR1.Size());
+					ApplyDelta(1, LinearProjection, DP1, DR1);
+				}
 
-					if (InvM0 > SMALL_NUMBER)
-					{
-						const FVec3 DP0 = InvM0 * DX;
-						const FVec3 DR0 = Utilities::Multiply(InvI0, FVec3::CrossProduct(Xs[0] - Ps[0], DX));
-						ApplyDelta(0, LinearProjection, DP0, DR0);
-					}
-
-					if (InvMs[1] > SMALL_NUMBER)
-					{
-						const FVec3 DP1 = -InvMs[1] * DX;
-						const FVec3 DR1 = Utilities::Multiply(InvIs[1], FVec3::CrossProduct(Xs[1] - Ps[1], -DX));
-						ApplyDelta(1, LinearProjection, DP1, DR1);
-					}
+				// Final position correction to fix errors from rotation in above
+				{
+					const FVec3 CX2 = Xs[1] - Xs[0];
+					ApplyPositionDelta(1, LinearProjection, -CX2);
 				}
 
 				// Apply a velocity correction to set the relative velocity at the joint to zero
@@ -1766,12 +1784,14 @@ namespace Chaos
 					{
 						const FVec3 DV0 = InvM0 * DV;
 						const FVec3 DW0 = Utilities::Multiply(InvI0, FVec3::CrossProduct(Xs[0] - Ps[0], DV));
+						UE_LOG(LogChaosJoint, Verbose, TEXT("Projection vel/angvel correction 0 %f %f"), DV0.Size(), DW0.Size());
 						ApplyVelocityDelta(0, LinearProjection, DV0, DW0);
 					}
 					if (InvMs[1] > SMALL_NUMBER)
 					{
 						const FVec3 DV1 = -InvMs[1] * DV;
-						const FVec3 DW1 = Utilities::Multiply(InvIs[1], FVec3::CrossProduct(Xs[1] - Ps[1], -DV));
+						const FVec3 DW1 = -Utilities::Multiply(InvIs[1], FVec3::CrossProduct(Xs[1] - Ps[1], DV));
+						UE_LOG(LogChaosJoint, Verbose, TEXT("Projection vel/angvel correction 1 %f %f"), DV1.Size(), DW1.Size());
 						ApplyVelocityDelta(1, LinearProjection, DV1, DW1);
 					}
 				}
@@ -1800,33 +1820,13 @@ namespace Chaos
 		const FVec3 CX = Xs[1] - Xs[0];
 		if (CX.Size() > ProjectionPositionTolerance)
 		{
-			const FReal InvM0 = ProjectionInvMassScale * InvMs[0];
-			const FMatrix33 InvI0 = ProjectionInvMassScale * InvIs[0];
-
-			FMatrix33 J = FMatrix33(0, 0, 0);
-			if (InvM0 > SMALL_NUMBER)
-			{
-				J += Utilities::ComputeJointFactorMatrix(Xs[0] - Ps[0], InvI0, InvM0);
-			}
-			if (InvMs[1] > SMALL_NUMBER)
-			{
-				J += Utilities::ComputeJointFactorMatrix(Xs[1] - Ps[1], InvIs[1], InvMs[1]);
-			}
+			FMatrix33 J = Utilities::ComputeJointFactorMatrix(Xs[1] - Ps[1], InvIs[1], InvMs[1]);
 			const FMatrix33 IJ = J.Inverse();
 			const FVec3 DX = Utilities::Multiply(IJ, CX);
 
-			if (InvM0 > SMALL_NUMBER)
-			{
-				const FVec3 DP0 = InvM0 * DX;
-				const FVec3 DR0 = Utilities::Multiply(InvI0, FVec3::CrossProduct(Xs[0] - Ps[0], DX));
-				ApplyDelta(0, LinearProjection, DP0, DR0);
-			}
-			if (InvMs[1] > SMALL_NUMBER)
-			{
-				const FVec3 DP1 = -InvMs[1] * DX;
-				const FVec3 DR1 = Utilities::Multiply(InvIs[1], FVec3::CrossProduct(Xs[1] - Ps[1], -DX));
-				ApplyDelta(1, LinearProjection, DP1, DR1);
-			}
+			const FVec3 DP1 = -InvMs[1] * DX;
+			const FVec3 DR1 = Utilities::Multiply(InvIs[1], FVec3::CrossProduct(Xs[1] - Ps[1], -DX));
+			ApplyDelta(1, LinearProjection, DP1, DR1);
 
 			++NumApplied;
 		}
@@ -1836,26 +1836,15 @@ namespace Chaos
 		const FReal Swing2Limit = JointSettings.AngularLimits[(int32)EJointAngularConstraintIndex::Swing2];
 		FVec3 SwingAxis = FVec3(0);
 		{
-			const FReal InvM0 = ProjectionInvMassScale * InvMs[0];
-			const FMatrix33 InvI0 = ProjectionInvMassScale * InvIs[0];
 			FVec3 SwingAxisLocal;
 			FReal DSwingAngle = 0.0f;
 			FPBDJointUtilities::GetEllipticalConeAxisErrorLocal(Rs[0], Rs[1], Swing2Limit, Swing1Limit, SwingAxisLocal, DSwingAngle);
 			SwingAxis = Rs[0] * SwingAxisLocal;
 			if (DSwingAngle > ProjectionAngleTolerance)
 			{
-				if (InvM0 > SMALL_NUMBER)
-				{
-					const FVec3 DR0 = DSwingAngle * SwingAxis;
-					const FVec3 DP0 = FVec3::CrossProduct(DR0, Xs[0] - Ps[0]);
-					ApplyDelta(0, AngularProjection, DP0, DR0);
-				}
-				if (InvMs[1] > SMALL_NUMBER)
-				{
-					const FVec3 DR1 = -DSwingAngle * SwingAxis;
-					const FVec3 DP1 = -FVec3::CrossProduct(DR1, Xs[1] - Ps[1]);
-					ApplyDelta(1, AngularProjection, DP1, DR1);
-				}
+				const FVec3 DR1 = -DSwingAngle * SwingAxis;
+				const FVec3 DP1 = -FVec3::CrossProduct(DR1, Xs[1] - Ps[1]);
+				ApplyDelta(1, AngularProjection, DP1, DR1);
 
 				++NumApplied;
 				bConeActive = true;
@@ -1864,14 +1853,16 @@ namespace Chaos
 
 		// Final position correction. Following this, both the position constraint and swing
 		// constraint should be exactly resolved.
-		const FVec3 CX2 = Xs[1] - Xs[0];
-		if (CX2.Size() > ProjectionPositionTolerance)
 		{
+			const FVec3 CX2 = Xs[1] - Xs[0];
 			ApplyPositionDelta(1, LinearProjection, -CX2);
 		}
 
 		// Velocity correction
 		{
+			const FReal InvM0 = VelProjectionInvMassScale * InvMs[0];
+			const FMatrix33 InvI0 = VelProjectionInvMassScale * InvIs[0];
+
 			// Recalculate swing axis. THis should not be necessary if we did not rotate the parent body.
 			// @todo(ccaulfield): see if we really need this.
 			//if (InvM0 > SMALL_NUMBER)
@@ -1887,9 +1878,6 @@ namespace Chaos
 			// we introduce an angular velocity that would violate the swing constraint.
 			for (int32 ProjIts = 0; ProjIts < 2; ++ProjIts)
 			{
-				const FReal InvM0 = VelProjectionInvMassScale * InvMs[0];
-				const FMatrix33 InvI0 = VelProjectionInvMassScale * InvIs[0];
-				
 				// Calculate velocity and angular velocity error at joint
 				const FVec3 V0 = Vs[0] + FVec3::CrossProduct(Ws[0], Xs[0] - Ps[0]);
 				const FVec3 V1 = Vs[1] + FVec3::CrossProduct(Ws[1], Xs[1] - Ps[1]);
@@ -1992,4 +1980,49 @@ namespace Chaos
 
 		return NumApplied;
 	}
+
+	// Projection for fixed point constraint and two locked swing axes
+	int32 FJointSolverGaussSeidel::ApplyPointDoubleSwingLockProjection(
+		const FReal Dt,
+		const FPBDJointSolverSettings& SolverSettings,
+		const FPBDJointSettings& JointSettings)
+	{
+		const FReal AngularProjection = FPBDJointUtilities::GetAngularProjection(SolverSettings, JointSettings);
+		const FReal ProjectionAngleTolerance = 0.0f;//AngleTolerance;
+
+		int32 NumApplied = 0;
+
+		// Correct position error
+		NumApplied += ApplyPointProjection(Dt, SolverSettings, JointSettings);
+
+		// Correct the swing error while maintaining position constraint by rotating about the swing axis through the joint position
+		{
+			FVec3 SwingAxisLocal;
+			FReal DSwingAngle = 0.0f;
+			FPBDJointUtilities::GetCircularConeAxisErrorLocal(Rs[0], Rs[1], 0.0f, SwingAxisLocal, DSwingAngle);
+			FVec3 SwingAxis = Rs[0] * SwingAxisLocal;
+			if (DSwingAngle > ProjectionAngleTolerance)
+			{
+				const FVec3 DR1 = -DSwingAngle * SwingAxis;
+				const FVec3 DP1 = -FVec3::CrossProduct(DR1, Xs[1] - Ps[1]);
+				ApplyDelta(1, AngularProjection, DP1, DR1);
+
+				++NumApplied;
+			}
+		}
+
+		// Correct angular velocity error
+		{
+			FVec3 TwistAxis = Rs[1] * FJointConstants::TwistAxis();
+
+			// Fix the angular velocity to remove all components not about twist axis
+			const FVec3 W0 = Ws[0] - FVec3::DotProduct(Ws[0], TwistAxis) * TwistAxis;
+			const FVec3 W1 = Ws[1] - FVec3::DotProduct(Ws[1], TwistAxis) * TwistAxis;
+			const FVec3 DW = W1 - W0;
+			ApplyVelocityDelta(1, 1.0f, FVec3(0), -DW);
+		}
+
+		return NumApplied;
+	}
+
 }
diff --git a/Engine/Source/Runtime/Experimental/Chaos/Private/Chaos/Joint/PBDJointSolverGaussSeidel.ispc b/Engine/Source/Runtime/Experimental/Chaos/Private/Chaos/Joint/PBDJointSolverGaussSeidel.ispc
--- a/Engine/Source/Runtime/Experimental/Chaos/Private/Chaos/Joint/PBDJointSolverGaussSeidel.ispc
+++ b/Engine/Source/Runtime/Experimental/Chaos/Private/Chaos/Joint/PBDJointSolverGaussSeidel.ispc
@@ -24,14 +24,10 @@ struct FJointSolverGaussSeidel
 
 	// World-space constraint state
 	FVector Xs[MaxConstrainedBodies];			// World-space joint connector positions
-	float LinearSoftLambda;						// XPBD constraint multipliers (net applied constraint-space deltas)
-	float LinearDriveLambda;					// XPBD constraint multipliers (net applied constraint-space deltas)
 	FVector4 Rs[MaxConstrainedBodies];			// World-space joint connector rotations
 
 	// World-space body state
 	FVector Ps[MaxConstrainedBodies];			// World-space particle CoM positions
-	float TwistSoftLambda;						// XPBD constraint multipliers (net applied constraint-space deltas)
-	float SwingSoftLambda;						// XPBD constraint multipliers (net applied constraint-space deltas)
 	FVector4 Qs[MaxConstrainedBodies];			// World-space particle CoM rotations
 	FVector Vs[MaxConstrainedBodies];			// World-space particle CoM velocities
 	FVector Ws[MaxConstrainedBodies];			// World-space particle CoM angular velocities
@@ -39,12 +35,23 @@ struct FJointSolverGaussSeidel
 
 	// XPBD Previous iteration world-space body state
 	FVector PrevPs[MaxConstrainedBodies];		// World-space particle CoM positions
+	FVector4 PrevQs[MaxConstrainedBodies];		// World-space particle CoM rotations
+	FVector PrevXs[MaxConstrainedBodies];		// World-space joint connector positions
+
+	// Accumulated Impulse and AngularImpulse (Impulse * Dt since they are mass multiplied position corrections)
+	FVector NetLinearImpulse;
+	FVector NetAngularImpulse;
+
+	// XPBD Accumulators (net impulse for each soft constraint/drive)
+	float LinearSoftLambda;						// XPBD constraint multipliers (net applied constraint-space deltas)
+	float TwistSoftLambda;						// XPBD constraint multipliers (net applied constraint-space deltas)
+	float SwingSoftLambda;						// XPBD constraint multipliers (net applied constraint-space deltas)
+	float LinearDriveLambda;					// XPBD constraint multipliers (net applied constraint-space deltas)
 	float TwistDriveLambda;						// XPBD constraint multipliers (net applied constraint-space deltas)
 	float Swing1DriveLambda;					// XPBD constraint multipliers (net applied constraint-space deltas)
 	float Swing2DriveLambda;					// XPBD constraint multipliers (net applied constraint-space deltas)
-	FVector4 PrevQs[MaxConstrainedBodies];		// World-space particle CoM rotations
-	FVector PrevXs[MaxConstrainedBodies];		// World-space joint connector positions
 	
+	// Tolerances below which we stop solving
 	float PositionTolerance;					// Distance error below which we consider a constraint or drive solved
 	float AngleTolerance;						// Angle error below which we consider a constraint or drive solved
 
@@ -276,9 +283,11 @@ export void ApplyPositionConstraintSoft(
 	const uniform FVector DR0 = DLambda * Multiply(M->InvIs[0], AngularAxis0);
 	const uniform FVector DR1 = -DLambda * Multiply(M->InvIs[1], AngularAxis1);
 
-	Lambda += DLambda;
 	ApplyPositionDelta(M, 1.0f, DP0, DP1);
 	ApplyRotationDelta(M, 1.0f, DR0, DR1);
+
+	Lambda += DLambda;
+	M->NetLinearImpulse = M->NetLinearImpulse + DLambda * Axis;
 }
 
 // See "XPBD: Position-Based Simulation of Compliant Constrained Dynamics"
@@ -319,8 +328,10 @@ export void ApplyRotationConstraintSoftKD(
 	//const FVec3 DR1 = IA1 * -DLambda;
 	const uniform FVector DR1 = Axis * -(DLambda * II1);
 
-	Lambda += DLambda;
 	ApplyRotationDelta(M, DIndex, 1.0f, DR1);
+
+	Lambda += DLambda;
+	M->NetAngularImpulse = M->NetAngularImpulse + DLambda * Axis;
 }
 
 export void ApplyRotationConstraintSoftDD(
@@ -362,8 +373,10 @@ export void ApplyRotationConstraintSoftDD(
 	const uniform FVector DR0 = Axis * (DLambda * II0);
 	const uniform FVector DR1 = Axis * -(DLambda * II1);
 
-	Lambda += DLambda;
 	ApplyRotationDelta(M, 1.0f, DR0, DR1);
+
+	Lambda += DLambda;
+	M->NetAngularImpulse = M->NetAngularImpulse + DLambda * Axis;
 }
 
 
@@ -385,6 +398,8 @@ export void ApplyPointPositionConstraintKD(
 	const uniform FVector DR1 = Multiply(M->InvIs[DIndex], VectorCross(M->Xs[DIndex] - M->Ps[DIndex], VectorNegate(DX)));
 
 	ApplyDelta(M, DIndex, LinearStiffness, DP1, DR1);
+
+	M->NetLinearImpulse = M->NetLinearImpulse + LinearStiffness * DX;
 }
 
 export void ApplyPointPositionConstraintDD(
@@ -406,6 +421,8 @@ export void ApplyPointPositionConstraintDD(
 
 	ApplyPositionDelta(M, LinearStiffness, DP0, DP1);
 	ApplyRotationDelta(M, LinearStiffness, DR0, DR1);
+
+	M->NetLinearImpulse = M->NetLinearImpulse + LinearStiffness * DX;
 }
 
 export void ApplyPointPositionProjection(
@@ -418,31 +435,19 @@ export void ApplyPointPositionProjection(
 		const uniform float InvM0 = M->InvMs[0] * M->ProjectionInvMassScale;
 		const uniform FMatrix33 InvI0 = Multiply(M->InvIs[0], M->ProjectionInvMassScale);
 
-		uniform FMatrix33 J = SetMatrix33(0, 0, 0);
-		if (InvM0 > SMALL_NUMBER)
-		{
-			J = AddAB(J, ComputeJointFactorMatrix(M->Xs[0] - M->Ps[0], InvI0, InvM0));
-		}
-		if (M->InvMs[1] > SMALL_NUMBER)
-		{
-			J = AddAB(J, ComputeJointFactorMatrix(M->Xs[1] - M->Ps[1], M->InvIs[1], M->InvMs[1]));
-		}
+		uniform FMatrix33 J = ComputeJointFactorMatrix(M->Xs[1] - M->Ps[1], M->InvIs[1], M->InvMs[1]);
 		const uniform FMatrix33 IJ = MatrixInverse(J);
 		const uniform FVector DX = Multiply(IJ, CX);
+		const uniform FVector NDX = VectorNegate(DX);
+		const uniform FVector DP1 = M->InvMs[1] * NDX;
+		const uniform FVector DR1 = Multiply(M->InvIs[1], VectorCross(M->Xs[1] - M->Ps[1], NDX));
+		ApplyDelta(M, 1, LinearProjection, DP1, DR1);
+	}
 
-		if (InvM0 > SMALL_NUMBER)
-		{
-			const uniform FVector DP0 = InvM0 * DX;
-			const uniform FVector DR0 = Multiply(InvI0, VectorCross(M->Xs[0] - M->Ps[0], DX));
-			ApplyDelta(M, 0, LinearProjection, DP0, DR0);
-		}
-		if (M->InvMs[1] > SMALL_NUMBER)
-		{
-			const uniform FVector NDX = VectorNegate(DX);
-			const uniform FVector DP1 = M->InvMs[1] * NDX;
-			const uniform FVector DR1 = Multiply(M->InvIs[1], VectorCross(M->Xs[1] - M->Ps[1], NDX));
-			ApplyDelta(M, 1, LinearProjection, DP1, DR1);
-		}
+	// Final position correction to fix errors from rotation in above
+	{
+		const uniform FVector NCX2 = VectorNegate(M->Xs[1] - M->Xs[0]);
+		ApplyPositionDelta(M, 1, LinearProjection, NCX2);
 	}
 
 	// Apply a velocity correction to set the relative velocity at the joint to zero
@@ -452,11 +457,11 @@ export void ApplyPointPositionProjection(
 
 		// Calculate joint-space mass
 		uniform FMatrix33 J = SetMatrix33(0, 0, 0);
-		if (InvM0 > 0)
+		if (InvM0 > SMALL_NUMBER)
 		{
 			J = AddAB(J, ComputeJointFactorMatrix(M->Xs[0] - M->Ps[0], InvI0, InvM0));
 		}
-		if (M->InvMs[1] > 0)
+		if (M->InvMs[1] > SMALL_NUMBER)
 		{
 			J = AddAB(J, ComputeJointFactorMatrix(M->Xs[1] - M->Ps[1], M->InvIs[1], M->InvMs[1]));
 		}
@@ -465,16 +470,15 @@ export void ApplyPointPositionProjection(
 		// Velocity correction
 		const uniform FVector V0 = M->Vs[0] + VectorCross(M->Ws[0], M->Xs[0] - M->Ps[0]);
 		const uniform FVector V1 = M->Vs[1] + VectorCross(M->Ws[1], M->Xs[1] - M->Ps[1]);
-		const uniform FVector CV = V1 - V0;
-		const uniform FVector DV = Multiply(IJ, CV);
+		const uniform FVector DV = Multiply(IJ, V1 - V0);
 
-		if (InvM0 > 0)
+		if (InvM0 > SMALL_NUMBER)
 		{
 			const uniform FVector DV0 = InvM0 * DV;
 			const uniform FVector DW0 = Multiply(InvI0, VectorCross(M->Xs[0] - M->Ps[0], DV));
 			ApplyVelocityDelta(M, 0, LinearProjection, DV0, DW0);
 		}
-		if (M->InvMs[1] > 0)
+		if (M->InvMs[1] > SMALL_NUMBER)
 		{
 			const uniform FVector DV1 = M->InvMs[1] * VectorNegate(DV);
 			const uniform FVector DW1 = Multiply(M->InvIs[1], VectorCross(M->Xs[1] - M->Ps[1], VectorNegate(DV)));
diff --git a/Engine/Source/Runtime/Experimental/Chaos/Private/Chaos/PBDJointConstraints.cpp b/Engine/Source/Runtime/Experimental/Chaos/Private/Chaos/PBDJointConstraints.cpp
--- a/Engine/Source/Runtime/Experimental/Chaos/Private/Chaos/PBDJointConstraints.cpp
+++ b/Engine/Source/Runtime/Experimental/Chaos/Private/Chaos/PBDJointConstraints.cpp
@@ -125,7 +125,6 @@ namespace Chaos
 		, LinearDriveStiffness(0)
 		, LinearDriveDamping(0)
 		, AngularDrivePositionTarget(FRotation3::FromIdentity())
-		, AngularDriveTargetAngles(FVec3(0, 0, 0))
 		, AngularDriveVelocityTarget(FVec3(0, 0, 0))
 		, bAngularSLerpPositionDriveEnabled(false)
 		, bAngularSLerpVelocityDriveEnabled(false)
@@ -136,6 +135,8 @@ namespace Chaos
 		, AngularDriveForceMode(EJointForceMode::Acceleration)
 		, AngularDriveStiffness(0)
 		, AngularDriveDamping(0)
+		, LinearBreakForce(0)
+		, AngularBreakTorque(0)
 	{
 		if (bChaos_Joint_ISPC_Enabled)
 		{
@@ -178,6 +179,7 @@ namespace Chaos
 		, Level(INDEX_NONE)
 		, Color(INDEX_NONE)
 		, IslandSize(0)
+		, bDisabled(false)
 	{
 	}
 
@@ -196,6 +198,8 @@ namespace Chaos
 		, MinParentMassRatio(0)
 		, MaxInertiaRatio(0)
 		, AngularConstraintPositionCorrection(1.0f)
+		, ProjectionInvMassScale(0)
+		, VelProjectionInvMassScale(0)
 		, bEnableTwistLimits(true)
 		, bEnableSwingLimits(true)
 		, bEnableDrives(true)
@@ -212,8 +216,6 @@ namespace Chaos
 		, SoftTwistDamping(0)
 		, SoftSwingStiffness(0)
 		, SoftSwingDamping(0)
-		, ProjectionInvMassScale(0)
-		, VelProjectionInvMassScale(0)
 	{
 	}
 
@@ -394,6 +396,30 @@ namespace Chaos
 	}
 
 
+	bool FPBDJointConstraints::IsConstraintEnabled(int32 ConstraintIndex) const
+	{
+		return !ConstraintStates[ConstraintIndex].bDisabled;
+	}
+
+
+	void FPBDJointConstraints::SetConstraintEnabled(int32 ConstraintIndex, bool bEnabled)
+	{
+		ConstraintStates[ConstraintIndex].bDisabled = !bEnabled;
+	}
+
+
+	void FPBDJointConstraints::BreakConstraint(int32 ConstraintIndex)
+	{
+		SetConstraintEnabled(ConstraintIndex, false);
+		BreakCallback(Handles[ConstraintIndex]);
+	}
+
+
+	void FPBDJointConstraints::FixConstraints(int32 ConstraintIndex)
+	{
+		SetConstraintEnabled(ConstraintIndex, true);
+	}
+
 	
 	void FPBDJointConstraints::SetPreApplyCallback(const FJointPreApplyCallback& Callback)
 	{
@@ -430,7 +456,19 @@ namespace Chaos
 		PostProjectCallback = nullptr;
 	}
 
-	
+
+	void FPBDJointConstraints::SetBreakCallback(const FJointBreakCallback& Callback)
+	{
+		BreakCallback = Callback;
+	}
+
+
+	void FPBDJointConstraints::ClearBreakCallback()
+	{
+		BreakCallback = nullptr;
+	}
+
+
 	const typename FPBDJointConstraints::FConstraintContainerHandle* FPBDJointConstraints::GetConstraintHandle(int32 ConstraintIndex) const
 	{
 		return Handles[ConstraintIndex];
@@ -892,7 +930,6 @@ namespace Chaos
 			GatherSolverJointState(JointIndex);
 		}
 
-		FJointSolverResult NetResult;
 		for (int32 PairIt = 0; PairIt < NumPairIts; ++PairIt)
 		{
 			UE_LOG(LogChaosJoint, VeryVerbose, TEXT("  Pair Iteration %d / %d"), PairIt, NumPairIts);
@@ -1027,6 +1064,11 @@ namespace Chaos
 	// This will converge slowly in some cases, particularly where resolving angular constraints violates position constraints and vice versa.
 	int32 FPBDJointConstraints::ApplySingle(const FReal Dt, const int32 ConstraintIndex, const int32 NumPairIts, const int32 It, const int32 NumIts)
 	{
+		if (!IsConstraintEnabled(ConstraintIndex))
+		{
+			return 0;
+		}
+
 		const TVector<TGeometryParticleHandle<FReal, 3>*, 2>& Constraint = ConstraintParticles[ConstraintIndex];
 		UE_LOG(LogChaosJoint, VeryVerbose, TEXT("Solve Joint Constraint %d %s %s (dt = %f; it = %d / %d)"), ConstraintIndex, *Constraint[0]->ToString(), *Constraint[1]->ToString(), Dt, It, NumIts);
 
@@ -1072,11 +1114,22 @@ namespace Chaos
 		UpdateParticleState(Particle0->CastToRigidParticle(), Dt, Solver.GetPrevP(0), Solver.GetPrevQ(0), Solver.GetP(0), Solver.GetQ(0), bUpdateVelocity);
 		UpdateParticleState(Particle1->CastToRigidParticle(), Dt, Solver.GetPrevP(1), Solver.GetPrevQ(1), Solver.GetP(1), Solver.GetQ(1), bUpdateVelocity);
 
+		// @todo(ccaulfield): The break limit should really be applied to the impulse in the solver to prevent 1-frame impulses larger than the threshold
+		if ((JointSettings.LinearBreakForce > 0.0f) || (JointSettings.AngularBreakTorque > 0.0f))
+		{
+			ApplyBreakThreshold(Dt, ConstraintIndex, Solver.GetNetLinearImpulse(), Solver.GetNetAngularImpulse());
+		}
+
 		return NumActive;
 	}
 
 	int32 FPBDJointConstraints::ApplyPushOutSingle(const FReal Dt, const int32 ConstraintIndex, const int32 NumPairIts, const int32 It, const int32 NumIts)
 	{
+		if (!IsConstraintEnabled(ConstraintIndex))
+		{
+			return 0;
+		}
+
 		const TVector<TGeometryParticleHandle<FReal, 3>*, 2>& Constraint = ConstraintParticles[ConstraintIndex];
 		UE_LOG(LogChaosJoint, VeryVerbose, TEXT("Project Joint Constraint %d %s %s (dt = %f; it = %d / %d)"), ConstraintIndex, *Constraint[0]->ToString(), *Constraint[1]->ToString(), Dt, It, NumIts);
 
@@ -1113,9 +1166,29 @@ namespace Chaos
 		UpdateParticleStateExplicit(Particle0->CastToRigidParticle(), Dt, Solver.GetP(0), Solver.GetQ(0), Solver.GetV(0), Solver.GetW(0));
 		UpdateParticleStateExplicit(Particle1->CastToRigidParticle(), Dt, Solver.GetP(1), Solver.GetQ(1), Solver.GetV(1), Solver.GetW(1));
 
+		// @todo(ccaulfield): should probably add to net impulses in push out too...
+
 		return NumActive;
 	}
 
+	void FPBDJointConstraints::ApplyBreakThreshold(const FReal Dt, int32 ConstraintIndex, const FVec3& LinearImpulse, const FVec3& AngularImpulse)
+	{
+		const FPBDJointSettings& JointSettings = ConstraintSettings[ConstraintIndex];
+
+		// LinearImpulse is not really an impulse - it is a mass-weighted position delta.
+		// The Threshold is a force limit, so we need to convert it to a position delta caused by that force in one timestep
+		const FReal LinearThreshold = JointSettings.LinearBreakForce * Dt * Dt;
+		const FReal AngularThreshold = JointSettings.AngularBreakTorque * Dt * Dt;
+
+		const FReal LinearThresholdSq = LinearThreshold * LinearThreshold;
+		const FReal AngularThresholdSq = AngularThreshold * AngularThreshold;
+		const bool bBreak = (LinearImpulse.SizeSquared() > LinearThresholdSq) || (AngularImpulse.SizeSquared() > AngularThresholdSq);
+		if (bBreak)
+		{
+			BreakConstraint(ConstraintIndex);
+		}
+	}
+
 
 	// Assign an Island, Level and Color to each constraint. Constraints must be processed in Level order, but
 	// constraints of the same color are independent and can be processed in parallel (SIMD or Task)
diff --git a/Engine/Source/Runtime/Experimental/Chaos/Public/Chaos/Joint/JointSolverConstraints.inl b/Engine/Source/Runtime/Experimental/Chaos/Public/Chaos/Joint/JointSolverConstraints.inl
--- a/Engine/Source/Runtime/Experimental/Chaos/Public/Chaos/Joint/JointSolverConstraints.inl
+++ b/Engine/Source/Runtime/Experimental/Chaos/Public/Chaos/Joint/JointSolverConstraints.inl
@@ -431,13 +431,14 @@ namespace Chaos
 		const FRotation3& R1,
 		const FRotation3& RTwist)
 	{
-		const FVec3 Axis = R1 * FJointConstants::TwistAxis();
-		const FReal Angle = FPBDJointUtilities::GetTwistAngle(RTwist);
-		const FReal TwistAngleTarget = JointSettings.AngularDriveTargetAngles[(int32)EJointAngularConstraintIndex::Twist];
-		const FReal DTwistAngle = Angle - TwistAngleTarget;
-
-		RowState.Axis = Axis;
-		RowState.CalculateError(DTwistAngle, RowData.Limit);
+		// @todo(ccaulfield): Reimplement with AngularDriveTarget
+		//const FVec3 Axis = R1 * FJointConstants::TwistAxis();
+		//const FReal Angle = FPBDJointUtilities::GetTwistAngle(RTwist);
+		//const FReal TwistAngleTarget = JointSettings.AngularDriveTargetAngles[(int32)EJointAngularConstraintIndex::Twist];
+		//const FReal DTwistAngle = Angle - TwistAngleTarget;
+
+		//RowState.Axis = Axis;
+		//RowState.CalculateError(DTwistAngle, RowData.Limit);
 	}
 
 	FORCEINLINE void FJointSolverConstraints::UpdateConeSwingDrive(
@@ -447,23 +448,24 @@ namespace Chaos
 		const FRotation3& R0,
 		const FRotation3& RSwing)
 	{
-		FReal SwingAngle;
-		FVec3 SwingAxisLocal;
-		RSwing.ToAxisAndAngleSafe(SwingAxisLocal, SwingAngle, FJointConstants::Swing1Axis(), AngleTolerance);
-		if (SwingAngle > PI)
-		{
-			SwingAngle = SwingAngle - (FReal)2 * PI;
-		}
-
-		// Circular swing target (max of Swing1, Swing2 targets)
-		// @todo(ccaulfield): what should cone target really do?
-		const FReal Swing1Target = JointSettings.AngularDriveTargetAngles[(int32)EJointAngularConstraintIndex::Swing1];
-		const FReal Swing2Target = JointSettings.AngularDriveTargetAngles[(int32)EJointAngularConstraintIndex::Swing2];
-		const FReal SwingAngleTarget = FMath::Max(Swing1Target, Swing2Target);
-		const FReal DSwingAngle = SwingAngle - SwingAngleTarget;
-
-		RowState.Axis = R0 * SwingAxisLocal;
-		RowState.CalculateError(DSwingAngle, RowData.Limit);
+		// @todo(ccaulfield): Reimplement with AngularDriveTarget
+		//FReal SwingAngle;
+		//FVec3 SwingAxisLocal;
+		//RSwing.ToAxisAndAngleSafe(SwingAxisLocal, SwingAngle, FJointConstants::Swing1Axis(), AngleTolerance);
+		//if (SwingAngle > PI)
+		//{
+		//	SwingAngle = SwingAngle - (FReal)2 * PI;
+		//}
+
+		//// Circular swing target (max of Swing1, Swing2 targets)
+		//// @todo(ccaulfield): what should cone target really do?
+		//const FReal Swing1Target = JointSettings.AngularDriveTargetAngles[(int32)EJointAngularConstraintIndex::Swing1];
+		//const FReal Swing2Target = JointSettings.AngularDriveTargetAngles[(int32)EJointAngularConstraintIndex::Swing2];
+		//const FReal SwingAngleTarget = FMath::Max(Swing1Target, Swing2Target);
+		//const FReal DSwingAngle = SwingAngle - SwingAngleTarget;
+
+		//RowState.Axis = R0 * SwingAxisLocal;
+		//RowState.CalculateError(DSwingAngle, RowData.Limit);
 	}
 
 	FORCEINLINE void FJointSolverConstraints::UpdateSwingDrive(
@@ -473,14 +475,15 @@ namespace Chaos
 		const FRotation3& R0,
 		const FRotation3& RSwing)
 	{
-		const FReal R01SwingYorZ = (FJointConstants::AxisIndex((EJointAngularConstraintIndex)RowData.ConstraintIndex) == 2) ? RSwing.Z : RSwing.Y;	// Can't index a quat :(
-		const FReal SwingAngle = 4.0f * FMath::Atan2(R01SwingYorZ, 1.0f + RSwing.W);
-		const FReal SwingAngleTarget = JointSettings.AngularDriveTargetAngles[(int32)RowData.ConstraintIndex];
-		const FReal DSwingAngle = SwingAngle - SwingAngleTarget;
-		const FVec3& AxisLocal = ((EJointAngularConstraintIndex)RowData.ConstraintIndex == EJointAngularConstraintIndex::Swing1) ? FJointConstants::Swing1Axis() : FJointConstants::Swing2Axis();
-
-		RowState.Axis = R0 * AxisLocal;
-		RowState.CalculateError(DSwingAngle, RowData.Limit);
+		// @todo(ccaulfield): Reimplement with AngularDriveTarget
+		//const FReal R01SwingYorZ = (FJointConstants::AxisIndex((EJointAngularConstraintIndex)RowData.ConstraintIndex) == 2) ? RSwing.Z : RSwing.Y;	// Can't index a quat :(
+		//const FReal SwingAngle = 4.0f * FMath::Atan2(R01SwingYorZ, 1.0f + RSwing.W);
+		//const FReal SwingAngleTarget = JointSettings.AngularDriveTargetAngles[(int32)RowData.ConstraintIndex];
+		//const FReal DSwingAngle = SwingAngle - SwingAngleTarget;
+		//const FVec3& AxisLocal = ((EJointAngularConstraintIndex)RowData.ConstraintIndex == EJointAngularConstraintIndex::Swing1) ? FJointConstants::Swing1Axis() : FJointConstants::Swing2Axis();
+
+		//RowState.Axis = R0 * AxisLocal;
+		//RowState.CalculateError(DSwingAngle, RowData.Limit);
 	}
 
 	FORCEINLINE void FJointSolverConstraints::UpdateSLerpDrive(
diff --git a/Engine/Source/Runtime/Experimental/Chaos/Public/Chaos/Joint/PBDJointSolverGaussSeidel.h b/Engine/Source/Runtime/Experimental/Chaos/Public/Chaos/Joint/PBDJointSolverGaussSeidel.h
--- a/Engine/Source/Runtime/Experimental/Chaos/Public/Chaos/Joint/PBDJointSolverGaussSeidel.h
+++ b/Engine/Source/Runtime/Experimental/Chaos/Public/Chaos/Joint/PBDJointSolverGaussSeidel.h
@@ -60,6 +60,16 @@ namespace Chaos
 			return Ws[Index];
 		}
 
+		FORCEINLINE const FVec3& GetNetLinearImpulse() const
+		{
+			return NetLinearImpulse;
+		}
+
+		FORCEINLINE const FVec3& GetNetAngularImpulse() const
+		{
+			return NetAngularImpulse;
+		}
+
 		FJointSolverGaussSeidel();
 
 		void Init(
@@ -361,6 +371,12 @@ namespace Chaos
 			const FPBDJointSolverSettings& SolverSettings,
 			const FPBDJointSettings& JointSettings);
 
+		int32 ApplyPointDoubleSwingLockProjection(
+			const FReal Dt,
+			const FPBDJointSolverSettings& SolverSettings,
+			const FPBDJointSettings& JointSettings);
+
+
 
 		// Local-space constraint settings
 		FRigidTransform3 XLs[MaxConstrainedBodies];	// Local-space joint connector transforms
@@ -369,14 +385,10 @@ namespace Chaos
 
 		// World-space constraint state
 		FVec3 Xs[MaxConstrainedBodies];				// World-space joint connector positions
-		FReal LinearSoftLambda;						// XPBD constraint multipliers (net applied constraint-space deltas)
-		FReal LinearDriveLambda;					// XPBD constraint multipliers (net applied constraint-space deltas)
 		FRotation3 Rs[MaxConstrainedBodies];		// World-space joint connector rotations
 
 		// World-space body state
 		FVec3 Ps[MaxConstrainedBodies];				// World-space particle CoM positions
-		FReal TwistSoftLambda;						// XPBD constraint multipliers (net applied constraint-space deltas)
-		FReal SwingSoftLambda;						// XPBD constraint multipliers (net applied constraint-space deltas)
 		FRotation3 Qs[MaxConstrainedBodies];		// World-space particle CoM rotations
 		FVec3 Vs[MaxConstrainedBodies];				// World-space particle CoM velocities
 		FVec3 Ws[MaxConstrainedBodies];				// World-space particle CoM angular velocities
@@ -384,15 +396,26 @@ namespace Chaos
 
 		// XPBD Previous iteration world-space body state
 		FVec3 PrevPs[MaxConstrainedBodies];			// World-space particle CoM positions
-		FReal TwistDriveLambda;						// XPBD constraint multipliers (net applied constraint-space deltas)
-		FReal Swing1DriveLambda;					// XPBD constraint multipliers (net applied constraint-space deltas)
-		FReal Swing2DriveLambda;					// XPBD constraint multipliers (net applied constraint-space deltas)
 		FRotation3 PrevQs[MaxConstrainedBodies];	// World-space particle CoM rotations
 		FVec3 PrevXs[MaxConstrainedBodies];			// World-space joint connector positions
 
+		// Accumulated Impulse and AngularImpulse (Impulse * Dt since they are mass multiplied position corrections)
+		FVec3 NetLinearImpulse;
+		FVec3 NetAngularImpulse;
+
+		// XPBD Accumulators (net impulse for each soft constraint/drive)
+		FReal LinearSoftLambda;
+		FReal TwistSoftLambda;
+		FReal SwingSoftLambda;
+		FReal LinearDriveLambda;
+		FReal TwistDriveLambda;
+		FReal Swing1DriveLambda;
+		FReal Swing2DriveLambda;
+
+		// Tolerances below which we stop solving
 		FReal PositionTolerance;					// Distance error below which we consider a constraint or drive solved
 		FReal AngleTolerance;						// Angle error below which we consider a constraint or drive solved
-
+	
 		FReal ProjectionInvMassScale;
 		FReal VelProjectionInvMassScale;
 	};
diff --git a/Engine/Source/Runtime/Experimental/Chaos/Public/Chaos/PBDJointConstraintTypes.h b/Engine/Source/Runtime/Experimental/Chaos/Public/Chaos/PBDJointConstraintTypes.h
--- a/Engine/Source/Runtime/Experimental/Chaos/Public/Chaos/PBDJointConstraintTypes.h
+++ b/Engine/Source/Runtime/Experimental/Chaos/Public/Chaos/PBDJointConstraintTypes.h
@@ -17,6 +17,8 @@ namespace Chaos
 
 	using FJointPostApplyCallback = TFunction<void(const FReal Dt, const TArray<FPBDJointConstraintHandle*>& InConstraintHandles)>;
 
+	using FJointBreakCallback = TFunction<void(FPBDJointConstraintHandle * ConstraintHandle)>;
+
 	enum class EJointMotionType : int32
 	{
 		Free,
@@ -131,9 +133,7 @@ namespace Chaos
 		FReal LinearDriveStiffness;
 		FReal LinearDriveDamping;
 
-		// @todo(ccaulfield): remove one of these
 		FRotation3 AngularDrivePositionTarget;
-		FVec3 AngularDriveTargetAngles;
 		FVec3 AngularDriveVelocityTarget;
 
 		bool bAngularSLerpPositionDriveEnabled;
@@ -145,6 +145,9 @@ namespace Chaos
 		EJointForceMode AngularDriveForceMode;
 		FReal AngularDriveStiffness;
 		FReal AngularDriveDamping;
+
+		FReal LinearBreakForce;
+		FReal AngularBreakTorque;
 	};
 
 	class CHAOS_API FPBDJointSolverSettings
@@ -168,6 +171,10 @@ namespace Chaos
 		// Angular stiffness
 		FReal AngularConstraintPositionCorrection;
 
+		// Projection settings
+		FReal ProjectionInvMassScale;
+		FReal VelProjectionInvMassScale;
+
 		// @todo(ccaulfield): remove these TEMP overrides for testing
 		bool bEnableTwistLimits;
 		bool bEnableSwingLimits;
@@ -185,47 +192,6 @@ namespace Chaos
 		FReal SoftTwistDamping;
 		FReal SoftSwingStiffness;
 		FReal SoftSwingDamping;
-		FReal ProjectionInvMassScale;
-		FReal VelProjectionInvMassScale;
-	};
-
-	class FJointSolverResult
-	{
-	public:
-		static FJointSolverResult MakeActive() { return { 1, 0 }; }
-		static FJointSolverResult MakeInactive() { return { 0, 1 }; }
-		static FJointSolverResult Make(int32 NumActive, int32 NumInactive) { return { NumActive, NumInactive }; }
-
-		FJointSolverResult()
-			: NumActive(0)
-			, NumInactive(0)
-		{
-		}
-
-		int32 GetNumActive() const { return NumActive; }
-		int32 GetNumInactive() const { return NumInactive; }
-
-		FJointSolverResult& operator+=(const FJointSolverResult& R)
-		{
-			NumActive += R.NumActive;
-			NumInactive += R.NumInactive;
-			return *this;
-		}
-
-		friend FJointSolverResult operator+(const FJointSolverResult& L, const FJointSolverResult& R)
-		{
-			return { L.NumActive + R.NumActive, L.NumInactive + R.NumInactive };
-		}
-
-	private:
-		FJointSolverResult(int32 InActive, int32 InSolved)
-			: NumActive(InActive)
-			, NumInactive(InSolved)
-		{
-		}
-
-		int32 NumActive;
-		int32 NumInactive;
 	};
 
 }
\ No newline at end of file
diff --git a/Engine/Source/Runtime/Experimental/Chaos/Public/Chaos/PBDJointConstraints.h b/Engine/Source/Runtime/Experimental/Chaos/Public/Chaos/PBDJointConstraints.h
--- a/Engine/Source/Runtime/Experimental/Chaos/Public/Chaos/PBDJointConstraints.h
+++ b/Engine/Source/Runtime/Experimental/Chaos/Public/Chaos/PBDJointConstraints.h
@@ -19,7 +19,6 @@ namespace Chaos
 {
 	class FJointSolverConstraints;
 	class FJointSolverGaussSeidel;
-	class FJointSolverResult;
 
 	class CHAOS_API FPBDJointConstraintHandle : public TContainerConstraintHandle<FPBDJointConstraints>
 	{
@@ -55,6 +54,7 @@ namespace Chaos
 		int32 Level;
 		int32 Color;
 		int32 IslandSize;
+		bool bDisabled;
 	};
 
 	/**
@@ -106,6 +106,26 @@ namespace Chaos
 		// @todo(ccaulfield): rename/remove  this
 		void RemoveConstraints(const TSet<TGeometryParticleHandle<FReal, 3>*>& RemovedParticles);
 
+		/*
+		 * Whether the constraint is enabled
+		 */
+		bool IsConstraintEnabled(int32 ConstraintIndex) const;
+
+		/*
+		 * Enable or disable a constraints
+		 */
+		void SetConstraintEnabled(int32 ConstraintIndex, bool bEnabled);
+
+		/*
+		 * Force a constraints to break
+		 */
+		void BreakConstraint(int32 ConstraintIndex);
+
+		/**
+		 * Repair a broken constraints (does not adjust particle positions)
+		 */
+		void FixConstraints(int32 ConstraintIndex);
+
 		void SetPreApplyCallback(const FJointPostApplyCallback& Callback);
 		void ClearPreApplyCallback();
 
@@ -115,6 +135,9 @@ namespace Chaos
 		void SetPostProjectCallback(const FJointPostApplyCallback& Callback);
 		void ClearPostProjectCallback();
 
+		void SetBreakCallback(const FJointBreakCallback& Callback);
+		void ClearBreakCallback();
+
 		//
 		// Constraint API
 		//
@@ -197,6 +220,7 @@ namespace Chaos
 		int32 ApplyBatch(const FReal Dt, const int32 BatchIndex, const int32 NumPairIts, const int32 It, const int32 NumIts);
 		int32 ApplySingle(const FReal Dt, const int32 ConstraintIndex, const int32 NumPairIts, const int32 It, const int32 NumIts);
 		int32 ApplyPushOutSingle(const FReal Dt, const int32 ConstraintIndex, const int32 NumPairIts, const int32 It, const int32 NumIts);
+		void ApplyBreakThreshold(const FReal Dt, int32 ConstraintIndex, const FVec3& LinearImpulse, const FVec3& AngularImpulse);
 
 		FPBDJointSolverSettings Settings;
 
@@ -213,6 +237,7 @@ namespace Chaos
 		FJointPreApplyCallback PreApplyCallback;
 		FJointPostApplyCallback PostApplyCallback;
 		FJointPostApplyCallback PostProjectCallback;
+		FJointBreakCallback BreakCallback;
 
 		// @todo(ccaulfield): optimize storage for joint solver
 		TArray<FJointSolverGaussSeidel> ConstraintSolvers;
diff --git a/Engine/Source/Runtime/Experimental/Chaos/Public/Chaos/ParticleHandle.h b/Engine/Source/Runtime/Experimental/Chaos/Public/Chaos/ParticleHandle.h
--- a/Engine/Source/Runtime/Experimental/Chaos/Public/Chaos/ParticleHandle.h
+++ b/Engine/Source/Runtime/Experimental/Chaos/Public/Chaos/ParticleHandle.h
@@ -1249,6 +1249,25 @@ public:
 		return TRotation<T, d>::FromIdentity();
 	}
 
+	T LinearEtherDrag() const
+	{
+		if (auto RigidHandle = Handle->CastToRigidParticle())
+		{
+			return RigidHandle->LinearEtherDrag();
+		}
+		return 0.0f;
+	}
+
+	T AngularEtherDrag() const
+	{
+		if (auto RigidHandle = Handle->CastToRigidParticle())
+		{
+			return RigidHandle->AngularEtherDrag();
+		}
+		return 0.0f;
+	}
+
+
 #if CHAOS_CHECKED
 	const FName& DebugName() const
 	{
