Chaos - joint break logic cleanup

diff --git a/Engine/Source/Runtime/Experimental/Chaos/Private/Chaos/PBDJointConstraints.cpp b/Engine/Source/Runtime/Experimental/Chaos/Private/Chaos/PBDJointConstraints.cpp
--- a/Engine/Source/Runtime/Experimental/Chaos/Private/Chaos/PBDJointConstraints.cpp
+++ b/Engine/Source/Runtime/Experimental/Chaos/Private/Chaos/PBDJointConstraints.cpp
@@ -1175,14 +1175,28 @@ namespace Chaos
 	{
 		const FPBDJointSettings& JointSettings = ConstraintSettings[ConstraintIndex];
 
-		// LinearImpulse is not really an impulse - it is a mass-weighted position delta.
+		// NOTE: LinearImpulse/AngularImpulse are not really an impulses - it is a mass-weighted position/rotation delta, or (impulse x dt).
 		// The Threshold is a force limit, so we need to convert it to a position delta caused by that force in one timestep
-		const FReal LinearThreshold = JointSettings.LinearBreakForce * Dt * Dt;
-		const FReal AngularThreshold = JointSettings.AngularBreakTorque * Dt * Dt;
 
-		const FReal LinearThresholdSq = LinearThreshold * LinearThreshold;
-		const FReal AngularThresholdSq = AngularThreshold * AngularThreshold;
-		const bool bBreak = (LinearImpulse.SizeSquared() > LinearThresholdSq) || (AngularImpulse.SizeSquared() > AngularThresholdSq);
+		bool bBreak = false;
+		if (!bBreak && (JointSettings.LinearBreakForce > 0.0f))
+		{
+			const FReal LinearThreshold = JointSettings.LinearBreakForce * Dt * Dt;
+			UE_LOG(LogChaosJoint, VeryVerbose, TEXT("Constraint %d Linear Break Check: %f / %f"), ConstraintIndex, LinearImpulse.Size(), LinearThreshold);
+
+			const FReal LinearThresholdSq = LinearThreshold * LinearThreshold;
+			bBreak = LinearImpulse.SizeSquared() > LinearThresholdSq;
+		}
+
+		if (!bBreak && (JointSettings.AngularBreakTorque > 0.0f))
+		{
+			const FReal AngularThreshold = JointSettings.AngularBreakTorque * Dt * Dt;
+			UE_LOG(LogChaosJoint, VeryVerbose, TEXT("Constraint %d Angular Break Check: %f / %f"), ConstraintIndex, AngularImpulse.Size(), AngularThreshold);
+
+			const FReal AngularThresholdSq = AngularThreshold * AngularThreshold;
+			bBreak = AngularImpulse.SizeSquared() > AngularThresholdSq;
+		}
+
 		if (bBreak)
 		{
 			BreakConstraint(ConstraintIndex);
