Re-added navmesh support for geometry collections. This is a naive implementation that periodically triggers nav updates on a timer rather than knowing the state of the object. Also the navigation export will push the render triangles to the navigation system until we
#rnx

diff --git a/Engine/Source/Runtime/Engine/Private/PhysicsEngine/Experimental/PhysScene_Chaos.cpp b/Engine/Source/Runtime/Engine/Private/PhysicsEngine/Experimental/PhysScene_Chaos.cpp
--- a/Engine/Source/Runtime/Engine/Private/PhysicsEngine/Experimental/PhysScene_Chaos.cpp
+++ b/Engine/Source/Runtime/Engine/Private/PhysicsEngine/Experimental/PhysScene_Chaos.cpp
@@ -2243,17 +2243,15 @@ void FPhysScene_ChaosInterface::SyncBodies(TSolver* Solver)
 {
 	DECLARE_SCOPE_CYCLE_COUNTER(TEXT("SyncBodies"), STAT_SyncBodies, STATGROUP_Physics);
 	TArray<FPhysScenePendingComponentTransform_Chaos> PendingTransforms;
-
 	TSet<FGeometryCollectionPhysicsProxy*> GCProxies;
 
 	{
 		Chaos::FPBDRigidActiveParticlesBufferAccessor Accessor(Solver->GetActiveParticlesBuffer());
 
-
 		const Chaos::FPBDRigidActiveParticlesBufferOut* ActiveParticleBuffer = Accessor.GetSolverOutData();
 		for (Chaos::TGeometryParticle<float, 3>* ActiveParticle : ActiveParticleBuffer->ActiveGameThreadParticles)
 		{
-		if (IPhysicsProxyBase* ProxyBase = ActiveParticle->GetProxy())
+			if (IPhysicsProxyBase* ProxyBase = ActiveParticle->GetProxy())
 			{
 				if (ProxyBase->GetType() == EPhysicsProxyType::SingleRigidParticleType)
 				{
@@ -2311,10 +2309,11 @@ void FPhysScene_ChaosInterface::SyncBodies(TSolver* Solver)
 		}
 	}
 	
-	for (auto* GCProxy : GCProxies)
+	for (FGeometryCollectionPhysicsProxy* GCProxy : GCProxies)
 	{
 		GCProxy->PullFromPhysicsState();
 	}
+
 	for (const FPhysScenePendingComponentTransform_Chaos& ComponentTransform : PendingTransforms)
 	{
 		if (ComponentTransform.OwningComp != nullptr)
diff --git a/Engine/Source/Runtime/Experimental/GeometryCollectionEngine/Private/GeometryCollection/GeometryCollectionComponent.cpp b/Engine/Source/Runtime/Experimental/GeometryCollectionEngine/Private/GeometryCollection/GeometryCollectionComponent.cpp
--- a/Engine/Source/Runtime/Experimental/GeometryCollectionEngine/Private/GeometryCollection/GeometryCollectionComponent.cpp
+++ b/Engine/Source/Runtime/Experimental/GeometryCollectionEngine/Private/GeometryCollection/GeometryCollectionComponent.cpp
@@ -22,6 +22,7 @@
 #include "PhysicsSolver.h"
 #include "Physics/PhysicsFiltering.h"
 #include "PhysicalMaterials/Experimental/ChaosPhysicalMaterial.h"
+#include "AI/NavigationSystemHelpers.h"
 
 #if WITH_EDITOR
 #include "AssetToolsModule.h"
@@ -74,6 +75,11 @@ void HackRegisterGeomAccelerator(UGeometryCollectionComponent& Component)
 }
 #endif
 
+// Size in CM used as a threshold for whether a geometry in the collection is collected and exported for
+// navigation purposes. Measured as the diagonal of the leaf node bounds.
+float GGeometryCollectionNavigationSizeThreshold = 20.0f;
+FAutoConsoleVariableRef CVarGeometryCollectionNavigationSizeThreshold(TEXT("p.GeometryCollectionNavigationSizeThreshold"), GGeometryCollectionNavigationSizeThreshold, TEXT("Size in CM used as a threshold for whether a geometry in the collection is collected and exported for navigation purposes. Measured as the diagonal of the leaf node bounds."));
+
 FGeomComponentCacheParameters::FGeomComponentCacheParameters()
 	: CacheMode(EGeometryCollectionCacheType::None)
 	, TargetCache(nullptr)
@@ -158,6 +164,8 @@ UGeometryCollectionComponent::UGeometryCollectionComponent(const FObjectInitiali
 	BodyInstance.SetObjectType(ECC_Destructible);
 
 	EventDispatcher = ObjectInitializer.CreateDefaultSubobject<UChaosGameplayEventDispatcher>(this, TEXT("GameplayEventDispatcher"));
+
+	bHasCustomNavigableGeometry = EHasCustomNavigableGeometry::Yes;
 }
 
 Chaos::FPhysicsSolver* GetSolver(const UGeometryCollectionComponent& GeometryCollectionComponent)
@@ -413,6 +421,125 @@ void UGeometryCollectionComponent::DispatchBreakEvent(const FChaosBreakEvent& Ev
 	}
 }
 
+bool UGeometryCollectionComponent::DoCustomNavigableGeometryExport(FNavigableGeometryExport& GeomExport) const
+{
+	TArray<FVector> OutVertexBuffer;
+	TArray<int32> OutIndexBuffer;
+
+	const FGeometryCollection* const Collection = RestCollection->GetGeometryCollection().Get();
+	check(Collection);
+
+	const float SizeThreshold = GGeometryCollectionNavigationSizeThreshold * GGeometryCollectionNavigationSizeThreshold;
+
+	// for all geometry. inspect bounding box build int list of transform indices.
+	int32 VertexCount = 0;
+	int32 FaceCountEstimate = 0;
+	TArray<int32> GeometryIndexBuffer;
+	TArray<int32> TransformIndexBuffer;
+
+	int32 NumGeometry = Collection->NumElements(FGeometryCollection::GeometryGroup);
+
+	const TManagedArray<FBox>& BoundingBox = Collection->BoundingBox;
+	const TManagedArray<int32>& TransformIndexArray = Collection->TransformIndex;
+	const TManagedArray<int32>& VertexCountArray = Collection->VertexCount;
+	const TManagedArray<int32>& FaceCountArray = Collection->FaceCount;
+	const TManagedArray<int32>& VertexStartArray = Collection->VertexStart;
+	const TManagedArray<FVector>& Vertex = Collection->Vertex;
+
+	for(int32 GeometryGroupIndex = 0; GeometryGroupIndex < NumGeometry; GeometryGroupIndex++)
+	{
+		if(BoundingBox[GeometryGroupIndex].GetSize().SizeSquared() > SizeThreshold)
+		{
+			TransformIndexBuffer.Add(TransformIndexArray[GeometryGroupIndex]);
+			GeometryIndexBuffer.Add(GeometryGroupIndex);
+			VertexCount += VertexCountArray[GeometryGroupIndex];
+			FaceCountEstimate += FaceCountArray[GeometryGroupIndex];
+		}
+	}
+
+	// Get all the geometry transforms in component space (they are stored natively in parent-bone space)
+	TArray<FTransform> GeomToComponent;
+	GeometryCollectionAlgo::GlobalMatrices(GetTransformArray(), GetParentArray(), TransformIndexBuffer, GeomToComponent);
+
+	OutVertexBuffer.AddUninitialized(VertexCount);
+
+	int32 DestVertex = 0;
+	//for each "subset" we care about 
+	for(int32 SubsetIndex = 0; SubsetIndex < GeometryIndexBuffer.Num(); ++SubsetIndex)
+	{
+		//find indices into the collection data
+		int32 GeometryIndex = GeometryIndexBuffer[SubsetIndex];
+		int32 TransformIndex = TransformIndexBuffer[SubsetIndex];
+		
+		int32 SourceGeometryVertexStart = VertexStartArray[GeometryIndex];
+		int32 SourceGeometryVertexCount = VertexCountArray[GeometryIndex];
+
+		ParallelFor(SourceGeometryVertexCount, [&](int32 PointIdx)
+			{
+				//extract vertex from source
+				int32 SourceGeometryVertexIndex = SourceGeometryVertexStart + PointIdx;
+				FVector const VertexInWorldSpace = GeomToComponent[SubsetIndex].TransformPosition(Vertex[SourceGeometryVertexIndex]);
+
+				int32 DestVertexIndex = DestVertex + PointIdx;
+				OutVertexBuffer[DestVertexIndex].X = VertexInWorldSpace.X;
+				OutVertexBuffer[DestVertexIndex].Y = VertexInWorldSpace.Y;
+				OutVertexBuffer[DestVertexIndex].Z = VertexInWorldSpace.Z;
+			});
+
+		DestVertex += SourceGeometryVertexCount;
+	}
+
+	//gather data needed for indices
+	const TManagedArray<int32>& FaceStartArray = Collection->FaceStart;
+	const TManagedArray<FIntVector>& Indices = Collection->Indices;
+	const TManagedArray<bool>& Visible = GetVisibleArray();
+	const TManagedArray<int32>& MaterialIndex = Collection->MaterialIndex;
+
+	//pre-allocate enough room (assuming all faces are visible)
+	OutIndexBuffer.AddUninitialized(3 * FaceCountEstimate);
+
+	//reset vertex counter so that we base the indices off the new location rather than the global vertex list
+	DestVertex = 0;
+	int32 DestinationIndex = 0;
+
+	//leaving index traversal in a different loop to help cache coherency of source data
+	for(int32 SubsetIndex = 0; SubsetIndex < GeometryIndexBuffer.Num(); ++SubsetIndex)
+	{
+		int32 GeometryIndex = GeometryIndexBuffer[SubsetIndex];
+
+		//for each index, subtract the starting vertex for that geometry to make it 0-based.  Then add the new starting vertex index for this geometry
+		int32 SourceGeometryVertexStart = VertexStartArray[GeometryIndex];
+		int32 SourceGeometryVertexCount = VertexCountArray[GeometryIndex];
+		int32 IndexDelta = DestVertex - SourceGeometryVertexStart;
+
+		int32 FaceStart = FaceStartArray[GeometryIndex];
+		int32 FaceCount = FaceCountArray[GeometryIndex];
+
+		//Copy the faces
+		for(int FaceIdx = FaceStart; FaceIdx < FaceStart + FaceCount; FaceIdx++)
+		{
+			if(Visible[FaceIdx])
+			{
+				OutIndexBuffer[DestinationIndex++] = Indices[FaceIdx].X + IndexDelta;
+				OutIndexBuffer[DestinationIndex++] = Indices[FaceIdx].Y + IndexDelta;
+				OutIndexBuffer[DestinationIndex++] = Indices[FaceIdx].Z + IndexDelta;
+			}
+		}
+
+		DestVertex += SourceGeometryVertexCount;
+	}
+
+	// Invisible faces make the index buffer smaller
+	OutIndexBuffer.SetNum(DestinationIndex);
+
+	// Push as a custom mesh to navigation system
+	// #CHAOSTODO This is pretty inefficient as it copies the whole buffer transforming each vert by the component to world
+	// transform. Investigate a move aware custom mesh for pre-transformed verts to speed this up.
+	GeomExport.ExportCustomMesh(OutVertexBuffer.GetData(), OutVertexBuffer.Num(), OutIndexBuffer.GetData(), OutIndexBuffer.Num(), GetComponentToWorld());
+
+	return true;
+}
+
 static void DispatchGeometryCollectionBreakEvent(const FChaosBreakEvent& Event)
 {
 	if (UGeometryCollectionComponent* const GC = Cast<UGeometryCollectionComponent>(Event.Component))
@@ -1008,16 +1135,16 @@ void UGeometryCollectionComponent::TickComponent(float DeltaTime, enum ELevelTic
 			bRenderStateDirty = false;
 			//DynamicCollection->MakeClean(); clean?
 
-// 			const UWorld* MyWorld = GetWorld();
-// 			if (MyWorld && MyWorld->IsGameWorld())
-// 			{
-// 				//cycle every 0xff frames
-// 				//@todo - Need way of seeing if the collection is actually changing
-// 				if (bNavigationRelevant && bRegistered && (((GFrameCounter + NavmeshInvalidationTimeSliceIndex) & 0xff) == 0))
-// 				{
-// 					UpdateNavigationData();
-// 				}
-// 			}
+			const UWorld* MyWorld = GetWorld();
+			if (MyWorld && MyWorld->IsGameWorld())
+			{
+				//cycle every 0xff frames
+				//@todo - Need way of seeing if the collection is actually changing
+				if (bNavigationRelevant && bRegistered && (((GFrameCounter + NavmeshInvalidationTimeSliceIndex) & 0xff) == 0))
+				{
+					UpdateNavigationData();
+				}
+			}
 		}
 	}
 }
diff --git a/Engine/Source/Runtime/Experimental/GeometryCollectionEngine/Public/GeometryCollection/GeometryCollectionComponent.h b/Engine/Source/Runtime/Experimental/GeometryCollectionEngine/Public/GeometryCollection/GeometryCollectionComponent.h
--- a/Engine/Source/Runtime/Experimental/GeometryCollectionEngine/Public/GeometryCollection/GeometryCollectionComponent.h
+++ b/Engine/Source/Runtime/Experimental/GeometryCollectionEngine/Public/GeometryCollection/GeometryCollectionComponent.h
@@ -464,8 +464,6 @@ public:
 	
 	bool GetIsObjectLoading() { return IsObjectLoading; }
 
-
-
 	/**
 	*
 	*/
@@ -533,6 +531,7 @@ public:
 	UPROPERTY(Transient, VisibleAnywhere, BlueprintReadWrite, Category = "Chaos")
 	bool CachePlayback;
 
+	bool DoCustomNavigableGeometryExport(FNavigableGeometryExport& GeomExport) const override;
 
 public:
 	UPROPERTY(BlueprintAssignable, Category = "Collision")
