Added desync concept to PT particles

diff --git a/Engine/Source/Programs/HeadlessChaos/Private/HeadlessChaosTestRewind.cpp b/Engine/Source/Programs/HeadlessChaos/Private/HeadlessChaosTestRewind.cpp
--- a/Engine/Source/Programs/HeadlessChaos/Private/HeadlessChaosTestRewind.cpp
+++ b/Engine/Source/Programs/HeadlessChaos/Private/HeadlessChaosTestRewind.cpp
@@ -659,24 +659,36 @@ namespace ChaosTest {
 			X[Step] = Particle->X();
 			TickSolverHelper(Module,Solver);
 
+			auto PTParticle = static_cast<FSingleParticlePhysicsProxy<TPBDRigidParticle<FReal,3>>*>(Particle->GetProxy())->GetHandle();
+			auto PTKinematic = static_cast<FSingleParticlePhysicsProxy<TKinematicGeometryParticle<FReal,3>>*>(Kinematic->GetProxy())->GetHandle();
+
 			//see that particle has desynced
 			if(Step < LastStep)
 			{
 				//If we're still in the past make sure future has been marked as desync
 				FGeometryParticleState State(*Particle);
 				EXPECT_EQ(EFutureQueryResult::Desync,RewindData->GetFutureStateAtFrame(State,Step));
+				EXPECT_EQ(PTParticle->SyncState(),ESyncState::HardDesync);
 
 				FGeometryParticleState KinState(*Kinematic);
 				const EFutureQueryResult KinFutureStatus = RewindData->GetFutureStateAtFrame(KinState,Step);
 				if(Step < 10)
 				{
 					EXPECT_EQ(KinFutureStatus,EFutureQueryResult::Ok);
+					EXPECT_EQ(PTKinematic->SyncState(),ESyncState::InSync);
 				}
 				else
 				{
 					EXPECT_EQ(KinFutureStatus,EFutureQueryResult::Desync);
+					EXPECT_EQ(PTKinematic->SyncState(),ESyncState::HardDesync);
 				}
 			}
+			else
+			{
+				//Last resim frame ran so everything is marked as in sync
+				EXPECT_EQ(PTParticle->SyncState(),ESyncState::InSync);
+				EXPECT_EQ(PTKinematic->SyncState(),ESyncState::InSync);
+			}
 		}
 
 		EXPECT_EQ(Kinematic->X()[2],50);	//Rewound kinematic and only did one update, so use that first update
diff --git a/Engine/Source/Runtime/Experimental/Chaos/Public/Chaos/GeometryParticles.h b/Engine/Source/Runtime/Experimental/Chaos/Public/Chaos/GeometryParticles.h
--- a/Engine/Source/Runtime/Experimental/Chaos/Public/Chaos/GeometryParticles.h
+++ b/Engine/Source/Runtime/Experimental/Chaos/Public/Chaos/GeometryParticles.h
@@ -275,6 +275,7 @@ namespace Chaos
 			TArrayCollection::AddArray(&MHasBounds);
 			TArrayCollection::AddArray(&MSpatialIdx);
 			TArrayCollection::AddArray(&MUserData);
+			TArrayCollection::AddArray(&MSyncState);
 #if CHAOS_CHECKED
 			TArrayCollection::AddArray(&MDebugName);
 #endif
@@ -303,6 +304,7 @@ namespace Chaos
 			, MHasBounds(MoveTemp(Other.MHasBounds))
 			, MSpatialIdx(MoveTemp(Other.MSpatialIdx))
 			, MUserData(MoveTemp(Other.MUserData))
+			, MSyncState(MoveTemp(Other.MSyncState))
 #if CHAOS_DETERMINISTIC
 			, MParticleIDs(MoveTemp(Other.MParticleIDs))
 #endif
@@ -320,6 +322,7 @@ namespace Chaos
 			TArrayCollection::AddArray(&MHasBounds);
 			TArrayCollection::AddArray(&MSpatialIdx);
 			TArrayCollection::AddArray(&MUserData);
+			TArrayCollection::AddArray(&MSyncState);
 #if CHAOS_DETERMINISTIC
 			TArrayCollection::AddArray(&MParticleIDs);
 #endif
@@ -352,6 +355,7 @@ namespace Chaos
 			TArrayCollection::AddArray(&MHasBounds);
 			TArrayCollection::AddArray(&MSpatialIdx);
 			TArrayCollection::AddArray(&MUserData);
+			TArrayCollection::AddArray(&MSyncState);
 #if CHAOS_DETERMINISTIC
 			TArrayCollection::AddArray(&MParticleIDs);
 #endif
@@ -378,6 +382,9 @@ namespace Chaos
 		CHAOS_API void*& UserData(const int32 Index) { return MUserData[Index]; }
 		CHAOS_API const void* UserData(const int32 Index) const { return MUserData[Index]; }
 
+		CHAOS_API ESyncState& SyncState(const int32 Index) { return MSyncState[Index].State; }
+		CHAOS_API ESyncState SyncState(const int32 Index) const { return MSyncState[Index].State; }
+
 		CHAOS_API TSerializablePtr<FImplicitObject> Geometry(const int32 Index) const { return MGeometry[Index]; }
 
 		CHAOS_API const TUniquePtr<FImplicitObject>& DynamicGeometry(const int32 Index) const { return MDynamicGeometry[Index]; }
@@ -604,6 +611,7 @@ namespace Chaos
 		TArrayCollectionArray<bool> MHasBounds;
 		TArrayCollectionArray<FSpatialAccelerationIdx> MSpatialIdx;
 		TArrayCollectionArray<void*> MUserData;
+		TArrayCollectionArray<FSyncState> MSyncState;
 
 		void UpdateShapesArray(const int32 Index)
 		{
diff --git a/Engine/Source/Runtime/Experimental/Chaos/Public/Chaos/GeometryParticlesfwd.h b/Engine/Source/Runtime/Experimental/Chaos/Public/Chaos/GeometryParticlesfwd.h
--- a/Engine/Source/Runtime/Experimental/Chaos/Public/Chaos/GeometryParticlesfwd.h
+++ b/Engine/Source/Runtime/Experimental/Chaos/Public/Chaos/GeometryParticlesfwd.h
@@ -11,6 +11,24 @@ namespace Chaos
 		Other
 	};
 
+	enum class ESyncState: uint8
+	{
+		InSync,	    //in sync with recorded data
+		SoftDesync, //recorded data still matches, but may interact with hard desynced particles
+		HardDesync, //recorded data mismatches, must run collision detection again
+	};
+
+	struct FSyncState
+	{
+		ESyncState State;
+
+		FSyncState()
+		: State(ESyncState::InSync)
+		{
+
+		}
+	};
+
 	template<class T, int d, EGeometryParticlesSimType SimType>
 	class TGeometryParticlesImp;
 
diff --git a/Engine/Source/Runtime/Experimental/Chaos/Public/Chaos/ParticleHandle.h b/Engine/Source/Runtime/Experimental/Chaos/Public/Chaos/ParticleHandle.h
--- a/Engine/Source/Runtime/Experimental/Chaos/Public/Chaos/ParticleHandle.h
+++ b/Engine/Source/Runtime/Experimental/Chaos/Public/Chaos/ParticleHandle.h
@@ -441,6 +441,16 @@ public:
 #endif
 	}
 
+	ESyncState SyncState() const
+	{
+		return GeometryParticles->SyncState(ParticleIdx);
+	}
+
+	void SetSyncState(ESyncState State)
+	{
+		GeometryParticles->SyncState(ParticleIdx) = State;
+	}
+
 	TSerializablePtr<FImplicitObject> Geometry() const { return GeometryParticles->Geometry(ParticleIdx); }
 	void SetGeometry(TSerializablePtr<FImplicitObject> InGeometry) { GeometryParticles->SetGeometry(ParticleIdx, InGeometry); }
 
diff --git a/Engine/Source/Runtime/Experimental/ChaosSolvers/Public/RewindData.h b/Engine/Source/Runtime/Experimental/ChaosSolvers/Public/RewindData.h
--- a/Engine/Source/Runtime/Experimental/ChaosSolvers/Public/RewindData.h
+++ b/Engine/Source/Runtime/Experimental/ChaosSolvers/Public/RewindData.h
@@ -894,7 +894,11 @@ public:
 					Info.Desync(CurFrame, LatestFrame);
 				}
 
-				
+				if(Info.bDesync)
+				{
+					//Any desync from GT is considered a hard desync - in theory we could make this more fine grained
+					Info.GetPTParticle()->SetSyncState(ESyncState::HardDesync);
+				}
 			}
 		}
 	}
@@ -905,6 +909,7 @@ public:
 		
 		if(IsResim())
 		{
+			const bool bLastResim = IsFinalResim();
 			ensure(IsInGameThread());
 			//snap particles forward that are not desynced or  do not have resim enabled
 			//TODO: handle desync case
@@ -922,6 +927,12 @@ public:
 						}
 					}
 				}
+
+				//Last resim so mark everything as in sync
+				if(bLastResim)
+				{
+					Info.GetPTParticle()->SetSyncState(ESyncState::InSync);
+				}
 			}
 		}
 
@@ -934,6 +945,11 @@ public:
 		return CurFrame < LatestFrame;
 	}
 
+	bool IsFinalResim() const
+	{
+		return (CurFrame + 1) == LatestFrame;
+	}
+
 	//Number of particles that we're currently storing history for
 	int32 GetNumDirtyParticles() const { return AllDirtyParticles.Num(); }
 
