Chaos Physics:
- Jitter fix: For AccImpClipping calculate zero impulse if not penetrating
- Jitter workaround: For AccImpClipping disable energy clamping

diff --git a/Engine/Source/Runtime/Experimental/Chaos/Private/Chaos/PBDCollisionConstraintsContact.cpp b/Engine/Source/Runtime/Experimental/Chaos/Private/Chaos/PBDCollisionConstraintsContact.cpp
--- a/Engine/Source/Runtime/Experimental/Chaos/Private/Chaos/PBDCollisionConstraintsContact.cpp
+++ b/Engine/Source/Runtime/Experimental/Chaos/Private/Chaos/PBDCollisionConstraintsContact.cpp
@@ -179,8 +179,11 @@ namespace Chaos
 			FVec3 ClippedAccumulatedImpulse(0.0f);  // To be calculated
 			{
 				//Project to normal
-				const float NewAccImpNormalSize = FVec3::DotProduct(NewUnclippedAccumulatedImpulse, Contact.Normal);
-				if (NewAccImpNormalSize > 0) // Clipping impulses to be positive
+				const FReal NewAccImpNormalSize = FVec3::DotProduct(NewUnclippedAccumulatedImpulse, Contact.Normal);
+				// Clipping impulses to be positive and contacts to be penetrating or touching
+				// ToDo: PenetrationBuffer avoids a lot of jitter added by the ApplypushOut function. Set this to 0 when that is fixed
+				const FReal PenetrationBuffer = 0.5;
+				if (NewAccImpNormalSize > 0 && Contact.Phi <= PenetrationBuffer)
 				{
 					if (Friction <= 0)
 					{
@@ -235,7 +238,8 @@ namespace Chaos
 				// Clamp the delta impulse to make sure we don't gain kinetic energy (ignore potential energy)
 				// Todo: Investigate Energy clamping to work with accumulated impulses
 				// else this might introduce a small amount of jitter, since we are clamping delta impulses instead of accumulated ones
-				OutDeltaImpulse = GetEnergyClampedImpulse(Particle0->CastToRigidParticle(), Particle1->CastToRigidParticle(), OutDeltaImpulse, VectorToPoint1, VectorToPoint2, Body1Velocity, Body2Velocity);
+				// Disable this until then
+				//OutDeltaImpulse = GetEnergyClampedImpulse(Particle0->CastToRigidParticle(), Particle1->CastToRigidParticle(), OutDeltaImpulse, VectorToPoint1, VectorToPoint2, Body1Velocity, Body2Velocity);
 			}
 
 			if (bIsRigidDynamic0)
