Fix ensure on bad proxy, also clear the proxy even if we can't process it. Fixes crash in unit tests

Benn.Gallagher

diff --git a/Engine/Source/Runtime/Experimental/ChaosSolvers/Private/PBDRigidsSolver.cpp b/Engine/Source/Runtime/Experimental/ChaosSolvers/Private/PBDRigidsSolver.cpp
--- a/Engine/Source/Runtime/Experimental/ChaosSolvers/Private/PBDRigidsSolver.cpp
+++ b/Engine/Source/Runtime/Experimental/ChaosSolvers/Private/PBDRigidsSolver.cpp
@@ -740,7 +740,7 @@ namespace Chaos
 				break;
 			}
 			default:
-			ensure("Unknown proxy type in physics solver.");
+			ensure(0 && TEXT("Unknown proxy type in physics solver."));
 			}
 		});
 
@@ -823,7 +823,12 @@ namespace Chaos
 					break;
 				}
 				default:
-				ensure("Unknown proxy type in physics solver.");
+				{
+					ensure(0 && TEXT("Unknown proxy type in physics solver."));
+					//Can't use, but we can still mark as "clean"
+					Dirty.Proxy->ResetDirtyIdx();
+				}
+				
 				}
 			});
 
