Small cleanup and better commenting on union group creation

diff --git a/Engine/Source/Runtime/Experimental/Chaos/Private/Chaos/PBDRigidClustering.cpp b/Engine/Source/Runtime/Experimental/Chaos/Private/Chaos/PBDRigidClustering.cpp
--- a/Engine/Source/Runtime/Experimental/Chaos/Private/Chaos/PBDRigidClustering.cpp
+++ b/Engine/Source/Runtime/Experimental/Chaos/Private/Chaos/PBDRigidClustering.cpp
@@ -406,9 +406,7 @@ namespace Chaos
 	//==========================================================================
 
 	template<class T_FPBDRigidsEvolution, class T_FPBDCollisionConstraint, class T, int d>
-	TPBDRigidClustering<T_FPBDRigidsEvolution, T_FPBDCollisionConstraint, T, d>::TPBDRigidClustering(
-		T_FPBDRigidsEvolution& InEvolution, 
-		TPBDRigidClusteredParticles<T, d>& InParticles)
+	TPBDRigidClustering<T_FPBDRigidsEvolution, T_FPBDCollisionConstraint, T, d>::TPBDRigidClustering(T_FPBDRigidsEvolution& InEvolution, TPBDRigidClusteredParticles<T, d>& InParticles)
 		: MEvolution(InEvolution)
 		, MParticles(InParticles)
 		, MCollisionImpulseArrayDirty(true)
@@ -423,13 +421,7 @@ namespace Chaos
 
 	DECLARE_CYCLE_STAT(TEXT("TPBDRigidClustering<>::CreateClusterParticle"), STAT_CreateClusterParticle, STATGROUP_Chaos);
 	template<class T_FPBDRigidsEvolution, class T_FPBDCollisionConstraint, class T, int d>
-	Chaos::TPBDRigidClusteredParticleHandle<float, 3>* 
-	TPBDRigidClustering<T_FPBDRigidsEvolution, T_FPBDCollisionConstraint, T, d>::CreateClusterParticle(
-		const int32 ClusterGroupIndex,
-		TArray<Chaos::TPBDRigidParticleHandle<T,d>*>&& Children,
-		const FClusterCreationParameters<T>& Parameters,
-		TSharedPtr<Chaos::FImplicitObject, ESPMode::ThreadSafe> ProxyGeometry,
-		const TRigidTransform<T, d>* ForceMassOrientation)
+	Chaos::TPBDRigidClusteredParticleHandle<float, 3>* TPBDRigidClustering<T_FPBDRigidsEvolution, T_FPBDCollisionConstraint, T, d>::CreateClusterParticle(const int32 ClusterGroupIndex, TArray<Chaos::TPBDRigidParticleHandle<T,d>*>&& Children, const FClusterCreationParameters<T>& Parameters, TSharedPtr<Chaos::FImplicitObject, ESPMode::ThreadSafe> ProxyGeometry, const TRigidTransform<T, d>* ForceMassOrientation)
 	{
 		SCOPE_CYCLE_COUNTER(STAT_CreateClusterParticle);
 
@@ -455,10 +447,7 @@ namespace Chaos
 		NewParticle->SetClusterGroupIndex(ClusterGroupIndex);
 		NewParticle->SetStrains(0.0);
 
-		//
 		// Update clustering data structures.
-		//
-
 		if (MChildren.Contains(NewParticle))
 		{
 			MChildren[NewParticle] = MoveTemp(Children);
@@ -467,6 +456,7 @@ namespace Chaos
 		{
 			MChildren.Add(NewParticle, MoveTemp(Children));
 		}
+
 		const TArray<TPBDRigidParticleHandle<T, d>*>& ChildrenArray = MChildren[NewParticle];
 		TSet<TPBDRigidParticleHandle<T, d>*> ChildrenSet(ChildrenArray);
 
@@ -488,8 +478,7 @@ namespace Chaos
 				ClusteredChild->ClusterIds().Id = NewParticle;
 				NewParticle->Strains() += ClusteredChild->Strains();
 
-				NewParticle->SetCollisionImpulses(
-					FMath::Max(NewParticle->CollisionImpulses(), ClusteredChild->CollisionImpulses()));
+				NewParticle->SetCollisionImpulses(FMath::Max(NewParticle->CollisionImpulses(), ClusteredChild->CollisionImpulses()));
 
 				const int32 NewCG = NewParticle->CollisionGroup();
 				const int32 ChildCG = ClusteredChild->CollisionGroup();
@@ -506,8 +495,14 @@ namespace Chaos
 		UpdateMassProperties(NewParticle, ChildrenSet, ForceMassOrientation);
 		UpdateGeometry(NewParticle, ChildrenSet, ProxyGeometry, Parameters);
 		GenerateConnectionGraph(NewParticle, Parameters);
+
 		NewParticle->SetSleeping(bClusterIsAsleep);
-		if (ClusterGroupIndex) AddToClusterUnion(ClusterGroupIndex, NewParticle);
+
+		if(ClusterGroupIndex)
+		{
+			AddToClusterUnion(ClusterGroupIndex, NewParticle);
+		}
+
 		return NewParticle;
 	}
 
@@ -602,43 +597,64 @@ namespace Chaos
 	void TPBDRigidClustering<T_FPBDRigidsEvolution, T_FPBDCollisionConstraint, T, d>::UnionClusterGroups()
 	{
 		SCOPE_CYCLE_COUNTER(STAT_UnionClusterGroups);
-		if (ClusterUnionMap.Num())
+
+		if(ClusterUnionMap.Num())
 		{
-			TMap < TPBDRigidParticleHandle<T, 3>*, TPBDRigidParticleHandle<T, 3>*> ClusterParents;
-			TMap < int32, TArray< TPBDRigidParticleHandle<T, 3>*>> NewClusterGroups;
-			for (TTuple<int32, TArray<TPBDRigidClusteredParticleHandle<T, 3>* >>& Group : ClusterUnionMap)
+			TMap<TPBDRigidParticleHandle<T, 3>*, TPBDRigidParticleHandle<T, 3>*> ChildToParentMap;
+			TMap<int32, TArray<TPBDRigidParticleHandle<T, 3>*>> NewClusterGroups;
+
+			// Walk the list of registered cluster groups
+			for(TTuple<int32, TArray<TPBDRigidClusteredParticleHandle<T, 3>* >>& Group : ClusterUnionMap)
 			{
 				int32 ClusterGroupID = Group.Key;
-				TArray<TPBDRigidClusteredParticleHandle<T, 3>* > Handles = Group.Value;
+				TArray<TPBDRigidClusteredParticleHandle<T, 3>*> Handles = Group.Value;
 
-				if (Handles.Num() > 1)
+				if(Handles.Num() > 1)
 				{
-					if (!NewClusterGroups.Contains(ClusterGroupID))
+					// First see if this is a new group
+					if(!NewClusterGroups.Contains(ClusterGroupID))
+					{
 						NewClusterGroups.Add(ClusterGroupID, TArray < TPBDRigidParticleHandle<T, 3>*>());
+					}
 
 					TArray<TPBDRigidParticleHandle<T, 3>*> ClusterBodies;
-					for (TPBDRigidClusteredParticleHandle<T, 3>* ActiveCluster : Handles)
+					for(TPBDRigidClusteredParticleHandle<T, 3>* ActiveCluster : Handles)
 					{
-						if (!ActiveCluster->Disabled())
+						if(!ActiveCluster->Disabled())
 						{
+							// If this is an external cluster (from the rest collection) we release its children and append them to the current group
 							TSet<TPBDRigidParticleHandle<T, 3>*> Children = ReleaseClusterParticles(ActiveCluster, nullptr, true);
 							NewClusterGroups[ClusterGroupID].Append(Children.Array());
-							for (auto& Child : Children) ClusterParents.Add(Child, ActiveCluster);
+							
+							for(TPBDRigidParticleHandle<T, 3>* Child : Children)
+							{
+								ChildToParentMap.Add(Child, ActiveCluster);
+							}
 						}
 					}
 				}
 			}
 
-			for (TTuple<int32, TArray<TPBDRigidParticleHandle<T, 3>* >>& Group : NewClusterGroups)
+			// For new cluster groups, create an internal cluster parent.
+			for(TTuple<int32, TArray<TPBDRigidParticleHandle<T, 3>* >>& Group : NewClusterGroups)
 			{
 				int32 ClusterGroupID = FMath::Abs(Group.Key);
-				TArray< TPBDRigidParticleHandle<T, 3>*> ActiveCluster = Group.Value;
-				FClusterCreationParameters<T> Parameters(0.3, 100, false, !!UnionsHaveCollisionParticles);  Parameters.ConnectionMethod = MClusterUnionConnectionType;
+
+				TArray<TPBDRigidParticleHandle<T, 3>*> ActiveCluster = Group.Value;
+
+				FClusterCreationParameters<T> Parameters(0.3, 100, false, !!UnionsHaveCollisionParticles);
+				Parameters.ConnectionMethod = MClusterUnionConnectionType;
 				TPBDRigidClusteredParticleHandleImp<float, 3, true>* Handle = CreateClusterParticle(-ClusterGroupID, MoveTemp(Group.Value), Parameters, TSharedPtr<FImplicitObject, ESPMode::ThreadSafe>());
 				Handle->SetInternalCluster(true);
+
 				MEvolution.SetPhysicsMaterial(Handle, MEvolution.GetPhysicsMaterial(ActiveCluster[0]));
-				for (auto& Constituent : ActiveCluster) MEvolution.DoInternalParticleInitilization( ClusterParents[Constituent], Handle);
+
+				for(TPBDRigidParticleHandle<T, 3>* Constituent : ActiveCluster)
+				{
+					MEvolution.DoInternalParticleInitilization(ChildToParentMap[Constituent], Handle);
+				}
 			}
+
 			ClusterUnionMap.Empty();
 		}
 	}
@@ -1946,9 +1962,7 @@ namespace Chaos
 	DECLARE_CYCLE_STAT(TEXT("TPBDRigidClustering<>::UpdateConnectivityGraphUsingDelaunayTriangulation"), STAT_UpdateConnectivityGraphUsingDelaunayTriangulation, STATGROUP_Chaos);
 
 	template<class T_FPBDRigidsEvolution, class T_FPBDCollisionConstraint, class T, int d>
-	void TPBDRigidClustering<T_FPBDRigidsEvolution, T_FPBDCollisionConstraint, T, d>::UpdateConnectivityGraphUsingDelaunayTriangulation(
-		Chaos::TPBDRigidClusteredParticleHandle<float, 3>* Parent,
-		const FClusterCreationParameters<T>& Parameters)
+	void TPBDRigidClustering<T_FPBDRigidsEvolution, T_FPBDCollisionConstraint, T, d>::UpdateConnectivityGraphUsingDelaunayTriangulation(Chaos::TPBDRigidClusteredParticleHandle<float, 3>* Parent, const FClusterCreationParameters<T>& Parameters)
 	{
 		SCOPE_CYCLE_COUNTER(STAT_UpdateConnectivityGraphUsingDelaunayTriangulation);
 
diff --git a/Engine/Source/Runtime/Experimental/Chaos/Public/Chaos/PBDRigidClustering.h b/Engine/Source/Runtime/Experimental/Chaos/Public/Chaos/PBDRigidClustering.h
--- a/Engine/Source/Runtime/Experimental/Chaos/Public/Chaos/PBDRigidClustering.h
+++ b/Engine/Source/Runtime/Experimental/Chaos/Public/Chaos/PBDRigidClustering.h
@@ -228,9 +228,18 @@ public:
 	/** If multi child proxy is used, this is the data needed */
 	const TArrayCollectionArray<TUniquePtr<TMultiChildProxyData<T, d>>>& GetMultiChildProxyDataArray() const { return MParticles.MultiChildProxyDataArray(); }
 
-	void AddToClusterUnion(int32 ClusterID, TPBDRigidClusteredParticleHandle<T, 3>* Handle) {
-		if (ClusterID <= 0) return;
-		if (!ClusterUnionMap.Contains(ClusterID)) ClusterUnionMap.Add(ClusterID, TArray<TPBDRigidClusteredParticleHandle<T, 3>*>());
+	void AddToClusterUnion(int32 ClusterID, TPBDRigidClusteredParticleHandle<T, 3>* Handle)
+	{
+		if(ClusterID <= 0)
+		{
+			return;
+		}
+
+		if(!ClusterUnionMap.Contains(ClusterID))
+		{
+			ClusterUnionMap.Add(ClusterID, TArray<TPBDRigidClusteredParticleHandle<T, 3>*>());
+		}
+
 		ClusterUnionMap[ClusterID].Add(Handle);
 	}
 
diff --git a/Engine/Source/Runtime/Experimental/Chaos/Public/Chaos/PBDRigidsEvolutionGBF.h b/Engine/Source/Runtime/Experimental/Chaos/Public/Chaos/PBDRigidsEvolutionGBF.h
--- a/Engine/Source/Runtime/Experimental/Chaos/Public/Chaos/PBDRigidsEvolutionGBF.h
+++ b/Engine/Source/Runtime/Experimental/Chaos/Public/Chaos/PBDRigidsEvolutionGBF.h
@@ -134,7 +134,10 @@ namespace Chaos
 
 		void DoInternalParticleInitilization(const TGeometryParticleHandle<float, 3>* OldParticle, const TGeometryParticleHandle<float, 3>* NewParticle) 
 		{ 
-			if (InternalParticleInitilization) InternalParticleInitilization(OldParticle, NewParticle); 
+			if(InternalParticleInitilization)
+			{
+				InternalParticleInitilization(OldParticle, NewParticle);
+			}
 		}
 
 
