Adding cases to GT and PT process during push for geometry collections. Removed PushToPhysicsState from geometry collections. Need to consider how to hook GC up to rewind/resim and fully implement these paths

diff --git a/Engine/Source/Runtime/Experimental/ChaosSolvers/Private/PBDRigidsSolver.cpp b/Engine/Source/Runtime/Experimental/ChaosSolvers/Private/PBDRigidsSolver.cpp
--- a/Engine/Source/Runtime/Experimental/ChaosSolvers/Private/PBDRigidsSolver.cpp
+++ b/Engine/Source/Runtime/Experimental/ChaosSolvers/Private/PBDRigidsSolver.cpp
@@ -739,6 +739,11 @@ namespace Chaos
 				ProcessProxyGT(Proxy,DataIdx,Dirty);
 				break;
 			}
+			case EPhysicsProxyType::GeometryCollectionType:
+			{
+				// Not invalid but doesn't currently use the remote data process
+				break;
+			}
 			default:
 			ensure(0 && TEXT("Unknown proxy type in physics solver."));
 			}
@@ -822,6 +827,13 @@ namespace Chaos
 					ProcessProxyPT(Proxy,DataIdx,Dirty,[this](const FUniqueIdx* UniqueIdx){ return Particles.CreateStaticParticles(1,UniqueIdx)[0]; });
 					break;
 				}
+				case EPhysicsProxyType::GeometryCollectionType:
+				{
+					// Currently no push needed for geometry collections and they handle the particle creation internally
+					// #TODO This skips the rewind data push so GC will not be rewindable until resolved.
+					Dirty.Proxy->ResetDirtyIdx();
+					break;
+				}
 				default:
 				{
 					ensure(0 && TEXT("Unknown proxy type in physics solver."));
diff --git a/Engine/Source/Runtime/Experimental/ChaosSolvers/Private/PhysicsProxy/GeometryCollectionPhysicsProxy.cpp b/Engine/Source/Runtime/Experimental/ChaosSolvers/Private/PhysicsProxy/GeometryCollectionPhysicsProxy.cpp
--- a/Engine/Source/Runtime/Experimental/ChaosSolvers/Private/PhysicsProxy/GeometryCollectionPhysicsProxy.cpp
+++ b/Engine/Source/Runtime/Experimental/ChaosSolvers/Private/PhysicsProxy/GeometryCollectionPhysicsProxy.cpp
@@ -1403,29 +1403,6 @@ void TGeometryCollectionPhysicsProxy<Traits>::BufferGameState()
 	//
 }
 
-template <typename Traits>
-void TGeometryCollectionPhysicsProxy<Traits>::PushToPhysicsState(const Chaos::FParticleData *InData)
-{
-	/*CONTEXT: GAMETHREAD->to->PHYSICSTHREAD
-	* Called on the game thread when the solver is about to advance forward.This
-	* callback should Enqueue commands on the PhysicsThread to update the state of
-	* the solver
-	*/
-
-	// We need to update physics thread world space bounding boxes.
-	for (FClusterHandle* Handle : SolverParticleHandles)
-	{
-		if (!Handle || !Handle->HasBounds())
-		{
-			continue;
-		}
-		const Chaos::TAABB<float, 3>& LocalBounds = Handle->LocalBounds();
-		const Chaos::TRigidTransform<float, 3> Xf(Handle->X(), Handle->R());
-		const Chaos::TAABB<float, 3> TransformedBBox = LocalBounds.TransformedAABB(Xf);
-		Handle->SetWorldSpaceInflatedBounds(TransformedBBox);
-	}
-}
-
 template <typename Traits>
 void TGeometryCollectionPhysicsProxy<Traits>::BufferPhysicsResults()
 {
diff --git a/Engine/Source/Runtime/Experimental/ChaosSolvers/Public/PhysicsProxy/GeometryCollectionPhysicsProxy.h b/Engine/Source/Runtime/Experimental/ChaosSolvers/Public/PhysicsProxy/GeometryCollectionPhysicsProxy.h
--- a/Engine/Source/Runtime/Experimental/ChaosSolvers/Public/PhysicsProxy/GeometryCollectionPhysicsProxy.h
+++ b/Engine/Source/Runtime/Experimental/ChaosSolvers/Public/PhysicsProxy/GeometryCollectionPhysicsProxy.h
@@ -118,20 +118,6 @@ public:
 	Chaos::FParticleData* NewData() { BufferGameState(); return nullptr; }
 	void BufferGameState();
 
-	/** 
-	 * Push game data to the physics thread.
-	 * 
-	 * The solver has determined that this proxy needs to update its physics state 
-	 * with game data, either as part of the solver registration process, or
-	 * a game thread particle instance was updated, which marked this proxy dirty.
-	 * 
-	 * Pulls data out of the \c GameToPhysInterchange, and updates solver particle 
-	 * data.
-	 *
-	 * Invoked by \c FPBDRigidsSolver::RegisterObject() and via particle dirtying. 
-	 */
-	void PushToPhysicsState(const Chaos::FParticleData*);
-
 	/** Called at the end of \c FPBDRigidsSolver::PushPhysicsStateExec(). */
 	void ClearAccumulatedData() {}
 
