Chaos
- new joint velocity projection method
- angular joint projection (disabled)

diff --git a/Engine/Source/Runtime/Experimental/Chaos/Private/Chaos/Joint/JointConstraintsCVars.cpp b/Engine/Source/Runtime/Experimental/Chaos/Private/Chaos/Joint/JointConstraintsCVars.cpp
--- a/Engine/Source/Runtime/Experimental/Chaos/Private/Chaos/Joint/JointConstraintsCVars.cpp
+++ b/Engine/Source/Runtime/Experimental/Chaos/Private/Chaos/Joint/JointConstraintsCVars.cpp
@@ -21,3 +21,9 @@ FAutoConsoleVariableRef CVarChaosJointDegenerateRotationLimit(TEXT("p.Chaos.Join
 
 bool bChaos_Joint_EllipticalFix = true;
 FAutoConsoleVariableRef CVarChaosJointEllipticalFix(TEXT("p.Chaos.Joint.EllipticalFix"), bChaos_Joint_EllipticalFix, TEXT("Enable the proper elliptical joint axis/error calculation"));
+
+float Chaos_Joint_VelProjectionAlpha = 0.1f;
+FAutoConsoleVariableRef CVarChaosJointVelProjectionScale(TEXT("p.Chaos.Joint.VelProjectionAlpha"), Chaos_Joint_VelProjectionAlpha, TEXT("How much of the velocity correction to apply during projection. Equivalent to (1-damping) for projection velocity delta"));
+
+bool bChaos_Joint_DisableSoftLimits = false;
+FAutoConsoleVariableRef CVarChaosJointDisableSoftLimits(TEXT("p.Chaos.Joint.DisableSoftLimits"), bChaos_Joint_DisableSoftLimits, TEXT("Disable soft limits (for debugging only)"));
diff --git a/Engine/Source/Runtime/Experimental/Chaos/Private/Chaos/Joint/JointConstraintsCVars.h b/Engine/Source/Runtime/Experimental/Chaos/Private/Chaos/Joint/JointConstraintsCVars.h
--- a/Engine/Source/Runtime/Experimental/Chaos/Private/Chaos/Joint/JointConstraintsCVars.h
+++ b/Engine/Source/Runtime/Experimental/Chaos/Private/Chaos/Joint/JointConstraintsCVars.h
@@ -18,3 +18,7 @@ extern int32 bChaos_Joint_MaxBatchSize;
 
 extern float Chaos_Joint_DegenerateRotationLimit;
 extern bool bChaos_Joint_EllipticalFix;
+
+extern float Chaos_Joint_VelProjectionAlpha;
+
+extern bool bChaos_Joint_DisableSoftLimits;
diff --git a/Engine/Source/Runtime/Experimental/Chaos/Private/Chaos/Joint/PBDJointSolverGaussSeidel.cpp b/Engine/Source/Runtime/Experimental/Chaos/Private/Chaos/Joint/PBDJointSolverGaussSeidel.cpp
--- a/Engine/Source/Runtime/Experimental/Chaos/Private/Chaos/Joint/PBDJointSolverGaussSeidel.cpp
+++ b/Engine/Source/Runtime/Experimental/Chaos/Private/Chaos/Joint/PBDJointSolverGaussSeidel.cpp
@@ -241,15 +241,22 @@ namespace Chaos
 		const EJointMotionType TwistMotion = JointSettings.AngularMotionTypes[(int32)EJointAngularConstraintIndex::Twist];
 		const EJointMotionType Swing1Motion = JointSettings.AngularMotionTypes[(int32)EJointAngularConstraintIndex::Swing1];
 		const EJointMotionType Swing2Motion = JointSettings.AngularMotionTypes[(int32)EJointAngularConstraintIndex::Swing2];
-		const bool bTwistSoft = JointSettings.bSoftTwistLimitsEnabled;
-		const bool bSwingSoft = JointSettings.bSoftSwingLimitsEnabled;
+		const bool bLinearSoft = FPBDJointUtilities::GetSoftLinearLimitEnabled(SolverSettings, JointSettings);
+		const bool bTwistSoft = FPBDJointUtilities::GetSoftTwistLimitEnabled(SolverSettings, JointSettings);
+		const bool bSwingSoft = FPBDJointUtilities::GetSoftSwingLimitEnabled(SolverSettings, JointSettings);
 
-		const bool bPointProject = 
-			(LinearProjection > 0.0f) 
+		const bool bPointProject =
+			(LinearProjection > 0.0f)
 			&& (LinearMotion[0] == EJointMotionType::Locked)
 			&& (LinearMotion[1] == EJointMotionType::Locked)
 			&& (LinearMotion[2] == EJointMotionType::Locked);
-		bool bConeProject = 
+		const bool bSphereProject =
+			(LinearProjection > 0.0f)
+			&& !bLinearSoft
+			&& (LinearMotion[0] == EJointMotionType::Limited)
+			&& (LinearMotion[1] == EJointMotionType::Limited)
+			&& (LinearMotion[2] == EJointMotionType::Limited);
+		bool bConeProject =
 			SolverSettings.bEnableSwingLimits 
 			&& (AngularProjection > 0.0f) 
 			&& !bSwingSoft 
@@ -260,21 +267,31 @@ namespace Chaos
 			&& (AngularProjection > 0.0f)
 			&& (Swing1Motion == EJointMotionType::Locked)
 			&& (Swing2Motion == EJointMotionType::Locked);
+		bool bTwistProject =
+			SolverSettings.bEnableTwistLimits
+			&& (AngularProjection > 0.0f)
+			&& !bTwistSoft
+			&& (TwistMotion != EJointMotionType::Free);
 
-		if (bPointProject && bConeProject)
+		if (bPointProject && bTwoSwingLockProject)
+		{
+			// Fixed point constraint with two locked swing axes and possible twist limit
+			NumActive += ApplyPointDoubleSwingLockProjection(Dt, SolverSettings, JointSettings, bPointProject, bTwoSwingLockProject, bTwistProject);
+		}
+		else if (bSphereProject && bTwoSwingLockProject)
 		{
-			// Fixed point constraint with cone limits
-			NumActive += ApplyPointConeProjection(Dt, SolverSettings, JointSettings);
+			// @todo(ccaulfield): support swing/twist limits with spherical limited position projection
+			NumActive += ApplySphereConeTwistProjection(Dt, SolverSettings, JointSettings, bSphereProject, false, false);
 		}
-		else if (bPointProject && bTwoSwingLockProject)
+		else if (bSphereProject)
 		{
-			// Fixed point constraint with two locked swing axes
-			NumActive += ApplyPointDoubleSwingLockProjection(Dt, SolverSettings, JointSettings);
+			// Limited spherical position constraint with possible cone and twist angular limits
+			NumActive += ApplySphereConeTwistProjection(Dt, SolverSettings, JointSettings, bSphereProject, bConeProject, bTwistProject);
 		}
-		else if (bPointProject)
+		else if (bPointProject || bConeProject || bTwistProject)
 		{
-			// Fixed point constraint with no angular limits
-			NumActive += ApplyPointProjection(Dt, SolverSettings, JointSettings);
+			// Fixed point constraint with possible cone and twist angular limits
+			NumActive += ApplyPointConeTwistProjection(Dt, SolverSettings, JointSettings, bPointProject, bConeProject, bTwistProject);
 		}
 
 		return NumActive;
@@ -292,8 +309,8 @@ namespace Chaos
 		EJointMotionType TwistMotion = JointSettings.AngularMotionTypes[(int32)EJointAngularConstraintIndex::Twist];
 		EJointMotionType Swing1Motion = JointSettings.AngularMotionTypes[(int32)EJointAngularConstraintIndex::Swing1];
 		EJointMotionType Swing2Motion = JointSettings.AngularMotionTypes[(int32)EJointAngularConstraintIndex::Swing2];
-		bool bTwistSoft = JointSettings.bSoftTwistLimitsEnabled;
-		bool bSwingSoft = JointSettings.bSoftSwingLimitsEnabled;
+		bool bTwistSoft = FPBDJointUtilities::GetSoftTwistLimitEnabled(SolverSettings, JointSettings);
+		bool bSwingSoft = FPBDJointUtilities::GetSoftSwingLimitEnabled(SolverSettings, JointSettings);
 
 		// If the twist axes are opposing, we cannot decompose the orientation into swing and twist angles, so just give up
 		const FVec3 Twist0 = Rs[0] * FJointConstants::TwistAxis();
@@ -1526,7 +1543,7 @@ namespace Chaos
 		const FReal Error = FMath::Max((FReal)0, Delta - JointSettings.LinearLimit);
 		if (FMath::Abs(Error) > PositionTolerance)
 		{
-			if (!JointSettings.bSoftLinearLimitsEnabled)
+			if (!FPBDJointUtilities::GetSoftLinearLimitEnabled(SolverSettings, JointSettings))
 			{
 				const FReal Stiffness = FPBDJointUtilities::GetLinearStiffness(SolverSettings, JointSettings);
 				ApplyPositionConstraint(Stiffness, Axis, Error);
@@ -1581,7 +1598,7 @@ namespace Chaos
 		
 		int32 NumActive = 0;
 
-		if ((AxialMotion == EJointMotionType::Limited) && JointSettings.bSoftLinearLimitsEnabled)
+		if ((AxialMotion == EJointMotionType::Limited) && FPBDJointUtilities::GetSoftLinearLimitEnabled(SolverSettings, JointSettings))
 		{
 			// Soft Axial constraint
 			const FReal AxialLimit = JointSettings.LinearLimit;
@@ -1608,7 +1625,7 @@ namespace Chaos
 			}
 		}
 
-		if ((RadialMotion == EJointMotionType::Limited) && JointSettings.bSoftLinearLimitsEnabled)
+		if ((RadialMotion == EJointMotionType::Limited) && FPBDJointUtilities::GetSoftLinearLimitEnabled(SolverSettings, JointSettings))
 		{
 			// Soft Radial constraint
 			const FReal RadialLimit = JointSettings.LinearLimit;
@@ -1677,7 +1694,7 @@ namespace Chaos
 		if (FMath::Abs(Delta) > Limit + PositionTolerance)
 		{
 			const FReal Error = (Delta > 0) ? Delta - Limit : Delta + Limit;
-			if ((AxialMotion == EJointMotionType::Limited) && JointSettings.bSoftLinearLimitsEnabled)
+			if ((AxialMotion == EJointMotionType::Limited) && FPBDJointUtilities::GetSoftLinearLimitEnabled(SolverSettings, JointSettings))
 			{
 				const FReal Stiffness = FPBDJointUtilities::GetSoftLinearStiffness(SolverSettings, JointSettings);
 				const FReal Damping = FPBDJointUtilities::GetSoftLinearDamping(SolverSettings, JointSettings);
@@ -1718,274 +1735,177 @@ namespace Chaos
 	}
 
 
-	// Projection for a fixed point constraint and no rotation constraint
-	int32 FJointSolverGaussSeidel::ApplyPointProjection(
+	int32 FJointSolverGaussSeidel::ApplyPointConeTwistProjection(
 		const FReal Dt,
 		const FPBDJointSolverSettings& SolverSettings,
-		const FPBDJointSettings& JointSettings)
-	{
-		FReal LinearProjection = FPBDJointUtilities::GetLinearProjection(SolverSettings, JointSettings);
-		const FReal ProjectionPositionTolerance = 0.0f;//PositionTolerance;
-
-		const FVec3 CX = Xs[1] - Xs[0];
-		bool bApplyProjection = true;//(CX.Size() > KINDA_SMALL_NUMBER); // + ProjectionPositionTolerance);
-		if (bApplyProjection)
-		{
-			UE_LOG(LogChaosJoint, Verbose, TEXT("Projection error %f"), CX.Size());
-
-			if (bChaos_Joint_ISPC_Enabled)
-			{
-#if INTEL_ISPC
-				ispc::ApplyPointPositionProjection((ispc::FJointSolverGaussSeidel*)this, (ispc::FVector&)CX, LinearProjection);
-#endif
-			}
-			else
-			{
-				// Apply a correction for the error in the joint connectors (apply position and rotation)
-				{
-					FMatrix33 J = Utilities::ComputeJointFactorMatrix(Xs[1] - Ps[1], InvIs[1], InvMs[1]);
-					const FMatrix33 IJ = J.Inverse();
-					const FVec3 DX = Utilities::Multiply(IJ, CX);
-					const FVec3 DP1 = -InvMs[1] * DX;
-					const FVec3 DR1 = Utilities::Multiply(InvIs[1], FVec3::CrossProduct(Xs[1] - Ps[1], -DX));
-					UE_LOG(LogChaosJoint, Verbose, TEXT("Projection pos/rot correction %f %f"), DP1.Size(), DR1.Size());
-					ApplyDelta(1, LinearProjection, DP1, DR1);
-				}
-
-				// Final position correction to fix errors from rotation in above
-				{
-					const FVec3 CX2 = Xs[1] - Xs[0];
-					ApplyPositionDelta(1, LinearProjection, -CX2);
-				}
-
-				// Apply a velocity correction to set the relative velocity at the joint to zero
-				{
-					const FReal InvM0 = VelProjectionInvMassScale * InvMs[0];
-					const FMatrix33 InvI0 = VelProjectionInvMassScale * InvIs[0];
-
-					// Calculate joint-space mass
-					FMatrix33 J = FMatrix33(0, 0, 0);
-					if (InvM0 > SMALL_NUMBER)
-					{
-						J = J + Utilities::ComputeJointFactorMatrix(Xs[0] - Ps[0], InvI0, InvM0);
-					}
-					if (InvMs[1] > SMALL_NUMBER)
-					{
-						J = J + Utilities::ComputeJointFactorMatrix(Xs[1] - Ps[1], InvIs[1], InvMs[1]);
-					}
-					const FMatrix33 IJ = J.Inverse();
-
-					// Velocity correction
-					const FVec3 V0 = Vs[0] + FVec3::CrossProduct(Ws[0], Xs[0] - Ps[0]);
-					const FVec3 V1 = Vs[1] + FVec3::CrossProduct(Ws[1], Xs[1] - Ps[1]);
-					const FVec3 DV = Utilities::Multiply(IJ, V1 - V0);
-
-					if (InvM0 > SMALL_NUMBER)
-					{
-						const FVec3 DV0 = InvM0 * DV;
-						const FVec3 DW0 = Utilities::Multiply(InvI0, FVec3::CrossProduct(Xs[0] - Ps[0], DV));
-						UE_LOG(LogChaosJoint, Verbose, TEXT("Projection vel/angvel correction 0 %f %f"), DV0.Size(), DW0.Size());
-						ApplyVelocityDelta(0, LinearProjection, DV0, DW0);
-					}
-					if (InvMs[1] > SMALL_NUMBER)
-					{
-						const FVec3 DV1 = -InvMs[1] * DV;
-						const FVec3 DW1 = -Utilities::Multiply(InvIs[1], FVec3::CrossProduct(Xs[1] - Ps[1], DV));
-						UE_LOG(LogChaosJoint, Verbose, TEXT("Projection vel/angvel correction 1 %f %f"), DV1.Size(), DW1.Size());
-						ApplyVelocityDelta(1, LinearProjection, DV1, DW1);
-					}
-				}
-			}
-			return 1;
-		}
-		return 0;
-	}
-
-	// Projection for fixed point constraint and cone limited angular constraint
-	int32 FJointSolverGaussSeidel::ApplyPointConeProjection(
-		const FReal Dt,
-		const FPBDJointSolverSettings& SolverSettings,
-		const FPBDJointSettings& JointSettings)
+		const FPBDJointSettings& JointSettings,
+		bool bApplyPositionCorrection,
+		bool bApplySwingCorrection,
+		bool bApplyTwistCorrection)
 	{
 		FReal LinearProjection = FPBDJointUtilities::GetLinearProjection(SolverSettings, JointSettings);
 		const FReal AngularProjection = FPBDJointUtilities::GetAngularProjection(SolverSettings, JointSettings);
 
 		int32 NumApplied = 0;
-		bool bConeActive = false;
 
 		const FReal ProjectionPositionTolerance = 0.0f;//PositionTolerance;
 		const FReal ProjectionAngleTolerance = 0.0f;//AngleTolerance;
 
+		FVec3 NetDP1 = FVec3(0);
+		FVec3 NetDR1 = FVec3(0);
+
 		// Correct the position error
-		const FVec3 CX = Xs[1] - Xs[0];
-		if (CX.Size() > ProjectionPositionTolerance)
+		if (bApplyPositionCorrection)
 		{
-			FMatrix33 J = Utilities::ComputeJointFactorMatrix(Xs[1] - Ps[1], InvIs[1], InvMs[1]);
-			const FMatrix33 IJ = J.Inverse();
-			const FVec3 DX = Utilities::Multiply(IJ, CX);
+			const FVec3 CX = Xs[1] - Xs[0];
+			if (CX.Size() > ProjectionPositionTolerance)
+			{
+				FMatrix33 J = Utilities::ComputeJointFactorMatrix(Xs[1] - Ps[1], InvIs[1], InvMs[1]);
+				const FMatrix33 IJ = J.Inverse();
+				const FVec3 DX = Utilities::Multiply(IJ, CX);
 
-			const FVec3 DP1 = -InvMs[1] * DX;
-			const FVec3 DR1 = Utilities::Multiply(InvIs[1], FVec3::CrossProduct(Xs[1] - Ps[1], -DX));
-			ApplyDelta(1, LinearProjection, DP1, DR1);
+				const FVec3 DP1 = -InvMs[1] * DX;
+				const FVec3 DR1 = Utilities::Multiply(InvIs[1], FVec3::CrossProduct(Xs[1] - Ps[1], -DX));
+				ApplyDelta(1, LinearProjection, DP1, DR1);
 
-			++NumApplied;
+				NetDP1 += DP1;
+				NetDR1 += DR1;
+				++NumApplied;
+			}
 		}
 
-		// Correct the cone error while maintaining position constraint by rotating about the swing axis through the joint position
-		const FReal Swing1Limit = JointSettings.AngularLimits[(int32)EJointAngularConstraintIndex::Swing1];
-		const FReal Swing2Limit = JointSettings.AngularLimits[(int32)EJointAngularConstraintIndex::Swing2];
-		FVec3 SwingAxis = FVec3(0);
+		// Correct the swing error while maintaining position constraint by rotating about the swing axis through the joint position
+		if (bApplySwingCorrection)
 		{
-			FVec3 SwingAxisLocal;
-			FReal DSwingAngle = 0.0f;
-			FPBDJointUtilities::GetEllipticalConeAxisErrorLocal(Rs[0], Rs[1], Swing2Limit, Swing1Limit, SwingAxisLocal, DSwingAngle);
-			SwingAxis = Rs[0] * SwingAxisLocal;
-			if (DSwingAngle > ProjectionAngleTolerance)
 			{
-				const FVec3 DR1 = -DSwingAngle * SwingAxis;
+				FVec3 SwingAxisLocal;
+				FReal DSwingAngle = 0.0f;
+				const FReal Swing1Limit = JointSettings.AngularLimits[(int32)EJointAngularConstraintIndex::Swing1];
+				const FReal Swing2Limit = JointSettings.AngularLimits[(int32)EJointAngularConstraintIndex::Swing2];
+				FPBDJointUtilities::GetEllipticalConeAxisErrorLocal(Rs[0], Rs[1], Swing2Limit, Swing1Limit, SwingAxisLocal, DSwingAngle);
+				FVec3 SwingAxis = Rs[0] * SwingAxisLocal;
+				if (DSwingAngle > ProjectionAngleTolerance)
+				{
+					const FVec3 DR1 = -DSwingAngle * SwingAxis;
+					const FVec3 DP1 = -FVec3::CrossProduct(DR1, Xs[1] - Ps[1]);
+					ApplyDelta(1, AngularProjection, DP1, DR1);
+
+					NetDP1 += DP1;
+					NetDR1 += DR1;
+					++NumApplied;
+				}
+			}
+		}
+
+		// Correct the Twist angle
+		if (bApplyTwistCorrection)
+		{
+			FVec3 TwistAxis;
+			FReal TwistAngle;
+			FPBDJointUtilities::GetTwistAxisAngle(Rs[0], Rs[1], TwistAxis, TwistAngle);
+			FReal DTwistAngle = 0;
+			const FReal TwistLimit = JointSettings.AngularLimits[(int32)EJointAngularConstraintIndex::Twist];
+			if (TwistAngle > TwistLimit)
+			{
+				DTwistAngle = TwistAngle - TwistLimit;
+			}
+			else if (TwistAngle < -TwistLimit)
+			{
+				DTwistAngle = TwistAngle + TwistLimit;
+			}
+
+			if (FMath::Abs(DTwistAngle) > ProjectionAngleTolerance)
+			{
+				const FVec3 DR1 = -DTwistAngle * TwistAxis;
 				const FVec3 DP1 = -FVec3::CrossProduct(DR1, Xs[1] - Ps[1]);
 				ApplyDelta(1, AngularProjection, DP1, DR1);
 
+				NetDP1 += DP1;
+				NetDR1 += DR1;
 				++NumApplied;
-				bConeActive = true;
 			}
 		}
 
 		// Final position correction. Following this, both the position constraint and swing
 		// constraint should be exactly resolved.
+		if (bApplyPositionCorrection)
 		{
-			const FVec3 CX2 = Xs[1] - Xs[0];
-			ApplyPositionDelta(1, LinearProjection, -CX2);
+			const FVec3 DP1 = -(Xs[1] - Xs[0]);
+			ApplyPositionDelta(1, LinearProjection, DP1);
+
+			NetDP1 += DP1;
 		}
 
-		// Velocity correction
+		if (Chaos_Joint_VelProjectionAlpha > 0.0f)
 		{
-			const FReal InvM0 = VelProjectionInvMassScale * InvMs[0];
-			const FMatrix33 InvI0 = VelProjectionInvMassScale * InvIs[0];
+			const FVec3 DV1 = NetDP1 / Dt;
+			const FVec3 DW1 = NetDR1 / Dt;
+			ApplyVelocityDelta(1, Chaos_Joint_VelProjectionAlpha * LinearProjection, DV1, DW1);
+		}
+
+		return NumApplied;
+	}
+
+
+	int32 FJointSolverGaussSeidel::ApplySphereConeTwistProjection(
+		const FReal Dt,
+		const FPBDJointSolverSettings& SolverSettings,
+		const FPBDJointSettings& JointSettings,
+		bool bApplyPositionCorrection,
+		bool bApplySwingCorrection,
+		bool bApplyTwistCorrection)
+	{
+		FReal LinearProjection = FPBDJointUtilities::GetLinearProjection(SolverSettings, JointSettings);
+		const FReal ProjectionPositionTolerance = 0.0f;//PositionTolerance;
 
-			// Recalculate swing axis. THis should not be necessary if we did not rotate the parent body.
-			// @todo(ccaulfield): see if we really need this.
-			//if (InvM0 > SMALL_NUMBER)
+		FVec3 NetDP1 = FVec3(0);
+		FVec3 NetDR1 = FVec3(0);
+
+		// Position correction
+		if (bApplyPositionCorrection)
+		{
+			FVec3 Axis;
+			FReal Delta;
+			FPBDJointUtilities::GetSphericalAxisDelta(Xs[0], Xs[1], Axis, Delta);
+			const FReal Error = FMath::Max((FReal)0, Delta - JointSettings.LinearLimit);
+			if (FMath::Abs(Error) > ProjectionPositionTolerance)
 			{
-				FVec3 SwingAxisLocal;
-				FReal DSwingAngle = 0.0f;
-				FPBDJointUtilities::GetEllipticalConeAxisErrorLocal(Rs[0], Rs[1], Swing2Limit, Swing1Limit, SwingAxisLocal, DSwingAngle);
-				SwingAxis = Rs[0] * SwingAxisLocal;
-			}
-
-			// We may iterate here if the swing constraint is active but the angular velocity does not required
-			// fixing (because it is not moving against the constraint), if when fixing the position constraint
-			// we introduce an angular velocity that would violate the swing constraint.
-			for (int32 ProjIts = 0; ProjIts < 2; ++ProjIts)
-			{
-				// Calculate velocity and angular velocity error at joint
-				const FVec3 V0 = Vs[0] + FVec3::CrossProduct(Ws[0], Xs[0] - Ps[0]);
-				const FVec3 V1 = Vs[1] + FVec3::CrossProduct(Ws[1], Xs[1] - Ps[1]);
-				const FReal W0 = FVec3::DotProduct(Ws[0], SwingAxis);
-				const FReal W1 = FVec3::DotProduct(Ws[1], SwingAxis);
-				const FVec3 CV = V1 - V0;
-				const FReal CW = FMath::Max(0.0f, W1 - W0);
-
-				// Calculate the joint-space velocity and angular velocity correction
-				// If the swing constraint is active, and we need to fix angular velocity about the axis,
-				// we have 4 simultaneous equations to solve (velocity at joint, and angular velocity about axis).
-				// If the swing constraint is inactive, or the angular velocity about the axis is negative,
-				// we have only 3 equations to solve (velocity at joint).
-				FVec3 DV = FVec3(0);
-				FReal DW = 0;
-				if (bConeActive && (CW > 0.0f))
-				{
-					// Calculate joint-space mass from the Jacobian
-					FMatrix33 J00 = FMatrix33(0, 0, 0);
-					FVec3 J01 = FVec3(0);
-					FReal J11 = 0;
-					if (InvM0 > SMALL_NUMBER)
-					{
-						J00 += Utilities::ComputeJointFactorMatrix(Xs[0] - Ps[0], InvI0, InvM0);
-						J01 += FVec3::CrossProduct(Xs[0] - Ps[0], Utilities::Multiply(InvI0, SwingAxis));
-						J11 += FVec3::DotProduct(SwingAxis, Utilities::Multiply(InvI0, SwingAxis));
-					}
-					if (InvMs[1] > SMALL_NUMBER)
-					{
-						J00 += Utilities::ComputeJointFactorMatrix(Xs[1] - Ps[1], InvIs[1], InvMs[1]);
-						J01 += FVec3::CrossProduct(Xs[1] - Ps[1], Utilities::Multiply(InvIs[1], SwingAxis));
-						J11 += FVec3::DotProduct(SwingAxis, Utilities::Multiply(InvIs[1], SwingAxis));
-					}
-					PMatrix<FReal, 4, 4>& J = reinterpret_cast<PMatrix<FReal, 4, 4>&>(J00);
-					J.M[3][0] = -J01.X;
-					J.M[3][1] = -J01.Y;
-					J.M[3][2] = -J01.Z;
-					J.M[0][3] = -J01.X;
-					J.M[1][3] = -J01.Y;
-					J.M[2][3] = -J01.Z;
-					J.M[3][3] = J11;
-					PMatrix<FReal, 4, 4> IJ = J.Inverse();
-
-					const TVector<FReal, 4> C = TVector<FReal, 4>(CV.X, CV.Y, CV.Z, CW);
-					const TVector<FReal, 4> D = IJ.TransformFVector4(C);
-					DV = FVec3(D.X, D.Y, D.Z);
-					DW = D.W;
-				}
-				else
-				{
-					// Calculate joint-space mass from the Jacobian
-					FMatrix33 J00 = FMatrix33(0, 0, 0);
-					if (InvM0 > SMALL_NUMBER)
-					{
-						J00 += Utilities::ComputeJointFactorMatrix(Xs[0] - Ps[0], InvI0, InvM0);
-					}
-					if (InvMs[1] > SMALL_NUMBER)
-					{
-						J00 += Utilities::ComputeJointFactorMatrix(Xs[1] - Ps[1], InvIs[1], InvMs[1]);
-					}
-					FMatrix33 IJ = J00.Inverse();
-
-					DV = Utilities::Multiply(IJ, CV);
-				}
+				const FVec3 AngularAxis1 = FVec3::CrossProduct(Xs[1] - Ps[1], Axis);
+				const FVec3 IA1 = Utilities::Multiply(InvIs[1], AngularAxis1);
+				const FReal II1 = FVec3::DotProduct(AngularAxis1, IA1);
+				const FReal IM = InvMs[1] + II1;
+				const FVec3 DX = Axis * Delta / IM;
 
-				// Apply the world-space velocity and angular velocity correction
-				FVec3 DV0 = InvM0 * DV;
-				FVec3 DW0 = Utilities::Multiply(InvI0, FVec3::CrossProduct(Xs[0] - Ps[0], DV)) + Utilities::Multiply(InvI0, SwingAxis) * DW;
-				if (InvM0 > SMALL_NUMBER)
-				{
-					ApplyVelocityDelta(0, 1.0f, DV0, DW0);
-				}
-				FVec3 DV1 = -InvMs[1] * DV;
-				FVec3 DWV1 = Utilities::Multiply(InvIs[1], FVec3::CrossProduct(Xs[1] - Ps[1], DV));
-				FVec3 DWW1 = Utilities::Multiply(InvIs[1], SwingAxis) * DW;
-				FVec3 DW1 = -(DWV1 + DWW1);
-				if (InvMs[1] > SMALL_NUMBER)
-				{
-					ApplyVelocityDelta(1, 1.0f, DV1, DW1);
-				}
+				const FVec3 DP1 = -InvMs[1] * DX;
+				const FVec3 DR1 = Utilities::Multiply(InvIs[1], FVec3::CrossProduct(Xs[1] - Ps[1], -DX));
+				ApplyDelta(1, LinearProjection, DP1, DR1);
 
-				const FVec3 V02 = Vs[0] + FVec3::CrossProduct(Ws[0], Xs[0] - Ps[0]);
-				const FVec3 V12 = Vs[1] + FVec3::CrossProduct(Ws[1], Xs[1] - Ps[1]);
-				const FReal W02 = FVec3::DotProduct(Ws[0], SwingAxis);
-				const FReal W12 = FVec3::DotProduct(Ws[1], SwingAxis);
-				const FVec3 CV2 = V12 - V02;
-				const FReal CW2 = FMath::Max(0.0f, W12 - W02);
-				const FVec3 DWOffAxis0 = DW0 - FVec3::DotProduct(DW0, SwingAxis) * SwingAxis;
-				const FVec3 DWOffAxis1 = DW1 - FVec3::DotProduct(DW1, SwingAxis) * SwingAxis;
-				const FVec3 DWOffAxisW1 = DWW1 - FVec3::DotProduct(DWW1, SwingAxis) * SwingAxis;
-				const FVec3 DWOffAxisV1 = DWV1 - FVec3::DotProduct(DWV1, SwingAxis) * SwingAxis;
-
-				// If we resolved the angvel constraint or it was ignored and not reactivated (by the vel constraint), we are done
-				if ((bConeActive && (CW > 0.0f)) || (CW2 <= 0.0f))
-				{
-					break;
-				}
+				NetDP1 += DP1;
+				NetDR1 += DR1;
 			}
 		}
 
-		return NumApplied;
+		// @todo(ccaulfiedl): apply swing and twist correction, taking into account whether position limit is active
+
+		// Velocity correction
+		if (Chaos_Joint_VelProjectionAlpha > 0.0f)
+		{
+			const FVec3 DV1 = NetDP1 / Dt;
+			const FVec3 DW1 = NetDR1 / Dt;
+			ApplyVelocityDelta(1, Chaos_Joint_VelProjectionAlpha * LinearProjection, DV1, DW1);
+		}
+
+		return 1;
 	}
 
-	// Projection for fixed point constraint and two locked swing axes
+
 	int32 FJointSolverGaussSeidel::ApplyPointDoubleSwingLockProjection(
 		const FReal Dt,
 		const FPBDJointSolverSettings& SolverSettings,
-		const FPBDJointSettings& JointSettings)
+		const FPBDJointSettings& JointSettings,
+		bool bApplyPositionCorrection,
+		bool bApplySwingCorrection,
+		bool bApplyTwistCorrection)
 	{
 		const FReal AngularProjection = FPBDJointUtilities::GetAngularProjection(SolverSettings, JointSettings);
 		const FReal ProjectionAngleTolerance = 0.0f;//AngleTolerance;
@@ -1993,9 +1913,16 @@ namespace Chaos
 		int32 NumApplied = 0;
 
 		// Correct position error
-		NumApplied += ApplyPointProjection(Dt, SolverSettings, JointSettings);
+		if (bApplyPositionCorrection)
+		{
+			NumApplied += ApplyPointConeTwistProjection(Dt, SolverSettings, JointSettings, bApplyPositionCorrection, false, false);
+		}
+
+		FVec3 NetDP1 = FVec3(0);
+		FVec3 NetDR1 = FVec3(0);
 
 		// Correct the swing error while maintaining position constraint by rotating about the swing axis through the joint position
+		if (bApplySwingCorrection)
 		{
 			FVec3 SwingAxisLocal;
 			FReal DSwingAngle = 0.0f;
@@ -2007,19 +1934,46 @@ namespace Chaos
 				const FVec3 DP1 = -FVec3::CrossProduct(DR1, Xs[1] - Ps[1]);
 				ApplyDelta(1, AngularProjection, DP1, DR1);
 
+				NetDP1 += DP1;
+				NetDR1 += DR1;
 				++NumApplied;
 			}
 		}
 
-		// Correct angular velocity error
+		if (bApplyTwistCorrection)
 		{
-			FVec3 TwistAxis = Rs[1] * FJointConstants::TwistAxis();
+			FVec3 TwistAxis;
+			FReal TwistAngle;
+			FPBDJointUtilities::GetTwistAxisAngle(Rs[0], Rs[1], TwistAxis, TwistAngle);
+			FReal DTwistAngle = 0;
+			const FReal TwistLimit = JointSettings.AngularLimits[(int32)EJointAngularConstraintIndex::Twist];
+			if (TwistAngle > TwistLimit)
+			{
+				DTwistAngle = TwistAngle - TwistLimit;
+			}
+			else if (TwistAngle < -TwistLimit)
+			{
+				DTwistAngle = TwistAngle + TwistLimit;
+			}
 
-			// Fix the angular velocity to remove all components not about twist axis
-			const FVec3 W0 = Ws[0] - FVec3::DotProduct(Ws[0], TwistAxis) * TwistAxis;
-			const FVec3 W1 = Ws[1] - FVec3::DotProduct(Ws[1], TwistAxis) * TwistAxis;
-			const FVec3 DW = W1 - W0;
-			ApplyVelocityDelta(1, 1.0f, FVec3(0), -DW);
+			if (FMath::Abs(DTwistAngle) > ProjectionAngleTolerance)
+			{
+				const FVec3 DR1 = -DTwistAngle * TwistAxis;
+				const FVec3 DP1 = -FVec3::CrossProduct(DR1, Xs[1] - Ps[1]);
+				ApplyDelta(1, AngularProjection, DP1, DR1);
+
+				NetDP1 += DP1;
+				NetDR1 += DR1;
+				++NumApplied;
+			}
+		}
+
+		// Correct angular velocity error
+		if (Chaos_Joint_VelProjectionAlpha > 0.0f)
+		{
+			const FVec3 DV1 = NetDP1 / Dt;
+			const FVec3 DW1 = NetDR1 / Dt;
+			ApplyVelocityDelta(1, Chaos_Joint_VelProjectionAlpha * AngularProjection, DV1, DW1);
 		}
 
 		return NumApplied;
diff --git a/Engine/Source/Runtime/Experimental/Chaos/Private/Chaos/PBDJointConstraintUtilities.cpp b/Engine/Source/Runtime/Experimental/Chaos/Private/Chaos/PBDJointConstraintUtilities.cpp
--- a/Engine/Source/Runtime/Experimental/Chaos/Private/Chaos/PBDJointConstraintUtilities.cpp
+++ b/Engine/Source/Runtime/Experimental/Chaos/Private/Chaos/PBDJointConstraintUtilities.cpp
@@ -1,6 +1,7 @@
 // Copyright Epic Games, Inc. All Rights Reserved.
 #include "Chaos/PBDJointConstraintUtilities.h"
 #include "Chaos/DenseMatrix.h"
+#include "Chaos/Joint/JointConstraintsCVars.h"
 #include "Chaos/ParticleHandle.h"
 #include "Chaos/Utilities.h"
 
@@ -391,8 +392,27 @@ namespace Chaos
 		return SwingAngleMax;
 	}
 
-	
-	// @todo(ccaulfield): separate linear soft and stiff
+	bool FPBDJointUtilities::GetSoftLinearLimitEnabled(
+		const FPBDJointSolverSettings& SolverSettings,
+		const FPBDJointSettings& JointSettings)
+	{
+		return JointSettings.bSoftLinearLimitsEnabled && !bChaos_Joint_DisableSoftLimits;
+	}
+
+	bool FPBDJointUtilities::GetSoftTwistLimitEnabled(
+		const FPBDJointSolverSettings& SolverSettings,
+		const FPBDJointSettings& JointSettings)
+	{
+		return JointSettings.bSoftTwistLimitsEnabled && !bChaos_Joint_DisableSoftLimits;
+	}
+
+	bool FPBDJointUtilities::GetSoftSwingLimitEnabled(
+		const FPBDJointSolverSettings& SolverSettings,
+		const FPBDJointSettings& JointSettings)
+	{
+		return JointSettings.bSoftSwingLimitsEnabled && !bChaos_Joint_DisableSoftLimits;
+	}
+
 	FReal FPBDJointUtilities::GetLinearStiffness(
 		const FPBDJointSolverSettings& SolverSettings,
 		const FPBDJointSettings& JointSettings)
diff --git a/Engine/Source/Runtime/Experimental/Chaos/Public/Chaos/Joint/PBDJointSolverGaussSeidel.h b/Engine/Source/Runtime/Experimental/Chaos/Public/Chaos/Joint/PBDJointSolverGaussSeidel.h
--- a/Engine/Source/Runtime/Experimental/Chaos/Public/Chaos/Joint/PBDJointSolverGaussSeidel.h
+++ b/Engine/Source/Runtime/Experimental/Chaos/Public/Chaos/Joint/PBDJointSolverGaussSeidel.h
@@ -361,20 +361,29 @@ namespace Chaos
 			const FPBDJointSolverSettings& SolverSettings,
 			const FPBDJointSettings& JointSettings);
 
-		int32 ApplyPointProjection(
+		int32 ApplyPointConeTwistProjection(
 			const FReal Dt,
 			const FPBDJointSolverSettings& SolverSettings,
-			const FPBDJointSettings& JointSettings);
+			const FPBDJointSettings& JointSettings,
+			bool bApplyPositionCorrection,
+			bool bApplySwingCorrection,
+			bool bApplyTwistCorrection);
 
-		int32 ApplyPointConeProjection(
+		int32 ApplySphereConeTwistProjection(
 			const FReal Dt,
 			const FPBDJointSolverSettings& SolverSettings,
-			const FPBDJointSettings& JointSettings);
+			const FPBDJointSettings& JointSettings,
+			bool bApplyPositionCorrection,
+			bool bApplySwingCorrection,
+			bool bApplyTwistCorrection);
 
 		int32 ApplyPointDoubleSwingLockProjection(
 			const FReal Dt,
 			const FPBDJointSolverSettings& SolverSettings,
-			const FPBDJointSettings& JointSettings);
+			const FPBDJointSettings& JointSettings,
+			bool bApplyPositionCorrection,
+			bool bApplySwingCorrection,
+			bool bApplyTwistCorrection);
 
 
 
diff --git a/Engine/Source/Runtime/Experimental/Chaos/Public/Chaos/PBDJointConstraintUtilities.h b/Engine/Source/Runtime/Experimental/Chaos/Public/Chaos/PBDJointConstraintUtilities.h
--- a/Engine/Source/Runtime/Experimental/Chaos/Public/Chaos/PBDJointConstraintUtilities.h
+++ b/Engine/Source/Runtime/Experimental/Chaos/Public/Chaos/PBDJointConstraintUtilities.h
@@ -152,6 +152,18 @@ namespace Chaos
 			const FReal MinParentMassRatio,
 			const FReal MaxInertiaRatio);
 
+		static bool GetSoftLinearLimitEnabled(
+			const FPBDJointSolverSettings& SolverSettings,
+			const FPBDJointSettings& JointSettings);
+
+		static bool GetSoftTwistLimitEnabled(
+			const FPBDJointSolverSettings& SolverSettings,
+			const FPBDJointSettings& JointSettings);
+
+		static bool GetSoftSwingLimitEnabled(
+			const FPBDJointSolverSettings& SolverSettings,
+			const FPBDJointSettings& JointSettings);
+
 		static FReal GetLinearStiffness(
 			const FPBDJointSolverSettings& SolverSettings,
 			const FPBDJointSettings& JointSettings);
