Chaos Cloth - Fix self-collision normals.

diff --git a/Engine/Source/Runtime/Experimental/Chaos/Private/Chaos/PBDCollisionSpringConstraintsBase.cpp b/Engine/Source/Runtime/Experimental/Chaos/Private/Chaos/PBDCollisionSpringConstraintsBase.cpp
--- a/Engine/Source/Runtime/Experimental/Chaos/Private/Chaos/PBDCollisionSpringConstraintsBase.cpp
+++ b/Engine/Source/Runtime/Experimental/Chaos/Private/Chaos/PBDCollisionSpringConstraintsBase.cpp
@@ -106,7 +106,9 @@ PBDCollisionSpringConstraintsBase<T, d>::PBDCollisionSpringConstraintsBase(
 			// TODO(mlentine): In theory this shouldn't happen as this means no collision but I'm not positive that the collision check is always correct.
 			//if (Bary.Y < 0 || Bary.Z < 0 || Bary.X < 0) return;
 			// TODO(mlentine): Incorporate history
-			TVector<T, d> Normal = TVector<T, d>::DotProduct(Result.Normal, PP0) > 0 ? Result.Normal : -Result.Normal;
+			TVector<T, d> Normal = TVector<T, d>::CrossProduct(P10, P20);
+			Normal.Normalize();
+			Normal = TVector<T, d>::DotProduct(Normal, PP0) > 0 ? Normal : -Normal;
 			CriticalSection.Lock();
 			MConstraints.Add({Index, Elem[0], Elem[1], Elem[2]});
 			MBarys.Add(Bary);
diff --git a/Engine/Source/Runtime/Experimental/Chaos/Public/Chaos/PBDCollisionSpringConstraints.h b/Engine/Source/Runtime/Experimental/Chaos/Public/Chaos/PBDCollisionSpringConstraints.h
--- a/Engine/Source/Runtime/Experimental/Chaos/Public/Chaos/PBDCollisionSpringConstraints.h
+++ b/Engine/Source/Runtime/Experimental/Chaos/Public/Chaos/PBDCollisionSpringConstraints.h
@@ -7,8 +7,6 @@
 #include "Chaos/PBDParticles.h"
 #include "Chaos/PBDConstraintContainer.h"
 
-#include <algorithm>
-
 namespace Chaos
 {
 template<class T, int d>
@@ -19,7 +17,7 @@ class TPBDCollisionSpringConstraints : public PBDCollisionSpringConstraintsBase<
 	using Base::MConstraints;
 
   public:
-	TPBDCollisionSpringConstraints(const TDynamicParticles<T, d>& InParticles, const TArray<TVector<int32, 3>>& Elements, const TSet<TVector<int32, 2>>& DisabledCollisionElements, const TArray<uint32>& DynamicGroupIds, const TArray<T>& PerGroupThicknesses, const T Dt, const T Stiffness = (T)1.5)
+	TPBDCollisionSpringConstraints(const TDynamicParticles<T, d>& InParticles, const TArray<TVector<int32, 3>>& Elements, const TSet<TVector<int32, 2>>& DisabledCollisionElements, const TArray<uint32>& DynamicGroupIds, const TArray<T>& PerGroupThicknesses, const T Dt, const T Stiffness = (T)1.0)
 	    : Base(InParticles, Elements, DisabledCollisionElements, DynamicGroupIds, PerGroupThicknesses, Dt, Stiffness) {}
 	virtual ~TPBDCollisionSpringConstraints() {}
 
@@ -36,19 +34,19 @@ class TPBDCollisionSpringConstraints : public PBDCollisionSpringConstraintsBase<
 		T Multiplier = 1;
 		if (InParticles.InvM(i1) > 0)
 		{
-			InParticles.P(i1) -= Multiplier * InParticles.InvM(i1) * Delta;
+			InParticles.P(i1) += Multiplier * InParticles.InvM(i1) * Delta;
 		}
 		if (InParticles.InvM(i2))
 		{
-			InParticles.P(i2) += Multiplier * InParticles.InvM(i2) * MBarys[i][0] * Delta;
+			InParticles.P(i2) -= Multiplier * InParticles.InvM(i2) * MBarys[i][0] * Delta;
 		}
 		if (InParticles.InvM(i3))
 		{
-			InParticles.P(i3) += Multiplier * InParticles.InvM(i3) * MBarys[i][1] * Delta;
+			InParticles.P(i3) -= Multiplier * InParticles.InvM(i3) * MBarys[i][1] * Delta;
 		}
 		if (InParticles.InvM(i4))
 		{
-			InParticles.P(i4) += Multiplier * InParticles.InvM(i4) * MBarys[i][2] * Delta;
+			InParticles.P(i4) -= Multiplier * InParticles.InvM(i4) * MBarys[i][2] * Delta;
 		}
 	}
 
