Get GC unit tests to compile

Benn.Gallagher, Brice.Criswell

diff --git a/Engine/Source/Runtime/Experimental/ChaosSolvers/Public/PhysicsProxy/GeometryCollectionPhysicsProxy.h b/Engine/Source/Runtime/Experimental/ChaosSolvers/Public/PhysicsProxy/GeometryCollectionPhysicsProxy.h
--- a/Engine/Source/Runtime/Experimental/ChaosSolvers/Public/PhysicsProxy/GeometryCollectionPhysicsProxy.h
+++ b/Engine/Source/Runtime/Experimental/ChaosSolvers/Public/PhysicsProxy/GeometryCollectionPhysicsProxy.h
@@ -154,7 +154,7 @@ public:
 	/** Enqueue a field \p Command to be processed by \c ProcessCommands() or 
 	 * \c FieldForcesUpdateCallback(). 
 	 */
-	void BufferCommand(Chaos::FPhysicsSolver* , const FFieldSystemCommand& Command) 
+	void BufferCommand(Chaos::TPBDRigidsSolver<Traits>* , const FFieldSystemCommand& Command) 
 	{ Commands.Add(Command); }
 
 	static void InitializeSharedCollisionStructures(Chaos::FErrorReporter& ErrorReporter, FGeometryCollection& RestCollection, const FSharedSimulationParameters& SharedParams);
