Fix static analysis

diff --git a/Engine/Source/Runtime/Experimental/ChaosSolvers/Public/RewindData.h b/Engine/Source/Runtime/Experimental/ChaosSolvers/Public/RewindData.h
--- a/Engine/Source/Runtime/Experimental/ChaosSolvers/Public/RewindData.h
+++ b/Engine/Source/Runtime/Experimental/ChaosSolvers/Public/RewindData.h
@@ -1330,7 +1330,7 @@ private:
 		}
 
 		auto* GTUnsafeParticle = PTParticle.GTGeometryParticle();
-		ensure(GTUnsafeParticle != nullptr);
+		check(GTUnsafeParticle != nullptr);
 		const int32 DirtyIdx = AllDirtyParticles.Add(FDirtyParticleInfo(*GTUnsafeParticle, PTParticle,UniqueIdx,CurFrame,Managers.Capacity()));
 		ParticleToAllDirtyIdx.Add(UniqueIdx,DirtyIdx);
 
