Chaos - Joint - fixed parameter edge cases

diff --git a/Engine/Source/Runtime/Experimental/Chaos/Private/Chaos/PBDJointConstraints.cpp b/Engine/Source/Runtime/Experimental/Chaos/Private/Chaos/PBDJointConstraints.cpp
--- a/Engine/Source/Runtime/Experimental/Chaos/Private/Chaos/PBDJointConstraints.cpp
+++ b/Engine/Source/Runtime/Experimental/Chaos/Private/Chaos/PBDJointConstraints.cpp
@@ -170,6 +170,54 @@ namespace Chaos
 		{
 			AngularLimits[(int32)EJointAngularConstraintIndex::Swing2] = 0;
 		}
+
+		// If we have a zero degree limit angle, lock the joint, or set a non-zero limit (to avoid division by zero in axis calculations)
+		const FReal MinAngularLimit = 0.01f;
+		if ((AngularMotionTypes[(int32)EJointAngularConstraintIndex::Twist] == EJointMotionType::Limited) && (AngularLimits[(int32)EJointAngularConstraintIndex::Twist] == 0))
+		{
+			if (bSoftTwistLimitsEnabled)
+			{
+				AngularLimits[(int32)EJointAngularConstraintIndex::Twist] = MinAngularLimit;
+			}
+			else
+			{
+				AngularMotionTypes[(int32)EJointAngularConstraintIndex::Twist] = EJointMotionType::Locked;
+			}
+		}
+		if ((AngularMotionTypes[(int32)EJointAngularConstraintIndex::Swing1] == EJointMotionType::Limited) && (AngularLimits[(int32)EJointAngularConstraintIndex::Swing1] == 0))
+		{
+			if (bSoftSwingLimitsEnabled)
+			{
+				AngularLimits[(int32)EJointAngularConstraintIndex::Swing1] = MinAngularLimit;
+			}
+			else
+			{
+				AngularMotionTypes[(int32)EJointAngularConstraintIndex::Swing1] = EJointMotionType::Locked;
+			}
+		}
+		if ((AngularMotionTypes[(int32)EJointAngularConstraintIndex::Swing2] == EJointMotionType::Limited) && (AngularLimits[(int32)EJointAngularConstraintIndex::Swing2] == 0))
+		{
+			if (bSoftSwingLimitsEnabled)
+			{
+				AngularLimits[(int32)EJointAngularConstraintIndex::Swing2] = MinAngularLimit;
+			}
+			else
+			{
+				AngularMotionTypes[(int32)EJointAngularConstraintIndex::Swing2] = EJointMotionType::Locked;
+			}
+		}
+
+		// SLerp drive is only allowed if no angular dofs are locked
+		if (bAngularSLerpPositionDriveEnabled || bAngularSLerpVelocityDriveEnabled)
+		{
+			if ((AngularMotionTypes[(int32)EJointAngularConstraintIndex::Twist] == EJointMotionType::Locked)
+				|| (AngularMotionTypes[(int32)EJointAngularConstraintIndex::Twist] == EJointMotionType::Locked)
+				|| (AngularMotionTypes[(int32)EJointAngularConstraintIndex::Twist] == EJointMotionType::Locked))
+			{
+				bAngularSLerpPositionDriveEnabled = false;
+				bAngularSLerpVelocityDriveEnabled = false;
+			}
+		}
 	}
 
 	
