Added Renderer and NaniteBuilder dependencies to GeometryCollectionEngine for upcoming Nanite support.

diff --git a/Engine/Source/Runtime/Experimental/GeometryCollectionEngine/GeometryCollectionEngine.Build.cs b/Engine/Source/Runtime/Experimental/GeometryCollectionEngine/GeometryCollectionEngine.Build.cs
--- a/Engine/Source/Runtime/Experimental/GeometryCollectionEngine/GeometryCollectionEngine.Build.cs
+++ b/Engine/Source/Runtime/Experimental/GeometryCollectionEngine/GeometryCollectionEngine.Build.cs
@@ -17,6 +17,7 @@ namespace UnrealBuildTool.Rules
                     "Engine",
                     "RenderCore",
                     "RHI",
+					"Renderer",
                     "Chaos",
 					"ChaosSolvers",
                     "PhysX",
@@ -37,15 +38,18 @@ namespace UnrealBuildTool.Rules
 
 	        if (!Target.bBuildRequiresCookedData)
 			{
-	            DynamicallyLoadedModuleNames.AddRange(new string[] { "DerivedDataCache" });
-		    }
+				DynamicallyLoadedModuleNames.Add("DerivedDataCache");
+			}
 
 			PrivateIncludePathModuleNames.Add("DerivedDataCache");
 
-			if(Target.bBuildEditor)
+			if (Target.bBuildEditor)
             {
-                PrivateDependencyModuleNames.Add("UnrealEd");
-            }
+				DynamicallyLoadedModuleNames.Add("NaniteBuilder");
+				PrivateIncludePathModuleNames.Add("NaniteBuilder");
+
+				PrivateDependencyModuleNames.Add("UnrealEd");
+			}
 
 			PrivateDefinitions.Add("CHAOS_INCLUDE_LEVEL_1=1");
 		}
