Add InstanceUniformShaderParameters.h include to GeometryCollectionObject.h for FNaniteInfo

diff --git a/Engine/Source/Runtime/Experimental/GeometryCollectionEngine/Public/GeometryCollection/GeometryCollectionObject.h b/Engine/Source/Runtime/Experimental/GeometryCollectionEngine/Public/GeometryCollection/GeometryCollectionObject.h
--- a/Engine/Source/Runtime/Experimental/GeometryCollectionEngine/Public/GeometryCollection/GeometryCollectionObject.h
+++ b/Engine/Source/Runtime/Experimental/GeometryCollectionEngine/Public/GeometryCollection/GeometryCollectionObject.h
@@ -4,6 +4,7 @@
 
 #include "CoreMinimal.h"
 #include "Misc/Crc.h"
+#include "InstanceUniformShaderParameters.h"
 #include "GeometryCollection/ManagedArray.h"
 #include "GeometryCollection/GeometryCollectionSimulationTypes.h"
 #include "GeometryCollectionObject.generated.h"
