Added initial custom render scene proxy for geometry collections driven by Nanite.

diff --git a/Engine/Source/Runtime/Experimental/GeometryCollectionEngine/Private/GeometryCollection/GeometryCollectionComponent.cpp b/Engine/Source/Runtime/Experimental/GeometryCollectionEngine/Private/GeometryCollection/GeometryCollectionComponent.cpp
--- a/Engine/Source/Runtime/Experimental/GeometryCollectionEngine/Private/GeometryCollection/GeometryCollectionComponent.cpp
+++ b/Engine/Source/Runtime/Experimental/GeometryCollectionEngine/Private/GeometryCollection/GeometryCollectionComponent.cpp
@@ -4,6 +4,7 @@
 
 #include "Async/ParallelFor.h"
 #include "Components/BoxComponent.h"
+#include "ComponentRecreateRenderStateContext.h"
 #include "GeometryCollection/GeometryCollectionObject.h"
 #include "GeometryCollection/GeometryCollectionAlgo.h"
 #include "GeometryCollection/GeometryCollectionComponentPluginPrivate.h"
@@ -33,6 +34,8 @@
 #include "PhysicsEngine/BodyInstance.h"
 #include "Chaos/ChaosGameplayEventDispatcher.h"
 
+#include "Rendering/NaniteResources.h"
+
 #if !(UE_BUILD_SHIPPING || UE_BUILD_TEST)
 #include "Logging/MessageLog.h"
 #include "Misc/UObjectToken.h"
@@ -75,6 +78,20 @@ void HackRegisterGeomAccelerator(UGeometryCollectionComponent& Component)
 }
 #endif
 
+static void RecreateGlobalRenderState(IConsoleVariable* Var)
+{
+	FGlobalComponentRecreateRenderStateContext Context;
+}
+
+int32 GGeometryCollectionNanite = 0;
+FAutoConsoleVariableRef CVarGeometryCollectionNanite(
+	TEXT("r.GeometryCollection.Nanite"),
+	GGeometryCollectionNanite,
+	TEXT("Render geometry collections using Nanite."),
+	FConsoleVariableDelegate::CreateStatic(&RecreateGlobalRenderState),
+	ECVF_RenderThreadSafe
+);
+
 // Size in CM used as a threshold for whether a geometry in the collection is collected and exported for
 // navigation purposes. Measured as the diagonal of the leaf node bounds.
 float GGeometryCollectionNavigationSizeThreshold = 20.0f;
@@ -340,42 +357,77 @@ void UGeometryCollectionComponent::CreateRenderState_Concurrent(FRegisterCompone
 
 	if (SceneProxy && RestCollection && RestCollection->HasVisibleGeometry())
 	{
-		FGeometryCollectionSceneProxy* const GeometryCollectionSceneProxy = static_cast<FGeometryCollectionSceneProxy*>(SceneProxy);
-#if GEOMETRYCOLLECTION_EDITOR_SELECTION
-		// Re-init subsections
-		if (bIsTransformSelectionModeEnabled)
-		{
-			GeometryCollectionSceneProxy->UseSubSections(true, false);  // Do not force reinit now, it'll be done in SetConstantData_RenderThread
-		}
-#endif  // #if GEOMETRYCOLLECTION_EDITOR_SELECTION
-
 		FGeometryCollectionConstantData* const ConstantData = ::new FGeometryCollectionConstantData;
 		InitConstantData(ConstantData);
 
 		FGeometryCollectionDynamicData* const DynamicData = ::new FGeometryCollectionDynamicData;
 		InitDynamicData(DynamicData);
 
-		ENQUEUE_RENDER_COMMAND(CreateRenderState)(
-			[GeometryCollectionSceneProxy, ConstantData, DynamicData](FRHICommandListImmediate& RHICmdList)
+		if (SceneProxy->IsNaniteMesh())
+		{
+			FNaniteGeometryCollectionSceneProxy* const GeometryCollectionSceneProxy = static_cast<FNaniteGeometryCollectionSceneProxy*>(SceneProxy);
+
+			// ...
+
+			// TODO: HACK: Pull transforms and delete data
+			delete ConstantData;
+			delete DynamicData;
+
+		#if GEOMETRYCOLLECTION_EDITOR_SELECTION
+			if (bIsTransformSelectionModeEnabled)
 			{
-				if (GeometryCollectionSceneProxy)
+				// ...
+			}
+		#endif
+
+			ENQUEUE_RENDER_COMMAND(CreateRenderState)(
+				[GeometryCollectionSceneProxy, ConstantData, DynamicData](FRHICommandListImmediate& RHICmdList)
 				{
-					GeometryCollectionSceneProxy->SetConstantData_RenderThread(ConstantData);
-					GeometryCollectionSceneProxy->SetDynamicData_RenderThread(DynamicData);
+					if (GeometryCollectionSceneProxy)
+					{
+						// ...
+					}
 				}
+			);
+		}
+		else
+		{
+			FGeometryCollectionSceneProxy* const GeometryCollectionSceneProxy = static_cast<FGeometryCollectionSceneProxy*>(SceneProxy);
+		#if GEOMETRYCOLLECTION_EDITOR_SELECTION
+			// Re-init subsections
+			if (bIsTransformSelectionModeEnabled)
+			{
+				GeometryCollectionSceneProxy->UseSubSections(true, false);  // Do not force reinit now, it'll be done in SetConstantData_RenderThread
 			}
-		);
+		#endif  // #if GEOMETRYCOLLECTION_EDITOR_SELECTION
+
+			ENQUEUE_RENDER_COMMAND(CreateRenderState)(
+				[GeometryCollectionSceneProxy, ConstantData, DynamicData](FRHICommandListImmediate& RHICmdList)
+				{
+					if (GeometryCollectionSceneProxy)
+					{
+						GeometryCollectionSceneProxy->SetConstantData_RenderThread(ConstantData);
+						GeometryCollectionSceneProxy->SetDynamicData_RenderThread(DynamicData);
+					}
+				}
+			);
+		}
 	}
 }
 
-
 FPrimitiveSceneProxy* UGeometryCollectionComponent::CreateSceneProxy()
 {
-	if (RestCollection)
+	if (!RestCollection)
+	{
+		return nullptr;
+	}
+
+	if (RestCollection->EnableNanite && GGeometryCollectionNanite != 0)
 	{
-		return new FGeometryCollectionSceneProxy(this);
+		return new FNaniteGeometryCollectionSceneProxy(this);
 	}
-	return nullptr;
+
+	return new FGeometryCollectionSceneProxy(this);
 }
 
 bool UGeometryCollectionComponent::ShouldCreatePhysicsState() const
diff --git a/Engine/Source/Runtime/Experimental/GeometryCollectionEngine/Private/GeometryCollection/GeometryCollectionObject.cpp b/Engine/Source/Runtime/Experimental/GeometryCollectionEngine/Private/GeometryCollection/GeometryCollectionObject.cpp
--- a/Engine/Source/Runtime/Experimental/GeometryCollectionEngine/Private/GeometryCollection/GeometryCollectionObject.cpp
+++ b/Engine/Source/Runtime/Experimental/GeometryCollectionEngine/Private/GeometryCollection/GeometryCollectionObject.cpp
@@ -20,6 +20,8 @@
 #include "Serialization/MemoryReader.h"
 #include "NaniteBuilder.h"
 #include "Rendering/NaniteResources.h"
+// TODO: Temp until new asset-agnostic builder API
+#include "StaticMeshResources.h"
 #endif
 
 #include "GeometryCollection/GeometryCollectionSimulationCoreTypes.h"
@@ -28,6 +30,8 @@
 
 DEFINE_LOG_CATEGORY_STATIC(UGeometryCollectionLogging, NoLogging, All);
 
+static const TCHAR* GeometryCollectionSelectedMaterial = TEXT("/Engine/EditorMaterials/GeometryCollection/SelectedGeometryMaterial.SelectedGeometryMaterial");
+
 #if ENABLE_COOK_STATS
 namespace GeometryCollectionCookStats
 {
@@ -240,7 +244,7 @@ void UGeometryCollection::InitializeMaterials()
 	Materials = FinalMaterials;
 
 	// Last Material is the selection one
-	UMaterialInterface* BoneSelectedMaterial = LoadObject<UMaterialInterface>(NULL, TEXT("/Engine/EditorMaterials/GeometryCollection/SelectedGeometryMaterial.SelectedGeometryMaterial"), NULL, LOAD_None, NULL);
+	UMaterialInterface* BoneSelectedMaterial = LoadObject<UMaterialInterface>(nullptr, GeometryCollectionSelectedMaterial, nullptr, LOAD_None, nullptr);
 	BoneSelectedMaterialIndex = Materials.Add(BoneSelectedMaterial);
 	
 	GeometryCollection->ReindexMaterials();
diff --git a/Engine/Source/Runtime/Experimental/GeometryCollectionEngine/Private/GeometryCollection/GeometryCollectionSceneProxy.cpp b/Engine/Source/Runtime/Experimental/GeometryCollectionEngine/Private/GeometryCollection/GeometryCollectionSceneProxy.cpp
--- a/Engine/Source/Runtime/Experimental/GeometryCollectionEngine/Private/GeometryCollection/GeometryCollectionSceneProxy.cpp
+++ b/Engine/Source/Runtime/Experimental/GeometryCollectionEngine/Private/GeometryCollection/GeometryCollectionSceneProxy.cpp
@@ -5,6 +5,8 @@
 #include "Async/ParallelFor.h"
 #include "Engine/Engine.h"
 #include "Materials/Material.h"
+#include "RenderCore/Public/CommonRenderResources.h"
+#include "Rendering/NaniteResources.h"
 #include "GeometryCollection/GeometryCollectionComponent.h"
 #include "GeometryCollection/GeometryCollectionAlgo.h"
 #if GEOMETRYCOLLECTION_EDITOR_SELECTION
@@ -1024,3 +1026,238 @@ void FGeometryCollectionSceneProxy::GetPreSkinnedLocalBounds(FBoxSphereBounds& O
 {
 	OutBounds = PreSkinnedBounds;
 }
+
+FNaniteGeometryCollectionSceneProxy::FNaniteGeometryCollectionSceneProxy(UGeometryCollectionComponent* Component)
+	: Nanite::FSceneProxyBase(Component)
+	//, RenderData(Component->GetStaticMesh()->RenderData.Get())
+	, GeometryCollection(Component->GetRestCollection())
+{
+	LLM_SCOPE(ELLMTag::Nanite);
+
+	// Nanite requires GPUScene
+	check(UseGPUScene(GMaxRHIShaderPlatform, GetScene().GetFeatureLevel()));
+
+	MaterialRelevance = Component->GetMaterialRelevance(Component->GetScene()->GetFeatureLevel());
+
+	// Nanite supports the GPUScene instance data buffer.
+	bSupportsInstanceDataBuffer = true;
+
+	bSupportsDistanceFieldRepresentation = false;
+	bSupportsMeshCardRepresentation = false;
+
+	// Use fast path that does not update static draw lists.
+	bStaticElementsAlwaysUseProxyPrimitiveUniformBuffer = true;
+
+	// We always use local vertex factory, which gets its primitive data from
+	// GPUScene, so we can skip expensive primitive uniform buffer updates.
+	bVFRequiresPrimitiveUniformBuffer = true;
+
+	// Indicates if 1 or more materials contain settings not supported by Nanite.
+	bHasMaterialErrors = false;
+
+	// Check if the assigned material can be rendered in Nanite. If not, default.
+	const bool IsRenderable = Nanite::FSceneProxy::IsNaniteRenderable(MaterialRelevance);
+
+	if (!IsRenderable)
+	{
+		UE_LOG(LogStaticMesh, Warning, TEXT("Nanite rendering was chosen for rendering mesh with materials that are not supported. Using default instead."));
+		bHasMaterialErrors = true;
+	}
+
+	const TSharedPtr<FGeometryCollection, ESPMode::ThreadSafe> Collection = GeometryCollection->GetGeometryCollection();
+
+	const auto& SectionsArray = Component->GetSectionsArray();
+	const int32 MeshSectionCount = SectionsArray.Num();
+
+	MaterialSections.Reserve(MeshSectionCount);
+
+	for (const auto& MeshSection : SectionsArray)
+	{
+		if (MeshSection.MaterialID == INDEX_NONE)
+		{
+			continue;
+		}
+
+		if (!MaterialSections.IsValidIndex(MeshSection.MaterialID))
+		{
+			MaterialSections.SetNumZeroed(MeshSection.MaterialID + 1);
+		}
+
+		FMaterialSection& Section = MaterialSections[MeshSection.MaterialID];
+		if (!Section.Material)
+		{
+			UMaterialInterface* Material = Component->GetMaterial(MeshSection.MaterialID);
+			Section.Material = Material;
+		}
+
+		const bool bInvalidMaterial = !Section.Material || Section.Material->GetBlendMode() != BLEND_Opaque;
+		if (bInvalidMaterial)
+		{
+			bHasMaterialErrors = true;
+		}
+
+		const bool bForceDefaultMaterial = /*!!FORCE_NANITE_DEFAULT_MATERIAL ||*/ bHasMaterialErrors;
+		if (bForceDefaultMaterial)
+		{
+			Section.Material = UMaterial::GetDefaultMaterial(MD_Surface);
+		}
+
+		// Should never be null here
+		check(Section.Material != nullptr);
+
+		// Should always be opaque blend mode here.
+		check(Section.Material->GetBlendMode() == BLEND_Opaque);
+	}
+
+	const int32 NumGeometry = Collection->NumElements(FGeometryCollection::GeometryGroup);
+	Instances.SetNumZeroed(NumGeometry);
+
+	const TManagedArray<FBox>& BoundingBoxes = Collection->BoundingBox;
+	const TManagedArray<FTransform>& Transforms = Collection->Transform;
+
+	for (int32 GeometryIndex = 0; GeometryIndex < NumGeometry; ++GeometryIndex)
+	{
+		if (GeometryIndex == 0)
+		{
+			continue; // TODO: Temp hack to hide complete mesh
+		}
+
+		FPrimitiveInstance& Instance = Instances[GeometryIndex];
+		Instance.PrimitiveId = ~uint32(0);
+		Instance.InstanceToLocal = Transforms[GeometryIndex].ToMatrixWithScale();
+		Instance.LocalToInstance = Instance.LocalToWorld.Inverse();
+		Instance.LocalToWorld = Instance.InstanceToLocal;
+		Instance.WorldToLocal = Instance.LocalToInstance;
+		Instance.RenderBounds = BoundingBoxes[GeometryIndex];
+		Instance.LocalBounds = Instance.RenderBounds.TransformBy(Instance.InstanceToLocal);
+		Instance.NaniteInfo = GeometryCollection->GetNaniteInfo(GeometryIndex);
+	}
+}
+
+FPrimitiveViewRelevance FNaniteGeometryCollectionSceneProxy::GetViewRelevance(const FSceneView* View) const
+{
+	LLM_SCOPE(ELLMTag::Nanite);
+
+#if WITH_EDITOR
+	const bool bOptimizedRelevance = false;
+#else
+	const bool bOptimizedRelevance = GNaniteOptimizedRelevance != 0;
+#endif
+
+	FPrimitiveViewRelevance Result;
+	Result.bDrawRelevance = IsShown(View) && View->Family->EngineShowFlags.NaniteMeshes;
+	Result.bShadowRelevance = IsShadowCast(View);
+	Result.bUsesLightingChannels = GetLightingChannelMask() != GetDefaultLightingChannelMask();
+
+	// Always render the Nanite mesh data with static relevance.
+	Result.bStaticRelevance = true;
+
+	// Should always be covered by constructor of Nanite scene proxy.
+	Result.bRenderInMainPass = true;
+
+	if (bOptimizedRelevance) // No dynamic relevance if optimized.
+	{
+		MaterialRelevance.SetPrimitiveViewRelevance(Result);
+		Result.bVelocityRelevance = IsMovable();
+	}
+	else
+	{
+#if WITH_EDITOR
+		//only check these in the editor
+		Result.bEditorStaticSelectionRelevance = (IsSelected() || IsHovered());
+#endif
+
+		bool bSetDynamicRelevance = false;
+
+		Result.bOpaque = true;
+
+		MaterialRelevance.SetPrimitiveViewRelevance(Result);
+		Result.bVelocityRelevance = Result.bOpaque && Result.bRenderInMainPass && IsMovable();
+	}
+
+	return Result;
+}
+
+#if WITH_EDITOR
+HHitProxy* FNaniteGeometryCollectionSceneProxy::CreateHitProxies(UPrimitiveComponent* Component, TArray<TRefCountPtr<HHitProxy>>& OutHitProxies)
+{
+	LLM_SCOPE(ELLMTag::Nanite);
+
+	if (Component->GetOwner())
+	{
+		// Generate separate hit proxies for each material section, so that we can perform hit tests against each one.
+		for (int32 SectionIndex = 0; SectionIndex < MaterialSections.Num(); ++SectionIndex)
+		{
+			FMaterialSection& Section = MaterialSections[SectionIndex];
+			HHitProxy* ActorHitProxy = new HActor(Component->GetOwner(), Component, SectionIndex, SectionIndex);
+			check(!Section.HitProxy);
+			Section.HitProxy = ActorHitProxy;
+			OutHitProxies.Add(ActorHitProxy);
+		}
+	}
+
+	// We don't want a default hit proxy, or to output any hit proxies (avoid 2x registration).
+	return nullptr;
+}
+#endif
+
+void FNaniteGeometryCollectionSceneProxy::DrawStaticElements(FStaticPrimitiveDrawInterface* PDI)
+{
+	LLM_SCOPE(ELLMTag::Nanite);
+
+	for (int32 SectionIndex = 0; SectionIndex < MaterialSections.Num(); ++SectionIndex)
+	{
+		const FMaterialSection& Section = MaterialSections[SectionIndex];
+		const UMaterialInterface* Material = Section.Material;
+		if (!Material)
+		{
+			continue;
+		}
+
+		FMaterialRenderProxy* MaterialProxy = Material->GetRenderProxy();
+		check(MaterialProxy);
+
+		FMeshBatch MeshBatch;
+		MeshBatch.SegmentIndex = SectionIndex;
+		MeshBatch.VertexFactory = Nanite::GGlobalResources.GetVertexFactory();
+		MeshBatch.Type = GRHISupportsRectTopology ? PT_RectList : PT_TriangleList;
+		MeshBatch.ReverseCulling = false;
+		MeshBatch.bDisableBackfaceCulling = true;
+		MeshBatch.DepthPriorityGroup = SDPG_World;
+		MeshBatch.LODIndex = -1;
+		MeshBatch.MaterialRenderProxy = MaterialProxy;
+		MeshBatch.bWireframe = false;
+		MeshBatch.bCanApplyViewModeOverrides = false;
+		MeshBatch.LCI = nullptr;// &MeshInfo;
+		MeshBatch.Elements[0].IndexBuffer = &GScreenRectangleIndexBuffer;
+		if (GRHISupportsRectTopology)
+		{
+			MeshBatch.Elements[0].FirstIndex = 9;
+			MeshBatch.Elements[0].NumPrimitives = 1;
+			MeshBatch.Elements[0].MinVertexIndex = 1;
+			MeshBatch.Elements[0].MaxVertexIndex = 3;
+		}
+		else
+		{
+			MeshBatch.Elements[0].FirstIndex = 0;
+			MeshBatch.Elements[0].NumPrimitives = 2;
+			MeshBatch.Elements[0].MinVertexIndex = 0;
+			MeshBatch.Elements[0].MaxVertexIndex = 3;
+		}
+		MeshBatch.Elements[0].NumInstances = 1;
+		MeshBatch.Elements[0].PrimitiveIdMode = PrimID_ForceZero;
+		MeshBatch.Elements[0].PrimitiveUniformBufferResource = &GIdentityPrimitiveUniformBuffer;
+
+#if WITH_EDITOR
+		HHitProxy* HitProxy = Section.HitProxy;
+		//check(HitProxy); // TODO: Is this valid? SME seems to have null proxies, but normal editor doesn't
+		PDI->SetHitProxy(HitProxy);
+#endif
+		PDI->DrawMesh(MeshBatch, FLT_MAX);
+	}
+}
+
+uint32 FNaniteGeometryCollectionSceneProxy::GetMemoryFootprint() const
+{
+	return sizeof(*this) + GetAllocatedSize();
+}
diff --git a/Engine/Source/Runtime/Experimental/GeometryCollectionEngine/Private/GeometryCollection/GeometryCollectionSceneProxy.h b/Engine/Source/Runtime/Experimental/GeometryCollectionEngine/Private/GeometryCollection/GeometryCollectionSceneProxy.h
--- a/Engine/Source/Runtime/Experimental/GeometryCollectionEngine/Private/GeometryCollection/GeometryCollectionSceneProxy.h
+++ b/Engine/Source/Runtime/Experimental/GeometryCollectionEngine/Private/GeometryCollection/GeometryCollectionSceneProxy.h
@@ -16,11 +16,16 @@
 #include "GeometryCollection/GeometryCollectionHitProxy.h"
 #endif
 
-
+class UGeometryCollection;
 class UGeometryCollectionComponent;
 struct FGeometryCollectionSection;
 struct HGeometryCollection;
 
+namespace Nanite
+{
+	struct FResources;
+}
+
 /** Index Buffer */
 class FGeometryCollectionIndexBuffer : public FIndexBuffer
 {
@@ -254,4 +259,45 @@ private:
 	/** Release subsections by emptying the associated arrays. */
 	void ReleaseSubSections_RenderThread();
 #endif
-};
\ No newline at end of file
+};
+
+class FNaniteGeometryCollectionSceneProxy : public Nanite::FSceneProxyBase
+{
+public:
+	FNaniteGeometryCollectionSceneProxy(UGeometryCollectionComponent* Component);
+
+	virtual ~FNaniteGeometryCollectionSceneProxy() = default;
+
+public:
+	// FPrimitiveSceneProxy interface.
+	virtual FPrimitiveViewRelevance	GetViewRelevance(const FSceneView* View) const override;
+#if WITH_EDITOR
+	virtual HHitProxy* CreateHitProxies(UPrimitiveComponent* Component, TArray<TRefCountPtr<HHitProxy> >& OutHitProxies) override;
+#endif
+	virtual void DrawStaticElements(FStaticPrimitiveDrawInterface* PDI) override;
+
+	virtual uint32 GetMemoryFootprint() const override;
+
+	const FORCEINLINE TArray<Nanite::FResources*>& GetResources() const
+	{
+		return Resources;
+	}
+
+	FORCEINLINE TArray<Nanite::FResources*>& GetResources()
+	{
+		return Resources;
+	}
+
+protected:
+	TArray<Nanite::FResources*> Resources;
+
+	// TODO: Should probably calculate this on the materials array above instead of on the component
+	//       Null and !Opaque are assigned default material unlike the component material relevance.
+	FMaterialRelevance MaterialRelevance;
+
+	uint32 bCastShadow : 1;
+	uint32 bReverseCulling : 1;
+	uint32 bHasMaterialErrors : 1;
+
+	const UGeometryCollection* GeometryCollection = nullptr;
+};
