Remove optimized relevancy from geometry collection Nanite scene proxy

diff --git a/Engine/Source/Runtime/Experimental/GeometryCollectionEngine/Private/GeometryCollection/GeometryCollectionSceneProxy.cpp b/Engine/Source/Runtime/Experimental/GeometryCollectionEngine/Private/GeometryCollection/GeometryCollectionSceneProxy.cpp
--- a/Engine/Source/Runtime/Experimental/GeometryCollectionEngine/Private/GeometryCollection/GeometryCollectionSceneProxy.cpp
+++ b/Engine/Source/Runtime/Experimental/GeometryCollectionEngine/Private/GeometryCollection/GeometryCollectionSceneProxy.cpp
@@ -1141,7 +1141,7 @@ FPrimitiveViewRelevance FNaniteGeometryCollectionSceneProxy::GetViewRelevance(co
 #if WITH_EDITOR
 	const bool bOptimizedRelevance = false;
 #else
-	const bool bOptimizedRelevance = GNaniteOptimizedRelevance != 0;
+	const bool bOptimizedRelevance = false;// GNaniteOptimizedRelevance != 0;
 #endif
 
 	FPrimitiveViewRelevance Result;
