Fix material TSB conversion for Nanite geometry collection data.

diff --git a/Engine/Source/Runtime/Experimental/GeometryCollectionEngine/Private/GeometryCollection/DerivedDataGeometryCollectionCooker.cpp b/Engine/Source/Runtime/Experimental/GeometryCollectionEngine/Private/GeometryCollection/DerivedDataGeometryCollectionCooker.cpp
--- a/Engine/Source/Runtime/Experimental/GeometryCollectionEngine/Private/GeometryCollection/DerivedDataGeometryCollectionCooker.cpp
+++ b/Engine/Source/Runtime/Experimental/GeometryCollectionEngine/Private/GeometryCollection/DerivedDataGeometryCollectionCooker.cpp
@@ -76,7 +76,8 @@ const TCHAR* FDerivedDataGeometryCollectionCooker::GetVersionString() const
 	static FString CachedNaniteVersionString;
 	if (CachedNaniteVersionString.IsEmpty())
 	{
-		CachedNaniteVersionString = FString::Printf(TEXT("%s_%s"), VersionString, *Nanite::IBuilderModule::Get().GetVersionString());
+		const TCHAR* NaniteVersionString = TEXT("CE9B3AB5F81A49B388E50FE4B8C19D7E");
+		CachedNaniteVersionString = FString::Printf(TEXT("%s_%s_%s"), VersionString, NaniteVersionString, *Nanite::IBuilderModule::Get().GetVersionString());
 	}
 
 	return GeometryCollection.EnableNanite ? *CachedNaniteVersionString : VersionString;
diff --git a/Engine/Source/Runtime/Experimental/GeometryCollectionEngine/Private/GeometryCollection/GeometryCollectionObject.cpp b/Engine/Source/Runtime/Experimental/GeometryCollectionEngine/Private/GeometryCollection/GeometryCollectionObject.cpp
--- a/Engine/Source/Runtime/Experimental/GeometryCollectionEngine/Private/GeometryCollection/GeometryCollectionObject.cpp
+++ b/Engine/Source/Runtime/Experimental/GeometryCollectionEngine/Private/GeometryCollection/GeometryCollectionObject.cpp
@@ -513,7 +513,7 @@ TArray<Nanite::FResources>& UGeometryCollection::CreateNaniteData(FGeometryColle
 		NaniteResource = {};
 
 		uint32 NumTexCoords = 1;// NumTextureCoord;
-		bool bHasColors = false;// true;
+		bool bHasColors = ColorArray.Num() > 0;
 
 		const int32 VertexStart = VertexStartArray[GeometryGroupIndex];
 		const int32 VertexCount = VertexCountArray[GeometryGroupIndex];
@@ -524,8 +524,15 @@ TArray<Nanite::FResources>& UGeometryCollection::CreateNaniteData(FGeometryColle
 		{
 			FStaticMeshBuildVertex& Vertex = BuildVertices.Emplace_GetRef();
 			Vertex.Position = VertexArray[VertexStart + VertexIndex];
-			Vertex.Color = FColor::White;
+			Vertex.Color = bHasColors ? ColorArray[VertexStart + VertexIndex].ToFColor(false /* sRGB */) : FColor::White;
+			Vertex.TangentX = FVector::ZeroVector;
+			Vertex.TangentY = FVector::ZeroVector;
+			Vertex.TangentZ = NormalArray[VertexStart + VertexIndex];
 			Vertex.UVs[0] = UVArray[VertexStart + VertexIndex];
+			if (Vertex.UVs[0].ContainsNaN())
+			{
+				Vertex.UVs[0] = FVector2D::ZeroVector;
+			}
 		}
 
 		const int32 FaceStart = FaceStartArray[GeometryGroupIndex];
