Initial support for dynamic GPU scene transform updates for Nanite destruction simulation. Many rough edges to file down, but this is fairly promising.

diff --git a/Engine/Source/Runtime/Experimental/GeometryCollectionEngine/Private/GeometryCollection/GeometryCollectionComponent.cpp b/Engine/Source/Runtime/Experimental/GeometryCollectionEngine/Private/GeometryCollection/GeometryCollectionComponent.cpp
--- a/Engine/Source/Runtime/Experimental/GeometryCollectionEngine/Private/GeometryCollection/GeometryCollectionComponent.cpp
+++ b/Engine/Source/Runtime/Experimental/GeometryCollectionEngine/Private/GeometryCollection/GeometryCollectionComponent.cpp
@@ -35,6 +35,7 @@
 #include "Chaos/ChaosGameplayEventDispatcher.h"
 
 #include "Rendering/NaniteResources.h"
+#include "PrimitiveSceneInfo.h"
 
 #if !(UE_BUILD_SHIPPING || UE_BUILD_TEST)
 #include "Logging/MessageLog.h"
@@ -369,10 +370,6 @@ void UGeometryCollectionComponent::CreateRenderState_Concurrent(FRegisterCompone
 
 			// ...
 
-			// TODO: HACK: Pull transforms and delete data
-			delete ConstantData;
-			delete DynamicData;
-
 		#if GEOMETRYCOLLECTION_EDITOR_SELECTION
 			if (bIsTransformSelectionModeEnabled)
 			{
@@ -385,7 +382,9 @@ void UGeometryCollectionComponent::CreateRenderState_Concurrent(FRegisterCompone
 				{
 					if (GeometryCollectionSceneProxy)
 					{
-						// ...
+						GeometryCollectionSceneProxy->SetConstantData_RenderThread(ConstantData);
+						GeometryCollectionSceneProxy->SetDynamicData_RenderThread(DynamicData);
+						GeometryCollectionSceneProxy->GetPrimitiveSceneInfo()->RequestGPUSceneUpdate();
 					}
 				}
 			);
@@ -393,13 +392,14 @@ void UGeometryCollectionComponent::CreateRenderState_Concurrent(FRegisterCompone
 		else
 		{
 			FGeometryCollectionSceneProxy* const GeometryCollectionSceneProxy = static_cast<FGeometryCollectionSceneProxy*>(SceneProxy);
+
 		#if GEOMETRYCOLLECTION_EDITOR_SELECTION
 			// Re-init subsections
 			if (bIsTransformSelectionModeEnabled)
 			{
 				GeometryCollectionSceneProxy->UseSubSections(true, false);  // Do not force reinit now, it'll be done in SetConstantData_RenderThread
 			}
-		#endif  // #if GEOMETRYCOLLECTION_EDITOR_SELECTION
+		#endif
 
 			ENQUEUE_RENDER_COMMAND(CreateRenderState)(
 				[GeometryCollectionSceneProxy, ConstantData, DynamicData](FRHICommandListImmediate& RHICmdList)
@@ -1623,16 +1623,33 @@ void UGeometryCollectionComponent::SendRenderDynamicData_Concurrent()
 		else
 		{
 			// Enqueue command to send to render thread
-			FGeometryCollectionSceneProxy* GeometryCollectionSceneProxy = static_cast<FGeometryCollectionSceneProxy*>(SceneProxy);
-			ENQUEUE_RENDER_COMMAND(SendRenderDynamicData)(
-				[GeometryCollectionSceneProxy, DynamicData](FRHICommandListImmediate& RHICmdList)
-				{
-					if (GeometryCollectionSceneProxy)
+			if (SceneProxy->IsNaniteMesh())
+			{
+				FNaniteGeometryCollectionSceneProxy* GeometryCollectionSceneProxy = static_cast<FNaniteGeometryCollectionSceneProxy*>(SceneProxy);
+				ENQUEUE_RENDER_COMMAND(SendRenderDynamicData)(
+					[GeometryCollectionSceneProxy, DynamicData](FRHICommandListImmediate& RHICmdList)
 					{
-						GeometryCollectionSceneProxy->SetDynamicData_RenderThread(DynamicData);
+						if (GeometryCollectionSceneProxy)
+						{
+							GeometryCollectionSceneProxy->SetDynamicData_RenderThread(DynamicData);
+							GeometryCollectionSceneProxy->GetPrimitiveSceneInfo()->RequestGPUSceneUpdate();
+						}
 					}
-				}
-			);
+				);
+			}
+			else
+			{
+				FGeometryCollectionSceneProxy* GeometryCollectionSceneProxy = static_cast<FGeometryCollectionSceneProxy*>(SceneProxy);
+				ENQUEUE_RENDER_COMMAND(SendRenderDynamicData)(
+					[GeometryCollectionSceneProxy, DynamicData](FRHICommandListImmediate& RHICmdList)
+					{
+						if (GeometryCollectionSceneProxy)
+						{
+							GeometryCollectionSceneProxy->SetDynamicData_RenderThread(DynamicData);
+						}
+					}
+				);
+			}
 		}
 
 		// mark collection clean now that we have rendered
@@ -2328,7 +2345,8 @@ bool UGeometryCollectionComponent::IsEqual(const TArray<FMatrix> &A, const TArra
 #if GEOMETRYCOLLECTION_EDITOR_SELECTION
 void UGeometryCollectionComponent::EnableTransformSelectionMode(bool bEnable)
 {
-	if (SceneProxy && RestCollection && RestCollection->HasVisibleGeometry())
+	// TODO: Support for Nanite?
+	if (SceneProxy && !SceneProxy->IsNaniteMesh() && RestCollection && RestCollection->HasVisibleGeometry())
 	{
 		static_cast<FGeometryCollectionSceneProxy*>(SceneProxy)->UseSubSections(bEnable, true);
 	}
diff --git a/Engine/Source/Runtime/Experimental/GeometryCollectionEngine/Private/GeometryCollection/GeometryCollectionSceneProxy.cpp b/Engine/Source/Runtime/Experimental/GeometryCollectionEngine/Private/GeometryCollection/GeometryCollectionSceneProxy.cpp
--- a/Engine/Source/Runtime/Experimental/GeometryCollectionEngine/Private/GeometryCollection/GeometryCollectionSceneProxy.cpp
+++ b/Engine/Source/Runtime/Experimental/GeometryCollectionEngine/Private/GeometryCollection/GeometryCollectionSceneProxy.cpp
@@ -1109,27 +1109,19 @@ FNaniteGeometryCollectionSceneProxy::FNaniteGeometryCollectionSceneProxy(UGeomet
 		check(Section.Material->GetBlendMode() == BLEND_Opaque);
 	}
 
-	const int32 NumGeometry = Collection->NumElements(FGeometryCollection::GeometryGroup);
-	Instances.SetNumZeroed(NumGeometry);
-
 	const TManagedArray<FBox>& BoundingBoxes = Collection->BoundingBox;
-	const TManagedArray<FTransform>& Transforms = Collection->Transform;
+	//const TManagedArray<FTransform>& Transforms = Collection->Transform;
+
+	// TODO: Not fully correct. Transforms are not always 1:1 with geometry groups or bones.
+
+	const int32 NumGeometry = Collection->NumElements(FGeometryCollection::GeometryGroup);
+	GeometryNaniteData.SetNumUninitialized(NumGeometry);
 
 	for (int32 GeometryIndex = 0; GeometryIndex < NumGeometry; ++GeometryIndex)
 	{
-		if (GeometryIndex == 0)
-		{
-			continue; // TODO: Temp hack to hide complete mesh
-		}
-
-		FPrimitiveInstance& Instance = Instances[GeometryIndex];
+		FGeometryNaniteData& Instance = GeometryNaniteData[GeometryIndex];
 		Instance.PrimitiveId = ~uint32(0);
-		Instance.InstanceToLocal = Transforms[GeometryIndex].ToMatrixWithScale();
-		Instance.LocalToInstance = Instance.LocalToWorld.Inverse();
-		Instance.LocalToWorld = Instance.InstanceToLocal;
-		Instance.WorldToLocal = Instance.LocalToInstance;
 		Instance.RenderBounds = BoundingBoxes[GeometryIndex];
-		Instance.LocalBounds = Instance.RenderBounds.TransformBy(Instance.InstanceToLocal);
 		Instance.NaniteInfo = GeometryCollection->GetNaniteInfo(GeometryIndex);
 	}
 }
@@ -1261,3 +1253,61 @@ uint32 FNaniteGeometryCollectionSceneProxy::GetMemoryFootprint() const
 {
 	return sizeof(*this) + GetAllocatedSize();
 }
+
+void FNaniteGeometryCollectionSceneProxy::SetConstantData_RenderThread(FGeometryCollectionConstantData* NewConstantData, bool ForceInit)
+{
+	check(NewConstantData->RestTransforms.Num() == GeometryNaniteData.Num());
+
+	Instances.SetNumUninitialized(NewConstantData->RestTransforms.Num() - 1); // TODO: Temp hack to hide complete mesh (generally 0th index)
+	for (int32 InstanceIndex = 0; InstanceIndex < Instances.Num(); ++InstanceIndex)
+	{
+		FPrimitiveInstance& Instance = Instances[InstanceIndex];
+		
+		Instance.InstanceToLocal = NewConstantData->RestTransforms[InstanceIndex + 1]; // TODO: Temp hack to hide complete mesh (generally 0th index)
+		Instance.LocalToInstance = Instance.LocalToWorld.Inverse();
+		Instance.LocalToWorld = Instance.InstanceToLocal;
+		Instance.WorldToLocal = Instance.LocalToInstance;
+
+		const FGeometryNaniteData& NaniteData = GeometryNaniteData[InstanceIndex + 1]; // TODO: Temp hack to hide complete mesh (generally 0th index)
+
+		Instance.PrimitiveId = NaniteData.PrimitiveId;
+		Instance.RenderBounds = NaniteData.RenderBounds;
+		Instance.LocalBounds = Instance.RenderBounds.TransformBy(Instance.InstanceToLocal);
+		Instance.NaniteInfo = NaniteData.NaniteInfo;
+	}
+
+	delete NewConstantData;
+}
+
+void FNaniteGeometryCollectionSceneProxy::SetDynamicData_RenderThread(FGeometryCollectionDynamicData* NewDynamicData)
+{
+	check(NewDynamicData->IsLoading || NewDynamicData->Transforms.Num() == GeometryNaniteData.Num());
+
+	if (NewDynamicData->IsDynamic)
+	{
+		Instances.SetNumUninitialized(NewDynamicData->Transforms.Num() - 1); // TODO: Temp hack to hide complete mesh (generally 0th index)
+		for (int32 InstanceIndex = 0; InstanceIndex < Instances.Num(); ++InstanceIndex)
+		{
+			FPrimitiveInstance& Instance = Instances[InstanceIndex];
+
+			Instance.InstanceToLocal = NewDynamicData->Transforms[InstanceIndex + 1]; // TODO: Temp hack to hide complete mesh (generally 0th index)
+			Instance.LocalToInstance = Instance.LocalToWorld.Inverse();
+			Instance.LocalToWorld = Instance.InstanceToLocal;
+			Instance.WorldToLocal = Instance.LocalToInstance;
+
+			//Instance.PrimitiveId = GeometryNaniteData[InstanceIndex].PrimitiveId;
+			//Instance.RenderBounds = GeometryNaniteData[InstanceIndex].RenderBounds;
+			Instance.LocalBounds = Instance.RenderBounds.TransformBy(Instance.InstanceToLocal);
+			//Instance.NaniteInfo = GeometryNaniteData[InstanceIndex].NaniteInfo;
+		}
+	}
+	else
+	{
+		// Rendering base geometry, use rest transforms rather than simulated transforms.
+		//Instances = RestInstances;
+		
+
+	}
+
+	delete NewDynamicData;
+}
diff --git a/Engine/Source/Runtime/Experimental/GeometryCollectionEngine/Private/GeometryCollection/GeometryCollectionSceneProxy.h b/Engine/Source/Runtime/Experimental/GeometryCollectionEngine/Private/GeometryCollection/GeometryCollectionSceneProxy.h
--- a/Engine/Source/Runtime/Experimental/GeometryCollectionEngine/Private/GeometryCollection/GeometryCollectionSceneProxy.h
+++ b/Engine/Source/Runtime/Experimental/GeometryCollectionEngine/Private/GeometryCollection/GeometryCollectionSceneProxy.h
@@ -278,6 +278,12 @@ public:
 
 	virtual uint32 GetMemoryFootprint() const override;
 
+	/** Called on render thread to setup static geometry for rendering */
+	void SetConstantData_RenderThread(FGeometryCollectionConstantData* NewConstantData, bool ForceInit = false);
+
+	/** Called on render thread to setup dynamic geometry for rendering */
+	void SetDynamicData_RenderThread(FGeometryCollectionDynamicData* NewDynamicData);
+
 	const FORCEINLINE TArray<Nanite::FResources*>& GetResources() const
 	{
 		return Resources;
@@ -300,4 +306,14 @@ protected:
 	uint32 bHasMaterialErrors : 1;
 
 	const UGeometryCollection* GeometryCollection = nullptr;
+
+	//TArray<FPrimitiveInstance> RestInstances;
+
+	struct FGeometryNaniteData
+	{
+		FBoxSphereBounds RenderBounds;
+		FNaniteInfo NaniteInfo;
+		uint32 PrimitiveId;
+	};
+	TArray<FGeometryNaniteData> GeometryNaniteData;
 };
diff --git a/Engine/Source/Runtime/Renderer/Private/PrimitiveSceneInfo.cpp b/Engine/Source/Runtime/Renderer/Private/PrimitiveSceneInfo.cpp
--- a/Engine/Source/Runtime/Renderer/Private/PrimitiveSceneInfo.cpp
+++ b/Engine/Source/Runtime/Renderer/Private/PrimitiveSceneInfo.cpp
@@ -1332,6 +1332,14 @@ void FPrimitiveSceneInfo::UnlinkAttachmentGroup()
 	}
 }
 
+void FPrimitiveSceneInfo::RequestGPUSceneUpdate()
+{
+	if (Scene && IsIndexValid())
+	{
+		AddPrimitiveToUpdateGPU(*Scene, GetIndex());
+	}
+}
+
 void FPrimitiveSceneInfo::GatherLightingAttachmentGroupPrimitives(TArray<FPrimitiveSceneInfo*, SceneRenderingAllocator>& OutChildSceneInfos)
 {
 #if ENABLE_NAN_DIAGNOSTIC
diff --git a/Engine/Source/Runtime/Renderer/Public/PrimitiveSceneInfo.h b/Engine/Source/Runtime/Renderer/Public/PrimitiveSceneInfo.h
--- a/Engine/Source/Runtime/Renderer/Public/PrimitiveSceneInfo.h
+++ b/Engine/Source/Runtime/Renderer/Public/PrimitiveSceneInfo.h
@@ -419,6 +419,9 @@ public:
 	/** Removes the primitive from the scene's attachment groups. */
 	void UnlinkAttachmentGroup();
 
+	/** Adds a request to update GPU scene representation. */
+	RENDERER_API void RequestGPUSceneUpdate();
+
 	/** 
 	 * Builds an array of all primitive scene info's in this primitive's attachment group. 
 	 * This only works on potential parents (!LightingAttachmentRoot.IsValid()) and will include the current primitive in the output array.
