Support arbitrary cardinality between geometry collection groups and transforms in the Nanite scene proxy. Also removes a number of hacks to hide the rest-pose geometry.

diff --git a/Engine/Source/Runtime/Experimental/GeometryCollectionEngine/Private/GeometryCollection/GeometryCollectionSceneProxy.cpp b/Engine/Source/Runtime/Experimental/GeometryCollectionEngine/Private/GeometryCollection/GeometryCollectionSceneProxy.cpp
--- a/Engine/Source/Runtime/Experimental/GeometryCollectionEngine/Private/GeometryCollection/GeometryCollectionSceneProxy.cpp
+++ b/Engine/Source/Runtime/Experimental/GeometryCollectionEngine/Private/GeometryCollection/GeometryCollectionSceneProxy.cpp
@@ -1110,9 +1110,6 @@ FNaniteGeometryCollectionSceneProxy::FNaniteGeometryCollectionSceneProxy(UGeomet
 	}
 
 	const TManagedArray<FBox>& BoundingBoxes = Collection->BoundingBox;
-	//const TManagedArray<FTransform>& Transforms = Collection->Transform;
-
-	// TODO: Not fully correct. Transforms are not always 1:1 with geometry groups or bones.
 
 	const int32 NumGeometry = Collection->NumElements(FGeometryCollection::GeometryGroup);
 	GeometryNaniteData.SetNumUninitialized(NumGeometry);
@@ -1256,24 +1253,32 @@ uint32 FNaniteGeometryCollectionSceneProxy::GetMemoryFootprint() const
 
 void FNaniteGeometryCollectionSceneProxy::SetConstantData_RenderThread(FGeometryCollectionConstantData* NewConstantData, bool ForceInit)
 {
-	check(NewConstantData->RestTransforms.Num() == GeometryNaniteData.Num());
+	const TManagedArray<int32>& TransformToGeometryIndices = GeometryCollection->GetGeometryCollection()->TransformToGeometryIndex;
+	check(NewConstantData->RestTransforms.Num() == TransformToGeometryIndices.Num());
+
+	Instances.Reset(NewConstantData->RestTransforms.Num());
 
-	Instances.SetNumUninitialized(NewConstantData->RestTransforms.Num() - 1); // TODO: Temp hack to hide complete mesh (generally 0th index)
-	for (int32 InstanceIndex = 0; InstanceIndex < Instances.Num(); ++InstanceIndex)
+	for (int32 TransformIndex = 0; TransformIndex < NewConstantData->RestTransforms.Num(); ++TransformIndex)
 	{
-		FPrimitiveInstance& Instance = Instances[InstanceIndex];
-		
-		Instance.InstanceToLocal = NewConstantData->RestTransforms[InstanceIndex + 1]; // TODO: Temp hack to hide complete mesh (generally 0th index)
-		Instance.LocalToInstance = Instance.LocalToWorld.Inverse();
-		Instance.LocalToWorld = Instance.InstanceToLocal;
-		Instance.WorldToLocal = Instance.LocalToInstance;
+		const int32 TransformToGeometryIndex = TransformToGeometryIndices[TransformIndex];
+		if (TransformToGeometryIndex == INDEX_NONE)
+		{
+			// Not geometry
+			continue;
+		}
 
-		const FGeometryNaniteData& NaniteData = GeometryNaniteData[InstanceIndex + 1]; // TODO: Temp hack to hide complete mesh (generally 0th index)
+		const FGeometryNaniteData& NaniteData = GeometryNaniteData[TransformToGeometryIndex];
 
-		Instance.PrimitiveId = NaniteData.PrimitiveId;
-		Instance.RenderBounds = NaniteData.RenderBounds;
-		Instance.LocalBounds = Instance.RenderBounds.TransformBy(Instance.InstanceToLocal);
-		Instance.NaniteInfo = NaniteData.NaniteInfo;
+		FPrimitiveInstance& Instance = Instances.Emplace_GetRef();
+
+		Instance.InstanceToLocal	= NewConstantData->RestTransforms[TransformIndex];
+		Instance.LocalToInstance	= Instance.LocalToWorld.Inverse();
+		Instance.LocalToWorld		= Instance.InstanceToLocal;
+		Instance.WorldToLocal		= Instance.LocalToInstance;
+		Instance.PrimitiveId		= NaniteData.PrimitiveId;
+		Instance.RenderBounds		= NaniteData.RenderBounds;
+		Instance.LocalBounds		= Instance.RenderBounds.TransformBy(Instance.InstanceToLocal);
+		Instance.NaniteInfo			= NaniteData.NaniteInfo;
 	}
 
 	delete NewConstantData;
@@ -1281,32 +1286,47 @@ void FNaniteGeometryCollectionSceneProxy::SetConstantData_RenderThread(FGeometry
 
 void FNaniteGeometryCollectionSceneProxy::SetDynamicData_RenderThread(FGeometryCollectionDynamicData* NewDynamicData)
 {
-	check(NewDynamicData->IsLoading || NewDynamicData->Transforms.Num() == GeometryNaniteData.Num());
+	//if (NewDynamicData->IsLoading)
+	{
+		// Ignored for now (rest transforms are already set)
+		//return;
+	}
 
+	// Are we currently simulating?
 	if (NewDynamicData->IsDynamic)
 	{
-		Instances.SetNumUninitialized(NewDynamicData->Transforms.Num() - 1); // TODO: Temp hack to hide complete mesh (generally 0th index)
-		for (int32 InstanceIndex = 0; InstanceIndex < Instances.Num(); ++InstanceIndex)
+		const TManagedArray<int32>& TransformToGeometryIndices = GeometryCollection->GetGeometryCollection()->TransformToGeometryIndex;
+		check(NewDynamicData->Transforms.Num() == TransformToGeometryIndices.Num());
+
+		Instances.Reset(NewDynamicData->Transforms.Num());
+
+		for (int32 TransformIndex = 0; TransformIndex < NewDynamicData->Transforms.Num(); ++TransformIndex)
 		{
-			FPrimitiveInstance& Instance = Instances[InstanceIndex];
+			const int32 TransformToGeometryIndex = TransformToGeometryIndices[TransformIndex];
+			if (TransformToGeometryIndex == INDEX_NONE)
+			{
+				// Not geometry
+				continue;
+			}
+
+			const FGeometryNaniteData& NaniteData = GeometryNaniteData[TransformToGeometryIndex];
 
-			Instance.InstanceToLocal = NewDynamicData->Transforms[InstanceIndex + 1]; // TODO: Temp hack to hide complete mesh (generally 0th index)
-			Instance.LocalToInstance = Instance.LocalToWorld.Inverse();
-			Instance.LocalToWorld = Instance.InstanceToLocal;
-			Instance.WorldToLocal = Instance.LocalToInstance;
+			FPrimitiveInstance& Instance = Instances.Emplace_GetRef();
 
-			//Instance.PrimitiveId = GeometryNaniteData[InstanceIndex].PrimitiveId;
-			//Instance.RenderBounds = GeometryNaniteData[InstanceIndex].RenderBounds;
-			Instance.LocalBounds = Instance.RenderBounds.TransformBy(Instance.InstanceToLocal);
-			//Instance.NaniteInfo = GeometryNaniteData[InstanceIndex].NaniteInfo;
+			Instance.InstanceToLocal	= NewDynamicData->Transforms[TransformIndex];
+			Instance.LocalToInstance	= Instance.LocalToWorld.Inverse();
+			Instance.LocalToWorld		= Instance.InstanceToLocal;
+			Instance.WorldToLocal		= Instance.LocalToInstance;
+			Instance.PrimitiveId		= NaniteData.PrimitiveId;
+			Instance.RenderBounds		= NaniteData.RenderBounds;
+			Instance.LocalBounds		= Instance.RenderBounds.TransformBy(Instance.InstanceToLocal);
+			Instance.NaniteInfo			= NaniteData.NaniteInfo;
 		}
 	}
 	else
 	{
 		// Rendering base geometry, use rest transforms rather than simulated transforms.
-		//Instances = RestInstances;
-		
-
+		// ...
 	}
 
 	delete NewDynamicData;
