Enable Nanite scene proxy construction by default for Chaos geometry collections.

diff --git a/Engine/Source/Runtime/Experimental/GeometryCollectionEngine/Private/GeometryCollection/GeometryCollectionComponent.cpp b/Engine/Source/Runtime/Experimental/GeometryCollectionEngine/Private/GeometryCollection/GeometryCollectionComponent.cpp
--- a/Engine/Source/Runtime/Experimental/GeometryCollectionEngine/Private/GeometryCollection/GeometryCollectionComponent.cpp
+++ b/Engine/Source/Runtime/Experimental/GeometryCollectionEngine/Private/GeometryCollection/GeometryCollectionComponent.cpp
@@ -84,7 +84,7 @@ static void RecreateGlobalRenderState(IConsoleVariable* Var)
 	FGlobalComponentRecreateRenderStateContext Context;
 }
 
-int32 GGeometryCollectionNanite = 0;
+int32 GGeometryCollectionNanite = 1;
 FAutoConsoleVariableRef CVarGeometryCollectionNanite(
 	TEXT("r.GeometryCollection.Nanite"),
 	GGeometryCollectionNanite,
