Correctly support geometry collection clustering with Nanite. Also fixed proxy recreation when EnableNanite is toggled in the editor.

diff --git a/Engine/Source/Runtime/Experimental/GeometryCollectionEngine/Private/GeometryCollection/GeometryCollectionComponent.cpp b/Engine/Source/Runtime/Experimental/GeometryCollectionEngine/Private/GeometryCollection/GeometryCollectionComponent.cpp
--- a/Engine/Source/Runtime/Experimental/GeometryCollectionEngine/Private/GeometryCollection/GeometryCollectionComponent.cpp
+++ b/Engine/Source/Runtime/Experimental/GeometryCollectionEngine/Private/GeometryCollection/GeometryCollectionComponent.cpp
@@ -1179,9 +1179,26 @@ void UGeometryCollectionComponent::TickComponent(float DeltaTime, enum ELevelTic
 	//UE_LOG(UGCC_LOG, Log, TEXT("GeometryCollectionComponent[%p]::TickComponent()"), this);
 	Super::TickComponent(DeltaTime, TickType, ThisTickFunction);
 
+#if WITH_EDITOR
+	if (IsRegistered() && SceneProxy && RestCollection)
+	{
+		const bool bWantNanite = RestCollection->EnableNanite && GGeometryCollectionNanite != 0;
+		const bool bHaveNanite = SceneProxy->IsNaniteMesh();
+		bool bRecreateProxy = bWantNanite != bHaveNanite;
+		if (bRecreateProxy)
+		{
+			// Wait until resources are released
+			FlushRenderingCommands();
+
+			FComponentReregisterContext ReregisterContext(this);
+			UpdateAllPrimitiveSceneInfosForSingleComponent(this);
+		}
+	}
+#endif
+
 #if WITH_CHAOS
 	//if (bRenderStateDirty && DynamicCollection)	//todo: always send for now
-	if(RestCollection)
+	if (RestCollection)
 	{
 		if(ensureMsgf(DynamicCollection, TEXT("No dynamic collection available for component %s during tick."), *GetName()))
 		{
diff --git a/Engine/Source/Runtime/Experimental/GeometryCollectionEngine/Private/GeometryCollection/GeometryCollectionObject.cpp b/Engine/Source/Runtime/Experimental/GeometryCollectionEngine/Private/GeometryCollection/GeometryCollectionObject.cpp
--- a/Engine/Source/Runtime/Experimental/GeometryCollectionEngine/Private/GeometryCollection/GeometryCollectionObject.cpp
+++ b/Engine/Source/Runtime/Experimental/GeometryCollectionEngine/Private/GeometryCollection/GeometryCollectionObject.cpp
@@ -16,6 +16,7 @@
 
 #if WITH_EDITOR
 #include "GeometryCollection/DerivedDataGeometryCollectionCooker.h"
+#include "GeometryCollection/GeometryCollectionComponent.h"
 #include "DerivedDataCacheInterface.h"
 #include "Serialization/MemoryReader.h"
 #include "NaniteBuilder.h"
@@ -547,10 +548,15 @@ TArray<Nanite::FResources>& UGeometryCollection::CreateNaniteData(FGeometryColle
 		BuildIndices.Reserve(FaceCount * 3);
 		for (int32 FaceIndex = 0; FaceIndex < FaceCount; ++FaceIndex)
 		{
+			if (!VisibleArray[FaceStart + FaceIndex]) // TODO: Always in range?
+			{
+				continue;
+			}
+
 			const FIntVector FaceIndices = IndicesArray[FaceStart + FaceIndex];
-			BuildIndices.Add(FaceIndices.X - VertexStart);// -(FaceStart * 3));
-			BuildIndices.Add(FaceIndices.Y - VertexStart);// -(FaceStart * 3));
-			BuildIndices.Add(FaceIndices.Z - VertexStart);// -(FaceStart * 3));
+			BuildIndices.Add(FaceIndices.X - VertexStart);
+			BuildIndices.Add(FaceIndices.Y - VertexStart);
+			BuildIndices.Add(FaceIndices.Z - VertexStart);
 
 			// TODO: Cleanup super hack
 			FStaticMeshSection& BuildSection = BuildSections.Emplace_GetRef();
@@ -561,6 +567,12 @@ TArray<Nanite::FResources>& UGeometryCollection::CreateNaniteData(FGeometryColle
 			BuildSection.NumTriangles = 1;
 		}
 
+		if (BuildIndices.Num() == 0)
+		{
+			// No visible faces of entire geometry, skip any building/rendering.
+			continue;
+		}
+
 		FMeshNaniteSettings NaniteSettings = {};
 		NaniteSettings.bEnabled = true;
 		NaniteSettings.PercentTriangles = 1.0f; // 100% - no reduction
@@ -582,12 +594,15 @@ void UGeometryCollection::InitResources()
 
 	ReleaseResources();
 
-	for (Nanite::FResources& Resource : NaniteResources)
+	if (EnableNanite)
 	{
-		// Skip resources that have their render data stripped
-		if (Resource.PageStreamingStates.Num() > 0)
+		for (Nanite::FResources& Resource : NaniteResources)
 		{
-			Resource.InitResources();
+			// Skip resources that have their render data stripped
+			if (Resource.PageStreamingStates.Num() > 0)
+			{
+				Resource.InitResources();
+			}
 		}
 	}
 
@@ -628,10 +643,21 @@ FGuid UGeometryCollection::GetStateGuid() const
 #if WITH_EDITOR
 void UGeometryCollection::PostEditChangeProperty(struct FPropertyChangedEvent& PropertyChangedEvent)
 {
-	if (PropertyChangedEvent.Property && PropertyChangedEvent.Property->GetFName() != GET_MEMBER_NAME_CHECKED(UGeometryCollection, Materials))
+	if (PropertyChangedEvent.Property)
 	{
-		InvalidateCollection();
-		CreateSimulationData();
+		if (PropertyChangedEvent.Property->GetFName() != GET_MEMBER_NAME_CHECKED(UGeometryCollection, Materials))
+		{
+			InvalidateCollection();
+			CreateSimulationData();
+		}
+
+		if (PropertyChangedEvent.Property->GetFName() == GET_MEMBER_NAME_CHECKED(UGeometryCollection, EnableNanite))
+		{
+			if (FApp::CanEverRender())
+			{
+				InitResources();
+			}
+		}
 	}
 }
 
@@ -662,7 +688,7 @@ void UGeometryCollection::PostLoad()
 {
 	Super::PostLoad();
 
-	// initialize rendering resources
+	// Initialize rendering resources.
 	if (FApp::CanEverRender())
 	{
 		InitResources();
diff --git a/Engine/Source/Runtime/Experimental/GeometryCollectionEngine/Private/GeometryCollection/GeometryCollectionSceneProxy.cpp b/Engine/Source/Runtime/Experimental/GeometryCollectionEngine/Private/GeometryCollection/GeometryCollectionSceneProxy.cpp
--- a/Engine/Source/Runtime/Experimental/GeometryCollectionEngine/Private/GeometryCollection/GeometryCollectionSceneProxy.cpp
+++ b/Engine/Source/Runtime/Experimental/GeometryCollectionEngine/Private/GeometryCollection/GeometryCollectionSceneProxy.cpp
@@ -1253,20 +1253,25 @@ uint32 FNaniteGeometryCollectionSceneProxy::GetMemoryFootprint() const
 
 void FNaniteGeometryCollectionSceneProxy::SetConstantData_RenderThread(FGeometryCollectionConstantData* NewConstantData, bool ForceInit)
 {
-	const TManagedArray<int32>& TransformToGeometryIndices = GeometryCollection->GetGeometryCollection()->TransformToGeometryIndex;
-	check(NewConstantData->RestTransforms.Num() == TransformToGeometryIndices.Num());
+	const TSharedPtr<FGeometryCollection, ESPMode::ThreadSafe> Collection = GeometryCollection->GetGeometryCollection();
+	const TManagedArray<int32>& TransformToGeometryIndices = Collection->TransformToGeometryIndex;
 
+	check(NewConstantData->RestTransforms.Num() == TransformToGeometryIndices.Num());
 	Instances.Reset(NewConstantData->RestTransforms.Num());
 
 	for (int32 TransformIndex = 0; TransformIndex < NewConstantData->RestTransforms.Num(); ++TransformIndex)
 	{
 		const int32 TransformToGeometryIndex = TransformToGeometryIndices[TransformIndex];
-		if (TransformToGeometryIndex == INDEX_NONE)
+		if (!Collection->IsGeometry(TransformIndex))
 		{
-			// Not geometry
 			continue;
 		}
 
+		//if (Collection->IsClustered(TransformIndex))
+		//{
+		//	continue;
+		//}
+
 		const FGeometryNaniteData& NaniteData = GeometryNaniteData[TransformToGeometryIndex];
 
 		FPrimitiveInstance& Instance = Instances.Emplace_GetRef();
@@ -1286,29 +1291,35 @@ void FNaniteGeometryCollectionSceneProxy::SetConstantData_RenderThread(FGeometry
 
 void FNaniteGeometryCollectionSceneProxy::SetDynamicData_RenderThread(FGeometryCollectionDynamicData* NewDynamicData)
 {
-	//if (NewDynamicData->IsLoading)
-	{
-		// Ignored for now (rest transforms are already set)
-		//return;
-	}
-
 	// Are we currently simulating?
 	if (NewDynamicData->IsDynamic)
 	{
-		const TManagedArray<int32>& TransformToGeometryIndices = GeometryCollection->GetGeometryCollection()->TransformToGeometryIndex;
-		check(NewDynamicData->Transforms.Num() == TransformToGeometryIndices.Num());
+		const TSharedPtr<FGeometryCollection, ESPMode::ThreadSafe> Collection = GeometryCollection->GetGeometryCollection();
+		const TManagedArray<int32>& TransformToGeometryIndices = Collection->TransformToGeometryIndex;
+		const TManagedArray<TSet<int32>>& TransformChildren = Collection->Children;
 
+		check(NewDynamicData->Transforms.Num() == TransformToGeometryIndices.Num());
+		check(NewDynamicData->Transforms.Num() == TransformChildren.Num());
 		Instances.Reset(NewDynamicData->Transforms.Num());
 
 		for (int32 TransformIndex = 0; TransformIndex < NewDynamicData->Transforms.Num(); ++TransformIndex)
 		{
 			const int32 TransformToGeometryIndex = TransformToGeometryIndices[TransformIndex];
-			if (TransformToGeometryIndex == INDEX_NONE)
+			if (!Collection->IsGeometry(TransformIndex))
 			{
-				// Not geometry
 				continue;
 			}
 
+			//if (Collection->IsClustered(TransformIndex))
+			//{
+			//	continue;
+			//}
+
+			//if (TransformChildren[TransformIndex].Num() > 0)
+			//{
+			//	continue;
+			//}
+
 			const FGeometryNaniteData& NaniteData = GeometryNaniteData[TransformToGeometryIndex];
 
 			FPrimitiveInstance& Instance = Instances.Emplace_GetRef();
