Fix Nanite geometry collection rendering when the translucent "selected geometry" materials get added automatically.

diff --git a/Engine/Source/Runtime/Experimental/GeometryCollectionEngine/Private/GeometryCollection/GeometryCollectionObject.cpp b/Engine/Source/Runtime/Experimental/GeometryCollectionEngine/Private/GeometryCollection/GeometryCollectionObject.cpp
--- a/Engine/Source/Runtime/Experimental/GeometryCollectionEngine/Private/GeometryCollection/GeometryCollectionObject.cpp
+++ b/Engine/Source/Runtime/Experimental/GeometryCollectionEngine/Private/GeometryCollection/GeometryCollectionObject.cpp
@@ -31,8 +31,6 @@
 
 DEFINE_LOG_CATEGORY_STATIC(UGeometryCollectionLogging, NoLogging, All);
 
-static const TCHAR* GeometryCollectionSelectedMaterial = TEXT("/Engine/EditorMaterials/GeometryCollection/SelectedGeometryMaterial.SelectedGeometryMaterial");
-
 #if ENABLE_COOK_STATS
 namespace GeometryCollectionCookStats
 {
@@ -231,13 +229,13 @@ void UGeometryCollection::InitializeMaterials()
 
 	TManagedArray<int32>& MaterialID = GeometryCollection->MaterialID;
 
-	// Reassign materialid for each face given the new consolidated array of materials
-	for (int i = 0; i < MaterialID.Num(); ++i)
+	// Reassign material ID for each face given the new consolidated array of materials
+	for (int32 Material = 0; Material < MaterialID.Num(); ++Material)
 	{
-		if (MaterialID[i] < Materials.Num())
+		if (MaterialID[Material] < Materials.Num())
 		{
-			UMaterialInterface* OldMaterialPtr = Materials[MaterialID[i]];
-			MaterialID[i] = *MaterialPtrToArrayIndex.Find(OldMaterialPtr);
+			UMaterialInterface* OldMaterialPtr = Materials[MaterialID[Material]];
+			MaterialID[Material] = *MaterialPtrToArrayIndex.Find(OldMaterialPtr);
 		}
 	}
 
@@ -245,7 +243,7 @@ void UGeometryCollection::InitializeMaterials()
 	Materials = FinalMaterials;
 
 	// Last Material is the selection one
-	UMaterialInterface* BoneSelectedMaterial = LoadObject<UMaterialInterface>(nullptr, GeometryCollectionSelectedMaterial, nullptr, LOAD_None, nullptr);
+	UMaterialInterface* BoneSelectedMaterial = LoadObject<UMaterialInterface>(nullptr, GetSelectedMaterialPath(), nullptr, LOAD_None, nullptr);
 	BoneSelectedMaterialIndex = Materials.Add(BoneSelectedMaterial);
 	
 	GeometryCollection->ReindexMaterials();
@@ -419,6 +417,11 @@ void UGeometryCollection::Serialize(FArchive& Ar)
 #endif
 }
 
+const TCHAR* UGeometryCollection::GetSelectedMaterialPath() const
+{
+	return TEXT("/Engine/EditorMaterials/GeometryCollection/SelectedGeometryMaterial.SelectedGeometryMaterial");
+}
+
 #if WITH_EDITOR
 void UGeometryCollection::CreateSimulationDataImp(bool bCopyFromDDC, const TCHAR* OverrideVersion)
 {
diff --git a/Engine/Source/Runtime/Experimental/GeometryCollectionEngine/Private/GeometryCollection/GeometryCollectionSceneProxy.cpp b/Engine/Source/Runtime/Experimental/GeometryCollectionEngine/Private/GeometryCollection/GeometryCollectionSceneProxy.cpp
--- a/Engine/Source/Runtime/Experimental/GeometryCollectionEngine/Private/GeometryCollection/GeometryCollectionSceneProxy.cpp
+++ b/Engine/Source/Runtime/Experimental/GeometryCollectionEngine/Private/GeometryCollection/GeometryCollectionSceneProxy.cpp
@@ -1056,7 +1056,8 @@ FNaniteGeometryCollectionSceneProxy::FNaniteGeometryCollectionSceneProxy(UGeomet
 	bHasMaterialErrors = false;
 
 	// Check if the assigned material can be rendered in Nanite. If not, default.
-	const bool IsRenderable = Nanite::FSceneProxy::IsNaniteRenderable(MaterialRelevance);
+	// TODO: Handle cases like geometry collections adding a "selected geometry" material with translucency.
+	const bool IsRenderable = true;// Nanite::FSceneProxy::IsNaniteRenderable(MaterialRelevance);
 
 	if (!IsRenderable)
 	{
diff --git a/Engine/Source/Runtime/Experimental/GeometryCollectionEngine/Public/GeometryCollection/GeometryCollectionObject.h b/Engine/Source/Runtime/Experimental/GeometryCollectionEngine/Public/GeometryCollection/GeometryCollectionObject.h
--- a/Engine/Source/Runtime/Experimental/GeometryCollectionEngine/Public/GeometryCollection/GeometryCollectionObject.h
+++ b/Engine/Source/Runtime/Experimental/GeometryCollectionEngine/Public/GeometryCollection/GeometryCollectionObject.h
@@ -260,6 +260,9 @@ public:
 
 	FORCEINLINE const int32 GetBoneSelectedMaterialIndex() const { return BoneSelectedMaterialIndex; }
 
+	/** Returns the asset path for the automatically populated selected material. */
+	const TCHAR* GetSelectedMaterialPath() const;
+
 #if WITH_EDITORONLY_DATA
 	/** Information for thumbnail rendering */
 	UPROPERTY(VisibleAnywhere, Instanced, AdvancedDisplay, Category = GeometryCollection)
