Removed Chaos option to ignore dirty elements. Unhooked the usage from Reverb.

diff --git a/Engine/Source/Runtime/Experimental/Chaos/Private/Chaos/AABBTree.cpp b/Engine/Source/Runtime/Experimental/Chaos/Private/Chaos/AABBTree.cpp
--- a/Engine/Source/Runtime/Experimental/Chaos/Private/Chaos/AABBTree.cpp
+++ b/Engine/Source/Runtime/Experimental/Chaos/Private/Chaos/AABBTree.cpp
@@ -4,6 +4,3 @@
 
 int32 FAABBTreeCVars::UpdateDirtyElementPayloadData = 1;
 FAutoConsoleVariableRef FAABBTreeCVars::CVarUpdateDirtyElementPayloadData(TEXT("p.aabbtree.updatedirtyelementpayloads"), FAABBTreeCVars::UpdateDirtyElementPayloadData, TEXT("Allow AABB tree elements to update internal payload data when they recieve a payload update"));
-
-int32 FAABBTreeCVars::IgnoreDirtyElements = 0;
-FAutoConsoleVariableRef FAABBTreeCVars::CVarIgnoreDirtyElements(TEXT("p.aabbtree.IgnoreDirtyElements"), FAABBTreeCVars::IgnoreDirtyElements, TEXT("Ignore AABB tree dirty elements"));
\ No newline at end of file
diff --git a/Engine/Source/Runtime/Experimental/Chaos/Public/Chaos/AABBTree.h b/Engine/Source/Runtime/Experimental/Chaos/Public/Chaos/AABBTree.h
--- a/Engine/Source/Runtime/Experimental/Chaos/Public/Chaos/AABBTree.h
+++ b/Engine/Source/Runtime/Experimental/Chaos/Public/Chaos/AABBTree.h
@@ -13,9 +13,7 @@
 struct FAABBTreeCVars
 {
 	static int32 UpdateDirtyElementPayloadData;
-	static CHAOS_API int32 IgnoreDirtyElements;
 	static FAutoConsoleVariableRef CVarUpdateDirtyElementPayloadData;
-	static FAutoConsoleVariableRef CVarIgnoreDirtyElements;
 };
 
 namespace Chaos
@@ -710,7 +708,7 @@ private:
 			}
 		}
 
-		if (bMutable && FAABBTreeCVars::IgnoreDirtyElements == 0)
+		if (bMutable)
 		{
 			//QUICK_SCOPE_CYCLE_COUNTER(QueryDirty);
 
