Chaos Cloth - LOD/graph node refactor.

diff --git a/Engine/Source/Runtime/ClothingSystemRuntimeNv/Public/ClothingSimulationNv.h b/Engine/Source/Runtime/ClothingSystemRuntimeNv/Public/ClothingSimulationNv.h
--- a/Engine/Source/Runtime/ClothingSystemRuntimeNv/Public/ClothingSimulationNv.h
+++ b/Engine/Source/Runtime/ClothingSystemRuntimeNv/Public/ClothingSimulationNv.h
@@ -39,18 +39,12 @@ public:
 	// Override the fill context function to also set the Nv specific simulation context members
 	virtual void Fill(const USkeletalMeshComponent* InComponent, float InDeltaSeconds, float InMaxPhysicsDelta) override;
 
-	// Set the RefToLocals array in the parent class using Nv specific predicted LOD information
-	virtual void FillRefToLocals(const USkeletalMeshComponent* InComponent) override;
-
 	// Set the world gravity in the parent class while preserving the Nv legacy code behavior
 	void virtual FillWorldGravity(const USkeletalMeshComponent* InComponent) override;
 
 	// Set WindVelocity in the parent class and Nv specific WindAdaption
 	virtual void FillWindVelocity(const USkeletalMeshComponent* InComponent) override;
 
-	// The predicted LOD of the skeletal mesh component running the simulation
-	int32 PredictedLod;
-
 	// Wind adaption, a measure of how quickly to adapt to the wind speed
 	// when using the legacy wind calculation mode
 	float WindAdaption;
@@ -198,7 +192,6 @@ public:
 
 	// IClothingSimulation Interface
 	virtual void CreateActor(USkeletalMeshComponent* InOwnerComponent, UClothingAssetBase* InAsset, int32 InSimDataIndex) override;
-	virtual void PostActorCreationInitialize() override {};
 	virtual IClothingSimulationContext* CreateContext() override;
 	virtual void Initialize() override;
 	virtual void Shutdown() override;
diff --git a/Engine/Source/Runtime/Engine/Classes/Components/SkeletalMeshComponent.h b/Engine/Source/Runtime/Engine/Classes/Components/SkeletalMeshComponent.h
--- a/Engine/Source/Runtime/Engine/Classes/Components/SkeletalMeshComponent.h
+++ b/Engine/Source/Runtime/Engine/Classes/Components/SkeletalMeshComponent.h
@@ -1265,6 +1265,9 @@ public:
 	/** Get the current clothing simulation (read only) */
 	const IClothingSimulation* GetClothingSimulation() const;
 
+	/** Get the current clothing simulation context (read only) */
+	const IClothingSimulationContext* GetClothingSimulationContext() const;
+
 	/** Get the current interactor for a clothing simulation, if the current simulation supports runtime interaction. */
 	UFUNCTION(BlueprintCallable, Category=ClothingSimulation)
 	UClothingSimulationInteractor* GetClothingSimulationInteractor() const;
diff --git a/Engine/Source/Runtime/Engine/Private/Components/SkeletalMeshComponent.cpp b/Engine/Source/Runtime/Engine/Private/Components/SkeletalMeshComponent.cpp
--- a/Engine/Source/Runtime/Engine/Private/Components/SkeletalMeshComponent.cpp
+++ b/Engine/Source/Runtime/Engine/Private/Components/SkeletalMeshComponent.cpp
@@ -1907,6 +1907,11 @@ const IClothingSimulation* USkeletalMeshComponent::GetClothingSimulation() const
 	return ClothingSimulation;
 }
 
+const IClothingSimulationContext* USkeletalMeshComponent::GetClothingSimulationContext() const
+{
+	return ClothingSimulationContext;
+}
+
 UClothingSimulationInteractor* USkeletalMeshComponent::GetClothingSimulationInteractor() const
 {
 	return ClothingInteractor;
diff --git a/Engine/Source/Runtime/Engine/Private/SkeletalMeshComponentPhysics.cpp b/Engine/Source/Runtime/Engine/Private/SkeletalMeshComponentPhysics.cpp
--- a/Engine/Source/Runtime/Engine/Private/SkeletalMeshComponentPhysics.cpp
+++ b/Engine/Source/Runtime/Engine/Private/SkeletalMeshComponentPhysics.cpp
@@ -2443,6 +2443,8 @@ void USkeletalMeshComponent::RecreateClothingActors()
 				TArray<UClothingAssetBase*> AssetsInUse;
 				SkeletalMesh->GetClothingAssetsInUse(AssetsInUse);
 
+				ClothingSimulation->FillContext(this, 0.f, ClothingSimulationContext);
+
 				const int32 NumMeshAssets = SkeletalMesh->MeshClothingAssets.Num();
 				for (int32 BaseAssetIndex = 0; BaseAssetIndex < NumMeshAssets; ++BaseAssetIndex)
 				{
@@ -2452,7 +2454,6 @@ void USkeletalMeshComponent::RecreateClothingActors()
 						ClothingSimulation->CreateActor(this, Asset, BaseAssetIndex);
 					}
 				}
-				ClothingSimulation->PostActorCreationInitialize();
 			}
 
 			// Retrieve the cloth sim data, or clear the data if the world isn't ready to sim
@@ -2531,9 +2532,9 @@ void USkeletalMeshComponent::GetWindForCloth_GameThread(FVector& WindDirection,
 #if WITH_CLOTH_COLLISION_DETECTION
 
 void USkeletalMeshComponent::FindClothCollisions(FClothCollisionData& OutCollisions)
-				{
+{
 	if(ClothingSimulation)
-					{
+	{
 		// Get collisions for this simulation, ignoring any externally added collisions
 		// (i.e. on grab the asset collisions, not environment etc.)
 		ClothingSimulation->GetCollisions(OutCollisions, false);
@@ -2643,8 +2644,8 @@ void USkeletalMeshComponent::CopyClothCollisionsToChildren()
 	{
 		USkeletalMeshComponent* pChild = Cast<USkeletalMeshComponent>(AttachedChild);
 		if(pChild && pChild->ClothingSimulation->ShouldSimulate())
-			{
-				ClothChildren.Add(pChild);
+		{
+			ClothChildren.Add(pChild);
 		}
 	}
 
diff --git a/Engine/Source/Runtime/Experimental/Chaos/Private/Chaos/PBDCollisionSpringConstraintsBase.cpp b/Engine/Source/Runtime/Experimental/Chaos/Private/Chaos/PBDCollisionSpringConstraintsBase.cpp
--- a/Engine/Source/Runtime/Experimental/Chaos/Private/Chaos/PBDCollisionSpringConstraintsBase.cpp
+++ b/Engine/Source/Runtime/Experimental/Chaos/Private/Chaos/PBDCollisionSpringConstraintsBase.cpp
@@ -54,9 +54,11 @@ class FMeshBuildDataProvider
 
 template<class T, int d>
 PBDCollisionSpringConstraintsBase<T, d>::PBDCollisionSpringConstraintsBase(
-    const TDynamicParticles<T, d>& InParticles, const TArray<TVector<int32, 3>>& Elements, const TSet<TVector<int32, 2>>& DisabledCollisionElements, const TArray<uint32>& DynamicGroupIds, const TArray<T>& PerGroupThicknesses, const T Dt, const T Stiffness)
+    const TPBDActiveView<TPBDParticles<T, d>>& ParticlesActiveView,  const TArray<TVector<int32, 3>>& Elements, const TSet<TVector<int32, 2>>& DisabledCollisionElements, const TArray<uint32>& DynamicGroupIds, const TArray<T>& PerGroupThicknesses, const T Dt, const T Stiffness)
     : MDynamicGroupIds(DynamicGroupIds), MPerGroupThicknesses(PerGroupThicknesses), MStiffness(Stiffness)
 {
+	const TDynamicParticles<T, d>& Particles = ParticlesActiveView.GetItems();
+
 	if (!Elements.Num())
 		return;
 #if PLATFORM_DESKTOP && PLATFORM_64BITS
@@ -69,53 +71,56 @@ PBDCollisionSpringConstraintsBase<T, d>::PBDCollisionSpringConstraintsBase(
 	const T Height = MaxThickness + MaxThickness;
 
 	TkDOPTree<const FMeshBuildDataProvider, uint32> DopTree;
-	TArray<FkDOPBuildCollisionTriangle<uint32>> BuildTraingleArray;
+	TArray<FkDOPBuildCollisionTriangle<uint32>> BuildTriangleArray;
 	for (int32 i = 0; i < Elements.Num(); ++i)
 	{
 		const auto& Elem = Elements[i];
-		BuildTraingleArray.Add(FkDOPBuildCollisionTriangle<uint32>(i, InParticles.X(Elem[0]), InParticles.X(Elem[1]), InParticles.X(Elem[2])));
+		BuildTriangleArray.Add(FkDOPBuildCollisionTriangle<uint32>(i, Particles.X(Elem[0]), Particles.X(Elem[1]), Particles.X(Elem[2])));
 	}
-	DopTree.Build(BuildTraingleArray);
+	DopTree.Build(BuildTriangleArray);
 	FMeshBuildDataProvider DopDataProvider(DopTree);
 	FCriticalSection CriticalSection;
-	PhysicsParallelFor(InParticles.Size(), [&](int32 Index) {
-		FkHitResult Result;
-		const auto& Start = InParticles.X(Index);
-		const auto End = Start + InParticles.V(Index) * Dt + InParticles.V(Index).GetSafeNormal() * Height;
-		FVector4 Start4(Start.X, Start.Y, Start.Z, 0);
-		FVector4 End4(End.X, End.Y, End.Z, 0);
-		TkDOPLineCollisionCheck<const FMeshBuildDataProvider, uint32> ray(Start4, End4, true, DopDataProvider, &Result);
-		if (DopTree.LineCheck(ray))
+
+	ParticlesActiveView.ParallelFor(
+		[this, &Elements, &DisabledCollisionElements, &CriticalSection, Dt, Height, &DopDataProvider, &DopTree](TPBDParticles<T, d>& Particles, int32 Index)
 		{
-			const auto& Elem = Elements[Result.Item];
-			if (DisabledCollisionElements.Contains({Index, Elem[0]}) || DisabledCollisionElements.Contains({Index, Elem[1]}) || DisabledCollisionElements.Contains({Index, Elem[2]}))
-				return;
-			TVector<T, 3> Bary;
-			TVector<T, d> P10 = InParticles.X(Elem[1]) - InParticles.X(Elem[0]);
-			TVector<T, d> P20 = InParticles.X(Elem[2]) - InParticles.X(Elem[0]);
-			TVector<T, d> PP0 = InParticles.X(Index) - InParticles.X(Elem[0]);
-			T Size10 = P10.SizeSquared();
-			T Size20 = P20.SizeSquared();
-			T ProjSides = TVector<T, d>::DotProduct(P10, P20);
-			T ProjP1 = TVector<T, d>::DotProduct(PP0, P10);
-			T ProjP2 = TVector<T, d>::DotProduct(PP0, P20);
-			T Denom = Size10 * Size20 - ProjSides * ProjSides;
-			Bary.Y = (Size20 * ProjP1 - ProjSides * ProjP2) / Denom;
-			Bary.Z = (Size10 * ProjP2 - ProjSides * ProjP1) / Denom;
-			Bary.X = 1.0f - Bary.Z - Bary.Y;
-			// TODO(mlentine): In theory this shouldn't happen as this means no collision but I'm not positive that the collision check is always correct.
-			//if (Bary.Y < 0 || Bary.Z < 0 || Bary.X < 0) return;
-			// TODO(mlentine): Incorporate history
-			TVector<T, d> Normal = TVector<T, d>::CrossProduct(P10, P20);
-			Normal.Normalize();
-			Normal = TVector<T, d>::DotProduct(Normal, PP0) > 0 ? Normal : -Normal;
-			CriticalSection.Lock();
-			MConstraints.Add({Index, Elem[0], Elem[1], Elem[2]});
-			MBarys.Add(Bary);
-			MNormals.Add(Normal);
-			CriticalSection.Unlock();
-		}
-	});
+			FkHitResult Result;
+			const auto& Start = Particles.X(Index);
+			const auto End = Start + Particles.V(Index) * Dt + Particles.V(Index).GetSafeNormal() * Height;
+			FVector4 Start4(Start.X, Start.Y, Start.Z, 0);
+			FVector4 End4(End.X, End.Y, End.Z, 0);
+			TkDOPLineCollisionCheck<const FMeshBuildDataProvider, uint32> Ray(Start4, End4, true, DopDataProvider, &Result);
+			if (DopTree.LineCheck(Ray))
+			{
+				const auto& Elem = Elements[Result.Item];
+				if (DisabledCollisionElements.Contains({Index, Elem[0]}) || DisabledCollisionElements.Contains({Index, Elem[1]}) || DisabledCollisionElements.Contains({Index, Elem[2]}))
+					return;
+				TVector<T, 3> Bary;
+				TVector<T, d> P10 = Particles.X(Elem[1]) - Particles.X(Elem[0]);
+				TVector<T, d> P20 = Particles.X(Elem[2]) - Particles.X(Elem[0]);
+				TVector<T, d> PP0 = Particles.X(Index) - Particles.X(Elem[0]);
+				T Size10 = P10.SizeSquared();
+				T Size20 = P20.SizeSquared();
+				T ProjSides = TVector<T, d>::DotProduct(P10, P20);
+				T ProjP1 = TVector<T, d>::DotProduct(PP0, P10);
+				T ProjP2 = TVector<T, d>::DotProduct(PP0, P20);
+				T Denom = Size10 * Size20 - ProjSides * ProjSides;
+				Bary.Y = (Size20 * ProjP1 - ProjSides * ProjP2) / Denom;
+				Bary.Z = (Size10 * ProjP2 - ProjSides * ProjP1) / Denom;
+				Bary.X = 1.0f - Bary.Z - Bary.Y;
+				// TODO(mlentine): In theory this shouldn't happen as this means no collision but I'm not positive that the collision check is always correct.
+				//if (Bary.Y < 0 || Bary.Z < 0 || Bary.X < 0) return;
+				// TODO(mlentine): Incorporate history
+				TVector<T, d> Normal = TVector<T, d>::CrossProduct(P10, P20);
+				Normal.Normalize();
+				Normal = TVector<T, d>::DotProduct(Normal, PP0) > 0 ? Normal : -Normal;
+				CriticalSection.Lock();
+				MConstraints.Add({Index, Elem[0], Elem[1], Elem[2]});
+				MBarys.Add(Bary);
+				MNormals.Add(Normal);
+				CriticalSection.Unlock();
+			}
+		});
 #endif
 }
 
diff --git a/Engine/Source/Runtime/Experimental/Chaos/Private/Chaos/PBDEvolution.cpp b/Engine/Source/Runtime/Experimental/Chaos/Private/Chaos/PBDEvolution.cpp
--- a/Engine/Source/Runtime/Experimental/Chaos/Private/Chaos/PBDEvolution.cpp
+++ b/Engine/Source/Runtime/Experimental/Chaos/Private/Chaos/PBDEvolution.cpp
@@ -14,7 +14,6 @@
 #include "Chaos/PerParticlePBDUpdateFromDeltaPosition.h"
 #include "ChaosStats.h"
 
-
 DECLARE_CYCLE_STAT(TEXT("Chaos PBD Advance Time"), STAT_ChaosPBDVAdvanceTime, STATGROUP_Chaos);
 DECLARE_CYCLE_STAT(TEXT("Chaos PBD Velocity Damping State Update"), STAT_ChaosPBDVelocityDampUpdateState, STATGROUP_Chaos);
 DECLARE_CYCLE_STAT(TEXT("Chaos PBD Velocity Field Update Forces"), STAT_ChaosPBDVelocityFieldUpdateForces, STATGROUP_Chaos);
@@ -29,11 +28,11 @@ DECLARE_CYCLE_STAT(TEXT("Chaos PBD Collider Kinematic Update"), STAT_CollisionKi
 using namespace Chaos;
 
 template<class T, int d>
-void TPBDEvolution<T, d>::AddGroups(int32 Num)
+void TPBDEvolution<T, d>::AddGroups(int32 NumGroups)
 {
 	// Add elements
 	const uint32 Offset = TArrayCollection::Size();
-	TArrayCollection::AddElementsHelper(Num);
+	TArrayCollection::AddElementsHelper(NumGroups);
 
 	// Set defaults
 	for (uint32 GroupId = Offset; GroupId < TArrayCollection::Size(); ++GroupId)
@@ -43,16 +42,26 @@ void TPBDEvolution<T, d>::AddGroups(int32 Num)
 		MGroupSelfCollisionThicknesses[GroupId] = MSelfCollisionThickness;
 		MGroupCoefficientOfFrictions[GroupId] = MCoefficientOfFriction;
 		MGroupDampings[GroupId] = MDamping;
-		MGroupVelocityFields[GroupId] = MakeUnique<TVelocityField<float, 3>>();
 	}
 }
 
+template<class T, int d>
+void TPBDEvolution<T, d>::ResetGroups()
+{
+	TArrayCollection::ResizeHelper(0);
+	AddGroups(1);  // Add default group
+}
+
 template<class T, int d>
 TPBDEvolution<T, d>::TPBDEvolution(TPBDParticles<T, d>&& InParticles, TKinematicGeometryClothParticles<T, d>&& InGeometryParticles, TArray<TVector<int32, 3>>&& CollisionTriangles,
     int32 NumIterations, T CollisionThickness, T SelfCollisionThickness, T CoefficientOfFriction, T Damping)
     : MParticles(MoveTemp(InParticles))
+	, MParticlesActiveView(MParticles)
 	, MCollisionParticles(MoveTemp(InGeometryParticles))
+	, MCollisionParticlesActiveView(MCollisionParticles)
 	, MCollisionTriangles(MoveTemp(CollisionTriangles))
+	, MConstraintInitsActiveView(MConstraintInits)
+	, MConstraintRulesActiveView(MConstraintRules)
 	, MNumIterations(NumIterations)
 	, MGravity(TVector<T, d>((T)0., (T)0., (T)-980.665))
 	, MCollisionThickness(CollisionThickness)
@@ -75,53 +84,110 @@ TPBDEvolution<T, d>::TPBDEvolution(TPBDParticles<T, d>&& InParticles, TKinematic
 	MCollisionParticles.AddArray(&MCollided);
 	MCollisionParticles.AddArray(&MCollisionParticleGroupIds);
 
-	SetParticleUpdateFunction(
+	MParticleUpdate =  // TODO(Kriss.Gossart): this callable seems redundant, might be worth taking it off
 		[PBDUpdateRule = 
-			TPerParticlePBDUpdateFromDeltaPosition<float, 3>()](TPBDParticles<T, d>& MParticlesInput, const T Dt) 
+			TPerParticlePBDUpdateFromDeltaPosition<T, d>()](TPBDActiveView<TPBDParticles<T, d>>& ParticlesView, const T Dt) 
 			{
 				// Don't bother with threaded execution if we don't have enough work to make it worth while.
-				const bool NonParallelUpdate = MParticlesInput.Size() > 1000; // TODO: 1000 is a guess, tune this!
-				PhysicsParallelFor(MParticlesInput.Size(), [&](int32 Index) {
-					PBDUpdateRule.Apply(MParticlesInput, Dt, Index);
-				}, NonParallelUpdate);
-			});
+				const int32 MinParallelBatchSize = 1000; // TODO: 1000 is a guess, tune this!
+
+				ParticlesView.ParallelFor(
+					[PBDUpdateRule, Dt](TPBDParticles<T, d>& Particles, int32 Index)
+					{
+						PBDUpdateRule.Apply(Particles, Dt, Index);
+					}, MinParallelBatchSize);
+			};
 }
 
 template<class T, int d>
-uint32 TPBDEvolution<T, d>::AddParticles(uint32 Num, uint32 GroupId)
+void TPBDEvolution<T, d>::ResetParticles()
 {
-	// Add new particles
-	const uint32 Offset = MParticles.Size();
-	MParticles.AddParticles(Num);
+	// Reset particles
+	MParticles.Resize(0);
+	MParticlesActiveView.Reset();
 
-	// Initialize the new particles' group ids
-	for (uint32 i = Offset; i < MParticles.Size(); ++i)
-	{
-		MParticleGroupIds[i] = GroupId;
-	}
+	// Reset particle groups
+	ResetGroups();
+}
 
-	// Resize group parameter arrays
-	const uint32 GroupSize = TArrayCollection::Size();
-	if (GroupId >= GroupSize)
+template<class T, int d>
+int32 TPBDEvolution<T, d>::AddParticleRange(int32 NumParticles, uint32 GroupId, bool bActivate)
+{
+	if (NumParticles)
 	{
-		AddGroups(GroupId + 1 - GroupSize);
+		const int32 Offset = (int32)MParticles.Size();
+
+		MParticles.AddParticles(NumParticles);
+
+		// Initialize the new particles' group ids
+		for (int32 i = Offset; i < (int32)MParticles.Size(); ++i)
+		{
+			MParticleGroupIds[i] = GroupId;
+		}
+
+		// Resize the group parameter arrays
+		const uint32 GroupSize = TArrayCollection::Size();
+		if (GroupId >= GroupSize)
+		{
+			AddGroups(GroupId + 1 - GroupSize);
+		}
+
+		// Add range
+		MParticlesActiveView.AddRange(NumParticles, bActivate);
+
+		return Offset;
 	}
-	return Offset;
+	return INDEX_NONE;
 }
 
 template<class T, int d>
-uint32 TPBDEvolution<T, d>::AddCollisionParticles(uint32 Num, uint32 GroupId)
+void TPBDEvolution<T, d>::ResetCollisionParticles(int32 NumParticles)
 {
-	// Add new particles
-	const uint32 Offset = MCollisionParticles.Size();
-	MCollisionParticles.AddParticles(Num);
+	MCollisionParticles.Resize(NumParticles);
+	MCollisionParticlesActiveView.Reset(NumParticles);
+}
 
-	// Initialize the new particles' group ids
-	for (uint32 i = Offset; i < MCollisionParticles.Size(); ++i)
+template<class T, int d>
+int32 TPBDEvolution<T, d>::AddCollisionParticleRange(int32 NumParticles, uint32 GroupId, bool bActivate)
+{
+	if (NumParticles)
 	{
-		MCollisionParticleGroupIds[i] = GroupId;
+		const int32 RangeOffset = (int32)MCollisionParticles.Size();
+
+		MCollisionParticles.AddParticles(NumParticles);
+
+		// Initialize the new particles' group ids
+		for (int32 i = RangeOffset; i < (int32)MCollisionParticles.Size(); ++i)
+		{
+			MCollisionParticleGroupIds[i] = GroupId;
+		}
+
+		// Add range
+		MCollisionParticlesActiveView.AddRange(NumParticles, bActivate);
+	
+		return RangeOffset;
 	}
-	return Offset;
+	return INDEX_NONE;
+}
+
+template<class T, int d>
+int32 TPBDEvolution<T, d>::AddConstraintInitRange(int32 NumConstraints, bool bActivate)
+{
+	// Add new constraint init functions
+	MConstraintInits.AddDefaulted(NumConstraints);
+
+	// Add range
+	return MConstraintInitsActiveView.AddRange(NumConstraints, bActivate);
+}
+
+template<class T, int d>
+int32 TPBDEvolution<T, d>::AddConstraintRuleRange(int32 NumConstraints, bool bActivate)
+{
+	// Add new constraint rule functions
+	MConstraintRules.AddDefaulted(NumConstraints);
+
+	// Add range
+	return MConstraintRulesActiveView.AddRange(NumConstraints, bActivate);
 }
 
 template<class T, int d>
@@ -136,70 +202,76 @@ void TPBDEvolution<T, d>::AdvanceOneTimeStep(const T Dt)
 
 	{
 		SCOPE_CYCLE_COUNTER(STAT_ChaosPBDVelocityDampUpdateState);
-		DampVelocityRule.UpdatePositionBasedState(MParticles);
+		DampVelocityRule.UpdatePositionBasedState(MParticlesActiveView);
 	}	
 
 	{
 		SCOPE_CYCLE_COUNTER(STAT_ChaosPBDVelocityFieldUpdateForces);
-		for (const TUniquePtr<FVelocityField>& VelocityField : MGroupVelocityFields)
+		for (const FVelocityField& VelocityField : MGroupVelocityFields)
 		{
-			VelocityField->UpdateForces(MParticles, Dt);
+			VelocityField.UpdateForces(MParticles, Dt);  // Update force per surface element
 		}
 	}	
 
+	memset(MCollided.GetData(), 0, MCollided.Num() * sizeof(bool));
+
 	// Don't bother with threaded execution if we don't have enough work to make it worth while.
-	const bool NonParallelUpdate = MParticles.Size() < 1000; // TODO: 1000 is a guess, tune this!
+	const int32 MinParallelBatchSize = 1000; // TODO: 1000 is a guess, tune this!
 
-	//PhysicsParallelFor(MCollisionParticles.Size(), [&](int32 Index) 
-	//{ MCollided[Index] = false; }, NonParallelUpdate);
-	//for(int32 i=0; i < MCollided.Num(); i++)
-	//	MCollided[i] = false;
-	memset(MCollided.GetData(), 0, MCollided.Num()*sizeof(bool));
 	{
 		SCOPE_CYCLE_COUNTER(STAT_ChaosPBDPreIterationUpdates);
-		PhysicsParallelFor(MParticles.Size(), [&](int32 Index)
-		{
-			const uint32 ParticleGroupId = MParticleGroupIds[Index];
 
-			InitForceRule.Apply(MParticles, Dt, Index); // F = TV(0)
-			MGroupGravityForces[ParticleGroupId].Apply(MParticles, Dt, Index); // F += M * G
-			for (TFunction<void(TPBDParticles<T, d>&, const T, const int32)>& ForceRule : MForceRules)
+		MParticlesActiveView.ParallelFor(
+			[this, Dt, &InitForceRule, &EulerStepVelocityRule, &DampVelocityRule, &EulerStepRule](TPBDParticles<T, d>& Particles, int32 Index)
 			{
-				ForceRule(MParticles, Dt, Index); // F += M * A
-			}
+				const uint32 ParticleGroupId = MParticleGroupIds[Index];
+
+				InitForceRule.Apply(Particles, Dt, Index); // F = TV(0)
+				MGroupGravityForces[ParticleGroupId].Apply(Particles, Dt, Index); // F += M * G
+				for (TFunction<void(TPBDParticles<T, d>&, const T, const int32)>& ForceRule : MForceRules)
+				{
+					ForceRule(Particles, Dt, Index); // F += M * A
+				}
 			
-			MGroupVelocityFields[ParticleGroupId]->Apply(MParticles, Dt, Index);
+				MGroupVelocityFields[ParticleGroupId].Apply(Particles, Dt, Index);
 
-			if (MKinematicUpdate)
-			{
-				MKinematicUpdate(MParticles, Dt, MTime + Dt, Index); // X = ...
-			}
-			EulerStepVelocityRule.Apply(MParticles, Dt, Index);
-			DampVelocityRule.Apply(MParticles, Dt, Index);
-			EulerStepRule.Apply(MParticles, Dt, Index);
-		}, NonParallelUpdate);
+				if (MKinematicUpdate)
+				{
+					MKinematicUpdate(Particles, Dt, MTime + Dt, Index); // X = ...
+				}
+				EulerStepVelocityRule.Apply(Particles, Dt, Index);
+				DampVelocityRule.Apply(Particles, Dt, Index);
+				EulerStepRule.Apply(Particles, Dt, Index);
+			}, MinParallelBatchSize);
 	}
 
 	if (MCollisionKinematicUpdate)
 	{
 		SCOPE_CYCLE_COUNTER(STAT_CollisionKinematicUpdate);
-		PhysicsParallelFor(MCollisionParticles.Size(), [&](int32 Index) {
-			MCollisionKinematicUpdate(MCollisionParticles, Dt, MTime + Dt, Index);
-		}, NonParallelUpdate);
+
+		MCollisionParticlesActiveView.SequentialFor(
+			[this, Dt](TKinematicGeometryClothParticles<T, d>& CollisionParticles, int32 Index)
+			{
+				MCollisionKinematicUpdate(CollisionParticles, Dt, MTime + Dt, Index);
+			});
 	}
 #if !COMPILE_WITHOUT_UNREAL_SUPPORT
-	TPBDCollisionSpringConstraints<T, d> SelfCollisionRule(MParticles, MCollisionTriangles, MDisabledCollisionElements, MParticleGroupIds, MGroupSelfCollisionThicknesses, Dt);
+	TPBDCollisionSpringConstraints<T, d> SelfCollisionRule(MParticlesActiveView, MCollisionTriangles, MDisabledCollisionElements, MParticleGroupIds, MGroupSelfCollisionThicknesses, Dt);
 #endif
 
-	for (TFunction<void()>& InitConstraintRule : MInitConstraintRules)
+	for (TFunction<void()>& ConstraintInit : MConstraintInits)
 	{
-		InitConstraintRule();  // Clear XPBD's Lambdas
+		ConstraintInit();  // Clear XPBD's Lambdas
 	}
 
 	// Do one extra collision pass at the start to decrease likelihood of cloth penetrating- TODO: Add option for more collision passed interleaved between constraints
 	{
 		SCOPE_CYCLE_COUNTER(STAT_ChaosPBDCollisionRule);
-		CollisionRule.ApplyPerParticle(MParticles, Dt);
+		MParticlesActiveView.ParallelFor(
+			[&CollisionRule, Dt](TPBDParticles<T, d>& Particles, int32 Index)
+			{
+				CollisionRule.Apply(Particles, Dt, Index);
+			});
 	}
 
 	for (int i = 0; i < MNumIterations; ++i)
@@ -217,17 +289,24 @@ void TPBDEvolution<T, d>::AdvanceOneTimeStep(const T Dt)
 #endif
 		{
 			SCOPE_CYCLE_COUNTER(STAT_ChaosPBDCollisionRule);
-			CollisionRule.ApplyPerParticle(MParticles, Dt);
+			MParticlesActiveView.ParallelFor(
+				[&CollisionRule, Dt](TPBDParticles<T, d>& Particles, int32 Index)
+				{
+					CollisionRule.Apply(Particles, Dt, Index);
+				},MinParallelBatchSize);
 		}
 	}
 	check(MParticleUpdate);
-	MParticleUpdate(MParticles, Dt); // V = (P - X) / Dt; X = P;
+	MParticleUpdate(MParticlesActiveView, Dt); // V = (P - X) / Dt; X = P;
+
 	if (MCoefficientOfFriction > 0)
 	{
 		SCOPE_CYCLE_COUNTER(STAT_ChaosPBDCollisionRuleFriction);
-		PhysicsParallelFor(MParticles.Size(), [&](int32 Index) {
-			CollisionRule.ApplyFriction(MParticles, Dt, Index);
-		}, NonParallelUpdate);
+		MParticlesActiveView.ParallelFor(
+			[&CollisionRule, Dt](TPBDParticles<T, d>& Particles, int32 Index)
+			{
+				CollisionRule.ApplyFriction(Particles, Dt, Index);
+			}, MinParallelBatchSize);
 	}
 
 	MTime += Dt;
diff --git a/Engine/Source/Runtime/Experimental/Chaos/Private/Chaos/TriangleMesh.cpp b/Engine/Source/Runtime/Experimental/Chaos/Private/Chaos/TriangleMesh.cpp
--- a/Engine/Source/Runtime/Experimental/Chaos/Private/Chaos/TriangleMesh.cpp
+++ b/Engine/Source/Runtime/Experimental/Chaos/Private/Chaos/TriangleMesh.cpp
@@ -105,9 +105,9 @@ void TTriangleMesh<T>::ResetAuxiliaryStructures()
 }
 
 template<class T>
-TPair<int32, int32> TTriangleMesh<T>::GetVertexRange() const
+TVector<int32, 2> TTriangleMesh<T>::GetVertexRange() const
 {
-	return TPair<int32, int32>(MStartIdx, MStartIdx + MNumIndices - 1);
+	return TVector<int32, 2>(MStartIdx, MStartIdx + MNumIndices - 1);
 }
 
 template<class T>
@@ -158,21 +158,20 @@ const TMap<int32, TSet<uint32>>& TTriangleMesh<T>::GetPointToNeighborsMap() cons
 }
 
 template<class T>
-const TMap<int32, TArray<int32>>& TTriangleMesh<T>::GetPointToTriangleMap() const
+TConstArrayView<TArray<int32>> TTriangleMesh<T>::GetPointToTriangleMap() const
 {
-	if (MPointToTriangleMap.Num())
+	if (!MPointToTriangleMap.Num())
 	{
-		return MPointToTriangleMap;
-	}
-	MPointToTriangleMap.Reserve(MNumIndices);
-	for (int i = 0; i < MElements.Num(); ++i)
-	{
-		for (int Axis = 0; Axis < 3; ++Axis)
+		MPointToTriangleMap.AddDefaulted(MNumIndices);
+		for (int i = 0; i < MElements.Num(); ++i)
 		{
-			MPointToTriangleMap.FindOrAdd(MElements[i][Axis]).Add(i);
+			for (int Axis = 0; Axis < 3; ++Axis)
+			{
+				MPointToTriangleMap[MElements[i][Axis] - MStartIdx].Add(i);  // Access MPointToTriangleMap with local index
+			}
 		}
 	}
-	return MPointToTriangleMap;
+	return TConstArrayView<TArray<int32>>(MPointToTriangleMap.GetData() - MStartIdx, MStartIdx + MNumIndices);  // Return an array view that is using global indexation
 }
 
 template<class T>
@@ -196,7 +195,7 @@ TArray<TVector<int32, 4>> TTriangleMesh<T>::GetUniqueAdjacentElements() const
 	for (int32 SurfaceIndex = MStartIdx; SurfaceIndex < MStartIdx + MNumIndices; ++SurfaceIndex)
 	{
 		TMap<int32, TArray<int32>> SubPointToTriangleMap;
-		for (auto TriangleIndex : MPointToTriangleMap[SurfaceIndex])
+		for (auto TriangleIndex : MPointToTriangleMap[SurfaceIndex - MStartIdx])  // Access MPointToTriangleMap with local index
 		{
 			SubPointToTriangleMap.FindOrAdd(MElements[TriangleIndex][0]).Add(TriangleIndex);
 			SubPointToTriangleMap.FindOrAdd(MElements[TriangleIndex][1]).Add(TriangleIndex);
@@ -252,7 +251,7 @@ TArray<TVector<int32, 4>> TTriangleMesh<T>::GetUniqueAdjacentElements() const
 }
 
 template<class T>
-TArray<TVector<T, 3>> TTriangleMesh<T>::GetFaceNormals(const TArrayView<const TVector<T, 3>>& Points, const bool ReturnEmptyOnError) const
+TArray<TVector<T, 3>> TTriangleMesh<T>::GetFaceNormals(const TConstArrayView<TVector<T, 3>>& Points, const bool ReturnEmptyOnError) const
 {
 	TArray<TVector<T, 3>> Normals;
 	GetFaceNormals(Normals, Points, ReturnEmptyOnError);
@@ -262,7 +261,7 @@ TArray<TVector<T, 3>> TTriangleMesh<T>::GetFaceNormals(const TArrayView<const TV
 // Note:	This function assumes Counter Clockwise triangle windings in a Left Handed coordinate system
 //			If this is not the case the returned face normals may need to be inverted
 template<class T>
-void TTriangleMesh<T>::GetFaceNormals(TArray<TVector<T, 3>>& Normals, const TArrayView<const TVector<T, 3>>& Points, const bool ReturnEmptyOnError) const
+void TTriangleMesh<T>::GetFaceNormals(TArray<TVector<T, 3>>& Normals, const TConstArrayView<TVector<T, 3>>& Points, const bool ReturnEmptyOnError) const
 {
 	Normals.Reset(MElements.Num());
 	if (ReturnEmptyOnError)
@@ -299,54 +298,49 @@ void TTriangleMesh<T>::GetFaceNormals(TArray<TVector<T, 3>>& Normals, const TArr
 }
 
 template<class T>
-TArray<TVector<T, 3>> TTriangleMesh<T>::GetPointNormals(const TArrayView<const TVector<T, 3>>& Points, const bool bReturnEmptyOnError, const bool bUseGlobalArray)
+TArray<TVector<T, 3>> TTriangleMesh<T>::GetPointNormals(const TConstArrayView<TVector<T, 3>>& Points, const bool ReturnEmptyOnError)
 {
-	TArray<TVector<T, 3>> FaceNormals = GetFaceNormals(Points, bReturnEmptyOnError);
 	TArray<TVector<T, 3>> PointNormals;
-	GetPointNormals(PointNormals, FaceNormals, bReturnEmptyOnError, bUseGlobalArray);
+	const TArray<TVector<T, 3>> FaceNormals = GetFaceNormals(Points, ReturnEmptyOnError);
+	if (FaceNormals.Num())
+	{
+		PointNormals.SetNumUninitialized(MNumIndices);
+		GetPointNormals(PointNormals, FaceNormals, /*bUseGlobalArray =*/ false);
+	}
 	return PointNormals;
 }
 
 template<class T>
-void TTriangleMesh<T>::GetPointNormals(TArray<TVector<T, 3>>& PointNormals, const TArray<TVector<T, 3>>& FaceNormals, const bool bReturnEmptyOnError, const bool bUseGlobalArray)
+void TTriangleMesh<T>::GetPointNormals(TArrayView<TVector<T, 3>> PointNormals, const TConstArrayView<TVector<T, 3>>& FaceNormals, const bool bUseGlobalArray)
 {
 	GetPointToTriangleMap(); // build MPointToTriangleMap
 	const TTriangleMesh<T>* ConstThis = this;
-	ConstThis->GetPointNormals(PointNormals, FaceNormals, bReturnEmptyOnError, bUseGlobalArray);
+	ConstThis->GetPointNormals(PointNormals, FaceNormals, bUseGlobalArray);
 }
 
 template<class T>
-void TTriangleMesh<T>::GetPointNormals(TArray<TVector<T, 3>>& PointNormals, const TArray<TVector<T, 3>>& FaceNormals, const bool bReturnEmptyOnError, const bool bUseGlobalArray) const
+void TTriangleMesh<T>::GetPointNormals(TArrayView<TVector<T, 3>> PointNormals, const TConstArrayView<TVector<T, 3>>& FaceNormals, const bool bUseGlobalArray) const
 {
 	check(MPointToTriangleMap.Num() != 0);
-	PointNormals.SetNum(MNumIndices);
-	for (auto Element : MPointToTriangleMap)
+	for (int32 Element = 0; Element < MNumIndices; ++Element)  // Iterate points with local indexes
 	{
-		checkSlow(Element.Key >= MStartIdx);
-		const int32 NormalIndex = bUseGlobalArray ? Element.Key : GlobalToLocal(Element.Key);  // Select whether the points normal indices match the points indices or start at 0
-		if (PointNormals.Num() <= NormalIndex)
+		const int32 NormalIndex = bUseGlobalArray ? LocalToGlobal(Element) : Element;  // Select whether the points normal indices match the points indices or start at 0
+		TVector<T, 3>& Normal = PointNormals[NormalIndex];
+		Normal = TVector<T, 3>(0);
+		const TArray<int32>& TriangleMap = MPointToTriangleMap[Element];  // Access MPointToTriangleMap with local index
+		for (int32 k = 0; k < TriangleMap.Num(); ++k)
 		{
-			PointNormals.SetNum(NormalIndex);
-		}
-		TVector<T, 3> Normal(0);
-		for (int32 k = 0; k < Element.Value.Num(); ++k)
-		{
-			if (FaceNormals.IsValidIndex(Element.Value[k]))
+			if (FaceNormals.IsValidIndex(TriangleMap[k]))
 			{
-				Normal += FaceNormals[Element.Value[k]];
-			}
-			else if (bReturnEmptyOnError)
-			{
-				PointNormals.Reset();
-				return;
+				Normal += FaceNormals[TriangleMap[k]];
 			}
 		}
-		PointNormals[NormalIndex] = Normal.GetSafeNormal();
+		Normal = Normal.GetSafeNormal();
 	}
 }
 
 template<class T>
-void AddTrianglesToHull(const TArrayView<const TVector<T, 3>>& Points, const int32 I0, const int32 I1, const int32 I2, const TPlane<T, 3>& SplitPlane, const TArray<int32>& InIndices, TArray<TVector<int32, 3>>& OutIndices)
+void AddTrianglesToHull(const TConstArrayView<TVector<T, 3>>& Points, const int32 I0, const int32 I1, const int32 I2, const TPlane<T, 3>& SplitPlane, const TArray<int32>& InIndices, TArray<TVector<int32, 3>>& OutIndices)
 {
 	int32 MaxD = 0; //This doesn't need to be initialized but we need to avoid the compiler warning
 	T MaxDistance = 0;
@@ -450,7 +444,7 @@ void AddTrianglesToHull(const TArrayView<const TVector<T, 3>>& Points, const int
 
 // @todo(mlentine, ocohen); Merge different hull creation versions
 template<class T>
-TTriangleMesh<T> TTriangleMesh<T>::GetConvexHullFromParticles(const TArrayView<const TVector<T, 3>>& Points)
+TTriangleMesh<T> TTriangleMesh<T>::GetConvexHullFromParticles(const TConstArrayView<TVector<T, 3>>& Points)
 {
 	TArray<TVector<int32, 3>> Indices;
 	if (Points.Num() <= 2)
@@ -738,7 +732,7 @@ TSet<int32> TTriangleMesh<T>::GetBoundaryPoints()
 template <class T>
 TMap<int32, int32> TTriangleMesh<T>::FindCoincidentVertexRemappings(
 	const TArray<int32>& TestIndices,
-	const TArrayView<const TVector<T, 3>>& Points)
+	const TConstArrayView<TVector<T, 3>>& Points)
 {
 	// From index -> To index
 	TMap<int32, int32> Remappings;
@@ -879,7 +873,7 @@ TArray<T> TTriangleMesh<T>::GetCurvatureOnEdges(const TArray<TVector<T, 3>>& Fac
 }
 
 template<class T>
-TArray<T> TTriangleMesh<T>::GetCurvatureOnEdges(const TArrayView<const TVector<T, 3>>& Points)
+TArray<T> TTriangleMesh<T>::GetCurvatureOnEdges(const TConstArrayView<TVector<T, 3>>& Points)
 {
 	const TArray<TVector<T, 3>> FaceNormals = GetFaceNormals(Points, false);
 	return GetCurvatureOnEdges(FaceNormals);
@@ -913,7 +907,7 @@ TArray<T> TTriangleMesh<T>::GetCurvatureOnPoints(const TArray<T>& EdgeCurvatures
 }
 
 template<class T>
-TArray<T> TTriangleMesh<T>::GetCurvatureOnPoints(const TArrayView<const TVector<T, 3>>& Points)
+TArray<T> TTriangleMesh<T>::GetCurvatureOnPoints(const TConstArrayView<TVector<T, 3>>& Points)
 {
 	const TArray<T> EdgeCurvatures = GetCurvatureOnEdges(Points);
 	return GetCurvatureOnPoints(EdgeCurvatures);
@@ -979,7 +973,7 @@ private:
 
 template<class T>
 TArray<int32> TTriangleMesh<T>::GetVertexImportanceOrdering(
-    const TArrayView<const TVector<T, 3>>& Points,
+    const TConstArrayView<TVector<T, 3>>& Points,
     const TArray<T>& PointCurvatures,
     TArray<int32>* CoincidentVertices,
     const bool RestrictToLocalIndexRange)
@@ -1029,7 +1023,7 @@ TArray<int32> TTriangleMesh<T>::GetVertexImportanceOrdering(
 		const bool IsFree = Neighbors == nullptr || Neighbors->Num() == 0;
 		Rank[Idx - Offset] = IsFree ? 1 : 0;
 		FoundFreeParticle |= IsFree;
-			}
+	}
 	if (FoundFreeParticle)
 	{
 		StableSort(&PointOrder[0], NumPoints, AscendingRankPred);
@@ -1066,7 +1060,7 @@ TArray<int32> TTriangleMesh<T>::GetVertexImportanceOrdering(
 	TAABB<T, 3> Bbox(Points[0], Points[0]);
 	for (int i = 1; i < NumPoints; i++)
 	{
-		Bbox.GrowToInclude(Points[i+Offset]);
+		Bbox.GrowToInclude(Points[i + Offset]);
 	}
 	const TVector<T, 3> Center = Bbox.Center();
 
@@ -1219,7 +1213,7 @@ TArray<int32> TTriangleMesh<T>::GetVertexImportanceOrdering(
 
 template<class T>
 TArray<int32>
-TTriangleMesh<T>::GetVertexImportanceOrdering(const TArrayView<const TVector<T, 3>>& Points, TArray<int32>* CoincidentVertices, const bool RestrictToLocalIndexRange)
+TTriangleMesh<T>::GetVertexImportanceOrdering(const TConstArrayView<TVector<T, 3>>& Points, TArray<int32>* CoincidentVertices, const bool RestrictToLocalIndexRange)
 {
 	const TArray<T> pointCurvatures = GetCurvatureOnPoints(Points);
 	return GetVertexImportanceOrdering(Points, pointCurvatures, CoincidentVertices, RestrictToLocalIndexRange);
diff --git a/Engine/Source/Runtime/Experimental/Chaos/Public/Chaos/PBDAnimDriveConstraint.h b/Engine/Source/Runtime/Experimental/Chaos/Public/Chaos/PBDAnimDriveConstraint.h
--- a/Engine/Source/Runtime/Experimental/Chaos/Public/Chaos/PBDAnimDriveConstraint.h
+++ b/Engine/Source/Runtime/Experimental/Chaos/Public/Chaos/PBDAnimDriveConstraint.h
@@ -3,63 +3,68 @@
 
 #include "Chaos/ParticleRule.h"
 #include "Chaos/Framework/Parallel.h"
+#include "Containers/ArrayView.h"
+#include "Containers/ContainersFwd.h"
 
 namespace Chaos
 {
-template<typename T, int d>
-class TPBDAnimDriveConstraint : public TParticleRule<T, d>
-{
-  public:
-	// InSpringNeutralPositions starts at index: InParticleIndexOffset
-	// InSpringStiffnessMultiplier starts at index: 0
-	TPBDAnimDriveConstraint(
-		const int InParticleIndexOffset
-		, const TArray<TVector<T, d>>* const InSpringNeutralPositions
-		, const TArray<T>* const InSpringStiffnessMultiplier
-		, const T InSpringStiffness
-	)
-		: ParticleIndexOffset(InParticleIndexOffset)
-		, SpringStiffness(InSpringStiffness)
-		, SpringNeutralPositions(InSpringNeutralPositions)
-		, SpringStiffnessMultiplier(InSpringStiffnessMultiplier)
-	{
-
-	}
-	virtual ~TPBDAnimDriveConstraint() {}
-
-	inline virtual void Apply(TPBDParticles<T, d>& InParticles, const T Dt) const override
+	template<typename T, int d>
+	class TPBDAnimDriveConstraint : public TParticleRule<T, d>
 	{
-		int32 ConstraintCount = SpringStiffnessMultiplier->Num();
-		PhysicsParallelFor(ConstraintCount, [&](int32 Index)
+	public:
+		TPBDAnimDriveConstraint(
+			const int32 InParticleOffset,
+			const int32 InParticleCount,
+			const TArray<TVector<T, d>>& InAnimationPositions,  // Use global indexation (will need adding ParticleOffset)
+			const TConstArrayView<T>& InSpringStiffnessMultiplier  // Use local indexation
+		)
+			: AnimationPositions(InAnimationPositions)
+			, SpringStiffnessMultiplier(InSpringStiffnessMultiplier)
+			, ParticleOffset(InParticleOffset)
+			, SpringStiffness((T)1.)
 		{
-			ApplyAnimDriveConstraint(InParticles, Dt, Index);
-		});
-	}
+			check(InSpringStiffnessMultiplier.Num() == InParticleCount);
+		}
+		virtual ~TPBDAnimDriveConstraint() {}
 
-	inline void ApplyAnimDriveConstraint(TPBDParticles<T, d>& InParticles, const T Dt, const int32 Index) const
-	{
-		
-		if (InParticles.InvM(Index + ParticleIndexOffset) == 0.0f)
+		inline virtual void Apply(TPBDParticles<T, d>& InParticles, const T Dt) const override
 		{
-			return;
+			const int32 ParticleCount = SpringStiffnessMultiplier.Num();
+			PhysicsParallelFor(ParticleCount, [&](int32 Index)  // TODO: profile need for parallel loop based on particle count
+			{
+				ApplyAnimDriveConstraint(InParticles, Dt, Index);
+			});
 		}
 
-		const TVector<T, d> NeutralPosition = (*SpringNeutralPositions)[Index + ParticleIndexOffset];
-		const TVector<T, d> ParticlePosition = InParticles.P(Index + ParticleIndexOffset);
+		inline void ApplyAnimDriveConstraint(TPBDParticles<T, d>& InParticles, const T Dt, const int32 Index) const
+		{
+			const int32 ParticleIndex = ParticleOffset + Index;
+			if (InParticles.InvM(ParticleIndex) == (T)0.)
+			{
+				return;
+			}
 
-		const T RelaxationFactor = FMath::Min(SpringStiffness * (*SpringStiffnessMultiplier)[Index], 1.0f);
-		InParticles.P(Index + ParticleIndexOffset) = RelaxationFactor * NeutralPosition + (1 - RelaxationFactor) * ParticlePosition;
-	}
+			const TVector<T, d>& NeutralPosition = AnimationPositions[ParticleIndex];
+			const TVector<T, d>& ParticlePosition = InParticles.P(ParticleIndex);
 
-	inline void SetSpringStiffness(T InSpringStiffness)
-	{
-		SpringStiffness = InSpringStiffness;
-	}
+			const T RelaxationFactor = FMath::Min(SpringStiffness * SpringStiffnessMultiplier[Index], (T)1.);
+			InParticles.P(ParticleIndex) = RelaxationFactor * NeutralPosition + ((T)1. - RelaxationFactor) * ParticlePosition;
+		}
+
+		inline void SetSpringStiffness(T InSpringStiffness)
+		{
+			SpringStiffness = FMath::Clamp(InSpringStiffness, (T)0., (T)1.);
+		}
+
+		inline T GetSpringStiffness() const
+		{
+			return SpringStiffness;
+		}
 
-private:
-	const int32 ParticleIndexOffset;
-	T SpringStiffness;
-	const TArray<TVector<T, d>>* SpringNeutralPositions; // Size: Same as full particle array // Starts at Index: ParticleIndexOffset
-	const TArray<T>* SpringStiffnessMultiplier; // Size: Number of Animation drive constraints to solve // Starts at Index 0
-};
+	private:
+		const TArray<TVector<T, d>>& AnimationPositions;  // Use global index (needs adding ParticleOffset)
+		const TConstArrayView<T> SpringStiffnessMultiplier; // Use local index
+		const int32 ParticleOffset;
+		T SpringStiffness;
+	};
 }
diff --git a/Engine/Source/Runtime/Experimental/Chaos/Public/Chaos/PBDCollisionSpringConstraints.h b/Engine/Source/Runtime/Experimental/Chaos/Public/Chaos/PBDCollisionSpringConstraints.h
--- a/Engine/Source/Runtime/Experimental/Chaos/Public/Chaos/PBDCollisionSpringConstraints.h
+++ b/Engine/Source/Runtime/Experimental/Chaos/Public/Chaos/PBDCollisionSpringConstraints.h
@@ -17,8 +17,8 @@ class TPBDCollisionSpringConstraints : public PBDCollisionSpringConstraintsBase<
 	using Base::MConstraints;
 
   public:
-	TPBDCollisionSpringConstraints(const TDynamicParticles<T, d>& InParticles, const TArray<TVector<int32, 3>>& Elements, const TSet<TVector<int32, 2>>& DisabledCollisionElements, const TArray<uint32>& DynamicGroupIds, const TArray<T>& PerGroupThicknesses, const T Dt, const T Stiffness = (T)1.0)
-	    : Base(InParticles, Elements, DisabledCollisionElements, DynamicGroupIds, PerGroupThicknesses, Dt, Stiffness) {}
+	TPBDCollisionSpringConstraints(const TPBDActiveView<TPBDParticles<T, d>>& ParticlesActiveView, const TArray<TVector<int32, 3>>& Elements, const TSet<TVector<int32, 2>>& DisabledCollisionElements, const TArray<uint32>& DynamicGroupIds, const TArray<T>& PerGroupThicknesses, const T Dt, const T Stiffness = (T)1.0)
+	    : Base(ParticlesActiveView, Elements, DisabledCollisionElements, DynamicGroupIds, PerGroupThicknesses, Dt, Stiffness) {}
 	virtual ~TPBDCollisionSpringConstraints() {}
 
 	void Apply(TPBDParticles<T, d>& InParticles, const T Dt, const int32 InConstraintIndex) const
diff --git a/Engine/Source/Runtime/Experimental/Chaos/Public/Chaos/PBDCollisionSpringConstraintsBase.h b/Engine/Source/Runtime/Experimental/Chaos/Public/Chaos/PBDCollisionSpringConstraintsBase.h
--- a/Engine/Source/Runtime/Experimental/Chaos/Public/Chaos/PBDCollisionSpringConstraintsBase.h
+++ b/Engine/Source/Runtime/Experimental/Chaos/Public/Chaos/PBDCollisionSpringConstraintsBase.h
@@ -4,6 +4,7 @@
 #if !COMPILE_WITHOUT_UNREAL_SUPPORT
 #include "Chaos/Array.h"
 #include "Chaos/PBDParticles.h"
+#include "Chaos/PBDActiveView.h"
 
 #include <unordered_set>
 
@@ -14,7 +15,7 @@ template<class T, int d>
 class PBDCollisionSpringConstraintsBase
 {
   public:
-	PBDCollisionSpringConstraintsBase(const TDynamicParticles<T, d>& InParticles, const TArray<TVector<int32, 3>>& Elements, const TSet<TVector<int32, 2>>& DisabledCollisionElements, const TArray<uint32>& DynamicGroupIds, const TArray<T>& PerGroupThicknesses, const T Dt, const T Stiffness = (T)1);
+	PBDCollisionSpringConstraintsBase(const TPBDActiveView<TPBDParticles<T, d>>& ParticlesActiveView, const TArray<TVector<int32, 3>>& Elements, const TSet<TVector<int32, 2>>& DisabledCollisionElements, const TArray<uint32>& DynamicGroupIds, const TArray<T>& PerGroupThicknesses, const T Dt, const T Stiffness = (T)1);
 	virtual ~PBDCollisionSpringConstraintsBase() {}
 
 	TVector<T, d> GetDelta(const TPBDParticles<T, d>& InParticles, const int32 i) const;
diff --git a/Engine/Source/Runtime/Experimental/Chaos/Public/Chaos/PBDEvolution.h b/Engine/Source/Runtime/Experimental/Chaos/Public/Chaos/PBDEvolution.h
--- a/Engine/Source/Runtime/Experimental/Chaos/Public/Chaos/PBDEvolution.h
+++ b/Engine/Source/Runtime/Experimental/Chaos/Public/Chaos/PBDEvolution.h
@@ -5,14 +5,15 @@
 #include "Chaos/PerParticleGravity.h"
 #include "Chaos/VelocityField.h"
 #include "Chaos/PBDParticles.h"
+#include "Chaos/PBDActiveView.h"
 #include "Chaos/Vector.h"
 
 namespace Chaos
 {
 template<class T, int d>
-class CHAOS_API TPBDEvolution  : private TArrayCollection
+class CHAOS_API TPBDEvolution : public TArrayCollection
 {
-  public:
+ public:
 	using FGravityForces = TPerParticleGravity<T, d>;
 	using FVelocityField = TVelocityField<T, d>;
 
@@ -22,31 +23,80 @@ class CHAOS_API TPBDEvolution  : private TArrayCollection
 
 	void AdvanceOneTimeStep(const T dt);
 
-	void SetKinematicUpdateFunction(TFunction<void(TPBDParticles<T, d>&, const T, const T, const int32)> KinematicUpdate) { MKinematicUpdate = KinematicUpdate; }
-	void SetCollisionKinematicUpdateFunction(TFunction<void(TKinematicGeometryClothParticles<T, d>&, const T, const T, const int32)> KinematicUpdate) { MCollisionKinematicUpdate = KinematicUpdate; }
-	void SetParticleUpdateFunction(TFunction<void(TPBDParticles<T, d>&, const T)> ParticleUpdate) { MParticleUpdate = ParticleUpdate; }
-	void AddPBDConstraintFunction(TFunction<void(TPBDParticles<T, d>&, const T)> ConstraintFunction) { MConstraintRules.Add(ConstraintFunction); }
-	void AddXPBDConstraintFunctions(TFunction<void()> InitConstraintFunction, TFunction<void(TPBDParticles<T, d>&, const T)> ConstraintFunction) { MInitConstraintRules.Add(InitConstraintFunction); MConstraintRules.Add(ConstraintFunction); }
-	void AddForceFunction(TFunction<void(TPBDParticles<T, d>&, const T, const int32)> ForceFunction) { MForceRules.Add(ForceFunction); }
+	// Remove all particles, will also reset all rules
+	void ResetParticles();
 
 	// Add particles and initialize group ids. Return the index of the first added particle.
-	uint32 AddParticles(uint32 Num, uint32 GroupId = 0);
-	// Add collision particles and initialize group ids. Return the index of the first added particle.
-	uint32 AddCollisionParticles(uint32 Num, uint32 GroupId = 0);
+	int32 AddParticleRange(int32 NumParticles, uint32 GroupId, bool bActivate);
+
+	// Return the number of particles of the block starting at Offset
+	int32 GetParticleRangeSize(int32 Offset) const { return MParticlesActiveView.GetRangeSize(Offset); }
 
+	// Set a block of particles active or inactive, using the index of the first added particle to identify the block.
+	void ActivateParticleRange(int32 Offset, bool bActivate)  { MParticlesActiveView.ActivateRange(Offset, bActivate); }
+
+	// Particles accessors
 	const TPBDParticles<T, d>& Particles() const { return MParticles; }
 	TPBDParticles<T, d>& Particles() { return MParticles; }
+	const TPBDActiveView<TPBDParticles<T, d>>& ParticlesActiveView() { return MParticlesActiveView; }
+
+	const TArray<uint32>& ParticleGroupIds() const { return MParticleGroupIds; }
+
+	// Remove all collision particles
+	void ResetCollisionParticles(int32 NumParticles = 0);
+
+	// Add collision particles and initialize group ids. Return the index of the first added particle.
+	// Use INDEX_NONE as GroupId for collision particles that affect all particle groups.
+	int32 AddCollisionParticleRange(int32 NumParticles, uint32 GroupId, bool bActivate);
+
+	// Set a block of collision particles active or inactive, using the index of the first added particle to identify the block.
+	void ActivateCollisionParticleRange(int32 Offset, bool bActivate) { MCollisionParticlesActiveView.ActivateRange(Offset, bActivate); }
+
+	// Return the number of particles of the block starting at Offset
+	int32 GetCollisionParticleRangeSize(int32 Offset) const { return MCollisionParticlesActiveView.GetRangeSize(Offset); }
+
+	// Collision particles accessors
+	const TKinematicGeometryClothParticles<T, d>& CollisionParticles() const { return MCollisionParticles; }
+	TKinematicGeometryClothParticles<T, d>& CollisionParticles() { return MCollisionParticles; }
+	const TArray<uint32>& CollisionParticleGroupIds() const { return MCollisionParticleGroupIds; }
+	const TPBDActiveView<TKinematicGeometryClothParticles<T, d>>& CollisionParticlesActiveView() { return MCollisionParticlesActiveView; }
+
+	// Reset all constraint init and rule functions.
+	void ResetConstraintRules() { MConstraintInits.Reset(); MConstraintRules.Reset(); MConstraintInitsActiveView.Reset(); MConstraintRulesActiveView.Reset();  };
+
+	// Add constraints. Return the index of the first added constraint.
+	int32 AddConstraintInitRange(int32 NumConstraints, bool bActivate);
+	int32 AddConstraintRuleRange(int32 NumConstraints, bool bActivate);
+	
+	// Return the number of particles of the block starting at Offset
+	int32 GetConstraintInitRangeSize(int32 Offset) const { return MConstraintInitsActiveView.GetRangeSize(Offset); }
+	int32 GetConstraintRuleRangeSize(int32 Offset) const { return MConstraintRulesActiveView.GetRangeSize(Offset); }
+
+	// Set a block of constraints active or inactive, using the index of the first added particle to identify the block.
+	void ActivateConstraintInitRange(int32 Offset, bool bActivate) { MConstraintInitsActiveView.ActivateRange(Offset, bActivate); }
+	void ActivateConstraintRuleRange(int32 Offset, bool bActivate) { MConstraintRulesActiveView.ActivateRange(Offset, bActivate); }
+
+	// Constraint accessors
+	const TArray<TFunction<void()>>& ConstraintInits() const { return MConstraintInits; }
+	TArray<TFunction<void()>>& ConstraintInits() { return MConstraintInits; }
+	const TArray<TFunction<void(TPBDParticles<T, d>&, const T)>>& ConstraintRules() const { return MConstraintRules; }
+	TArray<TFunction<void(TPBDParticles<T, d>&, const T)>>& ConstraintRules() { return MConstraintRules; }
+	
+	void SetKinematicUpdateFunction(TFunction<void(TPBDParticles<T, d>&, const T, const T, const int32)> KinematicUpdate) { MKinematicUpdate = KinematicUpdate; }
+	void SetCollisionKinematicUpdateFunction(TFunction<void(TKinematicGeometryClothParticles<T, d>&, const T, const T, const int32)> KinematicUpdate) { MCollisionKinematicUpdate = KinematicUpdate; }
+
+	void ResetForceFunctions() { MForceRules.Reset(); }
+	void AddForceFunction(TFunction<void(TPBDParticles<T, d>&, const T, const int32)> ForceFunction) { MForceRules.Add(ForceFunction); }
 
 	FGravityForces& GetGravityForces(const uint32 GroupId = 0) { check(GroupId < TArrayCollection::Size()); return MGroupGravityForces[GroupId]; }
 	const FGravityForces& GetGravityForces(const uint32 GroupId = 0) const { check(GroupId < TArrayCollection::Size()); return MGroupGravityForces[GroupId]; }
 
-	TUniquePtr<FVelocityField>& GetVelocityField(const uint32 GroupId = 0) { check(GroupId < TArrayCollection::Size()); return MGroupVelocityFields[GroupId]; }
-	const TUniquePtr<FVelocityField>& GetVelocityField(const uint32 GroupId = 0) const { check(GroupId < TArrayCollection::Size()); return MGroupVelocityFields[GroupId]; }
+	FVelocityField& GetVelocityField(const uint32 GroupId = 0) { check(GroupId < TArrayCollection::Size()); return MGroupVelocityFields[GroupId]; }
+	const FVelocityField& GetVelocityField(const uint32 GroupId = 0) const { check(GroupId < TArrayCollection::Size()); return MGroupVelocityFields[GroupId]; }
 
-	const TGeometryClothParticles<T, d>& CollisionParticles() const { return MCollisionParticles; }
-	TGeometryClothParticles<T, d>& CollisionParticles() { return MCollisionParticles; }
 	const bool Collided(int32 index) { return MCollided[index]; }
 
+	void ResetSelfCollision() { MCollisionTriangles.Reset(); MDisabledCollisionElements.Reset(); };
 	TArray<TVector<int32, 3>>& CollisionTriangles() { return MCollisionTriangles; }
 	TSet<TVector<int32, 2>>& DisabledCollisionElements() { return MDisabledCollisionElements; }
 
@@ -67,28 +117,42 @@ class CHAOS_API TPBDEvolution  : private TArrayCollection
 
 	T GetTime() const { return MTime; }
 
-	void ResetConstraintRules() { MInitConstraintRules.Reset(); MConstraintRules.Reset(); };
-	void ResetSelfCollision() { MCollisionTriangles.Reset(); MDisabledCollisionElements.Reset(); };
-
-	const TArrayCollectionArray<uint32>& ParticleGroupIds() const { return MParticleGroupIds; }
-
-  private:
+ private:
 	// Add simulation groups and set default values
-	void AddGroups(int32 Num);
-
+	void AddGroups(int32 NumGroups);
+	// Reset simulation groups
+	void ResetGroups();
+ 
+private:
 	TPBDParticles<T, d> MParticles;
+	TPBDActiveView<TPBDParticles<T, d>> MParticlesActiveView;
 	TKinematicGeometryClothParticles<T, d> MCollisionParticles;
+	TPBDActiveView<TKinematicGeometryClothParticles<T, d>> MCollisionParticlesActiveView;
+
 	TArray<TVector<int32, 3>> MCollisionTriangles;       // Used for self-collisions
 	TSet<TVector<int32, 2>> MDisabledCollisionElements;  // 
+
 	TArrayCollectionArray<bool> MCollided;
 	TArrayCollectionArray<uint32> MCollisionParticleGroupIds;  // Used for per group parameters for collision particles
 	TArrayCollectionArray<uint32> MParticleGroupIds;  // Used for per group parameters for particles
+
 	TArrayCollectionArray<FGravityForces> MGroupGravityForces;
 	TArrayCollectionArray<T> MGroupCollisionThicknesses;
 	TArrayCollectionArray<T> MGroupSelfCollisionThicknesses;
 	TArrayCollectionArray<T> MGroupCoefficientOfFrictions;
 	TArrayCollectionArray<T> MGroupDampings;
-	TArrayCollectionArray<TUniquePtr<FVelocityField>> MGroupVelocityFields;
+	TArrayCollectionArray<FVelocityField> MGroupVelocityFields;
+	
+	TArray<TFunction<void()>> MConstraintInits;
+	TPBDActiveView<TArray<TFunction<void()>>> MConstraintInitsActiveView;
+	TArray<TFunction<void(TPBDParticles<T, d>&, const T)>> MConstraintRules;
+	TPBDActiveView<TArray<TFunction<void(TPBDParticles<T, d>&, const T)>>> MConstraintRulesActiveView;
+
+	TArray<TFunction<void(TPBDParticles<T, d>&, const T, const int32)>> MForceRules;
+	TFunction<void(TPBDActiveView<TPBDParticles<T, d>>&, const T)> MParticleUpdate;
+	TFunction<void(TPBDParticles<T, d>&, const T, const T, const int32)> MKinematicUpdate;
+	TFunction<void(TKinematicGeometryClothParticles<T, d>&, const T, const T, const int32)> MCollisionKinematicUpdate;
+
 	int32 MNumIterations;
 	TVector<T, d> MGravity;
 	T MCollisionThickness;
@@ -96,12 +160,5 @@ class CHAOS_API TPBDEvolution  : private TArrayCollection
 	T MCoefficientOfFriction;
 	T MDamping;
 	T MTime;
-
-	TArray<TFunction<void(TPBDParticles<T, d>&, const T, const int32)>> MForceRules;
-	TArray<TFunction<void()>> MInitConstraintRules;
-	TArray<TFunction<void(TPBDParticles<T, d>&, const T)>> MConstraintRules;
-	TFunction<void(TPBDParticles<T, d>&, const T)> MParticleUpdate;
-	TFunction<void(TPBDParticles<T, d>&, const T, const T, const int32)> MKinematicUpdate;
-	TFunction<void(TKinematicGeometryClothParticles<T, d>&, const T, const T, const int32)> MCollisionKinematicUpdate;
 };
 }
diff --git a/Engine/Source/Runtime/Experimental/Chaos/Public/Chaos/PBDShapeConstraints.h b/Engine/Source/Runtime/Experimental/Chaos/Public/Chaos/PBDShapeConstraints.h
--- a/Engine/Source/Runtime/Experimental/Chaos/Public/Chaos/PBDShapeConstraints.h
+++ b/Engine/Source/Runtime/Experimental/Chaos/Public/Chaos/PBDShapeConstraints.h
@@ -1,9 +1,6 @@
 // Copyright Epic Games, Inc. All Rights Reserved.
 #pragma once
 
-#include "Chaos/Array.h"
-#include "Chaos/Map.h"
-#include "Chaos/PBDParticles.h"
 #include "Chaos/PBDShapeConstraintsBase.h"
 
 namespace Chaos
@@ -15,11 +12,17 @@ namespace Chaos
 
 	public:
 
-		TPBDShapeConstraints(const TDynamicParticles<T, d>& InParticles, uint32 InParticleIndexOffset, uint32 InParticleCount, const TArray<TVector<float, 3>>& TargetPositions, const T Stiffness = (T)1)
-			: Base(InParticles, TargetPositions, Stiffness), ParticleIndexOffset(InParticleIndexOffset), ParticleCount(InParticleCount)
+		TPBDShapeConstraints(
+			int32 InParticleOffset,
+			int32 InParticleCount,
+			const TArray<TVector<T, 3>>& StartPositions,
+			const TArray<TVector<T, 3>>& TargetPositions,
+			const T Stiffness = (T)1.
+		)
+			: Base(InParticleOffset, InParticleCount, StartPositions, TargetPositions, Stiffness)
 		{
 		}
-		virtual ~TPBDShapeConstraints() {}
+		virtual ~TPBDShapeConstraints() override {}
 
 		void Apply(TPBDParticles<T, d>& InParticles, const T Dt, const int32 Index) const
 		{
@@ -31,14 +34,10 @@ namespace Chaos
 
 		void Apply(TPBDParticles<T, d>& InParticles, const T Dt) const
 		{
-			for (uint32 Index = ParticleIndexOffset; Index < ParticleIndexOffset + ParticleCount; Index++)
+			for (int32 Index = MParticleOffset; Index < MParticleOffset + MTargetPositions.Num(); Index++)
 			{
 				Apply(InParticles, Dt, Index);
 			}
 		}
-
-	private:
-		uint32 ParticleIndexOffset;
-		uint32 ParticleCount;
 	};
 }
diff --git a/Engine/Source/Runtime/Experimental/Chaos/Public/Chaos/PBDShapeConstraintsBase.h b/Engine/Source/Runtime/Experimental/Chaos/Public/Chaos/PBDShapeConstraintsBase.h
--- a/Engine/Source/Runtime/Experimental/Chaos/Public/Chaos/PBDShapeConstraintsBase.h
+++ b/Engine/Source/Runtime/Experimental/Chaos/Public/Chaos/PBDShapeConstraintsBase.h
@@ -8,40 +8,55 @@
 
 namespace Chaos
 {
-template<class T, int d>
-class TPBDShapeConstraintsBase
-{
-  public:
-	TPBDShapeConstraintsBase(const TDynamicParticles<T, d>& InParticles, const TArray<TVector<T, 3>>& TargetPositions, const T Stiffness = (T)1)
-	    : MTargetPositions(TargetPositions), MStiffness(Stiffness)
+	template<class T, int d>
+	class TPBDShapeConstraintsBase
 	{
-		for (uint32 i = 0; i < InParticles.Size(); ++i)
+	public:
+		TPBDShapeConstraintsBase(
+			int32 InParticleOffset,
+			int32 InParticleCount,
+			const TArray<TVector<T, 3>>& StartPositions,
+			const TArray<TVector<T, 3>>& TargetPositions,
+			const T Stiffness
+		)
+			: MTargetPositions(TargetPositions)
+			, MParticleOffset(InParticleOffset)
+			, MStiffness(Stiffness)
 		{
-			const TVector<T, d>& P1 = InParticles.X(i);
-			const TVector<T, d>& P2 = MTargetPositions[i];
-			MDists.Add((P1 - P2).Size());
+			const int32 NumConstraints = InParticleCount;
+			MDists.SetNumUninitialized(InParticleCount);
+			for (int32 Index = 0; Index < InParticleCount; ++Index)
+			{
+				const int32 ParticleIndex = MParticleOffset + Index;
+				const TVector<T, d>& P1 = StartPositions[ParticleIndex];
+				const TVector<T, d>& P2 = MTargetPositions[ParticleIndex];
+				MDists[Index] = (P1 - P2).Size();
+			}
 		}
-	}
-	virtual ~TPBDShapeConstraintsBase() {}
+		virtual ~TPBDShapeConstraintsBase() {}
 
-	TVector<T, d> GetDelta(const TPBDParticles<T, d>& InParticles, const int32 i) const
-	{
-		if (InParticles.InvM(i) == 0)
-			return TVector<T, d>(0);
-		const TVector<T, d>& P1 = InParticles.P(i);
-		const TVector<T, d>& P2 = MTargetPositions[i];
-		TVector<T, d> Difference = P1 - P2;
-		float Distance = Difference.Size();
-		TVector<T, d> Direction = Difference / Distance;
-		TVector<T, d> Delta = (Distance - MDists[i]) * Direction;
-		return MStiffness * Delta / InParticles.InvM(i);
-	}
+		TVector<T, d> GetDelta(const TPBDParticles<T, d>& InParticles, const int32 Index) const
+		{
+			checkSlow(Index >= MParticleOffset && Index < MParticleOffset + MDists.Num())
+			if (InParticles.InvM(Index) == (T)0.)
+			{
+				return TVector<T, d>(0.);
+			}
+			const TVector<T, d>& P1 = InParticles.P(Index);
+			const TVector<T, d>& P2 = MTargetPositions[Index];
+			const TVector<T, d> Difference = P1 - P2;
+			const T Distance = Difference.Size();
+			const TVector<T, d> Direction = Difference / Distance;
+			const TVector<T, d> Delta = (Distance - MDists[Index - MParticleOffset]) * Direction;
+			return MStiffness * Delta / InParticles.InvM(Index);
+		}
 
-  protected:
-	const TArray<TVector<float, 3>>& MTargetPositions;
+	protected:
+		const TArray<TVector<float, 3>>& MTargetPositions;
+		const int32 MParticleOffset;
 
-  private:
-	TArray<T> MDists;
-	T MStiffness;
-};
+	private:
+		TArray<T> MDists;
+		T MStiffness;
+	};
 }
diff --git a/Engine/Source/Runtime/Experimental/Chaos/Public/Chaos/PBDSphericalConstraint.h b/Engine/Source/Runtime/Experimental/Chaos/Public/Chaos/PBDSphericalConstraint.h
--- a/Engine/Source/Runtime/Experimental/Chaos/Public/Chaos/PBDSphericalConstraint.h
+++ b/Engine/Source/Runtime/Experimental/Chaos/Public/Chaos/PBDSphericalConstraint.h
@@ -3,87 +3,156 @@
 
 #include "ParticleRule.h"
 #include "ChaosStats.h"
+#include "Containers/ArrayView.h"
+#include "Containers/ContainersFwd.h"
 
 DECLARE_CYCLE_STAT(TEXT("Chaos PBD Spherical Constraints"), STAT_PBD_Spherical, STATGROUP_Chaos);
+DECLARE_CYCLE_STAT(TEXT("Chaos PBD Spherical Backstop Constraints"), STAT_PBD_SphericalBackstop, STATGROUP_Chaos);
 
 namespace Chaos
 {
-
-template<typename T, int d>
-class PBDSphericalConstraint : public TParticleRule<T, d>
-{
-public:
-	PBDSphericalConstraint(
-		const uint32 InFirstParticleIndex
-		, const uint32 InParticleCount
-		, bool bInside
-		, const TArray<TVector<T, d>>* const InSpherePositions
-		, const TArray<T>* const InSphereRadii
-		, const TArray <T>* const InSphereOffsetDistances = nullptr
-		, const TArray<TVector<T, d>>* const InSphereOffsetDirections = nullptr
-	)
-		: FirstParticleIndex(InFirstParticleIndex), ParticleCount(InParticleCount)
-		, bConstraintInsideOfSphere(bInside)
-		, SpherePositions(InSpherePositions)
-		, SphereRadii(InSphereRadii)
-		, SphereOffsetDistances(InSphereOffsetDistances)
-		, SphereOffsetDirections(InSphereOffsetDirections)
-		, SphereRadiiMultiplier(1.f)
-	{
-		check(SphereRadii);
-		check(SpherePositions);
-	}
-	virtual ~PBDSphericalConstraint() {}
-
-	inline void Apply(TPBDParticles<T, d>& InParticles, const T Dt) const override
+	template<typename T, int d>
+	class TPBDSphericalConstraint : public TParticleRule<T, d>
 	{
-		SCOPE_CYCLE_COUNTER(STAT_PBD_Spherical);
-		const T CenterDeadZoneRadius = KINDA_SMALL_NUMBER; // We will not push the particle away in the dead zone
+	public:
+		TPBDSphericalConstraint(
+			const uint32 InParticleOffset,
+			const uint32 InParticleCount,
+			const TArray<TVector<T, d>>& InAnimationPositions,  // Use global indexation (will need adding ParticleOffset)
+			const TConstArrayView<T>& InSphereRadii  // Use local indexation
+		)
+			: AnimationPositions(InAnimationPositions)
+			, SphereRadii(InSphereRadii)
+			, ParticleOffset(InParticleOffset)
+			, SphereRadiiMultiplier((T)1.)
+		{
+			check(InSphereRadii.Num() == InParticleCount);
+		}
+		virtual ~TPBDSphericalConstraint() {}
 
-		PhysicsParallelFor(ParticleCount, [&](int32 Index)
+		inline virtual void Apply(TPBDParticles<T, d>& Particles, const T Dt) const override
 		{
-			if (InParticles.InvM(FirstParticleIndex + Index) == 0)
+			SCOPE_CYCLE_COUNTER(STAT_PBD_Spherical);
+
+			const int32 ParticleCount = SphereRadii.Num();
+
+			PhysicsParallelFor(ParticleCount, [&](int32 Index)  // TODO: profile need for parallel loop based on particle count
 			{
-				return;
-			}
+				const int32 ParticleIndex = ParticleOffset + Index;
+
+				if (Particles.InvM(ParticleIndex) == 0)
+				{
+					return;
+				}
 			
-			const T Radius = (*SphereRadii)[Index] * SphereRadiiMultiplier;
-			TVector<T, d> Centre = (*SpherePositions)[Index + FirstParticleIndex];
-			if (SphereOffsetDistances && SphereOffsetDirections)
-			{
-				Centre -= (*SphereOffsetDistances)[Index] * (*SphereOffsetDirections)[Index + FirstParticleIndex];
-			}
+				const T Radius = SphereRadii[Index] * SphereRadiiMultiplier;
+				const TVector<T, d>& Center = AnimationPositions[ParticleIndex];
 
-			const TVector<T, d> CenterToParticle = InParticles.P(FirstParticleIndex + Index) - Centre;
-			const T DistanceSquared = CenterToParticle.SizeSquared();
-			const bool bPushParticleToShpericalBoundery = ((DistanceSquared < Radius * Radius) ^ bConstraintInsideOfSphere) && DistanceSquared > CenterDeadZoneRadius * CenterDeadZoneRadius;
+				const TVector<T, d> CenterToParticle = Particles.P(ParticleIndex) - Center;
+				const T DistanceSquared = CenterToParticle.SizeSquared();
 
-			if (bPushParticleToShpericalBoundery)
-			{
-				const T Distance = sqrt(DistanceSquared);
-				const TVector<T, d> PositionOnSphere = (Radius / Distance) * (CenterToParticle);
-				InParticles.P(FirstParticleIndex + Index) = Centre + PositionOnSphere;
-			}
-		});
-	}
-
-	inline void SetSphereRadiiMultiplier(const T InSphereRadiiMultiplier)
+				static const T DeadZoneSquareRadius = SMALL_NUMBER; // We will not push the particle away in the dead zone
+				if (DistanceSquared > FMath::Square(Radius) + DeadZoneSquareRadius)
+				{
+					const T Distance = sqrt(DistanceSquared);
+					const TVector<T, d> PositionOnSphere = (Radius / Distance) * CenterToParticle;
+					Particles.P(ParticleIndex) = Center + PositionOnSphere;
+				}
+			});
+		}
+
+		inline void SetSphereRadiiMultiplier(const T InSphereRadiiMultiplier)
+		{
+			SphereRadiiMultiplier = FMath::Max((T)0., InSphereRadiiMultiplier);
+		}
+
+	protected:	
+		const TArray<TVector<T, d>>& AnimationPositions;  // Use global indexation (will need adding ParticleOffset)
+		const TConstArrayView<T> SphereRadii;  // Use local indexation
+		const int32 ParticleOffset;
+		T SphereRadiiMultiplier;
+	};
+
+	template<typename T, int d>
+	class TPBDSphericalBackstopConstraint : public TParticleRule<T, d>
 	{
-		SphereRadiiMultiplier = InSphereRadiiMultiplier;
-	}
+	public:
+		TPBDSphericalBackstopConstraint(
+			const int32 InParticleOffset,
+			const int32 InParticleCount,
+			const TArray<TVector<T, d>>& InAnimationPositions,  // Use global indexation (will need adding ParticleOffset)
+			const TArray<TVector<T, d>>& InAnimationNormals,  // Use global indexation (will need adding ParticleOffset)
+			const TConstArrayView<T>& InSphereRadii,  // Use local indexation
+			const TConstArrayView<T>& InSphereOffsetDistances  // Use local indexation
+		)
+			: AnimationPositions(InAnimationPositions)
+			, AnimationNormals(InAnimationNormals)
+			, SphereRadii(InSphereRadii)
+			, SphereOffsetDistances(InSphereOffsetDistances)
+			, ParticleOffset(InParticleOffset)
+			, SphereRadiiMultiplier((T)1.)
+		{
+			check(InSphereRadii.Num() == InParticleCount);
+			check(InSphereOffsetDistances.Num() == InParticleCount);
+		}
+		virtual ~TPBDSphericalBackstopConstraint() {}
+
+		inline virtual void Apply(TPBDParticles<T, d>& Particles, const T Dt) const override
+		{
+			SCOPE_CYCLE_COUNTER(STAT_PBD_SphericalBackstop);
+
+			const int32 ParticleCount = SphereRadii.Num();
+
+			PhysicsParallelFor(ParticleCount, [&](int32 Index)  // TODO: profile need for parallel loop based on particle count
+			{
+				const int32 ParticleIndex = ParticleOffset + Index;
 
-private:	
-	const uint32 FirstParticleIndex; // TODO: Get rid of this member variable (this is specific to a temporary cloth instancing implementation)
-	const uint32 ParticleCount;
-	const bool bConstraintInsideOfSphere; // Should the points be constrained to the inside of their respective spheres or the outsides
+				if (Particles.InvM(ParticleIndex) == 0)
+				{
+					return;
+				}
 
-	const TArray<TVector<T, d>>* const SpherePositions;  // Positions of spheres without adding an offset vector // Start at index FirstParticleIndex
-	const TArray<T>* const SphereRadii; // Start at index 0
+				const TVector<T, d>& AnimationPosition = AnimationPositions[Index];
+				const TVector<T, d>& AnimationNormal = AnimationNormals[Index];
 
-	const TArray<T>* const SphereOffsetDistances;  // Optional Sphere position offsets // start at index 0
-	const TArray<TVector<T, d>>* const SphereOffsetDirections; //	Optional Sphere offset directions (TBD: This is really implementation specific to backstop constraints, consider creating a specialized constraint instead).
+				const T SphereOffsetDistance = SphereOffsetDistances[Index];
+				const T Radius = SphereRadii[Index] * SphereRadiiMultiplier;
 
-	T SphereRadiiMultiplier;
-};
+				const TVector<T, d> Center = AnimationPosition - (Radius + SphereOffsetDistance) * AnimationNormal;
+
+				const TVector<T, d> CenterToParticle = Particles.P(ParticleIndex) - Center;
+				const T DistanceSquared = CenterToParticle.SizeSquared();
+
+				static const T DeadZoneSquareRadius = SMALL_NUMBER;
+				if (DistanceSquared < DeadZoneSquareRadius)
+				{
+					Particles.P(ParticleIndex) = AnimationPosition - SphereOffsetDistance * AnimationNormal;
+				}
+				else if (DistanceSquared < FMath::Square(Radius))
+				{
+					const TVector<T, d> PositionOnSphere = (Radius / sqrt(DistanceSquared)) * CenterToParticle;
+					Particles.P(ParticleIndex) = Center + PositionOnSphere;
+				}
+				// Else the particle is outside the sphere, and there is nothing to do
+			});
+		}
+
+		inline void SetSphereRadiiMultiplier(const T InSphereRadiiMultiplier)
+		{
+			SphereRadiiMultiplier = FMath::Max((T)0., InSphereRadiiMultiplier);
+		}
+
+		inline T GetSphereRadiiMultiplier() const
+		{
+			return SphereRadiiMultiplier;
+		}
 
+	private:
+		const TArray<TVector<T, d>>& AnimationPositions;  // Positions of spheres, use global indexation (will need adding ParticleOffset)
+		const TArray<TVector<T, d>>& AnimationNormals; // Sphere offset directions, use global indexation (will need adding ParticleOffset)
+		const TConstArrayView<T> SphereRadii; // Start at index 0, use local indexation
+		const TConstArrayView<T> SphereOffsetDistances;  // Sphere position offsets, use local indexation
+		const int32 ParticleOffset;
+		T SphereRadiiMultiplier;
+	};
 }
diff --git a/Engine/Source/Runtime/Experimental/Chaos/Public/Chaos/PerParticleDampVelocity.h b/Engine/Source/Runtime/Experimental/Chaos/Public/Chaos/PerParticleDampVelocity.h
--- a/Engine/Source/Runtime/Experimental/Chaos/Public/Chaos/PerParticleDampVelocity.h
+++ b/Engine/Source/Runtime/Experimental/Chaos/Public/Chaos/PerParticleDampVelocity.h
@@ -3,6 +3,7 @@
 
 #include "Chaos/Matrix.h"
 #include "Chaos/PerParticleRule.h"
+#include "Chaos/PBDActiveView.h"
 #include "GenericPlatform/GenericPlatformMath.h"
 
 namespace Chaos
@@ -18,6 +19,61 @@ namespace Chaos
 		virtual ~TPerParticleDampVelocity() {}
 
 		template<class T_PARTICLES>
+		inline void UpdatePositionBasedState(const TPBDActiveView<T_PARTICLES>& ParticlesActiveView)
+		{
+			static_assert(d == 3, "Damp Velocities currently only supports 3D vectors.");
+			
+			T_PARTICLES& Particles = ParticlesActiveView.GetItems();
+			
+			MXcm = TVector<T, d>(0.f, 0.f, 0.f);
+			MVcm = TVector<T, d>(0.f, 0.f, 0.f);
+			T Mcm = (T)0;
+
+			ParticlesActiveView.SequentialFor(
+				[this, &Mcm](T_PARTICLES& Particles, int32 Index)
+				{
+					if (!Particles.InvM(Index))
+					{
+						return;
+					}
+					MXcm += Particles.X(Index) * Particles.M(Index);
+					MVcm += Particles.V(Index) * Particles.M(Index);
+					Mcm += Particles.M(Index);
+				});
+
+			if (Mcm != 0.0f)
+			{
+				MXcm /= Mcm;
+				MVcm /= Mcm;
+			}
+
+			TVector<T, d> L = TVector<T, d>(0.f, 0.f, 0.f);
+			PMatrix<T, d, d> I(0);
+			ParticlesActiveView.SequentialFor(
+				[this, &L, &I](T_PARTICLES& Particles, int32 Index)
+				{
+					if (!Particles.InvM(Index))
+					{
+						return;
+					}
+					TVector<T, d> V = Particles.X(Index) - MXcm;
+					L += TVector<T, d>::CrossProduct(V, Particles.M(Index) * Particles.V(Index));
+					PMatrix<T, d, d> M(0, V[2], -V[1], -V[2], 0, V[0], V[1], -V[0], 0);
+					I += M.GetTransposed() * M * Particles.M(Index);
+				});
+
+#if COMPILE_WITHOUT_UNREAL_SUPPORT
+			MOmega = I.Determinant() > 1e-7 ? TRigidTransform<T, d>(I).InverseTransformVector(L) : TVector<T, d>(0);
+#else
+			const T Det = I.Determinant();
+			MOmega = Det < SMALL_NUMBER || !FGenericPlatformMath::IsFinite(Det) ?
+			    TVector<T, d>(0) :
+			    I.InverseTransformVector(L); // Calls FMatrix::InverseFast(), which tests against SMALL_NUMBER
+#endif
+		}
+
+		template<class T_PARTICLES>
+		UE_DEPRECATED(4.26, "Use TPBDActiveView instead")
 		inline void UpdatePositionBasedState(const T_PARTICLES& InParticles, const TArray<int32>& InActiveIndices)
 		{
 			static_assert(d == 3, "Damp Velocities currently only supports 3D vectors.");
diff --git a/Engine/Source/Runtime/Experimental/Chaos/Public/Chaos/PerParticlePBDCollisionConstraint.h b/Engine/Source/Runtime/Experimental/Chaos/Public/Chaos/PerParticlePBDCollisionConstraint.h
--- a/Engine/Source/Runtime/Experimental/Chaos/Public/Chaos/PerParticlePBDCollisionConstraint.h
+++ b/Engine/Source/Runtime/Experimental/Chaos/Public/Chaos/PerParticlePBDCollisionConstraint.h
@@ -4,7 +4,7 @@
 #include "Chaos/PerParticleRule.h"
 #include "Chaos/Transform.h"
 
-#include <memory>
+#include "HAL/PlatformMath.h"
 
 namespace Chaos
 {
@@ -18,29 +18,36 @@ class TPerParticlePBDCollisionConstraint : public TPerParticleRule<T, d>
 	};
 
   public:
-	TPerParticlePBDCollisionConstraint(const TKinematicGeometryParticlesImp<T, d, SimType>& InParticles, TArrayCollectionArray<bool>& Collided, TArrayCollectionArray<uint32>& DynamicGroupIds, TArrayCollectionArray<uint32>& KinematicGroupIds, const TArray<T>& PerGroupThickness, const TArray<T>& PerGroupFriction)
+	TPerParticlePBDCollisionConstraint(const TKinematicGeometryParticlesImp<T, d, SimType>& InParticles, TArray<bool>& Collided, TArray<uint32>& DynamicGroupIds, TArray<uint32>& KinematicGroupIds, const TArray<T>& PerGroupThickness, const TArray<T>& PerGroupFriction)
 	    : bFastPositionBasedFriction(true), MParticles(InParticles), MCollided(Collided), MDynamicGroupIds(DynamicGroupIds), MKinematicGroupIds(KinematicGroupIds), MPerGroupThickness(PerGroupThickness), MPerGroupFriction(PerGroupFriction) {}
 	virtual ~TPerParticlePBDCollisionConstraint() {}
 
 	inline void Apply(TPBDParticles<T, d>& InParticles, const T Dt, const int32 Index) const override //-V762
 	{
 		if (InParticles.InvM(Index) == 0)
+		{
 			return;
+		}
+
+		const uint32 DynamicGroupId = MDynamicGroupIds[Index];  // Particle group Id
+		
 		for (uint32 i = 0; i < MParticles.Size(); ++i)
 		{
-			if (MDynamicGroupIds[Index] != MKinematicGroupIds[i])
+			const uint32 KinematicGroupId = MKinematicGroupIds[i];  // Collision group Id
+
+			if (KinematicGroupId != (uint32)INDEX_NONE && DynamicGroupId != KinematicGroupId)
 			{
-				continue; // Bail out if the collision groups doesn't match the particle group id
+				continue; // Bail out if the collision groups doesn't match the particle group id, or use INDEX_NONE (= global collision that affects all particle)
 			}
 			TVector<T, d> Normal;
 			TRigidTransform<T, d> Frame(MParticles.X(i), MParticles.R(i));
 			const T Phi = MParticles.Geometry(i)->PhiWithNormal(Frame.InverseTransformPosition(InParticles.P(Index)), Normal);
-			const T Thickness = MPerGroupThickness[MDynamicGroupIds[Index]];
+			const T Thickness = MPerGroupThickness[DynamicGroupId];
 			if (Phi < Thickness)
 			{
 				const TVector<T, d> NormalWorld = Frame.TransformVector(Normal);
 				InParticles.P(Index) += (-Phi + Thickness) * NormalWorld;
-				const T CoefficientOfFriction = MPerGroupFriction[MDynamicGroupIds[Index]];
+				const T CoefficientOfFriction = MPerGroupFriction[DynamicGroupId];
 				if (CoefficientOfFriction > 0)
 				{
 					const T MaximumFrictionCorrectionPerStep = 1.0f; // 1cm absolute maximum correction the friction force can provide
@@ -89,7 +96,7 @@ class TPerParticlePBDCollisionConstraint : public TPerParticleRule<T, d>
 			const TVector<T, d> VTBody = MVelocityConstraints[Index].Velocity - VNBody * MVelocityConstraints[Index].Normal;
 			const TVector<T, d> VTRelative = InParticles.V(Index) - VN * MVelocityConstraints[Index].Normal - VTBody;
 			const T VTRelativeSize = VTRelative.Size();
-			const T VNMax = FGenericPlatformMath::Max(VN, VNBody);
+			const T VNMax = FMath::Max(VN, VNBody);
 			const T VNDelta = VNMax - VN;
 			const T CoefficientOfFriction = MPerGroupFriction[MDynamicGroupIds[Index]];
 			check(CoefficientOfFriction > 0);
@@ -102,9 +109,9 @@ class TPerParticlePBDCollisionConstraint : public TPerParticleRule<T, d>
 	bool bFastPositionBasedFriction;
 	// TODO(mlentine): Need a bb hierarchy
 	const TKinematicGeometryParticlesImp<T, d, SimType>& MParticles;
-	TArrayCollectionArray<bool>& MCollided;
-	TArrayCollectionArray<uint32>& MDynamicGroupIds;
-	TArrayCollectionArray<uint32>& MKinematicGroupIds;
+	TArray<bool>& MCollided;
+	TArray<uint32>& MDynamicGroupIds;
+	TArray<uint32>& MKinematicGroupIds;
 	mutable TMap<int32, VelocityConstraint> MVelocityConstraints;
 	const TArray<T>& MPerGroupThickness;
 	const TArray<T>& MPerGroupFriction;
diff --git a/Engine/Source/Runtime/Experimental/Chaos/Public/Chaos/TriangleMesh.h b/Engine/Source/Runtime/Experimental/Chaos/Public/Chaos/TriangleMesh.h
--- a/Engine/Source/Runtime/Experimental/Chaos/Public/Chaos/TriangleMesh.h
+++ b/Engine/Source/Runtime/Experimental/Chaos/Public/Chaos/TriangleMesh.h
@@ -5,8 +5,7 @@
 #include "Chaos/Map.h"
 #include "Chaos/Particles.h"
 #include "Chaos/SegmentMesh.h"
-
-#include <unordered_set>
+#include "Containers/ContainersFwd.h"
 
 namespace Chaos
 {
@@ -37,7 +36,7 @@ namespace Chaos
 		 *
 		 * If this mesh is empty, the second index of the range will be negative.
 		 */
-		TPair<int32, int32> GetVertexRange() const;
+		TVector<int32, 2> GetVertexRange() const;
 
 		/** Returns the set of vertices used by triangles. */
 		TSet<int32> GetVertices() const;
@@ -52,8 +51,8 @@ namespace Chaos
 		 */
 		void ExpandVertexRange(const int32 StartIdx, const int32 EndIdx)
 		{
-			const TPair<int32, int32> CurrRange = GetVertexRange();
-			if (StartIdx <= CurrRange.Key && EndIdx >= CurrRange.Value)
+			const TVector<int32, 2> CurrRange = GetVertexRange();
+			if (StartIdx <= CurrRange[0] && EndIdx >= CurrRange[1])
 			{
 				MStartIdx = StartIdx;
 				MNumIndices = EndIdx - StartIdx + 1;
@@ -91,7 +90,7 @@ namespace Chaos
 		const TMap<int32, TSet<uint32>>& GetPointToNeighborsMap() const;
 		const TSet<uint32>& GetNeighbors(const int32 Element) const { return GetPointToNeighborsMap()[Element]; }
 
-		const TMap<int32, TArray<int32>>& GetPointToTriangleMap() const;
+		TConstArrayView<TArray<int32>> GetPointToTriangleMap() const;  // Return an array view using global indexation. Only elements starting at MStartIdx will be valid!
 		const TArray<int32>& GetCoincidentTriangles(const int32 Element) const { return GetPointToTriangleMap()[Element]; }
 
 		TSet<int32> GetNRing(const int32 Element, const int32 N) const
@@ -137,25 +136,24 @@ namespace Chaos
 		/** The GetFaceNormals functions assume Counter Clockwise triangle windings in a Left Handed coordinate system
 			If this is not the case the returned face normals may be inverted
 		*/
-		TArray<TVector<T, 3>> GetFaceNormals(const TArrayView<const TVector<T, 3>>& Points, const bool ReturnEmptyOnError = true) const;
-		void GetFaceNormals(TArray<TVector<T, 3>>& Normals, const TArrayView<const TVector<T, 3>>& Points, const bool ReturnEmptyOnError = true) const;
-		/** Deprecated. Use TArrayView version. */
+		TArray<TVector<T, 3>> GetFaceNormals(const TConstArrayView<TVector<T, 3>>& Points, const bool ReturnEmptyOnError = true) const;
+		void GetFaceNormals(TArray<TVector<T, 3>>& Normals, const TConstArrayView<TVector<T, 3>>& Points, const bool ReturnEmptyOnError = true) const;
 		TArray<TVector<T, 3>> GetFaceNormals(const TParticles<T, 3>& InParticles, const bool ReturnEmptyOnError = true) const
 		{ return GetFaceNormals(InParticles.X(), ReturnEmptyOnError); }
 
-		TArray<TVector<T, 3>> GetPointNormals(const TArrayView<const TVector<T, 3>>& points, const bool bReturnEmptyOnError = false, const bool bUseGlobalArray = true);
-		void GetPointNormals(TArray<TVector<T, 3>>& PointNormals, const TArray<TVector<T, 3>>& FaceNormals, const bool ReturnEmptyOnError = false, const bool bUseGlobalArray = true);
+		TArray<TVector<T, 3>> GetPointNormals(const TConstArrayView<TVector<T, 3>>& points, const bool ReturnEmptyOnError = true);
+		TArray<TVector<T, 3>> GetPointNormals(const TParticles<T, 3>& InParticles, const bool ReturnEmptyOnError = true)
+		{ return GetPointNormals(InParticles.X(), ReturnEmptyOnError); }
+
+		void GetPointNormals(TArrayView<TVector<T, 3>> PointNormals, const TConstArrayView<TVector<T, 3>>& FaceNormals, const bool bUseGlobalArray);
 		/** \brief Get per-point normals. 
 		 * This const version of this function requires \c GetPointToTriangleMap() 
 		 * to be called prior to invoking this function. 
 		 * @param bUseGlobalArray When true, fill the array from the StartIdx to StartIdx + NumIndices - 1 positions, otherwise fill the array from the 0 to NumIndices - 1 positions.
 		 */
-		void GetPointNormals(TArray<TVector<T, 3>>& PointNormals, const TArray<TVector<T, 3>>& FaceNormals, const bool bReturnEmptyOnError, const bool bUseGlobalArray) const;
-		/** Deprecated. Use TArrayView version. */
-		TArray<TVector<T, 3>> GetPointNormals(const TParticles<T, 3>& InParticles, const bool bReturnEmptyOnError = false, const bool bUseGlobalArray = true)
-		{ return GetPointNormals(InParticles.X(), bReturnEmptyOnError, bUseGlobalArray); }
+		void GetPointNormals(TArrayView<TVector<T, 3>> PointNormals, const TConstArrayView<TVector<T, 3>>& FaceNormals, const bool bUseGlobalArray) const;
 
-		static TTriangleMesh<T> GetConvexHullFromParticles(const TArrayView<const TVector<T, 3>>& points);
+		static TTriangleMesh<T> GetConvexHullFromParticles(const TConstArrayView<TVector<T, 3>>& points);
 		/** Deprecated. Use TArrayView version. */
 		static TTriangleMesh<T> GetConvexHullFromParticles(const TParticles<T, 3>& InParticles)
 		{ return GetConvexHullFromParticles(InParticles.X()); }
@@ -177,7 +175,7 @@ namespace Chaos
 		 */
 		TArray<T> GetCurvatureOnEdges(const TArray<TVector<T, 3>>& faceNormals);
 		/** @brief Helper that generates face normals on the fly. */
-		TArray<T> GetCurvatureOnEdges(const TArrayView<const TVector<T, 3>>& points);
+		TArray<T> GetCurvatureOnEdges(const TConstArrayView<TVector<T, 3>>& points);
 
 		/**
 		 * @ret The maximum curvature at points from connected edges, specified in radians.
@@ -186,7 +184,7 @@ namespace Chaos
 		 */
 		TArray<T> GetCurvatureOnPoints(const TArray<T>& edgeCurvatures);
 		/** @brief Helper that generates edge curvatures on the fly. */
-		TArray<T> GetCurvatureOnPoints(const TArrayView<const TVector<T, 3>>& points);
+		TArray<T> GetCurvatureOnPoints(const TConstArrayView<TVector<T, 3>>& points);
 
 		/**
 		 * Get the set of point indices that live on the boundary (an edge with only 1 
@@ -201,7 +199,7 @@ namespace Chaos
 		 */
 		TMap<int32, int32> FindCoincidentVertexRemappings(
 			const TArray<int32>& TestIndices,
-			const TArrayView<const TVector<T, 3>>& Points);
+			const TConstArrayView<TVector<T, 3>>& Points);
 
 		/**
 		 * @ret An array of vertex indices ordered from most important to least.
@@ -211,13 +209,13 @@ namespace Chaos
 		 * @param RestrictToLocalIndexRange - ignores points outside of the index range used by this mesh.
 		 */
 		TArray<int32> GetVertexImportanceOrdering(
-		    const TArrayView<const TVector<T, 3>>& Points,
+		    const TConstArrayView<TVector<T, 3>>& Points,
 		    const TArray<T>& PointCurvatures,
 		    TArray<int32>* CoincidentVertices = nullptr,
 		    const bool RestrictToLocalIndexRange = false);
 		/** @brief Helper that generates point curvatures on the fly. */
 		TArray<int32> GetVertexImportanceOrdering(
-		    const TArrayView<const TVector<T, 3>>& Points,
+		    const TConstArrayView<TVector<T, 3>>& Points,
 		    TArray<int32>* CoincidentVertices = nullptr,
 		    const bool RestrictToLocalIndexRange = false);
 
@@ -269,9 +267,16 @@ namespace Chaos
 			return LocalIdx;
 		}
 
+		int32 LocalToGlobal(int32 LocalIdx) const
+		{
+			const int32 GlobalIdx = LocalIdx + MStartIdx;
+			check(GlobalIdx >= MStartIdx && GlobalIdx < MStartIdx + MNumIndices);
+			return GlobalIdx;
+		}
+
 		TArray<TVector<int32, 3>> MElements;
 
-		mutable TMap<int32, TArray<int32>> MPointToTriangleMap;
+		mutable TArray<TArray<int32>> MPointToTriangleMap;  // !! Unlike the TArrayView returned by GetPointToTriangleMap, this array starts at 0 for the point of index MStartIdx. Use GlobalToLocal to access with a global index. Note that this array's content is always indexed in global index.
 		mutable TMap<int32, TSet<uint32>> MPointToNeighborsMap;
 
 		TSegmentMesh<T> MSegmentMesh;
diff --git a/Engine/Source/Runtime/Experimental/Chaos/Public/Chaos/VelocityField.h b/Engine/Source/Runtime/Experimental/Chaos/Public/Chaos/VelocityField.h
--- a/Engine/Source/Runtime/Experimental/Chaos/Public/Chaos/VelocityField.h
+++ b/Engine/Source/Runtime/Experimental/Chaos/Public/Chaos/VelocityField.h
@@ -12,126 +12,142 @@ namespace Chaos
 	class TVelocityField
 	{
 	public:
+		// Construct an uninitialized field. Mesh, properties, and velocity will have to be set for this field to be valid.
 		TVelocityField()
-			: Elements(nullptr), QuarterRhoDragCoefficient((T)0.), bIsUniform(true) {}
+			: Range(-1)
+		{}
 
+		// Construct a uniform field.
+		TVelocityField(
+			const TTriangleMesh<T>& TriangleMesh,
+			const TVector<T, d>& InVelocity,
+			const T InDragCoefficient = (T)0.5,
+			const T InFluidDensity = (T)1.225e-6)
+			: PointToTriangleMap(TriangleMesh.GetPointToTriangleMap())
+			, Elements(TriangleMesh.GetElements())
+			, Velocity(InVelocity)
+			, Range(TriangleMesh.GetVertexRange())
+		{
+			Forces.SetNumUninitialized(Elements->Num());
+			SetProperties(InDragCoefficient, InFluidDensity);
+		}
+
+		// Construct a vector field.
 		TVelocityField(
 			const TTriangleMesh<T>& TriangleMesh,
 			TFunction<TVector<T, d>(const TVector<T, d>&)> InGetVelocity,
-			const bool bInIsUniform,
-			const T DragCoefficient = (T)0.5,
-			const T FluidDensity = (T)1.225e-6)
-			: Elements(&TriangleMesh.GetElements())
+			const T InDragCoefficient = (T)0.5,
+			const T InFluidDensity = (T)1.225e-6)
+			: PointToTriangleMap(TriangleMesh.GetPointToTriangleMap())
+			, Elements(TriangleMesh.GetElements())
+			, Velocity((T)0.)
 			, GetVelocity(InGetVelocity)
-			, bIsUniform(bInIsUniform)
-			, Offset(TriangleMesh.GetVertexRange().Key)
+			, Range(TriangleMesh.GetVertexRange())
 		{
 			Forces.SetNumUninitialized(Elements->Num());
-			SetProperties(DragCoefficient, FluidDensity);
-
-			const int32 NumPoints = TriangleMesh.GetVertexRange().Value + 1 - Offset;
-			PointToTriangleMap.SetNumUninitialized(NumPoints);
-			for (const TPair<int32, TArray<int32>>& PointToTriangle : TriangleMesh.GetPointToTriangleMap())
-			{
-				PointToTriangleMap[PointToTriangle.Key - Offset] = &PointToTriangle.Value;
-			}
+			SetProperties(InDragCoefficient, InFluidDensity);
 		}
 
 		virtual ~TVelocityField() {}
 
-		inline void UpdateForces(const TPBDParticles<T, d>& InParticles, const T Dt) const
+		inline void UpdateForces(const TPBDParticles<T, d>& InParticles, const T /*Dt*/) const
 		{
-			if (GetVelocity)
+			if (!GetVelocity)
 			{
-				if (bIsUniform)
+				for (int32 ElementIndex = 0; ElementIndex < Elements.Num(); ++ElementIndex)
 				{
-					const TVector<T, d> Velocity = GetVelocity(TVector<T, d>(0.));
-					for (int32 ElementIndex = 0; ElementIndex < Elements->Num(); ++ElementIndex)
-					{
-						UpdateUniformField(InParticles, Dt, ElementIndex, Velocity);
-					}
+					UpdateField(InParticles, ElementIndex, Velocity);
 				}
-				else
+			}
+			else
+			{
+				for (int32 ElementIndex = 0; ElementIndex < Elements.Num(); ++ElementIndex)
 				{
-					for (int32 ElementIndex = 0; ElementIndex < Elements->Num(); ++ElementIndex)
-					{
-						UpdateVectorField(InParticles, Dt, ElementIndex);
-					}
+					const TVector<int32, 3>& Element = Elements[ElementIndex];
+
+					// Get the triangle's position
+					const TVector<T, d>& SurfacePosition = (T)(1. / 3.) * (
+						InParticles.X(Element[0]) +
+						InParticles.X(Element[1]) +
+						InParticles.X(Element[2]));
+
+					UpdateField(InParticles, ElementIndex, GetVelocity(SurfacePosition));
 				}
 			}
 		}
 
 		inline void Apply(TPBDParticles<T, d>& InParticles, const T Dt, const int32 Index) const
 		{
-			const int32 PointIndex = Index - Offset;
-			if (PointToTriangleMap.IsValidIndex(PointIndex))
-			{
-				const TArray<int32>* ElementIndices = PointToTriangleMap[PointIndex];
+			check(Index >= Range[0] && Index <= Range[1]);  // The index should always match the original triangle mesh range
 
-				for (const int32 ElementIndex : *ElementIndices)
-				{
-					InParticles.F(Index) += Forces[ElementIndex];
-				}
+			const TArray<int32>& ElementIndices = PointToTriangleMap[Index];
+			for (const int32 ElementIndex : ElementIndices)
+			{
+				InParticles.F(Index) += Forces[ElementIndex];
 			}
 		}
 
-		void SetProperties(const T DragCoefficient, const T FluidDensity)
+		void SetProperties(const T InDragCoefficient, const T InFluidDensity)
 		{
+			DragCoefficient = InDragCoefficient;
+			FluidDensity = InFluidDensity;
 			QuarterRhoDragCoefficient = (T)0.25 * FluidDensity * DragCoefficient;
 		}
 
-		const TArray<TVector<int32, 3>>& GetElements() const { return *Elements; }
-		const TArray<TVector<T, d>>& GetForces() const { return Forces; }
+		void SetDragCoefficient(const T InDragCoefficient) { SetProperties(InDragCoefficient, FluidDensity); }
+		void SetFluidDensity(const T InFluidDensity) { SetProperties(DragCoefficient, InFluidDensity); }
 
-	private:
-		inline void UpdateUniformField(const TPBDParticles<T, d>& InParticles, const T Dt, int32 ElementIndex, const TVector<T, d>& Velocity) const
+		void SetGeometry(const TTriangleMesh<T>* TriangleMesh)
 		{
-			const TVector<int32, 3>& Element = (*Elements)[ElementIndex];
-
-			// Calculate direction and the relative velocity of the triangle to the flow
-			const TVector<T, d>& SurfaceVelocity = (T)(1. / 3.) * (
-				InParticles.V(Element[0]) + 
-				InParticles.V(Element[1]) + 
-				InParticles.V(Element[2]));
-			TVector<T, d> Direction = Velocity - SurfaceVelocity;
-			const T RelVelocity = Direction.SafeNormalize();
-
-			// Calculate the cross sectional area of the surface exposed to the flow
-			TVector<float, 3> Normal = TVector<float, 3>::CrossProduct(
-				InParticles.X(Element[1]) - InParticles.X(Element[0]),
-				InParticles.X(Element[2]) - InParticles.X(Element[0]));
-			const T DoubleArea = Normal.SafeNormalize();
-			const T DoubleCrossSectionalArea = DoubleArea * FMath::Abs(TVector<float, 3>::DotProduct(Direction, Normal));
+			if (TriangleMesh)
+			{
+				PointToTriangleMap = TriangleMesh->GetPointToTriangleMap();
+				Elements = TriangleMesh->GetElements();
+				Range = TriangleMesh->GetVertexRange();
+				Forces.SetNumUninitialized(Elements.Num());
+			}
+			else
+			{
+				PointToTriangleMap = TArrayView<TArray<int32>>();
+				Elements = TArrayView<TVector<int32, 3>>();
+				Range = TVector<int32, 2>(-1);
+				Forces.Reset();
+			}
+		}
 
-			// Calculate the drag force
-			const T Drag = QuarterRhoDragCoefficient * FMath::Square(RelVelocity) * DoubleCrossSectionalArea;
-			Forces[ElementIndex] = Direction * Drag;
+		void SetVelocity(const TVector<T, d>& InVelocity)
+		{
+			Velocity = InVelocity;
+			GetVelocity = TFunction<TVector<T, d>(const TVector<T, d>&)>();
 		}
 
-		inline void UpdateVectorField(const TPBDParticles<T, d>& InParticles, const T Dt, int32 ElementIndex) const
+		void SetVelocity(TFunction<TVector<T, d>(const TVector<T, d>&)> InGetVelocity)
 		{
-			const TVector<int32, 3>& Element = (*Elements)[ElementIndex];
+			GetVelocity = InGetVelocity;
+		}
 
-			// Get the triangle's position
-			const TVector<T, d>& SurfacePosition = (T)(1. / 3.) * (
-				InParticles.X(Element[0]) +
-				InParticles.X(Element[1]) +
-				InParticles.X(Element[2]));
+		const TConstArrayView<TVector<int32, 3>>& GetElements() const { return Elements; }
+		const TArray<TVector<T, d>>& GetForces() const { return Forces; }
 
-			// Calculate direction and the relative velocity of the triangle to the flow at the field's sampled position
+	private:
+		inline void UpdateField(const TPBDParticles<T, d>& InParticles, int32 ElementIndex, const TVector<T, d>& InVelocity) const
+		{
+			const TVector<int32, 3>& Element = Elements[ElementIndex];
+
+			// Calculate direction and the relative velocity of the triangle to the flow
 			const TVector<T, d>& SurfaceVelocity = (T)(1. / 3.) * (
-				InParticles.V(Element[0]) + 
-				InParticles.V(Element[1]) + 
+				InParticles.V(Element[0]) +
+				InParticles.V(Element[1]) +
 				InParticles.V(Element[2]));
-			TVector<T, d> Direction = GetVelocity(SurfacePosition) - SurfaceVelocity;
+			TVector<T, d> Direction = InVelocity - SurfaceVelocity;
 			const T RelVelocity = Direction.SafeNormalize();
 
 			// Calculate the cross sectional area of the surface exposed to the flow
-			 TVector<float, 3> Normal = TVector<float, 3>::CrossProduct(
+			TVector<T, d> Normal = TVector<T, d>::CrossProduct(
 				InParticles.X(Element[1]) - InParticles.X(Element[0]),
 				InParticles.X(Element[2]) - InParticles.X(Element[0]));
 			const T DoubleArea = Normal.SafeNormalize();
-			const T DoubleCrossSectionalArea = DoubleArea * FMath::Abs(TVector<float, 3>::DotProduct(Direction, Normal));
+			const T DoubleCrossSectionalArea = DoubleArea * FMath::Abs(TVector<T, d>::DotProduct(Direction, Normal));
 
 			// Calculate the drag force
 			const T Drag = QuarterRhoDragCoefficient * FMath::Square(RelVelocity) * DoubleCrossSectionalArea;
@@ -139,12 +155,14 @@ namespace Chaos
 		}
 
 	private:
-		TArray<const TArray<int32>*> PointToTriangleMap;
-		const TArray<TVector<int32, 3>>* Elements;
+		TConstArrayView<TArray<int32>> PointToTriangleMap;
+		TConstArrayView<TVector<int32, 3>> Elements;
+		TVector<T, d> Velocity;
 		TFunction<TVector<T, d>(const TVector<T, d>&)> GetVelocity;
 		mutable TArray<TVector<T, 3>> Forces;
 		float QuarterRhoDragCoefficient;
-		bool bIsUniform;
-		int32 Offset;
+		float DragCoefficient;
+		float FluidDensity;
+		TVector<int32, 2> Range;  // TODO: Remove? It is used by the check only
 	};
 }
