Chaos Cloth - Fix tapered capsule collision code that was running with an uninitialized normal since CL13811735.

Merging CL 14009486

diff --git a/Engine/Source/Runtime/Experimental/Chaos/Public/Chaos/TaperedCylinder.h b/Engine/Source/Runtime/Experimental/Chaos/Public/Chaos/TaperedCylinder.h
--- a/Engine/Source/Runtime/Experimental/Chaos/Public/Chaos/TaperedCylinder.h
+++ b/Engine/Source/Runtime/Experimental/Chaos/Public/Chaos/TaperedCylinder.h
@@ -116,11 +116,11 @@ namespace Chaos
 
 		T PhiWithNormal(const TVector<T, 3>& x, TVector<T, 3>& Normal) const
 		{
-			TVector<T, 3> Normal1, Normal2;
-			const T Distance1 = MPlane1.PhiWithNormal(x, Normal1);
+			const TVector<T, 3>& Normal1 = MPlane1.Normal();
+			const T Distance1 = MPlane1.SignedDistance(x);
 			if (Distance1 < SMALL_NUMBER)
 			{
-				ensure(MPlane2.PhiWithNormal(x, Normal2) > 0.);
+				ensure(MPlane2.SignedDistance(x) > (T)0.);
 				const TVector<T, 3> v = x - TVector<T, 3>(Normal1 * Distance1 + MPlane1.X());
 				if (v.Size() > MRadius1)
 				{
@@ -135,10 +135,10 @@ namespace Chaos
 					return -Distance1;
 				}
 			}
-			const T Distance2 = MHeight - Distance1; // Used to be Distance2 = MPlane2.PhiWithNormal(x, Normal2); but that would trigger the ensure on Distance2 being slightly larger than MHeight in some border cases
+			const TVector<T, 3>& Normal2 = MPlane2.Normal();  // Used to be Distance2 = MPlane2.PhiWithNormal(x, Normal2); but that would trigger 
+			const T Distance2 = MHeight - Distance1;          // the ensure on Distance2 being slightly larger than MHeight in some border cases
 			if (Distance2 < SMALL_NUMBER)
 			{
-				ensure(MPlane1.PhiWithNormal(x, Normal1) > 0.);
 				const TVector<T, 3> v = x - TVector<T, 3>(Normal2 * Distance2 + MPlane2.X());
 				if (v.Size() > MRadius2)
 				{
