CIS - Fix missing "using" showing with some compilers since CL 14112532.

diff --git a/Engine/Source/Runtime/Experimental/Chaos/Public/Chaos/PBDShapeConstraints.h b/Engine/Source/Runtime/Experimental/Chaos/Public/Chaos/PBDShapeConstraints.h
--- a/Engine/Source/Runtime/Experimental/Chaos/Public/Chaos/PBDShapeConstraints.h
+++ b/Engine/Source/Runtime/Experimental/Chaos/Public/Chaos/PBDShapeConstraints.h
@@ -9,6 +9,7 @@ namespace Chaos
 	class TPBDShapeConstraints : public TPBDShapeConstraintsBase<T, d>
 	{
 		typedef TPBDShapeConstraintsBase<T, d> Base;
+		using Base::MParticleOffset;
 
 	public:
 
