CIS - Fix another missing "using" showing with some compilers since CL 14112532.

diff --git a/Engine/Source/Runtime/Experimental/Chaos/Public/Chaos/PBDShapeConstraints.h b/Engine/Source/Runtime/Experimental/Chaos/Public/Chaos/PBDShapeConstraints.h
--- a/Engine/Source/Runtime/Experimental/Chaos/Public/Chaos/PBDShapeConstraints.h
+++ b/Engine/Source/Runtime/Experimental/Chaos/Public/Chaos/PBDShapeConstraints.h
@@ -10,6 +10,7 @@ namespace Chaos
 	{
 		typedef TPBDShapeConstraintsBase<T, d> Base;
 		using Base::MParticleOffset;
+		using Base::MTargetPositions;
 
 	public:
 
