CIS - Fix VelocityField error caught on Linux.

diff --git a/Engine/Source/Runtime/Experimental/Chaos/Public/Chaos/VelocityField.h b/Engine/Source/Runtime/Experimental/Chaos/Public/Chaos/VelocityField.h
--- a/Engine/Source/Runtime/Experimental/Chaos/Public/Chaos/VelocityField.h
+++ b/Engine/Source/Runtime/Experimental/Chaos/Public/Chaos/VelocityField.h
@@ -28,7 +28,7 @@ namespace Chaos
 			, Velocity(InVelocity)
 			, Range(TriangleMesh.GetVertexRange())
 		{
-			Forces.SetNumUninitialized(Elements->Num());
+			Forces.SetNumUninitialized(Elements.Num());
 			SetProperties(InDragCoefficient, InFluidDensity);
 		}
 
@@ -44,7 +44,7 @@ namespace Chaos
 			, GetVelocity(InGetVelocity)
 			, Range(TriangleMesh.GetVertexRange())
 		{
-			Forces.SetNumUninitialized(Elements->Num());
+			Forces.SetNumUninitialized(Elements.Num());
 			SetProperties(InDragCoefficient, InFluidDensity);
 		}
 
