CIS - Fix missing include that shows up in some platforms.

Merging CL 14156095

diff --git a/Engine/Source/Runtime/Experimental/Chaos/Public/Chaos/PBDActiveView.h b/Engine/Source/Runtime/Experimental/Chaos/Public/Chaos/PBDActiveView.h
--- a/Engine/Source/Runtime/Experimental/Chaos/Public/Chaos/PBDActiveView.h
+++ b/Engine/Source/Runtime/Experimental/Chaos/Public/Chaos/PBDActiveView.h
@@ -3,6 +3,7 @@
 
 #include "Containers/Array.h"
 #include "Algo/BinarySearch.h"
+#include "Chaos/Framework/Parallel.h"
 
 namespace Chaos
 {
