Hide Chaos Debug Draw Actor from WorldOutliner

diff --git a/Engine/Source/Runtime/Experimental/ChaosSolverEngine/Private/Chaos/ChaosDebugDrawComponent.cpp b/Engine/Source/Runtime/Experimental/ChaosSolverEngine/Private/Chaos/ChaosDebugDrawComponent.cpp
--- a/Engine/Source/Runtime/Experimental/ChaosSolverEngine/Private/Chaos/ChaosDebugDrawComponent.cpp
+++ b/Engine/Source/Runtime/Experimental/ChaosSolverEngine/Private/Chaos/ChaosDebugDrawComponent.cpp
@@ -224,6 +224,7 @@ void UChaosDebugDrawComponent::CreateDebugDrawActor(UWorld* World)
 	FActorSpawnParameters Params;
 	Params.Name = NAME_ChaosDebugDrawActor;
 	Params.ObjectFlags = Params.ObjectFlags | RF_Transient;
+	Params.bHideFromSceneOutliner = true;
 
 	AActor* Actor = World->SpawnActor<AActor>(FVector::ZeroVector, FRotator::ZeroRotator, Params);
 	
