Fixed non-editor build

#rnx

diff --git a/Engine/Source/Runtime/Experimental/ChaosSolverEngine/Private/Chaos/ChaosDebugDrawComponent.cpp b/Engine/Source/Runtime/Experimental/ChaosSolverEngine/Private/Chaos/ChaosDebugDrawComponent.cpp
--- a/Engine/Source/Runtime/Experimental/ChaosSolverEngine/Private/Chaos/ChaosDebugDrawComponent.cpp
+++ b/Engine/Source/Runtime/Experimental/ChaosSolverEngine/Private/Chaos/ChaosDebugDrawComponent.cpp
@@ -224,7 +224,9 @@ void UChaosDebugDrawComponent::CreateDebugDrawActor(UWorld* World)
 	FActorSpawnParameters Params;
 	Params.Name = NAME_ChaosDebugDrawActor;
 	Params.ObjectFlags = Params.ObjectFlags | RF_Transient;
+#if WITH_EDITOR
 	Params.bHideFromSceneOutliner = true;
+#endif
 
 	AActor* Actor = World->SpawnActor<AActor>(FVector::ZeroVector, FRotator::ZeroRotator, Params);
 	
