InteractiveToolsFramework: Add DrawHUD capability to Tools

#rnx

diff --git a/Engine/Plugins/Experimental/ModelingToolsEditorMode/Source/ModelingToolsEditorMode/Private/ModelingToolsEditorMode.cpp b/Engine/Plugins/Experimental/ModelingToolsEditorMode/Source/ModelingToolsEditorMode/Private/ModelingToolsEditorMode.cpp
--- a/Engine/Plugins/Experimental/ModelingToolsEditorMode/Source/ModelingToolsEditorMode/Private/ModelingToolsEditorMode.cpp
+++ b/Engine/Plugins/Experimental/ModelingToolsEditorMode/Source/ModelingToolsEditorMode/Private/ModelingToolsEditorMode.cpp
@@ -196,6 +196,15 @@ void FModelingToolsEditorMode::Render(const FSceneView* View, FViewport* Viewpor
 	}
 }
 
+void FModelingToolsEditorMode::DrawHUD(FEditorViewportClient* ViewportClient,FViewport* Viewport,const FSceneView* View,FCanvas* Canvas)
+{
+	FEdMode::DrawHUD(ViewportClient, Viewport, View, Canvas);
+	if (ToolsContext != nullptr)
+	{
+		ToolsContext->DrawHUD(Canvas);
+	}
+}
+
 
 
 
diff --git a/Engine/Plugins/Experimental/ModelingToolsEditorMode/Source/ModelingToolsEditorMode/Public/ModelingToolsEditorMode.h b/Engine/Plugins/Experimental/ModelingToolsEditorMode/Source/ModelingToolsEditorMode/Public/ModelingToolsEditorMode.h
--- a/Engine/Plugins/Experimental/ModelingToolsEditorMode/Source/ModelingToolsEditorMode/Public/ModelingToolsEditorMode.h
+++ b/Engine/Plugins/Experimental/ModelingToolsEditorMode/Source/ModelingToolsEditorMode/Public/ModelingToolsEditorMode.h
@@ -30,6 +30,8 @@ public:
 
 	virtual void Tick(FEditorViewportClient* ViewportClient, float DeltaTime) override;
 	virtual void Render(const FSceneView* View, FViewport* Viewport, FPrimitiveDrawInterface* PDI) override;
+	virtual void DrawHUD(FEditorViewportClient* ViewportClient,FViewport* Viewport,const FSceneView* View,FCanvas* Canvas) override;
+
 	virtual bool UsesToolkits() const override;
 
 	virtual void AddReferencedObjects(FReferenceCollector& Collector) override;
diff --git a/Engine/Source/Editor/Experimental/EditorInteractiveToolsFramework/Private/EdModeInteractiveToolsContext.cpp b/Engine/Source/Editor/Experimental/EditorInteractiveToolsFramework/Private/EdModeInteractiveToolsContext.cpp
--- a/Engine/Source/Editor/Experimental/EditorInteractiveToolsFramework/Private/EdModeInteractiveToolsContext.cpp
+++ b/Engine/Source/Editor/Experimental/EditorInteractiveToolsFramework/Private/EdModeInteractiveToolsContext.cpp
@@ -732,6 +732,10 @@ void UEdModeInteractiveToolsContext::Render(const FSceneView* View, FViewport* V
 	GizmoManager->Render(&RenderContext);
 }
 
+void UEdModeInteractiveToolsContext::DrawHUD(FCanvas* Canvas)
+{
+	ToolManager->DrawHUD(Canvas);
+}
 
 
 bool UEdModeInteractiveToolsContext::ProcessEditDelete()
diff --git a/Engine/Source/Editor/Experimental/EditorInteractiveToolsFramework/Public/EdModeInteractiveToolsContext.h b/Engine/Source/Editor/Experimental/EditorInteractiveToolsFramework/Public/EdModeInteractiveToolsContext.h
--- a/Engine/Source/Editor/Experimental/EditorInteractiveToolsFramework/Public/EdModeInteractiveToolsContext.h
+++ b/Engine/Source/Editor/Experimental/EditorInteractiveToolsFramework/Public/EdModeInteractiveToolsContext.h
@@ -60,6 +60,7 @@ public:
 
 	virtual void Tick(FEditorViewportClient* ViewportClient, float DeltaTime);
 	virtual void Render(const FSceneView* View, FViewport* Viewport, FPrimitiveDrawInterface* PDI);
+	virtual void DrawHUD(FCanvas* Canvas);
 
 	virtual bool ProcessEditDelete();
 
diff --git a/Engine/Source/Editor/UnrealEd/Private/Tools/UEdMode.cpp b/Engine/Source/Editor/UnrealEd/Private/Tools/UEdMode.cpp
--- a/Engine/Source/Editor/UnrealEd/Private/Tools/UEdMode.cpp
+++ b/Engine/Source/Editor/UnrealEd/Private/Tools/UEdMode.cpp
@@ -274,6 +274,12 @@ void UEdMode::Render(const FSceneView* View, FViewport* Viewport, FPrimitiveDraw
 
 void UEdMode::DrawHUD(FEditorViewportClient* ViewportClient, FViewport* Viewport, const FSceneView* View, FCanvas* Canvas)
 {
+	// give ToolsContext a chance to draw in screen space
+	if (ToolsContext != nullptr)
+	{
+		ToolsContext->DrawHUD(Canvas);
+	}
+
 	// Render the drag tool.
 	ViewportClient->RenderDragTool(View, Canvas);
 
diff --git a/Engine/Source/Runtime/Experimental/InteractiveToolsFramework/Private/InteractiveTool.cpp b/Engine/Source/Runtime/Experimental/InteractiveToolsFramework/Private/InteractiveTool.cpp
--- a/Engine/Source/Runtime/Experimental/InteractiveToolsFramework/Private/InteractiveTool.cpp
+++ b/Engine/Source/Runtime/Experimental/InteractiveToolsFramework/Private/InteractiveTool.cpp
@@ -49,6 +49,9 @@ void UInteractiveTool::Render(IToolsContextRenderAPI* RenderAPI)
 {
 }
 
+void UInteractiveTool::DrawHUD(FCanvas* Canvas)
+{
+}
 
 void UInteractiveTool::AddInputBehavior(UInputBehavior* Behavior)
 {
diff --git a/Engine/Source/Runtime/Experimental/InteractiveToolsFramework/Private/InteractiveToolManager.cpp b/Engine/Source/Runtime/Experimental/InteractiveToolsFramework/Private/InteractiveToolManager.cpp
--- a/Engine/Source/Runtime/Experimental/InteractiveToolsFramework/Private/InteractiveToolManager.cpp
+++ b/Engine/Source/Runtime/Experimental/InteractiveToolsFramework/Private/InteractiveToolManager.cpp
@@ -303,6 +303,19 @@ void UInteractiveToolManager::Render(IToolsContextRenderAPI* RenderAPI)
 	}
 }
 
+void UInteractiveToolManager::DrawHUD(FCanvas* Canvas)
+{
+	if (ActiveLeftTool != nullptr)
+	{
+		ActiveLeftTool->DrawHUD(Canvas);
+	}
+
+	if (ActiveRightTool != nullptr)
+	{
+		ActiveRightTool->DrawHUD(Canvas);
+	}
+}
+
 
 
 UInteractiveGizmoManager* UInteractiveToolManager::GetPairedGizmoManager()
@@ -483,4 +496,4 @@ FString FToolChangeWrapperChange::ToString() const
 
 
 
-#undef LOCTEXT_NAMESPACE
\ No newline at end of file
+#undef LOCTEXT_NAMESPACE
diff --git a/Engine/Source/Runtime/Experimental/InteractiveToolsFramework/Public/InteractiveTool.h b/Engine/Source/Runtime/Experimental/InteractiveToolsFramework/Public/InteractiveTool.h
--- a/Engine/Source/Runtime/Experimental/InteractiveToolsFramework/Public/InteractiveTool.h
+++ b/Engine/Source/Runtime/Experimental/InteractiveToolsFramework/Public/InteractiveTool.h
@@ -11,7 +11,7 @@
 #include "InteractiveTool.generated.h"
 
 class UInteractiveToolManager;
-
+class FCanvas;
 
 /** Passed to UInteractiveTool::Shutdown to indicate how Tool should shut itself down*/
 enum class EToolShutdownType
@@ -267,6 +267,7 @@ public:
 	 */
 	virtual void Render(IToolsContextRenderAPI* RenderAPI);
 
+	virtual void DrawHUD( FCanvas* Canvas );
 
 	/**
 	 * Non overrideable func which does processing and calls the tool's OnTick
diff --git a/Engine/Source/Runtime/Experimental/InteractiveToolsFramework/Public/InteractiveToolManager.h b/Engine/Source/Runtime/Experimental/InteractiveToolsFramework/Public/InteractiveToolManager.h
--- a/Engine/Source/Runtime/Experimental/InteractiveToolsFramework/Public/InteractiveToolManager.h
+++ b/Engine/Source/Runtime/Experimental/InteractiveToolsFramework/Public/InteractiveToolManager.h
@@ -223,6 +223,9 @@ public:
 	/** Render any active Tools. Called by UInteractiveToolsContext. */
 	virtual void Render(IToolsContextRenderAPI* RenderAPI);
 
+	/** Let active Tools do their screen space drawing. Called by UInteractiveToolsContext. */
+	virtual void DrawHUD(FCanvas* Canvas);
+
 
 	//
 	// access to APIs, etc
@@ -354,4 +357,4 @@ public:
 	virtual bool HasExpired(UObject* Object) const override;
 
 	virtual FString ToString() const override;
-};
\ No newline at end of file
+};
