Refactor of Nanite GeometryCollection serialization code to avoid the need to do tagged property versioning when Nanite::FResources is altered. This behavior now matches how regular StaticMesh serialization is working for Nanite resource data.

diff --git a/Engine/Source/Runtime/Core/Public/UObject/UE5MainStreamObjectVersion.h b/Engine/Source/Runtime/Core/Public/UObject/UE5MainStreamObjectVersion.h
--- a/Engine/Source/Runtime/Core/Public/UObject/UE5MainStreamObjectVersion.h
+++ b/Engine/Source/Runtime/Core/Public/UObject/UE5MainStreamObjectVersion.h
@@ -15,6 +15,9 @@ struct CORE_API FUE5MainStreamObjectVersion
 		// Nanite data added to Chaos geometry collections
 		GeometryCollectionNaniteData,
 
+		// Nanite Geometry Collection data moved to DDC
+		GeometryCollectionNaniteDDC,
+
 		// -----<new versions can be added above this line>-------------------------------------------------
 		VersionPlusOne,
 		LatestVersion = VersionPlusOne - 1
diff --git a/Engine/Source/Runtime/Experimental/GeometryCollectionEngine/Private/GeometryCollection/DerivedDataGeometryCollectionCooker.cpp b/Engine/Source/Runtime/Experimental/GeometryCollectionEngine/Private/GeometryCollection/DerivedDataGeometryCollectionCooker.cpp
--- a/Engine/Source/Runtime/Experimental/GeometryCollectionEngine/Private/GeometryCollection/DerivedDataGeometryCollectionCooker.cpp
+++ b/Engine/Source/Runtime/Experimental/GeometryCollectionEngine/Private/GeometryCollection/DerivedDataGeometryCollectionCooker.cpp
@@ -40,11 +40,16 @@ bool FDerivedDataGeometryCollectionCooker::Build(TArray<uint8>& OutData)
 		BuildSimulationData(ErrorReporter, *Collection, SharedParams);
 		Collection->Serialize(ChaosAr);
 
-		TArray<Nanite::FResources>& NaniteResources = GeometryCollection.CreateNaniteData(Collection);
-
-		for (Nanite::FResources& NaniteResource : NaniteResources)
+		if (GeometryCollection.EnableNanite)
+		{
+			TUniquePtr<FGeometryCollectionNaniteData> NaniteData = UGeometryCollection::CreateNaniteData(Collection);
+			NaniteData->Serialize(ChaosAr, &GeometryCollection);
+		}
+		else
 		{
-			NaniteResource.Serialize(ChaosAr, &GeometryCollection);
+			// No Nanite data, write out zero resources
+			int32 NumNaniteResources = 0;
+			Ar << NumNaniteResources;
 		}
 
 		if (false && ErrorReporter.EncounteredAnyErrors())
@@ -66,17 +71,12 @@ bool FDerivedDataGeometryCollectionCooker::Build(TArray<uint8>& OutData)
 
 const TCHAR* FDerivedDataGeometryCollectionCooker::GetVersionString() const
 {
-	if (OverrideVersion)
-	{
-		return OverrideVersion;	//force load old ddc if found. Not recommended
-	}
-
 	const TCHAR* VersionString = TEXT("A8A2C0FB45084FCB922FEC1139E11341");
 
 	static FString CachedNaniteVersionString;
 	if (CachedNaniteVersionString.IsEmpty())
 	{
-		const TCHAR* NaniteVersionString = TEXT("CE9B3AB5F81A5AB388E51FF4C8C19D7E");
+		const TCHAR* NaniteVersionString = TEXT("CE9B42F81A5XX388E51GGNC8C19D3E");
 		CachedNaniteVersionString = FString::Printf(TEXT("%s_%s_%s"), VersionString, NaniteVersionString, *Nanite::IBuilderModule::Get().GetVersionString());
 	}
 
@@ -95,5 +95,4 @@ FString FDerivedDataGeometryCollectionCooker::GetPluginSpecificCacheKeySuffix()
 	);
 }
 
-
 #endif // WITH_EDITOR
diff --git a/Engine/Source/Runtime/Experimental/GeometryCollectionEngine/Private/GeometryCollection/GeometryCollectionObject.cpp b/Engine/Source/Runtime/Experimental/GeometryCollectionEngine/Private/GeometryCollection/GeometryCollectionObject.cpp
--- a/Engine/Source/Runtime/Experimental/GeometryCollectionEngine/Private/GeometryCollection/GeometryCollectionObject.cpp
+++ b/Engine/Source/Runtime/Experimental/GeometryCollectionEngine/Private/GeometryCollection/GeometryCollectionObject.cpp
@@ -58,7 +58,6 @@ UGeometryCollection::UGeometryCollection(const FObjectInitializer& ObjectInitial
 	, CollisionParticlesFraction(1.0f)
 	, MaximumCollisionParticles(60)
 	, EnableRemovePiecesOnFracture(false)
-	, bRenderingResourcesInitialized(false)
 	, GeometryCollection(new FGeometryCollection())
 {
 	PersistentGuid = FGuid::NewGuid();
@@ -314,14 +313,6 @@ void UGeometryCollection::Serialize(FArchive& Ar)
 		Ar << bIsCookedOrCooking;
 	}
 
-#if WITH_EDITOR
-	if (Ar.CustomVer(FDestructionObjectVersion::GUID) == FDestructionObjectVersion::GeometryCollectionInDDC)
-	{
-		//This version only saved content into DDC so skip serializing, but copy from DDC at a specific version
-		bool bCopyFromDDC = Ar.IsLoading();
-		CreateSimulationDataImp(bCopyFromDDC, TEXT("8724C70A140146B5A2F4CF0C16083041"));
-	}
-#endif
 	//new versions serialize geometry collection after tagged properties
 	if (Ar.CustomVer(FDestructionObjectVersion::GUID) >= FDestructionObjectVersion::GeometryCollectionInDDCAndAsset)
 	{
@@ -362,49 +353,21 @@ void UGeometryCollection::Serialize(FArchive& Ar)
 #endif
 	}
 
-	if (Ar.CustomVer(FUE5MainStreamObjectVersion::GUID) >= FUE5MainStreamObjectVersion::GeometryCollectionNaniteData)
+	if (Ar.CustomVer(FUE5MainStreamObjectVersion::GUID) == FUE5MainStreamObjectVersion::GeometryCollectionNaniteData)
 	{
-		int32 NumNaniteResources = Ar.IsLoading() ? 0 : NaniteResources.Num();
-		if (!EnableNanite)
-		{
-			// Avoid serializing stale Nanite data when support is disabled.
-			// TODO: Should instead do this on toggle of the property?
-			// TODO: Might never even need this, as NaniteResources may always be empty due to CreateNaniteData() running before
-			NumNaniteResources = 0;
-		}
+		// This legacy version serialized structure information into archive, but the data is transient.
+		// Just load it and throw away here, it will be rebuilt later and resaved past this point.
 
+		int32 NumNaniteResources = 0;
 		Ar << NumNaniteResources;
 
-		if (NumNaniteResources == 0)
-		{
-			NaniteResources.Empty();
-		}
-		else
-		{
-			if (Ar.IsLoading())
-			{
-				NaniteResources.Reset(NumNaniteResources);
-				NaniteResources.AddDefaulted(NumNaniteResources);
-			}
-			
-			for (int32 ResourceIndex = 0; ResourceIndex < NumNaniteResources; ++ResourceIndex)
-			{
-				NaniteResources[ResourceIndex].Serialize(ChaosAr, this);
-			}
-
-			// Nanite data is currently 1:1 with each geometry group in the collection.
-			const int32 NumGeometryGroups = GeometryCollection->NumElements(FGeometryCollection::GeometryGroup);
-			if (NumGeometryGroups != NumNaniteResources)
-			{
-				Ar.SetError();
-			}
-		}
+		TArray<Nanite::FResources> NaniteResources;
+		NaniteResources.Reset(NumNaniteResources);
+		NaniteResources.AddDefaulted(NumNaniteResources);
 
-		if (!EnableNanite)
+		for (int32 ResourceIndex = 0; ResourceIndex < NumNaniteResources; ++ResourceIndex)
 		{
-			// Avoid serializing stale Nanite data when support is disabled.
-			// TODO: Should instead do this on toggle of the property?
-			NaniteResources.Empty();
+			NaniteResources[ResourceIndex].Serialize(ChaosAr, this);
 		}
 	}
 
@@ -423,7 +386,8 @@ const TCHAR* UGeometryCollection::GetSelectedMaterialPath() const
 }
 
 #if WITH_EDITOR
-void UGeometryCollection::CreateSimulationDataImp(bool bCopyFromDDC, const TCHAR* OverrideVersion)
+
+void UGeometryCollection::CreateSimulationDataImp(bool bCopyFromDDC)
 {
 	COOK_STAT(auto Timer = GeometryCollectionCookStats::UsageStats.TimeSyncWork());
 
@@ -433,7 +397,6 @@ void UGeometryCollection::CreateSimulationDataImp(bool bCopyFromDDC, const TCHAR
 	//Use the DDC to build simulation data. If we are loading in the editor we then serialize this data into the geometry collection
 	TArray<uint8> DDCData;
 	FDerivedDataGeometryCollectionCooker* GeometryCollectionCooker = new FDerivedDataGeometryCollectionCooker(*this);
-	GeometryCollectionCooker->SetOverrideVersion(OverrideVersion);
 
 	if (GeometryCollectionCooker->CanBuild())
 	{
@@ -454,10 +417,9 @@ void UGeometryCollection::CreateSimulationDataImp(bool bCopyFromDDC, const TCHAR
 			FMemoryReader Ar(DDCData, true);	// Must be persistent for BulkData to serialize
 			Chaos::FChaosArchive ChaosAr(Ar);
 			GeometryCollection->Serialize(ChaosAr);
-			for (Nanite::FResources& NaniteResource : NaniteResources)
-			{
-				NaniteResource.Serialize(ChaosAr, this);
-			}
+
+			NaniteData = MakeUnique<FGeometryCollectionNaniteData>();
+			NaniteData->Serialize(ChaosAr, this);
 		}
 	}
 }
@@ -467,18 +429,16 @@ void UGeometryCollection::CreateSimulationData()
 	CreateSimulationDataImp(/*bCopyFromDDC=*/false);
 }
 
-TArray<Nanite::FResources>& UGeometryCollection::CreateNaniteData(FGeometryCollection* Collection)
+TUniquePtr<FGeometryCollectionNaniteData> UGeometryCollection::CreateNaniteData(FGeometryCollection* Collection)
 {
-	TRACE_CPUPROFILER_EVENT_SCOPE_TEXT(TEXT("FDerivedDataGeometryCollectionCooker::Build::Nanite"));
+	TUniquePtr<FGeometryCollectionNaniteData> NaniteData;
 
-	NaniteResources.Empty();
-	if (!EnableNanite)
-	{
-		return NaniteResources;
-	}
+	TRACE_CPUPROFILER_EVENT_SCOPE_TEXT(TEXT("UGeometryCollection::CreateNaniteData"));
 
 	Nanite::IBuilderModule& NaniteBuilderModule = Nanite::IBuilderModule::Get();
 
+	NaniteData = MakeUnique<FGeometryCollectionNaniteData>();
+
 	// Transform Group
 	const TManagedArray<int32>& TransformToGeometryIndexArray = Collection->TransformToGeometryIndex;
 	const TManagedArray<int32>& SimulationTypeArray = Collection->SimulationType;
@@ -513,11 +473,11 @@ TArray<Nanite::FResources>& UGeometryCollection::CreateNaniteData(FGeometryColle
 	const TManagedArray<FGeometryCollectionSection>& Sections = Collection->Sections;
 
 	int32 NumGeometry = Collection->NumElements(FGeometryCollection::GeometryGroup);
-	NaniteResources.AddDefaulted(NumGeometry);
+	NaniteData->Resources.AddDefaulted(NumGeometry);
 
 	for (int32 GeometryGroupIndex = 0; GeometryGroupIndex < NumGeometry; GeometryGroupIndex++)
 	{
-		Nanite::FResources& NaniteResource = NaniteResources[GeometryGroupIndex];
+		Nanite::FResources& NaniteResource = NaniteData->Resources[GeometryGroupIndex];
 		NaniteResource = {};
 
 		uint32 NumTexCoords = 1;// NumTextureCoord;
@@ -585,45 +545,24 @@ TArray<Nanite::FResources>& UGeometryCollection::CreateNaniteData(FGeometryColle
 		}
 	}
 
-	return NaniteResources;
+	return NaniteData;
 }
 
 #endif
 
 void UGeometryCollection::InitResources()
 {
-	//LLM_SCOPE(ELLMTag::GeometryCollection);
-
-	ReleaseResources();
-
-	if (EnableNanite)
+	if (NaniteData)
 	{
-		for (Nanite::FResources& Resource : NaniteResources)
-		{
-			// Skip resources that have their render data stripped
-			if (Resource.PageStreamingStates.Num() > 0)
-			{
-				Resource.InitResources();
-			}
-		}
+		NaniteData->InitResources(this);
 	}
-
-	bRenderingResourcesInitialized = true;
 }
 
 void UGeometryCollection::ReleaseResources()
 {
-	if (bRenderingResourcesInitialized)
+	if (NaniteData)
 	{
-		for (Nanite::FResources& Resource : NaniteResources)
-		{
-			if (Resource.PageStreamingStates.Num() > 0)
-			{
-				Resource.ReleaseResources();
-			}
-		}
-
-		bRenderingResourcesInitialized = false;
+		NaniteData->ReleaseResources();
 	}
 }
 
@@ -643,23 +582,26 @@ FGuid UGeometryCollection::GetStateGuid() const
 }
 
 #if WITH_EDITOR
+
 void UGeometryCollection::PostEditChangeProperty(struct FPropertyChangedEvent& PropertyChangedEvent)
 {
 	if (PropertyChangedEvent.Property)
 	{
-		if (PropertyChangedEvent.Property->GetFName() != GET_MEMBER_NAME_CHECKED(UGeometryCollection, Materials))
+		if (PropertyChangedEvent.Property->GetFName() == GET_MEMBER_NAME_CHECKED(UGeometryCollection, EnableNanite))
 		{
 			InvalidateCollection();
-			CreateSimulationData();
-		}
+			EnsureDataIsCooked();
 
-		if (PropertyChangedEvent.Property->GetFName() == GET_MEMBER_NAME_CHECKED(UGeometryCollection, EnableNanite))
-		{
 			if (FApp::CanEverRender())
 			{
 				InitResources();
 			}
 		}
+		else if (PropertyChangedEvent.Property->GetFName() != GET_MEMBER_NAME_CHECKED(UGeometryCollection, Materials))
+		{
+			InvalidateCollection();
+			CreateSimulationData();
+		}
 	}
 }
 
@@ -668,7 +610,7 @@ bool UGeometryCollection::Modify(bool bAlwaysMarkDirty /*= true*/)
 	bool bSuperResult = Super::Modify(bAlwaysMarkDirty);
 
 	UPackage* Package = GetOutermost();
-	if(Package->IsDirty())
+	if (Package->IsDirty())
 	{
 		InvalidateCollection();
 	}
@@ -700,6 +642,85 @@ void UGeometryCollection::PostLoad()
 void UGeometryCollection::BeginDestroy()
 {
 	Super::BeginDestroy();
-
 	ReleaseResources();
 }
+
+FGeometryCollectionNaniteData::FGeometryCollectionNaniteData()
+{
+}
+
+FGeometryCollectionNaniteData::~FGeometryCollectionNaniteData()
+{
+}
+
+void FGeometryCollectionNaniteData::Serialize(FArchive& Ar, UGeometryCollection* Owner)
+{
+	if (Ar.IsSaving())
+	{
+		if (Owner->EnableNanite)
+		{
+			// Nanite data is currently 1:1 with each geometry group in the collection.
+			const int32 NumGeometryGroups = Owner->NumElements(FGeometryCollection::GeometryGroup);
+			if (NumGeometryGroups != Resources.Num())
+			{
+				Ar.SetError();
+			}
+		}
+
+		int32 NumNaniteResources = Resources.Num();
+		Ar << NumNaniteResources;
+
+		for (int32 ResourceIndex = 0; ResourceIndex < NumNaniteResources; ++ResourceIndex)
+		{
+			Resources[ResourceIndex].Serialize(Ar, Owner);
+		}
+	}
+	else if (Ar.IsLoading())
+	{
+		int32 NumNaniteResources = 0;
+		Ar << NumNaniteResources;
+
+		Resources.Reset(NumNaniteResources);
+		Resources.AddDefaulted(NumNaniteResources);
+
+		for (int32 ResourceIndex = 0; ResourceIndex < NumNaniteResources; ++ResourceIndex)
+		{
+			Resources[ResourceIndex].Serialize(Ar, Owner);
+		}
+
+		if (!Owner->EnableNanite)
+		{
+			Resources.Reset(0);
+		}
+	}
+}
+
+void FGeometryCollectionNaniteData::InitResources(UGeometryCollection* Owner)
+{
+	if (bIsInitialized)
+	{
+		ReleaseResources();
+	}
+
+	for (Nanite::FResources& Resource : Resources)
+	{
+		Resource.InitResources();
+	}
+
+	bIsInitialized = true;
+}
+
+void FGeometryCollectionNaniteData::ReleaseResources()
+{
+	if (!bIsInitialized)
+	{
+		return;
+	}
+
+	for (Nanite::FResources& Resource : Resources)
+	{
+		Resource.ReleaseResources();
+	}
+
+	bIsInitialized = false;
+}
diff --git a/Engine/Source/Runtime/Experimental/GeometryCollectionEngine/Public/GeometryCollection/DerivedDataGeometryCollectionCooker.h b/Engine/Source/Runtime/Experimental/GeometryCollectionEngine/Public/GeometryCollection/DerivedDataGeometryCollectionCooker.h
--- a/Engine/Source/Runtime/Experimental/GeometryCollectionEngine/Public/GeometryCollection/DerivedDataGeometryCollectionCooker.h
+++ b/Engine/Source/Runtime/Experimental/GeometryCollectionEngine/Public/GeometryCollection/DerivedDataGeometryCollectionCooker.h
@@ -15,10 +15,6 @@ class FDerivedDataGeometryCollectionCooker : public FDerivedDataPluginInterface
 {
 public:
 	FDerivedDataGeometryCollectionCooker(UGeometryCollection& InGeometryCollection);
-	void SetOverrideVersion(const TCHAR* InOverrideVersion)
-	{
-		OverrideVersion = InOverrideVersion;
-	}
 
 	virtual const TCHAR* GetPluginName() const override
 	{
@@ -53,7 +49,6 @@ public:
 
 private:
 	UGeometryCollection& GeometryCollection;
-	const TCHAR* OverrideVersion;	//force load old ddc content
 };
 
-#endif	//WITH_EDITOR
+#endif	// WITH_EDITOR
diff --git a/Engine/Source/Runtime/Experimental/GeometryCollectionEngine/Public/GeometryCollection/GeometryCollectionObject.h b/Engine/Source/Runtime/Experimental/GeometryCollectionEngine/Public/GeometryCollection/GeometryCollectionObject.h
--- a/Engine/Source/Runtime/Experimental/GeometryCollectionEngine/Public/GeometryCollection/GeometryCollectionObject.h
+++ b/Engine/Source/Runtime/Experimental/GeometryCollectionEngine/Public/GeometryCollection/GeometryCollectionObject.h
@@ -84,6 +84,32 @@ struct GEOMETRYCOLLECTIONENGINE_API FGeometryCollectionSizeSpecificData
 
 };
 
+class FGeometryCollectionNaniteData
+{
+public:
+	GEOMETRYCOLLECTIONENGINE_API FGeometryCollectionNaniteData();
+	GEOMETRYCOLLECTIONENGINE_API ~FGeometryCollectionNaniteData();
+
+	FORCEINLINE bool IsInitialized()
+	{
+		return bIsInitialized;
+	}
+
+	/** Serialization. */
+	void Serialize(FArchive& Ar, UGeometryCollection* Owner);
+
+	/** Initialize the render resources. */
+	void InitResources(UGeometryCollection* Owner);
+
+	/** Releases the render resources. */
+	GEOMETRYCOLLECTIONENGINE_API void ReleaseResources();
+
+	TArray<Nanite::FResources> Resources;
+
+private:
+	bool bIsInitialized = false;
+};
+
 /**
 * UGeometryCollectionObject (UObject)
 *
@@ -124,7 +150,8 @@ public:
 
 	FORCEINLINE FNaniteInfo GetNaniteInfo(int32 GeometryIndex) const
 	{
-		const Nanite::FResources& Resources = NaniteResources[GeometryIndex];
+		check(NaniteData && GeometryIndex >= 0 && GeometryIndex < NaniteData->Resources.Num());
+		const Nanite::FResources& Resources = NaniteData->Resources[GeometryIndex];
 		return FNaniteInfo(Resources.RuntimeResourceID, Resources.HierarchyOffset);
 	}
 
@@ -145,7 +172,7 @@ public:
 	void CreateSimulationData();
 
 	/** Create the Nanite rendering data. */
-	TArray<Nanite::FResources>& CreateNaniteData(FGeometryCollection* Collection);
+	static TUniquePtr<FGeometryCollectionNaniteData> CreateNaniteData(FGeometryCollection* Collection);
 #endif
 
 	void InitResources();
@@ -159,6 +186,9 @@ public:
 	FGuid GetIdGuid() const;
 	FGuid GetStateGuid() const;
 
+	/** Pointer to the data used to render this geometry collection with Nanite. */
+	TUniquePtr<class FGeometryCollectionNaniteData> NaniteData;
+
 	/** The editable mesh representation of this geometry collection */
 	class UObject* EditableMesh;
 	
@@ -269,13 +299,9 @@ public:
 	class UThumbnailInfo* ThumbnailInfo;
 #endif // WITH_EDITORONLY_DATA
 
-protected:
-	/** Tracks whether InitResources has been called, and rendering resources are initialized. */
-	uint8 bRenderingResourcesInitialized : 1;
-
 private:
 #if WITH_EDITOR
-	void CreateSimulationDataImp(bool bCopyFromDDC, const TCHAR* OverrideVersion = nullptr);
+	void CreateSimulationDataImp(bool bCopyFromDDC);
 #endif
 
 private:
@@ -301,8 +327,4 @@ private:
 	int32 BoneSelectedMaterialIndex;
 
 	TSharedPtr<FGeometryCollection, ESPMode::ThreadSafe> GeometryCollection;
-
-	TArray<Nanite::FResources> NaniteResources;
 };
-
-
