- Fixed PBViewActor compile errors (deleted unused code, and brought over TVector::Mid() from Dev-Physics)

diff --git a/Engine/Source/Runtime/Experimental/ChaosCore/Public/Chaos/Vector.h b/Engine/Source/Runtime/Experimental/ChaosCore/Public/Chaos/Vector.h
--- a/Engine/Source/Runtime/Experimental/ChaosCore/Public/Chaos/Vector.h
+++ b/Engine/Source/Runtime/Experimental/ChaosCore/Public/Chaos/Vector.h
@@ -474,10 +474,29 @@ namespace Chaos
 		{
 			return (X > Y && X > Z) ? X : (Y > Z ? Y : Z);
 		}
+		float Mid() const
+		{
+			return (*this)[MidIndex()];
+		}
 		float Min() const
 		{
 			return (X < Y && X < Z) ? X : (Y < Z ? Y : Z);
 		}
+		int32 MaxIndex() const
+		{
+			return (X > Y && X > Z) ? 0 : (Y > Z ? 1 : 2);
+		}
+		int32 MidIndex() const
+		{
+			int32 Idx[] = {0, 1, 2};
+			Idx[MinIndex()] = -1;
+			Idx[MaxIndex()] = -1;
+			return Idx[0] != -1 ? 0 : Idx[1] != -1 ? 1 : 2;
+		}
+		int32 MinIndex() const
+		{
+			return (X < Y && X < Z) ? 0 : (Y < Z ? 1 : 2);
+		}
 		TVector<float, 3> ComponentwiseMin(const TVector<float, 3>& Other) const { return {FMath::Min(X,Other.X), FMath::Min(Y,Other.Y), FMath::Min(Z,Other.Z)}; }
 		TVector<float, 3> ComponentwiseMax(const TVector<float, 3>& Other) const { return {FMath::Max(X,Other.X), FMath::Max(Y,Other.Y), FMath::Max(Z,Other.Z)}; }
 		static TVector<float, 3> Max(const TVector<float, 3>& V1, const TVector<float, 3>& V2)
@@ -699,8 +718,13 @@ namespace Chaos
 		FORCEINLINE T SizeSquared() const { return X * X + Y * Y + Z * Z; }
 
 		FORCEINLINE T Min() const { return FMath::Min3(X, Y, Z); }
+		FORCEINLINE T Mid() const { return (*this)[MidIndex()]; }
 		FORCEINLINE T Max() const { return FMath::Max3(X, Y, Z); }
 
+		int32 MaxIndex() const { return (X > Y&& X > Z) ? 0 : (Y > Z ? 1 : 2); }
+		int32 MidIndex() const { TVector<int32, 3> Idx(0,1,2); Idx[MinIndex()]=-1; Idx[MaxIndex()]=-1; return Idx.Max(); }
+		int32 MinIndex() const { return (X < Y && X < Z) ? 0 : (Y < Z ? 1 : 2); }
+
 		FORCEINLINE TVector<T, 3> ComponentwiseMin(const TVector<T, 3>& Other) const { return {FMath::Min(X,Other.X), FMath::Min(Y,Other.Y), FMath::Min(Z,Other.Z)}; }
 		FORCEINLINE TVector<T, 3> ComponentwiseMax(const TVector<T, 3>& Other) const { return {FMath::Max(X,Other.X), FMath::Max(Y,Other.Y), FMath::Max(Z,Other.Z)}; }
 
