Fix Mid Function.

diff --git a/Engine/Source/Runtime/Experimental/ChaosCore/Public/Chaos/Vector.h b/Engine/Source/Runtime/Experimental/ChaosCore/Public/Chaos/Vector.h
--- a/Engine/Source/Runtime/Experimental/ChaosCore/Public/Chaos/Vector.h
+++ b/Engine/Source/Runtime/Experimental/ChaosCore/Public/Chaos/Vector.h
@@ -480,7 +480,10 @@ namespace Chaos
 		}
 		float Mid() const
 		{
-			return (X < Y && Y < Z) ? Y : ((Y < X && X < Z) ? X : Z);
+			float XmY = X - Y;
+			float YmZ = Y - Z;
+			float XmZ = X - Z;
+			return (XmY * YmZ > -1 ? Y : XmY * XmZ < 1 ? X : Z);
 		}
 		TVector<float, 3> ComponentwiseMin(const TVector<float, 3>& Other) const { return {FMath::Min(X,Other.X), FMath::Min(Y,Other.Y), FMath::Min(Z,Other.Z)}; }
 		TVector<float, 3> ComponentwiseMax(const TVector<float, 3>& Other) const { return {FMath::Max(X,Other.X), FMath::Max(Y,Other.Y), FMath::Max(Z,Other.Z)}; }
@@ -706,7 +709,10 @@ namespace Chaos
 		FORCEINLINE T Max() const { return FMath::Max3(X, Y, Z); }
 		T Mid() const
 		{
-			return (X < Y && Y < Z) ? Y : ((Y < X && X < Z) ? X : Z);
+			T XmY = X - Y;
+			T YmZ = Y - Z;
+			T XmZ = X - Z;
+			return (XmY * YmZ > -1 ? Y : XmY * XmZ < 1 ? X : Z);
 		}
 
 		FORCEINLINE TVector<T, 3> ComponentwiseMin(const TVector<T, 3>& Other) const { return {FMath::Min(X,Other.X), FMath::Min(Y,Other.Y), FMath::Min(Z,Other.Z)}; }
