Merge Brice's UE4 constraints fix (CL 14249363) to UE5 main branch.

diff --git a/Engine/Source/Runtime/Experimental/Chaos/Private/Chaos/PBDJointConstraints.cpp b/Engine/Source/Runtime/Experimental/Chaos/Private/Chaos/PBDJointConstraints.cpp
--- a/Engine/Source/Runtime/Experimental/Chaos/Private/Chaos/PBDJointConstraints.cpp
+++ b/Engine/Source/Runtime/Experimental/Chaos/Private/Chaos/PBDJointConstraints.cpp
@@ -341,15 +341,16 @@ namespace Chaos
 		// However, it is possible to set it up so that the kinematic particle is the child which we don't support, so...
 		// If particle 0 is kinematic we make it the parent, otherwise particle 1 is the parent.
 		// @todo(ccaulfield): look into this and confirm/fix properly
-		if (!ConstraintParticles[ConstraintIndex][0]->CastToRigidParticle())
+		TPBDRigidParticleHandle<FReal, 3>* Particle1 = ConstraintParticles[ConstraintIndex][1]->CastToRigidParticle();
+		if (!Particle1 || Particle1->ObjectState() == EObjectStateType::Kinematic || Particle1->ObjectState() == EObjectStateType::Static)
 		{
-			Index0 = 0;
-			Index1 = 1;
+			Index0 = 1;
+			Index1 = 0;
 		}
 		else
 		{
-			Index0 = 1;
-			Index1 = 0;
+			Index0 = 0;
+			Index1 = 1;
 		}
 	}
 
