Chaos Cloth - Add backstop legacy parameter. Improve backstop debug draw. [Merging CL 14267658]

diff --git a/Engine/Plugins/Experimental/ChaosCloth/Source/Chaos/Private/ChaosCloth/ChaosClothConfig.cpp b/Engine/Plugins/Experimental/ChaosCloth/Source/Chaos/Private/ChaosCloth/ChaosClothConfig.cpp
--- a/Engine/Plugins/Experimental/ChaosCloth/Source/Chaos/Private/ChaosCloth/ChaosClothConfig.cpp
+++ b/Engine/Plugins/Experimental/ChaosCloth/Source/Chaos/Private/ChaosCloth/ChaosClothConfig.cpp
@@ -77,6 +77,8 @@ void UChaosClothConfig::MigrateFrom(const FClothConfig_Legacy& ClothConfig)
 	bUseGravityOverride = ClothConfig.bUseGravityOverride;
 	GravityScale = ClothConfig.GravityScale;
 	Gravity = ClothConfig.GravityOverride;
+
+	bUseLegacyBackstop = true;
 }
 
 void UChaosClothConfig::MigrateFrom(const UClothSharedConfigCommon* ClothSharedConfig)
@@ -124,6 +126,11 @@ void UChaosClothConfig::PostLoad()
 		MinPerParticleMass = 0.0001f;  // Override these values in case they might have been accidentally
 		bUseGeodesicDistance = true;   // changed from their default at any point in time
 	}
+
+	if (ChaosClothConfigCustomVersion < FChaosClothConfigCustomVersion::AddLegacyBackstopParameter)
+	{
+		bUseLegacyBackstop = true;
+	}
 }
 
 float UChaosClothConfig::GetMassValue() const
diff --git a/Engine/Plugins/Experimental/ChaosCloth/Source/Chaos/Private/ChaosCloth/ChaosClothConfigCustomVersion.h b/Engine/Plugins/Experimental/ChaosCloth/Source/Chaos/Private/ChaosCloth/ChaosClothConfigCustomVersion.h
--- a/Engine/Plugins/Experimental/ChaosCloth/Source/Chaos/Private/ChaosCloth/ChaosClothConfigCustomVersion.h
+++ b/Engine/Plugins/Experimental/ChaosCloth/Source/Chaos/Private/ChaosCloth/ChaosClothConfigCustomVersion.h
@@ -20,6 +20,8 @@ struct FChaosClothConfigCustomVersion
 		AddGravitySelfCollisionMigration = 3,
 		// Remove internal config parameters from UI
 		RemoveInternalConfigParameters = 4,
+		// Add a parameter to fix the backstop behavior without breaking PhysX assets backward compatibility
+		AddLegacyBackstopParameter = 5,
 
 		// -----<new versions can be added above this line>-------------------------------------------------
 		VersionPlusOne,
diff --git a/Engine/Plugins/Experimental/ChaosCloth/Source/Chaos/Private/ChaosCloth/ChaosClothConstraints.cpp b/Engine/Plugins/Experimental/ChaosCloth/Source/Chaos/Private/ChaosCloth/ChaosClothConstraints.cpp
--- a/Engine/Plugins/Experimental/ChaosCloth/Source/Chaos/Private/ChaosCloth/ChaosClothConstraints.cpp
+++ b/Engine/Plugins/Experimental/ChaosCloth/Source/Chaos/Private/ChaosCloth/ChaosClothConstraints.cpp
@@ -383,7 +383,7 @@ void FClothConstraints::SetMaximumDistanceConstraints(const TConstArrayView<floa
 	++NumConstraintRules;
 }
 
-void FClothConstraints::SetBackstopConstraints(const TConstArrayView<float>& BackstopDistances, const TConstArrayView<float>& BackstopRadiuses)
+void FClothConstraints::SetBackstopConstraints(const TConstArrayView<float>& BackstopDistances, const TConstArrayView<float>& BackstopRadiuses, bool bUseLegacyBackstop)
 {
 	BackstopConstraints = MakeShared<TPBDSphericalBackstopConstraint<float, 3>>(
 		ParticleOffset,
@@ -391,7 +391,8 @@ void FClothConstraints::SetBackstopConstraints(const TConstArrayView<float>& Bac
 		*AnimationPositions,
 		*AnimationNormals,
 		BackstopRadiuses,
-		BackstopDistances);
+		BackstopDistances,
+		bUseLegacyBackstop);
 	++NumConstraintRules;
 }
 
diff --git a/Engine/Plugins/Experimental/ChaosCloth/Source/Chaos/Private/ChaosCloth/ChaosClothingSimulation.cpp b/Engine/Plugins/Experimental/ChaosCloth/Source/Chaos/Private/ChaosCloth/ChaosClothingSimulation.cpp
--- a/Engine/Plugins/Experimental/ChaosCloth/Source/Chaos/Private/ChaosCloth/ChaosClothingSimulation.cpp
+++ b/Engine/Plugins/Experimental/ChaosCloth/Source/Chaos/Private/ChaosCloth/ChaosClothingSimulation.cpp
@@ -202,6 +202,7 @@ void FClothingSimulation::CreateActor(USkeletalMeshComponent* InOwnerComponent,
 		ClothConfig->FrictionCoefficient,
 		ClothConfig->bUseSelfCollisions,
 		ClothConfig->SelfCollisionThickness,
+		ClothConfig->bUseLegacyBackstop,
 		/*bUseLODIndexOverride =*/ false,
 		/*LODIndexOverride =*/ INDEX_NONE));
 
@@ -476,6 +477,7 @@ void FClothingSimulation::RefreshClothConfig(const IClothingSimulationContext* I
 			ClothConfig->FrictionCoefficient,
 			ClothConfig->bUseSelfCollisions,
 			ClothConfig->SelfCollisionThickness,
+			ClothConfig->bUseLegacyBackstop,
 			/*bUseLODIndexOverride =*/ false,
 			/*LODIndexOverride =*/ INDEX_NONE);
 
@@ -1106,10 +1108,11 @@ void FClothingSimulation::DebugDrawBackstops(FPrimitiveDrawInterface* PDI) const
 {
 	auto DrawBackstop = [PDI](const FVector& Position, const FVector& Normal, float Radius, const FVector& Axis, const FLinearColor& Color)
 	{
-		static const float ArcAngle = 25.0f; // Arc angle in degrees
 		static const float MaxCosAngle = 0.99f;
 		if (FMath::Abs(FVector::DotProduct(Normal, Axis)) < MaxCosAngle)
 		{
+			static const float ArcLength = 5.0f; // Arch length in cm
+			const float ArcAngle = 360.0f * ArcLength / FMath::Max((Radius * 2.0f * PI), ArcLength);
 			DrawArc(PDI, Position, Normal, FVector::CrossProduct(Axis, Normal).GetSafeNormal(), -ArcAngle / 2.0f, ArcAngle / 2.0f, Radius, Color);
 		}
 	};
@@ -1129,6 +1132,7 @@ void FClothingSimulation::DebugDrawBackstops(FPrimitiveDrawInterface* PDI) const
 		const FClothConstraints& ClothConstraints = Solver->GetClothConstraints(Offset);
 		if (const TPBDSphericalBackstopConstraint<float, 3>* const BackstopConstraint = ClothConstraints.GetBackstopConstraints().Get())
 		{
+			const bool bUseLegacyBackstop = BackstopConstraint->UseLegacyBackstop();
 			const TConstArrayView<float>& BackstopDistances = Cloth->GetWeightMaps(Solver.Get())[(int32)EChaosWeightMapTarget::BackstopDistance];
 			const TConstArrayView<float>& BackstopRadiuses = Cloth->GetWeightMaps(Solver.Get())[(int32)EChaosWeightMapTarget::BackstopRadius];
 			const TConstArrayView<TVector<float, 3>> AnimationPositions = Cloth->GetAnimationPositions(Solver.Get());
@@ -1147,16 +1151,16 @@ void FClothingSimulation::DebugDrawBackstops(FPrimitiveDrawInterface* PDI) const
 				const FVector AnimationPosition = LocalSpaceLocation + AnimationPositions[Index];
 				const FVector& AnimationNormal = AnimationNormals[Index];
 
-				// Draw a line to the current position
+				// Draw a line to show the current distance to the sphere
 				const FVector Pos0 = LocalSpaceLocation + AnimationPositions[Index];
-				const FVector Pos1 = Pos0 - BackstopDistance * AnimationNormal;
+				const FVector Pos1 = Pos0 - (bUseLegacyBackstop ? BackstopDistance - BackstopRadius : BackstopDistance) * AnimationNormal;
 				const FVector Pos2 = LocalSpaceLocation + ParticlePositions[Index];
 				DrawLine(PDI, Pos1, Pos2, ColorLight);
 
 				// Draw the sphere
 				if (BackstopRadius > 0.f)
 				{
-					const FVector Center = Pos0 - (BackstopRadius + BackstopDistance) * AnimationNormal;
+					const FVector Center = Pos0 - (bUseLegacyBackstop ? BackstopDistance : BackstopRadius + BackstopDistance) * AnimationNormal;
 					DrawBackstop(Center, AnimationNormal, BackstopRadius, FVector::ForwardVector, ColorDark);
 					DrawBackstop(Center, AnimationNormal, BackstopRadius, FVector::UpVector, ColorDark);
 					DrawBackstop(Center, AnimationNormal, BackstopRadius, FVector::RightVector, ColorDark);
@@ -1183,6 +1187,7 @@ void FClothingSimulation::DebugDrawBackstopDistances(FPrimitiveDrawInterface* PD
 		const FClothConstraints& ClothConstraints = Solver->GetClothConstraints(Offset);
 		if (const TPBDSphericalBackstopConstraint<float, 3>* const BackstopConstraint = ClothConstraints.GetBackstopConstraints().Get())
 		{
+			const bool bUseLegacyBackstop = BackstopConstraint->UseLegacyBackstop();
 			const TConstArrayView<float>& BackstopDistances = Cloth->GetWeightMaps(Solver.Get())[(int32)EChaosWeightMapTarget::BackstopDistance];
 			const TConstArrayView<float>& BackstopRadiuses = Cloth->GetWeightMaps(Solver.Get())[(int32)EChaosWeightMapTarget::BackstopRadius];
 			const TConstArrayView<TVector<float, 3>> AnimationPositions = Cloth->GetAnimationPositions(Solver.Get());
@@ -1202,7 +1207,7 @@ void FClothingSimulation::DebugDrawBackstopDistances(FPrimitiveDrawInterface* PD
 
 				// Draw a line to the sphere boundary
 				const FVector Pos0 = LocalSpaceLocation + AnimationPositions[Index];
-				const FVector Pos1 = Pos0 - BackstopDistance * AnimationNormal;
+				const FVector Pos1 = Pos0 - (bUseLegacyBackstop ? BackstopDistance - BackstopRadius : BackstopDistance) * AnimationNormal;
 				DrawLine(PDI, Pos0, Pos1, ColorDark);
 			}
 		}
diff --git a/Engine/Plugins/Experimental/ChaosCloth/Source/Chaos/Private/ChaosCloth/ChaosClothingSimulationCloth.cpp b/Engine/Plugins/Experimental/ChaosCloth/Source/Chaos/Private/ChaosCloth/ChaosClothingSimulationCloth.cpp
--- a/Engine/Plugins/Experimental/ChaosCloth/Source/Chaos/Private/ChaosCloth/ChaosClothingSimulationCloth.cpp
+++ b/Engine/Plugins/Experimental/ChaosCloth/Source/Chaos/Private/ChaosCloth/ChaosClothingSimulationCloth.cpp
@@ -226,7 +226,7 @@ void FClothingSimulationCloth::FLODData::Add(FClothingSimulationSolver* Solver,
 	const TConstArrayView<float>& BackstopRadiuses = WeightMaps[(int32)EChaosWeightMapTarget::BackstopRadius];
 	if (BackstopRadiuses.Num() && BackstopDistances.Num())
 	{
-		ClothConstraints.SetBackstopConstraints(BackstopDistances, BackstopRadiuses);
+		ClothConstraints.SetBackstopConstraints(BackstopDistances, BackstopRadiuses, Cloth->bUseLegacyBackstop);
 	}
 
 	// Animation Drive Constraints
@@ -348,6 +348,7 @@ FClothingSimulationCloth::FClothingSimulationCloth(
 	float InFrictionCoefficient,
 	bool bInUseSelfCollisions,
 	float InSelfCollisionThickness,
+	bool bInUseLegacyBackstop,
 	bool bInUseLODIndexOverride, 
 	int32 InLODIndexOverride)
 	: Mesh(nullptr)
@@ -381,6 +382,7 @@ FClothingSimulationCloth::FClothingSimulationCloth(
 	, FrictionCoefficient(InFrictionCoefficient)
 	, bUseSelfCollisions(bInUseSelfCollisions)
 	, SelfCollisionThickness(InSelfCollisionThickness)
+	, bUseLegacyBackstop(bInUseLegacyBackstop)
 	, bUseLODIndexOverride(bInUseLODIndexOverride)
 	, LODIndexOverride(InLODIndexOverride)
 	, bNeedsReset(false)
diff --git a/Engine/Plugins/Experimental/ChaosCloth/Source/Chaos/Public/ChaosCloth/ChaosClothConfig.h b/Engine/Plugins/Experimental/ChaosCloth/Source/Chaos/Public/ChaosCloth/ChaosClothConfig.h
--- a/Engine/Plugins/Experimental/ChaosCloth/Source/Chaos/Public/ChaosCloth/ChaosClothConfig.h
+++ b/Engine/Plugins/Experimental/ChaosCloth/Source/Chaos/Public/ChaosCloth/ChaosClothConfig.h
@@ -117,6 +117,13 @@ public:
 	UPROPERTY(EditAnywhere, Category = "Collision Properties", meta = (InlineEditConditionToggle))
 	bool bUseSelfCollisions = false;
 
+	// This parameter is automatically set by the migration code. It can be overriden here to use the old way of authoring the backstop distances.
+	// The legacy backstop requires the sphere radius to be included within the painted distance mask, making it difficult to author correctly. In this case the backstop distance is the distance from the animated mesh to the center of the corresponding backstop collision sphere.
+	// The non legacy backstop automatically adds the matching sphere's radius to the distance calculations at runtime to make for a simpler authoring of the backstop distances. In this case the backstop distance is the distance from the animated mesh to the surface of the backstop collision sphere.
+	// In both cases, a positive backstop distance goes against the corresponding animated mesh's normal, and a negative backstop distance goes along the corresponding animated mesh's normal.
+	UPROPERTY(EditAnywhere, Category = "Collision Properties")
+	bool bUseLegacyBackstop = false;
+
 	// The radius of the spheres used in self collision 
 	UPROPERTY(EditAnywhere, Category = "Collision Properties", meta = (UIMin = "0", UIMax = "100", ClampMin = "0", ClampMax = "1000", EditCondition = "bUseSelfCollisions"))
 	float SelfCollisionThickness = 2.0f;
diff --git a/Engine/Plugins/Experimental/ChaosCloth/Source/Chaos/Public/ChaosCloth/ChaosClothConstraints.h b/Engine/Plugins/Experimental/ChaosCloth/Source/Chaos/Public/ChaosCloth/ChaosClothConstraints.h
--- a/Engine/Plugins/Experimental/ChaosCloth/Source/Chaos/Public/ChaosCloth/ChaosClothConstraints.h
+++ b/Engine/Plugins/Experimental/ChaosCloth/Source/Chaos/Public/ChaosCloth/ChaosClothConstraints.h
@@ -47,7 +47,7 @@ namespace Chaos
 		void SetVolumeConstraints(TArray<TVector<int32, 3>>&& SurfaceElements, float VolumeStiffness);
 		void SetLongRangeConstraints(const TMap<int32, TSet<uint32>>& PointToNeighborsMap, float StrainLimitingStiffness, float LimitScale, bool bUseGeodesicDistance, bool bUseXPBDConstraints);
 		void SetMaximumDistanceConstraints(const TConstArrayView<float>& MaxDistances);
-		void SetBackstopConstraints(const TConstArrayView<float>& BackstopDistances, const TConstArrayView<float>& BackstopRadiuses);
+		void SetBackstopConstraints(const TConstArrayView<float>& BackstopDistances, const TConstArrayView<float>& BackstopRadiuses, bool bUseLegacyBackstop);
 		void SetAnimDriveConstraints(const TConstArrayView<float>& AnimDriveMultipliers);
 		void SetShapeTargetConstraints(float ShapeTargetStiffness);
 
diff --git a/Engine/Plugins/Experimental/ChaosCloth/Source/Chaos/Public/ChaosCloth/ChaosClothingSimulationCloth.h b/Engine/Plugins/Experimental/ChaosCloth/Source/Chaos/Public/ChaosCloth/ChaosClothingSimulationCloth.h
--- a/Engine/Plugins/Experimental/ChaosCloth/Source/Chaos/Public/ChaosCloth/ChaosClothingSimulationCloth.h
+++ b/Engine/Plugins/Experimental/ChaosCloth/Source/Chaos/Public/ChaosCloth/ChaosClothingSimulationCloth.h
@@ -55,6 +55,7 @@ namespace Chaos
 			float InFrictionCoefficient,
 			bool bInUseSelfCollisions,
 			float InSelfCollisionThickness,
+			bool bInUseLegacyBackstop,
 			bool bInUseLODIndexOverride,
 			int32 InLODIndexOverride);
 		~FClothingSimulationCloth();
@@ -185,6 +186,7 @@ namespace Chaos
 		float FrictionCoefficient;
 		bool bUseSelfCollisions;
 		float SelfCollisionThickness;
+		bool bUseLegacyBackstop;
 		bool bUseLODIndexOverride;
 		int32 LODIndexOverride;
 		bool bNeedsReset;
diff --git a/Engine/Source/Runtime/Experimental/Chaos/Public/Chaos/PBDSphericalConstraint.h b/Engine/Source/Runtime/Experimental/Chaos/Public/Chaos/PBDSphericalConstraint.h
--- a/Engine/Source/Runtime/Experimental/Chaos/Public/Chaos/PBDSphericalConstraint.h
+++ b/Engine/Source/Runtime/Experimental/Chaos/Public/Chaos/PBDSphericalConstraint.h
@@ -83,7 +83,8 @@ namespace Chaos
 			const TArray<TVector<T, d>>& InAnimationPositions,  // Use global indexation (will need adding ParticleOffset)
 			const TArray<TVector<T, d>>& InAnimationNormals,  // Use global indexation (will need adding ParticleOffset)
 			const TConstArrayView<T>& InSphereRadii,  // Use local indexation
-			const TConstArrayView<T>& InSphereOffsetDistances  // Use local indexation
+			const TConstArrayView<T>& InSphereOffsetDistances,  // Use local indexation
+			const bool bInUseLegacyBackstop  // Do not include the sphere radius in the distance calculations when this is true
 		)
 			: AnimationPositions(InAnimationPositions)
 			, AnimationNormals(InAnimationNormals)
@@ -91,6 +92,7 @@ namespace Chaos
 			, SphereOffsetDistances(InSphereOffsetDistances)
 			, ParticleOffset(InParticleOffset)
 			, SphereRadiiMultiplier((T)1.)
+			, bUseLegacyBackstop(bInUseLegacyBackstop)
 		{
 			check(InSphereRadii.Num() == InParticleCount);
 			check(InSphereOffsetDistances.Num() == InParticleCount);
@@ -101,6 +103,38 @@ namespace Chaos
 		{
 			SCOPE_CYCLE_COUNTER(STAT_PBD_SphericalBackstop);
 
+			if (bUseLegacyBackstop)
+			{
+				// SphereOffsetDistances includes the sphere radius
+				// This is harder to author, and does not follow the NvCloth specs.
+				// However, this is how it's been done in the Unreal Engine PhysX cloth implementation.
+				ApplyLegacyHelper(Particles, Dt);
+			}
+			else
+			{
+				// SphereOffsetDistances doesn't include the sphere radius
+				ApplyHelper(Particles, Dt);
+			}
+		}
+
+		inline void SetSphereRadiiMultiplier(const T InSphereRadiiMultiplier)
+		{
+			SphereRadiiMultiplier = FMath::Max((T)0., InSphereRadiiMultiplier);
+		}
+
+		inline T GetSphereRadiiMultiplier() const
+		{
+			return SphereRadiiMultiplier;
+		}
+
+		inline bool UseLegacyBackstop() const
+		{
+			return bUseLegacyBackstop;
+		}
+
+	private:
+		inline void ApplyHelper(TPBDParticles<T, d>& Particles, const T Dt) const
+		{
 			const int32 ParticleCount = SphereRadii.Num();
 
 			PhysicsParallelFor(ParticleCount, [&](int32 Index)  // TODO: profile need for parallel loop based on particle count
@@ -118,15 +152,14 @@ namespace Chaos
 				const T SphereOffsetDistance = SphereOffsetDistances[Index];
 				const T Radius = SphereRadii[Index] * SphereRadiiMultiplier;
 
-				const TVector<T, d> Center = AnimationPosition - (Radius + SphereOffsetDistance) * AnimationNormal;
-
+				const TVector<T, d> Center = AnimationPosition - (Radius + SphereOffsetDistance) * AnimationNormal;  // Non legacy version adds radius to the distance
 				const TVector<T, d> CenterToParticle = Particles.P(ParticleIndex) - Center;
 				const T DistanceSquared = CenterToParticle.SizeSquared();
 
 				static const T DeadZoneSquareRadius = SMALL_NUMBER;
 				if (DistanceSquared < DeadZoneSquareRadius)
 				{
-					Particles.P(ParticleIndex) = AnimationPosition - SphereOffsetDistance * AnimationNormal;
+					Particles.P(ParticleIndex) = AnimationPosition - SphereOffsetDistance * AnimationNormal;  // Non legacy version adds radius to the distance
 				}
 				else if (DistanceSquared < FMath::Square(Radius))
 				{
@@ -137,14 +170,41 @@ namespace Chaos
 			});
 		}
 
-		inline void SetSphereRadiiMultiplier(const T InSphereRadiiMultiplier)
+		inline void ApplyLegacyHelper(TPBDParticles<T, d>& Particles, const T Dt) const
 		{
-			SphereRadiiMultiplier = FMath::Max((T)0., InSphereRadiiMultiplier);
-		}
+			const int32 ParticleCount = SphereRadii.Num();
 
-		inline T GetSphereRadiiMultiplier() const
-		{
-			return SphereRadiiMultiplier;
+			PhysicsParallelFor(ParticleCount, [&](int32 Index)  // TODO: profile need for parallel loop based on particle count
+			{
+				const int32 ParticleIndex = ParticleOffset + Index;
+
+				if (Particles.InvM(ParticleIndex) == 0)
+				{
+					return;
+				}
+
+				const TVector<T, d>& AnimationPosition = AnimationPositions[Index];
+				const TVector<T, d>& AnimationNormal = AnimationNormals[Index];
+
+				const T SphereOffsetDistance = SphereOffsetDistances[Index];
+				const T Radius = SphereRadii[Index] * SphereRadiiMultiplier;
+
+				const TVector<T, d> Center = AnimationPosition - SphereOffsetDistance * AnimationNormal;  // Legacy version already includes the radius within the distance
+				const TVector<T, d> CenterToParticle = Particles.P(ParticleIndex) - Center;
+				const T DistanceSquared = CenterToParticle.SizeSquared();
+
+				static const T DeadZoneSquareRadius = SMALL_NUMBER;
+				if (DistanceSquared < DeadZoneSquareRadius)
+				{
+					Particles.P(ParticleIndex) = AnimationPosition - (SphereOffsetDistance - Radius) * AnimationNormal;  // Legacy version already includes the radius to the distance
+				}
+				else if (DistanceSquared < FMath::Square(Radius))
+				{
+					const TVector<T, d> PositionOnSphere = (Radius / sqrt(DistanceSquared)) * CenterToParticle;
+					Particles.P(ParticleIndex) = Center + PositionOnSphere;
+				}
+				// Else the particle is outside the sphere, and there is nothing to do
+			});
 		}
 
 	private:
@@ -154,5 +214,6 @@ namespace Chaos
 		const TConstArrayView<T> SphereOffsetDistances;  // Sphere position offsets, use local indexation
 		const int32 ParticleOffset;
 		T SphereRadiiMultiplier;
+		bool bUseLegacyBackstop;
 	};
 }
