Chaos Physics: Add CVAR to enable/disable contact solver relaxation

#[fyi] Michael.Lentine, Benn.Gallagher, Chris.Caulfield, Brice.Criswell

diff --git a/Engine/Source/Runtime/Experimental/Chaos/Private/Chaos/PBDCollisionConstraintsContact.cpp b/Engine/Source/Runtime/Experimental/Chaos/Private/Chaos/PBDCollisionConstraintsContact.cpp
--- a/Engine/Source/Runtime/Experimental/Chaos/Private/Chaos/PBDCollisionConstraintsContact.cpp
+++ b/Engine/Source/Runtime/Experimental/Chaos/Private/Chaos/PBDCollisionConstraintsContact.cpp
@@ -26,6 +26,9 @@ namespace Chaos
 		int32 Chaos_Collision_EnergyClampEnabled = 1;
 		FAutoConsoleVariableRef CVarChaosCollisionEnergyClampEnabled(TEXT("p.Chaos.Collision.EnergyClampEnabled"), Chaos_Collision_EnergyClampEnabled, TEXT("Whether to use energy clamping in collision apply step"));
 
+		int32 Chaos_Collision_RelaxationEnabled = 1;
+		FAutoConsoleVariableRef CVarChaosCollisionRelaxationEnabled(TEXT("p.Chaos.Collision.RelaxationEnabled"), Chaos_Collision_RelaxationEnabled, TEXT("Whether to reduce applied impulses during iterations for improved solver stability but reduced convergence"));
+
 		int32 Chaos_Collision_ForceApplyType = 0;
 		FAutoConsoleVariableRef CVarChaosCollisionAlternativeApply(TEXT("p.Chaos.Collision.ForceApplyType"), Chaos_Collision_ForceApplyType, TEXT("Force Apply step to use Velocity(1) or Position(2) modes"));
 
@@ -190,7 +193,7 @@ namespace Chaos
 			const FMatrix33 FactorInverse = Factor.Inverse();
 			FVec3 Impulse = FactorInverse * VelocityChange;  // Delta Impulse = (J.M^(-1).J^(T))^(-1).(ContactVelocityError)
 
-			if (bInApplyRelaxation)
+			if (bInApplyRelaxation && Chaos_Collision_RelaxationEnabled)
 			{
 				Impulse *= CalculateRelaxationFactor(IterationParameters);
 			}
@@ -477,7 +480,11 @@ namespace Chaos
 					Impulse = GetEnergyClampedImpulse(Particle0->CastToRigidParticle(), Particle1->CastToRigidParticle(), Impulse, VectorToPoint1, VectorToPoint2, Body1Velocity, Body2Velocity);
 				}
 
-				Impulse *= CalculateRelaxationFactor(IterationParameters);
+				if (Chaos_Collision_RelaxationEnabled)
+				{
+					Impulse *= CalculateRelaxationFactor(IterationParameters);
+				}
+				
 				AccumulatedImpulse += Impulse;
 
 				if (bIsRigidDynamic0)
