Chaos Cloth - Remove collision tests from collision particles that are not in the active LOD.

diff --git a/Engine/Source/Runtime/Experimental/Chaos/Private/Chaos/PBDEvolution.cpp b/Engine/Source/Runtime/Experimental/Chaos/Private/Chaos/PBDEvolution.cpp
--- a/Engine/Source/Runtime/Experimental/Chaos/Private/Chaos/PBDEvolution.cpp
+++ b/Engine/Source/Runtime/Experimental/Chaos/Private/Chaos/PBDEvolution.cpp
@@ -198,7 +198,7 @@ void TPBDEvolution<T, d>::AdvanceOneTimeStep(const T Dt)
 	TPerParticleEulerStepVelocity<T, d> EulerStepVelocityRule;
 	TPerGroupDampVelocity<T, d> DampVelocityRule(MParticleGroupIds, MGroupDampings);
 	TPerParticlePBDEulerStep<T, d> EulerStepRule;
-	TPerParticlePBDCollisionConstraint<T, d, EGeometryParticlesSimType::Other> CollisionRule(MCollisionParticles, MCollided, MParticleGroupIds, MCollisionParticleGroupIds, MGroupCollisionThicknesses, MGroupCoefficientOfFrictions);
+	TPerParticlePBDCollisionConstraint<T, d, EGeometryParticlesSimType::Other> CollisionRule(MCollisionParticlesActiveView, MCollided, MParticleGroupIds, MCollisionParticleGroupIds, MGroupCollisionThicknesses, MGroupCoefficientOfFrictions);
 
 	{
 		SCOPE_CYCLE_COUNTER(STAT_ChaosPBDVelocityDampUpdateState);
diff --git a/Engine/Source/Runtime/Experimental/Chaos/Public/Chaos/PerParticlePBDCollisionConstraint.h b/Engine/Source/Runtime/Experimental/Chaos/Public/Chaos/PerParticlePBDCollisionConstraint.h
--- a/Engine/Source/Runtime/Experimental/Chaos/Public/Chaos/PerParticlePBDCollisionConstraint.h
+++ b/Engine/Source/Runtime/Experimental/Chaos/Public/Chaos/PerParticlePBDCollisionConstraint.h
@@ -3,6 +3,7 @@
 
 #include "Chaos/PerParticleRule.h"
 #include "Chaos/Transform.h"
+#include "Chaos/PBDActiveView.h"
 
 #include "HAL/PlatformMath.h"
 
@@ -11,15 +12,18 @@ namespace Chaos
 template<class T, int d, EGeometryParticlesSimType SimType>
 class TPerParticlePBDCollisionConstraint : public TPerParticleRule<T, d>
 {
-	struct VelocityConstraint
+	struct FVelocityConstraint
 	{
 		TVector<T, d> Velocity;
 		TVector<T, d> Normal;
 	};
 
   public:
-	TPerParticlePBDCollisionConstraint(const TKinematicGeometryParticlesImp<T, d, SimType>& InParticles, TArray<bool>& Collided, TArray<uint32>& DynamicGroupIds, TArray<uint32>& KinematicGroupIds, const TArray<T>& PerGroupThickness, const TArray<T>& PerGroupFriction)
-	    : bFastPositionBasedFriction(true), MParticles(InParticles), MCollided(Collided), MDynamicGroupIds(DynamicGroupIds), MKinematicGroupIds(KinematicGroupIds), MPerGroupThickness(PerGroupThickness), MPerGroupFriction(PerGroupFriction) {}
+	typedef TKinematicGeometryParticlesImp<T, d, SimType> FCollisionParticles;
+
+	TPerParticlePBDCollisionConstraint(const TPBDActiveView<FCollisionParticles>& InParticlesActiveView, TArray<bool>& Collided, TArray<uint32>& DynamicGroupIds, TArray<uint32>& KinematicGroupIds, const TArray<T>& PerGroupThickness, const TArray<T>& PerGroupFriction)
+	    : bFastPositionBasedFriction(true), MCollisionParticlesActiveView(InParticlesActiveView), MCollided(Collided), MDynamicGroupIds(DynamicGroupIds), MKinematicGroupIds(KinematicGroupIds), MPerGroupThickness(PerGroupThickness), MPerGroupFriction(PerGroupFriction) {}
+
 	virtual ~TPerParticlePBDCollisionConstraint() {}
 
 	inline void Apply(TPBDParticles<T, d>& InParticles, const T Dt, const int32 Index) const override //-V762
@@ -30,18 +34,18 @@ class TPerParticlePBDCollisionConstraint : public TPerParticleRule<T, d>
 		}
 
 		const uint32 DynamicGroupId = MDynamicGroupIds[Index];  // Particle group Id
-		
-		for (uint32 i = 0; i < MParticles.Size(); ++i)
+
+		MCollisionParticlesActiveView.SequentialFor([this, &InParticles, &Dt, &Index, &DynamicGroupId](FCollisionParticles& CollisionParticles, int32 i)
 		{
 			const uint32 KinematicGroupId = MKinematicGroupIds[i];  // Collision group Id
 
 			if (KinematicGroupId != (uint32)INDEX_NONE && DynamicGroupId != KinematicGroupId)
 			{
-				continue; // Bail out if the collision groups doesn't match the particle group id, or use INDEX_NONE (= global collision that affects all particle)
+				return; // Bail out if the collision groups doesn't match the particle group id, or use INDEX_NONE (= global collision that affects all particle)
 			}
 			TVector<T, d> Normal;
-			TRigidTransform<T, d> Frame(MParticles.X(i), MParticles.R(i));
-			const T Phi = MParticles.Geometry(i)->PhiWithNormal(Frame.InverseTransformPosition(InParticles.P(Index)), Normal);
+			TRigidTransform<T, d> Frame(CollisionParticles.X(i), CollisionParticles.R(i));
+			const T Phi = CollisionParticles.Geometry(i)->PhiWithNormal(Frame.InverseTransformPosition(InParticles.P(Index)), Normal);
 			const T Thickness = MPerGroupThickness[DynamicGroupId];
 			if (Phi < Thickness)
 			{
@@ -54,8 +58,8 @@ class TPerParticlePBDCollisionConstraint : public TPerParticleRule<T, d>
 					if (bFastPositionBasedFriction)
 					{
 						const T Penetration = (-Phi + Thickness); // This is related to the Normal impulse
-						TVector<T, d> VectorToPoint = InParticles.P(Index) - MParticles.X(i);
-						const TVector<T, d> RelativeDisplacement = (InParticles.P(Index) - InParticles.X(Index)) - (MParticles.V(i) + TVector<T, d>::CrossProduct(MParticles.W(i), VectorToPoint)) * Dt; // This corresponds to the tangential velocity multiplied by dt (friction will drive this to zero if it is high enough)
+						TVector<T, d> VectorToPoint = InParticles.P(Index) - CollisionParticles.X(i);
+						const TVector<T, d> RelativeDisplacement = (InParticles.P(Index) - InParticles.X(Index)) - (CollisionParticles.V(i) + TVector<T, d>::CrossProduct(CollisionParticles.W(i), VectorToPoint)) * Dt; // This corresponds to the tangential velocity multiplied by dt (friction will drive this to zero if it is high enough)
 						const TVector<T, d> RelativeDisplacementTangent = RelativeDisplacement - TVector<T, d>::DotProduct(RelativeDisplacement, NormalWorld) * NormalWorld; // Project displacement into the tangential plane
 						const T RelativeDisplacementTangentLength = RelativeDisplacementTangent.Size();
 						T PositionCorrection = FMath::Clamp<T>(Penetration * CoefficientOfFriction, 0, RelativeDisplacementTangentLength);
@@ -71,16 +75,16 @@ class TPerParticlePBDCollisionConstraint : public TPerParticleRule<T, d>
 					else
 					{
 						// Note, to fix: Only use fast position based friction for now, since adding to TMaps here is not thread safe when calling Apply on multiple threads (will cause crash)
-						VelocityConstraint Constraint;
-						TVector<T, d> VectorToPoint = InParticles.P(Index) - MParticles.X(i);
-						Constraint.Velocity = MParticles.V(i) + TVector<T, d>::CrossProduct(MParticles.W(i), VectorToPoint);
+						FVelocityConstraint Constraint;
+						TVector<T, d> VectorToPoint = InParticles.P(Index) - CollisionParticles.X(i);
+						Constraint.Velocity = CollisionParticles.V(i) + TVector<T, d>::CrossProduct(CollisionParticles.W(i), VectorToPoint);
 						Constraint.Normal = Frame.TransformVector(Normal);
 						
 						MVelocityConstraints.Add(Index, Constraint);
 					}
 				}
 			}
-		}
+		});
 	}
 
 	void ApplyFriction(TPBDParticles<T, d>& InParticles, const T Dt, const int32 Index) const
@@ -108,11 +112,11 @@ class TPerParticlePBDCollisionConstraint : public TPerParticleRule<T, d>
   private:
 	bool bFastPositionBasedFriction;
 	// TODO(mlentine): Need a bb hierarchy
-	const TKinematicGeometryParticlesImp<T, d, SimType>& MParticles;
+	const TPBDActiveView<FCollisionParticles>& MCollisionParticlesActiveView;
 	TArray<bool>& MCollided;
 	TArray<uint32>& MDynamicGroupIds;
 	TArray<uint32>& MKinematicGroupIds;
-	mutable TMap<int32, VelocityConstraint> MVelocityConstraints;
+	mutable TMap<int32, FVelocityConstraint> MVelocityConstraints;
 	const TArray<T>& MPerGroupThickness;
 	const TArray<T>& MPerGroupFriction;
 };
