Chaos Physics: Restitution fixes brought back from 4.26

#[fyi]  Michael.Lentine, Benn.Gallagher, Chris.Caulfield, Ori.Cohen

diff --git a/Engine/Source/Programs/HeadlessChaos/Private/GeometryCollection/GeometryCollectionTestClustering.cpp b/Engine/Source/Programs/HeadlessChaos/Private/GeometryCollection/GeometryCollectionTestClustering.cpp
--- a/Engine/Source/Programs/HeadlessChaos/Private/GeometryCollection/GeometryCollectionTestClustering.cpp
+++ b/Engine/Source/Programs/HeadlessChaos/Private/GeometryCollection/GeometryCollectionTestClustering.cpp
@@ -461,7 +461,7 @@ namespace GeometryCollectionTest
 				BrokenFrame = Frame;
 			}
 
-			if ((BrokenFrame != INDEX_NONE) && (Frame > BrokenFrame))
+			if ((BrokenFrame != INDEX_NONE) && (Frame > BrokenFrame + 1)) // +1 so that the boxes have a bit of time to move away from each other
 			{
 				// The boxes are now moving independently - the bottom one is on the ground and should be stopped.
 				// The top one is still falling, so they should be closer together	
diff --git a/Engine/Source/Programs/HeadlessChaos/Private/HeadlessChaosTestRewind.cpp b/Engine/Source/Programs/HeadlessChaos/Private/HeadlessChaosTestRewind.cpp
--- a/Engine/Source/Programs/HeadlessChaos/Private/HeadlessChaosTestRewind.cpp
+++ b/Engine/Source/Programs/HeadlessChaos/Private/HeadlessChaosTestRewind.cpp
@@ -1128,7 +1128,9 @@ namespace ChaosTest {
 				TickSolverHelper(Module,Solver);
 			}
 
-			EXPECT_FLOAT_EQ(Dynamic->X()[2],10);
+			// We may end up a bit away from the surface (dt * V), due to solving for 0 velocity and not 0 position error
+			EXPECT_GE(Dynamic->X()[2], 10);
+			EXPECT_LE(Dynamic->X()[2], 11);
 		
 			const int RewindStep = 5;
 
@@ -1164,7 +1166,9 @@ namespace ChaosTest {
 			EXPECT_EQ(DesyncedParticles[0].MostDesynced,ESyncState::HardDesync);
 			EXPECT_EQ(DesyncedParticles[1].MostDesynced,ESyncState::HardDesync);
 
-			EXPECT_FLOAT_EQ(Dynamic->X()[2],9);
+			// We may end up a bit away from the surface (dt * V), due to solving for 0 velocity and not 0 position error
+			EXPECT_GE(Dynamic->X()[2],9);
+			EXPECT_LE(Dynamic->X()[2], 10);
 
 			Module->DestroySolver(Solver);
 		}
@@ -1346,7 +1350,9 @@ namespace ChaosTest {
 				Xs.Add(Dynamic->X());
 			}
 
-			EXPECT_FLOAT_EQ(Dynamic->X()[2],10);
+
+			EXPECT_GE(Dynamic->X()[2], 10);
+			EXPECT_LE(Dynamic->X()[2], 11);
 
 			const int RewindStep = 5;
 
@@ -1370,7 +1376,9 @@ namespace ChaosTest {
 			EXPECT_EQ(DesyncedParticles[0].MostDesynced,ESyncState::HardDesync);
 			EXPECT_EQ(DesyncedParticles[0].Particle,Kinematic.Get());
 
-			EXPECT_FLOAT_EQ(Dynamic->X()[2],10);
+			// We may end up a bit away from the surface (dt * V), due to solving for 0 velocity and not 0 position error
+			EXPECT_GE(Dynamic->X()[2],10);
+			EXPECT_LE(Dynamic->X()[2], 11);
 
 			Module->DestroySolver(Solver);
 		}
@@ -1421,7 +1429,9 @@ namespace ChaosTest {
 				Xs.Add(Dynamic->X());
 			}
 
-			EXPECT_FLOAT_EQ(Dynamic->X()[2],5);
+			// We may end up a bit away from the surface (dt * V), due to solving for 0 velocity and not 0 position error
+			EXPECT_GE(Dynamic->X()[2], 5);
+			EXPECT_LE(Dynamic->X()[2], 6);
 
 			const int RewindStep = 0;
 
@@ -1438,7 +1448,9 @@ namespace ChaosTest {
 				EXPECT_GE(Dynamic->X()[2],10);
 			}
 
-			EXPECT_FLOAT_EQ(Dynamic->X()[2],10);
+			// We may end up a bit away from the surface (dt * V), due to solving for 0 velocity and not 0 position error
+			EXPECT_GE(Dynamic->X()[2], 10);
+			EXPECT_LE(Dynamic->X()[2], 11);
 
 			//both desync
 			const TArray<FDesyncedParticleInfo> DesyncedParticles = RewindData->ComputeDesyncInfo();
@@ -1496,7 +1508,9 @@ namespace ChaosTest {
 				Xs.Add(Dynamic->X());
 			}
 
-			EXPECT_FLOAT_EQ(Dynamic->X()[2],5);
+			// We may end up a bit away from the surface (dt * V), due to solving for 0 velocity and not 0 position error
+			EXPECT_GE(Dynamic->X()[2], 5);
+			EXPECT_LE(Dynamic->X()[2], 6);
 
 			const int RewindStep = 0;
 
@@ -1514,7 +1528,9 @@ namespace ChaosTest {
 				EXPECT_VECTOR_FLOAT_EQ(Dynamic->X(),Xs[Step]);
 			}
 
-			EXPECT_FLOAT_EQ(Dynamic->X()[2],5);
+			// We may end up a bit away from the surface (dt * V), due to solving for 0 velocity and not 0 position error
+			EXPECT_GE(Dynamic->X()[2], 5);
+			EXPECT_LE(Dynamic->X()[2], 6);
 
 			//dynamic slave so only kinematic desyncs
 			const TArray<FDesyncedParticleInfo> DesyncedParticles = RewindData->ComputeDesyncInfo();
@@ -1820,7 +1836,7 @@ namespace ChaosTest {
 				}
 
 				TickSolverHelper(Module,Solver);
-				EXPECT_LT(Dynamic->X()[2],10);
+				EXPECT_LE(Dynamic->X()[2],10);
 
 				//kinematic desync will be known at end of frame because the simulation doesn't write results (so we know right away it's a desync)
 				if(Step < LastStep)
@@ -1889,7 +1905,9 @@ namespace ChaosTest {
 			Xs.Add(Dynamic->X());
 		}
 
-		EXPECT_FLOAT_EQ(Dynamic->X()[2],10);
+		// We may end up a bit away from the surface (dt * V), due to solving for 0 velocity and not 0 position error
+		EXPECT_GE(Dynamic->X()[2], 10);
+		EXPECT_LE(Dynamic->X()[2], 12);
 
 		const int RewindStep = 0;
 
@@ -1935,7 +1953,10 @@ namespace ChaosTest {
 		EXPECT_EQ(DesyncedParticles[1].MostDesynced,DesyncedParticles[1].Particle == Kinematic.Get() ? ESyncState::HardDesync : ESyncState::SoftDesync);
 
 		EXPECT_TRUE(bEverSoft);
-		EXPECT_FLOAT_EQ(Dynamic->X()[2],10);
+
+		// We may end up a bit away from the surface (dt * V), due to solving for 0 velocity and not 0 position error
+		EXPECT_GE(Dynamic->X()[2], 10);
+		EXPECT_LE(Dynamic->X()[2], 12);
 
 		Module->DestroySolver(Solver);
 	}
diff --git a/Engine/Source/Runtime/Experimental/Chaos/Private/Chaos/PBDCollisionConstraintsContact.cpp b/Engine/Source/Runtime/Experimental/Chaos/Private/Chaos/PBDCollisionConstraintsContact.cpp
--- a/Engine/Source/Runtime/Experimental/Chaos/Private/Chaos/PBDCollisionConstraintsContact.cpp
+++ b/Engine/Source/Runtime/Experimental/Chaos/Private/Chaos/PBDCollisionConstraintsContact.cpp
@@ -26,9 +26,12 @@ namespace Chaos
 		int32 Chaos_Collision_EnergyClampEnabled = 1;
 		FAutoConsoleVariableRef CVarChaosCollisionEnergyClampEnabled(TEXT("p.Chaos.Collision.EnergyClampEnabled"), Chaos_Collision_EnergyClampEnabled, TEXT("Whether to use energy clamping in collision apply step"));
 
-		int32 Chaos_Collision_RelaxationEnabled = 1;
+		int32 Chaos_Collision_RelaxationEnabled = 0; // TODO remove this feature soon
 		FAutoConsoleVariableRef CVarChaosCollisionRelaxationEnabled(TEXT("p.Chaos.Collision.RelaxationEnabled"), Chaos_Collision_RelaxationEnabled, TEXT("Whether to reduce applied impulses during iterations for improved solver stability but reduced convergence"));
 
+		int32 Chaos_Collision_PrevVelocityRestitutionEnabled = 0;
+		FAutoConsoleVariableRef CVarChaosCollisionPrevVelocityRestitutionEnabled(TEXT("p.Chaos.Collision.PrevVelocityRestitutionEnabled"), Chaos_Collision_PrevVelocityRestitutionEnabled, TEXT("If enabled restitution will be calculated on previous frame velocities instead of current frame velocities"));
+
 		int32 Chaos_Collision_ForceApplyType = 0;
 		FAutoConsoleVariableRef CVarChaosCollisionAlternativeApply(TEXT("p.Chaos.Collision.ForceApplyType"), Chaos_Collision_ForceApplyType, TEXT("Force Apply step to use Velocity(1) or Position(2) modes"));
 
@@ -137,6 +140,34 @@ namespace Chaos
 			return RelaxationFactor;
 		}
 
+		FReal CalculateRelativeNormalVelocityForRestitution(
+			const TGenericParticleHandle<FReal, 3> Particle0,
+			const TGenericParticleHandle<FReal, 3> Particle1,
+			const FRotation3& Q0,
+			const FRotation3& Q1,
+			const FVec3& ContactNormal,
+			const FVec3& VectorToPoint1,
+			const FVec3& VectorToPoint2,
+			FReal &outVelocitySize)
+		{
+			// Get previous world space position of the contact points relative the CoM
+			// Note: These particular points might not even have been in contact at the start of the frame
+			// VectorToContact point is transformed to local space, and then back to world space at the previous frame orientation
+			FRotation3 R0 = FParticleUtilitiesXR::GetCoMWorldRotation(Particle0); // Previous Rotation
+			FRotation3 R1 = FParticleUtilitiesXR::GetCoMWorldRotation(Particle1); // Previous Rotation
+
+			const FVec3 VectorToPoint1Prev = R0 * FRotation3::Conjugate(Q0) * VectorToPoint1;
+			const FVec3 VectorToPoint2Prev = R1 * FRotation3::Conjugate(Q1) * VectorToPoint2;
+			// 
+			const FVec3 ContactBody1VelocityPrev = FParticleUtilities::GetPreviousVelocityAtCoMRelativePosition(Particle0, VectorToPoint1Prev);
+			const FVec3 ContactBody2VelocityPrev = FParticleUtilities::GetPreviousVelocityAtCoMRelativePosition(Particle1, VectorToPoint2Prev);
+			const FVec3 RelativeVelocityPrev = ContactBody1VelocityPrev - ContactBody2VelocityPrev;
+			const FReal RelativeNormalVelocityForRestitution = FVec3::DotProduct(RelativeVelocityPrev, ContactNormal); // Note: using the current contact normal
+			outVelocitySize = RelativeVelocityPrev.Size();
+			
+			return RelativeNormalVelocityForRestitution;
+		}
+
 		// Calculate velocity corrections due to the given contact
 		// This function uses AccumulatedImpulse Clipping, so 
 		// AccumulatedImpulse is both an input and an output
@@ -182,13 +213,21 @@ namespace Chaos
 				(bIsRigidDynamic1 ? ComputeFactorMatrix3(VectorToPoint2, WorldSpaceInvI2, PBDRigid1->InvM()) : FMatrix33(0));
 			FVec3 AngularImpulse(0);
 
+			FReal RelativeNormalVelocityForRestitution = FVec3::DotProduct(RelativeVelocity, Contact.Normal); // Relative velocity in direction of the normal as used by restitution
+			FReal RelativeVelocityForRestitutionSize = RelativeVelocity.Size();
+			// Use the previous contact velocities to calculate the restitution response
+			if (Chaos_Collision_PrevVelocityRestitutionEnabled && bInApplyRestitution && Contact.Restitution > (FReal)0.0f)
+			{
+				RelativeNormalVelocityForRestitution = CalculateRelativeNormalVelocityForRestitution(Particle0, Particle1, Q0, Q1, Contact.Normal, VectorToPoint1, VectorToPoint2, RelativeVelocityForRestitutionSize);
+			}
+
 			// Resting contact if very close to the surface
-			const bool bApplyRestitution = bInApplyRestitution && (RelativeVelocity.Size() > ParticleParameters.RestitutionVelocityThreshold);
+			const bool bApplyRestitution = bInApplyRestitution && (RelativeVelocityForRestitutionSize > ParticleParameters.RestitutionVelocityThreshold);
 			const FReal Restitution = (bApplyRestitution) ? Contact.Restitution : (FReal)0;
 			const FReal Friction = bInApplyFriction ? Contact.Friction : (FReal)0; // Add friction even when pushing out
-			const FReal AngularFriction = bInApplyAngularFriction ? Contact.AngularFriction : (FReal)0; // Don't add angular friction in pushout since we don't have accumelated angular impulse clipping, todo: experiment with this later
+			const FReal AngularFriction = bInApplyAngularFriction ? Contact.AngularFriction : (FReal)0; // Don't add angular friction in pushout since we don't have accumulated angular impulse clipping, todo: experiment with this later
 
-			const FVec3 VelocityTarget = (-Restitution * FVec3::DotProduct(RelativeVelocity, Contact.Normal)) * Contact.Normal;
+			const FVec3 VelocityTarget = (-Restitution * RelativeNormalVelocityForRestitution) * Contact.Normal;
 			const FVec3 VelocityChange = VelocityTarget - RelativeVelocity;
 			const FMatrix33 FactorInverse = Factor.Inverse();
 			FVec3 Impulse = FactorInverse * VelocityChange;  // Delta Impulse = (J.M^(-1).J^(T))^(-1).(ContactVelocityError)
@@ -384,12 +423,25 @@ namespace Chaos
 			FRotation3 Q0 = FParticleUtilities::GetCoMWorldRotation(Particle0);
 			FRotation3 Q1 = FParticleUtilities::GetCoMWorldRotation(Particle1);
 
-			FVec3 VectorToPoint1 = Contact.Location - P0;
-			FVec3 VectorToPoint2 = Contact.Location - P1;
-			FVec3 Body1Velocity = FParticleUtilities::GetVelocityAtCoMRelativePosition(Particle0, VectorToPoint1);
-			FVec3 Body2Velocity = FParticleUtilities::GetVelocityAtCoMRelativePosition(Particle1, VectorToPoint2);
-			FVec3 RelativeVelocity = Body1Velocity - Body2Velocity;
-			FReal RelativeNormalVelocity = FVec3::DotProduct(RelativeVelocity, Contact.Normal);
+			const FVec3 VectorToPoint1 = Contact.Location - P0;
+			const FVec3 VectorToPoint2 = Contact.Location - P1;
+			const FVec3 Body1Velocity = FParticleUtilities::GetVelocityAtCoMRelativePosition(Particle0, VectorToPoint1);
+			const FVec3 Body2Velocity = FParticleUtilities::GetVelocityAtCoMRelativePosition(Particle1, VectorToPoint2);
+			const FVec3 RelativeVelocity = Body1Velocity - Body2Velocity;
+			const FReal RelativeNormalVelocity = FVec3::DotProduct(RelativeVelocity, Contact.Normal);
+			
+
+			FReal RelativeNormalVelocityForRestitution = RelativeNormalVelocity; // Relative velocity in direction of the normal as used by restitution
+			FReal RelativeVelocityForRestitutionSize = RelativeVelocity.Size();
+			// Use the previous contact velocities to calculate the restitution response
+			if (Chaos_Collision_PrevVelocityRestitutionEnabled && Contact.Restitution > (FReal)0.0f)
+			{
+				RelativeNormalVelocityForRestitution = CalculateRelativeNormalVelocityForRestitution(Particle0, Particle1, Q0, Q1, Contact.Normal, VectorToPoint1, VectorToPoint2, RelativeVelocityForRestitutionSize);
+			}
+
+			// Resting contact if very close to the surface
+			bool bApplyRestitution = (RelativeVelocityForRestitutionSize > ParticleParameters.RestitutionVelocityThreshold);
+			FReal Restitution = (bApplyRestitution) ? Contact.Restitution : (FReal)0;
 
 			if (RelativeNormalVelocity < 0) // ignore separating constraints
 			{
@@ -402,16 +454,13 @@ namespace Chaos
 					(bIsRigidDynamic1 ? ComputeFactorMatrix3(VectorToPoint2, WorldSpaceInvI2, PBDRigid1->InvM()) : FMatrix33(0));
 				FVec3 Impulse;
 				FVec3 AngularImpulse(0);
-
-				// Resting contact if very close to the surface
-				bool bApplyRestitution = (RelativeVelocity.Size() > ParticleParameters.RestitutionVelocityThreshold);
-				FReal Restitution = (bApplyRestitution) ? Contact.Restitution : (FReal)0;
+				
 				FReal Friction = Contact.Friction;
 				FReal AngularFriction = Contact.AngularFriction;
 
 				if (Friction > 0)
 				{
-					FVec3 VelocityChange = -(Restitution * RelativeNormalVelocity * Contact.Normal + RelativeVelocity);
+					FVec3 VelocityChange = -(Restitution * RelativeNormalVelocityForRestitution * Contact.Normal + RelativeVelocity);
 					FReal NormalVelocityChange = FVec3::DotProduct(VelocityChange, Contact.Normal);
 					FMatrix33 FactorInverse = Factor.Inverse();
 					FVec3 MinimalImpulse = FactorInverse * VelocityChange;
@@ -456,14 +505,14 @@ namespace Chaos
 							ImpulseDenominator = (FReal)1;
 						}
 
-						const FReal ImpulseMag = -(1 + Restitution) * RelativeNormalVelocity / ImpulseDenominator;
+						const FReal ImpulseMag = -(RelativeNormalVelocity + Restitution * RelativeNormalVelocityForRestitution) / ImpulseDenominator;
 						Impulse = ImpulseMag * (Contact.Normal - Friction * Tangent);
 					}
 				}
 				else
 				{
 					FReal ImpulseDenominator = FVec3::DotProduct(Contact.Normal, Factor * Contact.Normal);
-					FVec3 ImpulseNumerator = -(1 + Restitution) * FVec3::DotProduct(RelativeVelocity, Contact.Normal)* Contact.Normal;
+					FVec3 ImpulseNumerator = -(RelativeNormalVelocity + Restitution * RelativeNormalVelocityForRestitution) * Contact.Normal;
 					if (!ensureMsgf(FMath::Abs(ImpulseDenominator) > SMALL_NUMBER, TEXT("Contact:%s\n\nParticle:%s\n\nLevelset:%s\n\nFactor*Constraint.Normal:%s, ImpulseDenominator:%f"),
 						*Contact.ToString(),
 						*Particle0->ToString(),
@@ -861,7 +910,7 @@ namespace Chaos
 			// but also update the contact location if this is done
 			{
 				FVec3 DV0, DW0, DV1, DW1;
-				CalculateContactVelocityCorrections(Contact, Particle0, Particle1, IterationParameters, ParticleParameters, P0, Q0, P1, Q1, false, true, false, false, AccumulatedImpulse, DV0, DW0, DV1, DW1);
+				CalculateContactVelocityCorrections(Contact, Particle0, Particle1, IterationParameters, ParticleParameters, P0, Q0, P1, Q1, true, true, false, false, AccumulatedImpulse, DV0, DW0, DV1, DW1);
 				if (bIsRigidDynamic0)
 				{
 					PBDRigid0->V() += DV0;
diff --git a/Engine/Source/Runtime/Experimental/Chaos/Public/Chaos/Particle/ParticleUtilities.h b/Engine/Source/Runtime/Experimental/Chaos/Public/Chaos/Particle/ParticleUtilities.h
--- a/Engine/Source/Runtime/Experimental/Chaos/Public/Chaos/Particle/ParticleUtilities.h
+++ b/Engine/Source/Runtime/Experimental/Chaos/Public/Chaos/Particle/ParticleUtilities.h
@@ -121,6 +121,12 @@ namespace Chaos
 			return Particle->V() + FVec3::CrossProduct(Particle->W(), RelPos);
 		}
 
+		template<typename T_PARTICLEHANDLE>
+		static inline FVec3 GetPreviousVelocityAtCoMRelativePosition(T_PARTICLEHANDLE Particle, const FVec3& RelPos)
+		{
+			return Particle->PreV() + FVec3::CrossProduct(Particle->PreW(), RelPos);
+		}
+
 		template<typename T_PARTICLEHANDLE>
 		static inline FVec3 GetCoMWorldPosition(T_PARTICLEHANDLE Particle)
 		{
@@ -247,6 +253,12 @@ namespace Chaos
 			return Particle->V() + FVec3::CrossProduct(Particle->W(), RelPos);
 		}
 
+		template<typename T_PARTICLEHANDLE>
+		static inline FVec3 GetPreviousVelocityAtCoMRelativePosition(T_PARTICLEHANDLE Particle, const FVec3& RelPos)
+		{
+			return Particle->PreV() + FVec3::CrossProduct(Particle->PreW(), RelPos);
+		}
+
 		template<typename T_PARTICLEHANDLE>
 		static inline const FVec3& GetCoMWorldPosition(T_PARTICLEHANDLE Particle)
 		{
