Integrate 14387164 - Clear existing constraints during solver destroy to ensure we don't generate collision events with particles that have been unregistered during final flush.
Fixes headless chaos crash.

diff --git a/Engine/Source/Runtime/Experimental/Chaos/Private/Chaos/Framework/PhysicsSolverBase.cpp b/Engine/Source/Runtime/Experimental/Chaos/Private/Chaos/Framework/PhysicsSolverBase.cpp
--- a/Engine/Source/Runtime/Experimental/Chaos/Private/Chaos/Framework/PhysicsSolverBase.cpp
+++ b/Engine/Source/Runtime/Experimental/Chaos/Private/Chaos/Framework/PhysicsSolverBase.cpp
@@ -5,6 +5,7 @@
 #include "ProfilingDebugging/CsvProfiler.h"
 #include "ChaosStats.h"
 #include "Chaos/PendingSpatialData.h"
+#include "PBDRigidsSolver.h"
 
 namespace Chaos
 {	
@@ -147,6 +148,17 @@ namespace Chaos
 		{
 			Command();
 		}
+		
+		// GeometryCollection particles do not always remove collision constraints on unregister,
+		// explicitly clear constraints so we will not crash when filling collision events in advance.
+		InSolver.CastHelper([](auto& Concrete)
+		{
+			auto* Evolution = Concrete.GetEvolution();
+			if (Evolution)
+			{
+				Evolution->ResetConstraints();
+			}
+		});
 
 		// Advance in single threaded because we cannot block on an async task here if in multi threaded mode. see above comments.
 		InSolver.SetThreadingMode_External(EThreadingModeTemp::SingleThread);
diff --git a/Engine/Source/Runtime/Experimental/Chaos/Public/Chaos/PBDConstraintRule.h b/Engine/Source/Runtime/Experimental/Chaos/Public/Chaos/PBDConstraintRule.h
--- a/Engine/Source/Runtime/Experimental/Chaos/Public/Chaos/PBDConstraintRule.h
+++ b/Engine/Source/Runtime/Experimental/Chaos/Public/Chaos/PBDConstraintRule.h
@@ -156,6 +156,9 @@ namespace Chaos
 		// @todo(ccaulfield): remove uint version
 		virtual void RemoveConstraints(const TSet<TGeometryParticleHandle<FReal, 3>*>& RemovedParticles) { }
 
+		/** Remove all constraints */
+		virtual void ResetConstraints() {}
+
 		/** The number of constraints in the collection */
 		virtual int32 NumConstraints() const { return 0; }
 	};
@@ -341,6 +344,11 @@ namespace Chaos
 			Constraints.RemoveConstraints(InConstraints);
 		}
 
+		virtual void ResetConstraints()
+		{
+			Constraints.Reset();
+		}
+
 		virtual bool ApplyPushOut(const FReal Dt, int32 Island, const int32 It, const int32 NumIts) override
 		{
 			const typename FPBDConstraintColor::FLevelToColorToConstraintListMap& LevelToColorToConstraintListMap = GraphColor.GetIslandLevelToColorToConstraintListMap(Island);
diff --git a/Engine/Source/Runtime/Experimental/Chaos/Public/Chaos/PBDRigidsEvolution.h b/Engine/Source/Runtime/Experimental/Chaos/Public/Chaos/PBDRigidsEvolution.h
--- a/Engine/Source/Runtime/Experimental/Chaos/Public/Chaos/PBDRigidsEvolution.h
+++ b/Engine/Source/Runtime/Experimental/Chaos/Public/Chaos/PBDRigidsEvolution.h
@@ -453,6 +453,14 @@ public:
 		}
 	}
 
+	CHAOS_API void ResetConstraints()
+	{
+		for(FPBDConstraintGraphRule* ConstraintRule : ConstraintRules)
+		{
+			ConstraintRule->ResetConstraints();
+		}
+	}
+
 	//TEMP: this is only needed while clustering continues to use indices directly
 	const auto& GetActiveClusteredArray() const { return Particles.GetActiveClusteredArray(); }
 	const auto& GetNonDisabledClusteredArray() const { return Particles.GetNonDisabledClusteredArray(); }
