Fixes for stress test unity builds, incorrect exports on template classes (instead of specializations), global namespace and macro collisions.

diff --git a/Engine/Source/Runtime/Experimental/Chaos/Private/Chaos/Levelset.cpp b/Engine/Source/Runtime/Experimental/Chaos/Private/Chaos/Levelset.cpp
--- a/Engine/Source/Runtime/Experimental/Chaos/Private/Chaos/Levelset.cpp
+++ b/Engine/Source/Runtime/Experimental/Chaos/Private/Chaos/Levelset.cpp
@@ -37,12 +37,12 @@ FAutoConsoleVariableRef CVarNumOverlapSphereSamples(TEXT("p.LevelsetOverlapSpher
 int32 NumOverlapCapsuleSamples = 24;
 FAutoConsoleVariableRef CVarNumOverlapCapsuleSamples(TEXT("p.LevelsetOverlapCapsuleSamples"), NumOverlapCapsuleSamples, TEXT("Number of spiral points to generate for levelset-capsule overlaps"));
 
-using namespace Chaos;
-
 #define MAX_CLAMP(a, comp, b) (a >= comp ? b : a)
 #define MIN_CLAMP(a, comp, b) (a < comp ? b : a)
 #define RANGE_CLAMP(a, comp, b) ((a < 0 || comp <= a) ? b : a)
 
+namespace Chaos
+{
 template<class T, int d>
 TLevelSet<T, d>::TLevelSet(FErrorReporter& ErrorReporter, const TUniformGrid<T, d>& InGrid, const TParticles<T, d>& InParticles, const TTriangleMesh<T>& Mesh, const int32 BandWidth)
     : FImplicitObject(EImplicitObject::HasBoundingBox, ImplicitObjectType::LevelSet)
@@ -1739,5 +1739,10 @@ bool TLevelSet<T, d>::OverlapGeom(const TImplicitObjectScaled<FConvex>& QueryGeo
 {
 	return OverlapGeomImp(QueryGeom, QueryTM, Thickness, OutMTD);
 }
+}
+
+#undef MAX_CLAMP
+#undef MIN_CLAMP
+#undef RANGE_CLAMP
 
-template class Chaos::TLevelSet<float, 3>;
+template class CHAOS_API Chaos::TLevelSet<float, 3>;
diff --git a/Engine/Source/Runtime/Experimental/Chaos/Public/Chaos/Levelset.h b/Engine/Source/Runtime/Experimental/Chaos/Public/Chaos/Levelset.h
--- a/Engine/Source/Runtime/Experimental/Chaos/Public/Chaos/Levelset.h
+++ b/Engine/Source/Runtime/Experimental/Chaos/Public/Chaos/Levelset.h
@@ -13,6 +13,8 @@ namespace Chaos { class FErrorReporter; }
 namespace Chaos
 {
 
+struct FMTDInfo;
+
 template<class T>
 class TTriangleMesh;
 
@@ -25,7 +27,7 @@ class TCapsule;
 class FConvex;
 
 template<class T, int d>
-class CHAOS_API TLevelSet final : public FImplicitObject
+class TLevelSet final : public FImplicitObject
 {
   public:
 	using FImplicitObject::SignedDistance;
@@ -194,4 +196,4 @@ private:
 	TLevelSet() : FImplicitObject(EImplicitObject::HasBoundingBox, ImplicitObjectType::LevelSet) {}	//needed for serialization
 	friend FImplicitObject;	//needed for serialization
 };
-}
+}
\ No newline at end of file
