Fixed dynamic velocity vectors on non-Nanite geometry collection meshes.

diff --git a/Engine/Source/Runtime/Experimental/GeometryCollectionEngine/Private/GeometryCollection/GeometryCollectionSceneProxy.cpp b/Engine/Source/Runtime/Experimental/GeometryCollectionEngine/Private/GeometryCollection/GeometryCollectionSceneProxy.cpp
--- a/Engine/Source/Runtime/Experimental/GeometryCollectionEngine/Private/GeometryCollection/GeometryCollectionSceneProxy.cpp
+++ b/Engine/Source/Runtime/Experimental/GeometryCollectionEngine/Private/GeometryCollection/GeometryCollectionSceneProxy.cpp
@@ -808,8 +808,7 @@ FPrimitiveViewRelevance FGeometryCollectionSceneProxy::GetViewRelevance(const FS
 	Result.bTranslucentSelfShadow = bCastVolumetricTranslucentShadow;
 	MaterialRelevance.SetPrimitiveViewRelevance(Result);
 
-	// #todo(dmp): why does this crash?
-	// Result.bVelocityRelevance = IsMovable() && Result.bOpaque && Result.bRenderInMainPass;
+	Result.bVelocityRelevance = IsMovable() && Result.bOpaque && Result.bRenderInMainPass;
 
 	return Result;
 
