Fix for Geometry Collection not registering Nanite meshes with Nanite streamer when using Fracture tool

diff --git a/Engine/Source/Runtime/Experimental/GeometryCollectionEngine/Private/GeometryCollection/GeometryCollectionObject.cpp b/Engine/Source/Runtime/Experimental/GeometryCollectionEngine/Private/GeometryCollection/GeometryCollectionObject.cpp
--- a/Engine/Source/Runtime/Experimental/GeometryCollectionEngine/Private/GeometryCollection/GeometryCollectionObject.cpp
+++ b/Engine/Source/Runtime/Experimental/GeometryCollectionEngine/Private/GeometryCollection/GeometryCollectionObject.cpp
@@ -593,11 +593,6 @@ void UGeometryCollection::PostEditChangeProperty(struct FPropertyChangedEvent& P
 		{
 			InvalidateCollection();
 			EnsureDataIsCooked();
-
-			if (FApp::CanEverRender())
-			{
-				InitResources();
-			}
 		}
 		else if (PropertyChangedEvent.Property->GetFName() != GET_MEMBER_NAME_CHECKED(UGeometryCollection, Materials))
 		{
@@ -625,6 +620,12 @@ void UGeometryCollection::EnsureDataIsCooked()
 	if (StateGuid != LastBuiltGuid)
 	{
 		CreateSimulationDataImp(/*bCopyFromDDC=*/ true);
+		
+		if (FApp::CanEverRender())
+		{
+			NaniteData->InitResources(this);
+		}
+
 		LastBuiltGuid = StateGuid;
 	}
 }
