Change spatial acceleration type to enum class to stop global scope shadowing of non class enums

diff --git a/Engine/Source/Runtime/Experimental/Chaos/Public/Chaos/ISpatialAcceleration.h b/Engine/Source/Runtime/Experimental/Chaos/Public/Chaos/ISpatialAcceleration.h
--- a/Engine/Source/Runtime/Experimental/Chaos/Public/Chaos/ISpatialAcceleration.h
+++ b/Engine/Source/Runtime/Experimental/Chaos/Public/Chaos/ISpatialAcceleration.h
@@ -162,7 +162,8 @@ public:
 
 };
 
-enum ESpatialAcceleration
+using SpatialAccelerationType = uint8;	//see ESpatialAcceleration. Projects can add their own custom types by using enum values higher than ESpatialAcceleration::Unknown
+enum class ESpatialAcceleration : SpatialAccelerationType
 {
 	BoundingVolume,
 	AABBTree,
@@ -172,7 +173,11 @@ enum ESpatialAcceleration
 	//For custom types continue the enum after ESpatialAcceleration::Unknown
 };
 
-using SpatialAccelerationType = uint8;	//see ESpatialAcceleration. Projects can add their own custom types by using enum values higher than ESpatialAcceleration::Unknown
+inline bool SpatialAccelerationEqual(ESpatialAcceleration A, SpatialAccelerationType B) { return (SpatialAccelerationType)A == B; }
+inline bool operator==(ESpatialAcceleration A, SpatialAccelerationType B) { return SpatialAccelerationEqual(A,B); }
+inline bool operator==(SpatialAccelerationType A, ESpatialAcceleration B) { return SpatialAccelerationEqual(B,A); }
+inline bool operator!=(ESpatialAcceleration A, SpatialAccelerationType B) { return !SpatialAccelerationEqual(A,B); }
+inline bool operator!=(SpatialAccelerationType A, ESpatialAcceleration B) { return !SpatialAccelerationEqual(B,A); }
 
 template <typename TPayload>
 typename TEnableIf<!TIsPointer<TPayload>::Value, FUniqueIdx>::Type GetUniqueIdx(const TPayload& Payload)
@@ -240,10 +245,14 @@ class CHAOS_API ISpatialAcceleration
 {
 public:
 
-	ISpatialAcceleration(SpatialAccelerationType InType = ESpatialAcceleration::Unknown)
+	ISpatialAcceleration(SpatialAccelerationType InType = static_cast<SpatialAccelerationType>(ESpatialAcceleration::Unknown))
 		: Type(InType), SyncTimestamp(0), AsyncTimeSlicingComplete(true)
-	{
-	}
+	{}
+
+	ISpatialAcceleration(ESpatialAcceleration InType)
+		: ISpatialAcceleration(static_cast<SpatialAccelerationType>(InType))
+	{}
+
 	virtual ~ISpatialAcceleration() = default;
 
 	virtual bool IsAsyncTimeSlicingComplete() { return AsyncTimeSlicingComplete; }
