Avoid game-thread stalls during heavy asset compilation by scheduling Chaos solver as Normal priority instead of Background.

diff --git a/Engine/Source/Runtime/Experimental/Chaos/Private/Chaos/PBDRigidsEvolution.cpp b/Engine/Source/Runtime/Experimental/Chaos/Private/Chaos/PBDRigidsEvolution.cpp
--- a/Engine/Source/Runtime/Experimental/Chaos/Private/Chaos/PBDRigidsEvolution.cpp
+++ b/Engine/Source/Runtime/Experimental/Chaos/Private/Chaos/PBDRigidsEvolution.cpp
@@ -248,7 +248,13 @@ namespace Chaos
 	template <typename Traits>
 	ENamedThreads::Type TPBDRigidsEvolutionBase<Traits>::FChaosAccelerationStructureTask::GetDesiredThread()
 	{
+#if WITH_EDITOR
+		// Heavy async compilation could fill the background threads in Editor preventing us from making
+		// progress which would cause game-thread stalls. Schedule as normal for Editor to avoid this.
+		return ENamedThreads::AnyNormalThreadNormalTask;
+#else
 		return ENamedThreads::AnyBackgroundThreadNormalTask;
+#endif
 	}
 
 	template <typename Traits>
