Fixed Nanite Geometry Collection crashes and incorrect rendering on cooked platforms.

diff --git a/Engine/Source/Runtime/Core/Public/UObject/UE5MainStreamObjectVersion.h b/Engine/Source/Runtime/Core/Public/UObject/UE5MainStreamObjectVersion.h
--- a/Engine/Source/Runtime/Core/Public/UObject/UE5MainStreamObjectVersion.h
+++ b/Engine/Source/Runtime/Core/Public/UObject/UE5MainStreamObjectVersion.h
@@ -43,6 +43,9 @@ struct CORE_API FUE5MainStreamObjectVersion
 		// Serialize HLODLayer in WorldPartition HLODActorDesc
 		WorldPartitionHLODActorDescSerializeHLODLayer,
 
+		// Fixed Nanite Geometry Collection cooked data
+		GeometryCollectionNaniteCooked,
+
 		// -----<new versions can be added above this line>-------------------------------------------------
 		VersionPlusOne,
 		LatestVersion = VersionPlusOne - 1
diff --git a/Engine/Source/Runtime/Experimental/GeometryCollectionEngine/Private/GeometryCollection/GeometryCollectionComponent.cpp b/Engine/Source/Runtime/Experimental/GeometryCollectionEngine/Private/GeometryCollection/GeometryCollectionComponent.cpp
--- a/Engine/Source/Runtime/Experimental/GeometryCollectionEngine/Private/GeometryCollection/GeometryCollectionComponent.cpp
+++ b/Engine/Source/Runtime/Experimental/GeometryCollectionEngine/Private/GeometryCollection/GeometryCollectionComponent.cpp
@@ -528,6 +528,7 @@ FPrimitiveSceneProxy* UGeometryCollectionComponent::CreateSceneProxy()
 	// TODO: Abstract with a common helper
 	if (UseNanite(GetScene()->GetShaderPlatform()) &&
 		RestCollection->EnableNanite &&
+		RestCollection->NaniteData != nullptr &&
 		GGeometryCollectionNanite != 0)
 	{
 		return new FNaniteGeometryCollectionSceneProxy(this);
@@ -2050,12 +2051,16 @@ void UGeometryCollectionComponent::SendRenderDynamicData_Concurrent()
 		InitDynamicData(DynamicData);
 
 		//
-		// Only send dynamic data if the transform arrys on the past N frames are different
+		// Only send dynamic data if the transform arrays on the past N frames are different
 		//
 		bool PastFramesTransformsAreEqual = true;
-		for (int i = 0; i < TransformsAreEqual.Num(); ++i)
+		for (int32 TransformIndex = 0; TransformIndex < TransformsAreEqual.Num(); ++TransformIndex)
 		{
-			PastFramesTransformsAreEqual = PastFramesTransformsAreEqual && TransformsAreEqual[i];
+			if (!TransformsAreEqual[TransformIndex])
+			{
+				PastFramesTransformsAreEqual = false;
+				break;
+			}
 		}
 
 		if (PastFramesTransformsAreEqual)
diff --git a/Engine/Source/Runtime/Experimental/GeometryCollectionEngine/Private/GeometryCollection/GeometryCollectionObject.cpp b/Engine/Source/Runtime/Experimental/GeometryCollectionEngine/Private/GeometryCollection/GeometryCollectionObject.cpp
--- a/Engine/Source/Runtime/Experimental/GeometryCollectionEngine/Private/GeometryCollection/GeometryCollectionObject.cpp
+++ b/Engine/Source/Runtime/Experimental/GeometryCollectionEngine/Private/GeometryCollection/GeometryCollectionObject.cpp
@@ -392,6 +392,16 @@ void UGeometryCollection::Serialize(FArchive& Ar)
 		}
 	}
 
+	if (Ar.CustomVer(FUE5MainStreamObjectVersion::GUID) >= FUE5MainStreamObjectVersion::GeometryCollectionNaniteCooked)
+	{
+		if (NaniteData == nullptr)
+		{
+			NaniteData = MakeUnique<FGeometryCollectionNaniteData>();
+		}
+
+		NaniteData->Serialize(ChaosAr, this);
+	}
+
 #if WITH_EDITOR
 	//for all versions loaded, make sure sim data is up to date
 	if (Ar.IsLoading())
diff --git a/Engine/Source/Runtime/Experimental/GeometryCollectionEngine/Private/GeometryCollection/GeometryCollectionSceneProxy.cpp b/Engine/Source/Runtime/Experimental/GeometryCollectionEngine/Private/GeometryCollection/GeometryCollectionSceneProxy.cpp
--- a/Engine/Source/Runtime/Experimental/GeometryCollectionEngine/Private/GeometryCollection/GeometryCollectionSceneProxy.cpp
+++ b/Engine/Source/Runtime/Experimental/GeometryCollectionEngine/Private/GeometryCollection/GeometryCollectionSceneProxy.cpp
@@ -1123,6 +1123,22 @@ FNaniteGeometryCollectionSceneProxy::FNaniteGeometryCollectionSceneProxy(UGeomet
 		Instance.RenderBounds = BoundingBoxes[GeometryIndex];
 		Instance.NaniteInfo = GeometryCollection->GetNaniteInfo(GeometryIndex);
 	}
+
+	// Need to specify initial instance list, even with just identity transforms, so that the
+	// GPUScene instance data allocator reserves space for the instances early on. The instance
+	// transforms will be corrected during the first frame before any rendering occurs.
+	Instances.SetNumZeroed(NumGeometry);
+	for (int32 GeometryIndex = 0; GeometryIndex < NumGeometry; ++GeometryIndex)
+	{
+		FPrimitiveInstance& Instance = Instances[GeometryIndex];
+		Instance.PrimitiveId = ~uint32(0);
+		Instance.InstanceToLocal.SetIdentity();
+		Instance.LocalToInstance.SetIdentity();
+		Instance.LocalToWorld.SetIdentity();
+		Instance.WorldToLocal.SetIdentity();
+		Instance.RenderBounds = GeometryNaniteData[GeometryIndex].RenderBounds;
+		Instance.LocalBounds = Instance.RenderBounds;
+	}
 }
 
 FPrimitiveViewRelevance FNaniteGeometryCollectionSceneProxy::GetViewRelevance(const FSceneView* View) const
@@ -1270,11 +1286,6 @@ void FNaniteGeometryCollectionSceneProxy::SetConstantData_RenderThread(FGeometry
 			continue;
 		}
 
-		//if (Collection->IsClustered(TransformIndex))
-		//{
-		//	continue;
-		//}
-
 		const FGeometryNaniteData& NaniteData = GeometryNaniteData[TransformToGeometryIndex];
 
 		FPrimitiveInstance& Instance = Instances.Emplace_GetRef();
@@ -1313,16 +1324,6 @@ void FNaniteGeometryCollectionSceneProxy::SetDynamicData_RenderThread(FGeometryC
 				continue;
 			}
 
-			//if (Collection->IsClustered(TransformIndex))
-			//{
-			//	continue;
-			//}
-
-			//if (TransformChildren[TransformIndex].Num() > 0)
-			//{
-			//	continue;
-			//}
-
 			const FGeometryNaniteData& NaniteData = GeometryNaniteData[TransformToGeometryIndex];
 
 			FPrimitiveInstance& Instance = Instances.Emplace_GetRef();
