Update FIgnoreCollisionManager to use queue + timestamp to marshall data for async physics.

diff --git a/Engine/Source/Runtime/Experimental/Chaos/Private/Chaos/Collision/CollisionConstraintFlags.cpp b/Engine/Source/Runtime/Experimental/Chaos/Private/Chaos/Collision/CollisionConstraintFlags.cpp
--- a/Engine/Source/Runtime/Experimental/Chaos/Private/Chaos/Collision/CollisionConstraintFlags.cpp
+++ b/Engine/Source/Runtime/Experimental/Chaos/Private/Chaos/Collision/CollisionConstraintFlags.cpp
@@ -47,41 +47,42 @@ namespace Chaos
 		}
 	}
 
-	void FIgnoreCollisionManager::FlipBufferPreSolve()
+	void FIgnoreCollisionManager::PopStorageData_Internal(int32 ExternalTimestamp)
 	{
-		BufferedData->FlipProducer();
-
-		// Merge
-		for (auto& Elem : BufferedData->GetConsumerBuffer()->PendingActivations)
+		FStorageData* StorageData;
+		while(StorageDataQueue.Peek(StorageData) && StorageData->ExternalTimestamp <= ExternalTimestamp)
 		{
-			if( PendingActivations.Contains(Elem.Key) )
+			for (auto& Elem : StorageData->PendingActivations)
 			{
-				// the case where the key already existed should be avoided
-				// but the implementation is here for completeness. 
-				for (auto& Val : Elem.Value)
+				if (PendingActivations.Contains(Elem.Key))
 				{
-					if (!PendingActivations[Elem.Key].Contains(Val))
+					// the case where the key already existed should be avoided
+					// but the implementation is here for completeness. 
+					for (auto& Val : Elem.Value)
 					{
-						PendingActivations[Elem.Key].Add(Val);
+						if (!PendingActivations[Elem.Key].Contains(Val))
+						{
+							PendingActivations[Elem.Key].Add(Val);
+						}
 					}
 				}
+				else
+				{
+					PendingActivations.Add(Elem.Key, Elem.Value);
+				}
 			}
-			else
-			{
-				PendingActivations.Add(Elem.Key, Elem.Value);
-			}
-		}
 
-		for (auto& Item : BufferedData->GetConsumerBuffer()->PendingDeactivations)
-		{
-			if (!PendingDeactivations.Contains(Item))
+			for (auto& Item : StorageData->PendingDeactivations)
 			{
-				PendingDeactivations.Add(Item);
+				if (!PendingDeactivations.Contains(Item))
+				{
+					PendingDeactivations.Add(Item);
+				}
 			}
-		}
 
-		BufferedData->GetConsumerBufferMutable()->PendingActivations.Empty();
-		BufferedData->GetConsumerBufferMutable()->PendingDeactivations.Empty();
+			StorageDataQueue.Pop();
+			ReleaseStorageData(StorageData);
+		}
 	}
 
 	void FIgnoreCollisionManager::ProcessPendingQueues()
diff --git a/Engine/Source/Runtime/Experimental/Chaos/Private/PBDRigidsSolver.cpp b/Engine/Source/Runtime/Experimental/Chaos/Private/PBDRigidsSolver.cpp
--- a/Engine/Source/Runtime/Experimental/Chaos/Private/PBDRigidsSolver.cpp
+++ b/Engine/Source/Runtime/Experimental/Chaos/Private/PBDRigidsSolver.cpp
@@ -739,7 +739,7 @@ namespace Chaos
 
 		if (HasActiveParticles())
 		{
-			GetEvolution()->GetBroadPhase().GetIgnoreCollisionManager().FlipBufferPreSolve();
+			GetEvolution()->GetBroadPhase().GetIgnoreCollisionManager().PopStorageData_Internal(GetEvolution()->LatestExternalTimestampConsumed_Internal);
 		}
 	}
 
@@ -874,6 +874,8 @@ namespace Chaos
 			}
 		});
 
+		GetEvolution()->GetBroadPhase().GetIgnoreCollisionManager().PushProducerStorageData_External(MarshallingManager.GetExternalTimestamp_External());
+
 		MarshallingManager.Step_External(DeltaTime);
 	}
 
diff --git a/Engine/Source/Runtime/Experimental/Chaos/Public/Chaos/Collision/CollisionConstraintFlags.h b/Engine/Source/Runtime/Experimental/Chaos/Public/Chaos/Collision/CollisionConstraintFlags.h
--- a/Engine/Source/Runtime/Experimental/Chaos/Public/Chaos/Collision/CollisionConstraintFlags.h
+++ b/Engine/Source/Runtime/Experimental/Chaos/Public/Chaos/Collision/CollisionConstraintFlags.h
@@ -28,11 +28,20 @@ namespace Chaos
 		{
 			FPendingMap PendingActivations;
 			FParticleArray PendingDeactivations;
+			int32 ExternalTimestamp = INDEX_NONE;
+
+			void Reset()
+			{
+				PendingActivations.Reset();
+				PendingDeactivations.Reset();
+				ExternalTimestamp = INDEX_NONE;
+			}
 		};
 
 		FIgnoreCollisionManager()
+			: StorageDataProducer(nullptr)
 		{
-			BufferedData = FMultiBufferFactory<FStorageData>::CreateBuffer(EMultiBufferMode::Double);
+			StorageDataProducer = GetNewStorageData();
 		}
 
 		bool ContainsHandle(FHandleID Body0);
@@ -45,11 +54,43 @@ namespace Chaos
 
 		void RemoveIgnoreCollisionsFor(FHandleID Body0, FHandleID Body1);
 
-		const FPendingMap& GetPendingActivationsForGameThread() const { return BufferedData->AccessProducerBuffer()->PendingActivations; }
-		FPendingMap& GetPendingActivationsForGameThread() { return BufferedData->AccessProducerBuffer()->PendingActivations; }
+		FPendingMap& GetPendingActivationsForGameThread(int32 ExternalTimestamp) 
+		{
+			if (StorageDataProducer->ExternalTimestamp == INDEX_NONE)
+			{
+				StorageDataProducer->ExternalTimestamp = ExternalTimestamp;
+			}
+			else
+			{
+				ensure(StorageDataProducer->ExternalTimestamp == ExternalTimestamp);
+			}
+
+			return StorageDataProducer->PendingActivations;
+		}
+
+		FParticleArray& GetPendingDeactivationsForGameThread(int32 ExternalTimestamp)
+		{
+			if (StorageDataProducer->ExternalTimestamp == INDEX_NONE)
+			{
+				StorageDataProducer->ExternalTimestamp = ExternalTimestamp;
+			}
+			else
+			{
+				ensure(StorageDataProducer->ExternalTimestamp == ExternalTimestamp);
+			}
+
+			return StorageDataProducer->PendingDeactivations;
+		}
 
-		const FParticleArray& GetPendingDeactivationsForGameThread() const { return BufferedData->AccessProducerBuffer()->PendingDeactivations; }
-		FParticleArray& GetPendingDeactivationsForGameThread() { return BufferedData->AccessProducerBuffer()->PendingDeactivations; }
+		void PushProducerStorageData_External(int32 ExternalTimestamp)
+		{
+			if (StorageDataProducer->ExternalTimestamp != INDEX_NONE)
+			{
+				ensure(ExternalTimestamp == StorageDataProducer->ExternalTimestamp);
+				StorageDataQueue.Enqueue(StorageDataProducer);
+				StorageDataProducer = GetNewStorageData();
+			}
+		}
 
 		/*
 		*
@@ -59,14 +100,39 @@ namespace Chaos
 		/*
 		*
 		*/
-		void FlipBufferPreSolve();
+		void PopStorageData_Internal(int32 ExternalTimestamp);
 
 	private:
+
+		FStorageData* GetNewStorageData()
+		{
+			FStorageData* StorageData;
+			if (StorageDataFreePool.Dequeue(StorageData))
+			{
+				return StorageData;
+			}
+
+			StorageDataBackingBuffer.Emplace(MakeUnique<FStorageData>());
+			return StorageDataBackingBuffer.Last().Get();
+		}
+
+		void ReleaseStorageData(FStorageData *InStorageData)
+		{
+			InStorageData->Reset();
+			StorageDataFreePool.Enqueue(InStorageData);
+		}
+
 		FActiveMap IgnoreCollisionsList;
 
 		FPendingMap PendingActivations;
 		FParticleArray PendingDeactivations;
-		TUniquePtr<Chaos::IBufferResource<FStorageData>> BufferedData;
+
+		// Producer storage data, pending changes written here until pushed into queue.
+		FStorageData* StorageDataProducer;
+
+		TQueue<FStorageData*, EQueueMode::Spsc> StorageDataQueue; // Queue of storage data being passed to physics thread
+		TQueue<FStorageData*,EQueueMode::Spsc> StorageDataFreePool;	//free pool of storage data
+		TArray<TUniquePtr<FStorageData>> StorageDataBackingBuffer;	// Holds unique ptrs for storage data allocation
 	};
 
 } // Chaos
diff --git a/Engine/Source/Runtime/PhysicsCore/Private/ChaosEngineInterface.cpp b/Engine/Source/Runtime/PhysicsCore/Private/ChaosEngineInterface.cpp
--- a/Engine/Source/Runtime/PhysicsCore/Private/ChaosEngineInterface.cpp
+++ b/Engine/Source/Runtime/PhysicsCore/Private/ChaosEngineInterface.cpp
@@ -278,7 +278,8 @@ void FChaosEngineInterface::AddDisabledCollisionsFor_AssumesLocked(const TMap<FP
 		TArray< FPhysicsActorHandle >& DisabledCollisions = Elem.Value;
 		Chaos::FPhysicsSolver* Solver = ActorReference->GetProxy()->GetSolver<Chaos::FPhysicsSolver>();
 		Chaos::FIgnoreCollisionManager& CollisionManager = Solver->GetEvolution()->GetBroadPhase().GetIgnoreCollisionManager();
-		Chaos::FIgnoreCollisionManager::FPendingMap& PendingMap = CollisionManager.GetPendingActivationsForGameThread();
+		int32 ExternalTimestamp = Solver->GetMarshallingManager().GetExternalTimestamp_External();
+		Chaos::FIgnoreCollisionManager::FPendingMap& PendingMap = CollisionManager.GetPendingActivationsForGameThread(ExternalTimestamp);
 		if (PendingMap.Contains(ActorReference))
 		{
 			PendingMap.Remove(ActorReference);
@@ -293,7 +294,8 @@ void FChaosEngineInterface::RemoveDisabledCollisionsFor_AssumesLocked(TArray< FP
 	{
 		Chaos::FPhysicsSolver* Solver = Handle->GetProxy()->GetSolver<Chaos::FPhysicsSolver>();
 		Chaos::FIgnoreCollisionManager& CollisionManager = Solver->GetEvolution()->GetBroadPhase().GetIgnoreCollisionManager();
-		Chaos::FIgnoreCollisionManager::FParticleArray& PendingMap = CollisionManager.GetPendingDeactivationsForGameThread();
+		int32 ExternalTimestamp = Solver->GetMarshallingManager().GetExternalTimestamp_External();
+		Chaos::FIgnoreCollisionManager::FParticleArray& PendingMap = CollisionManager.GetPendingDeactivationsForGameThread(ExternalTimestamp);
 		if (!PendingMap.Contains(Handle))
 		{
 			PendingMap.Add(Handle);
