Fix missing include

diff --git a/Engine/Source/Runtime/Experimental/Chaos/Public/Chaos/Collision/CollisionConstraintFlags.h b/Engine/Source/Runtime/Experimental/Chaos/Public/Chaos/Collision/CollisionConstraintFlags.h
--- a/Engine/Source/Runtime/Experimental/Chaos/Public/Chaos/Collision/CollisionConstraintFlags.h
+++ b/Engine/Source/Runtime/Experimental/Chaos/Public/Chaos/Collision/CollisionConstraintFlags.h
@@ -6,6 +6,7 @@
 #include "Chaos/GeometryParticlesfwd.h"
 #include "Chaos/ParticleHandleFwd.h"
 #include "Chaos/PBDRigidsEvolutionFwd.h"
+#include "Containers/Queue.h"
 
 namespace Chaos
 {
