GeometryCollection -- Added Apply Changes button to trigger Simulation Data rebuild.

diff --git a/Engine/Plugins/Experimental/GeometryCollectionPlugin/Source/GeometryCollectionEditor/Private/GeometryCollection/DetailCustomizations/GeometryCollectionCustomization.cpp b/Engine/Plugins/Experimental/GeometryCollectionPlugin/Source/GeometryCollectionEditor/Private/GeometryCollection/DetailCustomizations/GeometryCollectionCustomization.cpp
new file mode 100644
--- /dev/null
+++ b/Engine/Plugins/Experimental/GeometryCollectionPlugin/Source/GeometryCollectionEditor/Private/GeometryCollection/DetailCustomizations/GeometryCollectionCustomization.cpp
@@ -0,0 +1,147 @@
+// Copyright Epic Games, Inc. All Rights Reserved.
+
+#include "GeometryCollectionCustomization.h"
+
+#include "GeometryCollection/GeometryCollectionObject.h"
+#include "DetailCategoryBuilder.h"
+#include "DetailLayoutBuilder.h"
+#include "IDetailGroup.h"
+#include "DetailWidgetRow.h"
+#include "Async/ParallelFor.h"
+#include "Widgets/Input/SButton.h"
+#include "Widgets/Input/SCheckBox.h"
+#include "Widgets/Layout/SBorder.h"
+#include "Widgets/Layout/SBox.h"
+#include "Widgets/Text/STextBlock.h"
+
+
+#define LOCTEXT_NAMESPACE "GeometryCollectionCustomization"
+
+TSharedRef<IDetailCustomization> FGeometryCollectionCustomization::MakeInstance()
+{
+	return MakeShareable(new FGeometryCollectionCustomization());
+}
+
+void FGeometryCollectionCustomization::CustomizeDetails(IDetailLayoutBuilder& DetailBuilder)
+{
+
+	// Store off the objects that we are editing for analysis in later function calls.
+	DetailBuilder.GetObjectsBeingCustomized(ObjectsCustomized);
+
+	IDetailCategoryBuilder& CollisionsCategory = DetailBuilder.EditCategory(TEXT("Collisions"));
+	
+	IDetailGroup& ManualUpdateGroup = CollisionsCategory.AddGroup("ManualUpdate", LOCTEXT("ManualUpdateLabel", "ManualUpdate"));
+	ManualUpdateGroup.HeaderRow()
+		[
+			SNew(SBorder)
+			.Padding(FMargin(3.0f))
+			[
+				SNew(SVerticalBox)
+				+SVerticalBox::Slot()
+				.Padding(FMargin(6.f, 3.f))
+				.AutoHeight()
+				[
+					SNew(SHorizontalBox)
+					+SHorizontalBox::Slot()
+					.VAlign(VAlign_Center)
+					.Padding(0.0f, 0.0f, 10.0f, 0.0f)
+					[
+						SNew(STextBlock)
+						.Text(LOCTEXT("UpdateChangesManually", "Update Changes Manually"))
+						.Font(IDetailLayoutBuilder::GetDetailFont())
+						.Justification(ETextJustify::Right)
+					]
+					+SHorizontalBox::Slot()
+					[
+						SNew(SCheckBox)
+						.ToolTipText(LOCTEXT("UpdateChangesManually_Tooltip", "Defer collision regeneration until manually triggered."))
+						.OnCheckStateChanged(this, &FGeometryCollectionCustomization::OnCheckStateChanged)
+					]
+				]
+				+SVerticalBox::Slot()
+				.Padding(FMargin(6.f, 3.f))
+				.AutoHeight()
+				[
+					SNew(SBox)
+					.Visibility(this, &FGeometryCollectionCustomization::ApplyButtonVisibility)
+					.WidthOverride(100.f)
+					.HeightOverride(25.f)
+					[
+						SNew(SButton)
+						.HAlign(HAlign_Center)
+						.VAlign(VAlign_Center)
+						.OnClicked(this, &FGeometryCollectionCustomization::OnApplyChanges)
+						.IsEnabled(this, &FGeometryCollectionCustomization::IsApplyNeeded)
+						.ToolTipText(LOCTEXT("ApplyChanges_Tooltip", "Regenerate collisions using new parameters."))
+						[
+							SNew(STextBlock)
+							.Text(LOCTEXT("ApplyChanges", "Apply Changes"))
+							.Font(IDetailLayoutBuilder::GetDetailFont())
+							.Justification(ETextJustify::Center)
+						]
+					]
+				]
+			]		
+		];
+		
+}
+
+bool FGeometryCollectionCustomization::IsApplyNeeded() const
+{
+	for (int32 i = 0; i < ObjectsCustomized.Num(); i++)
+	{
+		if (UGeometryCollection* GeometryCollection = Cast<UGeometryCollection>(ObjectsCustomized[i].Get()))
+		{
+			if (GeometryCollection->IsSimulationDataDirty())
+			{
+				return true;
+			}
+		}		
+	}
+	return false;
+}
+
+FReply FGeometryCollectionCustomization::OnApplyChanges()
+{
+	ParallelFor(ObjectsCustomized.Num(), [&ObjectsCustomized = ObjectsCustomized](int32 Idx)
+	{
+		if (UGeometryCollection* GeometryCollection = Cast<UGeometryCollection>(ObjectsCustomized[Idx].Get()))
+		{
+			if (GeometryCollection->IsSimulationDataDirty())
+			{
+				GeometryCollection->CreateSimulationData();
+			}
+		}
+	});
+
+	return FReply::Handled();
+}
+
+void FGeometryCollectionCustomization::OnCheckStateChanged(ECheckBoxState InNewState)
+{
+	bManualApplyActivated = (InNewState == ECheckBoxState::Checked);
+
+	for (TWeakObjectPtr<UObject> ObjectCustomized : ObjectsCustomized)
+	{
+		if (UGeometryCollection* GeometryCollection = Cast<UGeometryCollection>(ObjectCustomized.Get()))
+		{
+			GeometryCollection->bManualDataCreate = bManualApplyActivated;
+		}
+	}
+}
+
+EVisibility FGeometryCollectionCustomization::ApplyButtonVisibility() const
+{
+	if (bManualApplyActivated)
+	{
+		return EVisibility::Visible;
+	}
+	else
+	{
+		return EVisibility::Collapsed;
+	}
+}
+
+
+
+#undef LOCTEXT_NAMESPACE
\ No newline at end of file
diff --git a/Engine/Plugins/Experimental/GeometryCollectionPlugin/Source/GeometryCollectionEditor/Private/GeometryCollection/DetailCustomizations/GeometryCollectionCustomization.h b/Engine/Plugins/Experimental/GeometryCollectionPlugin/Source/GeometryCollectionEditor/Private/GeometryCollection/DetailCustomizations/GeometryCollectionCustomization.h
new file mode 100644
--- /dev/null
+++ b/Engine/Plugins/Experimental/GeometryCollectionPlugin/Source/GeometryCollectionEditor/Private/GeometryCollection/DetailCustomizations/GeometryCollectionCustomization.h
@@ -0,0 +1,26 @@
+// Copyright Epic Games, Inc. All Rights Reserved.
+
+#pragma once
+
+#include "IDetailCustomization.h"
+#include "Input/Reply.h"
+#include "Layout/Visibility.h"
+
+enum class ECheckBoxState : uint8;
+
+class FGeometryCollectionCustomization : public IDetailCustomization
+{
+public:
+	static TSharedRef<IDetailCustomization> MakeInstance();
+	virtual void CustomizeDetails(IDetailLayoutBuilder& DetailBuilder);
+
+private:
+	bool IsApplyNeeded() const;
+	FReply OnApplyChanges();
+	void OnCheckStateChanged(ECheckBoxState InNewState);
+	EVisibility ApplyButtonVisibility() const;
+
+	bool bManualApplyActivated = false;
+
+	TArray<TWeakObjectPtr<UObject>> ObjectsCustomized;
+};
diff --git a/Engine/Plugins/Experimental/GeometryCollectionPlugin/Source/GeometryCollectionEditor/Private/GeometryCollection/GeometryCollectionEditorPlugin.cpp b/Engine/Plugins/Experimental/GeometryCollectionPlugin/Source/GeometryCollectionEditor/Private/GeometryCollection/GeometryCollectionEditorPlugin.cpp
--- a/Engine/Plugins/Experimental/GeometryCollectionPlugin/Source/GeometryCollectionEditor/Private/GeometryCollection/GeometryCollectionEditorPlugin.cpp
+++ b/Engine/Plugins/Experimental/GeometryCollectionPlugin/Source/GeometryCollectionEditor/Private/GeometryCollection/GeometryCollectionEditorPlugin.cpp
@@ -21,6 +21,7 @@
 #include "GeometryCollection/DetailCustomizations/GeomComponentCacheCustomization.h"
 #include "GeometryCollection/DetailCustomizations/SelectedRigidBodyCustomization.h"
 #include "GeometryCollection/DetailCustomizations/WarningMessageCustomization.h"
+#include "GeometryCollection/DetailCustomizations/GeometryCollectionCustomization.h"
 #include "Styling/SlateStyle.h"
 #include "Styling/SlateStyleRegistry.h"
 #include "Styling/CoreStyle.h"
@@ -197,6 +198,9 @@ void IGeometryCollectionEditorPlugin::StartupModule()
 		PropertyModule->RegisterCustomPropertyTypeLayout("GeomComponentCacheParameters", FOnGetPropertyTypeCustomizationInstance::CreateStatic(&FGeomComponentCacheParametersCustomization::MakeInstance));
 		PropertyModule->RegisterCustomPropertyTypeLayout("GeometryCollectionDebugDrawActorSelectedRigidBody", FOnGetPropertyTypeCustomizationInstance::CreateStatic(&FSelectedRigidBodyCustomization::MakeInstance));
 		PropertyModule->RegisterCustomPropertyTypeLayout("GeometryCollectionDebugDrawWarningMessage", FOnGetPropertyTypeCustomizationInstance::CreateStatic(&FWarningMessageCustomization::MakeInstance));
+
+		const FName GeometryCollectionName = UGeometryCollection::StaticClass()->GetFName();
+		PropertyModule->RegisterCustomClassLayout(GeometryCollectionName, FOnGetDetailCustomizationInstance::CreateStatic(&FGeometryCollectionCustomization::MakeInstance));
 	}
 
 	// Register rigid body selection editor mode
@@ -252,6 +256,9 @@ void IGeometryCollectionEditorPlugin::ShutdownModule()
 			PropertyModule->UnregisterCustomPropertyTypeLayout("GeomCollectionCacheParameters");
 			PropertyModule->UnregisterCustomPropertyTypeLayout("GeometryCollectionDebugDrawActorSelectedRigidBody");
 			PropertyModule->UnregisterCustomPropertyTypeLayout("GeometryCollectionDebugDrawWarningMessage");
+
+			const FName GeometryCollectionName = UGeometryCollection::StaticClass()->GetFName();
+			PropertyModule->UnregisterCustomClassLayout(GeometryCollectionName);
 		}
 
 		// Unregister rigid body selection editor mode
diff --git a/Engine/Source/Runtime/Experimental/GeometryCollectionEngine/Private/GeometryCollection/GeometryCollectionObject.cpp b/Engine/Source/Runtime/Experimental/GeometryCollectionEngine/Private/GeometryCollection/GeometryCollectionObject.cpp
--- a/Engine/Source/Runtime/Experimental/GeometryCollectionEngine/Private/GeometryCollection/GeometryCollectionObject.cpp
+++ b/Engine/Source/Runtime/Experimental/GeometryCollectionEngine/Private/GeometryCollection/GeometryCollectionObject.cpp
@@ -61,9 +61,15 @@ UGeometryCollection::UGeometryCollection(const FObjectInitializer& ObjectInitial
 	, MaximumCollisionParticles(60)
 	, EnableRemovePiecesOnFracture(false)
 	, GeometryCollection(new FGeometryCollection())
+#if WITH_EDITOR
+	, bManualDataCreate(false)
+#endif
 {
 	PersistentGuid = FGuid::NewGuid();
 	InvalidateCollection();
+#if WITH_EDITOR
+	SimulationDataGuid = StateGuid;
+#endif
 }
 
 FGeometryCollectionSizeSpecificData::FGeometryCollectionSizeSpecificData()
@@ -462,6 +468,7 @@ void UGeometryCollection::CreateSimulationDataImp(bool bCopyFromDDC)
 void UGeometryCollection::CreateSimulationData()
 {
 	CreateSimulationDataImp(/*bCopyFromDDC=*/false);
+	SimulationDataGuid = StateGuid;
 }
 
 TUniquePtr<FGeometryCollectionNaniteData> UGeometryCollection::CreateNaniteData(FGeometryCollection* Collection)
@@ -606,6 +613,13 @@ void UGeometryCollection::InvalidateCollection()
 	StateGuid = FGuid::NewGuid();
 }
 
+#if WITH_EDITOR
+bool UGeometryCollection::IsSimulationDataDirty() const
+{
+	return StateGuid != SimulationDataGuid;
+}
+#endif
+
 FGuid UGeometryCollection::GetIdGuid() const
 {
 	return PersistentGuid;
@@ -630,7 +644,11 @@ void UGeometryCollection::PostEditChangeProperty(struct FPropertyChangedEvent& P
 		else if (PropertyChangedEvent.Property->GetFName() != GET_MEMBER_NAME_CHECKED(UGeometryCollection, Materials))
 		{
 			InvalidateCollection();
-			CreateSimulationData();
+			
+			if (!bManualDataCreate)
+			{
+				CreateSimulationData();
+			}
 		}
 	}
 }
diff --git a/Engine/Source/Runtime/Experimental/GeometryCollectionEngine/Public/GeometryCollection/GeometryCollectionObject.h b/Engine/Source/Runtime/Experimental/GeometryCollectionEngine/Public/GeometryCollection/GeometryCollectionObject.h
--- a/Engine/Source/Runtime/Experimental/GeometryCollectionEngine/Public/GeometryCollection/GeometryCollectionObject.h
+++ b/Engine/Source/Runtime/Experimental/GeometryCollectionEngine/Public/GeometryCollection/GeometryCollectionObject.h
@@ -41,61 +41,61 @@ struct GEOMETRYCOLLECTIONENGINE_API FGeometryCollectionSizeSpecificData
 	FGeometryCollectionSizeSpecificData();
 
 	/** The max size these settings apply to*/
-	UPROPERTY(EditAnywhere, Category = "ChaosPhysics|Collisions")
+	UPROPERTY(EditAnywhere, Category = "Collisions")
 	float MaxSize;
 
 	/*
 	*  CollisionType defines how to initialize the rigid collision structures.
 	*/
-	UPROPERTY(EditAnywhere, Category = "ChaosPhysics|Collisions")
+	UPROPERTY(EditAnywhere, Category = "Collisions")
 	ECollisionTypeEnum CollisionType;
 
 	/*
 	*  CollisionType defines how to initialize the rigid collision structures.
 	*/
-	UPROPERTY(EditAnywhere, Category = "ChaosPhysics|Collisions")
+	UPROPERTY(EditAnywhere, Category = "Collisions")
 	EImplicitTypeEnum ImplicitType;
 
 	/*
 	*  Resolution on the smallest axes for the level set. (def: 5)
 	*/
-	UPROPERTY(EditAnywhere, Category = "ChaosPhysics|Collisions")
+	UPROPERTY(EditAnywhere, Category = "Collisions")
 	int32 MinLevelSetResolution;
 
 	/*
 	*  Resolution on the smallest axes for the level set. (def: 10)
 	*/
-	UPROPERTY(EditAnywhere, Category = "ChaosPhysics|Collisions")
+	UPROPERTY(EditAnywhere, Category = "Collisions")
 	int32 MaxLevelSetResolution;
 
 	/*
 	*  Resolution on the smallest axes for the level set. (def: 5)
 	*/
-	UPROPERTY(EditAnywhere, Category = "ChaosPhysics|Collisions")
+	UPROPERTY(EditAnywhere, Category = "Collisions")
 	int32 MinClusterLevelSetResolution;
 
 	/*
 	*  Resolution on the smallest axes for the level set. (def: 10)
 	*/
-	UPROPERTY(EditAnywhere, Category = "ChaosPhysics|Collisions")
+	UPROPERTY(EditAnywhere, Category = "Collisions")
 	int32 MaxClusterLevelSetResolution;
 
 	/*
 	*  Resolution on the smallest axes for the level set. (def: 10)
 	*/
-	UPROPERTY(EditAnywhere, Category = "ChaosPhysics|Collisions")
+	UPROPERTY(EditAnywhere, Category = "Collisions")
 	int32 CollisionObjectReductionPercentage;
 
 	/**
 	 * Number of particles on the triangulated surface to use for collisions.
 	 */
-	UPROPERTY(EditAnywhere, Category = "ChaosPhysics|Collisions")
+	UPROPERTY(EditAnywhere, Category = "Collisions")
 	float CollisionParticlesFraction;
 
 	/**
 	 * Max number of particles.
 	 */
-	UPROPERTY(EditAnywhere, Category = "ChaosPhysics|Collisions")
+	UPROPERTY(EditAnywhere, Category = "Collisions")
 	int32 MaximumCollisionParticles;
 
 };
@@ -187,6 +187,14 @@ public:
 	void InvalidateCollection();
 
 #if WITH_EDITOR
+	/** Check to see if Simulation Data requires regeneration */
+	bool IsSimulationDataDirty() const;
+#endif
+
+#if WITH_EDITOR
+	/** If this flag is set, we only regenerate simulation data when requested via CreateSimulationData() */
+	bool bManualDataCreate;
+	
 	/** Create the simulation data that can be shared among all instances (mass, volume, etc...)*/
 	void CreateSimulationData();
 
@@ -223,88 +231,88 @@ public:
 	/*
 	*  CollisionType defines how to initialize the rigid collision structures.
 	*/
-	UPROPERTY(EditAnywhere, BlueprintReadOnly, Category = "ChaosPhysics|Collisions")
+	UPROPERTY(EditAnywhere, BlueprintReadOnly, Category = "Collisions")
 	ECollisionTypeEnum CollisionType;
 
 	/*
 	*  CollisionType defines how to initialize the rigid collision structures.
 	*/
-	UPROPERTY(EditAnywhere, BlueprintReadOnly, Category = "ChaosPhysics|Collisions")
+	UPROPERTY(EditAnywhere, BlueprintReadOnly, Category = "Collisions")
 	EImplicitTypeEnum ImplicitType;
 
 	/*
 	*  Resolution on the smallest axes for the level set. (def: 5)
 	*/
-	UPROPERTY(EditAnywhere, BlueprintReadOnly, Category = "ChaosPhysics|Collisions")
+	UPROPERTY(EditAnywhere, BlueprintReadOnly, Category = "Collisions")
 	int32 MinLevelSetResolution;
 
 	/*
 	*  Resolution on the smallest axes for the level set. (def: 10)
 	*/
-	UPROPERTY(EditAnywhere, BlueprintReadOnly, Category = "ChaosPhysics|Collisions")
+	UPROPERTY(EditAnywhere, BlueprintReadOnly, Category = "Collisions")
 	int32 MaxLevelSetResolution;
 
 	/*
 	*  Resolution on the smallest axes for the level set. (def: 5)
 	*/
-	UPROPERTY(EditAnywhere, BlueprintReadOnly, Category = "ChaosPhysics|Collisions")
+	UPROPERTY(EditAnywhere, BlueprintReadOnly, Category = "Collisions")
 	int32 MinClusterLevelSetResolution;
 
 	/*
 	*  Resolution on the smallest axes for the level set. (def: 10)
 	*/
-	UPROPERTY(EditAnywhere, BlueprintReadOnly, Category = "ChaosPhysics|Collisions")
+	UPROPERTY(EditAnywhere, BlueprintReadOnly, Category = "Collisions")
 	int32 MaxClusterLevelSetResolution;
 
 	/*
 	*  Resolution on the smallest axes for the level set. (def: 10)
 	*/
-	UPROPERTY(EditAnywhere, BlueprintReadOnly, Category = "ChaosPhysics|Collisions")
+	UPROPERTY(EditAnywhere, BlueprintReadOnly, Category = "Collisions")
 	float CollisionObjectReductionPercentage;
 
 	/**
 	* Mass As Density, units are in kg/m^3
 	*/
-	UPROPERTY(EditAnywhere, BlueprintReadOnly, Category = "ChaosPhysics|Collisions")
+	UPROPERTY(EditAnywhere, BlueprintReadOnly, Category = "Collisions")
 	bool bMassAsDensity;
 
 	/**
 	* Total Mass of Collection. If density, units are in kg/m^3
 	*/
-	UPROPERTY(EditAnywhere, BlueprintReadOnly, Category = "ChaosPhysics|Collisions")
+	UPROPERTY(EditAnywhere, BlueprintReadOnly, Category = "Collisions")
 	float Mass;
 
 	/**
 	* Smallest allowable mass (def:0.1)
 	*/
-	UPROPERTY(EditAnywhere, BlueprintReadOnly, Category = "ChaosPhysics|Collisions")
+	UPROPERTY(EditAnywhere, BlueprintReadOnly, Category = "Collisions")
 	float MinimumMassClamp;
 
 	/**
 	 * Number of particles on the triangulated surface to use for collisions.
 	 */
-	UPROPERTY(EditAnywhere, BlueprintReadOnly, Category = "ChaosPhysics|Collisions")
+	UPROPERTY(EditAnywhere, BlueprintReadOnly, Category = "Collisions")
 	float CollisionParticlesFraction;
 
 	/**
 	 * Max number of particles.
 	 */
-	UPROPERTY(EditAnywhere, BlueprintReadOnly, Category = "ChaosPhysics|Collisions")
+	UPROPERTY(EditAnywhere, BlueprintReadOnly, Category = "Collisions")
 	int32 MaximumCollisionParticles;
 
-	UPROPERTY(EditAnywhere, Category = "ChaosPhysics|Collisions")
+	UPROPERTY(EditAnywhere, Category = "Collisions")
 	TArray<FGeometryCollectionSizeSpecificData> SizeSpecificData;
 
 	/**
 	* Enable remove pieces on fracture
 	*/
-	UPROPERTY(EditAnywhere, BlueprintReadOnly, Category = "ChaosPhysics|Fracture")
+	UPROPERTY(EditAnywhere, BlueprintReadOnly, Category = "Fracture")
 	bool EnableRemovePiecesOnFracture;
 
 	/**
 	* Materials relating to remove on fracture
 	*/
-	UPROPERTY(EditAnywhere, Category = "ChaosPhysics|Fracture")
+	UPROPERTY(EditAnywhere, Category = "Fracture")
 	TArray<UMaterialInterface*> RemoveOnFractureMaterials;
 
 	FORCEINLINE const int32 GetBoneSelectedMaterialIndex() const { return BoneSelectedMaterialIndex; }
@@ -339,6 +347,9 @@ private:
 #if WITH_EDITOR
 	//Used to determine whether we need to cook content
 	FGuid LastBuiltGuid;
+
+	//Used to determine whether we need to regenerate simulation data
+	FGuid SimulationDataGuid;
 #endif
 
 	// #todo(dmp): rename to be consistent BoneSelectedMaterialID?
