Chaos : Ignore Collisions

- Flag removed particles for clean up in the collision manager.

diff --git a/Engine/Source/Runtime/Experimental/Chaos/Private/PBDRigidsSolver.cpp b/Engine/Source/Runtime/Experimental/Chaos/Private/PBDRigidsSolver.cpp
--- a/Engine/Source/Runtime/Experimental/Chaos/Private/PBDRigidsSolver.cpp
+++ b/Engine/Source/Runtime/Experimental/Chaos/Private/PBDRigidsSolver.cpp
@@ -472,6 +472,17 @@ namespace Chaos
 			GeometryParticlePhysicsProxies.RemoveSingleSwap((FGeometryParticlePhysicsProxy*)InProxy);
 		}
 
+
+		Chaos::FIgnoreCollisionManager& CollisionManager = GetEvolution()->GetBroadPhase().GetIgnoreCollisionManager();
+		{
+			int32 ExternalTimestamp = GetMarshallingManager().GetExternalTimestamp_External();
+			Chaos::FIgnoreCollisionManager::FParticleArray& PendingMap = CollisionManager.GetPendingDeactivationsForGameThread(ExternalTimestamp);
+			if (!PendingMap.Contains(GTParticle))
+			{
+				PendingMap.Add(GTParticle);
+			}
+		}
+
 		// Enqueue a command to remove the particle and delete the proxy
 		EnqueueCommandImmediate([InProxy, InParticleType, this]()
 		{
