Fixes to geometry collection clusters having misplaced children during mass updates and subsequently having a CoM way outside of their geometry bounds.
#jria none

diff --git a/Engine/Source/Runtime/Experimental/Chaos/Private/PhysicsProxy/GeometryCollectionPhysicsProxy.cpp b/Engine/Source/Runtime/Experimental/Chaos/Private/PhysicsProxy/GeometryCollectionPhysicsProxy.cpp
--- a/Engine/Source/Runtime/Experimental/Chaos/Private/PhysicsProxy/GeometryCollectionPhysicsProxy.cpp
+++ b/Engine/Source/Runtime/Experimental/Chaos/Private/PhysicsProxy/GeometryCollectionPhysicsProxy.cpp
@@ -1858,7 +1858,6 @@ void FGeometryCollectionPhysicsProxy::InitializeSharedCollisionStructures(
 				{
 					CalculateVolumeAndCenterOfMass(MassSpaceParticles, TriMesh->GetElements(), MassProperties.Volume, MassProperties.CenterOfMass);
 					InertiaComputationNeeded[GeometryIndex] = true;
-
 					if(MassProperties.Volume == 0)
 					{
 						FVector Extents = BoundingBox[GeometryIndex].GetExtent();
@@ -1868,6 +1867,7 @@ void FGeometryCollectionPhysicsProxy::InitializeSharedCollisionStructures(
 						float ExtentsXY = Extents.X * Extents.X + Extents.Y * Extents.Y;
 						MassProperties.InertiaTensor = PMatrix<float, 3, 3>(ExtentsYZ / 12., ExtentsXZ / 12., ExtentsXY / 12.);
 						MassProperties.CenterOfMass = BoundingBox[GeometryIndex].GetCenter();
+						CollectionMassToLocal[TransformGroupIndex] = FTransform(FQuat::Identity, MassProperties.CenterOfMass);
 						InertiaComputationNeeded[GeometryIndex] = false;
 					}
 
@@ -2148,7 +2148,7 @@ void FGeometryCollectionPhysicsProxy::InitializeSharedCollisionStructures(
 
  			if (CollectionSimulatableParticles[TransformGroupIndex])
  			{
-				FTransform GeometryWorldTransform = CollectionSpaceTransforms[TransformGroupIndex] * CollectionMassToLocal[TransformGroupIndex];
+				FTransform GeometryWorldTransform = CollectionMassToLocal[TransformGroupIndex] * CollectionSpaceTransforms[TransformGroupIndex];
 
  				PopulateSimulatedParticle(
  					Handles[TransformGroupIndex].Get(),
@@ -2238,7 +2238,7 @@ void FGeometryCollectionPhysicsProxy::InitializeSharedCollisionStructures(
 							// To move a particle from mass-space in the child to mass-space in the cluster parent, calculate
 							// the relative transform between the mass-space origin for both the parent and child before
 							// transforming the mass space particles into the parent mass-space.
-							const FTransform ChildMassToClusterMass = (CollectionSpaceTransforms[ChildTransformIdx] * CollectionMassToLocal[ChildTransformIdx]).GetRelativeTransform(CollectionSpaceTransforms[ClusterTransformIdx] * CollectionMassToLocal[ClusterTransformIdx]);
+							const FTransform ChildMassToClusterMass = (CollectionMassToLocal[ChildTransformIdx] * CollectionSpaceTransforms[ChildTransformIdx]).GetRelativeTransform(CollectionMassToLocal[ClusterTransformIdx] * CollectionSpaceTransforms[ClusterTransformIdx]);
 
 							ChildMesh->GetVertexSet(VertsAdded);
 							for (const int32 VertIdx : VertsAdded)
diff --git a/Engine/Source/Runtime/Experimental/GeometryCollectionEngine/Private/GeometryCollection/DerivedDataGeometryCollectionCooker.cpp b/Engine/Source/Runtime/Experimental/GeometryCollectionEngine/Private/GeometryCollection/DerivedDataGeometryCollectionCooker.cpp
--- a/Engine/Source/Runtime/Experimental/GeometryCollectionEngine/Private/GeometryCollection/DerivedDataGeometryCollectionCooker.cpp
+++ b/Engine/Source/Runtime/Experimental/GeometryCollectionEngine/Private/GeometryCollection/DerivedDataGeometryCollectionCooker.cpp
@@ -72,7 +72,7 @@ bool FDerivedDataGeometryCollectionCooker::Build(TArray<uint8>& OutData)
 
 const TCHAR* FDerivedDataGeometryCollectionCooker::GetVersionString() const
 {
-	const TCHAR* VersionString = TEXT("A8A2C0FB45084FCB922FEC1139E11341");
+	const TCHAR* VersionString = TEXT("3448151914544729A210FBB9B0524868");
 
 	static FString CachedNaniteVersionString;
 	if (CachedNaniteVersionString.IsEmpty())
