Fall back to default spring mass values if spring utility fails to compute. Currently a trailer with all wheels located at one end causes calculation to fail.

diff --git a/Engine/Plugins/Experimental/ChaosVehiclesPlugin/Source/ChaosVehicles/Private/ChaosWheeledVehicleMovementComponent.cpp b/Engine/Plugins/Experimental/ChaosVehiclesPlugin/Source/ChaosVehicles/Private/ChaosWheeledVehicleMovementComponent.cpp
--- a/Engine/Plugins/Experimental/ChaosVehiclesPlugin/Source/ChaosVehicles/Private/ChaosWheeledVehicleMovementComponent.cpp
+++ b/Engine/Plugins/Experimental/ChaosVehiclesPlugin/Source/ChaosVehicles/Private/ChaosWheeledVehicleMovementComponent.cpp
@@ -474,7 +474,14 @@ void UChaosWheeledVehicleMovementComponent::SetupSuspension()
 
 	// Calculate the mass that will rest on each of the springs
 	TArray<float> OutSprungMasses;
-	FSuspensionUtility::ComputeSprungMasses(LocalSpringPositions, TotalMass, OutSprungMasses);
+	if (!FSuspensionUtility::ComputeSprungMasses(LocalSpringPositions, TotalMass, OutSprungMasses))
+	{
+		// if the sprung mass calc fails fall back to something that will still simulate
+		for (int Index = 0; Index < OutSprungMasses.Num(); Index++)
+		{
+			OutSprungMasses[Index] = TotalMass / OutSprungMasses.Num();
+		}
+	}
 
 	// Calculate spring damping values we will use for physics simulation from the normalized damping ratio
 	for (int SpringIdx = 0; SpringIdx < PVehicle->Suspension.Num(); SpringIdx++)
diff --git a/Engine/Source/Runtime/Experimental/ChaosVehicles/ChaosVehiclesCore/Private/SuspensionUtility.cpp b/Engine/Source/Runtime/Experimental/ChaosVehicles/ChaosVehiclesCore/Private/SuspensionUtility.cpp
--- a/Engine/Source/Runtime/Experimental/ChaosVehicles/ChaosVehiclesCore/Private/SuspensionUtility.cpp
+++ b/Engine/Source/Runtime/Experimental/ChaosVehicles/ChaosVehiclesCore/Private/SuspensionUtility.cpp
@@ -2,6 +2,8 @@
 
 #include "SuspensionUtility.h"
 
+DEFINE_LOG_CATEGORY(LogVehicleUtility);
+
 bool FSuspensionUtility::ComputeSprungMasses(const TArray<FVector>& MassSpringPositions, const float TotalMass, TArray<float>& OutSprungMasses)
 {
 	/*
@@ -220,8 +222,11 @@ bool FSuspensionUtility::ComputeSprungMasses(const TArray<FVector>& MassSpringPo
 		const float Y = MassSpringPositions[Index].Y;
 		const float LLambda = (X * Lambda0) + (Y * Lambda1) + Lambda2;
 		OutSprungMasses[Index] = AverageMass - (0.5f * LLambda);
-		if (!ensureMsgf(OutSprungMasses[Index] >= 0.f, TEXT("Spring configuration is invalid! Please make sure the center of mass is located inside the area covered by the springs.")))
+		if (OutSprungMasses[Index] < 0.f)
 		{
+			// #TODO: it is technically possible for say a trailer to have all of its wheels at the rear and none at the front 
+			// hence violating this condition, need to think about handling this.
+			UE_LOG(LogVehicleUtility, Warning, TEXT("Spring configuration is invalid! Please make sure the center of mass is located inside the area covered by the springs."));
 			return false;
 		}
 	}
diff --git a/Engine/Source/Runtime/Experimental/ChaosVehicles/ChaosVehiclesCore/Public/SuspensionUtility.h b/Engine/Source/Runtime/Experimental/ChaosVehicles/ChaosVehiclesCore/Public/SuspensionUtility.h
--- a/Engine/Source/Runtime/Experimental/ChaosVehicles/ChaosVehiclesCore/Public/SuspensionUtility.h
+++ b/Engine/Source/Runtime/Experimental/ChaosVehicles/ChaosVehiclesCore/Public/SuspensionUtility.h
@@ -7,6 +7,8 @@
 #include "Templates/Function.h"
 #include "Containers/Array.h"
 
+DECLARE_LOG_CATEGORY_EXTERN(LogVehicleUtility, Log, All);
+
 
 struct FSuspensionUtility
 {
