MeshModelingTools: Fix tool icon not being unhighlighted on tool accept after CL 14809785. Also removing IsToolBuilderActive function since we plan to not use that anymore.

#rnx

diff --git a/Engine/Source/Editor/UnrealEd/Private/Tools/UEdMode.cpp b/Engine/Source/Editor/UnrealEd/Private/Tools/UEdMode.cpp
--- a/Engine/Source/Editor/UnrealEd/Private/Tools/UEdMode.cpp
+++ b/Engine/Source/Editor/UnrealEd/Private/Tools/UEdMode.cpp
@@ -266,7 +266,7 @@ void UEdMode::RegisterTool(TSharedPtr<FUICommandInfo> UICommand, FString ToolIde
 	CommandList->MapAction(UICommand,
 		FExecuteAction::CreateUObject(ToolsContext.Get(), &UEdModeInteractiveToolsContext::StartTool, ToolIdentifier),
 		FCanExecuteAction::CreateUObject(ToRawPtr(ToolsContext->ToolManager), &UInteractiveToolManager::CanActivateTool, EToolSide::Mouse, ToolIdentifier),
-		FIsActionChecked::CreateUObject(ToolsContext.Get(), &UEdModeInteractiveToolsContext::IsToolBuilderActive, EToolSide::Mouse, Builder),
+		FIsActionChecked::CreateUObject(ToolsContext.Get(), &UEdModeInteractiveToolsContext::IsToolActive, EToolSide::Mouse, ToolIdentifier),
 		EUIActionRepeatMode::RepeatDisabled);
 
 	RegisteredTools.Emplace(UICommand, ToolIdentifier);
diff --git a/Engine/Source/Runtime/Experimental/InteractiveToolsFramework/Private/InteractiveToolsContext.cpp b/Engine/Source/Runtime/Experimental/InteractiveToolsFramework/Private/InteractiveToolsContext.cpp
--- a/Engine/Source/Runtime/Experimental/InteractiveToolsFramework/Private/InteractiveToolsContext.cpp
+++ b/Engine/Source/Runtime/Experimental/InteractiveToolsFramework/Private/InteractiveToolsContext.cpp
@@ -121,9 +121,11 @@ void UInteractiveToolsContext::EndTool(EToolSide WhichSide, EToolShutdownType Sh
 	DeactivateActiveTool(WhichSide, ShutdownType);
 }
 
-bool UInteractiveToolsContext::IsToolBuilderActive(EToolSide WhichSide, UInteractiveToolBuilder* Builder)
+// Note: this takes FString by value so that it can be bound to delegates using CreateUObject. If you change it
+// to a reference, you'll need to use CreateWeakLambda instead and capture ToolIdentifier by value there.
+bool UInteractiveToolsContext::IsToolActive(EToolSide WhichSide, const FString ToolIdentifier) const
 {
-	return ToolManager->GetActiveToolBuilder(WhichSide) == Builder;
+	return GetActiveToolName(WhichSide) == ToolIdentifier;
 }
 
 void UInteractiveToolsContext::PostToolNotificationMessage(const FText& Message)
diff --git a/Engine/Source/Runtime/Experimental/InteractiveToolsFramework/Public/InteractiveToolsContext.h b/Engine/Source/Runtime/Experimental/InteractiveToolsFramework/Public/InteractiveToolsContext.h
--- a/Engine/Source/Runtime/Experimental/InteractiveToolsFramework/Public/InteractiveToolsContext.h
+++ b/Engine/Source/Runtime/Experimental/InteractiveToolsFramework/Public/InteractiveToolsContext.h
@@ -44,7 +44,7 @@ public:
 	bool CanCompleteActiveTool(EToolSide WhichSide) const;
 	bool StartTool(EToolSide WhichSide, const FString& ToolTypeIdentifier);
 	void EndTool(EToolSide WhichSide, EToolShutdownType ShutdownType);
-	bool IsToolBuilderActive(EToolSide WhichSide, UInteractiveToolBuilder* Builder);
+	bool IsToolActive(EToolSide WhichSide, const FString ToolIdentifier) const;
 
 public:
 	// forwards message to OnToolNotificationMessage delegate
