ToolsFramework: add FInputDeviceRay constructor that can automatically derive the members from a FInputDeviceState. Make FRay constructor explicit, fix up call sites that were relying on implicit construction.
#rnx

diff --git a/Engine/Source/Runtime/Experimental/InteractiveToolsFramework/Private/BaseBehaviors/MouseHoverBehavior.cpp b/Engine/Source/Runtime/Experimental/InteractiveToolsFramework/Private/BaseBehaviors/MouseHoverBehavior.cpp
--- a/Engine/Source/Runtime/Experimental/InteractiveToolsFramework/Private/BaseBehaviors/MouseHoverBehavior.cpp
+++ b/Engine/Source/Runtime/Experimental/InteractiveToolsFramework/Private/BaseBehaviors/MouseHoverBehavior.cpp
@@ -30,7 +30,7 @@ FInputCaptureRequest UMouseHoverBehavior::WantsHoverCapture(const FInputDeviceSt
 	{
 		Modifiers.UpdateModifiers(InputState, Target);
 
-		FInputRayHit Hit = Target->BeginHoverSequenceHitTest(InputState.Mouse.WorldRay);
+		FInputRayHit Hit = Target->BeginHoverSequenceHitTest(FInputDeviceRay(InputState));
 		if (Hit.bHit)
 		{
 			return FInputCaptureRequest::Begin(this, EInputCaptureSide::Any, Hit.HitDepth);
@@ -43,7 +43,7 @@ FInputCaptureUpdate UMouseHoverBehavior::BeginHoverCapture(const FInputDeviceSta
 {
 	check(Target != nullptr);
 	Modifiers.UpdateModifiers(InputState, Target);
-	Target->OnBeginHover(InputState.Mouse.WorldRay);
+	Target->OnBeginHover(FInputDeviceRay(InputState));
 	return FInputCaptureUpdate::Begin(this, eSide);
 }
 
@@ -51,7 +51,7 @@ FInputCaptureUpdate UMouseHoverBehavior::UpdateHoverCapture(const FInputDeviceSt
 {
 	check(Target != nullptr);
 	Modifiers.UpdateModifiers(InputState, Target);
-	if (Target->OnUpdateHover(InputState.Mouse.WorldRay))
+	if (Target->OnUpdateHover(FInputDeviceRay(InputState)))
 	{
 		return FInputCaptureUpdate::Continue();
 	}
diff --git a/Engine/Source/Runtime/Experimental/InteractiveToolsFramework/Private/BaseGizmos/ScalableSphereGizmo.cpp b/Engine/Source/Runtime/Experimental/InteractiveToolsFramework/Private/BaseGizmos/ScalableSphereGizmo.cpp
--- a/Engine/Source/Runtime/Experimental/InteractiveToolsFramework/Private/BaseGizmos/ScalableSphereGizmo.cpp
+++ b/Engine/Source/Runtime/Experimental/InteractiveToolsFramework/Private/BaseGizmos/ScalableSphereGizmo.cpp
@@ -338,12 +338,12 @@ FInputCaptureUpdate UScalableSphereGizmoInputBehavior::UpdateCapture(const FInpu
 	if (IsReleased(input))
 	{
 		bInputDragCaptured = false;
-		Gizmo->OnEndDrag(LastWorldRay);
+		Gizmo->OnEndDrag(FInputDeviceRay(LastWorldRay));
 		return FInputCaptureUpdate::End();
 	}
 
 	// Forward behavior to the Gizmo
-	Gizmo->OnUpdateDrag(LastWorldRay);
+	Gizmo->OnUpdateDrag(FInputDeviceRay(LastWorldRay));
 
 	return FInputCaptureUpdate::Continue();
 }
@@ -353,7 +353,7 @@ void UScalableSphereGizmoInputBehavior::ForceEndCapture(const FInputCaptureData&
 	if (bInputDragCaptured)
 	{
 		bInputDragCaptured = false;
-		Gizmo->OnEndDrag(LastWorldRay);
+		Gizmo->OnEndDrag(FInputDeviceRay(LastWorldRay));
 	}
 }
 
diff --git a/Engine/Source/Runtime/Experimental/InteractiveToolsFramework/Public/InputState.h b/Engine/Source/Runtime/Experimental/InteractiveToolsFramework/Public/InputState.h
--- a/Engine/Source/Runtime/Experimental/InteractiveToolsFramework/Public/InputState.h
+++ b/Engine/Source/Runtime/Experimental/InteractiveToolsFramework/Public/InputState.h
@@ -31,40 +31,6 @@ ENUM_CLASS_FLAGS(EInputDevices);
 
 
 
-/**
- * FInputDeviceRay represents a 3D ray created based on an input device.
- * If the device is a 2D input device like a mouse, then the ray may
- * have an associated 2D screen position.
- */
-struct FInputDeviceRay
-{
-	/** 3D ray in 3D scene, in world coordinates */
-	FRay WorldRay;
-
-	/** If true, WorldRay has 2D device position coordinates */
-	bool bHas2D;
-
-	/** 2D device position coordinates associated with the ray */
-	FVector2D ScreenPosition;
-
-
-	FInputDeviceRay(const FRay& WorldRayIn)
-	{
-		WorldRay = WorldRayIn;
-		bHas2D = false;
-		ScreenPosition = FVector2D(0, 0);
-	}
-
-	FInputDeviceRay(const FRay& WorldRayIn, const FVector2D& ScreenPositionIn)
-	{
-		WorldRay = WorldRayIn;
-		bHas2D = true;
-		ScreenPosition = ScreenPositionIn;
-	}
-};
-
-
-
 /*
  * FInputRayHit is returned by various hit-test interface functions.
  * Generally this is intended to be returned as the result of a hit-test with a FInputDeviceRay 
@@ -310,3 +276,56 @@ struct FInputDeviceState
 	}
 };
 
+
+
+
+/**
+ * FInputDeviceRay represents a 3D ray created based on an input device.
+ * If the device is a 2D input device like a mouse, then the ray may
+ * have an associated 2D screen position.
+ */
+struct FInputDeviceRay
+{
+	/** 3D ray in 3D scene, in world coordinates */
+	FRay WorldRay;
+
+	/** If true, WorldRay has 2D device position coordinates */
+	bool bHas2D = false;
+
+	/** 2D device position coordinates associated with the ray */
+	FVector2D ScreenPosition;
+
+
+	explicit FInputDeviceRay(const FRay& WorldRayIn)
+	{
+		WorldRay = WorldRayIn;
+		bHas2D = false;
+		ScreenPosition = FVector2D(0, 0);
+	}
+
+	FInputDeviceRay(const FRay& WorldRayIn, const FVector2D& ScreenPositionIn)
+	{
+		WorldRay = WorldRayIn;
+		bHas2D = true;
+		ScreenPosition = ScreenPositionIn;
+	}
+
+	FInputDeviceRay(const FInputDeviceState& Input)
+	{
+		if (Input.IsFromDevice(EInputDevices::Mouse))
+		{
+			WorldRay = Input.Mouse.WorldRay;
+			bHas2D = true;
+			ScreenPosition = Input.Mouse.Position2D;
+		}
+		else
+		{
+			ensure(false);
+			WorldRay = FRay(FVector::ZeroVector, FVector(0, 0, 1), true);
+			ScreenPosition = FVector2D(0, 0);
+			bHas2D = false;
+		}
+	}
+};
+
+
