Chaos : IgnoreCollisions
-Don't store particles in the deletion queue, just the unique indices.

diff --git a/Engine/Source/Runtime/Experimental/Chaos/Private/Chaos/Collision/CollisionConstraintFlags.cpp b/Engine/Source/Runtime/Experimental/Chaos/Private/Chaos/Collision/CollisionConstraintFlags.cpp
--- a/Engine/Source/Runtime/Experimental/Chaos/Private/Chaos/Collision/CollisionConstraintFlags.cpp
+++ b/Engine/Source/Runtime/Experimental/Chaos/Private/Chaos/Collision/CollisionConstraintFlags.cpp
@@ -96,7 +96,7 @@ namespace Chaos
 			TArray<FGeometryParticle*> DeletionList;
 			for (auto& Elem : PendingActivations)
 			{
-				int32 DeactiveIndex = PendingDeactivations.Find(Elem.Key);
+				int32 DeactiveIndex = PendingDeactivations.Find(Elem.Key->UniqueIdx());
 				if (DeactiveIndex != INDEX_NONE)
 				{
 					DeletionList.Add(Elem.Key);
@@ -168,14 +168,7 @@ namespace Chaos
 		{
 			for (auto Index = PendingDeactivations.Num() - 1; Index >= 0; Index--)
 			{
-				if (PendingDeactivations[Index])
-				{
-					FUniqueIdx ID0 = PendingDeactivations[Index]->UniqueIdx();
-					if (IgnoreCollisionsList.Contains(ID0))
-					{
-						IgnoreCollisionsList.Remove(ID0);
-					}
-				}
+				IgnoreCollisionsList.Remove(PendingDeactivations[Index]);
 			}
 			PendingDeactivations.Empty();
 		}
diff --git a/Engine/Source/Runtime/Experimental/Chaos/Private/PBDRigidsSolver.cpp b/Engine/Source/Runtime/Experimental/Chaos/Private/PBDRigidsSolver.cpp
--- a/Engine/Source/Runtime/Experimental/Chaos/Private/PBDRigidsSolver.cpp
+++ b/Engine/Source/Runtime/Experimental/Chaos/Private/PBDRigidsSolver.cpp
@@ -467,15 +467,13 @@ namespace Chaos
 			GeometryParticlePhysicsProxies.RemoveSingleSwap((FGeometryParticlePhysicsProxy*)InProxy);
 		}
 		
-
-
 		Chaos::FIgnoreCollisionManager& CollisionManager = GetEvolution()->GetBroadPhase().GetIgnoreCollisionManager();
 		{
 			int32 ExternalTimestamp = GetMarshallingManager().GetExternalTimestamp_External();
-			Chaos::FIgnoreCollisionManager::FParticleArray& PendingMap = CollisionManager.GetPendingDeactivationsForGameThread(ExternalTimestamp);
-			if (!PendingMap.Contains(GTParticle))
+			Chaos::FIgnoreCollisionManager::FDeactivationArray& PendingMap = CollisionManager.GetPendingDeactivationsForGameThread(ExternalTimestamp);
+			if (!PendingMap.Contains(GTParticle->UniqueIdx()))
 			{
-				PendingMap.Add(GTParticle);
+				PendingMap.Add(GTParticle->UniqueIdx());
 			}
 		}
 
diff --git a/Engine/Source/Runtime/Experimental/Chaos/Public/Chaos/Collision/CollisionConstraintFlags.h b/Engine/Source/Runtime/Experimental/Chaos/Public/Chaos/Collision/CollisionConstraintFlags.h
--- a/Engine/Source/Runtime/Experimental/Chaos/Public/Chaos/Collision/CollisionConstraintFlags.h
+++ b/Engine/Source/Runtime/Experimental/Chaos/Public/Chaos/Collision/CollisionConstraintFlags.h
@@ -22,13 +22,13 @@ namespace Chaos
 	public:
 		using FGeometryParticle = TGeometryParticle<FReal, 3>;
 		using FHandleID = FUniqueIdx;
-		using FParticleArray = TArray<FGeometryParticle*>;
+		using FDeactivationArray = TArray<FUniqueIdx>;
 		using FActiveMap = TMap<FHandleID, TArray<FHandleID> >;
-		using FPendingMap = TMap<FGeometryParticle*, FParticleArray >;
+		using FPendingMap = TMap<FGeometryParticle*, TArray<FGeometryParticle*> >;
 		struct FStorageData
 		{
 			FPendingMap PendingActivations;
-			FParticleArray PendingDeactivations;
+			FDeactivationArray PendingDeactivations;
 			int32 ExternalTimestamp = INDEX_NONE;
 
 			void Reset()
@@ -69,7 +69,7 @@ namespace Chaos
 			return StorageDataProducer->PendingActivations;
 		}
 
-		FParticleArray& GetPendingDeactivationsForGameThread(int32 ExternalTimestamp)
+		FDeactivationArray& GetPendingDeactivationsForGameThread(int32 ExternalTimestamp)
 		{
 			if (StorageDataProducer->ExternalTimestamp == INDEX_NONE)
 			{
@@ -126,7 +126,7 @@ namespace Chaos
 		FActiveMap IgnoreCollisionsList;
 
 		FPendingMap PendingActivations;
-		FParticleArray PendingDeactivations;
+		FDeactivationArray PendingDeactivations;
 
 		// Producer storage data, pending changes written here until pushed into queue.
 		FStorageData* StorageDataProducer;
diff --git a/Engine/Source/Runtime/PhysicsCore/Private/ChaosEngineInterface.cpp b/Engine/Source/Runtime/PhysicsCore/Private/ChaosEngineInterface.cpp
--- a/Engine/Source/Runtime/PhysicsCore/Private/ChaosEngineInterface.cpp
+++ b/Engine/Source/Runtime/PhysicsCore/Private/ChaosEngineInterface.cpp
@@ -307,10 +307,10 @@ void FChaosEngineInterface::RemoveDisabledCollisionsFor_AssumesLocked(TArray< FP
 			Chaos::FPhysicsSolver* Solver = Handle->GetProxy()->GetSolver<Chaos::FPhysicsSolver>();
 			Chaos::FIgnoreCollisionManager& CollisionManager = Solver->GetEvolution()->GetBroadPhase().GetIgnoreCollisionManager();
 			int32 ExternalTimestamp = Solver->GetMarshallingManager().GetExternalTimestamp_External();
-			Chaos::FIgnoreCollisionManager::FParticleArray& PendingMap = CollisionManager.GetPendingDeactivationsForGameThread(ExternalTimestamp);
-			if (!PendingMap.Contains(Handle))
+			Chaos::FIgnoreCollisionManager::FDeactivationArray& PendingMap = CollisionManager.GetPendingDeactivationsForGameThread(ExternalTimestamp);
+			if (!PendingMap.Contains(Handle->UniqueIdx()))
 			{
-				PendingMap.Add(Handle);
+				PendingMap.Add(Handle->UniqueIdx());
 			}
 		}
 	}
