Reduce O(n) operation cost to O(1) when unregistering physics proxy for instanced static mesh

  - 12m33s to 5m51s when generating a huge scene with Houdini

diff --git a/Engine/Source/Runtime/Experimental/Chaos/Private/PBDRigidsSolver.cpp b/Engine/Source/Runtime/Experimental/Chaos/Private/PBDRigidsSolver.cpp
--- a/Engine/Source/Runtime/Experimental/Chaos/Private/PBDRigidsSolver.cpp
+++ b/Engine/Source/Runtime/Experimental/Chaos/Private/PBDRigidsSolver.cpp
@@ -464,7 +464,7 @@ namespace Chaos
 		}
 		else
 		{
-			GeometryParticlePhysicsProxies.RemoveSingleSwap((FGeometryParticlePhysicsProxy*)InProxy);
+			GeometryParticlePhysicsProxies.Remove((FGeometryParticlePhysicsProxy*)InProxy);
 		}
 		
 		Chaos::FIgnoreCollisionManager& CollisionManager = GetEvolution()->GetBroadPhase().GetIgnoreCollisionManager();
diff --git a/Engine/Source/Runtime/Experimental/Chaos/Public/ChaosStats.h b/Engine/Source/Runtime/Experimental/Chaos/Public/ChaosStats.h
--- a/Engine/Source/Runtime/Experimental/Chaos/Public/ChaosStats.h
+++ b/Engine/Source/Runtime/Experimental/Chaos/Public/ChaosStats.h
@@ -72,6 +72,8 @@ DECLARE_CYCLE_STAT_EXTERN(TEXT("Object Parameter Update"), STAT_ParamUpdateObjec
 DECLARE_CYCLE_STAT_EXTERN(TEXT("Field Parameter Update"), STAT_ParamUpdateField, STATGROUP_Chaos, CHAOS_API);
 DECLARE_CYCLE_STAT_EXTERN(TEXT("Sync Events - Game Thread"), STAT_SyncEvents_GameThread, STATGROUP_Chaos, CHAOS_API);
 
+DECLARE_CYCLE_STAT_EXTERN(TEXT("Proxies Set Compaction"), STAT_ProxiesSetCompaction, STATGROUP_Chaos, CHAOS_API);
+
 DECLARE_CYCLE_STAT_EXTERN(TEXT("Physics Thread Stat Update"), STAT_PhysicsStatUpdate, STATGROUP_ChaosThread, CHAOS_API);
 DECLARE_FLOAT_COUNTER_STAT_EXTERN(TEXT("Physics Thread Time Actual (ms)"), STAT_PhysicsThreadTime, STATGROUP_ChaosThread, CHAOS_API);
 DECLARE_FLOAT_COUNTER_STAT_EXTERN(TEXT("Physics Thread Time Effective (ms)"), STAT_PhysicsThreadTimeEff, STATGROUP_ChaosThread, CHAOS_API);
diff --git a/Engine/Source/Runtime/Experimental/Chaos/Public/PBDRigidsSolver.h b/Engine/Source/Runtime/Experimental/Chaos/Public/PBDRigidsSolver.h
--- a/Engine/Source/Runtime/Experimental/Chaos/Public/PBDRigidsSolver.h
+++ b/Engine/Source/Runtime/Experimental/Chaos/Public/PBDRigidsSolver.h
@@ -160,6 +160,11 @@ namespace Chaos
 		template<typename Lambda>
 		void ForEachPhysicsProxy(Lambda InCallable)
 		{
+			// Compact is not strictly required here, but will offer the same performance characteristics as the multithreaded version
+			{
+				SCOPE_CYCLE_COUNTER(STAT_ProxiesSetCompaction);
+				GeometryParticlePhysicsProxies.Compact();
+			}
 			for (FGeometryParticlePhysicsProxy* Obj : GeometryParticlePhysicsProxies)
 			{
 				InCallable(Obj);
@@ -193,9 +198,15 @@ namespace Chaos
 		template<typename Lambda>
 		void ForEachPhysicsProxyParallel(Lambda InCallable)
 		{
+			// Make sure all elements are contiguous so we can use the set like an array.
+			// This only requires work if GeometryParticlePhysicsProxies has been dirtied since our last call.
+			{
+				SCOPE_CYCLE_COUNTER(STAT_ProxiesSetCompaction);
+				GeometryParticlePhysicsProxies.Compact();
+			}
 			Chaos::PhysicsParallelFor(GeometryParticlePhysicsProxies.Num(), [this, &InCallable](const int32 Index)
 			{
-				FGeometryParticlePhysicsProxy* Obj = GeometryParticlePhysicsProxies[Index];
+				FGeometryParticlePhysicsProxy* Obj = GeometryParticlePhysicsProxies[FSetElementId::FromInteger(Index)];
 				InCallable(Obj);
 			});
 			Chaos::PhysicsParallelFor(KinematicGeometryParticlePhysicsProxies.Num(), [this, &InCallable](const int32 Index)
@@ -432,7 +443,7 @@ namespace Chaos
 		// Proxies
 		//
 		TSharedPtr<FCriticalSection> MCurrentLock;
-		TArray< FGeometryParticlePhysicsProxy* > GeometryParticlePhysicsProxies;
+		TSet< FGeometryParticlePhysicsProxy* > GeometryParticlePhysicsProxies;
 		TArray< FKinematicGeometryParticlePhysicsProxy* > KinematicGeometryParticlePhysicsProxies;
 		TArray< FRigidParticlePhysicsProxy* > RigidParticlePhysicsProxies;
 		TArray< FSkeletalMeshPhysicsProxy* > SkeletalMeshPhysicsProxies; // dep
