ModelingMode: optimize UMeshSelectionSet::Remove functions that were doing a linear search for each element to remove. On large this could involve an array with millions of elements, each processed one-by-one with a linear search.
#rnx

diff --git a/Engine/Source/Runtime/Experimental/InteractiveToolsFramework/Private/SelectionSet.cpp b/Engine/Source/Runtime/Experimental/InteractiveToolsFramework/Private/SelectionSet.cpp
--- a/Engine/Source/Runtime/Experimental/InteractiveToolsFramework/Private/SelectionSet.cpp
+++ b/Engine/Source/Runtime/Experimental/InteractiveToolsFramework/Private/SelectionSet.cpp
@@ -71,15 +71,8 @@ void UMeshSelectionSet::AddIndices(EMeshSelectionElementType ElementType, const
 
 void UMeshSelectionSet::RemoveIndices(EMeshSelectionElementType ElementType, const TArray<int32>& Indices)
 {
-	TArray<int32>& CurElements = GetElements(ElementType);
-
-	// @todo if we are removing many elements it is maybe cheaper to make a new array...
-	int N = Indices.Num();
-	for (int32 k = 0; k < N; ++k)
-	{
-		CurElements.RemoveSwap(Indices[k]);
-	}
-	NotifySelectionSetModified();
+	TSet<int32> IndicesSet(Indices);
+	RemoveIndices(ElementType, IndicesSet);
 }
 
 
@@ -87,10 +80,16 @@ void UMeshSelectionSet::RemoveIndices(EMeshSelectionElementType ElementType, con
 {
 	TArray<int32>& CurElements = GetElements(ElementType);
 
-	// @todo if we are removing many elements it is maybe cheaper to make a new array...
-	for (int32 Index : Indices)
+	TArray<int32> NewElements;
+	NewElements.Reserve(FMath::Max(0, CurElements.Num() - Indices.Num()));
+	for (int32 Index : CurElements)
 	{
-		CurElements.RemoveSwap(Index);
+		if (Indices.Contains(Index) == false)
+		{
+			NewElements.Add(Index);
+		}
 	}
+	CurElements = MoveTemp(NewElements);
+
 	NotifySelectionSetModified();
 }
\ No newline at end of file
