Fix World::IsReadyForFinishDestroy checking if physics has not finished completion event instead of checking all pending tasks.(Completion event only represents blocking tasks, not all async).

diff --git a/Engine/Source/Runtime/Experimental/Chaos/Public/PBDRigidsSolver.h b/Engine/Source/Runtime/Experimental/Chaos/Public/PBDRigidsSolver.h
--- a/Engine/Source/Runtime/Experimental/Chaos/Public/PBDRigidsSolver.h
+++ b/Engine/Source/Runtime/Experimental/Chaos/Public/PBDRigidsSolver.h
@@ -403,7 +403,12 @@ namespace Chaos
 
 		virtual bool AreAnyTasksPending() const override
 		{
-			return GetEvolution()->AreAnyTasksPending();
+			if (IsPendingTasksComplete() == false || GetEvolution()->AreAnyTasksPending())
+			{
+				return true;
+			}
+
+			return false;
 		}
 
 		void BeginDestroy();
diff --git a/Engine/Source/Runtime/PhysicsCore/Private/ChaosScene.cpp b/Engine/Source/Runtime/PhysicsCore/Private/ChaosScene.cpp
--- a/Engine/Source/Runtime/PhysicsCore/Private/ChaosScene.cpp
+++ b/Engine/Source/Runtime/PhysicsCore/Private/ChaosScene.cpp
@@ -372,11 +372,6 @@ void FChaosScene::OnSyncBodies(Chaos::FPhysicsSolverBase* Solver)
 
 bool FChaosScene::AreAnyTasksPending() const
 {
-	if (!IsCompletionEventComplete())
-	{
-		return true;
-	}
-
 	const Chaos::FPBDRigidsSolver* Solver = GetSolver();
 	if (Solver && Solver->AreAnyTasksPending())
 	{
