Fix bad boxfalloff radius

diff --git a/Engine/Source/Runtime/Experimental/Chaos/Private/Field/FieldSystemNodes.cpp b/Engine/Source/Runtime/Experimental/Chaos/Private/Field/FieldSystemNodes.cpp
--- a/Engine/Source/Runtime/Experimental/Chaos/Private/Field/FieldSystemNodes.cpp
+++ b/Engine/Source/Runtime/Experimental/Chaos/Private/Field/FieldSystemNodes.cpp
@@ -521,7 +521,7 @@ void FBoxFalloff::Evaluator(const FFieldContext& Context, TArrayView<float>& Res
 {
 	const float DeltaRange = (MaxRange - MinRange);
 
-	static const float HalfBox = 0.5;
+	static const float HalfBox = 50.0;
 	static const FBox UnitBox(FVector(-HalfBox), FVector(HalfBox));
 
 	const int32 NumSamples = Context.SampleIndices.Num();
