when switching tools, don't accept the result of a tool when CanAccept() is false

diff --git a/Engine/Source/Runtime/Experimental/InteractiveToolsFramework/Private/InteractiveToolManager.cpp b/Engine/Source/Runtime/Experimental/InteractiveToolsFramework/Private/InteractiveToolManager.cpp
--- a/Engine/Source/Runtime/Experimental/InteractiveToolsFramework/Private/InteractiveToolManager.cpp
+++ b/Engine/Source/Runtime/Experimental/InteractiveToolsFramework/Private/InteractiveToolManager.cpp
@@ -128,7 +128,7 @@ bool UInteractiveToolManager::ActivateTool(EToolSide Side)
 
 	if (ActiveLeftTool != nullptr)
 	{
-		DeactivateTool(EToolSide::Left, EToolShutdownType::Accept);
+		DeactivateTool(EToolSide::Left, ActiveLeftTool->CanAccept() ? EToolShutdownType::Accept : EToolShutdownType::Cancel);
 	}
 
 	if (ActiveLeftBuilder == nullptr || ActivateToolInternal(Side) == false)
