MeshModelingTools: Made the scaling gizmos able to clamp to zero instead of going negative, and made the plane scaling gizmos able to scale uniformly in their plane.

#rnx

diff --git a/Engine/Source/Runtime/Experimental/InteractiveToolsFramework/Public/BaseGizmos/ParameterToTransformAdapters.h b/Engine/Source/Runtime/Experimental/InteractiveToolsFramework/Public/BaseGizmos/ParameterToTransformAdapters.h
--- a/Engine/Source/Runtime/Experimental/InteractiveToolsFramework/Public/BaseGizmos/ParameterToTransformAdapters.h
+++ b/Engine/Source/Runtime/Experimental/InteractiveToolsFramework/Public/BaseGizmos/ParameterToTransformAdapters.h
@@ -36,7 +36,7 @@ public:
 	 * Note that returned snap point will be projected onto the current translation origin/axis.
 	 * @return true if constraint point was found and should be used, false to ignore
 	 */
-	TUniqueFunction<bool(const FVector&, FVector&)> PositionConstraintFunction = [](const FVector&, FVector&) { return false; };
+	TUniqueFunction<bool(const FVector& Input, FVector& PreProjectedOutput)> PositionConstraintFunction = [](const FVector&, FVector&) { return false; };
 
 	virtual float GetParameter() const override
 	{
@@ -570,6 +570,12 @@ public:
 		FTransform NewTransform = InitialTransform;
 		FVector StartScale = InitialTransform.GetScale3D();
 		FVector NewScale = StartScale + ScaleDelta * CurScaleAxis;
+
+		if (bClampToZero)
+		{
+			NewScale = FVector::Max(FVector::ZeroVector, NewScale);
+		}
+
 		NewTransform.SetScale3D(NewScale);
 
 
@@ -614,8 +620,11 @@ public:
 
 	/** Coordinate delta is multiplied by this amount */
 	UPROPERTY()
-	float ScaleMultiplier = 0.05f;
+	float ScaleMultiplier = 0.02f;
 
+	/** If true, the minimal output scale will be zero. */
+	UPROPERTY()
+	bool bClampToZero = false;
 
 public:
 	/** Parameter is the line-equation parameter that this FloatParameterSource provides */
@@ -697,9 +706,20 @@ public:
 		// construct Scale as delta from initial position
 		FVector2D Delta = LastChange.GetChangeDelta() * ScaleMultiplier;
 
+		if (bUseEqualScaling)
+		{
+			Delta = FVector2D(FMath::Min(Delta.X, Delta.Y));
+		}
+
 		FTransform NewTransform = InitialTransform;
 		FVector StartScale = InitialTransform.GetScale3D();
 		FVector NewScale = StartScale + Delta.X*CurScaleAxisX + Delta.Y*CurScaleAxisY;
+
+		if (bClampToZero)
+		{
+			NewScale = FVector::Max(NewScale, FVector::ZeroVector);
+		}
+
 		NewTransform.SetScale3D(NewScale);
 
 		// apply position constraint
@@ -743,8 +763,15 @@ public:
 
 	/** Coordinate delta is multiplied by this amount */
 	UPROPERTY()
-	float ScaleMultiplier = 0.05f;
+	float ScaleMultiplier = 0.02f;
+
+	/** If true, the scaling will be done an equal amount in each axis, using the minimal value */
+	UPROPERTY()
+	bool bUseEqualScaling = false;
 
+	/** If true, the minimal output scale will be zero. */
+	UPROPERTY()
+	bool bClampToZero = false;
 public:
 	/** Parameter is the two line-equation parameters that this Vec2ParameterSource provides */
 	UPROPERTY()
