FROST-206 - Fix convex hull generation crash in physics editor
- Properly pass the suggested horizon epsilon parameter from caller sites
- Adjust the normal tolerance to better match the horizon epsilons value range
- Added extra debug support that can spit out OBJ format compatible mesh data to be able to debug visually the intermediate results

diff --git a/Engine/Source/Runtime/Engine/Private/PhysicsEngine/BodySetup.cpp b/Engine/Source/Runtime/Engine/Private/PhysicsEngine/BodySetup.cpp
--- a/Engine/Source/Runtime/Engine/Private/PhysicsEngine/BodySetup.cpp
+++ b/Engine/Source/Runtime/Engine/Private/PhysicsEngine/BodySetup.cpp
@@ -1926,7 +1926,9 @@ TArray<int32> FKConvexElem::GetChaosConvexIndices() const
 		}
 
 		TArray<Chaos::TVector<int32, 3>> Triangles;
-		Chaos::FConvexBuilder::BuildConvexHull(ConvexParticles, Triangles);
+		Chaos::FConvexBuilder::Params BuildParams;
+		BuildParams.HorizonEpsilon = Chaos::FConvexBuilder::SuggestEpsilon(ConvexParticles);
+		Chaos::FConvexBuilder::BuildConvexHull(ConvexParticles, Triangles, BuildParams);
 
 		ResultIndexData.Reserve(Triangles.Num() * 3);
 		for (Chaos::TVector<int32, 3> Tri : Triangles)
diff --git a/Engine/Source/Runtime/Experimental/Chaos/Public/Chaos/CollisionConvexMesh.h b/Engine/Source/Runtime/Experimental/Chaos/Public/Chaos/CollisionConvexMesh.h
--- a/Engine/Source/Runtime/Experimental/Chaos/Public/Chaos/CollisionConvexMesh.h
+++ b/Engine/Source/Runtime/Experimental/Chaos/Public/Chaos/CollisionConvexMesh.h
@@ -10,6 +10,16 @@
 
 #define DEBUG_HULL_GENERATION 0
 
+// Those flags allow to output geometric data in OBJ compatible format
+// INSTRUCTIONS:
+//     Get the data form the log, remove the log header part and save it to a .obj file,
+//     in a 3D viewer or DCC ( Windows 3D Viewer, Maya, Blender ... ), open/import the .obj file
+// WARNING: 
+//    - this needs DEBUG_HULL_GENERATION to also be enabled
+//    - this may produce a lot of data and slow down levels have assets with a lot of convexes
+#define DEBUG_HULL_GENERATION_HULLSTEPS_TO_OBJ 0
+#define DEBUG_HULL_GENERATION_BUILDHORIZON_TO_OBJ 0
+
 namespace Chaos
 {
 	// When encountering a triangle or quad in hull generation (3-points or 4 coplanar points) we will instead generate
@@ -301,7 +311,9 @@ namespace Chaos
 			if(NumParticlesToUse >= 4)
 			{
 				TArray<TVector<int32, 3>> Indices;
-				BuildConvexHull(*ParticlesToUse, Indices);
+				Params BuildParams;
+				BuildParams.HorizonEpsilon = Chaos::FConvexBuilder::SuggestEpsilon(*ParticlesToUse);
+				BuildConvexHull(*ParticlesToUse, Indices, BuildParams);
 				OutPlanes.Reserve(Indices.Num());
 				TMap<int32, int32> IndexMap; // maps original particle indices to output particle indices
 				int32 NewIdx = 0;
@@ -368,6 +380,28 @@ namespace Chaos
 			FHalfEdge* ConflictV = FindConflictVertex(InParticles, DummyFace->Next);
 			while(ConflictV)
 			{
+
+#if DEBUG_HULL_GENERATION
+#if DEBUG_HULL_GENERATION_HULLSTEPS_TO_OBJ
+				UE_LOG(LogChaos, VeryVerbose, TEXT("# ======================================================"));
+				const FVec3 ConflictPos = InParticles.X(ConflictV->Vertex);
+				UE_LOG(LogChaos, VeryVerbose, TEXT("# GENERATED HULL before adding Vtx %d (%f %f %f)"), ConflictV->Vertex, ConflictPos.X, ConflictPos.Y, ConflictPos.Z);
+				UE_LOG(LogChaos, VeryVerbose, TEXT("# ------------------------------------------------------"));
+				FConvexFace* Face = DummyFace->Next;
+				while (Face)
+				{
+					const FVector P1 = InParticles.X(Face->FirstEdge->Prev->Vertex);
+					const FVector P2 = InParticles.X(Face->FirstEdge->Next->Vertex);
+					const FVector P3 = InParticles.X(Face->FirstEdge->Vertex);
+					UE_LOG(LogChaos, VeryVerbose, TEXT("v %f %f %f"), P1.X, P1.Y, P1.Z);
+					UE_LOG(LogChaos, VeryVerbose, TEXT("v %f %f %f"), P2.X, P2.Y, P2.Z);
+					UE_LOG(LogChaos, VeryVerbose, TEXT("v %f %f %f"), P3.X, P3.Y, P3.Z);
+					UE_LOG(LogChaos, VeryVerbose, TEXT("f -3 -2 -1"));
+					Face = Face->Next;
+				}
+#endif
+#endif
+
 				AddVertex(Pool, InParticles, ConflictV, InParams);
 				ConflictV = FindConflictVertex(InParticles, DummyFace->Next);
 			}
@@ -603,7 +637,7 @@ namespace Chaos
 			TR->Next = RS;
 			FVec3 RSTNormal = ComputeFaceNormal(InParticles.X(RS->Vertex), InParticles.X(ST->Vertex), InParticles.X(TR->Vertex));
 			const FReal RSTNormalSize = RSTNormal.Size();
-			check(RSTNormalSize > 1e-4);
+			check(RSTNormalSize > 1e-6);
 			RSTNormal = RSTNormal * (1 / RSTNormalSize);
 			FConvexFace* RST = Pool.AllocConvexFace(TPlaneConcrete<FReal, 3>(InParticles.X(RS->Vertex), RSTNormal));
 			RST->FirstEdge = RS;
@@ -925,6 +959,26 @@ namespace Chaos
 					HorizonEdges.Add(Edge);
 				}
 			}
+
+#if DEBUG_HULL_GENERATION
+#if DEBUG_HULL_GENERATION_BUILDHORIZON_TO_OBJ
+			UE_LOG(LogChaos, VeryVerbose, TEXT("# ======================================================"));
+			const FVec3 ConflictPos = InParticles.X(ConflictV->Vertex);
+			UE_LOG(LogChaos, VeryVerbose, TEXT("# BUILD_HORIZON - Conflict Vertex = %d (%f %f %f)"), ConflictV->Vertex, ConflictPos.X, ConflictPos.Y, ConflictPos.Z);
+			UE_LOG(LogChaos, VeryVerbose, TEXT("# ------------------------------------------------------"));
+			for (TSet<FConvexFace*>::TConstIterator SetIt(Processed); SetIt; ++SetIt)
+			{
+				const FConvexFace* Face = *SetIt;
+				const FVector P1 = InParticles.X(Face->FirstEdge->Prev->Vertex);
+				const FVector P2 = InParticles.X(Face->FirstEdge->Next->Vertex);
+				const FVector P3 = InParticles.X(Face->FirstEdge->Vertex);
+				UE_LOG(LogChaos, VeryVerbose, TEXT("v %f %f %f"), P1.X, P1.Y, P1.Z);
+				UE_LOG(LogChaos, VeryVerbose, TEXT("v %f %f %f"), P2.X, P2.Y, P2.Z);
+				UE_LOG(LogChaos, VeryVerbose, TEXT("v %f %f %f"), P3.X, P3.Y, P3.Z);
+				UE_LOG(LogChaos, VeryVerbose, TEXT("f -3 -2 -1"));
+			}
+#endif
+#endif
 		}
 
 		static void BuildFaces(FMemPool& Pool, const TParticles<FReal, 3>& InParticles, const FHalfEdge* ConflictV, const TArray<FHalfEdge*>& HorizonEdges, const TArray<FConvexFace*> OldFaces, TArray<FConvexFace*>& NewFaces)
