GeometryCollections remove their particles from acceleration structure immediately when Unregistered.
Fix GC not removing or disabling particles from evolution.
Fixes crash during SQ against GC.

diff --git a/Engine/Source/Runtime/Engine/Private/PhysicsEngine/Experimental/PhysScene_Chaos.cpp b/Engine/Source/Runtime/Engine/Private/PhysicsEngine/Experimental/PhysScene_Chaos.cpp
--- a/Engine/Source/Runtime/Engine/Private/PhysicsEngine/Experimental/PhysScene_Chaos.cpp
+++ b/Engine/Source/Runtime/Engine/Private/PhysicsEngine/Experimental/PhysScene_Chaos.cpp
@@ -593,6 +593,17 @@ void FPhysScene_Chaos::RemoveObject(FGeometryParticlePhysicsProxy* InObject)
 void FPhysScene_Chaos::RemoveObject(FGeometryCollectionPhysicsProxy* InObject)
 {
 	Chaos::FPhysicsSolver* Solver = InObject->GetSolver<Chaos::FPhysicsSolver>();
+
+	for (TUniquePtr<Chaos::TGeometryParticle<Chaos::FReal, 3>>& GTParticleUnique : InObject->GetExternalParticles())
+	{
+		Chaos::TGeometryParticle<Chaos::FReal, 3>* GTParticle = GTParticleUnique.Get();
+		RemoveActorFromAccelerationStructure(GTParticle);
+		if (Solver)
+		{
+			Solver->UpdateParticleInAccelerationStructure_External(GTParticle, /*bDelete=*/true); // ensures deletion will be applied to structures passed from physics thread.
+		}
+	}
+
 	if(Solver && !Solver->UnregisterObject(InObject))
 	{
 		UE_LOG(LogChaos, Warning, TEXT("Attempted to remove an object that wasn't found in its solver's gamethread storage - it's likely the solver has been mistakenly changed."));
diff --git a/Engine/Source/Runtime/Experimental/Chaos/Private/PBDRigidsSolver.cpp b/Engine/Source/Runtime/Experimental/Chaos/Private/PBDRigidsSolver.cpp
--- a/Engine/Source/Runtime/Experimental/Chaos/Private/PBDRigidsSolver.cpp
+++ b/Engine/Source/Runtime/Experimental/Chaos/Private/PBDRigidsSolver.cpp
@@ -573,7 +573,7 @@ namespace Chaos
 		UE_LOG(LogPBDRigidsSolver, Verbose, TEXT("TPBDRigidsSolver::RegisterObject(FGeometryCollectionPhysicsProxy*)"));
 		GeometryCollectionPhysicsProxies_External.AddUnique(InProxy);
 		InProxy->SetSolver(this);
-		InProxy->Initialize();
+		InProxy->Initialize(GetEvolution());
 		InProxy->NewData(); // Buffers data on the proxy.
 		FParticlesType* InParticles = &GetParticles();
 
@@ -599,6 +599,9 @@ namespace Chaos
 
 		int32 NumRemoved = GeometryCollectionPhysicsProxies_External.Remove(InProxy);
 
+		// Particles are removed from acceleration structure in FPhysScene_Chaos::RemoveObject.
+
+
 		EnqueueCommandImmediate([InProxy, this]()
 			{
 				const TArray<Chaos::TPBDRigidClusteredParticleHandle<float, 3>*>& ParticleHandles = InProxy->GetSolverParticleHandles();
@@ -608,6 +611,7 @@ namespace Chaos
 				}
 				GeometryCollectionPhysicsProxies_Internal.RemoveSingle(InProxy);
 				InProxy->SyncBeforeDestroy();
+				InProxy->OnRemoveFromSolver(this);
 				delete InProxy;
 			});
 
diff --git a/Engine/Source/Runtime/Experimental/Chaos/Private/PhysicsProxy/GeometryCollectionPhysicsProxy.cpp b/Engine/Source/Runtime/Experimental/Chaos/Private/PhysicsProxy/GeometryCollectionPhysicsProxy.cpp
--- a/Engine/Source/Runtime/Experimental/Chaos/Private/PhysicsProxy/GeometryCollectionPhysicsProxy.cpp
+++ b/Engine/Source/Runtime/Experimental/Chaos/Private/PhysicsProxy/GeometryCollectionPhysicsProxy.cpp
@@ -25,6 +25,7 @@
 #include "GeometryCollection/ManagedArrayCollection.h"
 #include "Modules/ModuleManager.h"
 #include "Chaos/PullPhysicsDataImp.h"
+#include "Chaos/PBDRigidsEvolution.h"
 
 #ifndef TODO_REIMPLEMENT_INIT_COMMANDS
 #define TODO_REIMPLEMENT_INIT_COMMANDS 0
@@ -405,8 +406,8 @@ FGeometryCollectionPhysicsProxy::~FGeometryCollectionPhysicsProxy()
 float ReportHighParticleFraction = -1.f;
 FAutoConsoleVariableRef CVarReportHighParticleFraction(TEXT("p.gc.ReportHighParticleFraction"), ReportHighParticleFraction, TEXT("Report any objects with particle fraction above this threshold"));
 
-
-void FGeometryCollectionPhysicsProxy::Initialize()
+template <typename Traits>
+void FGeometryCollectionPhysicsProxy::Initialize(Chaos::TPBDRigidsEvolutionBase<Traits> *Evolution)
 {
 	check(IsInGameThread());
 
@@ -530,6 +531,8 @@ void FGeometryCollectionPhysicsProxy::Initialize()
 			GTParticles[Index] = Chaos::TGeometryParticle<Chaos::FReal, 3>::CreateParticle();
 			Chaos::TGeometryParticle<Chaos::FReal, 3>* P = GTParticles[Index].Get();
 
+			GTParticles[Index]->SetUniqueIdx(Evolution->GenerateUniqueIdx());
+
 			const FTransform& T = Parameters.WorldTransform * GameThreadCollection.Transform[Index];
 			P->SetX(T.GetTranslation(), false);
 			P->SetR(T.GetRotation(), false);
@@ -717,6 +720,7 @@ void FGeometryCollectionPhysicsProxy::InitializeBodiesPT(Chaos::TPBDRigidsSolver
 
 				// We're on the physics thread here but we've already set up the GT particles and we're just linking here
 				Handle->GTGeometryParticle() = GTParticles[Idx].Get();
+				Handle->SetUniqueIdx(GTParticles[Idx]->UniqueIdx());
 
 				check(SolverParticleHandles[Idx]->GetParticleType() == Handle->GetParticleType());
 				RigidsSolver->GetEvolution()->CreateParticle(Handle);
@@ -1349,10 +1353,26 @@ void FGeometryCollectionPhysicsProxy::OnRemoveFromSolver(Chaos::TPBDRigidsSolver
 {
 	const FGeometryDynamicCollection& DynamicCollection = PhysicsThreadCollection;
 
+	Chaos::TPBDRigidsEvolutionGBF<Traits>* Evolution = RBDSolver->GetEvolution();
+
 	for (const FClusterHandle* Handle : SolverClusterHandles)
 	{
 		RBDSolver->RemoveParticleToProxy(Handle);
 	}
+
+	for (FClusterHandle* Handle : SolverParticleHandles)
+	{	
+		if(Chaos::TPBDRigidClusteredParticleHandle<Chaos::FReal, 3> * Cluster = Handle->CastToClustered())
+		{
+			Evolution->GetRigidClustering().GetTopLevelClusterParents().Remove(Cluster);
+			Evolution->GetRigidClustering().GetChildrenMap().Remove(Cluster);
+			Evolution->DestroyParticle(Cluster);
+		}
+		else
+		{
+			Evolution->DestroyParticle(Handle);
+		}
+	}
 }
 
 void FGeometryCollectionPhysicsProxy::OnRemoveFromScene()
@@ -2776,7 +2796,8 @@ void FGeometryCollectionPhysicsProxy::ParameterUpdateCallback(Chaos::TPBDRigidsS
 		const Chaos::FClusterCreationParameters<float> & Parameters);\
 	template void FGeometryCollectionPhysicsProxy::FieldForcesUpdateCallback(\
 		Chaos::TPBDRigidsSolver<Chaos::Traits>* InSolver,\
-		FParticlesType& Particles,Chaos::TArrayCollectionArray<FVector>& Force,Chaos::TArrayCollectionArray<FVector>& Torque,const float Time);
+		FParticlesType& Particles,Chaos::TArrayCollectionArray<FVector>& Force,Chaos::TArrayCollectionArray<FVector>& Torque,const float Time);\
+	template void FGeometryCollectionPhysicsProxy::Initialize(Chaos::TPBDRigidsEvolutionBase<Chaos::Traits>* Evolution);
 
 #include "Chaos/EvolutionTraits.inl"
 #undef EVOLUTION_TRAIT
\ No newline at end of file
diff --git a/Engine/Source/Runtime/Experimental/Chaos/Public/Chaos/PBDRigidsSOAs.h b/Engine/Source/Runtime/Experimental/Chaos/Public/Chaos/PBDRigidsSOAs.h
--- a/Engine/Source/Runtime/Experimental/Chaos/Public/Chaos/PBDRigidsSOAs.h
+++ b/Engine/Source/Runtime/Experimental/Chaos/Public/Chaos/PBDRigidsSOAs.h
@@ -217,14 +217,35 @@ public:
 
 	void DestroyParticle(TGeometryParticleHandle<T, d>* Particle)
 	{
-		check(Particle->CastToClustered() == nullptr);	//not supported
-
 		auto PBDRigid = Particle->CastToRigidParticle();
 		if(PBDRigid)
 		{
 			RemoveFromMapAndArray(PBDRigid, ActiveParticlesToIndex, ActiveParticlesArray);
 			RemoveFromMapAndArray(PBDRigid, TransientDirtyToIndex, TransientDirtyArray);
 
+			if (auto PBDRigidClustered = Particle->CastToClustered())
+			{
+				if (Particle->GetParticleType() == Chaos::EParticleType::GeometryCollection)
+				{
+					bGeometryCollectionDirty = true;
+				}
+				else // clustered
+				{
+					RemoveFromMapAndArray(PBDRigidClustered,
+						NonDisabledClusteredToIndex, NonDisabledClusteredArray);
+
+					if (Particle->ObjectState() == EObjectStateType::Dynamic)
+					{
+						RemoveFromMapAndArray(PBDRigidClustered,
+							ActiveClusteredToIndex, ActiveClusteredArray);
+					}
+				}
+			}
+			else
+			{
+				Particle->MoveToSOA(*DynamicDisabledParticles);
+			}
+
 			// Check for sleep events referencing this particle
 			// TODO think about this case more
 			GetDynamicParticles().GetSleepDataLock().WriteLock();
diff --git a/Engine/Source/Runtime/Experimental/Chaos/Public/PhysicsProxy/GeometryCollectionPhysicsProxy.h b/Engine/Source/Runtime/Experimental/Chaos/Public/PhysicsProxy/GeometryCollectionPhysicsProxy.h
--- a/Engine/Source/Runtime/Experimental/Chaos/Public/PhysicsProxy/GeometryCollectionPhysicsProxy.h
+++ b/Engine/Source/Runtime/Experimental/Chaos/Public/PhysicsProxy/GeometryCollectionPhysicsProxy.h
@@ -25,6 +25,9 @@ namespace Chaos
 	template <typename T> struct FClusterCreationParameters;
 
 	struct FDirtyGeometryCollectionData;
+
+	template <typename Traits>
+	class TPBDRigidsEvolutionBase;
 }
 
 class FStubGeometryCollectionData : public Chaos::FParticleData 
@@ -58,7 +61,7 @@ public:
 	typedef TPhysicsProxy<FGeometryCollectionPhysicsProxy, FStubGeometryCollectionData> Base;
 	typedef FCollisionStructureManager::FSimplicial FSimplicial;
 	typedef Chaos::TPBDRigidParticleHandle<float, 3> FParticleHandle;
-	typedef Chaos::TPBDRigidClusteredParticleHandle<float, 3> FClusterHandle;
+	typedef Chaos::TPBDRigidClusteredParticleHandle<Chaos::FReal, 3> FClusterHandle;
 
 	/** Proxy publics */
 	using FInitFunc = TFunction<void(FSimulationParameters&)>;
@@ -90,7 +93,8 @@ public:
 	 * Construct \c PTDynamicCollection, copying attributes from the game thread, 
 	 * and prepare for simulation.
 	 */
-	void Initialize();
+	template <typename Traits>
+	void Initialize(Chaos::TPBDRigidsEvolutionBase<Traits>* Evolution);
 	void Reset() { }
 
 	/** 
@@ -204,6 +208,9 @@ public:
 		return GTParticles;
 	}
 
+	/* Implemented so we can construct TAccelerationStructureHandle. */
+	virtual void* GetHandleUnsafe() const override { return nullptr; }
+
 protected:
 	/**
 	 * Build a physics thread cluster parent particle.
@@ -345,7 +352,8 @@ private:
 		const Chaos::FClusterCreationParameters<float> & Parameters);\
 	extern template void FGeometryCollectionPhysicsProxy::FieldForcesUpdateCallback(\
 		Chaos::TPBDRigidsSolver<Chaos::Traits>* InSolver,\
-		FParticlesType& Particles,Chaos::TArrayCollectionArray<FVector>& Force,Chaos::TArrayCollectionArray<FVector>& Torque,const float Time);
+		FParticlesType& Particles,Chaos::TArrayCollectionArray<FVector>& Force,Chaos::TArrayCollectionArray<FVector>& Torque,const float Time);\
+	extern template void FGeometryCollectionPhysicsProxy::Initialize(Chaos::TPBDRigidsEvolutionBase<Chaos::Traits>* Evolution);
 
 #include "Chaos/EvolutionTraits.inl"
 #undef EVOLUTION_TRAIT
