MeshModelingTools: Added ability for gizmo to be repositioned using middle mouse, and used it iin Transform and PolyEdit tools. Added ability for gizmo to snap to nearby geometry when pressing Ctrl, and used it in Transform and PolyEdit tools.

#rnx

diff --git a/Engine/Plugins/Experimental/GeometryProcessing/Source/DynamicMesh/Private/Operations/GroupTopologyDeformer.cpp b/Engine/Plugins/Experimental/GeometryProcessing/Source/DynamicMesh/Private/Operations/GroupTopologyDeformer.cpp
--- a/Engine/Plugins/Experimental/GeometryProcessing/Source/DynamicMesh/Private/Operations/GroupTopologyDeformer.cpp
+++ b/Engine/Plugins/Experimental/GeometryProcessing/Source/DynamicMesh/Private/Operations/GroupTopologyDeformer.cpp
@@ -206,6 +206,11 @@ void FGroupTopologyDeformer::CalculateROI(const TArray<int>& HandleGroups, const
 			}
 		}
 	}
+
+	// ModifiedVertices currently holds all the ROI group vertices, but the fixed vertices
+	// don't actually change in position. However we do this after collecting the modified
+	// overlay normals, since those do change.
+	ModifiedVertices = ModifiedVertices.Difference(FixedBoundaryVertices);
 }
 
 
diff --git a/Engine/Plugins/Experimental/GeometryProcessing/Source/DynamicMesh/Public/Operations/GroupTopologyDeformer.h b/Engine/Plugins/Experimental/GeometryProcessing/Source/DynamicMesh/Public/Operations/GroupTopologyDeformer.h
--- a/Engine/Plugins/Experimental/GeometryProcessing/Source/DynamicMesh/Public/Operations/GroupTopologyDeformer.h
+++ b/Engine/Plugins/Experimental/GeometryProcessing/Source/DynamicMesh/Public/Operations/GroupTopologyDeformer.h
@@ -108,7 +108,7 @@ public:
 
 	/** @return the set of handle vertices */
 	const TSet<int>& GetHandleVertices() const { return HandleVertices; }
-	/** @return the set of all vertices that will be modified by the deformation*/
+	/** @return the set of all vertices whose positions will be modified by the deformation*/
 	const TSet<int>& GetModifiedVertices() const { return ModifiedVertices; }
 	/** @return the set of all overlay normals that will be modified by the deformation*/
 	const TSet<int>& GetModifiedOverlayNormals() const { return ModifiedOverlayNormals; }
diff --git a/Engine/Plugins/Experimental/GeometryProcessing/Source/GeometricObjects/Public/Spatial/MeshAABBTree3.h b/Engine/Plugins/Experimental/GeometryProcessing/Source/GeometricObjects/Public/Spatial/MeshAABBTree3.h
--- a/Engine/Plugins/Experimental/GeometryProcessing/Source/GeometricObjects/Public/Spatial/MeshAABBTree3.h
+++ b/Engine/Plugins/Experimental/GeometryProcessing/Source/GeometricObjects/Public/Spatial/MeshAABBTree3.h
@@ -127,7 +127,7 @@ public:
 	) const override
 	{
 		checkSlow(RootIndex >= 0);
-		if (RootIndex < 0 || !ensure(MeshTimestamp == Mesh->GetShapeTimestamp()))
+		if (RootIndex < 0 || (!Options.bAllowUnsafeModifiedMeshQueries && !ensure(MeshTimestamp == Mesh->GetShapeTimestamp())))
 		{
 			return IndexConstants::InvalidID;
 		}
@@ -249,7 +249,7 @@ public:
 		double MaxDist = TNumericLimits<double>::Max(), const FQueryOptions& Options = FQueryOptions())
 	{
 		checkSlow(RootIndex >= 0);
-		if (RootIndex < 0 || !ensure(MeshTimestamp == Mesh->GetShapeTimestamp()))
+		if (RootIndex < 0 || (!Options.bAllowUnsafeModifiedMeshQueries && !ensure(MeshTimestamp == Mesh->GetShapeTimestamp())))
 		{
 			return IndexConstants::InvalidID;
 		}
@@ -364,7 +364,7 @@ public:
 		NearestT = (Options.MaxDistance < TNumericLimits<float>::Max()) ? Options.MaxDistance : TNumericLimits<float>::Max();
 
 		checkSlow(RootIndex >= 0);
-		if (RootIndex < 0 || !ensure(MeshTimestamp == Mesh->GetShapeTimestamp()))
+		if (RootIndex < 0 || (!Options.bAllowUnsafeModifiedMeshQueries && !ensure(MeshTimestamp == Mesh->GetShapeTimestamp())))
 		{
 			return false;
 		}
@@ -458,7 +458,7 @@ public:
 	 */
 	virtual bool TestAnyHitTriangle(const FRay3d& Ray, const FQueryOptions& Options = FQueryOptions()) const
 	{
-		if (RootIndex < 0 || !ensure(MeshTimestamp == Mesh->GetShapeTimestamp()))
+		if (RootIndex < 0 || (!Options.bAllowUnsafeModifiedMeshQueries && !ensure(MeshTimestamp == Mesh->GetShapeTimestamp())))
 		{
 			return false;
 		}
@@ -575,7 +575,7 @@ public:
 	)
 	{
 		checkSlow(RootIndex >= 0);
-		if (RootIndex < 0 || !ensure(MeshTimestamp == Mesh->GetShapeTimestamp()))
+		if (RootIndex < 0 || (!Options.bAllowUnsafeModifiedMeshQueries && !ensure(MeshTimestamp == Mesh->GetShapeTimestamp())))
 		{
 			return FIndex2i::Invalid();
 		}
@@ -624,7 +624,7 @@ public:
 	virtual void DoTraversal(FTreeTraversal& Traversal, const FQueryOptions& Options = FQueryOptions()) const
 	{
 		checkSlow(RootIndex >= 0);
-		if (RootIndex < 0 || !ensure(MeshTimestamp == Mesh->GetShapeTimestamp()))
+		if (RootIndex < 0 || (!Options.bAllowUnsafeModifiedMeshQueries && !ensure(MeshTimestamp == Mesh->GetShapeTimestamp())))
 		{
 			return;
 		}
@@ -700,7 +700,7 @@ public:
 		const FQueryOptions& Options = FQueryOptions()
 	) const
 	{
-		if (!ensure(MeshTimestamp == Mesh->GetShapeTimestamp()))
+		if (!Options.bAllowUnsafeModifiedMeshQueries && !ensure(MeshTimestamp == Mesh->GetShapeTimestamp()))
 		{
 			return false;
 		}
@@ -750,7 +750,7 @@ public:
 		const FQueryOptions& Options = FQueryOptions(), const FQueryOptions& OtherTreeOptions = FQueryOptions()
 	) const
 	{
-		if (!ensure(MeshTimestamp == Mesh->GetShapeTimestamp()))
+		if (!Options.bAllowUnsafeModifiedMeshQueries && !ensure(MeshTimestamp == Mesh->GetShapeTimestamp()))
 		{
 			return false;
 		}
@@ -775,7 +775,7 @@ public:
 		const FQueryOptions& Options = FQueryOptions()
 	) const
 	{
-		if (!ensure(MeshTimestamp == Mesh->GetShapeTimestamp()))
+		if (!Options.bAllowUnsafeModifiedMeshQueries && !ensure(MeshTimestamp == Mesh->GetShapeTimestamp()))
 		{
 			return false;
 		}
@@ -813,7 +813,7 @@ public:
 
 		MeshIntersection::FIntersectionsQueryResult result;
 
-		if (!ensure(MeshTimestamp == Mesh->GetShapeTimestamp()))
+		if (!Options.bAllowUnsafeModifiedMeshQueries && !ensure(MeshTimestamp == Mesh->GetShapeTimestamp()))
 		{
 			return result;
 		}
@@ -847,7 +847,7 @@ public:
 
 		MeshIntersection::FIntersectionsQueryResult Result;
 
-		if (!ensure(MeshTimestamp == Mesh->GetShapeTimestamp()))
+		if (!Options.bAllowUnsafeModifiedMeshQueries && !ensure(MeshTimestamp == Mesh->GetShapeTimestamp()))
 		{
 			return Result;
 		}
@@ -859,7 +859,7 @@ public:
 
 	virtual bool TestSelfIntersection(bool bIgnoreTopoConnected = true, const FQueryOptions& Options = FQueryOptions()) const
 	{
-		if (!ensure(MeshTimestamp == Mesh->GetShapeTimestamp()))
+		if (!Options.bAllowUnsafeModifiedMeshQueries && !ensure(MeshTimestamp == Mesh->GetShapeTimestamp()))
 		{
 			return false;
 		}
diff --git a/Engine/Plugins/Experimental/GeometryProcessing/Source/GeometricObjects/Public/Spatial/SpatialInterfaces.h b/Engine/Plugins/Experimental/GeometryProcessing/Source/GeometricObjects/Public/Spatial/SpatialInterfaces.h
--- a/Engine/Plugins/Experimental/GeometryProcessing/Source/GeometricObjects/Public/Spatial/SpatialInterfaces.h
+++ b/Engine/Plugins/Experimental/GeometryProcessing/Source/GeometricObjects/Public/Spatial/SpatialInterfaces.h
@@ -50,6 +50,13 @@ public:
 		 */
 		TFunction<bool(int)> TriangleFilterF = nullptr;
 
+		// If true, then the IMeshSpatial may allow queries even when the underlying mesh has been
+		// modified without updating the queried structure. This may be useful, for instance, to 
+		// run queries against a mesh as it is being interactively modified, but it requires the 
+		// caller to pass an appropriate TriangleFilterF to make sure that modified triangles do
+		// not affect the query.
+		bool bAllowUnsafeModifiedMeshQueries = false;
+
 		FQueryOptions() {}
 		FQueryOptions(TFunction<bool(int)> TriangleFilterF) : TriangleFilterF(TriangleFilterF) {}
 		FQueryOptions(double MaxDistance, TFunction<bool(int)> TriangleFilterF = nullptr) : MaxDistance(MaxDistance), TriangleFilterF(TriangleFilterF) {}
diff --git a/Engine/Plugins/Experimental/MeshModelingToolset/Source/MeshModelingTools/Private/EditMeshPolygonsTool.cpp b/Engine/Plugins/Experimental/MeshModelingToolset/Source/MeshModelingTools/Private/EditMeshPolygonsTool.cpp
--- a/Engine/Plugins/Experimental/MeshModelingToolset/Source/MeshModelingTools/Private/EditMeshPolygonsTool.cpp
+++ b/Engine/Plugins/Experimental/MeshModelingToolset/Source/MeshModelingTools/Private/EditMeshPolygonsTool.cpp
@@ -7,6 +7,7 @@
 #include "CompGeom/PolygonTriangulation.h"
 #include "SegmentTypes.h"
 #include "DynamicMeshAttributeSet.h"
+#include "Mechanics/DragAlignmentMechanic.h"
 #include "MeshNormals.h"
 #include "ToolSceneQueriesUtil.h"
 #include "Intersection/IntersectionUtil.h"
@@ -254,6 +255,9 @@ void UEditMeshPolygonsTool::Setup()
 
 	MultiTransformer = NewObject<UMultiTransformer>(this);
 	MultiTransformer->Setup(GetToolManager()->GetPairedGizmoManager(), GetToolManager());
+	MultiTransformer->SetGizmoVisibility(false);
+	MultiTransformer->SetGizmoRepositionable(true);
+	MultiTransformer->SetDisallowNegativeScaling(true);
 	MultiTransformer->OnTransformStarted.AddUObject(this, &UEditMeshPolygonsTool::OnMultiTransformerTransformBegin);
 	MultiTransformer->OnTransformUpdated.AddUObject(this, &UEditMeshPolygonsTool::OnMultiTransformerTransformUpdate);
 	MultiTransformer->OnTransformCompleted.AddUObject(this, &UEditMeshPolygonsTool::OnMultiTransformerTransformEnd);
@@ -261,7 +265,18 @@ void UEditMeshPolygonsTool::Setup()
 		return CommonProps->bSnapToWorldGrid
 			&& GetToolManager()->GetContextQueriesAPI()->GetCurrentCoordinateSystem() == EToolContextCoordinateSystem::World;
 	});
-	MultiTransformer->SetGizmoVisibility(false);
+	// We allow non uniform scale even when the gizmo mode is set to "world" because we're not scaling components- we're
+	// moving vertices, so we don't care which axes we "scale" along.
+	MultiTransformer->SetIsNonUniformScaleAllowedFunction([]() {
+		return true;
+		});
+
+	DragAlignmentMechanic = NewObject<UDragAlignmentMechanic>(this);
+	DragAlignmentMechanic->Setup(this);
+	DragAlignmentMechanic->InitializeDeformedMeshRayCast(
+		[this]() { return &GetSpatial(); },
+		WorldTransform, &LinearDeformer); // Should happen after PrecomputeTopology so that LinearDeformer is valid
+	MultiTransformer->AddAlignmentMechanic(DragAlignmentMechanic);
 
 	if (bTriangleMode == false)
 	{
@@ -374,6 +389,7 @@ void UEditMeshPolygonsTool::Shutdown(EToolShutdownType ShutdownType)
 	SelectionMechanic->Properties->SaveProperties(this);
 
 	MultiTransformer->Shutdown();
+	DragAlignmentMechanic->Shutdown();
 	// We wait to shut down the selection mechanic in case we need to do work to store the selection.
 
 	if (EditPreview != nullptr)
@@ -746,6 +762,7 @@ void UEditMeshPolygonsTool::CacheUpdate_Gizmo()
 	LastUpdateGizmoScale = MultiTransformer->GetCurrentGizmoScale();
 	GetToolManager()->PostInvalidation();
 	bGizmoUpdatePending = true;
+	bLastUpdateUsedWorldFrame = (MultiTransformer->GetGizmoCoordinateSystem() == EToolContextCoordinateSystem::World);
 }
 
 void UEditMeshPolygonsTool::ComputeUpdate_Gizmo()
@@ -766,16 +783,37 @@ void UEditMeshPolygonsTool::ComputeUpdate_Gizmo()
 	FDynamicMesh3* Mesh = DynamicMeshComponent->GetMesh();
 	if (TranslationDelta.SquaredLength() > 0.0001 || RotateDelta.SquaredLength() > 0.0001 || CurScaleDelta.SquaredLength() > 0.0001)
 	{
-		LinearDeformer.UpdateSolution(Mesh, [&](FDynamicMesh3* TargetMesh, int VertIdx)
+		if (bLastUpdateUsedWorldFrame)
 		{
-			FVector3d PosLocal = TargetMesh->GetVertex(VertIdx);
-			FVector3d PosWorld = WorldTransform.TransformPosition(PosLocal);
-			FVector3d PosGizmo = InitialGizmoFrame.ToFramePoint(PosWorld);
-			PosGizmo = CurScale * PosGizmo;
-			FVector3d NewPosWorld = CurFrame.FromFramePoint(PosGizmo);
-			FVector3d NewPosLocal = WorldTransform.InverseTransformPosition(NewPosWorld);
-			return NewPosLocal;
-		});
+			// For a world frame gizmo, the scaling needs to happen in world aligned gizmo space, but the 
+			// rotation is still encoded in the local gizmo frame change.
+			FQuaterniond RotationToApply = CurFrame.Rotation * InitialGizmoFrame.Rotation.Inverse();
+			LinearDeformer.UpdateSolution(Mesh, [&](FDynamicMesh3* TargetMesh, int VertIdx)
+			{
+				FVector3d PosLocal = TargetMesh->GetVertex(VertIdx);
+				FVector3d PosWorld = WorldTransform.TransformPosition(PosLocal);
+				FVector3d PosWorldGizmo = PosWorld - InitialGizmoFrame.Origin;
+
+				FVector3d NewPosWorld = RotationToApply * (PosWorldGizmo * CurScale) + CurFrame.Origin;
+				FVector3d NewPosLocal = WorldTransform.InverseTransformPosition(NewPosWorld);
+				return NewPosLocal;
+			});
+		}
+		else
+		{
+			LinearDeformer.UpdateSolution(Mesh, [&](FDynamicMesh3* TargetMesh, int VertIdx)
+			{
+				// For a local gizmo, we just get the coordinates in the original frame, scale in that frame,
+				// then interpret them as coordinates in the new frame.
+				FVector3d PosLocal = TargetMesh->GetVertex(VertIdx);
+				FVector3d PosWorld = WorldTransform.TransformPosition(PosLocal);
+				FVector3d PosGizmo = InitialGizmoFrame.ToFramePoint(PosWorld);
+				PosGizmo = CurScale * PosGizmo;
+				FVector3d NewPosWorld = CurFrame.FromFramePoint(PosGizmo);
+				FVector3d NewPosLocal = WorldTransform.InverseTransformPosition(NewPosWorld);
+				return NewPosLocal;
+			});
+		}
 	}
 	else
 	{
@@ -983,6 +1021,8 @@ void UEditMeshPolygonsTool::Render(IToolsContextRenderAPI* RenderAPI)
 
 	SelectionMechanic->Render(RenderAPI);
 
+	DragAlignmentMechanic->Render(RenderAPI);
+
 	if (ExtrudeHeightMechanic != nullptr)
 	{
 		ExtrudeHeightMechanic->Render(RenderAPI);
diff --git a/Engine/Plugins/Experimental/MeshModelingToolset/Source/MeshModelingTools/Private/TransformMeshesTool.cpp b/Engine/Plugins/Experimental/MeshModelingToolset/Source/MeshModelingTools/Private/TransformMeshesTool.cpp
--- a/Engine/Plugins/Experimental/MeshModelingToolset/Source/MeshModelingTools/Private/TransformMeshesTool.cpp
+++ b/Engine/Plugins/Experimental/MeshModelingToolset/Source/MeshModelingTools/Private/TransformMeshesTool.cpp
@@ -3,6 +3,7 @@
 #include "TransformMeshesTool.h"
 #include "InteractiveToolManager.h"
 #include "InteractiveGizmoManager.h"
+#include "Mechanics/DragAlignmentMechanic.h"
 #include "ToolBuilderUtil.h"
 #include "ToolSetupUtil.h"
 #include "DynamicMesh3.h"
@@ -72,6 +73,10 @@ void UTransformMeshesTool::Setup()
 {
 	UInteractiveTool::Setup();
 
+	// Must be done before creating gizmos, so that we can bind the mechanic to them.
+	DragAlignmentMechanic = NewObject<UDragAlignmentMechanic>(this);
+	DragAlignmentMechanic->Setup(this);
+
 	UClickDragInputBehavior* ClickDragBehavior = NewObject<UClickDragInputBehavior>(this);
 	ClickDragBehavior->Initialize(this);
 	AddInputBehavior(ClickDragBehavior);
@@ -111,10 +116,6 @@ void UTransformMeshesTool::OnTick(float DeltaTime)
 	if (bCurSetPivotMode != TransformProps->bSetPivot)
 	{
 		bool bEnableSetPivot = TransformProps->bSetPivot;
-		if (TransformProps->TransformMode == ETransformMeshesTransformMode::SharedGizmoLocal)
-		{
-			bEnableSetPivot = false;
-		}
 		UpdateSetPivotModes(bEnableSetPivot);
 		bCurSetPivotMode = TransformProps->bSetPivot;
 	}
@@ -123,7 +124,7 @@ void UTransformMeshesTool::OnTick(float DeltaTime)
 
 void UTransformMeshesTool::Render(IToolsContextRenderAPI* RenderAPI)
 {
-	FPrimitiveDrawInterface* PDI = RenderAPI->GetPrimitiveDrawInterface();
+	DragAlignmentMechanic->Render(RenderAPI);
 }
 
 void UTransformMeshesTool::OnPropertyModified(UObject* PropertySet, FProperty* Property)
@@ -219,18 +220,23 @@ void UTransformMeshesTool::SetActiveGizmos_Single(bool bLocalRotations)
 	Transformable.TransformProxy = NewObject<UTransformProxy>(this);
 	Transformable.TransformProxy->bRotatePerObject = bLocalRotations;
 
+	TArray<const UPrimitiveComponent*> ComponentsToIgnoreInAlignment;
+
 	for (TUniquePtr<FPrimitiveComponentTarget>& Target : ComponentTargets)
 	{
 		Transformable.TransformProxy->AddComponent(Target->GetOwnerComponent());
+		ComponentsToIgnoreInAlignment.Add(Target->GetOwnerComponent());
 	}
 
 	// leave out nonuniform scale if we have multiple objects in non-local mode
 	bool bCanNonUniformScale = ComponentTargets.Num() == 1 || bLocalRotations;
 	ETransformGizmoSubElements GizmoElements = (bCanNonUniformScale) ?
 		ETransformGizmoSubElements::FullTranslateRotateScale : ETransformGizmoSubElements::TranslateRotateUniformScale;
-	Transformable.TransformGizmo = GizmoManager->CreateCustomTransformGizmo(GizmoElements, this);
+	Transformable.TransformGizmo = GizmoManager->CreateCustomRepositionableTransformGizmo(GizmoElements, this);
 	Transformable.TransformGizmo->SetActiveTarget(Transformable.TransformProxy);
 
+	DragAlignmentMechanic->AddToGizmo(Transformable.TransformGizmo, &ComponentsToIgnoreInAlignment);
+
 	ActiveGizmos.Add(Transformable);
 }
 
@@ -238,13 +244,22 @@ void UTransformMeshesTool::SetActiveGizmos_PerObject()
 {
 	check(ActiveGizmos.Num() == 0);
 
+	TArray<const UPrimitiveComponent*> ComponentsToIgnoreInAlignment;
 	for (TUniquePtr<FPrimitiveComponentTarget>& Target : ComponentTargets)
 	{
+		ComponentsToIgnoreInAlignment.Reset();
+
 		FTransformMeshesTarget Transformable;
 		Transformable.TransformProxy = NewObject<UTransformProxy>(this);
 		Transformable.TransformProxy->AddComponent(Target->GetOwnerComponent());
-		Transformable.TransformGizmo = GizmoManager->Create3AxisTransformGizmo(this);
+		ComponentsToIgnoreInAlignment.Add(Target->GetOwnerComponent());
+
+		ETransformGizmoSubElements GizmoElements = ETransformGizmoSubElements::FullTranslateRotateScale;
+		Transformable.TransformGizmo = GizmoManager->CreateCustomRepositionableTransformGizmo(GizmoElements, this);
 		Transformable.TransformGizmo->SetActiveTarget(Transformable.TransformProxy);
+
+		DragAlignmentMechanic->AddToGizmo(Transformable.TransformGizmo, &ComponentsToIgnoreInAlignment);
+
 		ActiveGizmos.Add(Transformable);
 	}
 }
@@ -314,7 +329,7 @@ void UTransformMeshesTool::OnClickDrag(const FInputDeviceRay& DragPos)
 {
 	bool bApplyToPivot = TransformProps->bSetPivot;
 
-	TArray<UPrimitiveComponent*> IgnoreComponents;
+	TArray<const UPrimitiveComponent*> IgnoreComponents;
 	if (bApplyToPivot == false)
 	{
 		int IgnoreIndex = (TransformProps->TransformMode == ETransformMeshesTransformMode::PerObjectGizmo) ?
diff --git a/Engine/Plugins/Experimental/MeshModelingToolset/Source/MeshModelingTools/Public/EditMeshPolygonsTool.h b/Engine/Plugins/Experimental/MeshModelingToolset/Source/MeshModelingTools/Public/EditMeshPolygonsTool.h
--- a/Engine/Plugins/Experimental/MeshModelingToolset/Source/MeshModelingTools/Public/EditMeshPolygonsTool.h
+++ b/Engine/Plugins/Experimental/MeshModelingToolset/Source/MeshModelingTools/Public/EditMeshPolygonsTool.h
@@ -25,6 +25,7 @@
 #include "EditMeshPolygonsTool.generated.h"
 
 class FMeshVertexChangeBuilder;
+class UDragAlignmentMechanic;
 class UGroupTopologyStorableSelection;
 
 /**
@@ -577,6 +578,9 @@ protected:
 	UPROPERTY()
 	UMultiTransformer* MultiTransformer = nullptr;
 
+	UPROPERTY()
+	UDragAlignmentMechanic* DragAlignmentMechanic = nullptr;
+
 	void OnMultiTransformerTransformBegin();
 	void OnMultiTransformerTransformUpdate();
 	void OnMultiTransformerTransformEnd();
@@ -603,6 +607,7 @@ protected:
 	bool bGizmoUpdatePending = false;
 	FFrame3d LastUpdateGizmoFrame;
 	FVector3d LastUpdateGizmoScale;
+	bool bLastUpdateUsedWorldFrame = false;
 	void ComputeUpdate_Gizmo();
 
 	TUniquePtr<FGroupTopology> Topology;
diff --git a/Engine/Plugins/Experimental/MeshModelingToolset/Source/MeshModelingTools/Public/TransformMeshesTool.h b/Engine/Plugins/Experimental/MeshModelingToolset/Source/MeshModelingTools/Public/TransformMeshesTool.h
--- a/Engine/Plugins/Experimental/MeshModelingToolset/Source/MeshModelingTools/Public/TransformMeshesTool.h
+++ b/Engine/Plugins/Experimental/MeshModelingToolset/Source/MeshModelingTools/Public/TransformMeshesTool.h
@@ -12,6 +12,7 @@
 
 class UBaseAxisTranslationGizmo;
 class UAxisAngleGizmo;
+class UDragAlignmentMechanic;
 class UTransformGizmo;
 class UTransformProxy;
 
@@ -96,7 +97,7 @@ public:
 	ETransformMeshesTransformMode TransformMode = ETransformMeshesTransformMode::SharedGizmo;
 
 
-	UPROPERTY(EditAnywhere, Category = Options, meta = (EditCondition = "TransformMode == ETransformMeshesTransformMode::SharedGizmo || TransformMode == ETransformMeshesTransformMode::PerObjectGizmo"))
+	UPROPERTY(EditAnywhere, Category = Options)
 	bool bSetPivot = false;
 
 
@@ -173,6 +174,8 @@ protected:
 	UPROPERTY()
 	TArray<FTransformMeshesTarget> ActiveGizmos;
 
+	UPROPERTY()
+	UDragAlignmentMechanic* DragAlignmentMechanic = nullptr;
 
 	ETransformMeshesTransformMode CurTransformMode;
 	void UpdateTransformMode(ETransformMeshesTransformMode NewMode);
diff --git a/Engine/Plugins/Experimental/MeshModelingToolset/Source/ModelingComponents/ModelingComponents.Build.cs b/Engine/Plugins/Experimental/MeshModelingToolset/Source/ModelingComponents/ModelingComponents.Build.cs
--- a/Engine/Plugins/Experimental/MeshModelingToolset/Source/ModelingComponents/ModelingComponents.Build.cs
+++ b/Engine/Plugins/Experimental/MeshModelingToolset/Source/ModelingComponents/ModelingComponents.Build.cs
@@ -44,6 +44,7 @@ public class ModelingComponents : ModuleRules
 			{
 				"CoreUObject",
 				"Engine",
+				"InputCore",
                 "RenderCore",
                 "RHI",
 
diff --git a/Engine/Plugins/Experimental/MeshModelingToolset/Source/ModelingComponents/Private/Mechanics/DragAlignmentMechanic.cpp b/Engine/Plugins/Experimental/MeshModelingToolset/Source/ModelingComponents/Private/Mechanics/DragAlignmentMechanic.cpp
new file mode 100644
--- /dev/null
+++ b/Engine/Plugins/Experimental/MeshModelingToolset/Source/ModelingComponents/Private/Mechanics/DragAlignmentMechanic.cpp
@@ -0,0 +1,322 @@
+// Copyright Epic Games, Inc. All Rights Reserved.
+
+#include "Mechanics/DragAlignmentMechanic.h"
+
+#include "BaseBehaviors/KeyAsModifierInputBehavior.h"
+#include "BaseGizmos/RepositionableTransformGizmo.h"
+#include "BaseGizmos/TransformGizmo.h"
+#include "BaseGizmos/TransformProxy.h"
+#include "InputCoreTypes.h" //EKeys
+#include "Operations/GroupTopologyDeformer.h"
+#include "SceneManagement.h" //FPrimitiveDrawInterface
+#include "ToolDataVisualizer.h"
+#include "ToolSceneQueriesUtil.h"
+
+void UDragAlignmentMechanic::Setup(UInteractiveTool* ParentToolIn)
+{	
+	UInteractionMechanic::Setup(ParentToolIn);
+
+	// Register modifier listeners.
+	UKeyAsModifierInputBehavior* AlignmentToggleBehavior = NewObject<UKeyAsModifierInputBehavior>();
+	AlignmentToggleBehavior->ModifierCheckFunc = FInputDeviceState::IsCtrlKeyDown;
+	AlignmentToggleBehavior->Initialize(this, AlignmentModifierID, EKeys::AnyKey);
+	ParentTool->AddInputBehavior(AlignmentToggleBehavior);
+
+	// Set up the function that casts rays into the world.
+	WorldRayCast = [this, ParentToolIn](const FRay& WorldRay, FHitResult& HitResult, 
+		const TArray<const UPrimitiveComponent*>* ComponentsToIgnore)
+	{
+		bool bSuccess = false;
+		FVector3d QueryPoint = WorldRay.PointAt(1);
+		FVector3d SnapPoint;
+		ToolSceneQueriesUtil::FSnapGeometry SnapGeometry;
+
+		// Try to snap to visible vertices/edges first
+		ToolSceneQueriesUtil::FFindSceneSnapPointParams SnapParams;
+		SnapParams.Tool = ParentTool.Get();
+		SnapParams.Point = &QueryPoint;
+		SnapParams.SnapPointOut = &SnapPoint;
+		SnapParams.bEdges = true;
+		SnapParams.bVertices = true;
+		SnapParams.VisualAngleThreshold = VisualAngleSnapThreshold;
+		SnapParams.SnapGeometryOut = &SnapGeometry;
+		SnapParams.ComponentsToIgnore = ComponentsToIgnore;
+
+		if (ToolSceneQueriesUtil::FindSceneSnapPoint(SnapParams))
+		{
+			HitResult = FHitResult(WorldRay.Origin, (FVector)SnapPoint);
+			HitResult.ImpactPoint = (FVector)SnapPoint;
+			HitResult.Distance = WorldRay.GetParameter(HitResult.ImpactPoint);
+			HitResult.Item = SnapGeometry.PointCount == 1 ? VERT_HIT_TYPE_ID : EDGE_HIT_TYPE_ID;
+			bSuccess = true;
+		}
+		// If we failed, try to hit a simple triangle
+		else if (ToolSceneQueriesUtil::FindNearestVisibleObjectHit(ParentToolIn->GetWorld(), HitResult, WorldRay, ComponentsToIgnore))
+		{
+			HitResult.Item = TRIANGLE_HIT_TYPE_ID;
+			bSuccess = true;
+		}
+
+		return bSuccess;
+	};
+}
+
+void UDragAlignmentMechanic::Shutdown()
+{
+}
+
+void UDragAlignmentMechanic::AddToGizmo(UTransformGizmo* TransformGizmo, const TArray<const UPrimitiveComponent*>* ComponentsToIgnoreInAlignment)
+{
+	// If we have components to ignore, we need a copy of the array so that the alignment
+	// functions can use them later.
+	TSharedPtr<TArray<const UPrimitiveComponent*>> ComponentsToIgnorePersistent;
+	if (ComponentsToIgnoreInAlignment)
+	{
+		ComponentsToIgnorePersistent = MakeShared<TArray<const UPrimitiveComponent*>>(*ComponentsToIgnoreInAlignment);
+	}
+
+	// Set the alignment functions.
+	TransformGizmo->SetWorldAlignmentFunctions(
+		[this]() { return bAlignmentToggle; },
+		[this, ComponentsToIgnorePersistent](const FRay& WorldRay, FVector& OutputPoint) { 
+			return CastRay(WorldRay, OutputPoint, ComponentsToIgnorePersistent.Get(), true); 
+		});
+
+	// If we're adding to a repositionable gizmo, we probably don't want to ignore components
+	// when repositioning the pivot, since ignoring components is mainly used to avoid hitting
+	// the actual object, which we actually do want to hit in moving moving the pivot.
+	// If we someday do want to ignore things in a pivot reposition function, we'll have to
+	// make this caller-configurable somehow.
+	if (URepositionableTransformGizmo* RepositionableGizmo = Cast<URepositionableTransformGizmo>(TransformGizmo))
+	{
+		RepositionableGizmo->SetPivotAlignmentFunctions(
+			[this]() { return bAlignmentToggle; },
+			[this](const FRay& WorldRay, FVector& OutputPoint) {
+				return CastRay(WorldRay, OutputPoint, nullptr, false); // don't ignore anything
+			});
+	}
+
+	// Register listeners to help with rendering. They get us the location that the gizmo ended
+	// up being placed, so that we can draw a line from the hit point to that location.
+	TransformGizmo->ActiveTarget->OnTransformChanged.AddWeakLambda(this, [this](UTransformProxy*, FTransform NewTransform){
+		OnGizmoTransformChanged(NewTransform);
+	});
+	TransformGizmo->ActiveTarget->OnPivotChanged.AddWeakLambda(this, [this](UTransformProxy*, FTransform NewTransform) {
+		OnGizmoTransformChanged(NewTransform);
+		});
+
+	// Register a listener to stop drawing once the gizmo is done moving.
+	TransformGizmo->ActiveTarget->OnEndTransformEdit.AddWeakLambda(this, [this](UTransformProxy*) { 
+		bPreviewEndpointsValid = false; 
+		bWaitingOnProjectedResult = false;
+	});
+	TransformGizmo->ActiveTarget->OnEndPivotEdit.AddWeakLambda(this, [this](UTransformProxy*) {
+		bPreviewEndpointsValid = false;
+		bWaitingOnProjectedResult = false;
+		});
+}
+
+void UDragAlignmentMechanic::OnGizmoTransformChanged(FTransform NewTransform)
+{
+	// Ignore changes that weren't preceded by a successful query (likely programatic)
+	if (bWaitingOnProjectedResult) 
+	{
+		LastProjectedResult = NewTransform.GetLocation();
+		bPreviewEndpointsValid = true;
+	}
+	bWaitingOnProjectedResult = false;
+}
+
+void UDragAlignmentMechanic::InitializeDeformedMeshRayCast(TFunction<FDynamicMeshAABBTree3* ()> GetSpatialIn, 
+	const FTransform3d &TargetTransform, const FGroupTopologyDeformer* LinearDeformer)
+{
+	// Use the deformer to create a way to filter out triangles we don't want to be hitting (the ones we're manipulating)
+	const FDynamicMesh3* Mesh = LinearDeformer->GetMesh();
+	MeshRayCastTriangleFilter = [Mesh, LinearDeformer](int32 Tid)
+	{
+		FIndex3i Vids = Mesh->GetTriangle(Tid);
+		for (int32 i = 0; i < 3; ++i)
+		{
+			if (LinearDeformer->GetModifiedVertices().Contains(Vids[i]))
+			{
+				return false;
+			}
+		}
+		return true;
+	};
+
+	// Set up the actual function that casts rays into the mesh.
+	MeshRayCast = [this, GetSpatialIn, TargetTransform, LinearDeformer](const FRay& WorldRay, FHitResult& HitResult, bool bUseFilter)
+	{
+		// Transform ray into local space
+		FRay3d LocalRay(TargetTransform.InverseTransformPosition(WorldRay.Origin),
+			TargetTransform.InverseTransformVector(WorldRay.Direction));
+		LocalRay.Direction.Normalize();
+
+		FDynamicMeshAABBTree3* Spatial = GetSpatialIn();
+		const FDynamicMesh3* Mesh = Spatial->GetMesh();
+
+		// Do the actual ray cast. Note that if we're using the filter, then we're presumably deforming the
+		// mesh, and the aabb tree will not be updated. In that case, we want to tell it to not perform
+		// that check, since we'll be relying on the filter to make sure that we're not hitting anything
+		// invalid.
+		double RayTValue;
+		int32 HitTid;
+		IMeshSpatial::FQueryOptions Options;
+		if (bUseFilter)
+		{
+			Options.TriangleFilterF = [this](int32 Tid) { return MeshRayCastTriangleFilter(Tid); };
+			Options.bAllowUnsafeModifiedMeshQueries = true;
+		}
+		if (!Spatial->FindNearestHitTriangle(LocalRay, RayTValue, HitTid, Options))
+		{
+			return false;
+		}
+
+		// See if we're close enough to one of the triangle's vertices to snap to it. We have
+		// to do this check in world space since scaling will affect things.
+		FVector3d WorldHitPoint = WorldRay.PointAt(RayTValue);
+
+		FVector3d Vert1, Vert2, Vert3;
+		Mesh->GetTriVertices(HitTid, Vert1, Vert2, Vert3);
+		Vert1 = TargetTransform.TransformPosition(Vert1);
+		Vert2 = TargetTransform.TransformPosition(Vert2);
+		Vert3 = TargetTransform.TransformPosition(Vert3);
+
+		FVector3d* ClosestVert = &Vert1;
+		double MinDistSquared = WorldHitPoint.DistanceSquared(Vert1);
+		double DistSquared = WorldHitPoint.DistanceSquared(Vert2);
+		if (DistSquared < MinDistSquared)
+		{
+			ClosestVert = &Vert2;
+			MinDistSquared = DistSquared;
+		}
+		DistSquared = WorldHitPoint.DistanceSquared(Vert3);
+		if (DistSquared < MinDistSquared)
+		{
+			ClosestVert = &Vert3;
+			MinDistSquared = DistSquared;
+		}
+
+		// Check if we're close enough to snap.
+		if (ToolSceneQueriesUtil::PointSnapQuery(CachedCameraState, WorldHitPoint, *ClosestVert, VisualAngleSnapThreshold))
+		{
+			HitResult = FHitResult(WorldRay.Origin, (FVector)*ClosestVert);
+			HitResult.ImpactPoint = (FVector)*ClosestVert;
+			HitResult.Item = VERT_HIT_TYPE_ID;
+		}
+		else
+		{
+			HitResult = FHitResult(WorldRay.Origin, (FVector)WorldHitPoint);
+			HitResult.ImpactPoint = (FVector)WorldHitPoint;
+			HitResult.Item = TRIANGLE_HIT_TYPE_ID;
+		}
+		HitResult.Distance = WorldRay.GetParameter(HitResult.ImpactPoint);
+		HitResult.ImpactNormal = (FVector)Mesh->GetTriNormal(HitTid);
+
+		return true;
+	};
+}
+
+/**
+ * Casts a ray into the scene to find an alignment point. Used in the functions bound
+ * to gizmos.
+ *
+ * @param bUseFilter If true, and if there is a dynamic mesh to hit, its triangles are
+ *  filtered to avoid hitting deformed triangles.
+ * @return true if something was hit.
+ */
+bool UDragAlignmentMechanic::CastRay(const FRay& WorldRay, FVector& OutputPoint,
+	const TArray<const UPrimitiveComponent*>* ComponentsToIgnore, bool bUseFilter)
+{
+	bool bHit = false;
+	FHitResult MeshHitResult;
+	if (MeshRayCast && MeshRayCast(WorldRay, MeshHitResult, bUseFilter))
+	{
+		bHit = true;
+		LastHitResult = MeshHitResult;
+		OutputPoint = LastHitResult.ImpactPoint;
+	}
+
+	FHitResult WorldHitResult;
+	if (WorldRayCast && WorldRayCast(WorldRay, WorldHitResult, ComponentsToIgnore) &&
+		(!bHit || WorldHitResult.Distance < MeshHitResult.Distance))
+	{
+		bHit = true;
+		LastHitResult = WorldHitResult;
+		OutputPoint = LastHitResult.ImpactPoint;
+	}
+
+	// We keep track of the last position we served the user, and the position
+	// that they moved a transform. When we have both, we can render a visualization
+	// connecting the two. Right now we (may) have the first of the two.
+	bWaitingOnProjectedResult = bHit;
+	bPreviewEndpointsValid = false;
+	return bHit;
+}
+
+void UDragAlignmentMechanic::Render(IToolsContextRenderAPI* RenderAPI)
+{
+	// Cache the camera state
+	GetParentTool()->GetToolManager()->GetContextQueriesAPI()->GetCurrentViewState(CachedCameraState);
+
+	if (!bPreviewEndpointsValid)
+	{
+		return;
+	}
+
+	FColor Color = FColor::Orange;
+	FViewCameraState CameraState = RenderAPI->GetCameraState();
+	FPrimitiveDrawInterface* PDI = RenderAPI->GetPrimitiveDrawInterface();
+
+	float PDIScale = CameraState.GetPDIScalingFactor();
+	double VisualAngleRenderThreshold = 0.5;
+
+	FToolDataVisualizer Renderer;
+	Renderer.BeginFrame(RenderAPI, CameraState);
+
+	if (LastHitResult.Item == TRIANGLE_HIT_TYPE_ID)
+	{
+		// We'll draw a circle with an X in it lying on the face we hit.
+		FFrame3d HitFrame = FFrame3d(LastHitResult.ImpactPoint, LastHitResult.ImpactNormal);
+
+		float TickLength = (float)ToolSceneQueriesUtil::CalculateDimensionFromVisualAngleD(CameraState, HitFrame.Origin, 0.8);
+		float TickThickness = 4 * PDIScale;
+
+		Renderer.DrawLine(
+			HitFrame.PointAt(-TickLength, -TickLength, 0),
+			HitFrame.PointAt(TickLength, TickLength, 0), Color, TickThickness, false);
+		Renderer.DrawLine(
+			HitFrame.PointAt(-TickLength, TickLength, 0),
+			HitFrame.PointAt(TickLength, -TickLength, 0), Color, TickThickness, false);
+
+		Renderer.DrawCircle(
+			HitFrame.Origin, HitFrame.Z(), 2.0 * TickLength, 8, Color, 1.0, false);
+	}
+	else
+	{
+		PDI->DrawPoint(LastHitResult.ImpactPoint, Color, 10.0f * PDIScale, SDPG_Foreground);
+	}
+
+	// If the two endpoints are far enough apart, draw a line connecting them.
+	if (ToolSceneQueriesUtil::CalculateNormalizedViewVisualAngleD(
+		CameraState, LastHitResult.ImpactPoint, LastProjectedResult) > VisualAngleRenderThreshold)
+	{
+		PDI->DrawLine(LastHitResult.ImpactPoint, (FVector)LastProjectedResult, Color,
+			SDPG_Foreground, 0.5f * PDIScale, 0.0f, true);
+	}
+
+	Renderer.EndFrame();
+}
+
+void UDragAlignmentMechanic::OnUpdateModifierState(int ModifierID, bool bIsOn)
+{
+	if (ModifierID == AlignmentModifierID)
+	{
+		bAlignmentToggle = bIsOn;
+	}
+	if (!bAlignmentToggle)
+	{
+		bPreviewEndpointsValid = false;
+	}
+}
diff --git a/Engine/Plugins/Experimental/MeshModelingToolset/Source/ModelingComponents/Private/ToolSceneQueriesUtil.cpp b/Engine/Plugins/Experimental/MeshModelingToolset/Source/ModelingComponents/Private/ToolSceneQueriesUtil.cpp
--- a/Engine/Plugins/Experimental/MeshModelingToolset/Source/ModelingComponents/Private/ToolSceneQueriesUtil.cpp
+++ b/Engine/Plugins/Experimental/MeshModelingToolset/Source/ModelingComponents/Private/ToolSceneQueriesUtil.cpp
@@ -135,54 +135,70 @@ bool ToolSceneQueriesUtil::IsPointVisible(const FViewCameraState& CameraState, c
 	return true;
 }
 
-
 bool ToolSceneQueriesUtil::FindSceneSnapPoint(const UInteractiveTool* Tool, const FVector3d& Point, FVector3d& SnapPointOut,
-	bool bVertices, bool bEdges, double VisualAngleThreshold, 
+	bool bVertices, bool bEdges, double VisualAngleThreshold,
 	FSnapGeometry* SnapGeometry, FVector* DebugTriangleOut)
 {
-	double UseThreshold = (VisualAngleThreshold <= 0) ? GetDefaultVisualAngleSnapThreshD() : VisualAngleThreshold;
+	FFindSceneSnapPointParams Params;
+
+	Params.Tool = Tool;
+	Params.Point = &Point;
+	Params.SnapPointOut = &SnapPointOut;
+	Params.bVertices = bVertices;
+	Params.bEdges = bEdges;
+	Params.VisualAngleThreshold = VisualAngleThreshold;
+	Params.SnapGeometryOut = SnapGeometry;
+	Params.DebugTriangleOut = DebugTriangleOut;
+
+	return FindSceneSnapPoint(Params);
+}
+
+bool ToolSceneQueriesUtil::FindSceneSnapPoint(FFindSceneSnapPointParams& Params)
+{
+	double UseThreshold = (Params.VisualAngleThreshold <= 0) ? GetDefaultVisualAngleSnapThreshD() : Params.VisualAngleThreshold;
 
 	FViewCameraState CameraState;
-	Tool->GetToolManager()->GetContextQueriesAPI()->GetCurrentViewState(CameraState);
+	Params.Tool->GetToolManager()->GetContextQueriesAPI()->GetCurrentViewState(CameraState);
 	UseThreshold *= CameraState.GetFOVAngleNormalizationFactor();
 
-	IToolsContextQueriesAPI* QueryAPI = Tool->GetToolManager()->GetContextQueriesAPI();
+	IToolsContextQueriesAPI* QueryAPI = Params.Tool->GetToolManager()->GetContextQueriesAPI();
 	FSceneSnapQueryRequest Request;
 	Request.RequestType = ESceneSnapQueryType::Position;
 	Request.TargetTypes = ESceneSnapQueryTargetType::None;
-	if (bVertices)
+	if (Params.bVertices)
 	{
 		Request.TargetTypes |= ESceneSnapQueryTargetType::MeshVertex;
 	}
-	if (bEdges)
+	if (Params.bEdges)
 	{
 		Request.TargetTypes |= ESceneSnapQueryTargetType::MeshEdge;
 	}
-	Request.Position = (FVector)Point;
+	Request.Position = (FVector)*Params.Point;
 	Request.VisualAngleThresholdDegrees = UseThreshold;
-	
+	Request.ComponentsToIgnore = Params.ComponentsToIgnore;
+
 	TArray<FSceneSnapQueryResult> Results;
 	if (QueryAPI->ExecuteSceneSnapQuery(Request, Results))
 	{
-		SnapPointOut = Results[0].Position;
+		*Params.SnapPointOut = Results[0].Position;
 
-		if (SnapGeometry != nullptr)
+		if (Params.SnapGeometryOut != nullptr)
 		{
 			int iSnap = Results[0].TriSnapIndex;
-			SnapGeometry->Points[0] = Results[0].TriVertices[iSnap];
-			SnapGeometry->PointCount = 1;
+			Params.SnapGeometryOut->Points[0] = Results[0].TriVertices[iSnap];
+			Params.SnapGeometryOut->PointCount = 1;
 			if (Results[0].TargetType == ESceneSnapQueryTargetType::MeshEdge)
 			{
-				SnapGeometry->Points[1] = Results[0].TriVertices[(iSnap+1)%3];
-				SnapGeometry->PointCount = 2;
+				Params.SnapGeometryOut->Points[1] = Results[0].TriVertices[(iSnap+1)%3];
+				Params.SnapGeometryOut->PointCount = 2;
 			}
 		}
 
-		if (DebugTriangleOut != nullptr)
+		if (Params.DebugTriangleOut != nullptr)
 		{
-			DebugTriangleOut[0] = Results[0].TriVertices[0];
-			DebugTriangleOut[1] = Results[0].TriVertices[1];
-			DebugTriangleOut[2] = Results[0].TriVertices[2];
+			Params.DebugTriangleOut[0] = Results[0].TriVertices[0];
+			Params.DebugTriangleOut[1] = Results[0].TriVertices[1];
+			Params.DebugTriangleOut[2] = Results[0].TriVertices[2];
 		}
 
 		return true;
@@ -252,7 +268,7 @@ bool ToolSceneQueriesUtil::IsVisibleObjectHit(const FHitResult& HitResult)
 
 
 bool ToolSceneQueriesUtil::FindNearestVisibleObjectHit(UWorld* World, FHitResult& HitResultOut, const FVector& Start, const FVector& End,
-	const TArray<UPrimitiveComponent*>* IgnoreComponents, const TArray<UPrimitiveComponent*>* InvisibleComponentsToInclude)
+	const TArray<const UPrimitiveComponent*>* IgnoreComponents, const TArray<const UPrimitiveComponent*>* InvisibleComponentsToInclude)
 {
 	FCollisionObjectQueryParams ObjectQueryParams(FCollisionObjectQueryParams::AllObjects);
 	FCollisionQueryParams QueryParams = FCollisionQueryParams::DefaultQueryParam;
@@ -286,7 +302,7 @@ bool ToolSceneQueriesUtil::FindNearestVisibleObjectHit(UWorld* World, FHitResult
 
 
 bool ToolSceneQueriesUtil::FindNearestVisibleObjectHit(UWorld* World, FHitResult& HitResultOut, const FRay& Ray,
-	const TArray<UPrimitiveComponent*>* IgnoreComponents, const TArray<UPrimitiveComponent*>* InvisibleComponentsToInclude)
+	const TArray<const UPrimitiveComponent*>* IgnoreComponents, const TArray<const UPrimitiveComponent*>* InvisibleComponentsToInclude)
 {
 	return FindNearestVisibleObjectHit(World, HitResultOut, Ray.Origin, Ray.PointAt(HALF_WORLD_MAX), IgnoreComponents, InvisibleComponentsToInclude);
 }
\ No newline at end of file
diff --git a/Engine/Plugins/Experimental/MeshModelingToolset/Source/ModelingComponents/Private/Transforms/MultiTransformer.cpp b/Engine/Plugins/Experimental/MeshModelingToolset/Source/ModelingComponents/Private/Transforms/MultiTransformer.cpp
--- a/Engine/Plugins/Experimental/MeshModelingToolset/Source/ModelingComponents/Private/Transforms/MultiTransformer.cpp
+++ b/Engine/Plugins/Experimental/MeshModelingToolset/Source/ModelingComponents/Private/Transforms/MultiTransformer.cpp
@@ -24,6 +24,10 @@ void UMultiTransformer::Setup(UInteractiveGizmoManager* GizmoManagerIn, IToolCon
 	TransformProxy->OnTransformChanged.AddUObject(this, &UMultiTransformer::OnProxyTransformChanged);
 	TransformProxy->OnBeginTransformEdit.AddUObject(this, &UMultiTransformer::OnBeginProxyTransformEdit);
 	TransformProxy->OnEndTransformEdit.AddUObject(this, &UMultiTransformer::OnEndProxyTransformEdit);
+	TransformProxy->OnPivotChanged.AddWeakLambda(this, [this](UTransformProxy* Proxy, FTransform Transform) {
+		ActiveGizmoFrame = FFrame3d(Transform);
+		ActiveGizmoScale = FVector3d(Transform.GetScale3D());
+	});
 }
 
 
@@ -90,11 +94,57 @@ void UMultiTransformer::SetGizmoVisibility(bool bVisible)
 	}
 }
 
+void UMultiTransformer::SetGizmoRepositionable(bool bOn)
+{
+	if (bRepositionableGizmo != bOn)
+	{
+		bRepositionableGizmo = bOn;
+		if (TransformGizmo)
+		{
+			UpdateShowGizmoState(false);
+			UpdateShowGizmoState(true);
+		}
+	}
+}
+
+EToolContextCoordinateSystem UMultiTransformer::GetGizmoCoordinateSystem()
+{ 
+	return TransformGizmo ? TransformGizmo->CurrentCoordinateSystem : GizmoCoordSystem; 
+}
+
+
 void UMultiTransformer::SetSnapToWorldGridSourceFunc(TUniqueFunction<bool()> EnableSnapFunc)
 {
 	EnableSnapToWorldGridFunc = MoveTemp(EnableSnapFunc);
 }
 
+void UMultiTransformer::SetIsNonUniformScaleAllowedFunction(TFunction<bool()> IsNonUniformScaleAllowedIn)
+{
+	IsNonUniformScaleAllowed = IsNonUniformScaleAllowedIn;
+	if (TransformGizmo != nullptr)
+	{
+		TransformGizmo->SetIsNonUniformScaleAllowedFunction(IsNonUniformScaleAllowed);
+	}
+}
+
+void UMultiTransformer::SetDisallowNegativeScaling(bool bDisallow)
+{
+	bDisallowNegativeScaling = bDisallow;
+	if (TransformGizmo != nullptr)
+	{
+		TransformGizmo->SetDisallowNegativeScaling(bDisallowNegativeScaling);
+	}
+}
+
+void UMultiTransformer::AddAlignmentMechanic(UDragAlignmentMechanic* AlignmentMechanic)
+{
+	DragAlignmentMechanic = AlignmentMechanic;
+	if (TransformGizmo)
+	{
+		DragAlignmentMechanic->AddToGizmo(TransformGizmo);
+	}
+}
+
 void UMultiTransformer::Tick(float DeltaTime)
 {
 	if (TransformGizmo != nullptr)
@@ -181,15 +231,24 @@ void UMultiTransformer::UpdateShowGizmoState(bool bNewVisibility)
 	else
 	{
 		check(TransformGizmo == nullptr);
-		TransformGizmo = GizmoManager->CreateCustomTransformGizmo(ActiveGizmoSubElements, this);
+		TransformGizmo = bRepositionableGizmo ? GizmoManager->CreateCustomRepositionableTransformGizmo(ActiveGizmoSubElements, this)
+			: GizmoManager->CreateCustomTransformGizmo(ActiveGizmoSubElements, this);
 		if (bForceGizmoCoordSystem)
 		{
 			TransformGizmo->bUseContextCoordinateSystem = false;
 			TransformGizmo->CurrentCoordinateSystem = GizmoCoordSystem;
 		}
+		if (IsNonUniformScaleAllowed)
+		{
+			TransformGizmo->SetIsNonUniformScaleAllowedFunction(IsNonUniformScaleAllowed);
+		}
+		TransformGizmo->SetDisallowNegativeScaling(bDisallowNegativeScaling);
+		if (DragAlignmentMechanic)
+		{
+			DragAlignmentMechanic->AddToGizmo(TransformGizmo);
+		}
 		TransformGizmo->SetActiveTarget(TransformProxy, TransactionProvider);
 		TransformGizmo->ReinitializeGizmoTransform(ActiveGizmoFrame.ToFTransform());
 		TransformGizmo->SetVisibility(bShouldBeVisible);
 	}
 }
-
diff --git a/Engine/Plugins/Experimental/MeshModelingToolset/Source/ModelingComponents/Public/Mechanics/DragAlignmentMechanic.h b/Engine/Plugins/Experimental/MeshModelingToolset/Source/ModelingComponents/Public/Mechanics/DragAlignmentMechanic.h
new file mode 100644
--- /dev/null
+++ b/Engine/Plugins/Experimental/MeshModelingToolset/Source/ModelingComponents/Public/Mechanics/DragAlignmentMechanic.h
@@ -0,0 +1,91 @@
+// Copyright Epic Games, Inc. All Rights Reserved.
+
+#pragma once
+
+#include "CoreMinimal.h"
+
+#include "BaseBehaviors/BehaviorTargetInterfaces.h"
+#include "DynamicMeshAABBTree3.h"
+#include "ToolContextInterfaces.h" // FViewCameraState
+#include "Engine/EngineTypes.h" //FHitResult
+#include "InteractionMechanic.h"
+#include "TransformTypes.h"
+#include "VectorTypes.h"
+
+#include "DragAlignmentMechanic.generated.h"
+
+class FGroupTopologyDeformer;
+class UPrimitiveComponent;
+class UTransformGizmo;
+
+/**
+ * Mechanic that can be added to (potentially multiple) UTransformGizmo object to allow them to snap to
+ * scene geometry in rotation and translation when the Ctrl key is pressed.
+ */
+UCLASS()
+class MODELINGCOMPONENTS_API UDragAlignmentMechanic : public UInteractionMechanic, public IModifierToggleBehaviorTarget
+{
+	GENERATED_BODY()
+
+public:
+
+	/**
+	 * Adds this mechanic to the given gizmo. Can be called on multiple gizmos, and works both for UTransformGizmo
+	 * and URepositionableTransformGizmo.
+	 * For repositioning, the gizmo will not ignore anything in its alignment (so that the gizmo can be aligned to the
+	 * target object(s)). For movement, the gizmo will ignore ComponentsToIgnore, as well as modified triangles if used
+	 * for deformation.
+	 */
+	void AddToGizmo(UTransformGizmo* TransformGizmo, const TArray<const UPrimitiveComponent*>* ComponentsToIgnore = nullptr);
+
+	/**
+	 * Optional- allows the use of the mechanic in mesh deformation. Specifically, alignment will attempt to align
+	 * to the dynamic mesh underlying the given FDynamicMeshAABBTree3, and ignore deformed triangles when the gizmo
+	 * is moved.
+	 * This works with only one deformed mesh, so it probably wouldn't be used with multiple gizmos.
+	 */
+	void InitializeDeformedMeshRayCast(
+		TFunction<FDynamicMeshAABBTree3* ()> GetSpatialIn,
+		const FTransform3d& TargetTransform, const FGroupTopologyDeformer* LinearDeformer);
+
+	virtual void Setup(UInteractiveTool* ParentToolIn) override;
+	virtual void Shutdown() override;
+
+	/**
+	 * This will draw a line from the last input position to the last output position. This
+	 * gets cleared when the transform is ended. 
+	 */
+	virtual void Render(IToolsContextRenderAPI* RenderAPI) override;
+
+	// IModifierToggleBehaviorTarget implementation
+	virtual void OnUpdateModifierState(int ModifierID, bool bIsOn) override;
+
+protected:
+	// These are modifiable in case we someday want to make it easier to customize them further.
+	TUniqueFunction<bool(const FRay& InputRay, FHitResult& OutputHit, const TArray<const UPrimitiveComponent*>* ComponentsToIgnore)> WorldRayCast = nullptr;
+	TUniqueFunction<bool(const FRay&, FHitResult&, bool bUseFilter)> MeshRayCast = nullptr;
+	TUniqueFunction<bool(int32 Tid)> MeshRayCastTriangleFilter = [](int32 Tid) { return true; };
+
+	// Used to identify the type of object we hit, so we can choose different renderings.
+	static const int32 TRIANGLE_HIT_TYPE_ID = 1;
+	static const int32 VERT_HIT_TYPE_ID = 2;
+	static const int32 EDGE_HIT_TYPE_ID = 3;
+
+	// Used for snapping to verts
+	FViewCameraState CachedCameraState;
+	float VisualAngleSnapThreshold = 1.0;
+
+	// We keep track of the last position we gave the user and the last position that
+	// they moved their gizmo. This lets us render lines connecting the two.
+	bool bPreviewEndpointsValid = false;
+	bool bWaitingOnProjectedResult = false;
+	FHitResult LastHitResult;
+	FVector3d LastProjectedResult;
+
+	bool bAlignmentToggle = false;
+	int32 AlignmentModifierID = 1;
+
+	virtual bool CastRay(const FRay& WorldRay, FVector& OutputPoint,
+		const TArray<const UPrimitiveComponent*>* ComponentsToIgnore, bool bUseFilter);
+	void OnGizmoTransformChanged(FTransform NewTransform);
+};
\ No newline at end of file
diff --git a/Engine/Plugins/Experimental/MeshModelingToolset/Source/ModelingComponents/Public/Selection/SelectClickedAction.h b/Engine/Plugins/Experimental/MeshModelingToolset/Source/ModelingComponents/Public/Selection/SelectClickedAction.h
--- a/Engine/Plugins/Experimental/MeshModelingToolset/Source/ModelingComponents/Public/Selection/SelectClickedAction.h
+++ b/Engine/Plugins/Experimental/MeshModelingToolset/Source/ModelingComponents/Public/Selection/SelectClickedAction.h
@@ -36,8 +36,8 @@ public:
 
 	// These lists can be used to modify which components are hit tested when doing a ray cast.
 	// By default, all visible components are hit tested.
-	TArray<UPrimitiveComponent*> VisibleComponentsToIgnore;
-	TArray<UPrimitiveComponent*> InvisibleComponentsToHitTest;
+	TArray<const UPrimitiveComponent*> VisibleComponentsToIgnore;
+	TArray<const UPrimitiveComponent*> InvisibleComponentsToHitTest;
 
 	// can alternately track shift modifier, however client must register this modifier w/ behavior
 	static const int ShiftModifier = 1;
diff --git a/Engine/Plugins/Experimental/MeshModelingToolset/Source/ModelingComponents/Public/ToolSceneQueriesUtil.h b/Engine/Plugins/Experimental/MeshModelingToolset/Source/ModelingComponents/Public/ToolSceneQueriesUtil.h
--- a/Engine/Plugins/Experimental/MeshModelingToolset/Source/ModelingComponents/Public/ToolSceneQueriesUtil.h
+++ b/Engine/Plugins/Experimental/MeshModelingToolset/Source/ModelingComponents/Public/ToolSceneQueriesUtil.h
@@ -90,6 +90,25 @@ namespace ToolSceneQueriesUtil
 		int PointCount = 0;
 	};
 
+	struct FFindSceneSnapPointParams
+	{
+		// Required inputs/outputs
+		const UInteractiveTool* Tool = nullptr;
+		const FVector3d* Point = nullptr;
+		FVector3d* SnapPointOut = nullptr;
+
+		bool bVertices = true; // If true, try to snap to vertices
+		bool bEdges = false; // If true, try to snap to edges
+		double VisualAngleThreshold = 0; // Visual angle threshold to use. If 0, GetDefaultVisualAngleSnapThresh() is used
+		const TArray<const UPrimitiveComponent*>* ComponentsToIgnore = nullptr;
+		FSnapGeometry* SnapGeometryOut = nullptr; // world-space position of the snap geometry(point / line / polygon)
+		FVector* DebugTriangleOut = nullptr; // if non - null, triangle containing snap is returned if a snap is found
+	};
+
+	/**
+	 * Run a query against the scene to find the best SnapPointOut for the given Point
+	 */
+	MODELINGCOMPONENTS_API bool FindSceneSnapPoint(FFindSceneSnapPointParams& Params);
 
 	/**
 	 * Run a query against the scene to find the best SnapPointOut for the given Point
@@ -104,13 +123,11 @@ namespace ToolSceneQueriesUtil
 		bool bVertices = true, bool bEdges = false, double VisualAngleThreshold = 0, 
 		FSnapGeometry* SnapGeometry = nullptr, FVector* DebugTriangleOut = nullptr);
 
-
 	/**
 	 * Run a query against the scene to find the nearest WorldGrid snap point
 	 */
 	MODELINGCOMPONENTS_API bool FindWorldGridSnapPoint(const UInteractiveTool* Tool, const FVector3d& QueryPoint, FVector3d& GridSnapPointOut);
 
-
 	/**
 	 * @return true if HitResult is a hit on a visible Component of a visible Actor (provides correct result in Editor)
 	 */
@@ -126,8 +143,8 @@ namespace ToolSceneQueriesUtil
 	 * @return true if a visible hit was found
 	 */
 	MODELINGCOMPONENTS_API bool FindNearestVisibleObjectHit(UWorld* World, FHitResult& HitResultOut, const FVector& Start, const FVector& End,
-		const TArray<UPrimitiveComponent*>* IgnoreComponents = nullptr, 
-		const TArray<UPrimitiveComponent*>* InvisibleComponentsToInclude = nullptr);
+		const TArray<const UPrimitiveComponent*>* IgnoreComponents = nullptr, 
+		const TArray<const UPrimitiveComponent*>* InvisibleComponentsToInclude = nullptr);
 
 
 	/**
@@ -139,7 +156,7 @@ namespace ToolSceneQueriesUtil
 	 * @return true if a visible hit was found
 	 */
 	MODELINGCOMPONENTS_API bool FindNearestVisibleObjectHit(UWorld* World, FHitResult& HitResultOut, const FRay& Ray,
-		const TArray<UPrimitiveComponent*>* IgnoreComponents = nullptr, 
-		const TArray<UPrimitiveComponent*>* InvisibleComponentsToInclude = nullptr);
+		const TArray<const UPrimitiveComponent*>* IgnoreComponents = nullptr, 
+		const TArray<const UPrimitiveComponent*>* InvisibleComponentsToInclude = nullptr);
 
 }
\ No newline at end of file
diff --git a/Engine/Plugins/Experimental/MeshModelingToolset/Source/ModelingComponents/Public/Transforms/MultiTransformer.h b/Engine/Plugins/Experimental/MeshModelingToolset/Source/ModelingComponents/Public/Transforms/MultiTransformer.h
--- a/Engine/Plugins/Experimental/MeshModelingToolset/Source/ModelingComponents/Public/Transforms/MultiTransformer.h
+++ b/Engine/Plugins/Experimental/MeshModelingToolset/Source/ModelingComponents/Public/Transforms/MultiTransformer.h
@@ -3,12 +3,15 @@
 #pragma once
 
 #include "CoreMinimal.h"
+
 #include "ToolContextInterfaces.h"
 #include "ToolDataVisualizer.h"
 #include "FrameTypes.h"
 #include "InteractiveGizmo.h"
+#include "Mechanics/DragAlignmentMechanic.h"
 #include "BaseGizmos/TransformGizmo.h"
 #include "InteractiveGizmoManager.h"
+
 #include "MultiTransformer.generated.h"
 
 
@@ -52,11 +55,22 @@ public:
 
 	virtual void SetEnabledGizmoSubElements(ETransformGizmoSubElements EnabledSubElements);
 
+	virtual void SetGizmoRepositionable(bool bOn);
+
+	virtual EToolContextCoordinateSystem GetGizmoCoordinateSystem();
+
 	void SetSnapToWorldGridSourceFunc(TUniqueFunction<bool()> EnableSnapFunc);
+	void SetIsNonUniformScaleAllowedFunction(TFunction<bool()> IsNonUniformScaleAllowedIn);
+
+	void SetDisallowNegativeScaling(bool bDisallow);
+
+	void AddAlignmentMechanic(UDragAlignmentMechanic* AlignmentMechanic);
 
 public:
 	DECLARE_MULTICAST_DELEGATE(FMultiTransformerEvent);
 
+	// Note that the following delegates don't fire on pivot repositioning drags.
+
 	/** This delegate is fired when a drag is started */
 	FMultiTransformerEvent OnTransformStarted;
 
@@ -66,8 +80,7 @@ public:
 	/** This delegate is fired when the drag is completed */
 	FMultiTransformerEvent OnTransformCompleted;
 
-
-public:
+protected:
 	UPROPERTY()
 	UInteractiveGizmoManager* GizmoManager;
 
@@ -84,13 +97,23 @@ public:
 	FFrame3d ActiveGizmoFrame;
 	FVector3d ActiveGizmoScale;
 
+	bool bRepositionableGizmo = false;
+
+	bool bDisallowNegativeScaling = false;
+
 	UPROPERTY()
 	UTransformGizmo* TransformGizmo;
 
 	UPROPERTY()
 	UTransformProxy* TransformProxy;
 
+	// We have to hold on to the mechanic only because the MultiTransformer has the capacity to delete and
+	// recreate its gizmo, in which case we'll need to attach the alignment mechanic again.
+	UPROPERTY()
+	UDragAlignmentMechanic* DragAlignmentMechanic = nullptr;
+
 	TUniqueFunction<bool()> EnableSnapToWorldGridFunc;
+	TFunction<bool()> IsNonUniformScaleAllowed;
 
 	// called on PlaneTransformProxy.OnTransformChanged
 	void OnProxyTransformChanged(UTransformProxy* Proxy, FTransform Transform);
@@ -99,6 +122,4 @@ public:
 	bool bInGizmoEdit = false;
 
 	void UpdateShowGizmoState(bool bNewVisibility);
-
-
 };
\ No newline at end of file
diff --git a/Engine/Source/Editor/Experimental/EditorInteractiveToolsFramework/Private/EdModeInteractiveToolsContext.cpp b/Engine/Source/Editor/Experimental/EditorInteractiveToolsFramework/Private/EdModeInteractiveToolsContext.cpp
--- a/Engine/Source/Editor/Experimental/EditorInteractiveToolsFramework/Private/EdModeInteractiveToolsContext.cpp
+++ b/Engine/Source/Editor/Experimental/EditorInteractiveToolsFramework/Private/EdModeInteractiveToolsContext.cpp
@@ -71,7 +71,8 @@ static bool IsVisibleObjectHit_Internal(const FHitResult& HitResult)
 	return true;
 }
 
-static bool FindNearestVisibleObjectHit_Internal(UWorld* World, FHitResult& HitResultOut, const FVector& Start, const FVector& End, bool bIsSceneGeometrySnapQuery)
+static bool FindNearestVisibleObjectHit_Internal(UWorld* World, FHitResult& HitResultOut, const FVector& Start, const FVector& End,
+	bool bIsSceneGeometrySnapQuery, const TArray<const UPrimitiveComponent*>* ComponentsToIgnore)
 {
 	FCollisionObjectQueryParams ObjectQueryParams(FCollisionObjectQueryParams::AllObjects);
 	FCollisionQueryParams QueryParams = FCollisionQueryParams::DefaultQueryParam;
@@ -89,7 +90,8 @@ static bool FindNearestVisibleObjectHit_Internal(UWorld* World, FHitResult& HitR
 	{
 		if (CurResult.Distance < NearestVisible)
 		{
-			if (IsVisibleObjectHit_Internal(CurResult))
+			if (IsVisibleObjectHit_Internal(CurResult) 
+				&& (!ComponentsToIgnore || !ComponentsToIgnore->Contains(CurResult.Component.Get())))
 			{
 				HitResultOut = CurResult;
 				NearestVisible = CurResult.Distance;
@@ -249,7 +251,8 @@ public:
 		FVector RayDirection = Request.Position - RayStart; RayDirection.Normalize();
 		FVector RayEnd = RayStart + HALF_WORLD_MAX * RayDirection;
 		FHitResult HitResult;
-		bool bHitWorld = FindNearestVisibleObjectHit_Internal(EditorModeManager->GetWorld(), HitResult, RayStart, RayEnd, true);
+		bool bHitWorld = FindNearestVisibleObjectHit_Internal(EditorModeManager->GetWorld(), HitResult, 
+			RayStart, RayEnd, true, Request.ComponentsToIgnore);
 		if (bHitWorld && HitResult.FaceIndex >= 0)
 		{
 			float VisualAngle = OpeningAngleDeg(Request.Position, HitResult.ImpactPoint, RayStart);
@@ -992,14 +995,6 @@ bool UEdModeInteractiveToolsContext::InputKey(FEditorViewportClient* ViewportCli
 				{
 					return false;
 				}
-				// TODO: This should no longer be necessary: test and remove.
-				// This is a special-case hack for UMultiClickSequenceInputBehavior, because it holds capture across multiple
-				// mouse clicks, which prevents alt+mouse navigation from working between clicks (very annoying in draw polygon).
-				// Remove this special-case once that tool is fixed to use CollectSurfacePathMechanic instead
-				if (Event == IE_Pressed && bIsLeftMouse && ViewportClient->IsAltPressed() && InputRouter->HasActiveMouseCapture())
-				{
-					return false;
-				}
 
 				FInputDeviceState InputState = CurrentMouseState;
 				InputState.InputDevice = EInputDevices::Mouse;
diff --git a/Engine/Source/Runtime/Experimental/InteractiveToolsFramework/Private/BaseGizmos/AxisAngleGizmo.cpp b/Engine/Source/Runtime/Experimental/InteractiveToolsFramework/Private/BaseGizmos/AxisAngleGizmo.cpp
--- a/Engine/Source/Runtime/Experimental/InteractiveToolsFramework/Private/BaseGizmos/AxisAngleGizmo.cpp
+++ b/Engine/Source/Runtime/Experimental/InteractiveToolsFramework/Private/BaseGizmos/AxisAngleGizmo.cpp
@@ -22,7 +22,7 @@ void UAxisAngleGizmo::Setup()
 	UInteractiveGizmo::Setup();
 
 	// Add default mouse input behavior
-	UClickDragInputBehavior* MouseBehavior = NewObject<UClickDragInputBehavior>();
+	MouseBehavior = NewObject<UClickDragInputBehavior>();
 	MouseBehavior->Modifiers.RegisterModifier(SnapAngleModifierID, FInputDeviceState::IsShiftKeyDown);
 	MouseBehavior->Initialize(this);
 	MouseBehavior->SetDefaultPriority(FInputCapturePriority(FInputCapturePriority::DEFAULT_GIZMO_PRIORITY));
@@ -68,8 +68,8 @@ void UAxisAngleGizmo::OnClickPress(const FInputDeviceRay& PressPos)
 {
 	check(bInInteraction == false);
 
-	RotationOrigin = GizmoMath::ProjectPointOntoLine(LastHitPosition, AxisSource->GetOrigin(), RotationAxis);
 	RotationAxis = AxisSource->GetDirection();
+	RotationOrigin = GizmoMath::ProjectPointOntoLine(LastHitPosition, AxisSource->GetOrigin(), RotationAxis);
 
 	if (AxisSource->HasTangentVectors())
 	{
@@ -94,6 +94,12 @@ void UAxisAngleGizmo::OnClickPress(const FInputDeviceRay& PressPos)
 		RotationOrigin, RotationAxis, RotationPlaneX, RotationPlaneY);
 	InteractionCurAngle = InteractionStartAngle;
 
+	// Figure out the angle of the closest axis so that we know what to snap if using alignment.
+	float StartAngleAbs = FMath::Abs(InteractionStartAngle);
+	ClosestAxisStartAngle = StartAngleAbs <= PI / 4 ? 0
+		: StartAngleAbs >= PI * 3 / 4 ? PI
+		: InteractionStartAngle > 0 ? PI / 2 : - PI / 2;
+
 	InitialTargetAngle = AngleSource->GetParameter();
 	AngleSource->BeginModify();
 
@@ -111,29 +117,48 @@ void UAxisAngleGizmo::OnClickDrag(const FInputDeviceRay& DragPos)
 {
 	check(bInInteraction);
 
-	bool bIntersects; FVector IntersectionPoint;
-	GizmoMath::RayPlaneIntersectionPoint(
-		RotationOrigin, RotationAxis,
-		DragPos.WorldRay.Origin, DragPos.WorldRay.Direction,
-		bIntersects, IntersectionPoint);
+	FVector HitPoint;
+	float DeltaAngle;
 
-	if (bIntersects == false)
+	// See if we should use custom destination funtion
+	FCustomDestinationParams Params;
+	Params.WorldRay = &DragPos.WorldRay;
+	if (ShouldUseCustomDestinationFunc() && CustomDestinationFunc(Params, HitPoint))
 	{
-		return;
+		InteractionCurPoint = HitPoint;
+		InteractionCurAngle = GizmoMath::ComputeAngleInPlane(InteractionCurPoint,
+			RotationOrigin, RotationAxis, RotationPlaneX, RotationPlaneY);
+		
+		// Align nearest axis to destination point.
+		DeltaAngle = InteractionCurAngle - ClosestAxisStartAngle;
 	}
+	else
+	{
+		bool bIntersects; FVector IntersectionPoint;
+		GizmoMath::RayPlaneIntersectionPoint(
+			RotationOrigin, RotationAxis,
+			DragPos.WorldRay.Origin, DragPos.WorldRay.Direction,
+			bIntersects, IntersectionPoint);
 
-	InteractionCurPoint = IntersectionPoint;
+		if (bIntersects == false)
+		{
+			return;
+		}
 
-	InteractionCurAngle = GizmoMath::ComputeAngleInPlane(InteractionCurPoint,
-		RotationOrigin, RotationAxis, RotationPlaneX, RotationPlaneY);
+		InteractionCurPoint = IntersectionPoint;
 
-	float DeltaAngle = InteractionCurAngle - InteractionStartAngle;
-	if (bEnableSnapAngleModifier)
-	{
-		DeltaAngle = GizmoMath::SnapToIncrement(DeltaAngle, 3.14159f / 4.0f);				// hardcoded 45-degree snap
+		InteractionCurAngle = GizmoMath::ComputeAngleInPlane(InteractionCurPoint,
+			RotationOrigin, RotationAxis, RotationPlaneX, RotationPlaneY);
+
+		DeltaAngle = InteractionCurAngle - InteractionStartAngle;
+		if (bEnableSnapAngleModifier)
+		{
+			DeltaAngle = GizmoMath::SnapToIncrement(DeltaAngle, 3.14159f / 4.0f);				// hardcoded 45-degree snap
+		}
+		
 	}
-	float NewAngle = InitialTargetAngle + DeltaAngle;
 
+	float NewAngle = InitialTargetAngle + DeltaAngle;
 	AngleSource->SetParameter(NewAngle);
 }
 
diff --git a/Engine/Source/Runtime/Experimental/InteractiveToolsFramework/Private/BaseGizmos/AxisPositionGizmo.cpp b/Engine/Source/Runtime/Experimental/InteractiveToolsFramework/Private/BaseGizmos/AxisPositionGizmo.cpp
--- a/Engine/Source/Runtime/Experimental/InteractiveToolsFramework/Private/BaseGizmos/AxisPositionGizmo.cpp
+++ b/Engine/Source/Runtime/Experimental/InteractiveToolsFramework/Private/BaseGizmos/AxisPositionGizmo.cpp
@@ -22,7 +22,7 @@ void UAxisPositionGizmo::Setup()
 	UInteractiveGizmo::Setup();
 
 	// Add default mouse input behavior
-	UClickDragInputBehavior* MouseBehavior = NewObject<UClickDragInputBehavior>();
+	MouseBehavior = NewObject<UClickDragInputBehavior>();
 	MouseBehavior->Initialize(this);
 	MouseBehavior->SetDefaultPriority(FInputCapturePriority(FInputCapturePriority::DEFAULT_GIZMO_PRIORITY));
 	AddInputBehavior(MouseBehavior);
@@ -59,15 +59,25 @@ void UAxisPositionGizmo::OnClickPress(const FInputDeviceRay& PressPos)
 	InteractionOrigin = LastHitPosition;
 	InteractionAxis = AxisSource->GetDirection();
 
-	FVector RayNearestPt; float RayNearestParam;
+	// Find interaction start point and parameter.
+	FVector NearestPt; float RayNearestParam;
 	GizmoMath::NearestPointOnLineToRay(InteractionOrigin, InteractionAxis,
 		PressPos.WorldRay.Origin, PressPos.WorldRay.Direction,
 		InteractionStartPoint, InteractionStartParameter,
-		RayNearestPt, RayNearestParam);
+		NearestPt, RayNearestParam);
 
-	float DirectionSign = FVector::DotProduct(InteractionStartPoint - AxisSource->GetOrigin(), InteractionAxis);
+	FVector AxisOrigin = AxisSource->GetOrigin();
+
+	float DirectionSign = FVector::DotProduct(InteractionStartPoint - AxisOrigin, InteractionAxis);
 	ParameterSign = (bEnableSignedAxis && DirectionSign < 0) ? -1.0f : 1.0f;
 
+	// Figure out how the parameter would need to be adjusted to bring the axis origin to the
+	// interaction start point. This is used when aligning the axis origin to a custom destination.
+	float AxisOriginParamValue;
+	GizmoMath::NearestPointOnLine(InteractionOrigin, InteractionAxis, AxisOrigin,
+		NearestPt, AxisOriginParamValue);
+	InteractionStartAxisOriginParameterOffset = InteractionStartParameter - AxisOriginParamValue;
+
 	InteractionCurPoint = InteractionStartPoint;
 	InteractionStartParameter *= ParameterSign;
 	InteractionCurParameter = InteractionStartParameter;
@@ -81,24 +91,39 @@ void UAxisPositionGizmo::OnClickPress(const FInputDeviceRay& PressPos)
 	{
 		StateTarget->BeginUpdate();
 	}
+
+	OnClickDrag(PressPos);
 }
 
 void UAxisPositionGizmo::OnClickDrag(const FInputDeviceRay& DragPos)
 {
-	FVector AxisNearestPt; float AxisNearestParam;
-	FVector RayNearestPt; float RayNearestParam;
-	GizmoMath::NearestPointOnLineToRay(InteractionOrigin, InteractionAxis,
-		DragPos.WorldRay.Origin, DragPos.WorldRay.Direction,
-		AxisNearestPt, AxisNearestParam,
-		RayNearestPt, RayNearestParam);
+	FVector HitPoint;
 
-	InteractionCurPoint = AxisNearestPt;
-	InteractionCurParameter = ParameterSign * AxisNearestParam;
+	// See if we should use the custom destination function.
+	FCustomDestinationParams Params;
+	Params.WorldRay = &DragPos.WorldRay;
+	if (ShouldUseCustomDestinationFunc() && CustomDestinationFunc(Params, HitPoint))
+	{
+		GizmoMath::NearestPointOnLine(InteractionOrigin, InteractionAxis, HitPoint,
+			InteractionCurPoint, InteractionCurParameter);
+		InteractionCurParameter += InteractionStartAxisOriginParameterOffset;
+	}
+	else
+	{
+		float RayNearestParam; float AxisNearestParam;
+		FVector RayNearestPt; 
+		GizmoMath::NearestPointOnLineToRay(InteractionOrigin, InteractionAxis,
+			DragPos.WorldRay.Origin, DragPos.WorldRay.Direction,
+			InteractionCurPoint, AxisNearestParam,
+			RayNearestPt, RayNearestParam);
+
+		InteractionCurParameter = ParameterSign * AxisNearestParam;
+	}
 
 	float DeltaParam = InteractionCurParameter - InteractionStartParameter;
-	float NewValue = InitialTargetParameter + DeltaParam;
+	float NewParamValue = InitialTargetParameter + DeltaParam;
 
-	ParameterSource->SetParameter(NewValue);
+	ParameterSource->SetParameter(NewParamValue);
 }
 
 void UAxisPositionGizmo::OnClickRelease(const FInputDeviceRay& ReleasePos)
diff --git a/Engine/Source/Runtime/Experimental/InteractiveToolsFramework/Private/BaseGizmos/AxisSources.cpp b/Engine/Source/Runtime/Experimental/InteractiveToolsFramework/Private/BaseGizmos/AxisSources.cpp
--- a/Engine/Source/Runtime/Experimental/InteractiveToolsFramework/Private/BaseGizmos/AxisSources.cpp
+++ b/Engine/Source/Runtime/Experimental/InteractiveToolsFramework/Private/BaseGizmos/AxisSources.cpp
@@ -26,14 +26,20 @@ FVector UGizmoComponentAxisSource::GetDirection() const
 
 void UGizmoComponentAxisSource::GetTangentVectors(FVector& TangentXOut, FVector& TangentYOut) const
 {
+	// Note that in giving the tangent vectors, we need to give them in an order that
+	// preserves the handedness of our coordinate system so that we can use the function
+	// for generating planes in which we can measure rotation angles about the input axis.
+	// So, TangentXOut, TangentYOut, and Input(Z) should give a left-handed frame (since 
+	// unreal is left handed).
+
 	const FTransform& WorldTransform = Component->GetComponentToWorld();
 	TangentXOut = FVector(0, 1, 0);
 	TangentYOut = FVector(0, 0, 1);
 	int Index = FMath::Clamp(AxisIndex, 0, 2);
 	if (Index == 1)
 	{
-		TangentXOut = FVector(-1, 0, 0);
-		TangentYOut = FVector(0, 0, 1);
+		TangentXOut = FVector(0, 0, 1);
+		TangentYOut = FVector(1, 0, 0);
 	}
 	else if (Index == 2)
 	{
diff --git a/Engine/Source/Runtime/Experimental/InteractiveToolsFramework/Private/BaseGizmos/HitTargets.cpp b/Engine/Source/Runtime/Experimental/InteractiveToolsFramework/Private/BaseGizmos/HitTargets.cpp
--- a/Engine/Source/Runtime/Experimental/InteractiveToolsFramework/Private/BaseGizmos/HitTargets.cpp
+++ b/Engine/Source/Runtime/Experimental/InteractiveToolsFramework/Private/BaseGizmos/HitTargets.cpp
@@ -9,7 +9,7 @@
 FInputRayHit UGizmoComponentHitTarget::IsHit(const FInputDeviceRay& ClickPos) const
 {
 	FInputRayHit Hit;
-	if (Component)
+	if (Component && (!Condition || Condition(ClickPos)))
 	{
 		// if a gizmo is not visible it cannot be hit
 		bool bVisible = Component->IsVisible() && ( Component->GetOwner() && Component->GetOwner()->IsHidden() == false );
diff --git a/Engine/Source/Runtime/Experimental/InteractiveToolsFramework/Private/BaseGizmos/PlanePositionGizmo.cpp b/Engine/Source/Runtime/Experimental/InteractiveToolsFramework/Private/BaseGizmos/PlanePositionGizmo.cpp
--- a/Engine/Source/Runtime/Experimental/InteractiveToolsFramework/Private/BaseGizmos/PlanePositionGizmo.cpp
+++ b/Engine/Source/Runtime/Experimental/InteractiveToolsFramework/Private/BaseGizmos/PlanePositionGizmo.cpp
@@ -22,7 +22,7 @@ void UPlanePositionGizmo::Setup()
 	UInteractiveGizmo::Setup();
 
 	// Add default mouse input behavior
-	UClickDragInputBehavior* MouseBehavior = NewObject<UClickDragInputBehavior>();
+	MouseBehavior = NewObject<UClickDragInputBehavior>();
 	MouseBehavior->Initialize(this);
 	MouseBehavior->SetDefaultPriority(FInputCapturePriority(FInputCapturePriority::DEFAULT_GIZMO_PRIORITY));
 	AddInputBehavior(MouseBehavior);
@@ -75,15 +75,24 @@ void UPlanePositionGizmo::OnClickPress(const FInputDeviceRay& PressPos)
 
 	InteractionStartPoint = InteractionCurPoint = IntersectionPoint;
 
-	float DirectionSignX = FVector::DotProduct(InteractionStartPoint - AxisSource->GetOrigin(), InteractionAxisX);
+	FVector AxisOrigin = AxisSource->GetOrigin();
+
+	float DirectionSignX = FVector::DotProduct(InteractionStartPoint - AxisOrigin, InteractionAxisX);
 	ParameterSigns.X = (bEnableSignedAxis && DirectionSignX < 0) ? -1.0f : 1.0f;
 	ParameterSigns.X *= (bFlipX) ? -1.0f : 1.0f;
-	float DirectionSignY = FVector::DotProduct(InteractionStartPoint - AxisSource->GetOrigin(), InteractionAxisY);
+	float DirectionSignY = FVector::DotProduct(InteractionStartPoint - AxisOrigin, InteractionAxisY);
 	ParameterSigns.Y = (bEnableSignedAxis && DirectionSignY < 0) ? -1.0f : 1.0f;
 	ParameterSigns.Y *= (bFlipY) ? -1.0f : 1.0f;
 
 	InteractionStartParameter = GizmoMath::ComputeCoordinatesInPlane(IntersectionPoint,
 			InteractionOrigin, InteractionNormal, InteractionAxisX, InteractionAxisY);
+
+	// Figure out how the parameters would need to be adjusted to bring the axis origin to the
+	// interaction start point. This is used when aligning the axis origin to a custom destination.
+	FVector2D OriginParamValue = GizmoMath::ComputeCoordinatesInPlane(AxisOrigin,
+		InteractionOrigin, InteractionNormal, InteractionAxisX, InteractionAxisY);
+	InteractionStartOriginParameterOffset = InteractionStartParameter - OriginParamValue;
+
 	InteractionStartParameter.X *= ParameterSigns.X;
 	InteractionStartParameter.Y *= ParameterSigns.Y;
 	InteractionCurParameter = InteractionStartParameter;
@@ -101,27 +110,44 @@ void UPlanePositionGizmo::OnClickPress(const FInputDeviceRay& PressPos)
 
 void UPlanePositionGizmo::OnClickDrag(const FInputDeviceRay& DragPos)
 {
-	bool bIntersects; FVector IntersectionPoint;
-	GizmoMath::RayPlaneIntersectionPoint(
-		InteractionOrigin, InteractionNormal,
-		DragPos.WorldRay.Origin, DragPos.WorldRay.Direction,
-		bIntersects, IntersectionPoint);
+	FVector HitPoint;
+	FVector2D NewParamValue;
 
-	if (bIntersects == false)
+	// See if we should use the custom destination function.
+	FCustomDestinationParams Params;
+	Params.WorldRay = &DragPos.WorldRay;
+	if (ShouldUseCustomDestinationFunc() && CustomDestinationFunc(Params, HitPoint))
 	{
-		return;
+		InteractionCurPoint = GizmoMath::ProjectPointOntoPlane(HitPoint, InteractionOrigin, InteractionNormal);
+		InteractionCurParameter = GizmoMath::ComputeCoordinatesInPlane(InteractionCurPoint,
+			InteractionOrigin, InteractionNormal, InteractionAxisX, InteractionAxisY);
+
+		InteractionCurParameter += InteractionStartOriginParameterOffset;
 	}
+	else
+	{
+		bool bIntersects; 
+		GizmoMath::RayPlaneIntersectionPoint(
+			InteractionOrigin, InteractionNormal,
+			DragPos.WorldRay.Origin, DragPos.WorldRay.Direction,
+			bIntersects, HitPoint);
 
-	InteractionCurPoint = IntersectionPoint;
-	InteractionCurParameter = GizmoMath::ComputeCoordinatesInPlane(IntersectionPoint,
-		InteractionOrigin, InteractionNormal, InteractionAxisX, InteractionAxisY);
-	InteractionCurParameter.X *= ParameterSigns.X;
-	InteractionCurParameter.Y *= ParameterSigns.Y;
+		if (bIntersects == false)
+		{
+			return;
+		}
+		InteractionCurPoint = HitPoint;
+
+		InteractionCurParameter = GizmoMath::ComputeCoordinatesInPlane(InteractionCurPoint,
+			InteractionOrigin, InteractionNormal, InteractionAxisX, InteractionAxisY);
+		InteractionCurParameter.X *= ParameterSigns.X;
+		InteractionCurParameter.Y *= ParameterSigns.Y;
+	}
 
 	FVector2D DeltaParam = InteractionCurParameter - InteractionStartParameter;
-	FVector2D NewValue = InitialTargetParameter + DeltaParam;
+	NewParamValue = InitialTargetParameter + DeltaParam;
 
-	ParameterSource->SetParameter(NewValue);
+	ParameterSource->SetParameter(NewParamValue);
 }
 
 void UPlanePositionGizmo::OnClickRelease(const FInputDeviceRay& ReleasePos)
diff --git a/Engine/Source/Runtime/Experimental/InteractiveToolsFramework/Private/BaseGizmos/RepositionableTransformGizmo.cpp b/Engine/Source/Runtime/Experimental/InteractiveToolsFramework/Private/BaseGizmos/RepositionableTransformGizmo.cpp
new file mode 100644
--- /dev/null
+++ b/Engine/Source/Runtime/Experimental/InteractiveToolsFramework/Private/BaseGizmos/RepositionableTransformGizmo.cpp
@@ -0,0 +1,250 @@
+// Copyright Epic Games, Inc. All Rights Reserved.
+
+#include "BaseGizmos/RepositionableTransformGizmo.h"
+
+#include "BaseBehaviors/ClickDragBehavior.h"
+#include "BaseGizmos/AxisAngleGizmo.h"
+#include "BaseGizmos/AxisPositionGizmo.h"
+#include "BaseGizmos/HitTargets.h"
+#include "BaseGizmos/PlanePositionGizmo.h"
+#include "BaseGizmos/TransformProxy.h"
+#include "BaseGizmos/TransformSources.h"
+
+#include "InteractiveGizmo.h"
+
+#define LOCTEXT_NAMESPACE "URepositionableTransformGizmo"
+
+
+UInteractiveGizmo* URepositionableTransformGizmoBuilder::BuildGizmo(const FToolBuilderState& SceneState) const
+{
+	UTransformGizmo* NewGizmo = NewObject<URepositionableTransformGizmo>(SceneState.GizmoManager);
+	NewGizmo->SetWorld(SceneState.World);
+
+	// use default gizmo actor if client has not given us a new builder
+	NewGizmo->SetGizmoActorBuilder((GizmoActorBuilder) ? GizmoActorBuilder : MakeShared<FTransformGizmoActorFactory>());
+
+	// override default hover function if proposed
+	if (UpdateHoverFunction)
+	{
+		NewGizmo->SetUpdateHoverFunction(UpdateHoverFunction);
+	}
+
+	if (UpdateCoordSystemFunction)
+	{
+		NewGizmo->SetUpdateCoordSystemFunction(UpdateCoordSystemFunction);
+	}
+
+	return NewGizmo;
+}
+
+
+void URepositionableTransformGizmo::SetActiveTarget(UTransformProxy* Target, IToolContextTransactionProvider* TransactionProvider)
+{
+	// We're going to add on to the original transform gizmo, adding extra gizmos centered around
+	// existing gizmo components. These will respond to middle clicks and take "set pivot" code
+	// paths in their transform source and state target.
+	UTransformGizmo::SetActiveTarget(Target, TransactionProvider);
+
+	// Create the alternate transform source.
+	UGizmoTransformProxyTransformSource* ProxyTransformSource = UGizmoTransformProxyTransformSource::Construct(ActiveTarget, this);
+	ProxyTransformSource->bSetPivotMode = true;
+	USceneComponent* GizmoComponent = GizmoActor->GetRootComponent();
+	UGizmoScaledAndUnscaledTransformSources* TransformSource = UGizmoScaledAndUnscaledTransformSources::Construct(
+		ProxyTransformSource, GizmoComponent, this);
+
+	// The state target gives a different name to the transactions and issues different begin/end calls.
+	RepositionStateTarget = UGizmoTransformChangeStateTarget::Construct(GizmoComponent,
+		LOCTEXT("TransformGizmoTransaction", "Reposition Pivot"), 
+		Cast<IToolContextTransactionProvider>(StateTarget->TransactionManager.GetObject()), 
+		this);
+	TUniquePtr<FTransformProxyChangeSource> ProxyChangeSource = MakeUnique<FTransformProxyChangeSource>(Target);
+	ProxyChangeSource->bSetPivotMode = true;
+	RepositionStateTarget->DependentChangeSources.Add(MoveTemp(ProxyChangeSource));
+
+	// Add on the extra gizmos to existing components. We use the base class addition functions and
+	// then go back to tweak the parts that need changing using our own "Modify..." functions.
+	// We also store the repositioning gizmos in an additional separate list so that we can
+	// modify their alignment functions separately if we need to.
+	if (GizmoActor->TranslateX != nullptr)
+	{
+		UInteractiveGizmo* SubGizmo = AddAxisTranslationGizmo(GizmoActor->TranslateX, GizmoComponent, AxisXSource, TransformSource, RepositionStateTarget);
+		ModifyPivotAxisGizmo(SubGizmo);
+		PivotAlignmentGizmos.Add(SubGizmo);
+	}
+	if (GizmoActor->TranslateY != nullptr)
+	{
+		UInteractiveGizmo* SubGizmo = AddAxisTranslationGizmo(GizmoActor->TranslateY, GizmoComponent, AxisYSource, TransformSource, RepositionStateTarget);
+		ModifyPivotAxisGizmo(SubGizmo);
+		PivotAlignmentGizmos.Add(SubGizmo);
+	}
+	if (GizmoActor->TranslateZ != nullptr)
+	{
+		UInteractiveGizmo* SubGizmo = AddAxisTranslationGizmo(GizmoActor->TranslateZ, GizmoComponent, AxisZSource, TransformSource, RepositionStateTarget);
+		ModifyPivotAxisGizmo(SubGizmo);
+		PivotAlignmentGizmos.Add(SubGizmo);
+	}
+
+
+	if (GizmoActor->TranslateYZ != nullptr)
+	{
+		UInteractiveGizmo* SubGizmo = AddPlaneTranslationGizmo(GizmoActor->TranslateYZ, GizmoComponent, AxisXSource, TransformSource, RepositionStateTarget);
+		ModifyPivotPlaneGizmo(SubGizmo);
+		PivotAlignmentGizmos.Add(SubGizmo);
+	}
+	if (GizmoActor->TranslateXZ != nullptr)
+	{
+		UInteractiveGizmo* SubGizmo = AddPlaneTranslationGizmo(GizmoActor->TranslateXZ, GizmoComponent, AxisYSource, TransformSource, RepositionStateTarget);
+		ModifyPivotPlaneGizmo(SubGizmo);
+		PivotAlignmentGizmos.Add(SubGizmo);
+	}
+	if (GizmoActor->TranslateXY != nullptr)
+	{
+		UInteractiveGizmo* SubGizmo = AddPlaneTranslationGizmo(GizmoActor->TranslateXY, GizmoComponent, AxisZSource, TransformSource, RepositionStateTarget);
+		ModifyPivotPlaneGizmo(SubGizmo);
+		PivotAlignmentGizmos.Add(SubGizmo);
+	}
+
+
+	if (GizmoActor->RotateX != nullptr)
+	{
+		UInteractiveGizmo* SubGizmo = AddAxisRotationGizmo(GizmoActor->RotateX, GizmoComponent, AxisXSource, TransformSource, RepositionStateTarget);
+		ModifyPivotRotateGizmo(SubGizmo);
+		PivotAlignmentGizmos.Add(SubGizmo);
+	}
+	if (GizmoActor->RotateY != nullptr)
+	{
+		UInteractiveGizmo* SubGizmo = AddAxisRotationGizmo(GizmoActor->RotateY, GizmoComponent, AxisYSource, TransformSource, RepositionStateTarget);
+		ModifyPivotRotateGizmo(SubGizmo);
+		PivotAlignmentGizmos.Add(SubGizmo);
+	}
+	if (GizmoActor->RotateZ != nullptr)
+	{
+		UInteractiveGizmo* SubGizmo = AddAxisRotationGizmo(GizmoActor->RotateZ, GizmoComponent, AxisZSource, TransformSource, RepositionStateTarget);
+		ModifyPivotRotateGizmo(SubGizmo);
+		PivotAlignmentGizmos.Add(SubGizmo);
+	}
+}
+
+/** 
+ * Modifies a UAxisPositionGizmo after its addition via AddAxisTranslationGizmo for use in 
+ * gizmo repositioning.
+ */
+void URepositionableTransformGizmo::ModifyPivotAxisGizmo(UInteractiveGizmo* SubGizmoIn)
+{
+	UAxisPositionGizmo* SubGizmo = Cast<UAxisPositionGizmo>(SubGizmoIn);
+	check(SubGizmo);
+	SubGizmo->MouseBehavior->SetUseMiddleMouseButton();
+	SubGizmo->ShouldUseCustomDestinationFunc = [this]() { return ShouldAlignPivot(); };
+	SubGizmo->CustomDestinationFunc =
+		[this](const UAxisPositionGizmo::FCustomDestinationParams& Params, FVector& OutputPoint) {
+		return PivotAlignmentRayCaster(*Params.WorldRay, OutputPoint);
+	};
+
+	UGizmoComponentHitTarget* HitTarget = Cast<UGizmoComponentHitTarget>(SubGizmo->HitTarget.GetObject());
+	check(HitTarget);
+	HitTarget->UpdateHoverFunction = nullptr; // Already taken care of by the normal gizmo
+}
+
+/**
+ * Modifies a UPlanePositionGizmo after its addition via AddPlaneTranslationGizmo for use in
+ * gizmo repositioning.
+ */
+void URepositionableTransformGizmo::ModifyPivotPlaneGizmo(UInteractiveGizmo* SubGizmoIn)
+{
+	UPlanePositionGizmo* SubGizmo = Cast<UPlanePositionGizmo>(SubGizmoIn);
+	check(SubGizmo);
+	SubGizmo->MouseBehavior->SetUseMiddleMouseButton();
+	SubGizmo->ShouldUseCustomDestinationFunc = [this]() { return ShouldAlignPivot(); };
+	SubGizmo->CustomDestinationFunc =
+		[this](const UPlanePositionGizmo::FCustomDestinationParams& Params, FVector& OutputPoint) {
+		return PivotAlignmentRayCaster(*Params.WorldRay, OutputPoint);
+	};
+
+	UGizmoComponentHitTarget* HitTarget = Cast<UGizmoComponentHitTarget>(SubGizmo->HitTarget.GetObject());
+	check(HitTarget);
+	HitTarget->UpdateHoverFunction = nullptr; // Already taken care of by the normal gizmo
+}
+
+/**
+ * Modifies a UAxisAngleGizmo after its addition via AddAxisRotationGizmo for use in
+ * gizmo repositioning.
+ */
+void URepositionableTransformGizmo::ModifyPivotRotateGizmo(UInteractiveGizmo* SubGizmoIn)
+{
+	UAxisAngleGizmo* SubGizmo = Cast<UAxisAngleGizmo>(SubGizmoIn);
+	check(SubGizmo);
+	SubGizmo->MouseBehavior->SetUseMiddleMouseButton();
+	SubGizmo->ShouldUseCustomDestinationFunc = [this]() { return ShouldAlignPivot(); };
+	SubGizmo->CustomDestinationFunc =
+		[this](const UAxisAngleGizmo::FCustomDestinationParams& Params, FVector& OutputPoint) {
+		return PivotAlignmentRayCaster(*Params.WorldRay, OutputPoint);
+	};
+
+	UGizmoComponentHitTarget* HitTarget = Cast<UGizmoComponentHitTarget>(SubGizmo->HitTarget.GetObject());
+	check(HitTarget);
+	HitTarget->UpdateHoverFunction = nullptr; // Already taken care of by the normal gizmo
+
+	// We don't want to be able to reposition the rotation of the gizmo when using the world coordinate system,
+	// since the axes must stay aligned with world axes. However we can't hide the component either because it
+	// is used for regular rotation. Thus, we add a check that disables the reposition gizmo hit testing if
+	// the coordinate system is not local.
+	HitTarget->Condition = [this](const FInputDeviceRay&) {
+		return (CurrentCoordinateSystem == EToolContextCoordinateSystem::Local);
+	};
+}
+
+void URepositionableTransformGizmo::SetWorldAlignmentFunctions(
+	TUniqueFunction<bool()>&& ShouldAlignTranslationIn,
+	TUniqueFunction<bool(const FRay&, FVector&)>&& TranslationAlignmentRayCasterIn)
+{
+	// The base class function will do most of what we want.
+	UTransformGizmo::SetWorldAlignmentFunctions(MoveTemp(ShouldAlignTranslationIn), MoveTemp(TranslationAlignmentRayCasterIn));
+
+	// However, we need to adjust our saved pivot alignment functions so changes of the gizmo target keep the settings.
+	// Since the functions are TUniqueFunction, we can't actually copy them here, we need to create new functors that
+	// forward to the other ones.
+	ShouldAlignPivot = [this]() { return ShouldAlignDestination(); };
+	PivotAlignmentRayCaster = [this](const FRay& WorldRay, FVector& OutputPoint) { 
+		return DestinationAlignmentRayCaster(WorldRay, OutputPoint); };
+
+}
+
+void URepositionableTransformGizmo::SetPivotAlignmentFunctions(
+	TUniqueFunction<bool()>&& ShouldAlignPivotIn,
+	TUniqueFunction<bool(const FRay&, FVector&)>&& PivotAlignmentRayCasterIn)
+{
+	// Save these so that later changes of gizmo target keep the settings.
+	ShouldAlignPivot = MoveTemp(ShouldAlignPivotIn);
+	PivotAlignmentRayCaster = MoveTemp(PivotAlignmentRayCasterIn);
+
+	// We allow this function to be called after Setup(), so modify any existing translation sub gizmos.
+	for (UInteractiveGizmo* SubGizmo : PivotAlignmentGizmos)
+	{
+		if (UAxisPositionGizmo* CastGizmo = Cast<UAxisPositionGizmo>(SubGizmo))
+		{
+			CastGizmo->ShouldUseCustomDestinationFunc = [this]() { return ShouldAlignPivot(); };
+			CastGizmo->CustomDestinationFunc =
+				[this](const UAxisPositionGizmo::FCustomDestinationParams& Params, FVector& OutputPoint) {
+				return PivotAlignmentRayCaster(*Params.WorldRay, OutputPoint);
+			};
+		}
+		if (UPlanePositionGizmo* CastGizmo = Cast<UPlanePositionGizmo>(SubGizmo))
+		{
+			CastGizmo->ShouldUseCustomDestinationFunc = [this]() { return ShouldAlignPivot(); };
+			CastGizmo->CustomDestinationFunc =
+				[this](const UPlanePositionGizmo::FCustomDestinationParams& Params, FVector& OutputPoint) {
+				return PivotAlignmentRayCaster(*Params.WorldRay, OutputPoint);
+			};
+		}
+		if (UAxisAngleGizmo* CastGizmo = Cast<UAxisAngleGizmo>(SubGizmo))
+		{
+			CastGizmo->ShouldUseCustomDestinationFunc = [this]() { return ShouldAlignPivot(); };
+			CastGizmo->CustomDestinationFunc =
+				[this](const UAxisAngleGizmo::FCustomDestinationParams& Params, FVector& OutputPoint) {
+				return PivotAlignmentRayCaster(*Params.WorldRay, OutputPoint);
+			};
+		}
+	}
+}
+
+#undef LOCTEXT_NAMESPACE
diff --git a/Engine/Source/Runtime/Experimental/InteractiveToolsFramework/Private/BaseGizmos/TransformGizmo.cpp b/Engine/Source/Runtime/Experimental/InteractiveToolsFramework/Private/BaseGizmos/TransformGizmo.cpp
--- a/Engine/Source/Runtime/Experimental/InteractiveToolsFramework/Private/BaseGizmos/TransformGizmo.cpp
+++ b/Engine/Source/Runtime/Experimental/InteractiveToolsFramework/Private/BaseGizmos/TransformGizmo.cpp
@@ -261,6 +261,93 @@ void UTransformGizmo::SetUpdateCoordSystemFunction(TFunction<void(UPrimitiveComp
 	UpdateCoordSystemFunction = CoordSysFunction;
 }
 
+void UTransformGizmo::SetWorldAlignmentFunctions(
+	TUniqueFunction<bool()>&& ShouldAlignTranslationIn,
+	TUniqueFunction<bool(const FRay&, FVector&)>&& TranslationAlignmentRayCasterIn)
+{
+	// Save these so that later changes of gizmo target keep the settings.
+	ShouldAlignDestination = MoveTemp(ShouldAlignTranslationIn);
+	DestinationAlignmentRayCaster = MoveTemp(TranslationAlignmentRayCasterIn);
+
+	// We allow this function to be called after Setup(), so modify any existing translation/rotation sub gizmos.
+	// Unfortunately we keep all the sub gizmos in one list, and the scaling gizmos are differentiated from the
+	// translation ones mainly in the components they use. So this ends up being a slightly messy set of checks,
+	// but it didn't seem worth keeping a segregated list for something that will only happen once.
+	for (UInteractiveGizmo* SubGizmo : this->ActiveGizmos)
+	{
+		if (UAxisPositionGizmo* CastGizmo = Cast<UAxisPositionGizmo>(SubGizmo))
+		{
+			if (UGizmoComponentHitTarget* CastHitTarget = Cast<UGizmoComponentHitTarget>(CastGizmo->HitTarget.GetObject()))
+			{
+				if (CastHitTarget->Component == GizmoActor->TranslateX
+					|| CastHitTarget->Component == GizmoActor->TranslateY
+					|| CastHitTarget->Component == GizmoActor->TranslateZ)
+				{
+					CastGizmo->ShouldUseCustomDestinationFunc = [this]() { return ShouldAlignDestination(); };
+					CastGizmo->CustomDestinationFunc = 
+						[this](const UAxisPositionGizmo::FCustomDestinationParams& Params, FVector& OutputPoint) { 
+						return DestinationAlignmentRayCaster(*Params.WorldRay, OutputPoint); 
+					};
+				}
+			}
+		}
+		if (UPlanePositionGizmo* CastGizmo = Cast<UPlanePositionGizmo>(SubGizmo))
+		{
+			if (UGizmoComponentHitTarget* CastHitTarget = Cast<UGizmoComponentHitTarget>(CastGizmo->HitTarget.GetObject()))
+			{
+				if (CastHitTarget->Component == GizmoActor->TranslateXY
+					|| CastHitTarget->Component == GizmoActor->TranslateXZ
+					|| CastHitTarget->Component == GizmoActor->TranslateYZ)
+				{
+					CastGizmo->ShouldUseCustomDestinationFunc = [this]() { return ShouldAlignDestination(); };
+					CastGizmo->CustomDestinationFunc =
+						[this](const UPlanePositionGizmo::FCustomDestinationParams& Params, FVector& OutputPoint) {
+						return DestinationAlignmentRayCaster(*Params.WorldRay, OutputPoint);
+					};
+				}
+			}
+		}
+		if (UAxisAngleGizmo* CastGizmo = Cast<UAxisAngleGizmo>(SubGizmo))
+		{
+			CastGizmo->ShouldUseCustomDestinationFunc = [this]() { return ShouldAlignDestination(); };
+			CastGizmo->CustomDestinationFunc =
+				[this](const UAxisAngleGizmo::FCustomDestinationParams& Params, FVector& OutputPoint) {
+				return DestinationAlignmentRayCaster(*Params.WorldRay, OutputPoint);
+			};
+		}
+	}
+}
+
+void UTransformGizmo::SetDisallowNegativeScaling(bool bDisallow)
+{
+	if (bDisallowNegativeScaling != bDisallow)
+	{
+		bDisallowNegativeScaling = bDisallow;
+		for (UInteractiveGizmo* SubGizmo : this->ActiveGizmos)
+		{
+			if (UAxisPositionGizmo* CastGizmo = Cast<UAxisPositionGizmo>(SubGizmo))
+			{
+				if (UGizmoAxisScaleParameterSource* ParamSource = Cast<UGizmoAxisScaleParameterSource>(CastGizmo->ParameterSource.GetObject()))
+				{
+					ParamSource->bClampToZero = bDisallow;
+				}
+			}
+			if (UPlanePositionGizmo* CastGizmo = Cast<UPlanePositionGizmo>(SubGizmo))
+			{
+				if (UGizmoPlaneScaleParameterSource* ParamSource = Cast<UGizmoPlaneScaleParameterSource>(CastGizmo->ParameterSource.GetObject()))
+				{
+					ParamSource->bClampToZero = bDisallow;
+				}
+			}
+		}
+	}
+}
+
+void UTransformGizmo::SetIsNonUniformScaleAllowedFunction(TUniqueFunction<bool()>&& IsNonUniformScaleAllowedIn)
+{
+	IsNonUniformScaleAllowed = MoveTemp(IsNonUniformScaleAllowedIn);
+}
+
 
 void UTransformGizmo::Setup()
 {
@@ -338,9 +425,10 @@ void UTransformGizmo::Tick(float DeltaTime)
 		}
 	}
 
+	bool bShouldShowNonUniformScale = IsNonUniformScaleAllowed();
 	for (UPrimitiveComponent* Component : NonuniformScaleComponents)
 	{
-		Component->SetVisibility(bUseLocalAxes);
+		Component->SetVisibility(bShouldShowNonUniformScale);
 	}
 
 	UpdateCameraAxisSource();
@@ -365,33 +453,9 @@ void UTransformGizmo::SetActiveTarget(UTransformProxy* Target, IToolContextTrans
 	GizmoTransform.SetScale3D(FVector(1, 1, 1));
 	GizmoComponent->SetWorldTransform(GizmoTransform);
 
-	// save current scale because gizmo is not scaled
-	SeparateChildScale = TargetTransform.GetScale3D();
-
-	UGizmoComponentWorldTransformSource* ComponentTransformSource =
-		UGizmoComponentWorldTransformSource::Construct(GizmoComponent, this);
-	FSeparateScaleProvider ScaleProvider = {
-		[this]() { return this->SeparateChildScale; },
-		[this](FVector Scale) { this->SeparateChildScale = Scale; }
-	};
-	ScaledTransformSource = UGizmoScaledTransformSource::Construct(ComponentTransformSource, ScaleProvider, this);
-
-	// Target tracks location of GizmoComponent. Note that TransformUpdated is not called during undo/redo transactions!
-	// We currently rely on the transaction system to undo/redo target object locations. This will not work during runtime...
-	GizmoComponent->TransformUpdated.AddLambda(
-		[this](USceneComponent* Component, EUpdateTransformFlags /*UpdateTransformFlags*/, ETeleportType /*Teleport*/) {
-		//FTransform NewXForm = Component->GetComponentToWorld();
-		//NewXForm.SetScale3D(this->CurTargetScale);
-		FTransform NewXForm = ScaledTransformSource->GetTransform();
-		this->ActiveTarget->SetTransform(NewXForm);
-	});
-	ScaledTransformSource->OnTransformChanged.AddLambda(
-		[this](IGizmoTransformSource* Source)
-	{
-		FTransform NewXForm = ScaledTransformSource->GetTransform();
-		this->ActiveTarget->SetTransform(NewXForm);
-	});
-
+	UGizmoScaledAndUnscaledTransformSources* TransformSource = UGizmoScaledAndUnscaledTransformSources::Construct(
+		UGizmoTransformProxyTransformSource::Construct(ActiveTarget, this), 
+		GizmoComponent, this);
 
 	// This state target emits an explicit FChange that moves the GizmoActor root component during undo/redo.
 	// It also opens/closes the Transaction that saves/restores the target object locations.
@@ -402,7 +466,6 @@ void UTransformGizmo::SetActiveTarget(UTransformProxy* Target, IToolContextTrans
 	StateTarget = UGizmoTransformChangeStateTarget::Construct(GizmoComponent,
 		LOCTEXT("UTransformGizmoTransaction", "Transform"), TransactionProvider, this);
 	StateTarget->DependentChangeSources.Add(MakeUnique<FTransformProxyChangeSource>(Target));
-	StateTarget->ExternalDependentChangeSources.Add(this);
 
 	CameraAxisSource = NewObject<UGizmoConstantFrameAxisSource>(this);
 
@@ -414,50 +477,50 @@ void UTransformGizmo::SetActiveTarget(UTransformProxy* Target, IToolContextTrans
 	// todo should we hold onto these?
 	if (GizmoActor->TranslateX != nullptr)
 	{
-		AddAxisTranslationGizmo(GizmoActor->TranslateX, GizmoComponent, AxisXSource, ScaledTransformSource, StateTarget);
+		AddAxisTranslationGizmo(GizmoActor->TranslateX, GizmoComponent, AxisXSource, TransformSource, StateTarget);
 		ActiveComponents.Add(GizmoActor->TranslateX);
 	}
 	if (GizmoActor->TranslateY != nullptr)
 	{
-		AddAxisTranslationGizmo(GizmoActor->TranslateY, GizmoComponent, AxisYSource, ScaledTransformSource, StateTarget);
+		AddAxisTranslationGizmo(GizmoActor->TranslateY, GizmoComponent, AxisYSource, TransformSource, StateTarget);
 		ActiveComponents.Add(GizmoActor->TranslateY);
 	}
 	if (GizmoActor->TranslateZ != nullptr)
 	{
-		AddAxisTranslationGizmo(GizmoActor->TranslateZ, GizmoComponent, AxisZSource, ScaledTransformSource, StateTarget);
+		AddAxisTranslationGizmo(GizmoActor->TranslateZ, GizmoComponent, AxisZSource, TransformSource, StateTarget);
 		ActiveComponents.Add(GizmoActor->TranslateZ);
 	}
 
 
 	if (GizmoActor->TranslateYZ != nullptr)
 	{
-		AddPlaneTranslationGizmo(GizmoActor->TranslateYZ, GizmoComponent, AxisXSource, ScaledTransformSource, StateTarget);
+		AddPlaneTranslationGizmo(GizmoActor->TranslateYZ, GizmoComponent, AxisXSource, TransformSource, StateTarget);
 		ActiveComponents.Add(GizmoActor->TranslateYZ);
 	}
 	if (GizmoActor->TranslateXZ != nullptr)
 	{
-		AddPlaneTranslationGizmo(GizmoActor->TranslateXZ, GizmoComponent, AxisYSource, ScaledTransformSource, StateTarget);
+		AddPlaneTranslationGizmo(GizmoActor->TranslateXZ, GizmoComponent, AxisYSource, TransformSource, StateTarget);
 		ActiveComponents.Add(GizmoActor->TranslateXZ);
 	}
 	if (GizmoActor->TranslateXY != nullptr)
 	{
-		AddPlaneTranslationGizmo(GizmoActor->TranslateXY, GizmoComponent, AxisZSource, ScaledTransformSource, StateTarget);
+		AddPlaneTranslationGizmo(GizmoActor->TranslateXY, GizmoComponent, AxisZSource, TransformSource, StateTarget);
 		ActiveComponents.Add(GizmoActor->TranslateXY);
 	}
 
 	if (GizmoActor->RotateX != nullptr)
 	{
-		AddAxisRotationGizmo(GizmoActor->RotateX, GizmoComponent, AxisXSource, ScaledTransformSource, StateTarget);
+		AddAxisRotationGizmo(GizmoActor->RotateX, GizmoComponent, AxisXSource, TransformSource, StateTarget);
 		ActiveComponents.Add(GizmoActor->RotateX);
 	}
 	if (GizmoActor->RotateY != nullptr)
 	{
-		AddAxisRotationGizmo(GizmoActor->RotateY, GizmoComponent, AxisYSource, ScaledTransformSource, StateTarget);
+		AddAxisRotationGizmo(GizmoActor->RotateY, GizmoComponent, AxisYSource, TransformSource, StateTarget);
 		ActiveComponents.Add(GizmoActor->RotateY);
 	}
 	if (GizmoActor->RotateZ != nullptr)
 	{
-		AddAxisRotationGizmo(GizmoActor->RotateZ, GizmoComponent, AxisZSource, ScaledTransformSource, StateTarget);
+		AddAxisRotationGizmo(GizmoActor->RotateZ, GizmoComponent, AxisZSource, TransformSource, StateTarget);
 		ActiveComponents.Add(GizmoActor->RotateZ);
 	}
 
@@ -471,45 +534,44 @@ void UTransformGizmo::SetActiveTarget(UTransformProxy* Target, IToolContextTrans
 
 	if (GizmoActor->UniformScale != nullptr)
 	{
-		AddUniformScaleGizmo(GizmoActor->UniformScale, GizmoComponent, CameraAxisSource, CameraAxisSource, ScaledTransformSource, StateTarget);
+		AddUniformScaleGizmo(GizmoActor->UniformScale, GizmoComponent, CameraAxisSource, CameraAxisSource, TransformSource, StateTarget);
 		ActiveComponents.Add(GizmoActor->UniformScale);
 	}
 
 	if (GizmoActor->AxisScaleX != nullptr)
 	{
-		AddAxisScaleGizmo(GizmoActor->AxisScaleX, GizmoComponent, AxisXSource, UnitAxisXSource, ScaledTransformSource, StateTarget);
+		AddAxisScaleGizmo(GizmoActor->AxisScaleX, GizmoComponent, AxisXSource, UnitAxisXSource, TransformSource, StateTarget);
 		ActiveComponents.Add(GizmoActor->AxisScaleX);
 		NonuniformScaleComponents.Add(GizmoActor->AxisScaleX);
 	}
 	if (GizmoActor->AxisScaleY != nullptr)
 	{
-		AddAxisScaleGizmo(GizmoActor->AxisScaleY, GizmoComponent, AxisYSource, UnitAxisYSource, ScaledTransformSource, StateTarget);
+		AddAxisScaleGizmo(GizmoActor->AxisScaleY, GizmoComponent, AxisYSource, UnitAxisYSource, TransformSource, StateTarget);
 		ActiveComponents.Add(GizmoActor->AxisScaleY);
 		NonuniformScaleComponents.Add(GizmoActor->AxisScaleY);
 	}
 	if (GizmoActor->AxisScaleZ != nullptr)
 	{
-		AddAxisScaleGizmo(GizmoActor->AxisScaleZ, GizmoComponent, AxisZSource, UnitAxisZSource, ScaledTransformSource, StateTarget);
+		AddAxisScaleGizmo(GizmoActor->AxisScaleZ, GizmoComponent, AxisZSource, UnitAxisZSource, TransformSource, StateTarget);
 		ActiveComponents.Add(GizmoActor->AxisScaleZ);
 		NonuniformScaleComponents.Add(GizmoActor->AxisScaleZ);
 	}
 
 	if (GizmoActor->PlaneScaleYZ != nullptr)
 	{
-		AddPlaneScaleGizmo(GizmoActor->PlaneScaleYZ, GizmoComponent, AxisXSource, UnitAxisXSource, ScaledTransformSource, StateTarget);
+		AddPlaneScaleGizmo(GizmoActor->PlaneScaleYZ, GizmoComponent, AxisXSource, UnitAxisXSource, TransformSource, StateTarget);
 		ActiveComponents.Add(GizmoActor->PlaneScaleYZ);
 		NonuniformScaleComponents.Add(GizmoActor->PlaneScaleYZ);
 	}
 	if (GizmoActor->PlaneScaleXZ != nullptr)
 	{
-		UPlanePositionGizmo* Gizmo = (UPlanePositionGizmo *)AddPlaneScaleGizmo(GizmoActor->PlaneScaleXZ, GizmoComponent, AxisYSource, UnitAxisYSource, ScaledTransformSource, StateTarget);
-		Gizmo->bFlipX = true;		// unclear why this is necessary...possibly a handedness issue?
+		UPlanePositionGizmo* Gizmo = (UPlanePositionGizmo *)AddPlaneScaleGizmo(GizmoActor->PlaneScaleXZ, GizmoComponent, AxisYSource, UnitAxisYSource, TransformSource, StateTarget);
 		ActiveComponents.Add(GizmoActor->PlaneScaleXZ);
 		NonuniformScaleComponents.Add(GizmoActor->PlaneScaleXZ);
 	}
 	if (GizmoActor->PlaneScaleXY != nullptr)
 	{
-		AddPlaneScaleGizmo(GizmoActor->PlaneScaleXY, GizmoComponent, AxisZSource, UnitAxisZSource, ScaledTransformSource, StateTarget);
+		AddPlaneScaleGizmo(GizmoActor->PlaneScaleXY, GizmoComponent, AxisZSource, UnitAxisZSource, TransformSource, StateTarget);
 		ActiveComponents.Add(GizmoActor->PlaneScaleXY);
 		NonuniformScaleComponents.Add(GizmoActor->PlaneScaleXY);
 	}
@@ -539,19 +601,24 @@ void UTransformGizmo::SetNewGizmoTransform(const FTransform& NewTransform)
 
 	StateTarget->BeginUpdate();
 
-	SeparateChildScale = NewTransform.GetScale3D();
-
 	USceneComponent* GizmoComponent = GizmoActor->GetRootComponent();
 	GizmoComponent->SetWorldTransform(NewTransform);
-	//ActiveTarget->SetTransform(NewTransform);		// this will happen in the GizmoComponent.TransformUpdated delegate handler above
+	ActiveTarget->SetTransform(NewTransform);
 
 	StateTarget->EndUpdate();
 }
 
 
+// TODO: This should either be named to "SetScale" or removed, since it can be done with ReinitializeGizmoTransform
 void UTransformGizmo::SetNewChildScale(const FVector& NewChildScale)
 {
-	SeparateChildScale = NewChildScale;
+	FTransform NewTransform = ActiveTarget->GetTransform();
+	NewTransform.SetScale3D(NewChildScale);
+
+	bool bSavedSetPivotMode = ActiveTarget->bSetPivotMode;
+	ActiveTarget->bSetPivotMode = true;
+	ActiveTarget->SetTransform(NewTransform);
+	ActiveTarget->bSetPivotMode = bSavedSetPivotMode;
 }
 
 
@@ -593,6 +660,12 @@ UInteractiveGizmo* UTransformGizmo::AddAxisTranslationGizmo(
 
 	TranslateGizmo->StateTarget = Cast<UObject>(StateTargetIn);
 
+	TranslateGizmo->ShouldUseCustomDestinationFunc = [this]() { return ShouldAlignDestination(); };
+	TranslateGizmo->CustomDestinationFunc =
+		[this](const UAxisPositionGizmo::FCustomDestinationParams& Params, FVector& OutputPoint) {
+		return DestinationAlignmentRayCaster(*Params.WorldRay, OutputPoint);
+	};
+
 	ActiveGizmos.Add(TranslateGizmo);
 	return TranslateGizmo;
 }
@@ -628,6 +701,12 @@ UInteractiveGizmo* UTransformGizmo::AddPlaneTranslationGizmo(
 
 	TranslateGizmo->StateTarget = Cast<UObject>(StateTargetIn);
 
+	TranslateGizmo->ShouldUseCustomDestinationFunc = [this]() { return ShouldAlignDestination(); };
+	TranslateGizmo->CustomDestinationFunc =
+		[this](const UPlanePositionGizmo::FCustomDestinationParams& Params, FVector& OutputPoint) {
+		return DestinationAlignmentRayCaster(*Params.WorldRay, OutputPoint);
+	};
+
 	ActiveGizmos.Add(TranslateGizmo);
 	return TranslateGizmo;
 }
@@ -665,6 +744,12 @@ UInteractiveGizmo* UTransformGizmo::AddAxisRotationGizmo(
 
 	RotateGizmo->StateTarget = Cast<UObject>(StateTargetIn);
 
+	RotateGizmo->ShouldUseCustomDestinationFunc = [this]() { return ShouldAlignDestination(); };
+	RotateGizmo->CustomDestinationFunc =
+		[this](const UAxisAngleGizmo::FCustomDestinationParams& Params, FVector& OutputPoint) {
+		return DestinationAlignmentRayCaster(*Params.WorldRay, OutputPoint);
+	};
+
 	ActiveGizmos.Add(RotateGizmo);
 
 	return RotateGizmo;
@@ -689,7 +774,7 @@ UInteractiveGizmo* UTransformGizmo::AddAxisScaleGizmo(
 
 	// parameter source maps axis-parameter-change to translation of TransformSource's transform
 	UGizmoAxisScaleParameterSource* ParamSource = UGizmoAxisScaleParameterSource::Construct(ParameterAxisSource, TransformSource, this);
-	//ParamSource->PositionConstraintFunction = [this](const FVector& Pos, FVector& Snapped) { return PositionSnapFunction(Pos, Snapped); };
+	ParamSource->bClampToZero = bDisallowNegativeScaling;
 	ScaleGizmo->ParameterSource = ParamSource;
 
 	// sub-component provides hit target
@@ -725,7 +810,8 @@ UInteractiveGizmo* UTransformGizmo::AddPlaneScaleGizmo(
 
 	// parameter source maps axis-parameter-change to translation of TransformSource's transform
 	UGizmoPlaneScaleParameterSource* ParamSource = UGizmoPlaneScaleParameterSource::Construct(ParameterAxisSource, TransformSource, this);
-	//ParamSource->PositionConstraintFunction = [this](const FVector& Pos, FVector& Snapped) { return PositionSnapFunction(Pos, Snapped); };
+	ParamSource->bClampToZero = bDisallowNegativeScaling;
+	ParamSource->bUseEqualScaling = true;
 	ScaleGizmo->ParameterSource = ParamSource;
 
 	// sub-component provides hit target
@@ -864,53 +950,4 @@ FQuat UTransformGizmo::RotationSnapFunction(const FQuat& DeltaRotation) const
 	return SnappedDeltaRotation;
 }
 
-void UTransformGizmo::BeginChange()
-{
-	ActiveChange = MakeUnique<FTransformGizmoTransformChange>();
-	ActiveChange->ChildScaleBefore = SeparateChildScale;
-}
-
-TUniquePtr<FToolCommandChange> UTransformGizmo::EndChange()
-{
-	ActiveChange->ChildScaleAfter = SeparateChildScale;
-	return MoveTemp(ActiveChange);
-	ActiveChange = nullptr;
-}
-
-UObject* UTransformGizmo::GetChangeTarget()
-{
-	return this;
-}
-
-FText UTransformGizmo::GetChangeDescription()
-{
-	return LOCTEXT("TransformGizmoChangeDescription", "Transform Change");
-}
-
-
-void UTransformGizmo::ExternalSetChildScale(const FVector& NewScale)
-{
-	SeparateChildScale = NewScale;
-}
-
-
-void FTransformGizmoTransformChange::Apply(UObject* Object)
-{
-	UTransformGizmo* Gizmo = CastChecked<UTransformGizmo>(Object);
-	Gizmo->ExternalSetChildScale(ChildScaleAfter);
-}
-
-
-void FTransformGizmoTransformChange::Revert(UObject* Object)
-{
-	UTransformGizmo* Gizmo = CastChecked<UTransformGizmo>(Object);
-	Gizmo->ExternalSetChildScale(ChildScaleBefore);
-}
-
-FString FTransformGizmoTransformChange::ToString() const
-{
-	return FString(TEXT("TransformGizmo Change"));
-}
-
-
 #undef LOCTEXT_NAMESPACE
diff --git a/Engine/Source/Runtime/Experimental/InteractiveToolsFramework/Private/BaseGizmos/TransformProxy.cpp b/Engine/Source/Runtime/Experimental/InteractiveToolsFramework/Private/BaseGizmos/TransformProxy.cpp
--- a/Engine/Source/Runtime/Experimental/InteractiveToolsFramework/Private/BaseGizmos/TransformProxy.cpp
+++ b/Engine/Source/Runtime/Experimental/InteractiveToolsFramework/Private/BaseGizmos/TransformProxy.cpp
@@ -19,7 +19,7 @@ void UTransformProxy::AddComponent(USceneComponent* Component, bool bModifyCompo
 	Objects.Add(NewObj);
 
 	UpdateSharedTransform();
-	OnTransformChanged.Broadcast(this, SharedTransform);
+	OnPivotChanged.Broadcast(this, SharedTransform);
 }
 
 
@@ -36,6 +36,7 @@ void UTransformProxy::SetTransform(const FTransform& TransformIn)
 	if (bSetPivotMode)
 	{
 		UpdateObjectTransforms();
+		OnPivotChanged.Broadcast(this, SharedTransform);
 	}
 	else
 	{
@@ -55,7 +56,15 @@ void UTransformProxy::EndTransformEditSequence()
 	OnEndTransformEdit.Broadcast(this);
 }
 
+void UTransformProxy::BeginPivotEditSequence()
+{
+	OnBeginPivotEdit.Broadcast(this);
+}
 
+void UTransformProxy::EndPivotEditSequence()
+{
+	OnEndPivotEdit.Broadcast(this);
+}
 
 
 
@@ -66,12 +75,17 @@ void UTransformProxy::UpdateObjects()
 		FTransform CombinedTransform;
 		if (bRotatePerObject && Objects.Num() > 1)
 		{
-			FTransform Temp = SharedTransform.GetRelativeTransform(InitialSharedTransform);
-
+			// We want to apply the compare the shared transform to the shared transform that existed
+			// at the time the object's StartTransform was set, then apply the changes to the StartTransform.
+			// It may seem that FTransform::RelativeTransform() might be of use here, but that gives the
+			// transform from the point of view of the initial frame, which would give us incorrect
+			// translation if the initial frame axes didn't line up with world axes.
 			CombinedTransform = Obj.StartTransform;
-			CombinedTransform.AddToTranslation(Temp.GetTranslation());
-			CombinedTransform.ConcatenateRotation(Temp.GetRotation());
-			CombinedTransform.SetScale3D(CombinedTransform.GetScale3D() * Temp.GetScale3D());
+
+			CombinedTransform.AddToTranslation(SharedTransform.GetTranslation() - InitialSharedTransform.GetTranslation());
+			CombinedTransform.ConcatenateRotation(InitialSharedTransform.GetRotation().Inverse());
+			CombinedTransform.ConcatenateRotation(SharedTransform.GetRotation());
+			CombinedTransform.SetScale3D(CombinedTransform.GetScale3D() * SharedTransform.GetScale3D() / InitialSharedTransform.GetScale3D());
 		}
 		else
 		{
@@ -139,6 +153,8 @@ void UTransformProxy::UpdateObjectTransforms()
 		Obj.RelativeTransform = Obj.StartTransform;
 		Obj.RelativeTransform.SetToRelativeTransform(SharedTransform);
 	}
+
+	InitialSharedTransform = SharedTransform;
 }
 
 
@@ -148,13 +164,21 @@ void UTransformProxy::UpdateObjectTransforms()
 void FTransformProxyChange::Apply(UObject* Object)
 {
 	UTransformProxy* Proxy = CastChecked<UTransformProxy>(Object);
+
+	bool bSavedSetPivotMode = Proxy->bSetPivotMode;
+	Proxy->bSetPivotMode = bSetPivotMode;
 	Proxy->SetTransform(To);
+	Proxy->bSetPivotMode = bSavedSetPivotMode;
 }
 
 void FTransformProxyChange::Revert(UObject* Object)
 {
 	UTransformProxy* Proxy = CastChecked<UTransformProxy>(Object);
+
+	bool bSavedSetPivotMode = Proxy->bSetPivotMode;
+	Proxy->bSetPivotMode = bSetPivotMode;
 	Proxy->SetTransform(From);
+	Proxy->bSetPivotMode = bSavedSetPivotMode;
 }
 
 
@@ -164,7 +188,16 @@ void FTransformProxyChangeSource::BeginChange()
 	{
 		ActiveChange = MakeUnique<FTransformProxyChange>();
 		ActiveChange->From = Proxy->GetTransform();
-		Proxy->BeginTransformEditSequence();
+		ActiveChange->bSetPivotMode = bSetPivotMode;
+
+		if (bSetPivotMode)
+		{
+			Proxy->BeginPivotEditSequence();
+		}
+		else
+		{
+			Proxy->BeginTransformEditSequence();
+		}
 	}
 }
 
@@ -172,7 +205,15 @@ TUniquePtr<FToolCommandChange> FTransformProxyChangeSource::EndChange()
 {
 	if (Proxy.IsValid())
 	{
-		Proxy->EndTransformEditSequence();
+		if (bSetPivotMode)
+		{
+			Proxy->EndPivotEditSequence();
+		}
+		else
+		{
+			Proxy->EndTransformEditSequence();
+		}
+
 		ActiveChange->To = Proxy->GetTransform();
 		return MoveTemp(ActiveChange);
 	}
diff --git a/Engine/Source/Runtime/Experimental/InteractiveToolsFramework/Private/BaseGizmos/TransformSources.cpp b/Engine/Source/Runtime/Experimental/InteractiveToolsFramework/Private/BaseGizmos/TransformSources.cpp
--- a/Engine/Source/Runtime/Experimental/InteractiveToolsFramework/Private/BaseGizmos/TransformSources.cpp
+++ b/Engine/Source/Runtime/Experimental/InteractiveToolsFramework/Private/BaseGizmos/TransformSources.cpp
@@ -57,6 +57,37 @@ FTransform UGizmoTransformProxyTransformSource::GetTransform() const
 
 void UGizmoTransformProxyTransformSource::SetTransform(const FTransform& NewTransform)
 {
-	Proxy->SetTransform(NewTransform);
+	if (bSetPivotMode)
+	{
+		bool bProxySetPivotOriginal = Proxy->bSetPivotMode;
+		Proxy->bSetPivotMode = true;
+		Proxy->SetTransform(NewTransform);
+		Proxy->bSetPivotMode = bProxySetPivotOriginal;
+
+		OnPivotChanged.Broadcast(this);
+	}
+	else
+	{
+		Proxy->SetTransform(NewTransform);
+		OnTransformChanged.Broadcast(this);
+	}
+}
+
+FTransform UGizmoScaledAndUnscaledTransformSources::GetTransform() const
+{
+	return ScaledTransformSource->GetTransform();
+}
+
+void UGizmoScaledAndUnscaledTransformSources::SetTransform(const FTransform& NewTransform)
+{
+	if (UnscaledTransformSource)
+	{
+		FTransform Unscaled(NewTransform);
+		Unscaled.SetScale3D(FVector::OneVector);
+		UnscaledTransformSource->SetTransform(Unscaled);
+	}
+
+	ScaledTransformSource->SetTransform(NewTransform);
+
 	OnTransformChanged.Broadcast(this);
 }
\ No newline at end of file
diff --git a/Engine/Source/Runtime/Experimental/InteractiveToolsFramework/Private/InteractiveGizmoManager.cpp b/Engine/Source/Runtime/Experimental/InteractiveToolsFramework/Private/InteractiveGizmoManager.cpp
--- a/Engine/Source/Runtime/Experimental/InteractiveToolsFramework/Private/InteractiveGizmoManager.cpp
+++ b/Engine/Source/Runtime/Experimental/InteractiveToolsFramework/Private/InteractiveGizmoManager.cpp
@@ -9,6 +9,7 @@
 #include "BaseGizmos/AxisAngleGizmo.h"
 #include "BaseGizmos/TransformGizmo.h"
 #include "BaseGizmos/IntervalGizmo.h"
+#include "BaseGizmos/RepositionableTransformGizmo.h"
 #include "BaseGizmos/ScalableSphereGizmo.h"
 
 #define LOCTEXT_NAMESPACE "UInteractiveGizmoManager"
@@ -269,6 +270,7 @@ FString UInteractiveGizmoManager::DefaultPlanePositionBuilderIdentifier = TEXT("
 FString UInteractiveGizmoManager::DefaultAxisAngleBuilderIdentifier = TEXT("StandardXFormAxisRotationGizmo");
 FString UInteractiveGizmoManager::DefaultThreeAxisTransformBuilderIdentifier = TEXT("DefaultThreeAxisTransformBuilderIdentifier");
 const FString UInteractiveGizmoManager::CustomThreeAxisTransformBuilderIdentifier = TEXT("CustomThreeAxisTransformBuilderIdentifier");
+const FString UInteractiveGizmoManager::CustomRepositionableThreeAxisTransformBuilderIdentifier = TEXT("CustomRepositionableThreeAxisTransformBuilderIdentifier");
 FString UInteractiveGizmoManager::DefaultScalableSphereBuilderIdentifier = TEXT("DefaultScalableSphereBuilderIdentifier");
 
 void UInteractiveGizmoManager::RegisterDefaultGizmos()
@@ -287,10 +289,16 @@ void UInteractiveGizmoManager::RegisterDefaultGizmos()
 	UTransformGizmoBuilder* TransformBuilder = NewObject<UTransformGizmoBuilder>();
 	RegisterGizmoType(DefaultThreeAxisTransformBuilderIdentifier, TransformBuilder);
 
-	CustomThreeAxisBuilder = NewObject<UTransformGizmoBuilder>();
-	CustomThreeAxisBuilder->GizmoActorBuilder = MakeShared<FTransformGizmoActorFactory>();
+	GizmoActorBuilder = MakeShared<FTransformGizmoActorFactory>();
+
+	UTransformGizmoBuilder* CustomThreeAxisBuilder = NewObject<UTransformGizmoBuilder>();
+	CustomThreeAxisBuilder->GizmoActorBuilder = GizmoActorBuilder;
 	RegisterGizmoType(CustomThreeAxisTransformBuilderIdentifier, CustomThreeAxisBuilder);
 
+	URepositionableTransformGizmoBuilder* CustomRepositionableThreeAxisBuilder = NewObject<URepositionableTransformGizmoBuilder>();
+	CustomRepositionableThreeAxisBuilder->GizmoActorBuilder = GizmoActorBuilder;
+	RegisterGizmoType(CustomRepositionableThreeAxisTransformBuilderIdentifier, CustomRepositionableThreeAxisBuilder);
+
 	UIntervalGizmoBuilder* IntervalGizmoBuilder = NewObject<UIntervalGizmoBuilder>();
 	RegisterGizmoType(UIntervalGizmo::GizmoName, IntervalGizmoBuilder);
 
@@ -311,11 +319,20 @@ UTransformGizmo* UInteractiveGizmoManager::Create3AxisTransformGizmo(void* Owner
 UTransformGizmo* UInteractiveGizmoManager::CreateCustomTransformGizmo(ETransformGizmoSubElements Elements, void* Owner, const FString& InstanceIdentifier)
 {
 	check(bDefaultGizmosRegistered);
-	CustomThreeAxisBuilder->GizmoActorBuilder->EnableElements = Elements;
+	GizmoActorBuilder->EnableElements = Elements;
 	UInteractiveGizmo* NewGizmo = CreateGizmo(CustomThreeAxisTransformBuilderIdentifier, InstanceIdentifier, Owner);
 	check(NewGizmo);
 	return Cast<UTransformGizmo>(NewGizmo);
 }
 
+UTransformGizmo* UInteractiveGizmoManager::CreateCustomRepositionableTransformGizmo(ETransformGizmoSubElements Elements, void* Owner, const FString& InstanceIdentifier)
+{
+	check(bDefaultGizmosRegistered);
+	GizmoActorBuilder->EnableElements = Elements;
+	UInteractiveGizmo* NewGizmo = CreateGizmo(CustomRepositionableThreeAxisTransformBuilderIdentifier, InstanceIdentifier, Owner);
+	check(NewGizmo);
+	return Cast<UTransformGizmo>(NewGizmo);
+}
+
 
 #undef LOCTEXT_NAMESPACE
diff --git a/Engine/Source/Runtime/Experimental/InteractiveToolsFramework/Public/BaseGizmos/AxisAngleGizmo.h b/Engine/Source/Runtime/Experimental/InteractiveToolsFramework/Public/BaseGizmos/AxisAngleGizmo.h
--- a/Engine/Source/Runtime/Experimental/InteractiveToolsFramework/Public/BaseGizmos/AxisAngleGizmo.h
+++ b/Engine/Source/Runtime/Experimental/InteractiveToolsFramework/Public/BaseGizmos/AxisAngleGizmo.h
@@ -9,7 +9,7 @@
 #include "BaseGizmos/GizmoComponents.h"
 #include "AxisAngleGizmo.generated.h"
 
-
+class UClickDragInputBehavior;
 
 
 UCLASS()
@@ -65,6 +65,35 @@ public:
 	UPROPERTY()
 	TScriptInterface<IGizmoStateTarget> StateTarget;
 
+	/** The mouse click behavior of the gizmo is accessible so that it can be modified to use different mouse keys. */
+	UPROPERTY()
+	UClickDragInputBehavior* MouseBehavior;
+
+	/** 
+	 * This gets checked to see if we should use the custom ray caster to get a destination point for the gizmo, rather
+	 * than grabbing the intersection with the gizmo plane.
+	 */
+	TUniqueFunction<bool()> ShouldUseCustomDestinationFunc = []() {return false; };
+
+	struct FCustomDestinationParams
+	{
+		// Right now we use the custom destination function for aligning to items in the scene, which
+		// we just need the world ray for. If we want to use functions that use other inputs as the
+		// basis for the destination, we would add those parameters here and would make sure that the
+		// gizmo passes them in.
+		const FRay* WorldRay = nullptr;
+	};
+
+	/**
+	 * If ShouldUseCustomDestinationFunc() returns true, this function is used to get a destination point, and
+	 * the output parameter is picked in such a way that the closest axis in the plane of the gizmo (the positive or
+	 * negative x or y, if we're rotating around z, for instance) moves to the closest point in the plane to the
+	 * destination point.
+	 * Used, for instance, for aligning to items in the scene.
+	 */
+	TUniqueFunction<bool(const FCustomDestinationParams& WorldRay, FVector& OutputPoint)> CustomDestinationFunc =
+		[](const FCustomDestinationParams& Params, FVector& OutputPoint) { return false; };
+
 public:
 	UPROPERTY()
 	bool bInInteraction = false;
@@ -98,6 +127,10 @@ protected:
 	FVector LastHitPosition;
 	float InitialTargetAngle;
 
+	// The angle of the closest axis in the plane of the angle gizmo (will be 0, pi/2, pi, or 3pi/2).
+	// Used for snapping the nearest axis to a ray cast point when using CustomDestinationRayCaster.
+	float ClosestAxisStartAngle;
+
 	static const int SnapAngleModifierID = 1;		// identifier we associate with the shift key
 	bool bEnableSnapAngleModifier = false;			// flag we use to keep track of modifier state
 
diff --git a/Engine/Source/Runtime/Experimental/InteractiveToolsFramework/Public/BaseGizmos/AxisPositionGizmo.h b/Engine/Source/Runtime/Experimental/InteractiveToolsFramework/Public/BaseGizmos/AxisPositionGizmo.h
--- a/Engine/Source/Runtime/Experimental/InteractiveToolsFramework/Public/BaseGizmos/AxisPositionGizmo.h
+++ b/Engine/Source/Runtime/Experimental/InteractiveToolsFramework/Public/BaseGizmos/AxisPositionGizmo.h
@@ -10,7 +10,7 @@
 #include "AxisPositionGizmo.generated.h"
 
 
-
+class UClickDragInputBehavior;
 
 UCLASS()
 class INTERACTIVETOOLSFRAMEWORK_API UAxisPositionGizmoBuilder : public UInteractiveGizmoBuilder
@@ -80,12 +80,37 @@ public:
 	UPROPERTY()
 	TScriptInterface<IGizmoStateTarget> StateTarget;
 
+	/** The mouse click behavior of the gizmo is accessible so that it can be modified to use different mouse keys. */
+	UPROPERTY()
+	UClickDragInputBehavior* MouseBehavior;
 
 public:
 	/** If enabled, then the sign on the parameter delta is always "increasing" when moving away from the origin point, rather than just being a projection onto the axis */
 	UPROPERTY()
 	bool bEnableSignedAxis = false;
 
+	/** 
+	 * This gets checked to see if we should use the custom destination function to get a destination point for
+	 * the gizmo, rather than grabbing the closest point on axis to ray.
+	 */
+	TUniqueFunction<bool()> ShouldUseCustomDestinationFunc = []() {return false; };
+
+	struct FCustomDestinationParams
+	{
+		// Right now we use the custom destination function for aligning to items in the scene, which
+		// we just need the world ray for. If we want to use functions that use other inputs as the
+		// basis for the destination (for instance, just the line parameter), we would add those
+		// parameters here and would make sure that the gizmo passes them in.
+		const FRay* WorldRay = nullptr;
+	};
+
+	/**
+	 * If ShouldUseCustomDestinationFunc() returns true, this function gets queried to get a destination point.
+	 * The gizmo parameter will then be picked in such a way that the gizmo origin moves to the closest point
+	 * on the axis to the destination point. Can be used, for example, to align to items in the scene.
+	 */
+	TUniqueFunction<bool(const FCustomDestinationParams& WorldRay, FVector& OutputPoint)> CustomDestinationFunc =
+		[](const FCustomDestinationParams& Params, FVector& OutputPoint) { return false; };
 
 public:
 	/** If true, we are in an active click+drag interaction, otherwise we are not */
@@ -122,5 +147,7 @@ public:
 protected:
 	FVector LastHitPosition;
 	float InitialTargetParameter;
+
+	float InteractionStartAxisOriginParameterOffset = 0;
 };
 
diff --git a/Engine/Source/Runtime/Experimental/InteractiveToolsFramework/Public/BaseGizmos/HitTargets.h b/Engine/Source/Runtime/Experimental/InteractiveToolsFramework/Public/BaseGizmos/HitTargets.h
--- a/Engine/Source/Runtime/Experimental/InteractiveToolsFramework/Public/BaseGizmos/HitTargets.h
+++ b/Engine/Source/Runtime/Experimental/InteractiveToolsFramework/Public/BaseGizmos/HitTargets.h
@@ -62,6 +62,14 @@ public:
 	UPROPERTY()
 	UPrimitiveComponent* Component;
 
+	/**
+	 * If set, this condition is checked before performing the hit test. This gives a way 
+	 * to disable the hit test without hiding the component. This is useful, for instance,
+	 * in a repositionable transform gizmo in world-coordinate mode, where the rotation
+	 * components need to be hittable for movement, but not for repositioning.
+	 */
+	TFunction<bool(const FInputDeviceRay&)> Condition = nullptr;
+
 	/** This function is called to update hover state of the target */
 	TFunction<void(bool)> UpdateHoverFunction;
 
diff --git a/Engine/Source/Runtime/Experimental/InteractiveToolsFramework/Public/BaseGizmos/PlanePositionGizmo.h b/Engine/Source/Runtime/Experimental/InteractiveToolsFramework/Public/BaseGizmos/PlanePositionGizmo.h
--- a/Engine/Source/Runtime/Experimental/InteractiveToolsFramework/Public/BaseGizmos/PlanePositionGizmo.h
+++ b/Engine/Source/Runtime/Experimental/InteractiveToolsFramework/Public/BaseGizmos/PlanePositionGizmo.h
@@ -9,7 +9,7 @@
 #include "BaseGizmos/GizmoComponents.h"
 #include "PlanePositionGizmo.generated.h"
 
-
+class UClickDragInputBehavior;
 
 
 UCLASS()
@@ -80,6 +80,10 @@ public:
 	UPROPERTY()
 	TScriptInterface<IGizmoStateTarget> StateTarget;
 
+	/** The mouse click behavior of the gizmo is accessible so that it can be modified to use different mouse keys. */
+	UPROPERTY()
+	UClickDragInputBehavior* MouseBehavior;
+
 public:
 	/** If enabled, then the sign on the parameter delta is always "increasing" when moving away from the origin point, rather than just being a projection onto the axis */
 	UPROPERTY()
@@ -93,6 +97,29 @@ public:
 	UPROPERTY()
 	bool bFlipY = false;
 
+	/** 
+	 * This gets checked to see if we should use the custom ray caster to get a destination point for the gizmo, rather
+	 * than grabbing the intersection with the gizmo plane.
+	 */
+	TUniqueFunction<bool()> ShouldUseCustomDestinationFunc = []() {return false; };
+
+	struct FCustomDestinationParams
+	{
+		// Right now we use the custom destination function for aligning to items in the scene, which
+		// we just need the world ray for. If we want to use functions that use other inputs as the
+		// basis for the destination, we would add those parameters here and would make sure that the 
+		// gizmo passes them in.
+		const FRay* WorldRay = nullptr;
+	};
+
+	/**
+	 * If ShouldUseCustomDestinationFunc() returns true, this function is used to get a destination point, and
+	 * the output parameters are picked in such a way that the axis origin moves to the closest point in the plane 
+	 * to the destination point.
+	 * Used, for instance, for aligning to items in the scene.
+	 */
+	TUniqueFunction<bool(const FCustomDestinationParams& WorldRay, FVector& OutputPoint)> CustomDestinationFunc =
+		[](const FCustomDestinationParams& Params, FVector& OutputPoint) { return false; };
 
 public:
 	/** If true, we are in an active click+drag interaction, otherwise we are not */
@@ -136,5 +163,8 @@ public:
 protected:
 	FVector LastHitPosition;
 	FVector2D InitialTargetParameter;
+
+	FVector2D InteractionStartOriginParameterOffset;
+	FInputDeviceRay LastInputRay = FInputDeviceRay(FRay());
 };
 
diff --git a/Engine/Source/Runtime/Experimental/InteractiveToolsFramework/Public/BaseGizmos/RepositionableTransformGizmo.h b/Engine/Source/Runtime/Experimental/InteractiveToolsFramework/Public/BaseGizmos/RepositionableTransformGizmo.h
new file mode 100644
--- /dev/null
+++ b/Engine/Source/Runtime/Experimental/InteractiveToolsFramework/Public/BaseGizmos/RepositionableTransformGizmo.h
@@ -0,0 +1,75 @@
+// Copyright Epic Games, Inc. All Rights Reserved.
+
+#pragma once
+
+#include "CoreMinimal.h"
+
+#include "BaseGizmos/TransformGizmo.h"
+#include "InteractiveGizmoBuilder.h"
+#include "ToolContextInterfaces.h"
+
+#include "RepositionableTransformGizmo.generated.h"
+
+class IToolContextTransactionProvider;
+class UTransformProxy;
+
+UCLASS()
+class INTERACTIVETOOLSFRAMEWORK_API URepositionableTransformGizmoBuilder : public UTransformGizmoBuilder
+{
+	GENERATED_BODY()
+
+public:
+
+	virtual UInteractiveGizmo* BuildGizmo(const FToolBuilderState& SceneState) const override;
+};
+
+
+/**
+ * A transform gizmo that also allows the user to reposition it by middle-clicking rotation/translation components.
+ */
+UCLASS()
+class INTERACTIVETOOLSFRAMEWORK_API URepositionableTransformGizmo : public UTransformGizmo
+{
+	GENERATED_BODY()
+
+public:
+
+	/**
+	 * Set the active target object for the Gizmo
+	 * @param Target active target
+	 * @param TransactionProvider optional IToolContextTransactionProvider implementation to use - by default uses GizmoManager
+	 */
+	virtual void SetActiveTarget(UTransformProxy* Target, IToolContextTransactionProvider* TransactionProvider = nullptr) override;
+
+	/**
+	 * Allows the user to provide functions to use for aligning the gizmo destination to items in the scene. Note that this affects
+	 * both the movement and repositioning.
+	 */
+	virtual void SetWorldAlignmentFunctions(
+		TUniqueFunction<bool()>&& ShouldAlignTranslationIn,
+		TUniqueFunction<bool(const FRay&, FVector&)>&& TranslationAlignmentRayCasterIn) override;
+
+	/**
+	 * Allows the user to provide functions to use in aligning the gizmo during repositioning only.
+	 * This is because the repositining is likely to want different raycasts that don't ignore the 
+	 * moved components themselves, so that the gizmo can be aligned to a target component.
+	 */
+	virtual void SetPivotAlignmentFunctions(
+		TUniqueFunction<bool()>&& ShouldAlignPivotIn,
+		TUniqueFunction<bool(const FRay&, FVector&)>&& PivotAlignmentRayCasterIn);
+protected:
+
+	TUniqueFunction<bool()> ShouldAlignPivot = []() { return false; };
+	TUniqueFunction<bool(const FRay&, FVector&)> PivotAlignmentRayCaster = [](const FRay&, FVector&) {return false; };
+
+	/** Subset of ActiveGizmos, for use by SetPivotAlignmentFunctions */
+	TArray<UInteractiveGizmo*> PivotAlignmentGizmos;
+
+	UPROPERTY()
+	UGizmoTransformChangeStateTarget* RepositionStateTarget;
+
+	// Helper functions
+	void ModifyPivotAxisGizmo(UInteractiveGizmo* SubGizmoIn);
+	void ModifyPivotPlaneGizmo(UInteractiveGizmo* SubGizmoIn);
+	void ModifyPivotRotateGizmo(UInteractiveGizmo* SubGizmoIn);
+};
diff --git a/Engine/Source/Runtime/Experimental/InteractiveToolsFramework/Public/BaseGizmos/TransformGizmo.h b/Engine/Source/Runtime/Experimental/InteractiveToolsFramework/Public/BaseGizmos/TransformGizmo.h
--- a/Engine/Source/Runtime/Experimental/InteractiveToolsFramework/Public/BaseGizmos/TransformGizmo.h
+++ b/Engine/Source/Runtime/Experimental/InteractiveToolsFramework/Public/BaseGizmos/TransformGizmo.h
@@ -19,8 +19,8 @@ class IGizmoStateTarget;
 class UGizmoConstantFrameAxisSource;
 class UGizmoComponentAxisSource;
 class UGizmoTransformChangeStateTarget;
-class UGizmoScaledTransformSource;
 class FTransformGizmoTransformChange;
+class FRay;
 
 /**
  * ATransformGizmoActor is an Actor type intended to be used with UTransformGizmo,
@@ -236,7 +236,7 @@ public:
  * 
  */
 UCLASS()
-class INTERACTIVETOOLSFRAMEWORK_API UTransformGizmo : public UInteractiveGizmo, public IToolCommandChangeSource
+class INTERACTIVETOOLSFRAMEWORK_API UTransformGizmo : public UInteractiveGizmo
 {
 	GENERATED_BODY()
 
@@ -246,6 +246,37 @@ public:
 	virtual void SetGizmoActorBuilder(TSharedPtr<FTransformGizmoActorFactory> Builder);
 	virtual void SetUpdateHoverFunction(TFunction<void(UPrimitiveComponent*, bool)> HoverFunction);
 	virtual void SetUpdateCoordSystemFunction(TFunction<void(UPrimitiveComponent*, EToolContextCoordinateSystem)> CoordSysFunction);
+	
+	/**
+	 * If used, binds alignment functions to the sub gizmos that they can use to align to geometry in the scene. 
+	 * Specifically, translation and rotation gizmos will check ShouldAlignDestination() to see if they should
+	 * use the custom ray caster (this allows the behavior to respond to modifier key presses, for instance),
+	 * and then use DestinationAlignmentRayCaster() to find a point to align to.
+	 * Subgizmos align to the point in different ways, usually by projecting onto the axis or plane that they
+	 * operate in.
+	 */
+	virtual void SetWorldAlignmentFunctions(
+		TUniqueFunction<bool()>&& ShouldAlignDestination,
+		TUniqueFunction<bool(const FRay&, FVector&)>&& DestinationAlignmentRayCaster
+		);
+
+	/**
+	 * By default, non-uniform scaling handles appear (assuming they exist in the gizmo to begin with), 
+	 * when CurrentCoordinateSystem == EToolContextCoordinateSystem::Local, since components can only be
+	 * locally scaled. However, this can be changed to a custom check here, perhaps to hide them in extra
+	 * conditions or to always show them (if the gizmo is not scaling a component).
+	 */
+	virtual void SetIsNonUniformScaleAllowedFunction(
+		TUniqueFunction<bool()>&& IsNonUniformScaleAllowed
+	);
+
+	/**
+	 * By default, the nonuniform scale components can scale negatively. However, they can be made to clamp
+	 * to zero instead by passing true here. This is useful for using the gizmo to flatten geometry.
+	 *
+	 * TODO: Should this affect uniform scaling too?
+	 */
+	virtual void SetDisallowNegativeScaling(bool bDisallow);
 
 	// UInteractiveGizmo overrides
 	virtual void Setup() override;
@@ -431,42 +462,16 @@ protected:
 	UPROPERTY()
 	UGizmoTransformChangeStateTarget* StateTarget;
 
-
 	/**
-	 * This TransformSource wraps a UGizmoComponentWorldTransformSource that is on the Gizmo Actor directly.
-	 * It tracks the scaling separately (SeparateChildScale is provided as the storage for the scaling).
-	 * This allows the various scaling handles to update the Transform without actually scaling the Gizmo itself.
+	 * These are used to let the translation subgizmos use raycasts into the scene to align the gizmo with scene geometry.
+	 * See comment for SetWorldAlignmentFunctions().
 	 */
-	UPROPERTY()
-	UGizmoScaledTransformSource* ScaledTransformSource;
-
-	/**
-	 * Scaling on the active target UTransformProxy is stored here. Undo/Redo events update this value via an FTransformGizmoTransformChange.
-	 */
-	FVector SeparateChildScale = FVector(1,1,1);
-
-	/**
-	 * This function is called by FTransformGizmoTransformChange to update SeparateChildScale on Undo/Redo.
-	 */
-	void ExternalSetChildScale(const FVector& NewScale);
-
-	friend class FTransformGizmoTransformChange;
-	/**
-	 * Ongoing transform change is tracked here, initialized in the IToolCommandChangeSource implementation.
-	 * Note that currently we only handle the SeparateChildScale this way. The transform changes on the target
-	 * objects are dealt with using the editor Transaction system (and on the Proxy using a FTransformProxyChangeSource attached to the StateTarget).
-	 * The Transaction system is not available at runtime, we should track the entire change in the FChange and update the proxy
-	 */
-	TUniquePtr<FTransformGizmoTransformChange> ActiveChange;
-
-public:
-	// IToolCommandChangeSource implementation, used to initialize and emit ActiveChange
-	virtual void BeginChange();
-	virtual TUniquePtr<FToolCommandChange> EndChange();
-	virtual UObject* GetChangeTarget();
-	virtual FText GetChangeDescription();
+	TUniqueFunction<bool()> ShouldAlignDestination = []() { return false; };
+	TUniqueFunction<bool(const FRay&, FVector&)> DestinationAlignmentRayCaster = [](const FRay&, FVector&) {return false; };
 
+	TUniqueFunction<bool()> IsNonUniformScaleAllowed = [this]() { return CurrentCoordinateSystem == EToolContextCoordinateSystem::Local; };
 
+	bool bDisallowNegativeScaling = false;
 protected:
 
 
@@ -517,27 +522,3 @@ protected:
 	FQuat RotationSnapFunction(const FQuat& DeltaRotation) const;
 
 };
-
-
-
-
-/**
- * FChange for a UTransformGizmo that applies transform change.
- * Currently only handles the scaling part of a transform.
- */
-class INTERACTIVETOOLSFRAMEWORK_API FTransformGizmoTransformChange : public FToolCommandChange
-{
-public:
-	FVector ChildScaleBefore;
-	FVector ChildScaleAfter;
-
-	/** Makes the change to the object */
-	virtual void Apply(UObject* Object) override;
-
-	/** Reverts change to the object */
-	virtual void Revert(UObject* Object) override;
-
-	/** Describes this change (for debugging) */
-	virtual FString ToString() const override;
-};
-
diff --git a/Engine/Source/Runtime/Experimental/InteractiveToolsFramework/Public/BaseGizmos/TransformProxy.h b/Engine/Source/Runtime/Experimental/InteractiveToolsFramework/Public/BaseGizmos/TransformProxy.h
--- a/Engine/Source/Runtime/Experimental/InteractiveToolsFramework/Public/BaseGizmos/TransformProxy.h
+++ b/Engine/Source/Runtime/Experimental/InteractiveToolsFramework/Public/BaseGizmos/TransformProxy.h
@@ -59,22 +59,48 @@ public:
 	 */
 	virtual void EndTransformEditSequence();
 
+	/**
+	 * Clients should call this before a sequence of SetTransform calls that have bSetPivotMode as true
+	 * (see comment in BeginTransformEditSequence).
+	 */
+	virtual void BeginPivotEditSequence();
+
+	/**
+	 * Clients should call this when done with a sequence of SetTransform that have bSetPivotMode as true 
+	 * (see comment in BeginTransformEditSequence).
+	 */
+	virtual void EndPivotEditSequence();
+
 public:
 	/**
-	 * This delegate is fired whenever the internal transform changes, ie
-	 * on AddComponent and SetTransform
+	 * This delegate is fired whenever the transform changes in a way that updates the contained components. I.e.,
+	 * SetTransform is called with bSetPivotMode being false.
 	 */
 	DECLARE_MULTICAST_DELEGATE_TwoParams(FOnTransformChanged, UTransformProxy*, FTransform);
 	FOnTransformChanged OnTransformChanged;
 
-	/** This delegate is fired when BeginTransformEditSequence() is called to indicate that a transform change has started */
+	/** This delegate is fired when BeginTransformEditSequence() is called to indicate that a sequence of transform updates has started */
 	DECLARE_MULTICAST_DELEGATE_OneParam(FOnBeginTransformEdit, UTransformProxy*);
 	FOnBeginTransformEdit OnBeginTransformEdit;
 
-	/** This delegate is fired when EndTransformEditSequence() is called to indicate that a transform change has ended */
+	/** This delegate is fired when EndTransformEditSequence() is called to indicate that a sequence of transform updates has ended */
 	DECLARE_MULTICAST_DELEGATE_OneParam(FOnEndTransformEdit, UTransformProxy*);
 	FOnEndTransformEdit OnEndTransformEdit;
 
+	/**
+	 * This delegate is fired whenever the internal transform changes due to a pivot reposition, ie on AddComponent and
+	 * when SetTransform is called with bSetPivotMode being true.
+	 */
+	DECLARE_MULTICAST_DELEGATE_TwoParams(FOnPivotChanged, UTransformProxy*, FTransform);
+	FOnPivotChanged OnPivotChanged;
+
+	/** This delegate is fired when BeginTransformEditSequence() is called to indicate that a sequence of pivot updates has started */
+	DECLARE_MULTICAST_DELEGATE_OneParam(FOnBeginPivotEdit, UTransformProxy*);
+	FOnBeginPivotEdit OnBeginPivotEdit;
+
+	/** This delegate is fired when EndTransformEditSequence() is called to indicate that a sequence of pivot updates has ended */
+	DECLARE_MULTICAST_DELEGATE_OneParam(FOnEndPivotEdit, UTransformProxy*);
+	FOnEndPivotEdit OnEndPivotEdit;
 
 	/**
 	 * If true, relative rotation of shared transform is applied to objects before relative translation (ie they rotate in place)
@@ -133,6 +159,7 @@ class INTERACTIVETOOLSFRAMEWORK_API FTransformProxyChange : public FToolCommandC
 public:
 	FTransform From;
 	FTransform To;
+	bool bSetPivotMode = false;
 
 	virtual void Apply(UObject* Object) override;
 	virtual void Revert(UObject* Object) override;
@@ -161,4 +188,6 @@ public:
 	virtual TUniquePtr<FToolCommandChange> EndChange() override;
 	virtual UObject* GetChangeTarget() override;
 	virtual FText GetChangeDescription() override;
+
+	bool bSetPivotMode = false;
 };
\ No newline at end of file
diff --git a/Engine/Source/Runtime/Experimental/InteractiveToolsFramework/Public/BaseGizmos/TransformSources.h b/Engine/Source/Runtime/Experimental/InteractiveToolsFramework/Public/BaseGizmos/TransformSources.h
--- a/Engine/Source/Runtime/Experimental/InteractiveToolsFramework/Public/BaseGizmos/TransformSources.h
+++ b/Engine/Source/Runtime/Experimental/InteractiveToolsFramework/Public/BaseGizmos/TransformSources.h
@@ -76,6 +76,9 @@ public:
 
 
 /**
+ * @deprecated This class was used with UGizmoScaledTransformSource, but UGizmoScaledAndUnscaledTransformSources
+ *  will likely handle any similar use cases in a cleaner way.
+ *
  * FSeparateScaleProvider provides TFunction-implementable API that sets/gets a
  * Scaling Vector from an external source. 
  */
@@ -87,6 +90,9 @@ struct FSeparateScaleProvider
 
 
 /**
+ * @deprecated Use UGizmoScaledAndUnscaledTransformSources instead.
+ *
+ * Old description:
  * UGizmoScaledTransformSource wraps another IGizmoTransformSource implementation and adds a
  * separate scaling vector to the Transform. The main use of this class is to support scaling
  * in a 3D gizmo without actually scaling the Gizmo itself. Generally our pattern is to apply
@@ -165,6 +171,15 @@ public:
 	UPROPERTY()
 	UTransformProxy* Proxy;
 
+	/**
+	 * The underlying proxy is modified with its SetPivotMode flag set to this. Thus, the transform source
+	 * can be used for either proxy repositioning or actual use.
+	 */
+	bool bSetPivotMode = false;
+
+	DECLARE_MULTICAST_DELEGATE_OneParam(FOnGizmoTransformSourcePivotChanged, IGizmoTransformSource*);
+	FOnGizmoTransformSourcePivotChanged OnPivotChanged;
+
 public:
 	/**
 	 * Construct a default instance of UGizmoComponentWorldTransformSource with the given Proxy
@@ -179,3 +194,64 @@ public:
 	}
 };
 
+
+/**
+ * A wrapper around two IGizmoTransformSource's that generally forwards transforms to/from its ScaledTransformSource,
+ * but also forwards an unscaled version of the transform to UnscaledTransformSource on SetTransform calls.
+ * This handles the common case of wanting to apply the entire transform to one IGizmoTransformSource, but only
+ * the unscaled transform to a gizmo component (since we don't want to scale the gizmo component but do want to 
+ * rotate/translate it).
+ */
+UCLASS()
+class INTERACTIVETOOLSFRAMEWORK_API UGizmoScaledAndUnscaledTransformSources : public UGizmoBaseTransformSource
+{
+	GENERATED_BODY()
+public:
+
+	/** Gets the transform from ScaledTransformSource. */
+	virtual FTransform GetTransform() const override;
+
+	/** Calls SetTransform on ScaledTransformSource and passes the unscaled version to UnscaledTransformSource. */
+	virtual void SetTransform(const FTransform& NewTransform) override;
+
+	UPROPERTY()
+	TScriptInterface<IGizmoTransformSource> ScaledTransformSource;
+
+	UPROPERTY()
+	TScriptInterface<IGizmoTransformSource> UnscaledTransformSource;
+
+public:
+	/**
+	 * Constructs a UGizmoScaledAndUnscaledTransformSources by wrapping the provided IGizmoTransformSource as
+	 * its scaled transform source, and the given gizmo in a UGizmoComponentWorldTransformSource as its
+	 * unscaled transform source.
+	 */
+	static UGizmoScaledAndUnscaledTransformSources* Construct(
+		IGizmoTransformSource* ScaledSource,
+		USceneComponent* GizmoComponentIn,
+		UObject* Outer = (UObject*)GetTransientPackage())
+	{
+		UGizmoScaledAndUnscaledTransformSources* NewSource = NewObject<UGizmoScaledAndUnscaledTransformSources>(Outer);
+		NewSource->ScaledTransformSource = Cast<UObject>(ScaledSource);
+
+		UGizmoComponentWorldTransformSource* WrappedGizmo = UGizmoComponentWorldTransformSource::Construct(GizmoComponentIn, Outer);
+		WrappedGizmo->bModifyComponentOnTransform = false;
+
+		NewSource->UnscaledTransformSource = WrappedGizmo;
+		return NewSource;
+	}
+
+	/**
+	 * Constructs a UGizmoScaledAndUnscaledTransformSources around the given two IGizmoTransformSource's.
+	 */
+	static UGizmoScaledAndUnscaledTransformSources* Construct(
+		IGizmoTransformSource* ScaledSource,
+		IGizmoTransformSource* UnscaledSource,
+		UObject* Outer = (UObject*)GetTransientPackage())
+	{
+		UGizmoScaledAndUnscaledTransformSources* NewSource = NewObject<UGizmoScaledAndUnscaledTransformSources>(Outer);
+		NewSource->ScaledTransformSource = Cast<UObject>(ScaledSource);
+		NewSource->UnscaledTransformSource = Cast<UObject>(UnscaledSource);
+		return NewSource;
+	}
+};
diff --git a/Engine/Source/Runtime/Experimental/InteractiveToolsFramework/Public/InteractiveGizmoManager.h b/Engine/Source/Runtime/Experimental/InteractiveToolsFramework/Public/InteractiveGizmoManager.h
--- a/Engine/Source/Runtime/Experimental/InteractiveToolsFramework/Public/InteractiveGizmoManager.h
+++ b/Engine/Source/Runtime/Experimental/InteractiveToolsFramework/Public/InteractiveGizmoManager.h
@@ -11,6 +11,7 @@
 #include "ToolContextInterfaces.h"
 #include "InteractiveGizmoManager.generated.h"
 
+class FTransformGizmoActorFactory;
 class UTransformGizmo;
 class UTransformGizmoBuilder;
 
@@ -210,6 +211,7 @@ public:
 	 */
 	virtual UTransformGizmo* CreateCustomTransformGizmo(ETransformGizmoSubElements Elements, void* Owner = nullptr, const FString& InstanceIdentifier = FString());
 
+	virtual UTransformGizmo* CreateCustomRepositionableTransformGizmo(ETransformGizmoSubElements Elements, void* Owner = nullptr, const FString& InstanceIdentifier = FString());
 
 public:
 	// builder identifiers for default gizmo types. Perhaps should have an API for this...
@@ -218,6 +220,7 @@ public:
 	static FString DefaultAxisAngleBuilderIdentifier;
 	static FString DefaultThreeAxisTransformBuilderIdentifier;
 	static const FString CustomThreeAxisTransformBuilderIdentifier;
+	static const FString CustomRepositionableThreeAxisTransformBuilderIdentifier;
 	static FString DefaultScalableSphereBuilderIdentifier;
 
 protected:
@@ -240,5 +243,5 @@ protected:
 	TMap<FString, UInteractiveGizmoBuilder*> GizmoBuilders;
 
 	bool bDefaultGizmosRegistered = false;
-	UTransformGizmoBuilder* CustomThreeAxisBuilder = nullptr;
+	TSharedPtr<FTransformGizmoActorFactory> GizmoActorBuilder;
 };
diff --git a/Engine/Source/Runtime/Experimental/InteractiveToolsFramework/Public/ToolContextInterfaces.h b/Engine/Source/Runtime/Experimental/InteractiveToolsFramework/Public/ToolContextInterfaces.h
--- a/Engine/Source/Runtime/Experimental/InteractiveToolsFramework/Public/ToolContextInterfaces.h
+++ b/Engine/Source/Runtime/Experimental/InteractiveToolsFramework/Public/ToolContextInterfaces.h
@@ -147,6 +147,9 @@ struct INTERACTIVETOOLSFRAMEWORK_API FSceneSnapQueryRequest
 
 	/** Snap input rotation delta */
 	FQuat DeltaRotation;
+
+	/** Optional: components to consider invisible even if they aren't. */
+	const TArray<const UPrimitiveComponent*>* ComponentsToIgnore = nullptr;
 };
 
 
