Chaos : Plasticity

-Linear plasticity thresholds.

diff --git a/Engine/Source/Editor/DetailCustomizations/Private/PhysicsConstraintComponentDetails.cpp b/Engine/Source/Editor/DetailCustomizations/Private/PhysicsConstraintComponentDetails.cpp
--- a/Engine/Source/Editor/DetailCustomizations/Private/PhysicsConstraintComponentDetails.cpp
+++ b/Engine/Source/Editor/DetailCustomizations/Private/PhysicsConstraintComponentDetails.cpp
@@ -333,6 +333,9 @@ void FPhysicsConstraintComponentDetails::AddLinearLimits(IDetailLayoutBuilder& D
 #endif
 	LinearLimitCat.AddProperty(ProfilePropertiesProperty->GetChildHandle(GET_MEMBER_NAME_CHECKED(FConstraintProfileProperties, bLinearBreakable)));
 	LinearLimitCat.AddProperty(ProfilePropertiesProperty->GetChildHandle(GET_MEMBER_NAME_CHECKED(FConstraintProfileProperties, LinearBreakThreshold)));
+
+	LinearLimitCat.AddProperty(ProfilePropertiesProperty->GetChildHandle(GET_MEMBER_NAME_CHECKED(FConstraintProfileProperties, bLinearPlasticity)).ToSharedRef());
+	LinearLimitCat.AddProperty(ProfilePropertiesProperty->GetChildHandle(GET_MEMBER_NAME_CHECKED(FConstraintProfileProperties, LinearPlasticityThreshold)).ToSharedRef());
 }
 
 
diff --git a/Engine/Source/Runtime/Engine/Classes/PhysicsEngine/ConstraintInstance.h b/Engine/Source/Runtime/Engine/Classes/PhysicsEngine/ConstraintInstance.h
--- a/Engine/Source/Runtime/Engine/Classes/PhysicsEngine/ConstraintInstance.h
+++ b/Engine/Source/Runtime/Engine/Classes/PhysicsEngine/ConstraintInstance.h
@@ -46,6 +46,10 @@ struct ENGINE_API FConstraintProfileProperties
 	UPROPERTY(EditAnywhere, AdvancedDisplay, Category = Linear, meta = (editcondition = "bLinearBreakable", ClampMin = "0.0"))
 	float LinearBreakThreshold;
 
+	/** Percent threshold from target position needed to reset the spring rest length.*/
+	UPROPERTY(EditAnywhere, AdvancedDisplay, Category = Linear, meta = (editcondition = "bLinearPlasticity", ClampMin = "0.0"))
+	float LinearPlasticityThreshold;
+
 	/** Torque needed to break the joint. */
 	UPROPERTY(EditAnywhere, AdvancedDisplay, Category = Angular, meta = (editcondition = "bAngularBreakable", ClampMin = "0.0"))
 	float AngularBreakThreshold;
@@ -113,6 +117,10 @@ struct ENGINE_API FConstraintProfileProperties
 	UPROPERTY(EditAnywhere, AdvancedDisplay, Category = Linear)
 	uint8 bLinearBreakable : 1;
 
+	/** Whether it is possible to reset spr4ing rest length from the linear deformation. */
+	UPROPERTY(EditAnywhere, AdvancedDisplay, Category = Linear)
+	uint8 bLinearPlasticity : 1;
+
 	FConstraintProfileProperties();
 
 	/** Updates physx joint properties from unreal properties (limits, drives, flags, etc...) */
@@ -472,6 +480,17 @@ public:
 		UpdateBreakable();
 	}
 
+	/** Sets the Linear Plasticity properties
+	*	@param bInLinearPlasticity 	Whether it is possible to reset the target angles
+	*	@param InLinearPlasticityThreshold	Delta from target needed to reset the target joint
+	*/
+	void SetLinearPlasticity(bool bInLinearPlasticity, float InLinearPlasticityThreshold)
+	{
+		ProfileInstance.bLinearPlasticity = bInLinearPlasticity;
+		ProfileInstance.LinearPlasticityThreshold = InLinearPlasticityThreshold;
+		UpdatePlasticity();
+	}
+
 	/** Sets the Angular Breakable properties
 	*	@param bInAngularBreakable		Whether it is possible to break the joint with angular force
 	*	@param InAngularBreakThreshold	Torque needed to break the joint
diff --git a/Engine/Source/Runtime/Engine/Classes/PhysicsEngine/PhysicsConstraintComponent.h b/Engine/Source/Runtime/Engine/Classes/PhysicsEngine/PhysicsConstraintComponent.h
--- a/Engine/Source/Runtime/Engine/Classes/PhysicsEngine/PhysicsConstraintComponent.h
+++ b/Engine/Source/Runtime/Engine/Classes/PhysicsEngine/PhysicsConstraintComponent.h
@@ -269,6 +269,13 @@ public:
 	UFUNCTION(BlueprintCallable, Category = "Physics|Components|PhysicsConstraint")
 	void SetLinearBreakable(bool bLinearBreakable, float LinearBreakThreshold);
 
+	/** Sets the Linear Plasticity properties
+	*	@param bLinearPlasticity		    Whether it is possible to reset the target angle from the Linear displacement
+	*	@param LinearPlasticityThreshold	Percent deformation needed to reset the rest length of the joint
+	*/
+	UFUNCTION(BlueprintCallable, Category = "Physics|Components|PhysicsConstraint")
+	void SetLinearPlasticity(bool bLinearPlasticity, float LinearPlasticityThreshold);
+
 	/** Sets the Angular Breakable properties
 	*	@param bAngularBreakable		Whether it is possible to break the joint with angular force
 	*	@param AngularBreakThreshold	Torque needed to break the joint
diff --git a/Engine/Source/Runtime/Engine/Private/PhysicsEngine/ConstraintInstance.cpp b/Engine/Source/Runtime/Engine/Private/PhysicsEngine/ConstraintInstance.cpp
--- a/Engine/Source/Runtime/Engine/Private/PhysicsEngine/ConstraintInstance.cpp
+++ b/Engine/Source/Runtime/Engine/Private/PhysicsEngine/ConstraintInstance.cpp
@@ -129,6 +129,7 @@ FConstraintProfileProperties::FConstraintProfileProperties()
 	, ProjectionLinearAlpha(1.0f)
 	, ProjectionAngularAlpha(0.0f)
 	, LinearBreakThreshold(300.f)
+	, LinearPlasticityThreshold(10.f)
 	, AngularBreakThreshold(500.f)
 	, AngularPlasticityThreshold(10.f)
 	, bDisableCollision(false)
@@ -138,6 +139,7 @@ FConstraintProfileProperties::FConstraintProfileProperties()
 	, bAngularBreakable(false)
 	, bAngularPlasticity(false)
 	, bLinearBreakable(false)
+	, bLinearPlasticity(false)
 {
 }
 
@@ -184,7 +186,7 @@ void FConstraintInstance::UpdatePlasticity()
 
 void FConstraintProfileProperties::UpdatePlasticity_AssumesLocked(const FPhysicsConstraintHandle& InConstraintRef) const
 {
-	const float LinearPlasticityLimit = MAX_FLT; // todo(Chaos) : implement
+	const float LinearPlasticityLimit = bLinearPlasticity ? LinearPlasticityThreshold : FLT_MAX;
 	const float AngularPlasticityLimit = bAngularPlasticity ? FMath::DegreesToRadians(AngularPlasticityThreshold) : MAX_FLT;
 
 	FPhysicsInterface::SetPlasticityLimits_AssumesLocked(InConstraintRef, LinearPlasticityLimit, AngularPlasticityLimit);
diff --git a/Engine/Source/Runtime/Engine/Private/PhysicsEngine/Experimental/PhysInterface_Chaos.cpp b/Engine/Source/Runtime/Engine/Private/PhysicsEngine/Experimental/PhysInterface_Chaos.cpp
--- a/Engine/Source/Runtime/Engine/Private/PhysicsEngine/Experimental/PhysInterface_Chaos.cpp
+++ b/Engine/Source/Runtime/Engine/Private/PhysicsEngine/Experimental/PhysInterface_Chaos.cpp
@@ -390,7 +390,10 @@ void FPhysInterface_Chaos::UpdateLinearDrive_AssumesLocked(const FPhysicsConstra
 				Constraint->SetLinearPositionDriveXEnabled(InDriveParams.XDrive.bEnablePositionDrive);
 				Constraint->SetLinearPositionDriveYEnabled(InDriveParams.YDrive.bEnablePositionDrive);
 				Constraint->SetLinearPositionDriveZEnabled(InDriveParams.ZDrive.bEnablePositionDrive);
-				Constraint->SetLinearDrivePositionTarget(InDriveParams.PositionTarget);
+				if (FMath::IsNearlyEqual(Constraint->GetLinearPlasticityLimit(), FLT_MAX))
+				{
+					Constraint->SetLinearDrivePositionTarget(InDriveParams.PositionTarget);
+				}
 			}
 
 			bool bVelocityDriveEnabled = InDriveParams.IsVelocityDriveEnabled();
diff --git a/Engine/Source/Runtime/Engine/Private/PhysicsEngine/ImmediatePhysics/ImmediatePhysicsChaos/ImmediatePhysicsJointHandle_Chaos.cpp b/Engine/Source/Runtime/Engine/Private/PhysicsEngine/ImmediatePhysics/ImmediatePhysicsChaos/ImmediatePhysicsJointHandle_Chaos.cpp
--- a/Engine/Source/Runtime/Engine/Private/PhysicsEngine/ImmediatePhysics/ImmediatePhysicsChaos/ImmediatePhysicsJointHandle_Chaos.cpp
+++ b/Engine/Source/Runtime/Engine/Private/PhysicsEngine/ImmediatePhysics/ImmediatePhysicsChaos/ImmediatePhysicsJointHandle_Chaos.cpp
@@ -112,6 +112,7 @@ namespace ImmediatePhysics_Chaos
 		ConstraintSettings.AngularDriveForceMode = EJointForceMode::Acceleration;
 
 		ConstraintSettings.LinearBreakForce = (Profile.bLinearBreakable) ? Chaos::ConstraintSettings::LinearBreakScale() * Profile.LinearBreakThreshold : FLT_MAX;
+		ConstraintSettings.LinearPlasticityLimit = (Profile.bLinearPlasticity) ? FMath::Clamp((float)Profile.LinearPlasticityThreshold, 0.f, 1.f) : FLT_MAX;
 		ConstraintSettings.AngularBreakTorque = (Profile.bAngularBreakable) ? Chaos::ConstraintSettings::AngularBreakScale() * Profile.AngularBreakThreshold : FLT_MAX;
 		ConstraintSettings.AngularPlasticityLimit = (Profile.bAngularPlasticity) ? Profile.AngularPlasticityThreshold : FLT_MAX;
 
diff --git a/Engine/Source/Runtime/Engine/Private/PhysicsEngine/PhysicsConstraintComponent.cpp b/Engine/Source/Runtime/Engine/Private/PhysicsEngine/PhysicsConstraintComponent.cpp
--- a/Engine/Source/Runtime/Engine/Private/PhysicsEngine/PhysicsConstraintComponent.cpp
+++ b/Engine/Source/Runtime/Engine/Private/PhysicsEngine/PhysicsConstraintComponent.cpp
@@ -656,6 +656,11 @@ void UPhysicsConstraintComponent::SetLinearBreakable(bool bLinearBreakable, floa
 	ConstraintInstance.SetLinearBreakable(bLinearBreakable, LinearBreakThreshold);
 }
 
+void UPhysicsConstraintComponent::SetLinearPlasticity(bool bLinearPlasticity, float LinearPlasticityThreshold)
+{
+	ConstraintInstance.SetLinearPlasticity(bLinearPlasticity, LinearPlasticityThreshold);
+}
+
 void UPhysicsConstraintComponent::SetAngularBreakable(bool bAngularBreakable, float AngularBreakThreshold)
 {
 	ConstraintInstance.SetAngularBreakable(bAngularBreakable, AngularBreakThreshold);
diff --git a/Engine/Source/Runtime/Experimental/Chaos/Private/Chaos/PBDJointConstraints.cpp b/Engine/Source/Runtime/Experimental/Chaos/Private/Chaos/PBDJointConstraints.cpp
--- a/Engine/Source/Runtime/Experimental/Chaos/Private/Chaos/PBDJointConstraints.cpp
+++ b/Engine/Source/Runtime/Experimental/Chaos/Private/Chaos/PBDJointConstraints.cpp
@@ -1313,9 +1313,9 @@ namespace Chaos
 			ApplyBreakThreshold(Dt, ConstraintIndex, Solver.GetNetLinearImpulse(), Solver.GetNetAngularImpulse());
 		}
 
-		if (/*(JointSettings.LinearPlasticityLimit != FLT_MAX) ||*/ (JointSettings.AngularPlasticityLimit != FLT_MAX)) // todo(chaos) : Implement Linear Plasticity
+		if ((JointSettings.LinearPlasticityLimit != FLT_MAX) || (JointSettings.AngularPlasticityLimit != FLT_MAX))
 		{
-			ApplyPlasticityLimits(Dt, ConstraintIndex, Solver.GetP(1) - Solver.GetP(0), Solver.GetQ(0).Inverse() * Solver.GetQ(1));
+			ApplyPlasticityLimits(Dt, ConstraintIndex, Particle1->X() - Particle0->X(), Particle0->R().Inverse() * Particle1->R());
 		}
 
 		return Solver.GetIsActive() || !bChaos_Joint_EarlyOut_Enabled;
@@ -1427,10 +1427,26 @@ namespace Chaos
 	void FPBDJointConstraints::ApplyPlasticityLimits(const FReal Dt, int32 ConstraintIndex, const FVec3& LinearDisplacement, const FRotation3& AngularDisplacement)
 	{
 		FPBDJointSettings& JointSettings = ConstraintSettings[ConstraintIndex];
+		FTransformPair& ConstraintFrame = ConstraintFrames[ConstraintIndex];
+
+		if (!FMath::IsNearlyEqual(JointSettings.LinearPlasticityLimit, FLT_MAX))
+		{
+			FTransform JointTransform = ConstraintFrame[0].GetRelativeTransform(ConstraintFrame[1]);
+			const FReal Delta = LinearDisplacement.Size();
+			const FReal TargetDelta = JointTransform.GetTranslation().Size();
+			if (!FMath::IsNearlyZero(TargetDelta) && !FMath::IsNearlyZero(Delta))
+			{
+				FReal Ratio = Delta / TargetDelta;
+				if ((1.f - Ratio) > JointSettings.LinearPlasticityLimit)
+				{
+					JointTransform.ScaleTranslation(Ratio);
+					ConstraintFrame[1] = JointTransform.GetRelativeTransformReverse(ConstraintFrame[0]);
+				}
+			}
+		}
 
-		// @todo(chaos) : Implement Linear case
 
-		if (JointSettings.AngularPlasticityLimit)
+		if (!FMath::IsNearlyEqual(JointSettings.AngularPlasticityLimit, FLT_MAX))
 		{
 			const FReal AngleDeg = JointSettings.AngularDrivePositionTarget.AngularDistance(AngularDisplacement);
 			if (AngleDeg > JointSettings.AngularPlasticityLimit)
diff --git a/Engine/Source/Runtime/Experimental/Chaos/Private/PhysicsProxy/JointConstraintProxy.cpp b/Engine/Source/Runtime/Experimental/Chaos/Private/PhysicsProxy/JointConstraintProxy.cpp
--- a/Engine/Source/Runtime/Experimental/Chaos/Private/PhysicsProxy/JointConstraintProxy.cpp
+++ b/Engine/Source/Runtime/Experimental/Chaos/Private/PhysicsProxy/JointConstraintProxy.cpp
@@ -345,7 +345,7 @@ void TJointConstraintProxy<Chaos::FJointConstraint>::PushStateOnPhysicsThread(Ch
 			if (DirtyFlagsBuffer.IsDirty(Chaos::EJointConstraintFlags::LinearBreakForce))
 			{
 				ConstraintSettings.LinearBreakForce = JointSettingsBuffer.LinearBreakForce;
-				ConstraintSettings.LinearPlasticityLimit = JointSettingsBuffer.LinearPlasticityLimit;
+				ConstraintSettings.LinearPlasticityLimit = FMath::Clamp((float)JointSettingsBuffer.LinearPlasticityLimit, 0.f, 1.f);
 			}
 
 			if (DirtyFlagsBuffer.IsDirty(Chaos::EJointConstraintFlags::AngularBreakTorque))
