Finer adjustment for FConvex convex generation since we are using the proper input thresholds

diff --git a/Engine/Source/Runtime/Experimental/Chaos/Public/Chaos/CollisionConvexMesh.h b/Engine/Source/Runtime/Experimental/Chaos/Public/Chaos/CollisionConvexMesh.h
--- a/Engine/Source/Runtime/Experimental/Chaos/Public/Chaos/CollisionConvexMesh.h
+++ b/Engine/Source/Runtime/Experimental/Chaos/Public/Chaos/CollisionConvexMesh.h
@@ -637,7 +637,7 @@ namespace Chaos
 			TR->Next = RS;
 			FVec3 RSTNormal = ComputeFaceNormal(InParticles.X(RS->Vertex), InParticles.X(ST->Vertex), InParticles.X(TR->Vertex));
 			const FReal RSTNormalSize = RSTNormal.Size();
-			check(RSTNormalSize > 1e-6);
+			check(RSTNormalSize > FLT_EPSILON);
 			RSTNormal = RSTNormal * (1 / RSTNormalSize);
 			FConvexFace* RST = Pool.AllocConvexFace(TPlaneConcrete<FReal, 3>(InParticles.X(RS->Vertex), RSTNormal));
 			RST->FirstEdge = RS;
