Chaos : Clustering

-Fixed crash in particle cleanup.

diff --git a/Engine/Source/Runtime/Experimental/Chaos/Private/PhysicsProxy/GeometryCollectionPhysicsProxy.cpp b/Engine/Source/Runtime/Experimental/Chaos/Private/PhysicsProxy/GeometryCollectionPhysicsProxy.cpp
--- a/Engine/Source/Runtime/Experimental/Chaos/Private/PhysicsProxy/GeometryCollectionPhysicsProxy.cpp
+++ b/Engine/Source/Runtime/Experimental/Chaos/Private/PhysicsProxy/GeometryCollectionPhysicsProxy.cpp
@@ -1362,15 +1362,18 @@ void FGeometryCollectionPhysicsProxy::OnRemoveFromSolver(Chaos::TPBDRigidsSolver
 
 	for (FClusterHandle* Handle : SolverParticleHandles)
 	{	
-		if(Chaos::TPBDRigidClusteredParticleHandle<Chaos::FReal, 3> * Cluster = Handle->CastToClustered())
+		if (Handle)
 		{
-			Evolution->GetRigidClustering().GetTopLevelClusterParents().Remove(Cluster);
-			Evolution->GetRigidClustering().GetChildrenMap().Remove(Cluster);
-			Evolution->DestroyParticle(Cluster);
-		}
-		else
-		{
-			Evolution->DestroyParticle(Handle);
+			if (Chaos::TPBDRigidClusteredParticleHandle<Chaos::FReal, 3>* Cluster = Handle->CastToClustered())
+			{
+				Evolution->GetRigidClustering().GetTopLevelClusterParents().Remove(Cluster);
+				Evolution->GetRigidClustering().GetChildrenMap().Remove(Cluster);
+				Evolution->DestroyParticle(Cluster);
+			}
+			else
+			{
+				Evolution->DestroyParticle(Handle);
+			}
 		}
 	}
 }
