Chaos : DebugDrawing
-Moved recent debug draw collision particles command to UE5
... to use:
... ... p.Chaos.DebugDrawing 1
... ... p.Chaos.DebugDraw.ShowCollisionParticles 1

diff --git a/Engine/Source/Runtime/Experimental/Chaos/Private/Chaos/ChaosDebugDraw.cpp b/Engine/Source/Runtime/Experimental/Chaos/Private/Chaos/ChaosDebugDraw.cpp
--- a/Engine/Source/Runtime/Experimental/Chaos/Private/Chaos/ChaosDebugDraw.cpp
+++ b/Engine/Source/Runtime/Experimental/Chaos/Private/Chaos/ChaosDebugDraw.cpp
@@ -9,6 +9,7 @@
 #include "Chaos/ImplicitObjectScaled.h"
 #include "Chaos/ImplicitObjectTransformed.h"
 #include "Chaos/ImplicitObjectUnion.h"
+#include "Chaos/Levelset.h"
 #include "Chaos/Particle/ParticleUtilities.h"
 #include "Chaos/ParticleHandle.h"
 #include "Chaos/PBDCollisionConstraints.h"
@@ -30,6 +31,9 @@ namespace Chaos
 		bool bChaosDebugDebugDrawShapeBounds = false;
 		FAutoConsoleVariableRef CVarChaosDebugDrawShapeBounds(TEXT("p.Chaos.DebugDraw.ShowShapeBounds"), bChaosDebugDebugDrawShapeBounds, TEXT("Whether to show the bounds of each shape in DrawShapes"));
 
+		bool bChaosDebugDebugDrawCollisionParticles = false;
+		FAutoConsoleVariableRef CVarChaosDebugDrawShapeParticles(TEXT("p.Chaos.DebugDraw.ShowCollisionParticles"), bChaosDebugDebugDrawCollisionParticles, TEXT("Whether to show the collision particles if present"));
+
 		bool bChaosDebugDebugDrawInactiveContacts = true;
 		FAutoConsoleVariableRef CVarChaosDebugDrawInactiveContacts(TEXT("p.Chaos.DebugDraw.ShowInactiveContacts"), bChaosDebugDebugDrawInactiveContacts, TEXT("Whether to show inactive contacts (ones that contributed no impulses or pushout)"));
 
@@ -71,15 +75,15 @@ namespace Chaos
 		//
 		//
 
-		void DrawShapesImpl(const FRigidTransform3& ShapeTransform, const FImplicitObject* Shape, const FColor& Color, const FChaosDebugDrawSettings& Settings);
+		void DrawShapesImpl(const TGeometryParticleHandle<FReal, 3>* Particle, const FRigidTransform3& ShapeTransform, const FImplicitObject* Shape, const FColor& Color, const FChaosDebugDrawSettings& Settings);
 
 		void DrawShape(const FRigidTransform3& ShapeTransform, const FImplicitObject* Shape, const FColor& Color, const FChaosDebugDrawSettings* Settings)
 		{
-			DrawShapesImpl(ShapeTransform, Shape, Color, GetChaosDebugDrawSettings(Settings));
+			DrawShapesImpl(nullptr, ShapeTransform, Shape, Color, GetChaosDebugDrawSettings(Settings));
 		}
 
 		template <bool bInstanced>
-		void DrawShapesScaledImpl(const FRigidTransform3& ShapeTransform, const FImplicitObject* Shape, const FColor& Color, const FChaosDebugDrawSettings& Settings)
+		void DrawShapesScaledImpl(const TGeometryParticleHandle<FReal, 3>* Particle, const FRigidTransform3& ShapeTransform, const FImplicitObject* Shape, const FColor& Color, const FChaosDebugDrawSettings& Settings)
 		{
 			const EImplicitObjectType PackedType = Shape->GetType();
 			const EImplicitObjectType InnerType = GetInnerType(PackedType);
@@ -109,7 +113,7 @@ namespace Chaos
 			{
 				const TImplicitObjectScaled<FConvex, bInstanced>* Scaled = Shape->template GetObject<TImplicitObjectScaled<FConvex, bInstanced>>();
 				ScaleTM.SetScale3D(Scaled->GetScale());
-				DrawShapesImpl(ScaleTM * ShapeTransform, Scaled->GetUnscaledObject(), Color, Settings);
+				DrawShapesImpl(Particle, ScaleTM * ShapeTransform, Scaled->GetUnscaledObject(), Color, Settings);
 				break;
 			}
 			case ImplicitObjectType::TaperedCylinder:
@@ -125,7 +129,7 @@ namespace Chaos
 			}
 		}
 
-		void DrawShapesInstancedImpl(const FRigidTransform3& ShapeTransform, const FImplicitObject* Shape, const FColor& Color, const FChaosDebugDrawSettings& Settings)
+		void DrawShapesInstancedImpl(const TGeometryParticleHandle<FReal, 3>* Particle, const FRigidTransform3& ShapeTransform, const FImplicitObject* Shape, const FColor& Color, const FChaosDebugDrawSettings& Settings)
 		{
 			const EImplicitObjectType PackedType = Shape->GetType();
 			const EImplicitObjectType InnerType = GetInnerType(PackedType);
@@ -153,7 +157,7 @@ namespace Chaos
 			case ImplicitObjectType::Convex:
 			{
 				const TImplicitObjectInstanced<FConvex>* Instanced = Shape->template GetObject<TImplicitObjectInstanced<FConvex>>();
-				DrawShapesImpl(ShapeTransform, Instanced->GetInstancedObject(), Color, Settings);
+				DrawShapesImpl(Particle, ShapeTransform, Instanced->GetInstancedObject(), Color, Settings);
 				break;
 			}
 			case ImplicitObjectType::TaperedCylinder:
@@ -169,7 +173,7 @@ namespace Chaos
 			}
 		}
 
-		void DrawShapesImpl(const FRigidTransform3& ShapeTransform, const FImplicitObject* Shape, const FColor& Color, const FChaosDebugDrawSettings& Settings)
+		void DrawShapesImpl(const TGeometryParticleHandle<FReal, 3>* Particle, const FRigidTransform3& ShapeTransform, const FImplicitObject* Shape, const FColor& Color, const FChaosDebugDrawSettings& Settings)
 		{
 			const EImplicitObjectType PackedType = Shape->GetType(); // Type includes scaling and instancing data
 			const EImplicitObjectType InnerType = GetInnerType(Shape->GetType());
@@ -179,17 +183,17 @@ namespace Chaos
 			{
 				if (IsInstanced(PackedType))
 				{
-					DrawShapesScaledImpl<true>(ShapeTransform, Shape, Color, Settings);
+					DrawShapesScaledImpl<true>(Particle, ShapeTransform, Shape, Color, Settings);
 				}
 				else
 				{
-					DrawShapesScaledImpl<false>(ShapeTransform, Shape, Color, Settings);
+					DrawShapesScaledImpl<false>(Particle, ShapeTransform, Shape, Color, Settings);
 				}
 				return;
 			}
 			else if (IsInstanced(PackedType))
 			{
-				DrawShapesInstancedImpl(ShapeTransform, Shape, Color, Settings);
+				DrawShapesInstancedImpl(Particle, ShapeTransform, Shape, Color, Settings);
 				return;
 			}
 
@@ -224,7 +228,7 @@ namespace Chaos
 			{
 				const TImplicitObjectTransformed<FReal, 3>* Transformed = Shape->template GetObject<TImplicitObjectTransformed<FReal, 3>>();
 				FRigidTransform3 TransformedTransform = FRigidTransform3(ShapeTransform.TransformPosition(Transformed->GetTransform().GetLocation()), ShapeTransform.GetRotation() * Transformed->GetTransform().GetRotation());
-				DrawShapesImpl(TransformedTransform, Transformed->GetTransformedObject(), Color, Settings);
+				DrawShapesImpl(Particle, TransformedTransform, Transformed->GetTransformedObject(), Color, Settings);
 				break;
 			}
 			case ImplicitObjectType::Union:
@@ -232,7 +236,7 @@ namespace Chaos
 				const FImplicitObjectUnion* Union = Shape->template GetObject<FImplicitObjectUnion>();
 				for (auto& UnionShape : Union->GetObjects())
 				{
-					DrawShapesImpl(ShapeTransform, UnionShape.Get(), Color, Settings);
+					DrawShapesImpl(Particle, ShapeTransform, UnionShape.Get(), Color, Settings);
 				}
 				break;
 			}
@@ -290,6 +294,22 @@ namespace Chaos
 				break;
 			}
 
+			if (bChaosDebugDebugDrawCollisionParticles && (Particle != nullptr))
+			{
+				if (const TPBDRigidParticleHandle<FReal, 3>* Rigid = Particle->CastToRigidParticle())
+				{
+					const TUniquePtr<TBVHParticles<FReal, 3>>& Particles = Rigid->CollisionParticles();
+					if (Particles != nullptr)
+					{
+						for (int32 ParticleIndex = 0; ParticleIndex < (int32)Particles->Size(); ++ParticleIndex)
+						{
+							FVec3 P = ShapeTransform.TransformPosition(Particles->X(ParticleIndex));
+							FDebugDrawQueue::GetInstance().DrawDebugPoint(P, Color, false, KINDA_SMALL_NUMBER, Settings.DrawPriority, Settings.PointSize);
+						}
+					}
+				}
+			}
+
 			if (bChaosDebugDebugDrawShapeBounds)
 			{
 				const FColor ShapeBoundsColor = FColor::Orange;
@@ -311,7 +331,7 @@ namespace Chaos
 				Color = FColor(InColor.R / 2, InColor.G / 2, InColor.B / 2, InColor.A);
 			}
 
-			DrawShapesImpl(FRigidTransform3(P, Q), Particle->Geometry().Get(), Color, Settings);
+			DrawShapesImpl(Particle, FRigidTransform3(P, Q), Particle->Geometry().Get(), Color, Settings);
 		}
 
 		void DrawParticleShapesImpl(const FRigidTransform3& SpaceTransform, const TGeometryParticle<FReal, 3>* Particle, const FColor& InColor, const FChaosDebugDrawSettings& Settings)
@@ -326,7 +346,7 @@ namespace Chaos
 				Color = FColor(InColor.R / 2, InColor.G / 2, InColor.B / 2, InColor.A);
 			}
 
-			DrawShapesImpl(FRigidTransform3(P, Q), Particle->Geometry().Get(), Color, Settings);
+			DrawShapesImpl(Particle->Handle(), FRigidTransform3(P, Q), Particle->Geometry().Get(), Color, Settings);
 		}
 
 		void DrawParticleBoundsImpl(const FRigidTransform3& SpaceTransform, const TGeometryParticleHandle<FReal, 3>* InParticle, const FReal Dt, const FReal BoundsThickness, const FReal BoundsThicknessVelocityInflation, const FChaosDebugDrawSettings& Settings)
