Chaos optimizations (during PVF simulation on large a world):
- Changed Chaos deactivation array to a set, as it can grow over a million instances (saves ~12.5%).
- Removed TVector destructor as it generates destructors iterators to empty functions for every instance on stack (saves ~8.5%).

Saves ~3:40 in a large wqorld test case.

diff --git a/Engine/Source/Runtime/Experimental/Chaos/Private/Chaos/Collision/CollisionConstraintFlags.cpp b/Engine/Source/Runtime/Experimental/Chaos/Private/Chaos/Collision/CollisionConstraintFlags.cpp
--- a/Engine/Source/Runtime/Experimental/Chaos/Private/Chaos/Collision/CollisionConstraintFlags.cpp
+++ b/Engine/Source/Runtime/Experimental/Chaos/Private/Chaos/Collision/CollisionConstraintFlags.cpp
@@ -74,10 +74,7 @@ namespace Chaos
 
 			for (auto& Item : StorageData->PendingDeactivations)
 			{
-				if (!PendingDeactivations.Contains(Item))
-				{
-					PendingDeactivations.Add(Item);
-				}
+				PendingDeactivations.Add(Item);
 			}
 
 			StorageDataQueue.Pop();
@@ -96,12 +93,10 @@ namespace Chaos
 			TArray<FHandleID> DeletionList;
 			for (auto& Elem : PendingActivations)
 			{
-				int32 DeactiveIndex = PendingDeactivations.Find(Elem.Key);
-				if (DeactiveIndex != INDEX_NONE)
+				if (PendingDeactivations.Remove(Elem.Key))
 				{
 					DeletionList.Add(Elem.Key);
 					PreculledParticles.Add(Elem.Key);
-					PendingDeactivations.RemoveAtSwap(DeactiveIndex, 1);
 				}
 			}
 			for (FHandleID Del : DeletionList)
@@ -145,9 +140,9 @@ namespace Chaos
 		// remove relationships that exist and have been initialized. 
 		if (PendingDeactivations.Num())
 		{
-			for (auto Index = PendingDeactivations.Num() - 1; Index >= 0; Index--)
+			for (auto& Idx : PendingDeactivations)
 			{
-				IgnoreCollisionsList.Remove(PendingDeactivations[Index]);
+				IgnoreCollisionsList.Remove(Idx);
 			}
 			PendingDeactivations.Empty();
 		}
diff --git a/Engine/Source/Runtime/Experimental/Chaos/Private/PBDRigidsSolver.cpp b/Engine/Source/Runtime/Experimental/Chaos/Private/PBDRigidsSolver.cpp
--- a/Engine/Source/Runtime/Experimental/Chaos/Private/PBDRigidsSolver.cpp
+++ b/Engine/Source/Runtime/Experimental/Chaos/Private/PBDRigidsSolver.cpp
@@ -474,11 +474,8 @@ namespace Chaos
 		Chaos::FIgnoreCollisionManager& CollisionManager = GetEvolution()->GetBroadPhase().GetIgnoreCollisionManager();
 		{
 			int32 ExternalTimestamp = GetMarshallingManager().GetExternalTimestamp_External();
-			Chaos::FIgnoreCollisionManager::FDeactivationArray& PendingMap = CollisionManager.GetPendingDeactivationsForGameThread(ExternalTimestamp);
-			if (!PendingMap.Contains(GTParticle->UniqueIdx()))
-			{
-				PendingMap.Add(GTParticle->UniqueIdx());
-			}
+			Chaos::FIgnoreCollisionManager::FDeactivationSet& PendingMap = CollisionManager.GetPendingDeactivationsForGameThread(ExternalTimestamp);
+			PendingMap.Add(GTParticle->UniqueIdx());
 		}
 
 		// Enqueue a command to remove the particle and delete the proxy
diff --git a/Engine/Source/Runtime/Experimental/Chaos/Public/Chaos/Collision/CollisionConstraintFlags.h b/Engine/Source/Runtime/Experimental/Chaos/Public/Chaos/Collision/CollisionConstraintFlags.h
--- a/Engine/Source/Runtime/Experimental/Chaos/Public/Chaos/Collision/CollisionConstraintFlags.h
+++ b/Engine/Source/Runtime/Experimental/Chaos/Public/Chaos/Collision/CollisionConstraintFlags.h
@@ -21,13 +21,13 @@ namespace Chaos
 	{
 	public:
 		using FHandleID = FUniqueIdx;
-		using FDeactivationArray = TArray<FUniqueIdx>;
+		using FDeactivationSet = TSet<FUniqueIdx>;
 		using FActiveMap = TMap<FHandleID, TArray<FHandleID> >;
 		using FPendingMap = TMap<FHandleID, TArray<FHandleID> >;
 		struct FStorageData
 		{
 			FPendingMap PendingActivations;
-			FDeactivationArray PendingDeactivations;
+			FDeactivationSet PendingDeactivations;
 			int32 ExternalTimestamp = INDEX_NONE;
 
 			void Reset()
@@ -68,7 +68,7 @@ namespace Chaos
 			return StorageDataProducer->PendingActivations;
 		}
 
-		FDeactivationArray& GetPendingDeactivationsForGameThread(int32 ExternalTimestamp)
+		FDeactivationSet& GetPendingDeactivationsForGameThread(int32 ExternalTimestamp)
 		{
 			if (StorageDataProducer->ExternalTimestamp == INDEX_NONE)
 			{
@@ -125,7 +125,7 @@ namespace Chaos
 		FActiveMap IgnoreCollisionsList;
 
 		FPendingMap PendingActivations;
-		FDeactivationArray PendingDeactivations;
+		FDeactivationSet PendingDeactivations;
 
 		// Producer storage data, pending changes written here until pushed into queue.
 		FStorageData* StorageDataProducer;
diff --git a/Engine/Source/Runtime/Experimental/ChaosCore/Public/Chaos/Vector.h b/Engine/Source/Runtime/Experimental/ChaosCore/Public/Chaos/Vector.h
--- a/Engine/Source/Runtime/Experimental/ChaosCore/Public/Chaos/Vector.h
+++ b/Engine/Source/Runtime/Experimental/ChaosCore/Public/Chaos/Vector.h
@@ -41,7 +41,6 @@ namespace Chaos
 				V[i] = Element;
 			}
 		}
-		~TVector() {}
 
 		TVector(std::initializer_list<T> InElements)
 		{
@@ -388,7 +387,6 @@ namespace Chaos
 			Stream.read(reinterpret_cast<char*>(&Y), sizeof(Y));
 			Stream.read(reinterpret_cast<char*>(&Z), sizeof(Z));
 		}
-		~TVector() {}
 		void Write(std::ostream& Stream) const
 		{
 			Stream.write(reinterpret_cast<const char*>(&X), sizeof(X));
@@ -588,7 +586,6 @@ namespace Chaos
 			X = ((float)InVector[0]);
 			Y = ((float)InVector[1]);
 		}
-		~TVector() {}
 		void Write(std::ostream& Stream) const
 		{
 			Stream.write(reinterpret_cast<const char*>(&X), sizeof(X));
@@ -681,7 +678,6 @@ namespace Chaos
 			Stream.read(reinterpret_cast<char*>(&Y), sizeof(T));
 			Stream.read(reinterpret_cast<char*>(&Z), sizeof(T));
 		}
-		FORCEINLINE ~TVector() {}
 		FORCEINLINE void Write(std::ostream& Stream) const
 		{
 			Stream.write(reinterpret_cast<const char*>(&X), sizeof(T));
@@ -839,8 +835,6 @@ namespace Chaos
 			Stream.read(reinterpret_cast<char*>(&X), sizeof(int32));
 			Stream.read(reinterpret_cast<char*>(&Y), sizeof(int32));
 		}
-		FORCEINLINE ~TVector()
-		{}
 
 		FORCEINLINE int32 Num() const { return 2; }
 
diff --git a/Engine/Source/Runtime/PhysicsCore/Private/ChaosEngineInterface.cpp b/Engine/Source/Runtime/PhysicsCore/Private/ChaosEngineInterface.cpp
--- a/Engine/Source/Runtime/PhysicsCore/Private/ChaosEngineInterface.cpp
+++ b/Engine/Source/Runtime/PhysicsCore/Private/ChaosEngineInterface.cpp
@@ -329,11 +329,8 @@ void FChaosEngineInterface::RemoveDisabledCollisionsFor_AssumesLocked(TArray< FP
 			Chaos::FIgnoreCollisionManager& CollisionManager = Solver->GetEvolution()->GetBroadPhase().GetIgnoreCollisionManager();
 			int32 ExternalTimestamp = Solver->GetMarshallingManager().GetExternalTimestamp_External();
 
-			Chaos::FIgnoreCollisionManager::FDeactivationArray& DeactivationMap = CollisionManager.GetPendingDeactivationsForGameThread(ExternalTimestamp);
-			if (!DeactivationMap.Contains(ActorReference->UniqueIdx()))
-			{
-				DeactivationMap.Add(ActorReference->UniqueIdx());
-			}
+			Chaos::FIgnoreCollisionManager::FDeactivationSet& DeactivationMap = CollisionManager.GetPendingDeactivationsForGameThread(ExternalTimestamp);
+			DeactivationMap.Add(ActorReference->UniqueIdx());
 		}
 	}
 }
