Chaos Cloth - Re-enable ISPC collisions in Shipping builds also. [Edigrate CL 15036670]

diff --git a/Engine/Source/Runtime/Experimental/Chaos/Public/Chaos/PerParticlePBDCollisionConstraint.h b/Engine/Source/Runtime/Experimental/Chaos/Public/Chaos/PerParticlePBDCollisionConstraint.h
--- a/Engine/Source/Runtime/Experimental/Chaos/Public/Chaos/PerParticlePBDCollisionConstraint.h
+++ b/Engine/Source/Runtime/Experimental/Chaos/Public/Chaos/PerParticlePBDCollisionConstraint.h
@@ -12,7 +12,7 @@
 #if !INTEL_ISPC
 const bool bChaos_PerParticleCollision_ISPC_Enabled = false;
 #elif UE_BUILD_SHIPPING
-const bool bChaos_PerParticleCollision_ISPC_Enabled = false;  // TODO: Enable once ISPC collision works
+const bool bChaos_PerParticleCollision_ISPC_Enabled = true;
 #else
 extern CHAOS_API bool bChaos_PerParticleCollision_ISPC_Enabled;
 #endif
