Chaos : GeometryCollection
-moved recursive order to GeometryCollectionAlgo

diff --git a/Engine/Source/Runtime/Experimental/Chaos/Private/GeometryCollection/GeometryCollectionAlgo.cpp b/Engine/Source/Runtime/Experimental/Chaos/Private/GeometryCollection/GeometryCollectionAlgo.cpp
--- a/Engine/Source/Runtime/Experimental/Chaos/Private/GeometryCollection/GeometryCollectionAlgo.cpp
+++ b/Engine/Source/Runtime/Experimental/Chaos/Private/GeometryCollection/GeometryCollectionAlgo.cpp
@@ -1203,5 +1203,69 @@ namespace GeometryCollectionAlgo
 		return true;
 	}
 
+	TArray<int32> ComputeRecursiveOrder(const FGeometryCollection& Collection)
+	{
+		const int32 NumTransforms = Collection.NumElements(FGeometryCollection::TransformGroup);
+		const TManagedArray<int32>& Parent = Collection.Parent;
+		const TManagedArray<TSet<int32>>& Children = Collection.Children;
+
+		//traverse cluster hierarchy in depth first and record order
+		struct FClusterProcessing
+		{
+			int32 TransformGroupIndex;
+			enum
+			{
+				None,
+				VisitingChildren
+			} State;
+
+			FClusterProcessing(int32 InIndex) : TransformGroupIndex(InIndex), State(None) {};
+		};
+
+		TArray<FClusterProcessing> ClustersToProcess;
+		//enqueue all roots
+		for (int32 TransformGroupIndex = 0; TransformGroupIndex < NumTransforms; TransformGroupIndex++)
+		{
+			if (Parent[TransformGroupIndex] == FGeometryCollection::Invalid && Children[TransformGroupIndex].Num() > 0)
+			{
+				ClustersToProcess.Emplace(TransformGroupIndex);
+			}
+		}
+
+		TArray<int32> TransformOrder;
+		TransformOrder.Reserve(NumTransforms);
+
+		while (ClustersToProcess.Num())
+		{
+			FClusterProcessing CurCluster = ClustersToProcess.Pop();
+			const int32 ClusterTransformIdx = CurCluster.TransformGroupIndex;
+			if (CurCluster.State == FClusterProcessing::VisitingChildren)
+			{
+				//children already visited
+				TransformOrder.Add(ClusterTransformIdx);
+			}
+			else
+			{
+				if (Children[ClusterTransformIdx].Num())
+				{
+					CurCluster.State = FClusterProcessing::VisitingChildren;
+					ClustersToProcess.Add(CurCluster);
+
+					//order of children doesn't matter as long as all children appear before parent
+					for (int32 ChildIdx : Children[ClusterTransformIdx])
+					{
+						ClustersToProcess.Emplace(ChildIdx);
+					}
+				}
+				else
+				{
+					TransformOrder.Add(ClusterTransformIdx);
+				}
+			}
+		}
+
+		return TransformOrder;
+	}
+
 
 }
diff --git a/Engine/Source/Runtime/Experimental/Chaos/Private/PhysicsProxy/GeometryCollectionPhysicsProxy.cpp b/Engine/Source/Runtime/Experimental/Chaos/Private/PhysicsProxy/GeometryCollectionPhysicsProxy.cpp
--- a/Engine/Source/Runtime/Experimental/Chaos/Private/PhysicsProxy/GeometryCollectionPhysicsProxy.cpp
+++ b/Engine/Source/Runtime/Experimental/Chaos/Private/PhysicsProxy/GeometryCollectionPhysicsProxy.cpp
@@ -139,70 +139,6 @@ TArray<int32> ComputeTransformToGeometryMap(const FGeometryCollection& Collectio
 	return TransformToGeometryMap;
 }
 
-//Computes the order of transform indices so that children in a tree always appear before their parents. Handles forests
-TArray<int32> ComputeRecursiveOrder(const FGeometryCollection& Collection)
-{
-	const int32 NumTransforms = Collection.NumElements(FGeometryCollection::TransformGroup);
-	const TManagedArray<int32>& Parent = Collection.Parent;
-	const TManagedArray<TSet<int32>>& Children = Collection.Children;
-
-	//traverse cluster hierarchy in depth first and record order
-	struct FClusterProcessing
-	{
-		int32 TransformGroupIndex;
-		enum
-		{
-			None,
-			VisitingChildren
-		} State;
-
-		FClusterProcessing(int32 InIndex) : TransformGroupIndex(InIndex), State(None) {};
-	};
-
-	TArray<FClusterProcessing> ClustersToProcess;
-	//enqueue all roots
-	for(int32 TransformGroupIndex = 0; TransformGroupIndex < NumTransforms; TransformGroupIndex++)
-	{
-		if(Parent[TransformGroupIndex] == FGeometryCollection::Invalid && Children[TransformGroupIndex].Num() > 0)
-		{
-			ClustersToProcess.Emplace(TransformGroupIndex);
-		}
-	}
-
-	TArray<int32> TransformOrder;
-	TransformOrder.Reserve(NumTransforms);
-
-	while(ClustersToProcess.Num())
-	{
-		FClusterProcessing CurCluster = ClustersToProcess.Pop();
-		const int32 ClusterTransformIdx = CurCluster.TransformGroupIndex;
-		if(CurCluster.State == FClusterProcessing::VisitingChildren)
-		{
-			//children already visited
-			TransformOrder.Add(ClusterTransformIdx);
-		}
-		else
-		{
-			if(Children[ClusterTransformIdx].Num())
-			{
-				CurCluster.State = FClusterProcessing::VisitingChildren;
-				ClustersToProcess.Add(CurCluster);
-
-				//order of children doesn't matter as long as all children appear before parent
-				for(int32 ChildIdx : Children[ClusterTransformIdx])
-				{
-					ClustersToProcess.Emplace(ChildIdx);
-				}
-			}
-			else
-			{
-				TransformOrder.Add(ClusterTransformIdx);
-			}
-		}
-	}
-
-	return TransformOrder;
-}
 
 DECLARE_CYCLE_STAT(TEXT("FGeometryCollectionPhysicsProxy::PopulateSimulatedParticle"), STAT_PopulateSimulatedParticle, STATGROUP_Chaos);
 void PopulateSimulatedParticle(
@@ -837,7 +773,7 @@ void FGeometryCollectionPhysicsProxy::InitializeBodiesPT(Chaos::TPBDRigidsSolver
 		if (Parameters.EnableClustering)// && Parameters.CacheType != EGeometryCollectionCacheType::Play)
 		{
 			// "RecursiveOrder" means bottom up - children come before their parents.
-			const TArray<int32> RecursiveOrder = ComputeRecursiveOrder(*RestCollection);
+			const TArray<int32> RecursiveOrder = GeometryCollectionAlgo::ComputeRecursiveOrder(*RestCollection);
 
 			// Propagate simulated particle flags up the hierarchy from children 
 			// to their parents, grandparents, etc...
@@ -2193,7 +2129,7 @@ void FGeometryCollectionPhysicsProxy::InitializeSharedCollisionStructures(
  			}
  		}
 
-		const TArray<int32> RecursiveOrder = ComputeRecursiveOrder(RestCollection);
+		const TArray<int32> RecursiveOrder = GeometryCollectionAlgo::ComputeRecursiveOrder(RestCollection);
 		const TArray<int32> TransformToGeometry = ComputeTransformToGeometryMap(RestCollection);
 
 		TArray<bool> IsClusterSimulated;
diff --git a/Engine/Source/Runtime/Experimental/Chaos/Public/GeometryCollection/GeometryCollectionAlgo.h b/Engine/Source/Runtime/Experimental/Chaos/Public/GeometryCollection/GeometryCollectionAlgo.h
--- a/Engine/Source/Runtime/Experimental/Chaos/Public/GeometryCollection/GeometryCollectionAlgo.h
+++ b/Engine/Source/Runtime/Experimental/Chaos/Public/GeometryCollection/GeometryCollectionAlgo.h
@@ -264,4 +264,12 @@ namespace GeometryCollectionAlgo
 	CHAOS_API
 	HasValidGeometryReferences(const FGeometryCollection* GeometryCollection);
 
+
+	/*
+	* Computes the order of transform indices so that children in a tree always appear before their parents. Handles forests
+	*/
+	TArray<int32>
+	CHAOS_API
+	ComputeRecursiveOrder(const FGeometryCollection& Collection);
+
 }
