Chaos : GeometryCollection

- Updated SimulationType for older attributes as they are loaded in.

diff --git a/Engine/Source/Runtime/Experimental/Chaos/Private/GeometryCollection/GeometryCollection.cpp b/Engine/Source/Runtime/Experimental/Chaos/Private/GeometryCollection/GeometryCollection.cpp
--- a/Engine/Source/Runtime/Experimental/Chaos/Private/GeometryCollection/GeometryCollection.cpp
+++ b/Engine/Source/Runtime/Experimental/Chaos/Private/GeometryCollection/GeometryCollection.cpp
@@ -9,14 +9,14 @@
 #include "GeometryCollection/GeometryCollectionAlgo.h"
 #include "GeometryCollection/GeometryCollectionUtility.h"
 #include "GeometryCollection/GeometryCollectionProximityUtility.h"
+#include "GeometryCollection/GeometryCollectionClusteringUtility.h"
 
 #include <iostream>
 #include <fstream>
 #include "Chaos/ChaosArchive.h"
 #include "Voronoi/Voronoi.h"
 
-DEFINE_LOG_CATEGORY_STATIC(FGeometryCollectionLogging, NoLogging, All);
-
+DEFINE_LOG_CATEGORY_STATIC(FGeometryCollectionLogging, Log, All);
 
 // @todo: update names 
 const FName FGeometryCollection::FacesGroup = "Faces";
@@ -887,6 +887,52 @@ void FGeometryCollection::Serialize(Chaos::FChaosArchive& Ar)
 		RemoveAttribute("ExplodedTransform", FTransformCollection::TransformGroup);
 		RemoveAttribute("ExplodedVector", FTransformCollection::TransformGroup);
 
+		// Version 5 introduced accurate SimulationType tagging
+		if (Version < 5)
+		{
+			UE_LOG(FGeometryCollectionLogging, Warning, TEXT("GeometryCollection is has inaccurate simulation type tags. Updating tags based on transform topology."));
+			TManagedArray<bool>* SimulatableParticles = FindAttribute<bool>(FGeometryCollection::SimulatableParticlesAttribute, FTransformCollection::TransformGroup);
+			TArray<bool> RigidChildren; RigidChildren.Init(false,NumElements(FTransformCollection::TransformGroup));
+			const TArray<int32> RecursiveOrder = GeometryCollectionAlgo::ComputeRecursiveOrder(*this);
+			for (const int32 TransformGroupIndex : RecursiveOrder)
+			{
+				SimulationType[TransformGroupIndex] = ESimulationTypes::FST_None;
+
+				if(!Children[TransformGroupIndex].Num())
+				{ // leaf nodes
+					if (TransformToGeometryIndex[TransformGroupIndex] > INDEX_NONE)
+					{
+						SimulationType[TransformGroupIndex] = ESimulationTypes::FST_Rigid;
+					}
+
+					if (SimulatableParticles && !(*SimulatableParticles)[TransformGroupIndex])
+					{
+						SimulationType[TransformGroupIndex] = ESimulationTypes::FST_None;
+					}
+				}
+				else
+				{ // interior nodes
+					if (RigidChildren[TransformGroupIndex])
+					{
+						SimulationType[TransformGroupIndex] = ESimulationTypes::FST_Clustered;
+					}
+					else if (TransformToGeometryIndex[TransformGroupIndex] > INDEX_NONE)
+					{
+						SimulationType[TransformGroupIndex] = ESimulationTypes::FST_Rigid;
+					}
+				}
+
+				if (SimulationType[TransformGroupIndex] != ESimulationTypes::FST_None &&
+					Parent[TransformGroupIndex] != INDEX_NONE )
+				{
+					RigidChildren[Parent[TransformGroupIndex]] = true;
+				}
+			}
+
+			// Structure is conditioned, now considered up to date.
+			Version = 5;
+		}
+
 	}
 }
 
diff --git a/Engine/Source/Runtime/Experimental/Chaos/Private/GeometryCollection/ManagedArrayCollection.cpp b/Engine/Source/Runtime/Experimental/Chaos/Private/GeometryCollection/ManagedArrayCollection.cpp
--- a/Engine/Source/Runtime/Experimental/Chaos/Private/GeometryCollection/ManagedArrayCollection.cpp
+++ b/Engine/Source/Runtime/Experimental/Chaos/Private/GeometryCollection/ManagedArrayCollection.cpp
@@ -12,6 +12,7 @@ int8 FManagedArrayCollection::Invalid = INDEX_NONE;
 
 FManagedArrayCollection::FManagedArrayCollection()
 {
+	Version = 5;
 }
 
 static const FName GuidName("GUID");
@@ -402,7 +403,7 @@ FString FManagedArrayCollection::ToString() const
 
 void FManagedArrayCollection::Serialize(Chaos::FChaosArchive& Ar)
 {
-	int Version = 4;
+	if (Ar.IsSaving()) Version = 5;
 	Ar << Version;
 
 	if (Ar.IsLoading())
diff --git a/Engine/Source/Runtime/Experimental/Chaos/Public/GeometryCollection/ManagedArrayCollection.h b/Engine/Source/Runtime/Experimental/Chaos/Public/GeometryCollection/ManagedArrayCollection.h
--- a/Engine/Source/Runtime/Experimental/Chaos/Public/GeometryCollection/ManagedArrayCollection.h
+++ b/Engine/Source/Runtime/Experimental/Chaos/Public/GeometryCollection/ManagedArrayCollection.h
@@ -464,6 +464,11 @@ protected:
 	*/
 	void SyncGroupSizeFrom(const FManagedArrayCollection& InCollection, FName Group);
 
+	/**
+	 * Version to indicate need for conditioning to current expected data layout during serialization loading.
+	 */
+	int32 Version;
+
 
 	/**
 	*   operator<<(FGroupInfo)
