GeometryCollection -- Asset now has properties for default clustering behaviour. The properties are written to the actor on placement.

diff --git a/Engine/Plugins/Experimental/GeometryCollectionPlugin/Source/GeometryCollectionEditor/Private/GeometryCollection/ActorFactoryGeometryCollection.cpp b/Engine/Plugins/Experimental/GeometryCollectionPlugin/Source/GeometryCollectionEditor/Private/GeometryCollection/ActorFactoryGeometryCollection.cpp
--- a/Engine/Plugins/Experimental/GeometryCollectionPlugin/Source/GeometryCollectionEditor/Private/GeometryCollection/ActorFactoryGeometryCollection.cpp
+++ b/Engine/Plugins/Experimental/GeometryCollectionPlugin/Source/GeometryCollectionEditor/Private/GeometryCollection/ActorFactoryGeometryCollection.cpp
@@ -49,6 +49,13 @@ void UActorFactoryGeometryCollection::PostSpawnActor(UObject* Asset, AActor* New
 	// Change properties
 	NewGeometryCollectionActor->GetGeometryCollectionComponent()->SetRestCollection(GeometryCollection);
 
+	// Set configured clustering properties.
+	NewGeometryCollectionActor->GetGeometryCollectionComponent()->EnableClustering = GeometryCollection->EnableClustering;
+	NewGeometryCollectionActor->GetGeometryCollectionComponent()->ClusterGroupIndex = GeometryCollection->ClusterGroupIndex;
+	NewGeometryCollectionActor->GetGeometryCollectionComponent()->MaxClusterLevel = GeometryCollection->MaxClusterLevel;
+	NewGeometryCollectionActor->GetGeometryCollectionComponent()->DamageThreshold = GeometryCollection->DamageThreshold;
+	NewGeometryCollectionActor->GetGeometryCollectionComponent()->ClusterConnectionType = GeometryCollection->ClusterConnectionType;
+
 	// Init Component
 	NewGeometryCollectionActor->GetGeometryCollectionComponent()->RegisterComponent();
 }
diff --git a/Engine/Source/Runtime/Experimental/GeometryCollectionEngine/Private/GeometryCollection/GeometryCollectionObject.cpp b/Engine/Source/Runtime/Experimental/GeometryCollectionEngine/Private/GeometryCollection/GeometryCollectionObject.cpp
--- a/Engine/Source/Runtime/Experimental/GeometryCollectionEngine/Private/GeometryCollection/GeometryCollectionObject.cpp
+++ b/Engine/Source/Runtime/Experimental/GeometryCollectionEngine/Private/GeometryCollection/GeometryCollectionObject.cpp
@@ -50,6 +50,11 @@ UGeometryCollection::UGeometryCollection(const FObjectInitializer& ObjectInitial
 	#if WITH_EDITOR
 	, bManualDataCreate(false)
 	#endif
+	, EnableClustering(true)
+	, ClusterGroupIndex(0)
+	, MaxClusterLevel(100)
+	, DamageThreshold({ 250.0 })
+	, ClusterConnectionType(EClusterConnectionTypeEnum::Chaos_PointImplicit)
 	, EnableNanite(false)
 	, CollisionType(ECollisionTypeEnum::Chaos_Volumetric)
 	, ImplicitType(EImplicitTypeEnum::Chaos_Implicit_Box)
diff --git a/Engine/Source/Runtime/Experimental/GeometryCollectionEngine/Public/GeometryCollection/GeometryCollectionObject.h b/Engine/Source/Runtime/Experimental/GeometryCollectionEngine/Public/GeometryCollection/GeometryCollectionObject.h
--- a/Engine/Source/Runtime/Experimental/GeometryCollectionEngine/Public/GeometryCollection/GeometryCollectionObject.h
+++ b/Engine/Source/Runtime/Experimental/GeometryCollectionEngine/Public/GeometryCollection/GeometryCollectionObject.h
@@ -8,6 +8,7 @@
 #include "InstanceUniformShaderParameters.h"
 #include "GeometryCollection/ManagedArray.h"
 #include "GeometryCollection/GeometryCollectionSimulationTypes.h"
+#include "Chaos/ChaosSolverActor.h"
 #include "GeometryCollectionObject.generated.h"
 
 class UMaterialInterface;
@@ -218,6 +219,25 @@ public:
 	/** Pointer to the data used to render this geometry collection with Nanite. */
 	TUniquePtr<class FGeometryCollectionNaniteData> NaniteData;
 	
+	UPROPERTY(EditAnywhere, Category = "Clustering")
+	bool EnableClustering;
+
+	/** Maximum level for cluster breaks. */
+	UPROPERTY(EditAnywhere, Category = "Clustering")
+	int32 ClusterGroupIndex;
+
+	/** Maximum level for cluster breaks. */
+	UPROPERTY(EditAnywhere, Category = "Clustering")
+	int32 MaxClusterLevel;
+
+	/** Damage threshold for clusters at different levels. */
+	UPROPERTY(EditAnywhere, Category = "Clustering")
+	TArray<float> DamageThreshold;
+
+	/** */
+	UPROPERTY(EditAnywhere, Category = "Clustering")
+	EClusterConnectionTypeEnum ClusterConnectionType;
+
 	UPROPERTY(BlueprintReadWrite, EditAnywhere, Category = "GeometrySource")
 	TArray<FGeometryCollectionSource> GeometrySource;
 	
