Properly handle degenarated faces during convex generation by providing an alternate normal

diff --git a/Engine/Source/Runtime/Experimental/Chaos/Public/Chaos/CollisionConvexMesh.h b/Engine/Source/Runtime/Experimental/Chaos/Public/Chaos/CollisionConvexMesh.h
--- a/Engine/Source/Runtime/Experimental/Chaos/Public/Chaos/CollisionConvexMesh.h
+++ b/Engine/Source/Runtime/Experimental/Chaos/Public/Chaos/CollisionConvexMesh.h
@@ -637,8 +637,15 @@ namespace Chaos
 			TR->Next = RS;
 			FVec3 RSTNormal = ComputeFaceNormal(InParticles.X(RS->Vertex), InParticles.X(ST->Vertex), InParticles.X(TR->Vertex));
 			const FReal RSTNormalSize = RSTNormal.Size();
-			check(RSTNormalSize > FLT_EPSILON);
-			RSTNormal = RSTNormal * (1 / RSTNormalSize);
+			if (RSTNormalSize > 0)
+			{
+				RSTNormal = RSTNormal * (1 / RSTNormalSize);
+			}
+			else
+			{
+				// degenerated face, use a valid neighbor face normal 
+				RSTNormal = RS->Twin->Face->Plane.Normal();
+			}
 			FConvexFace* RST = Pool.AllocConvexFace(TPlaneConcrete<FReal, 3>(InParticles.X(RS->Vertex), RSTNormal));
 			RST->FirstEdge = RS;
 			RS->Face = RST;
