fix comment on CollisionSampleFraction

#rnx

diff --git a/Engine/Source/Runtime/Experimental/GeometryCollectionEngine/Public/GeometryCollection/GeometryCollectionComponent.h b/Engine/Source/Runtime/Experimental/GeometryCollectionEngine/Public/GeometryCollection/GeometryCollectionComponent.h
--- a/Engine/Source/Runtime/Experimental/GeometryCollectionEngine/Public/GeometryCollection/GeometryCollectionComponent.h
+++ b/Engine/Source/Runtime/Experimental/GeometryCollectionEngine/Public/GeometryCollection/GeometryCollectionComponent.h
@@ -463,7 +463,7 @@ public:
 	UPROPERTY(EditAnywhere, BlueprintReadWrite, Category = "ChaosPhysics|Collisions")
 	int32 CollisionGroup;
 
-	/** Uniform Friction */
+	/** Fraction of collision sample particles to keep */
 	UPROPERTY(EditAnywhere, BlueprintReadWrite, Category = "ChaosPhysics|Collisions")
 	float CollisionSampleFraction;
 
