Fix HGeometryCollectionBone having the wrong parent proxy

diff --git a/Engine/Source/Runtime/Experimental/GeometryCollectionEngine/Private/GeometryCollection/GeometryCollectionHitProxy.cpp b/Engine/Source/Runtime/Experimental/GeometryCollectionEngine/Private/GeometryCollection/GeometryCollectionHitProxy.cpp
--- a/Engine/Source/Runtime/Experimental/GeometryCollectionEngine/Private/GeometryCollection/GeometryCollectionHitProxy.cpp
+++ b/Engine/Source/Runtime/Experimental/GeometryCollectionEngine/Private/GeometryCollection/GeometryCollectionHitProxy.cpp
@@ -4,5 +4,5 @@
 
 #if GEOMETRYCOLLECTION_EDITOR_SELECTION
 IMPLEMENT_HIT_PROXY(HGeometryCollection, HActor);
-IMPLEMENT_HIT_PROXY(HGeometryCollectionBone, HActor);
+IMPLEMENT_HIT_PROXY(HGeometryCollectionBone, HHitProxy);
 #endif  // #if GEOMETRYCOLLECTION_EDITOR_SELECTION
