Chaos - fix box margin in SupportCoreScaled

diff --git a/Engine/Source/Runtime/Experimental/Chaos/Private/Chaos/CollisionResolution.cpp b/Engine/Source/Runtime/Experimental/Chaos/Private/Chaos/CollisionResolution.cpp
--- a/Engine/Source/Runtime/Experimental/Chaos/Private/Chaos/CollisionResolution.cpp
+++ b/Engine/Source/Runtime/Experimental/Chaos/Private/Chaos/CollisionResolution.cpp
@@ -790,7 +790,7 @@ namespace Chaos
 		template <typename TriMeshType>
 		FContactPoint BoxTriangleMeshContactPoint(const FImplicitBox3& A, const FRigidTransform3& ATransform, const TriMeshType& B, const FRigidTransform3& BTransform, const FReal CullDistance, const FReal ShapePadding)
 		{
-			return GJKImplicitContactPoint<TBox<float, 3>>(TBox<float, 3>(A), ATransform, B, BTransform, CullDistance, ShapePadding);
+			return GJKImplicitContactPoint<TBox<float, 3>>(A, ATransform, B, BTransform, CullDistance, ShapePadding);
 		}
 
 		template <typename TriMeshType>
diff --git a/Engine/Source/Runtime/Experimental/Chaos/Public/Chaos/AABB.h b/Engine/Source/Runtime/Experimental/Chaos/Public/Chaos/AABB.h
--- a/Engine/Source/Runtime/Experimental/Chaos/Public/Chaos/AABB.h
+++ b/Engine/Source/Runtime/Experimental/Chaos/Public/Chaos/AABB.h
@@ -355,7 +355,7 @@ namespace Chaos
 			return BestNormal;
 		}
 
-		FORCEINLINE TVector<T, d> Support(const TVector<T, d>& Direction, const T Thickness) const 
+		FORCEINLINE_DEBUGGABLE TVector<T, d> Support(const TVector<T, d>& Direction, const T Thickness) const
 		{
 			TVector<T, d> ChosenPt;
 			for (int Axis = 0; Axis < d; ++Axis)
@@ -382,7 +382,7 @@ namespace Chaos
 		}
 
 		// Support vertex in the specified direction, assuming each face has been moved inwards by InMargin
-		FORCEINLINE TVector<T, d> SupportCore(const TVector<T, d>& Direction, float InMargin) const
+		FORCEINLINE_DEBUGGABLE TVector<T, d> SupportCore(const TVector<T, d>& Direction, float InMargin) const
 		{
 			TVector<T, d> ChosenPt;
 			for (int Axis = 0; Axis < d; ++Axis)
diff --git a/Engine/Source/Runtime/Experimental/Chaos/Public/Chaos/Box.h b/Engine/Source/Runtime/Experimental/Chaos/Public/Chaos/Box.h
--- a/Engine/Source/Runtime/Experimental/Chaos/Public/Chaos/Box.h
+++ b/Engine/Source/Runtime/Experimental/Chaos/Public/Chaos/Box.h
@@ -267,22 +267,22 @@ namespace Chaos
 		}
 
 		// Returns a position on the shape
-		FORCEINLINE TVector<T, d> Support(const TVector<T, d>& Direction, const T Thickness) const
+		FORCEINLINE_DEBUGGABLE TVector<T, d> Support(const TVector<T, d>& Direction, const T Thickness) const
 		{
 			return AABB.Support(Direction, Thickness);
 		}
 
 		// Returns a position on the core shape excluding the margin
-		FORCEINLINE TVector<T, d> SupportCore(const TVector<T, d>& Direction, float InMargin) const
+		FORCEINLINE_DEBUGGABLE TVector<T, d> SupportCore(const TVector<T, d>& Direction, float InMargin) const
 		{
 			return AABB.SupportCore(Direction, InMargin + GetMargin());
 		}
 
-		FORCEINLINE TVector<T, d> SupportCoreScaled(const TVector<T, d>& Direction, float InMargin, const TVector<T, d>& Scale) const
+		FORCEINLINE_DEBUGGABLE TVector<T, d> SupportCoreScaled(const TVector<T, d>& Direction, float InMargin, const TVector<T, d>& Scale) const
 		{
 			const FReal InvScale = 1.0f / Scale[0];
 			const FReal NetMargin = InvScale * (InMargin + GetMargin());
-			return SupportCore(Direction * Scale, NetMargin) * Scale;
+			return AABB.SupportCore(Direction * Scale, NetMargin) * Scale;
 		}
 
 		// Returns a winding order multiplier used in the manifold clipping and required when we have negative scales (See ImplicitObjectScaled)
