Chaos : Destruction
- Added Size Specific damage thresholds.

diff --git a/Engine/Source/Runtime/Experimental/Chaos/Private/PhysicsProxy/GeometryCollectionPhysicsProxy.cpp b/Engine/Source/Runtime/Experimental/Chaos/Private/PhysicsProxy/GeometryCollectionPhysicsProxy.cpp
--- a/Engine/Source/Runtime/Experimental/Chaos/Private/PhysicsProxy/GeometryCollectionPhysicsProxy.cpp
+++ b/Engine/Source/Runtime/Experimental/Chaos/Private/PhysicsProxy/GeometryCollectionPhysicsProxy.cpp
@@ -938,6 +938,26 @@ DECLARE_CYCLE_STAT(TEXT("FGeometryCollectionPhysicsProxy::BuildClusters"), STAT_
 DECLARE_CYCLE_STAT(TEXT("FGeometryCollectionPhysicsProxy::BuildClusters:GlobalMatrices"), STAT_BuildClustersGlobalMatrices, STATGROUP_Chaos);
 
 
+int32 FindSizeSpecificIdx(const TArray<FSharedSimulationSizeSpecificData>& SizeSpecificData, const FBox& Bounds)
+{
+	const FVector Extents = Bounds.GetExtent();
+	const float Size = Extents.GetAbsMin();
+	check(SizeSpecificData.Num());
+	int32 UseIdx = 0;
+	float PreSize = FLT_MAX;
+	for (int32 Idx = SizeSpecificData.Num() - 1; Idx >= 0; --Idx)
+	{
+		ensureMsgf(PreSize >= SizeSpecificData[Idx].MaxSize, TEXT("SizeSpecificData is not sorted"));
+		PreSize = SizeSpecificData[Idx].MaxSize;
+		if (Size < SizeSpecificData[Idx].MaxSize)
+			UseIdx = Idx;
+		else
+			break;
+	}
+	return UseIdx;
+}
+
+
 template <typename Traits>
 Chaos::TPBDRigidClusteredParticleHandle<float, 3>* 
 FGeometryCollectionPhysicsProxy::BuildClusters(
@@ -1044,8 +1064,24 @@ FGeometryCollectionPhysicsProxy::BuildClusters(
 		Damage = FLT_MAX;
 	}
 
+	if (Parameters.bUseSizeSpecificDamageThresholds)
+	{
+		const TManagedArray<FGeometryDynamicCollection::FSharedImplicit>& Implicit = DynamicCollection.Implicits;
+		if (Implicit[CollectionClusterIndex] && Implicit[CollectionClusterIndex]->HasBoundingBox())
+		{
+			FBox LocalBoundingBox(Implicit[CollectionClusterIndex]->BoundingBox().Min(), Implicit[CollectionClusterIndex]->BoundingBox().Max());
+			const int32 SizeSpecificIdx = FindSizeSpecificIdx(Parameters.Shared.SizeSpecificData, LocalBoundingBox);
+			if (0<= SizeSpecificIdx && SizeSpecificIdx < Parameters.Shared.SizeSpecificData.Num() )
+			{
+				const FSharedSimulationSizeSpecificData& SizeSpecificData = Parameters.Shared.SizeSpecificData[SizeSpecificIdx];
+				Damage = SizeSpecificData.DamageThreshold;
+			}
+		}
+	}
+
 	Parent->SetStrains(Damage);
 
+
 	// #BGTODO This will not automatically update - material properties should only ever exist in the material, not in other arrays
 	const Chaos::FChaosPhysicsMaterial* CurMaterial = Solver->GetSimMaterials().Get(Parameters.PhysicalMaterialHandle.InnerHandle);
 	if(CurMaterial)
@@ -1652,24 +1688,6 @@ bool FGeometryCollectionPhysicsProxy::PullFromPhysicsState(const Chaos::FDirtyGe
 //==============================================================================
 
 
-int32 FindSizeSpecificIdx(const TArray<FSharedSimulationSizeSpecificData>& SizeSpecificData, const FBox& Bounds)
-{
-	const FVector Extents = Bounds.GetExtent();
-	const float Size = Extents.GetAbsMin();
-	check(SizeSpecificData.Num());
-	int32 UseIdx = 0;
-	float PreSize = FLT_MAX;
-	for (int32 Idx = SizeSpecificData.Num() - 1; Idx >=0 ; --Idx)
-	{
-		ensureMsgf(PreSize >= SizeSpecificData[Idx].MaxSize, TEXT("SizeSpecificData is not sorted"));
-		PreSize = SizeSpecificData[Idx].MaxSize;
-		if (Size < SizeSpecificData[Idx].MaxSize)
-			UseIdx = Idx;
-		else
-			break;
-	}
-	return UseIdx;
-}
 
 /** 
 	NOTE - Making any changes to data stored on the rest collection below MUST be accompanied
diff --git a/Engine/Source/Runtime/Experimental/Chaos/Public/GeometryCollection/GeometryCollectionSimulationCoreTypes.h b/Engine/Source/Runtime/Experimental/Chaos/Public/GeometryCollection/GeometryCollectionSimulationCoreTypes.h
--- a/Engine/Source/Runtime/Experimental/Chaos/Public/GeometryCollection/GeometryCollectionSimulationCoreTypes.h
+++ b/Engine/Source/Runtime/Experimental/Chaos/Public/GeometryCollection/GeometryCollectionSimulationCoreTypes.h
@@ -26,6 +26,7 @@ struct FSharedSimulationSizeSpecificData
 		, CollisionObjectReductionPercentage(0.f)
 		, CollisionParticlesFraction(1.f)
 		, MaximumCollisionParticles(60)
+		, DamageThreshold(250.f)
 	{
 	}
 
@@ -39,6 +40,7 @@ struct FSharedSimulationSizeSpecificData
 	float CollisionObjectReductionPercentage;
 	float CollisionParticlesFraction;
 	int32 MaximumCollisionParticles;
+	float DamageThreshold;
 
 	bool operator<(const FSharedSimulationSizeSpecificData& Rhs) const { return MaxSize < Rhs.MaxSize; }
 };
@@ -232,6 +234,7 @@ struct FSimulationParameters
 		, EnableClustering(true)
 		, ClusterGroupIndex(0)
 		, MaxClusterLevel(100)
+		, bUseSizeSpecificDamageThresholds(false)
 		, DamageThreshold({250.f})
 		, ClusterConnectionMethod(Chaos::FClusterCreationParameters<float>::EConnectionMethod::PointImplicit)
 		, CollisionGroup(0)
@@ -260,6 +263,7 @@ struct FSimulationParameters
 		, EnableClustering(Other.EnableClustering)
 		, ClusterGroupIndex(Other.ClusterGroupIndex)
 		, MaxClusterLevel(Other.MaxClusterLevel)
+		, bUseSizeSpecificDamageThresholds(Other.bUseSizeSpecificDamageThresholds)
 		, DamageThreshold(Other.DamageThreshold)
 		, ClusterConnectionMethod(Other.ClusterConnectionMethod)
 		, CollisionGroup(Other.CollisionGroup)
@@ -307,6 +311,7 @@ struct FSimulationParameters
 	bool EnableClustering;
 	int32 ClusterGroupIndex;
 	int32 MaxClusterLevel;
+	bool bUseSizeSpecificDamageThresholds;
 	TArray<float> DamageThreshold;
 	Chaos::FClusterCreationParameters<float>::EConnectionMethod ClusterConnectionMethod;
 
diff --git a/Engine/Source/Runtime/Experimental/GeometryCollectionEngine/Private/GeometryCollection/GeometryCollectionComponent.cpp b/Engine/Source/Runtime/Experimental/GeometryCollectionEngine/Private/GeometryCollection/GeometryCollectionComponent.cpp
--- a/Engine/Source/Runtime/Experimental/GeometryCollectionEngine/Private/GeometryCollection/GeometryCollectionComponent.cpp
+++ b/Engine/Source/Runtime/Experimental/GeometryCollectionEngine/Private/GeometryCollection/GeometryCollectionComponent.cpp
@@ -212,6 +212,7 @@ UGeometryCollectionComponent::UGeometryCollectionComponent(const FObjectInitiali
 	, InitializationState(ESimulationInitializationState::Unintialized)
 	, ObjectType(EObjectStateTypeEnum::Chaos_Object_Dynamic)
 	, EnableClustering(true)
+	, bUseSizeSpecificDamageThreshold(false)
 	, ClusterGroupIndex(0)
 	, MaxClusterLevel(100)
 	, DamageThreshold({250.0})
@@ -1763,6 +1764,7 @@ void UGeometryCollectionComponent::OnCreatePhysicsState()
 				SimulationParameters.EnableClustering = EnableClustering;
 				SimulationParameters.ClusterGroupIndex = EnableClustering ? ClusterGroupIndex : 0;
 				SimulationParameters.MaxClusterLevel = MaxClusterLevel;
+				SimulationParameters.bUseSizeSpecificDamageThresholds = bUseSizeSpecificDamageThreshold;
 				SimulationParameters.DamageThreshold = DamageThreshold;
 				SimulationParameters.ClusterConnectionMethod = (Chaos::FClusterCreationParameters<float>::EConnectionMethod)ClusterConnectionType;
 				SimulationParameters.CollisionGroup = CollisionGroup;
diff --git a/Engine/Source/Runtime/Experimental/GeometryCollectionEngine/Private/GeometryCollection/GeometryCollectionObject.cpp b/Engine/Source/Runtime/Experimental/GeometryCollectionEngine/Private/GeometryCollection/GeometryCollectionObject.cpp
--- a/Engine/Source/Runtime/Experimental/GeometryCollectionEngine/Private/GeometryCollection/GeometryCollectionObject.cpp
+++ b/Engine/Source/Runtime/Experimental/GeometryCollectionEngine/Private/GeometryCollection/GeometryCollectionObject.cpp
@@ -104,6 +104,7 @@ void FillSharedSimulationSizeSpecificData(FSharedSimulationSizeSpecificData& ToD
 	ToData.CollisionObjectReductionPercentage = FromData.CollisionObjectReductionPercentage;
 	ToData.CollisionParticlesFraction = FromData.CollisionParticlesFraction;
 	ToData.MaximumCollisionParticles = FromData.MaximumCollisionParticles;
+	ToData.DamageThreshold = FromData.DamageThreshold;
 }
 
 float KgCm3ToKgM3(float Density)
diff --git a/Engine/Source/Runtime/Experimental/GeometryCollectionEngine/Public/GeometryCollection/GeometryCollectionComponent.h b/Engine/Source/Runtime/Experimental/GeometryCollectionEngine/Public/GeometryCollection/GeometryCollectionComponent.h
--- a/Engine/Source/Runtime/Experimental/GeometryCollectionEngine/Public/GeometryCollection/GeometryCollectionComponent.h
+++ b/Engine/Source/Runtime/Experimental/GeometryCollectionEngine/Public/GeometryCollection/GeometryCollectionComponent.h
@@ -451,9 +451,12 @@ public:
 	UPROPERTY(EditAnywhere, BlueprintReadWrite, Category = "ChaosPhysics|Clustering")
 	int32 MaxClusterLevel;
 
+	UPROPERTY(EditAnywhere, BlueprintReadWrite, Category = "ChaosPhysics|Clustering", meta = (EditCondition = "!bUseSizeSpecificDamageThreshold"))
+	TArray<float> DamageThreshold;
+
 	/** Damage threshold for clusters at different levels. */
 	UPROPERTY(EditAnywhere, BlueprintReadWrite, Category = "ChaosPhysics|Clustering")
-	TArray<float> DamageThreshold;
+	bool bUseSizeSpecificDamageThreshold;
 
 	/** */
 	UPROPERTY(EditAnywhere, BlueprintReadWrite, Category = "ChaosPhysics|Clustering")
diff --git a/Engine/Source/Runtime/Experimental/GeometryCollectionEngine/Public/GeometryCollection/GeometryCollectionObject.h b/Engine/Source/Runtime/Experimental/GeometryCollectionEngine/Public/GeometryCollection/GeometryCollectionObject.h
--- a/Engine/Source/Runtime/Experimental/GeometryCollectionEngine/Public/GeometryCollection/GeometryCollectionObject.h
+++ b/Engine/Source/Runtime/Experimental/GeometryCollectionEngine/Public/GeometryCollection/GeometryCollectionObject.h
@@ -99,6 +99,12 @@ struct GEOMETRYCOLLECTIONENGINE_API FGeometryCollectionSizeSpecificData
 	UPROPERTY(EditAnywhere, Category = "Collisions")
 	int32 MaximumCollisionParticles;
 
+	/**
+	 * Max number of particles.
+	*/
+	UPROPERTY(EditAnywhere, Category = "Collisions")
+	int32 DamageThreshold;
+
 };
 
 class FGeometryCollectionNaniteData
