Chaos : Destruction
- Volume based size specific data.
#okforgithum public, internal

diff --git a/Engine/Source/Runtime/Experimental/Chaos/Private/PhysicsProxy/GeometryCollectionPhysicsProxy.cpp b/Engine/Source/Runtime/Experimental/Chaos/Private/PhysicsProxy/GeometryCollectionPhysicsProxy.cpp
--- a/Engine/Source/Runtime/Experimental/Chaos/Private/PhysicsProxy/GeometryCollectionPhysicsProxy.cpp
+++ b/Engine/Source/Runtime/Experimental/Chaos/Private/PhysicsProxy/GeometryCollectionPhysicsProxy.cpp
@@ -938,10 +938,8 @@ DECLARE_CYCLE_STAT(TEXT("FGeometryCollectionPhysicsProxy::BuildClusters"), STAT_
 DECLARE_CYCLE_STAT(TEXT("FGeometryCollectionPhysicsProxy::BuildClusters:GlobalMatrices"), STAT_BuildClustersGlobalMatrices, STATGROUP_Chaos);
 
 
-int32 FindSizeSpecificIdx(const TArray<FSharedSimulationSizeSpecificData>& SizeSpecificData, const FBox& Bounds)
+int32 FindSizeSpecificIdx(const TArray<FSharedSimulationSizeSpecificData>& SizeSpecificData, float Volume)
 {
-	const FVector Extents = Bounds.GetExtent();
-	const float Size = Extents.GetAbsMin();
 	check(SizeSpecificData.Num());
 	int32 UseIdx = 0;
 	float PreSize = FLT_MAX;
@@ -949,7 +947,7 @@ int32 FindSizeSpecificIdx(const TArray<FSharedSimulationSizeSpecificData>& SizeS
 	{
 		ensureMsgf(PreSize >= SizeSpecificData[Idx].MaxSize, TEXT("SizeSpecificData is not sorted"));
 		PreSize = SizeSpecificData[Idx].MaxSize;
-		if (Size < SizeSpecificData[Idx].MaxSize)
+		if (Volume < SizeSpecificData[Idx].MaxSize)
 			UseIdx = Idx;
 		else
 			break;
@@ -957,6 +955,10 @@ int32 FindSizeSpecificIdx(const TArray<FSharedSimulationSizeSpecificData>& SizeS
 	return UseIdx;
 }
 
+int32 FindSizeSpecificIdx(const TArray<FSharedSimulationSizeSpecificData>& SizeSpecificData, const FBox& Bounds)
+{
+	return FindSizeSpecificIdx(SizeSpecificData, Bounds.GetVolume());
+}
 
 template <typename Traits>
 Chaos::TPBDRigidClusteredParticleHandle<float, 3>* 
