Trivial cleanup

diff --git a/Engine/Source/Runtime/Experimental/GeometryCollectionEngine/Private/GeometryCollection/GeometryCollectionSceneProxy.cpp b/Engine/Source/Runtime/Experimental/GeometryCollectionEngine/Private/GeometryCollection/GeometryCollectionSceneProxy.cpp
--- a/Engine/Source/Runtime/Experimental/GeometryCollectionEngine/Private/GeometryCollection/GeometryCollectionSceneProxy.cpp
+++ b/Engine/Source/Runtime/Experimental/GeometryCollectionEngine/Private/GeometryCollection/GeometryCollectionSceneProxy.cpp
@@ -185,7 +185,7 @@ void FGeometryCollectionSceneProxy::InitResources()
 	NumVertices = ConstantData->Vertices.Num();
 	NumIndices = ConstantData->Indices.Num()*3;	
 
-	// taken from this, and expanded here to accomodate modifications for
+	// taken from this, and expanded here to accommodate modifications for
 	// GeometryCollection vertex factory data (transform and bonemap)
 	// VertexBuffers.InitWithDummyData(&VertexFactory, GetRequiredVertexCount());
 
@@ -307,7 +307,7 @@ void FGeometryCollectionSceneProxy::BuildGeometry( const FGeometryCollectionCons
 	check(ConstantDataIn->Indices.Num() * 3 == NumIndices);
 
 	OutIndices.SetNumUninitialized(NumIndices);
-	ParallelFor (ConstantDataIn->Indices.Num(), [&](int32 IndexIdx)
+	ParallelFor(ConstantDataIn->Indices.Num(), [&](int32 IndexIdx)
 	{
 		OutIndices[IndexIdx * 3 ]    = ConstantDataIn->Indices[IndexIdx].X;
 		OutIndices[IndexIdx * 3 + 1] = ConstantDataIn->Indices[IndexIdx].Y;
