Chaos::FComplex
-- FComplex now uses an array of FVec3 instead of TParticles to represent vertices.

diff --git a/Engine/Plugins/Experimental/ChaosCloth/Source/Chaos/Private/ChaosCloth/ChaosClothingSimulation.cpp b/Engine/Plugins/Experimental/ChaosCloth/Source/Chaos/Private/ChaosCloth/ChaosClothingSimulation.cpp
--- a/Engine/Plugins/Experimental/ChaosCloth/Source/Chaos/Private/ChaosCloth/ChaosClothingSimulation.cpp
+++ b/Engine/Plugins/Experimental/ChaosCloth/Source/Chaos/Private/ChaosCloth/ChaosClothingSimulation.cpp
@@ -1075,17 +1075,17 @@ static void DrawConvex(FPrimitiveDrawInterface* PDI, const FConvex& Convex, cons
 			// Find the two surface points that belong to both Plane1 and Plane2
 			uint32 ParticleIndex1 = INDEX_NONE;
 
-			const TParticles<float, 3>& SurfaceParticles = Convex.GetSurfaceParticles();
-			for (uint32 ParticleIndex = 0; ParticleIndex < SurfaceParticles.Size(); ++ParticleIndex)
+			const TArray<FVec3>& Vertices = Convex.GetVertices();
+			for (int32 ParticleIndex = 0; ParticleIndex < Vertices.Num(); ++ParticleIndex)
 			{
-				const TVector<float, 3>& X = SurfaceParticles.X(ParticleIndex);
+				const FVec3& X = Vertices[ParticleIndex];
 
 				if (FMath::Square(Plane1.SignedDistance(X)) < KINDA_SMALL_NUMBER && 
 					FMath::Square(Plane2.SignedDistance(X)) < KINDA_SMALL_NUMBER)
 				{
 					if (ParticleIndex1 != INDEX_NONE)
 					{
-						const TVector<float, 3>& X1 = SurfaceParticles.X(ParticleIndex1);
+						const FVec3& X1 = Vertices[ParticleIndex1];
 						const FVector Position1 = Position + Rotation.RotateVector(X1);
 						const FVector Position2 = Position + Rotation.RotateVector(X);
 						DrawLine(PDI, Position1, Position2, Color);
diff --git a/Engine/Plugins/Experimental/ChaosCloth/Source/Chaos/Private/ChaosCloth/ChaosClothingSimulationCollider.cpp b/Engine/Plugins/Experimental/ChaosCloth/Source/Chaos/Private/ChaosCloth/ChaosClothingSimulationCollider.cpp
--- a/Engine/Plugins/Experimental/ChaosCloth/Source/Chaos/Private/ChaosCloth/ChaosClothingSimulationCollider.cpp
+++ b/Engine/Plugins/Experimental/ChaosCloth/Source/Chaos/Private/ChaosCloth/ChaosClothingSimulationCollider.cpp
@@ -228,15 +228,15 @@ void FClothingSimulationCollider::FLODData::Add(
 			if (Planes.Num() == NumPlanes)
 			{
 				// Retrieve particles
-				TParticles<float, 3> SurfaceParticles;
-				SurfaceParticles.Resize(NumSurfacePoints);
+				TArray<FVec3> Vertices;
+				Vertices.SetNum(NumSurfacePoints);
 				for (int32 ParticleIndex = 0; ParticleIndex < NumSurfacePoints; ++ParticleIndex)
 				{
-					SurfaceParticles.X(ParticleIndex) = Convex.SurfacePoints[ParticleIndex];
+					Vertices[ParticleIndex] = Convex.SurfacePoints[ParticleIndex];
 				}
 
 				// Setup the collision particle geometry
-				Solver->SetCollisionGeometry(ConvexOffset, Index, MakeUnique<FConvex>(MoveTemp(Planes), MoveTemp(SurfaceParticles)));
+				Solver->SetCollisionGeometry(ConvexOffset, Index, MakeUnique<FConvex>(MoveTemp(Planes), MoveTemp(Vertices)));
 			}
 			else
 			{
@@ -529,11 +529,11 @@ void FClothingSimulationCollider::ExtractPhysicsAssetCollision(FClothCollisionDa
 				}
 
 				// Copy surface points
-				const uint32 NumSurfacePoints = ChaosConvex.GetSurfaceParticles().Size();
+				const int32 NumSurfacePoints = ChaosConvex.NumVertices();
 				Convex.SurfacePoints.Reserve(NumSurfacePoints);
-				for (uint32 ParticleIndex = 0; ParticleIndex < NumSurfacePoints; ++ParticleIndex)
+				for (int32 ParticleIndex = 0; ParticleIndex < NumSurfacePoints; ++ParticleIndex)
 				{
-					Convex.SurfacePoints.Add(ChaosConvex.GetSurfaceParticles().X(ParticleIndex));
+					Convex.SurfacePoints.Add(ChaosConvex.GetVertex(ParticleIndex));
 				}
 #endif  // #if PHYSICS_INTERFACE_PHYSX #elif WITH_CHAOS
 
diff --git a/Engine/Plugins/Experimental/Water/Source/Runtime/Private/OceanCollisionComponent.cpp b/Engine/Plugins/Experimental/Water/Source/Runtime/Private/OceanCollisionComponent.cpp
--- a/Engine/Plugins/Experimental/Water/Source/Runtime/Private/OceanCollisionComponent.cpp
+++ b/Engine/Plugins/Experimental/Water/Source/Runtime/Private/OceanCollisionComponent.cpp
@@ -94,15 +94,14 @@ void UOceanCollisionComponent::UpdateBodySetup(const TArray<FKConvexElem>& Conve
 	for (FKConvexElem& Elem : CachedBodySetup->AggGeom.ConvexElems)
 	{
 		int32 NumHullVerts = Elem.VertexData.Num();
-
-		Chaos::TParticles<Chaos::FReal, 3> ConvexParticles;
-		ConvexParticles.AddParticles(NumHullVerts);
+		TArray<Chaos::FVec3> ConvexVertices;
+		ConvexVertices.SetNum(NumHullVerts);
 		for (int32 VertIndex = 0; VertIndex < NumHullVerts; ++VertIndex)
 		{
-			const FVector& HullVert = Elem.VertexData[VertIndex];
-			ConvexParticles.X(VertIndex) = FVector(HullVert.X, HullVert.Y, HullVert.Z);
+			ConvexVertices[VertIndex] = Elem.VertexData[VertIndex];
 		}
-		TSharedPtr<Chaos::FConvex, ESPMode::ThreadSafe> ChaosConvex = MakeShared<Chaos::FConvex, ESPMode::ThreadSafe>(ConvexParticles, 0.0f);
+
+		TSharedPtr<Chaos::FConvex, ESPMode::ThreadSafe> ChaosConvex = MakeShared<Chaos::FConvex, ESPMode::ThreadSafe>(ConvexVertices, 0.0f);
 
 		Elem.SetChaosConvexMesh(MoveTemp(ChaosConvex));
 	}
diff --git a/Engine/Plugins/Experimental/Water/Source/Runtime/Private/WaterBodyRiverActor.cpp b/Engine/Plugins/Experimental/Water/Source/Runtime/Private/WaterBodyRiverActor.cpp
--- a/Engine/Plugins/Experimental/Water/Source/Runtime/Private/WaterBodyRiverActor.cpp
+++ b/Engine/Plugins/Experimental/Water/Source/Runtime/Private/WaterBodyRiverActor.cpp
@@ -293,14 +293,15 @@ void URiverGenerator::UpdateSplineMesh(USplineMeshComponent* MeshComp, int32 Spl
 			for (FKConvexElem& Elem : MeshComp->GetBodySetup()->AggGeom.ConvexElems)
 			{
 				int32 NumHullVerts = Elem.VertexData.Num();
-				Chaos::TParticles<Chaos::FReal, 3> ConvexParticles;
-				ConvexParticles.AddParticles(NumHullVerts);
+				TArray<Chaos::FVec3> ConvexVertices;
+				ConvexVertices.SetNum(NumHullVerts);
 				for (int32 VertIndex = 0; VertIndex < NumHullVerts; ++VertIndex)
 				{
-					const FVector& HullVert = Elem.VertexData[VertIndex];
-					ConvexParticles.X(VertIndex) = FVector(HullVert.X, HullVert.Y, HullVert.Z);
+					ConvexVertices[VertIndex] = Elem.VertexData[VertIndex];
 				}
-				TSharedPtr<Chaos::FConvex, ESPMode::ThreadSafe> ChaosConvex = MakeShared<Chaos::FConvex, ESPMode::ThreadSafe>(ConvexParticles, 0.0f);
+
+				TSharedPtr<Chaos::FConvex, ESPMode::ThreadSafe> ChaosConvex = MakeShared<Chaos::FConvex, ESPMode::ThreadSafe>(ConvexVertices, 0.0f);
+
 				Elem.SetChaosConvexMesh(MoveTemp(ChaosConvex));
 			}
 
diff --git a/Engine/Source/Programs/HeadlessChaos/Private/HeadlessChaosTestConvex.cpp b/Engine/Source/Programs/HeadlessChaos/Private/HeadlessChaosTestConvex.cpp
--- a/Engine/Source/Programs/HeadlessChaos/Private/HeadlessChaosTestConvex.cpp
+++ b/Engine/Source/Programs/HeadlessChaos/Private/HeadlessChaosTestConvex.cpp
@@ -19,23 +19,23 @@ namespace ChaosTest
 	// leave the hull with only one face per normal.
 	void TestConvexBuilderConvexBoxFaceMerge(const FVec3* Vertices, const int32 NumVertices)
 	{
-		TParticles<FReal, 3> Particles;
-		Particles.AddParticles(NumVertices);
+		TArray<Chaos::FVec3> Particles;
+		Particles.SetNum(NumVertices);
 		for (int32 ParticleIndex = 0; ParticleIndex < NumVertices; ++ParticleIndex)
 		{
-			Particles.X(ParticleIndex) = Vertices[ParticleIndex];
+			Particles[ParticleIndex] = Vertices[ParticleIndex];
 		}
 
 		TArray<TPlaneConcrete<FReal, 3>> Planes;
 		TArray<TArray<int32>> FaceVertices;
-		TParticles<FReal, 3> SurfaceParticles;
+		TArray<Chaos::FVec3> SurfaceParticles;
 		TAABB<FReal, 3> LocalBounds;
 
 		FConvexBuilder::Build(Particles, Planes, FaceVertices, SurfaceParticles, LocalBounds);
 		FConvexBuilder::MergeFaces(Planes, FaceVertices, SurfaceParticles);
 
 		// Check that we have the right number of faces and particles
-		EXPECT_EQ(SurfaceParticles.Size(), 8);
+		EXPECT_EQ(SurfaceParticles.Num(), 8);
 		EXPECT_EQ(Planes.Num(), 6);
 		EXPECT_EQ(FaceVertices.Num(), 6);
 
@@ -47,9 +47,9 @@ namespace ChaosTest
 			{
 				int32 VertexIndex1 = Chaos::Utilities::WrapIndex(VertexIndex0 + 1, 0, FaceVertices[FaceIndex].Num());
 				int32 VertexIndex2 = Chaos::Utilities::WrapIndex(VertexIndex0 + 2, 0, FaceVertices[FaceIndex].Num());
-				const FVec3 Vertex0 = SurfaceParticles.X(FaceVertices[FaceIndex][VertexIndex0]);
-				const FVec3 Vertex1 = SurfaceParticles.X(FaceVertices[FaceIndex][VertexIndex1]);
-				const FVec3 Vertex2 = SurfaceParticles.X(FaceVertices[FaceIndex][VertexIndex2]);
+				const FVec3 Vertex0 = SurfaceParticles[FaceVertices[FaceIndex][VertexIndex0]];
+				const FVec3 Vertex1 = SurfaceParticles[FaceVertices[FaceIndex][VertexIndex1]];
+				const FVec3 Vertex2 = SurfaceParticles[FaceVertices[FaceIndex][VertexIndex2]];
 
 				// All vertices should lie in a plane at the same distance
 				const FReal Dist0 = FVec3::DotProduct(Vertex0, Planes[FaceIndex].Normal());
@@ -119,11 +119,11 @@ namespace ChaosTest
 	// Check that the convex structure data is consistent
 	void TestConvexStructureData(const FVec3* Vertices, const int32 NumVertices)
 	{
-		TParticles<FReal, 3> Particles;
-		Particles.AddParticles(NumVertices);
+		TArray<Chaos::FVec3> Particles;
+		Particles.SetNum(NumVertices);
 		for (int32 ParticleIndex = 0; ParticleIndex < NumVertices; ++ParticleIndex)
 		{
-			Particles.X(ParticleIndex) = Vertices[ParticleIndex];
+			Particles[ParticleIndex] = Vertices[ParticleIndex];
 		}
 
 		FConvex Convex(Particles, 0.0f);
diff --git a/Engine/Source/Programs/HeadlessChaos/Private/HeadlessChaosTestEPA.cpp b/Engine/Source/Programs/HeadlessChaos/Private/HeadlessChaosTestEPA.cpp
--- a/Engine/Source/Programs/HeadlessChaos/Private/HeadlessChaosTestEPA.cpp
+++ b/Engine/Source/Programs/HeadlessChaos/Private/HeadlessChaosTestEPA.cpp
@@ -663,7 +663,7 @@ namespace ChaosTest
 					{{768.000000, -448.000000, 3.81469727e-06}, {6.62273836e-09, 6.62273836e-09, -1.00000000}}
 				});
 
-			TParticles<FReal, 3> SurfaceParticles(
+			TArray<FVec3> SurfaceParticles(
 				{
 					{0.000000000, -1024.00000, 2.84217094e-14},
 					{768.000000, -1024.00000, 2.84217094e-14},
@@ -758,7 +758,7 @@ namespace ChaosTest
 					{{0.000000000,0.000000000,0.000000000},{-0.000000000,0.000000000,1.00000000}}
 					});
 
-				TParticles<FReal,3> SurfaceParticles(
+				TArray<FVec3> SurfaceParticles(
 					{
 						{0.000000000,-512.000000,-32.0000000},
 					{512.000000,0.000000000,-32.0000000},
@@ -811,7 +811,7 @@ namespace ChaosTest
 		}
 
 		// Defining boat geom data outside of single test scope as it's used in multiple.
-		const TArray<TVector<FReal, 3>> BoatSurfaceParticles(
+		const TArray<Chaos::FVec3> BoatSurfaceVertices(
 			{
 				{-118.965088, -100.379936, 105.818298},
 				{-128.562881, 80.0933762, 107.703270},
@@ -1021,7 +1021,7 @@ namespace ChaosTest
 					{{0.000000000,0.000000000,0.000000000},{-0.000000000,0.000000000,1.00000000}}
 				});
 
-			TParticles<FReal, 3> GroundSurfaceParticles(
+			TArray<FVec3> GroundSurfaceParticles(
 				{
 					{0.000000000,-512.000000,-32.0000000},
 					{512.000000,0.000000000,-32.0000000},
@@ -1042,8 +1042,8 @@ namespace ChaosTest
 
 
 			auto BoatConvexPlanesCopy = BoatConvexPlanes;
-			auto BoatSurfaceParticlesCopy = TParticles<FReal, 3>(CopyTemp(BoatSurfaceParticles));
-			FConvex BoatConvex = FConvex(MoveTemp(BoatConvexPlanesCopy), MoveTemp(BoatSurfaceParticlesCopy));
+			TArray<Chaos::FVec3> BoatSurfaceVerticesCopy = TArray<Chaos::FVec3>(CopyTemp(BoatSurfaceVertices));
+			FConvex BoatConvex = FConvex(MoveTemp(BoatConvexPlanesCopy), MoveTemp(BoatSurfaceVerticesCopy));
 
 
 			TRigidTransform<FReal, 3> BoatTransform(FVec3(-5421.507324, 2335.360840, 6.972876), FQuat(-0.016646, -0.008459, 0.915564, -0.401738), FVec3(1.0f));
@@ -1074,8 +1074,8 @@ namespace ChaosTest
 			ExpectedNormal.Normalize();
 
 			auto BoatConvexPlanesCopy = BoatConvexPlanes;
-			auto BoatSurfaceParticlesCopy = TParticles<FReal, 3>(CopyTemp(BoatSurfaceParticles));
-			FConvex BoatConvex = FConvex(MoveTemp(BoatConvexPlanesCopy), MoveTemp(BoatSurfaceParticlesCopy));
+			TArray<Chaos::FVec3> BoatSurfaceVerticesCopy = TArray<Chaos::FVec3>(CopyTemp(BoatSurfaceVertices));
+			FConvex BoatConvex = FConvex(MoveTemp(BoatConvexPlanesCopy), MoveTemp(BoatSurfaceVerticesCopy));
 
 			TRigidTransform<FReal, 3> QueryTM(FVec3(-10831.1875, -11206.3750, -6140.38135), FQuat(-0.524916053, -0.0370868668, -0.126528814, 0.840879321), FVec3(1.0f));
 
@@ -1186,7 +1186,7 @@ namespace ChaosTest
 		GTEST_TEST(EPATests, EPARealFailures_ConvexTrimeshRotationalDifferencesBreakNormal)
 		{
 			using namespace Chaos;
-			TParticles<FReal, 3> ConvexBoxSurfaceParticles(
+			TArray<FVec3> ConvexBoxSurfaceParticles(
 				{
 				{50.0999985, -50.1124992, -50.1250000},
 				{-50.0999985, 50.1250000, -50.1250000},
diff --git a/Engine/Source/Programs/HeadlessChaos/Private/HeadlessChaosTestGJK.cpp b/Engine/Source/Programs/HeadlessChaos/Private/HeadlessChaosTestGJK.cpp
--- a/Engine/Source/Programs/HeadlessChaos/Private/HeadlessChaosTestGJK.cpp
+++ b/Engine/Source/Programs/HeadlessChaos/Private/HeadlessChaosTestGJK.cpp
@@ -553,12 +553,12 @@ namespace ChaosTest
 
 		{
 			//Tetrahedron
-			TParticles<T, 3> HullParticles;
-			HullParticles.AddParticles(4);
-			HullParticles.X(0) = { -1,-1,-1 };
-			HullParticles.X(1) = { 1,-1,-1 };
-			HullParticles.X(2) = { 0,1,-1 };
-			HullParticles.X(3) = { 0,0,1 };
+			TArray<Chaos::FVec3> HullParticles;
+			HullParticles.SetNum(4);
+			HullParticles[0] = { -1,-1,-1 };
+			HullParticles[1] = { 1,-1,-1 };
+			HullParticles[2] = { 0,1,-1 };
+			HullParticles[3] = { 0,0,1 };
 			FConvex B(HullParticles, 0.0f);
 
 			for (const TVector<T, 3>& InitialDir : InitialDirs)
@@ -588,11 +588,11 @@ namespace ChaosTest
 
 		{
 			//Triangle
-			TParticles<T, 3> TriangleParticles;
-			TriangleParticles.AddParticles(3);
-			TriangleParticles.X(0) = { -1,-1,-1 };
-			TriangleParticles.X(1) = { 1,-1,-1 };
-			TriangleParticles.X(2) = { 0,1,-1 };
+			TArray<Chaos::FVec3> TriangleParticles;
+			TriangleParticles.SetNum(3);
+			TriangleParticles[0] = { -1,-1,-1 };
+			TriangleParticles[1] = { 1,-1,-1 };
+			TriangleParticles[2] = { 0,1,-1 };
 			FConvex B(TriangleParticles, 0.0f);
 
 			//triangle
@@ -932,12 +932,12 @@ namespace ChaosTest
 	{
 		typedef TVector<T, 3> TVector3;
 		//Tetrahedron
-		TParticles<T, 3> HullParticles;
-		HullParticles.AddParticles(4);
-		HullParticles.X(0) = { 3,0,4 };
-		HullParticles.X(1) = { 3,1,0 };
-		HullParticles.X(2) = { 3,-1,0 };
-		HullParticles.X(3) = { 4,0,2 };
+		TArray<Chaos::FVec3> HullParticles;
+		HullParticles.SetNum(4);
+		HullParticles[0] = { 3,0,4 };
+		HullParticles[1] = { 3,1,0 };
+		HullParticles[2] = { 3,-1,0 };
+		HullParticles[3] = { 4,0,2 };
 		FConvex A(HullParticles, 0.0f);
 		TSphere<T, 3> B(TVector<T, 3>(0, 0, 0), 1);
 
@@ -1368,19 +1368,19 @@ namespace ChaosTest
 
 		{
 			//based on real sweep from game
-			TParticles<T, 3> ConvexParticles;
-			ConvexParticles.AddParticles(10);
+			TArray<Chaos::FVec3> ConvexParticles;
+			ConvexParticles.SetNum(10);
 
-			ConvexParticles.X(0) = { 51870.2305, 54369.6719, 19200.0000 };
-			ConvexParticles.X(1) = { -91008.5625, -59964.0000, -19199.9629 };
-			ConvexParticles.X(2) = { 51870.2305, 54369.6758, -19199.9668 };
-			ConvexParticles.X(3) = { 22164.4883, 124647.500, -19199.9961 };
-			ConvexParticles.X(4) = { 34478.5000, 123975.492, -19199.9961 };
-			ConvexParticles.X(5) = { -91008.5000, -59963.9375, 19200.0000 };
-			ConvexParticles.X(6) = { -91008.5000, 33715.5625, 19200.0000 };
-			ConvexParticles.X(7) = { 34478.4961, 123975.500, 19200.0000 };
-			ConvexParticles.X(8) = { 22164.4922, 124647.500, 19200.0000 };
-			ConvexParticles.X(9) = { -91008.5000, 33715.5625, -19199.9961 };
+			ConvexParticles[0] = { 51870.2305, 54369.6719, 19200.0000 };
+			ConvexParticles[1] = { -91008.5625, -59964.0000, -19199.9629 };
+			ConvexParticles[2] = { 51870.2305, 54369.6758, -19199.9668 };
+			ConvexParticles[3] = { 22164.4883, 124647.500, -19199.9961 };
+			ConvexParticles[4] = { 34478.5000, 123975.492, -19199.9961 };
+			ConvexParticles[5] = { -91008.5000, -59963.9375, 19200.0000 };
+			ConvexParticles[6] = { -91008.5000, 33715.5625, 19200.0000 };
+			ConvexParticles[7] = { 34478.4961, 123975.500, 19200.0000 };
+			ConvexParticles[8] = { 22164.4922, 124647.500, 19200.0000 };
+			ConvexParticles[9] = { -91008.5000, 33715.5625, -19199.9961 };
 
 
 			const FConvex A(ConvexParticles, 0.0f);
@@ -1400,17 +1400,17 @@ namespace ChaosTest
 	void GJKCapsuleConvexInitialOverlapSweep()
 	{
 		{
-			TParticles<T,3> ConvexParticles;
-			ConvexParticles.AddParticles(8);
-
-			ConvexParticles.X(0) ={-256.000031,12.0000601,384.000061};
-			ConvexParticles.X(1) ={256.000031,12.0000601,384.000061};
-			ConvexParticles.X(2) ={256.000031,12.0000601,6.10351563e-05};
-			ConvexParticles.X(3) ={-256.000031,-11.9999399,6.10351563e-05};
-			ConvexParticles.X(4) ={-256.000031,12.0000601,6.10351563e-05};
-			ConvexParticles.X(5) ={-256.000031,-11.9999399,384.000061};
-			ConvexParticles.X(6) ={256.000031,-11.9999399,6.10351563e-05};
-			ConvexParticles.X(7) ={256.000031,-11.9999399,384.000061};
+			TArray<Chaos::FVec3> ConvexParticles;
+			ConvexParticles.SetNum(8);
+
+			ConvexParticles[0] ={-256.000031,12.0000601,384.000061};
+			ConvexParticles[1] ={256.000031,12.0000601,384.000061};
+			ConvexParticles[2] ={256.000031,12.0000601,6.10351563e-05};
+			ConvexParticles[3] ={-256.000031,-11.9999399,6.10351563e-05};
+			ConvexParticles[4] ={-256.000031,12.0000601,6.10351563e-05};
+			ConvexParticles[5] ={-256.000031,-11.9999399,384.000061};
+			ConvexParticles[6] ={256.000031,-11.9999399,6.10351563e-05};
+			ConvexParticles[7] ={256.000031,-11.9999399,384.000061};
 
 			TUniquePtr<FConvex> UniqueConvex = MakeUnique<FConvex>(ConvexParticles, 0.0f);
 			TSerializablePtr<FConvex> AConv(UniqueConvex);
@@ -1434,25 +1434,25 @@ namespace ChaosTest
 		}
 
 		{
-			TParticles<T,3> ConvexParticles;
-			ConvexParticles.AddParticles(16);
-
-			ConvexParticles.X(0) ={-127.216454,203.240234,124.726524};
-			ConvexParticles.X(1) ={125.708847,203.240295,124.726524};
-			ConvexParticles.X(2) ={-120.419685,207.124924,-0.386817127};
-			ConvexParticles.X(3) ={-32.9052734,91.5147095,199.922119};
-			ConvexParticles.X(4) ={118.912071,91.3693237,155.363205};
-			ConvexParticles.X(5) ={31.3977623,91.5147705,199.922150};
-			ConvexParticles.X(6) ={115.392204,91.6678925,162.647476};
-			ConvexParticles.X(7) ={-120.419701,91.1026840,-0.386809498};
-			ConvexParticles.X(8) ={118.912086,207.124985,-0.386806667};
-			ConvexParticles.X(9) ={118.912086,91.1027603,-0.386806667};
-			ConvexParticles.X(10) ={-120.419685,91.3692703,155.363174};
-			ConvexParticles.X(11) ={-110.103012,199.020554,160.910324};
-			ConvexParticles.X(12) ={-116.899742,91.6678467,162.647491};
-			ConvexParticles.X(13) ={31.3977337,194.240265,194.534988};
-			ConvexParticles.X(14) ={-32.9052925,194.240204,194.534958};
-			ConvexParticles.X(15) ={108.595482,199.020599,160.910309};
+			TArray<TVec3<T>> ConvexParticles;
+			ConvexParticles.SetNum(16);
+
+			ConvexParticles[0] ={-127.216454,203.240234,124.726524};
+			ConvexParticles[1] ={125.708847,203.240295,124.726524};
+			ConvexParticles[2] ={-120.419685,207.124924,-0.386817127};
+			ConvexParticles[3] ={-32.9052734,91.5147095,199.922119};
+			ConvexParticles[4] ={118.912071,91.3693237,155.363205};
+			ConvexParticles[5] ={31.3977623,91.5147705,199.922150};
+			ConvexParticles[6] ={115.392204,91.6678925,162.647476};
+			ConvexParticles[7] ={-120.419701,91.1026840,-0.386809498};
+			ConvexParticles[8] ={118.912086,207.124985,-0.386806667};
+			ConvexParticles[9] ={118.912086,91.1027603,-0.386806667};
+			ConvexParticles[10] ={-120.419685,91.3692703,155.363174};
+			ConvexParticles[11] ={-110.103012,199.020554,160.910324};
+			ConvexParticles[12] ={-116.899742,91.6678467,162.647491};
+			ConvexParticles[13] ={31.3977337,194.240265,194.534988};
+			ConvexParticles[14] ={-32.9052925,194.240204,194.534958};
+			ConvexParticles[15] ={108.595482,199.020599,160.910309};
 
 			auto Convex = MakeShared<FConvex, ESPMode::ThreadSafe>(ConvexParticles, 0.0f);
 			const auto& A = *Convex;
@@ -1503,18 +1503,18 @@ namespace ChaosTest
 
 		{
 			//capsule vs triangle as we make the sweep longer the world space point of impact should stay the same
-			TParticles<T,3> ConvexParticles;
-			ConvexParticles.AddParticles(3);
+			TArray<TVec3<T>> ConvexParticles;
+			ConvexParticles.SetNum(3);
 
-			ConvexParticles.X(0) ={7400.00000, 12600.0000, 206.248123};
-			ConvexParticles.X(1) ={7500.00000, 12600.0000, 199.994904};
-			ConvexParticles.X(2) ={7500.00000, 12700.0000, 189.837433};
+			ConvexParticles[0] ={7400.00000, 12600.0000, 206.248123};
+			ConvexParticles[1] ={7500.00000, 12600.0000, 199.994904};
+			ConvexParticles[2] ={7500.00000, 12700.0000, 189.837433};
 			
 			TUniquePtr<FConvex> UniqueConvex = MakeUnique<FConvex>(ConvexParticles, 0.0f);
 			TSerializablePtr<FConvex> AConv(UniqueConvex);
 			const TImplicitObjectScaled<FConvex> AConvScaled(AConv,TVec3<T>(1.0,1.0,1.0));
 
-			TTriangle<T> A(ConvexParticles.X(0),ConvexParticles.X(1),ConvexParticles.X(2));
+			TTriangle<T> A(ConvexParticles[0],ConvexParticles[1],ConvexParticles[2]);
 
 			const TVec3<T> Pt0(0.0,0.0,-29.6999969);
 			TVec3<T> Pt1 = Pt0;
@@ -1744,8 +1744,7 @@ namespace ChaosTest
 		const FVec3 Scale = FVec3(50.0f);
 		const FReal Margin = 0.75f;
 
-		TParticles<FReal, 3> CoreShapeParticles(MoveTemp(CoreShapeVerts));
-		TUniquePtr<FImplicitConvex3> CoreConvexShapePtr = MakeUnique<FImplicitConvex3>(CoreShapeParticles, 0.0f);
+		TUniquePtr<FImplicitConvex3> CoreConvexShapePtr = MakeUnique<FImplicitConvex3>(CoreShapeVerts, 0.0f);
 		const TImplicitObjectScaled<FImplicitConvex3> ShapeA(MakeSerializable(CoreConvexShapePtr), Scale, Margin);
 		const TImplicitObjectScaled<FImplicitConvex3> ShapeB(MakeSerializable(CoreConvexShapePtr), Scale, Margin);
 		const FRigidTransform3 TransformA(FVec3(0.000000000f, 0.000000000f, 182.378937f), FRotation3::FromElements(0.000000000f, 0.000000000f, 0.707106650f, 0.707106888f));	// Top
diff --git a/Engine/Source/Programs/HeadlessChaos/Private/HeadlessChaosTestImplicits.cpp b/Engine/Source/Programs/HeadlessChaos/Private/HeadlessChaosTestImplicits.cpp
--- a/Engine/Source/Programs/HeadlessChaos/Private/HeadlessChaosTestImplicits.cpp
+++ b/Engine/Source/Programs/HeadlessChaos/Private/HeadlessChaosTestImplicits.cpp
@@ -1221,14 +1221,14 @@ namespace ChaosTest {
 		{
 			TParticles<T, 3> Particles;
 			Particles.AddParticles(9);
-			Particles.X(0) = TVector<float, 3>(-1, -1, -1);
-			Particles.X(1) = TVector<float, 3>(-1, -1, 1);
-			Particles.X(2) = TVector<float, 3>(-1, 1, -1);
-			Particles.X(3) = TVector<float, 3>(-1, 1, 1);
-			Particles.X(4) = TVector<float, 3>(1, -1, -1);
-			Particles.X(5) = TVector<float, 3>(1, -1, 1);
-			Particles.X(6) = TVector<float, 3>(1, 1, -1);
-			Particles.X(7) = TVector<float, 3>(1, 1, 1);
+			Particles.X(0) = Chaos::TVec3<float>(-1, -1, -1);
+			Particles.X(1) = Chaos::TVec3<float>(-1, -1, 1);
+			Particles.X(2) = Chaos::TVec3<float>(-1, 1, -1);
+			Particles.X(3) = Chaos::TVec3<float>(-1, 1, 1);
+			Particles.X(4) = Chaos::TVec3<float>(1, -1, -1);
+			Particles.X(5) = Chaos::TVec3<float>(1, -1, 1);
+			Particles.X(6) = Chaos::TVec3<float>(1, 1, -1);
+			Particles.X(7) = Chaos::TVec3<float>(1, 1, 1);
 			Particles.X(8) = TVector<float, 3>(0, 0, 0);
 			const TTriangleMesh<T> TriMesh = TTriangleMesh<T>::GetConvexHullFromParticles(Particles);
 			EXPECT_EQ(TriMesh.GetSurfaceElements().Num(), 12);
@@ -1239,17 +1239,23 @@ namespace ChaosTest {
 				EXPECT_NE(Tri.Z, 8);
 			}
 
-			FConvex Convex(Particles, 0.0f);
-			const TParticles<T, 3>& CulledParticles = Convex.GetSurfaceParticles();
-			EXPECT_EQ(CulledParticles.Size(), 8);
+			TArray<FVec3> Vertices;
+			Vertices.SetNum((int32)Particles.Size());
+			for (int32 VertexIndex = 0; VertexIndex < (int32)Particles.Size(); ++VertexIndex)
+			{
+				Vertices[VertexIndex] = Particles.X(VertexIndex);
+			}
+			FConvex Convex(Vertices, 0.0f);
+			const TArray<TVec3<T>>& CulledParticles = Convex.GetVertices();
+			EXPECT_EQ(CulledParticles.Num(), 8);
 
-			for (uint32 Idx = 0; Idx < CulledParticles.Size(); ++Idx)
+			for (int32 Idx = 0; Idx < CulledParticles.Num(); ++Idx)
 			{
-				EXPECT_NE(Particles.X(8), CulledParticles.X(Idx));	//interior particle gone
+				EXPECT_NE(Particles.X(8), CulledParticles[Idx]);	//interior particle gone
 				bool bFound = false;
 				for (uint32 InnerIdx = 0; InnerIdx < Particles.Size(); ++InnerIdx)	//remaining particles are from the original set
 				{
-					if (Particles.X(InnerIdx) == CulledParticles.X(Idx))
+					if (Particles.X(InnerIdx) == CulledParticles[Idx])
 					{
 						bFound = true;
 						break;
@@ -1284,31 +1290,30 @@ namespace ChaosTest {
 	{
 		{
 			//degenerates
-			Chaos::TParticles<T, 3> Particles;
-			Particles.AddParticles(3);
-			Particles.X(0) = TVector<float, 3>(-1, -1, -1);
-			Particles.X(1) = TVector<float, 3>(1, -1, -1);
-			Particles.X(2) = TVector<float, 3>(1, 1, -1);
+			TArray<Chaos::TVec3<T>> Particles;
+			Particles.SetNum(3);
+			Particles[0] = Chaos::TVec3<float>(-1, -1, -1);
+			Particles[1] = Chaos::TVec3<float>(1, -1, -1);
+			Particles[2] = Chaos::TVec3<float>(1, 1, -1);
 			TArray<TVector<int32, 3>>Indices;
 			Chaos::FConvexBuilder::BuildConvexHull(Particles, Indices);
 			EXPECT_EQ(Indices.Num(), 0);
-			Particles.AddParticles(1);
-			Particles.X(3) = TVector<float, 3>(2, 3, -1);
+			Particles.Add(Chaos::TVec3<float>( 2, 3, -1 ));
 			Chaos::FConvexBuilder::BuildConvexHull(Particles, Indices);
 			EXPECT_EQ(Indices.Num(), 0);
 		}
 		{
-			Chaos::TParticles<T, 3> Particles;
-			Particles.AddParticles(9);
-			Particles.X(0) = TVector<float, 3>(-1, -1, -1);
-			Particles.X(1) = TVector<float, 3>(-1, -1, 1);
-			Particles.X(2) = TVector<float, 3>(-1, 1, -1);
-			Particles.X(3) = TVector<float, 3>(-1, 1, 1);
-			Particles.X(4) = TVector<float, 3>(1, -1, -1);
-			Particles.X(5) = TVector<float, 3>(1, -1, 1);
-			Particles.X(6) = TVector<float, 3>(1, 1, -1);
-			Particles.X(7) = TVector<float, 3>(1, 1, 1);
-			Particles.X(8) = TVector<float, 3>(0, 0, 0);
+			TArray <Chaos::TVec3<T>> Particles;
+			Particles.SetNum(9);
+			Particles[0] = Chaos::TVec3<float>(-1, -1, -1);
+			Particles[1] = Chaos::TVec3<float>(-1, -1, 1);
+			Particles[2] = Chaos::TVec3<float>(-1, 1, -1);
+			Particles[3] = Chaos::TVec3<float>(-1, 1, 1);
+			Particles[4] = Chaos::TVec3<float>(1, -1, -1);
+			Particles[5] = Chaos::TVec3<float>(1, -1, 1);
+			Particles[6] = Chaos::TVec3<float>(1, 1, -1);
+			Particles[7] = Chaos::TVec3<float>(1, 1, 1);
+			Particles[8] = Chaos::TVec3<float>(0, 0, 0);
 			TArray<TVector<int32, 3>>Indices;
 			Chaos::FConvexBuilder::BuildConvexHull(Particles, Indices);
 			EXPECT_EQ(Indices.Num(), 12);
@@ -1320,26 +1325,26 @@ namespace ChaosTest {
 			}
 		}
 		{
-			Chaos::TParticles<T, 3> Particles;
-			Particles.AddParticles(5);
-			Particles.X(0) = TVector<float, 3>(-1, -1, -1);
-			Particles.X(1) = TVector<float, 3>(1, -1, -1);
-			Particles.X(2) = TVector<float, 3>(1, 1, -1);
-			Particles.X(3) = TVector<float, 3>(0, 0, 0.5);
-			Particles.X(4) = (Particles.X(3) - Particles.X(1)) * 0.5 + Particles.X(1) + TVector<float, 3>(0, 0, 0.1);
+			TArray<TVec3<T>> Particles;
+			Particles.SetNum(5);
+			Particles[0] = Chaos::TVec3<float>(-1, -1, -1);
+			Particles[1] = Chaos::TVec3<float>(1, -1, -1);
+			Particles[2] = Chaos::TVec3<float>(1, 1, -1);
+			Particles[3] = Chaos::TVec3<float>(0, 0, 0.5);
+			Particles[4] = (Particles[3]- Particles[1]) * 0.5 + Particles[1]+ Chaos::TVec3<float>(0, 0, 0.1);
 			TArray<TVector<int32, 3>> Indices;
 			Chaos::FConvexBuilder::BuildConvexHull(Particles, Indices);
 			EXPECT_EQ(Indices.Num(), 6);
 		}
 		{
-			Chaos::TParticles<T, 3> Particles;
-			Particles.AddParticles(6);
-			Particles.X(0) = TVector<float, 3>(-1, -1, -1);
-			Particles.X(1) = TVector<float, 3>(1, -1, -1);
-			Particles.X(2) = TVector<float, 3>(1, 1, -1);
-			Particles.X(3) = TVector<float, 3>(0, 0, 0.5);
-			Particles.X(4) = (Particles.X(3) - Particles.X(1)) * 0.5 + Particles.X(1) + TVector<float, 3>(0, 0, 0.1);
-			Particles.X(5) = Particles.X(4) + TVector<float, 3>(-0.1, 0, 0);
+			TArray<TVec3<T>> Particles;
+			Particles.SetNum(6);
+			Particles[0] = Chaos::TVec3<float>(-1, -1, -1);
+			Particles[1] = Chaos::TVec3<float>(1, -1, -1);
+			Particles[2] = Chaos::TVec3<float>(1, 1, -1);
+			Particles[3] = Chaos::TVec3<float>(0, 0, 0.5);
+			Particles[4] = (Particles[3]- Particles[1]) * 0.5 + Particles[1] + Chaos::TVec3<float>(0, 0, 0.1);
+			Particles[5] = Particles[4]+ Chaos::TVec3<float>(-0.1, 0, 0);
 			TArray<TVector<int32, 3>> Indices;
 			Chaos::FConvexBuilder::BuildConvexHull(Particles, Indices);
 			EXPECT_EQ(Indices.Num(), 8);
@@ -1349,17 +1354,17 @@ namespace ChaosTest {
 			// a large enough epsilon for building horizons in hull generation
 			// (tested to fail with 1e-1) we will generate a non-convex hull
 			// Using a scaled epsilon resolves this case
-			Chaos::TParticles<T, 3> Particles;
-			Particles.AddParticles(9);
-			Particles.X(0) = TVector<float, 3>(-1, -1, -1);
-			Particles.X(1) = TVector<float, 3>(-1, -1, 1);
-			Particles.X(2) = TVector<float, 3>(-1, 1, -1);
-			Particles.X(3) = TVector<float, 3>(-1, 1, 1);
-			Particles.X(4) = TVector<float, 3>(1, -1, -1);
-			Particles.X(5) = TVector<float, 3>(1, -1, 1);
-			Particles.X(6) = TVector<float, 3>(1, 1, -1);
-			Particles.X(7) = TVector<float, 3>(1, 1, 1);
-			Particles.X(8) = TVector<float, 3>(0.966962576, -0.0577232838, 0.959515572);
+			TArray<TVec3<T>> Particles;
+			Particles.SetNum(9);
+			Particles[0] = Chaos::TVec3<float>(-1, -1, -1);
+			Particles[1] = Chaos::TVec3<float>(-1, -1, 1);
+			Particles[2] = Chaos::TVec3<float>(-1, 1, -1);
+			Particles[3] = Chaos::TVec3<float>(-1, 1, 1);
+			Particles[4] = Chaos::TVec3<float>(1, -1, -1);
+			Particles[5] = Chaos::TVec3<float>(1, -1, 1);
+			Particles[6] = Chaos::TVec3<float>(1, 1, -1);
+			Particles[7] = Chaos::TVec3<float>(1, 1, 1);
+			Particles[8] = Chaos::TVec3<float>(0.966962576, -0.0577232838, 0.959515572);
 			
 			TArray<TVector<int32, 3>> Indices;
 			Chaos::FConvexBuilder::Params BuildParams;
@@ -1373,7 +1378,7 @@ namespace ChaosTest {
 			{
 				for (int32 i = 0; i < 3; ++i)
 				{
-					TVector<T, 3> V = Particles.X(Tri[i]);
+					TVector<T, 3> V = Particles[Tri[i]];
 					TVector<T, 3> VAbs = V.GetAbs();
 					T Max = VAbs.GetMax();
 					EXPECT_GE(Max, 1 - 1e-2);
@@ -1385,21 +1390,21 @@ namespace ChaosTest {
 			// only the original box - ignoring all interior and coplanar points.
 			// Note: If hull generation is changed to support non-triangular faces the conditions here
 			// will need to change as a correct hull in that method will produce only 6 faces not 12
-			Chaos::TParticles<T, 3> Particles;
+			TArray<TVec3<T>> Particles;
 			int32 NumParticles = 3600;
-			Particles.AddParticles(NumParticles);
-			Particles.X(0) = TVector<float, 3>(-1, -1, -1);
-			Particles.X(1) = TVector<float, 3>(-1, -1, 1);
-			Particles.X(2) = TVector<float, 3>(-1, 1, -1);
-			Particles.X(3) = TVector<float, 3>(-1, 1, 1);
-			Particles.X(4) = TVector<float, 3>(1, -1, -1);
-			Particles.X(5) = TVector<float, 3>(1, -1, 1);
-			Particles.X(6) = TVector<float, 3>(1, 1, -1);
-			Particles.X(7) = TVector<float, 3>(1, 1, 1);
+			Particles.SetNum(NumParticles);
+			Particles[0] = Chaos::TVec3<float>(-1, -1, -1);
+			Particles[1] = Chaos::TVec3<float>(-1, -1, 1);
+			Particles[2] = Chaos::TVec3<float>(-1, 1, -1);
+			Particles[3] = Chaos::TVec3<float>(-1, 1, 1);
+			Particles[4] = Chaos::TVec3<float>(1, -1, -1);
+			Particles[5] = Chaos::TVec3<float>(1, -1, 1);
+			Particles[6] = Chaos::TVec3<float>(1, 1, -1);
+			Particles[7] = Chaos::TVec3<float>(1, 1, 1);
 			FRandomStream Stream(42);
 			for(int i = 8; i < NumParticles; ++i)
 			{
-				Particles.X(i) = TVector<float, 3>(Stream.FRandRange(-1.f, 1.f), Stream.FRandRange(-1.f, 1.f), Stream.FRandRange(-1.f, 1.f));
+				Particles[i]= Chaos::TVec3<float>(Stream.FRandRange(-1.f, 1.f), Stream.FRandRange(-1.f, 1.f), Stream.FRandRange(-1.f, 1.f));
 			}
 			TArray<TVector<int32, 3>> Indices;
 
@@ -1413,7 +1418,7 @@ namespace ChaosTest {
 			{
 				for(int i = 0; i < 3; ++i)
 				{
-					TVector<T, 3> V = Particles.X(Tri[i]);
+					TVector<T, 3> V = Particles[Tri[i]];
 					TVector<T, 3> VAbs = V.GetAbs();
 					T Max = VAbs.GetMax();
 					EXPECT_GE(Max, 1 - 1e-2);
@@ -1426,46 +1431,46 @@ namespace ChaosTest {
 	template <typename T>
 	void Simplify()
 	{
-		Chaos::TParticles<T, 3> Particles;
-		Particles.AddParticles(18);
-		Particles.X(0) = TVector<float, 3>(0, 0, 12.0f);
-		Particles.X(1) = TVector<float, 3>(-0.707f, -0.707f, 10.0f);
-		Particles.X(2) = TVector<float, 3>(0, -1, 10.0f);
-		Particles.X(3) = TVector<float, 3>(0.707f, -0.707f, 10.0f);
-		Particles.X(4) = TVector<float, 3>(1, 0, 10.0f);
-		Particles.X(5) = TVector<float, 3>(0.707f, 0.707f, 10.0f);
-		Particles.X(6) = TVector<float, 3>(0.0f, 1.0f, 10.0f);
-		Particles.X(7) = TVector<float, 3>(-0.707f, 0.707f, 10.0f);
-		Particles.X(8) = TVector<float, 3>(-1.0f, 0.0f, 10.0f);
-		Particles.X(9) = TVector<float, 3>(-0.707f, -0.707f, 0.0f);
-		Particles.X(10) = TVector<float, 3>(0, -1, 0.0f);
-		Particles.X(11) = TVector<float, 3>(0.707f, -0.707f, 0.0f);
-		Particles.X(12) = TVector<float, 3>(1, 0, 0.0f);
-		Particles.X(13) = TVector<float, 3>(0.707f, 0.707f, 0.0f);
-		Particles.X(14) = TVector<float, 3>(0.0f, 1.0f, 0.0f);
-		Particles.X(15) = TVector<float, 3>(-0.707f, 0.707f, 0.0f);
-		Particles.X(16) = TVector<float, 3>(-1.0f, 0.0f, 0.0f);
-		Particles.X(17) = TVector<float, 3>(0, 0, -2.0f);
-
+		TArray<TVec3<T>> Particles;
+		Particles.SetNum(18);
+		Particles[0] = Chaos::TVec3<float>(0, 0, 12.0f);
+		Particles[1] = Chaos::TVec3<float>(-0.707f, -0.707f, 10.0f);
+		Particles[2] = Chaos::TVec3<float>(0, -1, 10.0f);
+		Particles[3] = Chaos::TVec3<float>(0.707f, -0.707f, 10.0f);
+		Particles[4] = Chaos::TVec3<float>(1, 0, 10.0f);
+		Particles[5] = Chaos::TVec3<float>(0.707f, 0.707f, 10.0f);
+		Particles[6] = Chaos::TVec3<float>(0.0f, 1.0f, 10.0f);
+		Particles[7] = Chaos::TVec3<float>(-0.707f, 0.707f, 10.0f);
+		Particles[8] = Chaos::TVec3<float>(-1.0f, 0.0f, 10.0f);
+		Particles[9] = Chaos::TVec3<float>(-0.707f, -0.707f, 0.0f);
+		Particles[10] = Chaos::TVec3<float>(0, -1, 0.0f);
+		Particles[11] = Chaos::TVec3<float>(0.707f, -0.707f, 0.0f);
+		Particles[12] = Chaos::TVec3<float>(1, 0, 0.0f);
+		Particles[13] = Chaos::TVec3<float>(0.707f, 0.707f, 0.0f);
+		Particles[14] = Chaos::TVec3<float>(0.0f, 1.0f, 0.0f);
+		Particles[15] = Chaos::TVec3<float>(-0.707f, 0.707f, 0.0f);
+		Particles[16] = Chaos::TVec3<float>(-1.0f, 0.0f, 0.0f);
+		Particles[17] = Chaos::TVec3<float>(0, 0, -2.0f);
+					
 		FConvex Convex(Particles, 0.0f);
 
 		// capture original details
-		uint32 OriginalNumberParticles = Convex.GetSurfaceParticles().Size();
+		int32 OriginalNumberParticles = Convex.NumVertices();
 		int32 OriginalNumberFaces = Convex.GetFaces().Num();
 		TAABB<T, 3> OriginalBoundingBox = Convex.BoundingBox();
 
-		const TParticles<T, 3>& CulledParticles = Convex.GetSurfaceParticles();
+		const TArray<TVec3<T>>& CulledParticles = Convex.GetVertices();
 		const TArray<TPlaneConcrete<T, 3>> Planes = Convex.GetFaces();
 
 		// set target number of particles in simplified convex
 		FConvexBuilder::PerformGeometryReduction = 1;
-		FConvexBuilder::ParticlesThreshold = 10;
+		FConvexBuilder::VerticesThreshold = 10;
 
 		// simplify
 		Convex.PerformanceWarningAndSimplifaction();
 
 		// capture new details
-		uint32 NewNumberParticles = Convex.GetSurfaceParticles().Size();
+		int32 NewNumberParticles = Convex.NumVertices();
 		int32 NewNumberFaces = Convex.GetFaces().Num();
 		TAABB<T, 3> NewBoundingBox = Convex.BoundingBox();
 
diff --git a/Engine/Source/Programs/HeadlessChaos/Private/HeadlessChaosTestMostOpposing.cpp b/Engine/Source/Programs/HeadlessChaos/Private/HeadlessChaosTestMostOpposing.cpp
--- a/Engine/Source/Programs/HeadlessChaos/Private/HeadlessChaosTestMostOpposing.cpp
+++ b/Engine/Source/Programs/HeadlessChaos/Private/HeadlessChaosTestMostOpposing.cpp
@@ -77,15 +77,15 @@ namespace ChaosTest
 		TVector<T, 3> Normal;
 		int32 FaceIndex;
 
-		TParticles<T, 3> Particles;
-		Particles.AddParticles(6);
-		Particles.X(0) = TVector<T, 3>(1, 1, 1);
-		Particles.X(1) = TVector<T, 3>(5, 1, 1);
-		Particles.X(2) = TVector<T, 3>(1, 5, 1);
+		TArray<TVec3<T>> Particles;
+		Particles.SetNum(6);
+		Particles[0] = TVec3<T>(1, 1, 1);
+		Particles[1] = TVec3<T>(5, 1, 1);
+		Particles[2] = TVec3<T>(1, 5, 1);
 
-		Particles.X(3) = TVector<T, 3>(1, 1, 1);
-		Particles.X(4) = TVector<T, 3>(1, 5, 1);
-		Particles.X(5) = TVector<T, 3>(1, 1, -5);
+		Particles[3] = TVec3<T>(1, 1, 1);
+		Particles[4] = TVec3<T>(1, 5, 1);
+		Particles[5] = TVec3<T>(1, 1, -5);
 
 		FConvex Convex(MoveTemp(Particles), 0.0f);
 
@@ -133,15 +133,15 @@ namespace ChaosTest
 		TVector<T, 3> Normal;
 		int32 FaceIndex;
 
-		TParticles<T, 3> Particles;
-		Particles.AddParticles(6);
-		Particles.X(0) = TVector<T, 3>(0, -1, 1);
-		Particles.X(1) = TVector<T, 3>(1, -1, -1);
-		Particles.X(2) = TVector<T, 3>(0, 1, 1);
+		TArray<TVec3<T>> Particles;
+		Particles.SetNum(6);
+		Particles[0] = TVec3<T>(0, -1, 1);
+		Particles[1] = TVec3<T>(1, -1, -1);
+		Particles[2] = TVec3<T>(0, 1, 1);
 
-		Particles.X(3) = TVector<T, 3>(0, -1, 1);
-		Particles.X(4) = TVector<T, 3>(0, 1, 1);
-		Particles.X(5) = TVector<T, 3>(-1, -1, -1);
+		Particles[3] = TVec3<T>(0, -1, 1);
+		Particles[4] = TVec3<T>(0, 1, 1);
+		Particles[5] = TVec3<T>(-1, -1, -1);
 
 		TUniquePtr<FImplicitObject> Convex = MakeUnique<FConvex>(MoveTemp(Particles), 0.0f);
 
diff --git a/Engine/Source/Programs/HeadlessChaos/Private/HeadlessChaosTestOneShotManifolds.cpp b/Engine/Source/Programs/HeadlessChaos/Private/HeadlessChaosTestOneShotManifolds.cpp
--- a/Engine/Source/Programs/HeadlessChaos/Private/HeadlessChaosTestOneShotManifolds.cpp
+++ b/Engine/Source/Programs/HeadlessChaos/Private/HeadlessChaosTestOneShotManifolds.cpp
@@ -251,11 +251,10 @@ namespace ChaosTest
 			FVec3( HalfSize,  HalfSize,  HalfSize),
 			FVec3( HalfSize, -HalfSize,  HalfSize),
 		};
-		TParticles<FReal, 3> BoxParticles(MoveTemp(BoxVerts));
 
 		// First box rotated 45 degrees about Z and placed at the oirigin
 		// Second box rotated 45 degrees about Z and placed along X axis so that we have edge-edge contact with specified Phi
-		FConvex Convex(BoxParticles, Margin);
+		FConvex Convex(BoxVerts, Margin);
 		FRigidTransform3 Convex1Transform = FRigidTransform3(FVec3(0.0f, 0.0f, 0.0f), FRotation3::FromAxisAngle(FVec3(0.0f, 0.0f, 1.0f), FMath::DegreesToRadians(45)));
 		FRigidTransform3 Convex2Transform = FRigidTransform3(FVec3(Offset + ExpectedPhi, 0.0f, 0.0f), FRotation3::FromAxisAngle(FVec3(0.0f, 1.0f, 0.0f), FMath::DegreesToRadians(45)));
 
diff --git a/Engine/Source/Programs/HeadlessChaos/Private/HeadlessChaosTestUtility.cpp b/Engine/Source/Programs/HeadlessChaos/Private/HeadlessChaosTestUtility.cpp
--- a/Engine/Source/Programs/HeadlessChaos/Private/HeadlessChaosTestUtility.cpp
+++ b/Engine/Source/Programs/HeadlessChaos/Private/HeadlessChaosTestUtility.cpp
@@ -543,17 +543,17 @@ namespace ChaosTest {
 		Particle->R() = TRotation<T, 3>::MakeFromEuler(TVector<T, 3>(0.f, 0.f, 0.f)).GetNormalized();
 		Particle->W() = TVector<T, 3>(0.f, 0.f, 0.f);
 
-		Chaos::TParticles<T, 3> Cube;
-		Cube.AddParticles(9);
-		Cube.X(0) = TVector<float, 3>(-1000, -1000, -20);
-		Cube.X(1) = TVector<float, 3>(-1000, -1000, 0);
-		Cube.X(2) = TVector<float, 3>(-1000, 1000, -20);
-		Cube.X(3) = TVector<float, 3>(-1000, 1000, 0);
-		Cube.X(4) = TVector<float, 3>(1000, -1000, -20);
-		Cube.X(5) = TVector<float, 3>(1000, -1000, 0);
-		Cube.X(6) = TVector<float, 3>(1000, 1000, -20);
-		Cube.X(7) = TVector<float, 3>(1000, 1000, 0);
-		Cube.X(8) = TVector<float, 3>(0, 0, 0);
+		TArray<Chaos::FVec3> Cube;
+		Cube.SetNum(9);
+		Cube[0] = Chaos::FVec3(-1000, -1000, -20);
+		Cube[1] = Chaos::FVec3(-1000, -1000, 0);
+		Cube[2] = Chaos::FVec3(-1000, 1000, -20);
+		Cube[3] = Chaos::FVec3(-1000, 1000, 0);
+		Cube[4] = Chaos::FVec3(1000, -1000, -20);
+		Cube[5] = Chaos::FVec3(1000, -1000, 0);
+		Cube[6] = Chaos::FVec3(1000, 1000, -20);
+		Cube[7] = Chaos::FVec3(1000, 1000, 0);
+		Cube[8] = Chaos::FVec3(0, 0, 0);
 
 		Particle->SetDynamicGeometry(MakeUnique<FConvex>(Cube, 0.0f));
 
@@ -594,17 +594,17 @@ namespace ChaosTest {
 	template<class T>
 	void AppendDynamicParticleConvexBox(TPBDRigidParticleHandle<T, 3> & InParticles, const TVector<T, 3>& Scale, FReal Margin)
 	{
-		Chaos::TParticles<T, 3> Cube;
-		Cube.AddParticles(9);
-		Cube.X(0) = TVector<float, 3>(-1, -1, -1)*Scale;
-		Cube.X(1) = TVector<float, 3>(-1, -1, 1)*Scale;
-		Cube.X(2) = TVector<float, 3>(-1, 1, -1)*Scale;
-		Cube.X(3) = TVector<float, 3>(-1, 1, 1)*Scale;
-		Cube.X(4) = TVector<float, 3>(1, -1, -1)*Scale;
-		Cube.X(5) = TVector<float, 3>(1, -1, 1)*Scale;
-		Cube.X(6) = TVector<float, 3>(1, 1, -1)*Scale;
-		Cube.X(7) = TVector<float, 3>(1, 1, 1)*Scale;
-		Cube.X(8) = TVector<float, 3>(0, 0, 0);
+		TArray<FVec3> Cube;
+		Cube.SetNum(9);
+		Cube[0] = Chaos::FVec3(-1, -1, -1)*Scale;
+		Cube[1] = Chaos::FVec3(-1, -1, 1)*Scale;
+		Cube[2] = Chaos::FVec3(-1, 1, -1)*Scale;
+		Cube[3] = Chaos::FVec3(-1, 1, 1)*Scale;
+		Cube[4] = Chaos::FVec3(1, -1, -1)*Scale;
+		Cube[5] = Chaos::FVec3(1, -1, 1)*Scale;
+		Cube[6] = Chaos::FVec3(1, 1, -1)*Scale;
+		Cube[7] = Chaos::FVec3(1, 1, 1)*Scale;
+		Cube[8] = Chaos::FVec3(0, 0, 0);
 
 		InParticles.X() = TVector<T, 3>(0.f, 0.f, 0.f);
 		InParticles.V() = TVector<T, 3>(0.f, 0.f, 0.f);
@@ -614,8 +614,8 @@ namespace ChaosTest {
 		InParticles.Q() = InParticles.R();
 
 		// TODO: Change this error prone API to set bounds more automatically. This is easy to forget
-		InParticles.SetLocalBounds(TAABB<T, 3>(Cube.X(0), Cube.X(7)));
-		InParticles.SetWorldSpaceInflatedBounds(TAABB<T, 3>(Cube.X(0), Cube.X(7)));
+		InParticles.SetLocalBounds(TAABB<T, 3>(Cube[0], Cube[7]));
+		InParticles.SetWorldSpaceInflatedBounds(TAABB<T, 3>(Cube[0], Cube[7]));
 		InParticles.SetHasBounds(true);
 
 		InParticles.M() = 1.f;
@@ -771,9 +771,8 @@ namespace ChaosTest {
 			FVec3(HalfSize.X,  HalfSize.Y,  HalfSize.Z),
 			FVec3(HalfSize.X, -HalfSize.Y,  HalfSize.Z),
 		};
-		TParticles<FReal, 3> BoxParticles(MoveTemp(BoxVerts));
-
-		TSharedPtr<FImplicitConvex3, ESPMode::ThreadSafe> BoxConvex = MakeShared<FImplicitConvex3, ESPMode::ThreadSafe>(BoxParticles, 0.0f);
+		
+		TSharedPtr<FImplicitConvex3, ESPMode::ThreadSafe> BoxConvex = MakeShared<FImplicitConvex3, ESPMode::ThreadSafe>(BoxVerts, 0.0f);
 
 		return TImplicitObjectScaled<FImplicitConvex3>(BoxConvex, BoxScale, Margin);
 	}
diff --git a/Engine/Source/Runtime/Core/Public/UObject/UE5MainStreamObjectVersion.h b/Engine/Source/Runtime/Core/Public/UObject/UE5MainStreamObjectVersion.h
--- a/Engine/Source/Runtime/Core/Public/UObject/UE5MainStreamObjectVersion.h
+++ b/Engine/Source/Runtime/Core/Public/UObject/UE5MainStreamObjectVersion.h
@@ -88,6 +88,9 @@ struct CORE_API FUE5MainStreamObjectVersion
 		// Make UMeshDescriptionBase serializable
 		SerializeMeshDescriptionBase,
 
+		// Chaos FConvex uses array of FVec3s for vertices instead of particles
+		ConvexUsesVerticesArray,
+
 		// -----<new versions can be added above this line>-------------------------------------------------
 		VersionPlusOne,
 		LatestVersion = VersionPlusOne - 1
diff --git a/Engine/Source/Runtime/Engine/Private/PhysicsEngine/BodySetup.cpp b/Engine/Source/Runtime/Engine/Private/PhysicsEngine/BodySetup.cpp
--- a/Engine/Source/Runtime/Engine/Private/PhysicsEngine/BodySetup.cpp
+++ b/Engine/Source/Runtime/Engine/Private/PhysicsEngine/BodySetup.cpp
@@ -1910,18 +1910,18 @@ TArray<int32> FKConvexElem::GetChaosConvexIndices() const
 	const int32 NumVerts = VertexData.Num();
 	if (NumVerts > 0)
 	{
-		Chaos::TParticles<Chaos::FReal, 3> ConvexParticles;
-		ConvexParticles.AddParticles(NumVerts);
+		TArray<Chaos::FVec3> ConvexVertices;
+		ConvexVertices.SetNum(NumVerts);
 
 		for (int32 VertIndex = 0; VertIndex < NumVerts; ++VertIndex)
 		{
-			ConvexParticles.X(VertIndex) = VertexData[VertIndex];
+			ConvexVertices[VertIndex] = VertexData[VertIndex];
 		}
 
 		TArray<Chaos::TVector<int32, 3>> Triangles;
 		Chaos::FConvexBuilder::Params BuildParams;
-		BuildParams.HorizonEpsilon = Chaos::FConvexBuilder::SuggestEpsilon(ConvexParticles);
-		Chaos::FConvexBuilder::BuildConvexHull(ConvexParticles, Triangles, BuildParams);
+		BuildParams.HorizonEpsilon = Chaos::FConvexBuilder::SuggestEpsilon(ConvexVertices);
+		Chaos::FConvexBuilder::BuildConvexHull(ConvexVertices, Triangles, BuildParams);
 
 		ResultIndexData.Reserve(Triangles.Num() * 3);
 		for (Chaos::TVector<int32, 3> Tri : Triangles)
diff --git a/Engine/Source/Runtime/Engine/Private/PhysicsEngine/Experimental/ChaosDerivedData.cpp b/Engine/Source/Runtime/Engine/Private/PhysicsEngine/Experimental/ChaosDerivedData.cpp
--- a/Engine/Source/Runtime/Engine/Private/PhysicsEngine/Experimental/ChaosDerivedData.cpp
+++ b/Engine/Source/Runtime/Engine/Private/PhysicsEngine/Experimental/ChaosDerivedData.cpp
@@ -381,18 +381,17 @@ void FChaosDerivedDataCooker::BuildConvexMeshes(TArray<TUniquePtr<Chaos::FImplic
 				Bounds.GrowToInclude(HullVert);
 			}
 
-			// Create the surface particle for the convex
-			TParticles<FReal, 3> ConvexParticles;
-			ConvexParticles.AddParticles(NumHullVerts);
-
+			// Create the vertices for the convex
+			TArray<FVec3> ConvexVertices;
+			ConvexVertices.SetNum(NumHullVerts);
 			for (int32 VertIndex = 0; VertIndex < NumHullVerts; ++VertIndex)
 			{
-				const FVector& HullVert = HullVerts[VertIndex];
-				ConvexParticles.X(VertIndex) = FVector(bMirrored ? -HullVert.X : HullVert.X, HullVert.Y, HullVert.Z);
+				const FVec3& HullVert = HullVerts[VertIndex];
+				ConvexVertices[VertIndex] = FVec3(bMirrored ? -HullVert.X : HullVert.X, HullVert.Y, HullVert.Z);
 			}
 
 			// Margin is always zero on convex shapes - they are intended to be instanced
-			OutConvexes.Emplace(new Chaos::FConvex(ConvexParticles, 0.0f));
+			OutConvexes.Emplace(new Chaos::FConvex(ConvexVertices, 0.0f));
 		}
 	};
 
diff --git a/Engine/Source/Runtime/Experimental/Chaos/Private/Chaos/ChaosDebugDraw.cpp b/Engine/Source/Runtime/Experimental/Chaos/Private/Chaos/ChaosDebugDraw.cpp
--- a/Engine/Source/Runtime/Experimental/Chaos/Private/Chaos/ChaosDebugDraw.cpp
+++ b/Engine/Source/Runtime/Experimental/Chaos/Private/Chaos/ChaosDebugDraw.cpp
@@ -267,7 +267,10 @@ namespace Chaos
 					{
 						// TODO: This is horrendously slow. Figure out a way to cache
 						// the generated trimeshes on the debug draw queue instance.
-						const TParticles<FReal, 3>& Particles = Convex->GetSurfaceParticles();
+						
+						// Copy the vertices then move them into particles..
+						TArray<FVec3> Vertices = Convex->GetVertices();
+						const TParticles<FReal, 3>& Particles(MoveTemp(Vertices));
 						TTriangleMesh<FReal> Triangles = TTriangleMesh<FReal>::GetConvexHullFromParticles(Particles);
 						for (const TVector<int32, 3>& Elem : Triangles.GetElements())
 						{
diff --git a/Engine/Source/Runtime/Experimental/Chaos/Private/Chaos/CollisionConvexMesh.cpp b/Engine/Source/Runtime/Experimental/Chaos/Private/Chaos/CollisionConvexMesh.cpp
--- a/Engine/Source/Runtime/Experimental/Chaos/Private/Chaos/CollisionConvexMesh.cpp
+++ b/Engine/Source/Runtime/Experimental/Chaos/Private/Chaos/CollisionConvexMesh.cpp
@@ -10,11 +10,11 @@ namespace Chaos
 
 	int32 FConvexBuilder::PerformGeometryReduction = 0;
 
-	int32 FConvexBuilder::ParticlesThreshold = 50;
+	int32 FConvexBuilder::VerticesThreshold = 50;
 
 	FAutoConsoleVariableRef CVarConvexGeometryCheckEnable(TEXT("p.Chaos.ConvexGeometryCheckEnable"), FConvexBuilder::PerformGeometryCheck, TEXT("Perform convex geometry complexity check for Chaos physics."));
 
 	FAutoConsoleVariableRef CVarConvexGeometrySimplifyEnable(TEXT("p.Chaos.PerformGeometryReduction"), FConvexBuilder::PerformGeometryReduction, TEXT("Perform convex geometry simplification to increase performance in Chaos physics."));
 
-	FAutoConsoleVariableRef CVarConvexParticlesWarningThreshold(TEXT("p.Chaos.ConvexParticlesWarningThreshold"), FConvexBuilder::ParticlesThreshold, TEXT("Threshold beyond which we warn about collision geometry complexity."));
+	FAutoConsoleVariableRef CVarConvexParticlesWarningThreshold(TEXT("p.Chaos.ConvexParticlesWarningThreshold"), FConvexBuilder::VerticesThreshold, TEXT("Threshold beyond which we warn about collision geometry complexity."));
 }
diff --git a/Engine/Source/Runtime/Experimental/Chaos/Private/Chaos/Convex.cpp b/Engine/Source/Runtime/Experimental/Chaos/Private/Chaos/Convex.cpp
--- a/Engine/Source/Runtime/Experimental/Chaos/Private/Chaos/Convex.cpp
+++ b/Engine/Source/Runtime/Experimental/Chaos/Private/Chaos/Convex.cpp
@@ -38,7 +38,7 @@ namespace Chaos
 		//
 		//  @todo(chaos) : Collision Manifold 
 		//     Create a correspondence between the faces and surface particles on construction.
-		//     The correspondence will provide an index between the Planes and the SurfaceParticles,
+		//     The correspondence will provide an index between the Planes and the Vertices,
 		//     removing the need for the exhaustive search here. 
 		//
 
@@ -50,13 +50,13 @@ namespace Chaos
 			FReal AbsOfSignedDistance = FMath::Abs(Plane.SignedDistance(Position));
 			if (AbsOfSignedDistance < SearchDist)
 			{
-				for (int32 Fdx = 0; Fdx < (int32)SurfaceParticles.Size(); Fdx++)
+				for (int32 Fdx = 0; Fdx < (int32)Vertices.Num(); Fdx++)
 				{
 					if (!IncludedParticles.Contains(Fdx))
 					{
-						if (FMath::Abs(Plane.SignedDistance(SurfaceParticles.X(Fdx))) < SearchDist)
+						if (FMath::Abs(Plane.SignedDistance(Vertices[Fdx])) < SearchDist)
 						{
-							FaceVertices.Add(SurfaceParticles.X(Fdx));
+							FaceVertices.Add(Vertices[Fdx]);
 							IncludedParticles.Add(Fdx);
 						}
 					}
@@ -169,7 +169,7 @@ namespace Chaos
 	// Store the structure data with the convex. This is used by manifold generation, for example
 	void FConvex::CreateStructureData(TArray<TArray<int32>>&& PlaneVertexIndices)
 	{
-		StructureData.SetPlaneVertices(MoveTemp(PlaneVertexIndices), SurfaceParticles.Size());
+		StructureData.SetPlaneVertices(MoveTemp(PlaneVertexIndices), Vertices.Num());
 	}
 
 	void FConvex::MovePlanesAndRebuild(const FReal InDelta)
@@ -271,12 +271,6 @@ namespace Chaos
 
 
 		// Use the new surface points
-		TParticles<FReal, 3> NewParticles;
-		NewParticles.AddParticles(NewPoints.Num());
-		for (int32 PointIndex = 0; PointIndex < NewPoints.Num(); ++PointIndex)
-		{
-			NewParticles.X(PointIndex) = NewPoints[PointIndex];
-		}
-		SurfaceParticles = MoveTemp(NewParticles);
+		Vertices = MoveTemp(NewPoints);
 	}
 }
\ No newline at end of file
diff --git a/Engine/Source/Runtime/Experimental/Chaos/Private/Chaos/Levelset.cpp b/Engine/Source/Runtime/Experimental/Chaos/Private/Chaos/Levelset.cpp
--- a/Engine/Source/Runtime/Experimental/Chaos/Private/Chaos/Levelset.cpp
+++ b/Engine/Source/Runtime/Experimental/Chaos/Private/Chaos/Levelset.cpp
@@ -1495,17 +1495,7 @@ void GetGeomSurfaceSamples(const TCapsule<FReal>& InGeom, TArray<FVec3>& OutSamp
 
 void GetGeomSurfaceSamples(const FConvex& InGeom, TArray<FVec3>& OutSamples)
 {
-	const TParticles<FReal, 3>& Particles = InGeom.GetSurfaceParticles();
-
-	const int32 NumParticles = Particles.Size();
-
-	OutSamples.Reset();
-	OutSamples.AddUninitialized(NumParticles);
-
-	for(int32 Index = 0; Index < NumParticles; ++Index)
-	{
-		OutSamples[Index] = Particles.X(Index);
-	}
+	OutSamples = InGeom.GetVertices();
 }
 
 template<typename InnerT>
diff --git a/Engine/Source/Runtime/Experimental/Chaos/Public/Chaos/CollisionConvexMesh.h b/Engine/Source/Runtime/Experimental/Chaos/Public/Chaos/CollisionConvexMesh.h
--- a/Engine/Source/Runtime/Experimental/Chaos/Public/Chaos/CollisionConvexMesh.h
+++ b/Engine/Source/Runtime/Experimental/Chaos/Public/Chaos/CollisionConvexMesh.h
@@ -42,23 +42,23 @@ namespace Chaos
 			FReal HorizonEpsilon;
 		};
 
-		static FReal SuggestEpsilon(const TParticles<FReal, 3>& InParticles)
+		static FReal SuggestEpsilon(const TArray<FVec3>& InVertices)
 		{
 			// Create a scaled epsilon for our input data set. FLT_EPSILON is the distance between 1.0 and the next value
 			// above 1.0 such that 1.0 + FLT_EPSILON != 1.0. Floats aren't equally disanced though so big or small numbers
 			// don't work well with it. Here we take the max absolute of each axis and scale that for a wider margin and
 			// use that to scale FLT_EPSILON to get a more relevant value.
 			TVector<FReal, 3> MaxAxes(TNumericLimits<FReal>::Lowest());
-			const int32 NumParticles = InParticles.Size();
+			const int32 NumVertices = InVertices.Num();
 
-			if(NumParticles <= 1)
+			if (NumVertices <= 1)
 			{
 				return FLT_EPSILON;
 			}
 
-			for(int32 Index = 0; Index < NumParticles; ++Index)
+			for (int32 Index = 0; Index < NumVertices; ++Index)
 			{
-				TVector<FReal, 3> PositionAbs = InParticles.X(Index).GetAbs();
+				TVector<FReal, 3> PositionAbs = InVertices[Index].GetAbs();
 
 				MaxAxes[0] = FMath::Max(MaxAxes[0], PositionAbs[0]);
 				MaxAxes[1] = FMath::Max(MaxAxes[1], PositionAbs[1]);
@@ -91,24 +91,24 @@ namespace Chaos
 			return FMath::IsNearlyEqual(DPointDistance, 0, KINDA_SMALL_NUMBER);
 		}
 
-		static bool IsPlanarShape(const TParticles<FReal, 3>& InParticles, FVec3& OutNormal)
+		static bool IsPlanarShape(const TArray<FVec3>& InVertices, FVec3& OutNormal)
 		{
 			bool bResult = false;
-			const int32 NumParticles = InParticles.Size();
+			const int32 NumVertices = InVertices.Num();
 			
-			if(NumParticles <= 3)
+			if (NumVertices <= 3)
 			{
 				// Nothing, point, line or triangle, not a planar set
 				return false;
 			}
 			else // > 3 points
 			{
-				FPlane TriPlane(InParticles.X(0), InParticles.X(1), InParticles.X(2));
+				FPlane TriPlane(InVertices[0], InVertices[1], InVertices[2]);
 				OutNormal = FVec3(TriPlane.X, TriPlane.Y, TriPlane.Z);
 
-				for(int32 Index = 3; Index < NumParticles; ++Index)
+				for (int32 Index = 3; Index < NumVertices; ++Index)
 				{
-					const float PointPlaneDot = FMath::Abs(TriPlane.PlaneDot(InParticles.X(Index)));
+					const float PointPlaneDot = FMath::Abs(TriPlane.PlaneDot(InVertices[Index]));
 					if(!FMath::IsNearlyEqual(PointPlaneDot, 0, KINDA_SMALL_NUMBER))
 					{
 						return false;
@@ -247,45 +247,45 @@ namespace Chaos
 
 	public:
 
-		static void Build(const TParticles<FReal, 3>& InParticles, TArray <TPlaneConcrete<FReal, 3>>& OutPlanes, TArray<TArray<int32>>& OutFaceIndices, TParticles<FReal, 3>& OutSurfaceParticles, TAABB<FReal, 3>& OutLocalBounds)
+		static void Build(const TArray<FVec3>& InVertices, TArray <TPlaneConcrete<FReal, 3>>& OutPlanes, TArray<TArray<int32>>& OutFaceIndices, TArray<FVec3>& OutVertices, TAABB<FReal, 3>& OutLocalBounds)
 		{
 			OutPlanes.Reset();
-			OutSurfaceParticles.Resize(0);
+			OutVertices.Reset();
 			OutLocalBounds = TAABB<FReal, 3>::EmptyAABB();
 
-			const uint32 NumParticlesIn = InParticles.Size();
-			if(NumParticlesIn == 0)
+			const int32 NumVerticesIn = InVertices.Num();
+			if(NumVerticesIn == 0)
 			{
 				return;
 			}
 
-			const TParticles<FReal, 3>* ParticlesToUse = &InParticles;
-			TParticles<FReal, 3> ModifiedParticles;
+			const TArray<FVec3>* VerticesToUse = &InVertices;
+			TArray<FVec3> ModifiedVertices;
 
 			// For triangles and planar shapes, create a very thin prism as a convex
-			auto Inflate = [](const TParticles<FReal, 3>& Source, TParticles<FReal, 3>& Destination, const FVec3& Normal, float Inflation)
+			auto Inflate = [](const TArray<FVec3>& Source, TArray<FVec3>& Destination, const FVec3& Normal, float Inflation)
 			{
-				const int32 NumSource = Source.Size();
-				Destination.Resize(0);
-				Destination.AddParticles(NumSource * 2);
+				const int32 NumSource = Source.Num();
+				Destination.Reset();
+				Destination.SetNum(NumSource * 2);
 
 				for(int32 Index = 0; Index < NumSource; ++Index)
 				{
-					Destination.X(Index) = Source.X(Index);
-					Destination.X(NumSource + Index) = Source.X(Index) + Normal * Inflation;
+					Destination[Index] = Source[Index];
+					Destination[NumSource + Index] = Source[Index] + Normal * Inflation;
 				}
 			};
 
 			FVec3 PlanarNormal(0);
-			if(NumParticlesIn == 3)
+			if (NumVerticesIn == 3)
 			{
-				const bool bIsValidTriangle = IsValidTriangle(InParticles.X(0), InParticles.X(1), InParticles.X(2), PlanarNormal);
+				const bool bIsValidTriangle = IsValidTriangle(InVertices[0], InVertices[1], InVertices[2], PlanarNormal);
 
 				//TODO_SQ_IMPLEMENTATION: should do proper cleanup to avoid this
 				if(ensureMsgf(bIsValidTriangle, TEXT("FConvexBuilder::Build(): Generated invalid triangle!")))
 				{
-					Inflate(InParticles, ModifiedParticles, PlanarNormal, TriQuadPrismInflation());
-					ParticlesToUse = &ModifiedParticles;
+					Inflate(InVertices, ModifiedVertices, PlanarNormal, TriQuadPrismInflation());
+					VerticesToUse = &ModifiedVertices;
 					UE_LOG(LogChaos, Verbose, TEXT("Encountered a triangle in convex hull generation. Will prepare a prism of thickness %.5f in place of a triangle."), TriQuadPrismInflation());
 				}
 				else
@@ -293,27 +293,27 @@ namespace Chaos
 					return;
 				}
 			}
-			else if(IsPlanarShape(InParticles, PlanarNormal))
+			else if (IsPlanarShape(InVertices, PlanarNormal))
 			{
-				Inflate(InParticles, ModifiedParticles, PlanarNormal, TriQuadPrismInflation());
-				ParticlesToUse = &ModifiedParticles;
+				Inflate(InVertices, ModifiedVertices, PlanarNormal, TriQuadPrismInflation());
+				VerticesToUse = &ModifiedVertices;
 				UE_LOG(LogChaos, Verbose, TEXT("Encountered a planar shape in convex hull generation. Will prepare a prism of thickness %.5f in place of a triangle."), TriQuadPrismInflation());
 			}
 
-			const int32 NumParticlesToUse = ParticlesToUse->Size();
+			const int32 NumVerticesToUse = VerticesToUse->Num();
 
-			OutLocalBounds = TAABB<FReal, 3>(ParticlesToUse->X(0), ParticlesToUse->X(0));
-			for(int32 ParticleIndex = 0; ParticleIndex < NumParticlesToUse; ++ParticleIndex)
+			OutLocalBounds = TAABB<FReal, 3>((*VerticesToUse)[0], (*VerticesToUse)[0]);
+			for (int32 VertexIndex = 0; VertexIndex < NumVerticesToUse; ++VertexIndex)
 			{
-				OutLocalBounds.GrowToInclude(ParticlesToUse->X(ParticleIndex));
+				OutLocalBounds.GrowToInclude((*VerticesToUse)[VertexIndex]);
 			}
 
-			if(NumParticlesToUse >= 4)
+			if (NumVerticesToUse >= 4)
 			{
 				TArray<TVector<int32, 3>> Indices;
 				Params BuildParams;
-				BuildParams.HorizonEpsilon = Chaos::FConvexBuilder::SuggestEpsilon(*ParticlesToUse);
-				BuildConvexHull(*ParticlesToUse, Indices, BuildParams);
+				BuildParams.HorizonEpsilon = Chaos::FConvexBuilder::SuggestEpsilon(*VerticesToUse);
+				BuildConvexHull(*VerticesToUse, Indices, BuildParams);
 				OutPlanes.Reserve(Indices.Num());
 				TMap<int32, int32> IndexMap; // maps original particle indices to output particle indices
 				int32 NewIdx = 0;
@@ -330,7 +330,7 @@ namespace Chaos
 
 				for(const TVector<int32, 3>& Idx : Indices)
 				{
-					FVec3 Vs[3] = {ParticlesToUse->X(Idx[0]), ParticlesToUse->X(Idx[1]), ParticlesToUse->X(Idx[2])};
+					FVec3 Vs[3] = { (*VerticesToUse)[Idx[0]], (*VerticesToUse)[Idx[1]], (*VerticesToUse)[Idx[2]] };
 					const FVec3 Normal = FVec3::CrossProduct(Vs[1] - Vs[0], Vs[2] - Vs[0]).GetUnsafeNormal();
 					OutPlanes.Add(TPlaneConcrete<FReal, 3>(Vs[0], Normal));
 					TArray<int32> FaceIndices;
@@ -341,21 +341,21 @@ namespace Chaos
 					OutFaceIndices.Add(FaceIndices);
 				}
 
-				OutSurfaceParticles.AddParticles(IndexMap.Num());
+				OutVertices.SetNum(IndexMap.Num());
 				for(const auto& Elem : IndexMap)
 				{
-					OutSurfaceParticles.X(Elem.Value) = ParticlesToUse->X(Elem.Key);
+					OutVertices[Elem.Value] = (*VerticesToUse)[Elem.Key];
 				}
 			}
 
-			UE_CLOG(OutSurfaceParticles.Size() == 0, LogChaos, Warning, TEXT("Convex hull generation produced zero convex particles, collision will fail for this primitive."));
+			UE_CLOG(OutVertices.Num() == 0, LogChaos, Warning, TEXT("Convex hull generation produced zero convex particles, collision will fail for this primitive."));
 		}
 
-		static void BuildConvexHull(const TParticles<FReal, 3>& InParticles, TArray<TVector<int32, 3>>& OutIndices, const Params& InParams = Params())
+		static void BuildConvexHull(const TArray<FVec3>& InVertices, TArray<TVector<int32, 3>>& OutIndices, const Params& InParams = Params())
 		{
 			OutIndices.Reset();
 			FMemPool Pool;
-			FConvexFace* Faces = BuildInitialHull(Pool, InParticles);
+			FConvexFace* Faces = BuildInitialHull(Pool, InVertices);
 			if(Faces == nullptr)
 			{
 				return;
@@ -377,22 +377,22 @@ namespace Chaos
 			DummyFace->Next = Faces;
 			Faces->Prev = DummyFace;
 
-			FHalfEdge* ConflictV = FindConflictVertex(InParticles, DummyFace->Next);
+			FHalfEdge* ConflictV = FindConflictVertex(InVertices, DummyFace->Next);
 			while(ConflictV)
 			{
 
 #if DEBUG_HULL_GENERATION
 #if DEBUG_HULL_GENERATION_HULLSTEPS_TO_OBJ
 				UE_LOG(LogChaos, VeryVerbose, TEXT("# ======================================================"));
-				const FVec3 ConflictPos = InParticles.X(ConflictV->Vertex);
+				const FVec3 ConflictPos = InVertices[ConflictV->Vertex];
 				UE_LOG(LogChaos, VeryVerbose, TEXT("# GENERATED HULL before adding Vtx %d (%f %f %f)"), ConflictV->Vertex, ConflictPos.X, ConflictPos.Y, ConflictPos.Z);
 				UE_LOG(LogChaos, VeryVerbose, TEXT("# ------------------------------------------------------"));
 				FConvexFace* Face = DummyFace->Next;
 				while (Face)
 				{
-					const FVector P1 = InParticles.X(Face->FirstEdge->Prev->Vertex);
-					const FVector P2 = InParticles.X(Face->FirstEdge->Next->Vertex);
-					const FVector P3 = InParticles.X(Face->FirstEdge->Vertex);
+					const FVector P1 = InVertices[Face->FirstEdge->Prev->Vertex];
+					const FVector P2 = InVertices[Face->FirstEdge->Next->Vertex];
+					const FVector P3 = InVertices[Face->FirstEdge->Vertex];
 					UE_LOG(LogChaos, VeryVerbose, TEXT("v %f %f %f"), P1.X, P1.Y, P1.Z);
 					UE_LOG(LogChaos, VeryVerbose, TEXT("v %f %f %f"), P2.X, P2.Y, P2.Z);
 					UE_LOG(LogChaos, VeryVerbose, TEXT("v %f %f %f"), P3.X, P3.Y, P3.Z);
@@ -402,8 +402,8 @@ namespace Chaos
 #endif
 #endif
 
-				AddVertex(Pool, InParticles, ConflictV, InParams);
-				ConflictV = FindConflictVertex(InParticles, DummyFace->Next);
+				AddVertex(Pool, InVertices, ConflictV, InParams);
+				ConflictV = FindConflictVertex(InVertices, DummyFace->Next);
 			}
 
 			FConvexFace* Cur = DummyFace->Next;
@@ -416,21 +416,21 @@ namespace Chaos
 			}
 		}
 
-		static TTriangleMesh<FReal> BuildConvexHullTriMesh(const TParticles<FReal, 3>& InParticles)
+		static TTriangleMesh<FReal> BuildConvexHullTriMesh(const TArray<FVec3>& InVertices)
 		{
 			TArray<TVector<int32, 3>> Indices;
-			BuildConvexHull(InParticles, Indices);
+			BuildConvexHull(InVertices, Indices);
 			return TTriangleMesh<FReal>(MoveTemp(Indices));
 		}
 
-		static CHAOS_API bool IsPerformanceWarning(int32 NumPlanes, int32 NumParticles)
+		static CHAOS_API bool IsPerformanceWarning(int32 NumPlanes, int32 NumVertices)
 		{
 			if (!PerformGeometryCheck)
 			{
 				return false;
 			}
 
-			return (NumParticles > ParticlesThreshold);
+			return (NumVertices > VerticesThreshold);
 		}
 
 		static CHAOS_API bool IsGeometryReductionEnabled()
@@ -438,12 +438,12 @@ namespace Chaos
 			return (PerformGeometryReduction>0)?true:false;
 		}
 
-		static FString PerformanceWarningString(int32 NumPlanes, int32 NumParticles)
+		static FString PerformanceWarningString(int32 NumPlanes, int32 NumVertices)
 		{
-			return FString::Printf(TEXT("Planes %d, SurfaceParticles %d"), NumPlanes, NumParticles);
+			return FString::Printf(TEXT("Planes %d, Vertices %d"), NumPlanes, NumVertices);
 		}
 
-		static CHAOS_API void Simplify(TArray <TPlaneConcrete<FReal, 3>>& InOutPlanes, TArray<TArray<int32>>& InOutFaces, TParticles<FReal, 3>& InOutParticles, TAABB<FReal, 3>& InOutLocalBounds)
+		static CHAOS_API void Simplify(TArray <TPlaneConcrete<FReal, 3>>& InOutPlanes, TArray<TArray<int32>>& InOutFaces, TArray<FVec3>& InOutVertices, TAABB<FReal, 3>& InOutLocalBounds)
 		{
 			struct TPair
 			{
@@ -452,17 +452,12 @@ namespace Chaos
 				uint32 B;
 			};
 
-			uint32 NumberOfParticlesRequired = ParticlesThreshold;
-			uint32 NumberOfParticlesWeHave = InOutParticles.Size();
-			int32 NumToDelete = NumberOfParticlesWeHave - NumberOfParticlesRequired;
+			uint32 NumberOfVerticesRequired = VerticesThreshold;
+			uint32 NumberOfVerticesWeHave = InOutVertices.Num();
+			int32 NumToDelete = NumberOfVerticesWeHave - NumberOfVerticesRequired;
 
-			uint32 Size = InOutParticles.Size();
-			TArray<FVec3> Particles;
-			Particles.AddUninitialized(Size);
-			for (uint32 A = 0; A < Size; A++)
-			{
-				Particles[A] = InOutParticles.X(A);
-			}
+			int32 Size = InOutVertices.Num();
+			TArray<FVec3> Vertices(InOutVertices);
 
 			TArray<bool> IsDeleted;
 			IsDeleted.Reset();
@@ -475,15 +470,15 @@ namespace Chaos
 					TPair ClosestPair;
 					float ClosestDistSqr = FLT_MAX;
 
-					for (uint32 A = 0; A < (Size - 1); A++)
+					for (int32 A = 0; A < (Size - 1); A++)
 					{
 						if (!IsDeleted[A])
 						{
-							for (uint32 B = A + 1; B < Size; B++)
+							for (int32 B = A + 1; B < Size; B++)
 							{
 								if (!IsDeleted[B])
 								{
-									FVec3 Vec = Particles[A] - Particles[B];
+									FVec3 Vec = Vertices[A] - Vertices[B];
 									float LengthSqr = Vec.SizeSquared();
 									if (LengthSqr < ClosestDistSqr)
 									{
@@ -499,29 +494,28 @@ namespace Chaos
 					if (ClosestPair.A != -1)
 					{
 						// merge to mid point
-						Particles[ClosestPair.A] = Particles[ClosestPair.A] + (Particles[ClosestPair.B] - Particles[ClosestPair.A]) * 0.5f;
+						Vertices[ClosestPair.A] = Vertices[ClosestPair.A] + (Vertices[ClosestPair.B] - Vertices[ClosestPair.A]) * 0.5f;
 						IsDeleted[ClosestPair.B] = true;
 					}
 				}
 			}
 
-			TParticles<FReal, 3> TmpParticles;
-			for (int Idx = 0; Idx < Particles.Num(); Idx++)
+			TArray<FVec3> TmpVertices;
+			for (int Idx = 0; Idx < Vertices.Num(); Idx++)
 			{
-				// Only add particles that have not been merged away
+				// Only add vertices that have not been merged away
 				if (!IsDeleted[Idx])
 				{
-					TmpParticles.AddParticles(1);
-					TmpParticles.X(TmpParticles.Size() - 1) = Particles[Idx];
+					TmpVertices.Add(Vertices[Idx]);
 				}
 			}
 
-			Build(TmpParticles, InOutPlanes, InOutFaces, InOutParticles, InOutLocalBounds);
-			check(InOutParticles.Size() > 3);
+			Build(TmpVertices, InOutPlanes, InOutFaces, InOutVertices, InOutLocalBounds);
+			check(InOutVertices.Num() > 3);
 		}
 
 		// Convert multi-triangle faces to single n-gons
-		static void MergeFaces(TArray<TPlaneConcrete<FReal, 3>>& InOutPlanes, TArray<TArray<int32>>& InOutFaceVertexIndices, const TParticles<FReal, 3>& SurfaceParticles)
+		static void MergeFaces(TArray<TPlaneConcrete<FReal, 3>>& InOutPlanes, TArray<TArray<int32>>& InOutFaceVertexIndices, const TArray<FVec3>& Vertices)
 		{
 			const FReal NormalThreshold = 1.e-4f;
 
@@ -556,26 +550,26 @@ namespace Chaos
 			// Re-order the face vertices to form the face half-edges
 			for (int32 PlaneIndex0 = 0; PlaneIndex0 < InOutPlanes.Num(); ++PlaneIndex0)
 			{
-				SortFaceVerticesCCW(InOutPlanes[PlaneIndex0], InOutFaceVertexIndices[PlaneIndex0], SurfaceParticles);
+				SortFaceVerticesCCW(InOutPlanes[PlaneIndex0], InOutFaceVertexIndices[PlaneIndex0], Vertices);
 			}
 		}
 
 		// Reorder the vertices to be counter-clockwise about the normal
-		static void SortFaceVerticesCCW(const TPlaneConcrete<FReal, 3>& Face, TArray<int32>& InOutFaceVertexIndices, const TParticles<FReal, 3>& SurfaceParticles)
+		static void SortFaceVerticesCCW(const TPlaneConcrete<FReal, 3>& Face, TArray<int32>& InOutFaceVertexIndices, const TArray<FVec3>& Vertices)
 		{
 			FMatrix33 FaceMatrix = FRotationMatrix::MakeFromZ(Face.Normal());
 
 			FVec3 Centroid = FVec3(0);
 			for (int32 VertexIndex = 0; VertexIndex < InOutFaceVertexIndices.Num(); ++VertexIndex)
 			{
-				Centroid += SurfaceParticles.X(InOutFaceVertexIndices[VertexIndex]);
+				Centroid += Vertices[InOutFaceVertexIndices[VertexIndex]];
 			}
 			Centroid /= FReal(InOutFaceVertexIndices.Num());
 
 			// [2, -2] based on clockwise angle about the normal
-			auto VertexScore = [&Centroid, &FaceMatrix, &SurfaceParticles](int32 VertexIndex)
+			auto VertexScore = [&Centroid, &FaceMatrix, &Vertices](int32 VertexIndex)
 			{
-				const FVec3 CentroidOffsetDir = (SurfaceParticles.X(VertexIndex) - Centroid).GetSafeNormal();
+				const FVec3 CentroidOffsetDir = (Vertices[VertexIndex] - Centroid).GetSafeNormal();
 				const FReal DotX = FVec3::DotProduct(CentroidOffsetDir, FaceMatrix.GetAxis(0));
 				const FReal DotY = FVec3::DotProduct(CentroidOffsetDir, FaceMatrix.GetAxis(1));
 				const FReal Score = (DotX >= 0.0f) ? 1.0f + DotY : -1.0f - DotY;
@@ -591,21 +585,21 @@ namespace Chaos
 		}
 
 		// Generate the vertex indices for all planes in CCW order (used to serialize old data that did not have structure data)
-		static void BuildPlaneVertexIndices(const TArray<TPlaneConcrete<FReal, 3>>& InPlanes, const TParticles<FReal, 3>& SurfaceParticles, TArray<TArray<int32>>& OutFaceVertexIndices, const FReal DistanceTolerance = 1.e-3f)
+		static void BuildPlaneVertexIndices(const TArray<TPlaneConcrete<FReal, 3>>& InPlanes, const TArray<FVec3>& Vertices, TArray<TArray<int32>>& OutFaceVertexIndices, const FReal DistanceTolerance = 1.e-3f)
 		{
 			OutFaceVertexIndices.Reset(InPlanes.Num());
 			for (int32 PlaneIndex = 0; PlaneIndex < InPlanes.Num(); ++PlaneIndex)
 			{
-				for (int32 VertexIndex = 0; VertexIndex < (int32)SurfaceParticles.Size(); ++VertexIndex)
+				for (int32 VertexIndex = 0; VertexIndex < Vertices.Num(); ++VertexIndex)
 				{
-					const FReal PlaneVertexDistance = FVec3::DotProduct(InPlanes[PlaneIndex].Normal(), SurfaceParticles.X(VertexIndex) - InPlanes[PlaneIndex].X());
+					const FReal PlaneVertexDistance = FVec3::DotProduct(InPlanes[PlaneIndex].Normal(), Vertices[VertexIndex] - InPlanes[PlaneIndex].X());
 					if (FMath::Abs(PlaneVertexDistance) < DistanceTolerance)
 					{
 						OutFaceVertexIndices[PlaneIndex].Add(VertexIndex);
 					}
 				}
 
-				SortFaceVerticesCCW(InPlanes[PlaneIndex], OutFaceVertexIndices[PlaneIndex], SurfaceParticles);
+				SortFaceVerticesCCW(InPlanes[PlaneIndex], OutFaceVertexIndices[PlaneIndex], Vertices);
 			}
 		}
 
@@ -613,11 +607,11 @@ namespace Chaos
 #if PLATFORM_MAC || PLATFORM_LINUX
 		static CHAOS_API int32 PerformGeometryCheck;
 		static CHAOS_API int32 PerformGeometryReduction;
-		static CHAOS_API int32 ParticlesThreshold;
+		static CHAOS_API int32 VerticesThreshold;
 #else
 		static int32 PerformGeometryCheck;
 		static int32 PerformGeometryReduction;
-		static int32 ParticlesThreshold;
+		static int32 VerticesThreshold;
 #endif
 
 	private:
@@ -627,7 +621,7 @@ namespace Chaos
 			return FVec3::CrossProduct((B - A), (C - A));
 		}
 
-		static FConvexFace* CreateFace(FMemPool& Pool, const TParticles<FReal, 3>& InParticles, FHalfEdge* RS, FHalfEdge* ST, FHalfEdge* TR)
+		static FConvexFace* CreateFace(FMemPool& Pool, const TArray<FVec3>& InVertices, FHalfEdge* RS, FHalfEdge* ST, FHalfEdge* TR)
 		{
 			RS->Prev = TR;
 			RS->Next = ST;
@@ -635,7 +629,7 @@ namespace Chaos
 			ST->Next = TR;
 			TR->Prev = ST;
 			TR->Next = RS;
-			FVec3 RSTNormal = ComputeFaceNormal(InParticles.X(RS->Vertex), InParticles.X(ST->Vertex), InParticles.X(TR->Vertex));
+			FVec3 RSTNormal = ComputeFaceNormal(InVertices[RS->Vertex], InVertices[ST->Vertex], InVertices[TR->Vertex]);
 			const FReal RSTNormalSize = RSTNormal.Size();
 			if (RSTNormalSize > 0)
 			{
@@ -646,7 +640,7 @@ namespace Chaos
 				// degenerated face, use a valid neighbor face normal 
 				RSTNormal = RS->Twin->Face->Plane.Normal();
 			}
-			FConvexFace* RST = Pool.AllocConvexFace(TPlaneConcrete<FReal, 3>(InParticles.X(RS->Vertex), RSTNormal));
+			FConvexFace* RST = Pool.AllocConvexFace(TPlaneConcrete<FReal, 3>(InVertices[RS->Vertex], RSTNormal));
 			RST->FirstEdge = RS;
 			RS->Face = RST;
 			ST->Face = RST;
@@ -654,7 +648,7 @@ namespace Chaos
 			return RST;
 		}
 
-		static void StealConflictList(FMemPool& Pool, const TParticles<FReal, 3>& InParticles, FHalfEdge* OldList, FConvexFace** Faces, int32 NumFaces)
+		static void StealConflictList(FMemPool& Pool, const TArray<FVec3>& InVertices, FHalfEdge* OldList, FConvexFace** Faces, int32 NumFaces)
 		{
 			FHalfEdge* Cur = OldList;
 			while(Cur)
@@ -663,7 +657,7 @@ namespace Chaos
 				int32 MaxIdx = -1;
 				for(int32 Idx = 0; Idx < NumFaces; ++Idx)
 				{
-					FReal Distance = Faces[Idx]->Plane.SignedDistance(InParticles.X(Cur->Vertex));
+					FReal Distance = Faces[Idx]->Plane.SignedDistance(InVertices[Cur->Vertex]);
 					if(Distance > MaxD)
 					{
 						MaxD = Distance;
@@ -675,8 +669,8 @@ namespace Chaos
 				if(!bDeleteVertex)
 				{
 					//let's make sure faces created with this new conflict vertex will be valid. The plane check above is not sufficient because long thin triangles will have a plane with its point at one of these. Combined with normal and precision we can have errors
-					auto PretendNormal = [&InParticles](FHalfEdge* A, FHalfEdge* B, FHalfEdge* C) {
-						return FVec3::CrossProduct(InParticles.X(B->Vertex) - InParticles.X(A->Vertex), InParticles.X(C->Vertex) - InParticles.X(A->Vertex)).SizeSquared();
+					auto PretendNormal = [&InVertices](FHalfEdge* A, FHalfEdge* B, FHalfEdge* C) {
+						return FVec3::CrossProduct(InVertices[B->Vertex] - InVertices[A->Vertex], InVertices[C->Vertex] - InVertices[A->Vertex]).SizeSquared();
 					};
 					FHalfEdge* Edge = Faces[MaxIdx]->FirstEdge;
 					do
@@ -713,16 +707,16 @@ namespace Chaos
 			}
 		}
 
-		static FConvexFace* BuildInitialHull(FMemPool& Pool, const TParticles<FReal, 3>& InParticles)
+		static FConvexFace* BuildInitialHull(FMemPool& Pool, const TArray<FVec3>& InVertices)
 		{
-			if(InParticles.Size() < 4) //not enough points
+			if (InVertices.Num() < 4) //not enough points
 			{
 				return nullptr;
 			}
 
 			constexpr FReal Epsilon = 1e-4;
 
-			const int32 NumParticles = InParticles.Size();
+			const int32 NumVertices = InVertices.Num();
 
 			//We store the vertex directly in the half-edge. We use its next to group free vertices by context list
 			//create a starting triangle by finding min/max on X and max on Y
@@ -735,13 +729,13 @@ namespace Chaos
 			DummyHalfEdge->Next = nullptr;
 			FHalfEdge* Prev = DummyHalfEdge;
 
-			for(int32 i = 0; i < NumParticles; ++i)
+			for (int32 i = 0; i < NumVertices; ++i)
 			{
 				FHalfEdge* VHalf = Pool.AllocHalfEdge(i); //todo(ocohen): preallocate these
 				Prev->Next = VHalf;
 				VHalf->Prev = Prev;
 				VHalf->Next = nullptr;
-				const FVec3& V = InParticles.X(i);
+				const FVec3& V = InVertices[i];
 
 				if(V[0] < MinX)
 				{
@@ -758,7 +752,7 @@ namespace Chaos
 			}
 
 			check(A && B);
-			if(A == B || (InParticles.X(A->Vertex) - InParticles.X(B->Vertex)).SizeSquared() < Epsilon) //infinitely thin
+			if (A == B || (InVertices[A->Vertex] - InVertices[B->Vertex]).SizeSquared() < Epsilon) //infinitely thin
 			{
 				return nullptr;
 			}
@@ -777,11 +771,11 @@ namespace Chaos
 
 			//find C so that we get the biggest base triangle
 			FReal MaxTriSize = Epsilon;
-			const FVec3 AToB = InParticles.X(B->Vertex) - InParticles.X(A->Vertex);
+			const FVec3 AToB = InVertices[B->Vertex] - InVertices[A->Vertex];
 			FHalfEdge* C = nullptr;
 			for(FHalfEdge* V = DummyHalfEdge->Next; V; V = V->Next)
 			{
-				FReal TriSize = FVec3::CrossProduct(AToB, InParticles.X(V->Vertex) - InParticles.X(A->Vertex)).SizeSquared();
+				FReal TriSize = FVec3::CrossProduct(AToB, InVertices[V->Vertex] - InVertices[A->Vertex]).SizeSquared();
 				if(TriSize > MaxTriSize)
 				{
 					MaxTriSize = TriSize;
@@ -802,7 +796,7 @@ namespace Chaos
 			}
 
 			//find farthest D along normal
-			const FVec3 AToC = InParticles.X(C->Vertex) - InParticles.X(A->Vertex);
+			const FVec3 AToC = InVertices[C->Vertex] - InVertices[A->Vertex];
 			const FVec3 Normal = FVec3::CrossProduct(AToB, AToC);
 
 			FReal MaxPosDistance = Epsilon;
@@ -811,7 +805,7 @@ namespace Chaos
 			FHalfEdge* NegD = nullptr;
 			for(FHalfEdge* V = DummyHalfEdge->Next; V; V = V->Next)
 			{
-				FReal Dot = FVec3::DotProduct(InParticles.X(V->Vertex) - InParticles.X(A->Vertex), Normal);
+				FReal Dot = FVec3::DotProduct(InVertices[V->Vertex] - InVertices[A->Vertex], Normal);
 				if(Dot > MaxPosDistance)
 				{
 					MaxPosDistance = Dot;
@@ -853,10 +847,10 @@ namespace Chaos
 			}
 
 			FConvexFace* Faces[4];
-			Faces[0] = CreateFace(Pool, InParticles, Edges[0], Edges[1], Edges[2]); //base
-			Faces[1] = CreateFace(Pool, InParticles, Pool.AllocHalfEdge(Edges[1]->Vertex), Pool.AllocHalfEdge(Edges[0]->Vertex), Edges[3]);
-			Faces[2] = CreateFace(Pool, InParticles, Pool.AllocHalfEdge(Edges[0]->Vertex), Pool.AllocHalfEdge(Edges[2]->Vertex), Pool.AllocHalfEdge(Edges[3]->Vertex));
-			Faces[3] = CreateFace(Pool, InParticles, Pool.AllocHalfEdge(Edges[2]->Vertex), Pool.AllocHalfEdge(Edges[1]->Vertex), Pool.AllocHalfEdge(Edges[3]->Vertex));
+			Faces[0] = CreateFace(Pool, InVertices, Edges[0], Edges[1], Edges[2]); //base
+			Faces[1] = CreateFace(Pool, InVertices, Pool.AllocHalfEdge(Edges[1]->Vertex), Pool.AllocHalfEdge(Edges[0]->Vertex), Edges[3]);
+			Faces[2] = CreateFace(Pool, InVertices, Pool.AllocHalfEdge(Edges[0]->Vertex), Pool.AllocHalfEdge(Edges[2]->Vertex), Pool.AllocHalfEdge(Edges[3]->Vertex));
+			Faces[3] = CreateFace(Pool, InVertices, Pool.AllocHalfEdge(Edges[2]->Vertex), Pool.AllocHalfEdge(Edges[1]->Vertex), Pool.AllocHalfEdge(Edges[3]->Vertex));
 
 			auto MakeTwins = [](FHalfEdge* E1, FHalfEdge* E2) {
 				E1->Twin = E2;
@@ -879,11 +873,11 @@ namespace Chaos
 			Faces[3]->Next = nullptr;
 
 			//split up the conflict list
-			StealConflictList(Pool, InParticles, DummyHalfEdge->Next, Faces, 4);
+			StealConflictList(Pool, InVertices, DummyHalfEdge->Next, Faces, 4);
 			return Faces[0];
 		}
 
-		static FHalfEdge* FindConflictVertex(const TParticles<FReal, 3>& InParticles, FConvexFace* FaceList)
+		static FHalfEdge* FindConflictVertex(const TArray<FVec3>& InVertices, FConvexFace* FaceList)
 		{
 			UE_CLOG(DEBUG_HULL_GENERATION, LogChaos, VeryVerbose, TEXT("Finding conflict vertex"));
 
@@ -896,7 +890,7 @@ namespace Chaos
 				for(FHalfEdge* CurFaceVertex = CurFace->ConflictList; CurFaceVertex; CurFaceVertex = CurFaceVertex->Next)
 				{
 					//is it faster to cache this from stealing stage?
-					FReal Dist = FVec3::DotProduct(InParticles.X(CurFaceVertex->Vertex), CurFace->Plane.Normal());
+					FReal Dist = FVec3::DotProduct(InVertices[CurFaceVertex->Vertex], CurFace->Plane.Normal());
 					if(Dist > MaxD)
 					{
 						MaxD = Dist;
@@ -930,13 +924,13 @@ namespace Chaos
 			return nullptr;
 		}
 
-		static void BuildHorizon(const TParticles<FReal, 3>& InParticles, FHalfEdge* ConflictV, TArray<FHalfEdge*>& HorizonEdges, TArray<FConvexFace*>& FacesToDelete, const Params& InParams)
+		static void BuildHorizon(const TArray<FVec3>& InVertices, FHalfEdge* ConflictV, TArray<FHalfEdge*>& HorizonEdges, TArray<FConvexFace*>& FacesToDelete, const Params& InParams)
 		{
 			//We must flood fill from the initial face and mark edges of faces the conflict vertex cannot see
 			//In order to return a CCW ordering we must traverse each face in CCW order from the edge we crossed over
 			//This should already be the ordering in the half edge
 			const FReal Epsilon = InParams.HorizonEpsilon;
-			const FVec3 V = InParticles.X(ConflictV->Vertex);
+			const FVec3 V = InVertices[ConflictV->Vertex];
 			TSet<FConvexFace*> Processed;
 			TArray<FHalfEdge*> Queue;
 			check(ConflictV->Face);
@@ -970,15 +964,15 @@ namespace Chaos
 #if DEBUG_HULL_GENERATION
 #if DEBUG_HULL_GENERATION_BUILDHORIZON_TO_OBJ
 			UE_LOG(LogChaos, VeryVerbose, TEXT("# ======================================================"));
-			const FVec3 ConflictPos = InParticles.X(ConflictV->Vertex);
+			const FVec3 ConflictPos = InVertices[ConflictV->Vertex];
 			UE_LOG(LogChaos, VeryVerbose, TEXT("# BUILD_HORIZON - Conflict Vertex = %d (%f %f %f)"), ConflictV->Vertex, ConflictPos.X, ConflictPos.Y, ConflictPos.Z);
 			UE_LOG(LogChaos, VeryVerbose, TEXT("# ------------------------------------------------------"));
 			for (TSet<FConvexFace*>::TConstIterator SetIt(Processed); SetIt; ++SetIt)
 			{
 				const FConvexFace* Face = *SetIt;
-				const FVector P1 = InParticles.X(Face->FirstEdge->Prev->Vertex);
-				const FVector P2 = InParticles.X(Face->FirstEdge->Next->Vertex);
-				const FVector P3 = InParticles.X(Face->FirstEdge->Vertex);
+				const FVector P1 = InVertices[Face->FirstEdge->Prev->Vertex];
+				const FVector P2 = InVertices[Face->FirstEdge->Next->Vertex];
+				const FVector P3 = InVertices[Face->FirstEdge->Vertex];
 				UE_LOG(LogChaos, VeryVerbose, TEXT("v %f %f %f"), P1.X, P1.Y, P1.Z);
 				UE_LOG(LogChaos, VeryVerbose, TEXT("v %f %f %f"), P2.X, P2.Y, P2.Z);
 				UE_LOG(LogChaos, VeryVerbose, TEXT("v %f %f %f"), P3.X, P3.Y, P3.Z);
@@ -988,7 +982,7 @@ namespace Chaos
 #endif
 		}
 
-		static void BuildFaces(FMemPool& Pool, const TParticles<FReal, 3>& InParticles, const FHalfEdge* ConflictV, const TArray<FHalfEdge*>& HorizonEdges, const TArray<FConvexFace*> OldFaces, TArray<FConvexFace*>& NewFaces)
+		static void BuildFaces(FMemPool& Pool, const TArray<FVec3>& InVertices, const FHalfEdge* ConflictV, const TArray<FHalfEdge*>& HorizonEdges, const TArray<FConvexFace*> OldFaces, TArray<FConvexFace*>& NewFaces)
 		{
 			//The HorizonEdges are in CCW order. We must make new faces and edges to join from ConflictV to these edges
 			check(HorizonEdges.Num() >= 3);
@@ -1011,7 +1005,7 @@ namespace Chaos
 				PrevEdge = HorizonNext;
 
 				//link new faces together
-				FConvexFace* NewFace = CreateFace(Pool, InParticles, NewHorizonEdge, HorizonNext, V);
+				FConvexFace* NewFace = CreateFace(Pool, InVertices, NewHorizonEdge, HorizonNext, V);
 				if(NewFaces.Num() > 0)
 				{
 					NewFace->Prev = NewFaces[NewFaces.Num() - 1];
@@ -1032,7 +1026,7 @@ namespace Chaos
 			//redistribute conflict list
 			for(FConvexFace* OldFace : OldFaces)
 			{
-				StealConflictList(Pool, InParticles, OldFace->ConflictList, &NewFaces[0], NewFaces.Num());
+				StealConflictList(Pool, InVertices, OldFace->ConflictList, &NewFaces[0], NewFaces.Num());
 			}
 
 			//insert all new faces after conflict vertex face
@@ -1048,13 +1042,13 @@ namespace Chaos
 			StartFace->Prev = OldFace;
 		}
 
-		static void AddVertex(FMemPool& Pool, const TParticles<FReal, 3>& InParticles, FHalfEdge* ConflictV, const Params& InParams)
+		static void AddVertex(FMemPool& Pool, const TArray<FVec3>& InVertices, FHalfEdge* ConflictV, const Params& InParams)
 		{
 			UE_CLOG(DEBUG_HULL_GENERATION, LogChaos, VeryVerbose, TEXT("Adding Vertex %d"), ConflictV->Vertex);
 
 			TArray<FHalfEdge*> HorizonEdges;
 			TArray<FConvexFace*> FacesToDelete;
-			BuildHorizon(InParticles, ConflictV, HorizonEdges, FacesToDelete, InParams);
+			BuildHorizon(InVertices, ConflictV, HorizonEdges, FacesToDelete, InParams);
 
 #if DEBUG_HULL_GENERATION
 			FString HorizonString(TEXT("\tHorizon: ("));
@@ -1067,7 +1061,7 @@ namespace Chaos
 #endif
 
 			TArray<FConvexFace*> NewFaces;
-			BuildFaces(Pool, InParticles, ConflictV, HorizonEdges, FacesToDelete, NewFaces);
+			BuildFaces(Pool, InVertices, ConflictV, HorizonEdges, FacesToDelete, NewFaces);
 
 #if DEBUG_HULL_GENERATION
 			FString NewFaceString(TEXT("\tNew Faces: "));
diff --git a/Engine/Source/Runtime/Experimental/Chaos/Public/Chaos/Convex.h b/Engine/Source/Runtime/Experimental/Chaos/Public/Chaos/Convex.h
--- a/Engine/Source/Runtime/Experimental/Chaos/Public/Chaos/Convex.h
+++ b/Engine/Source/Runtime/Experimental/Chaos/Public/Chaos/Convex.h
@@ -93,7 +93,7 @@ namespace Chaos
 		FConvex(FConvex&& Other)
 		    : FImplicitObject(EImplicitObject::IsConvex | EImplicitObject::HasBoundingBox, ImplicitObjectType::Convex)
 			, Planes(MoveTemp(Other.Planes))
-		    , SurfaceParticles(MoveTemp(Other.SurfaceParticles))
+		    , Vertices(MoveTemp(Other.Vertices))
 		    , LocalBoundingBox(MoveTemp(Other.LocalBoundingBox))
 			, StructureData(MoveTemp(Other.StructureData))
 			, Volume(MoveTemp(Other.Volume))
@@ -106,15 +106,15 @@ namespace Chaos
 		// @todo(chaos): Keep track of invalid state and ensure on volume or COM access?
 		// @todo(chaos): Add plane vertex indices in the constructor and call CreateStructureData
 		// @todo(chaos): Merge planes? Or assume the input is a good convex hull?
-		FConvex(TArray<TPlaneConcrete<FReal, 3>>&& InPlanes, TParticles<FReal, 3>&& InSurfaceParticles)
+		FConvex(TArray<TPlaneConcrete<FReal, 3>>&& InPlanes, TArray<FVec3>&& InVertices)
 		    : FImplicitObject(EImplicitObject::IsConvex | EImplicitObject::HasBoundingBox, ImplicitObjectType::Convex)
 			, Planes(MoveTemp(InPlanes))
-		    , SurfaceParticles(MoveTemp(InSurfaceParticles))
+		    , Vertices(MoveTemp(InVertices))
 		    , LocalBoundingBox(TAABB<FReal, 3>::EmptyAABB())
 		{
-			for (uint32 ParticleIndex = 0; ParticleIndex < SurfaceParticles.Size(); ++ParticleIndex)
+			for (int32 ParticleIndex = 0; ParticleIndex < Vertices.Num(); ++ParticleIndex)
 			{
-				LocalBoundingBox.GrowToInclude(SurfaceParticles.X(ParticleIndex));
+				LocalBoundingBox.GrowToInclude(Vertices[ParticleIndex]);
 			}
 
 			// For now we approximate COM and volume with the bounding box
@@ -122,29 +122,30 @@ namespace Chaos
 			Volume = LocalBoundingBox.GetVolume();
 		}
 
-		FConvex(const TParticles<FReal, 3>& InParticles, const FReal InMargin)
+		FConvex(const TArray<FVec3>& InVertices, const FReal InMargin)
 		    : FImplicitObject(EImplicitObject::IsConvex | EImplicitObject::HasBoundingBox, ImplicitObjectType::Convex)
 		{
-			const uint32 NumParticles = InParticles.Size();
-			if (NumParticles == 0)
+			const int32 NumVertices = InVertices.Num();
+			if (NumVertices == 0)
 			{
 				return;
 			}
 
 			TArray<TArray<int32>> FaceIndices;
-			FConvexBuilder::Build(InParticles, Planes, FaceIndices, SurfaceParticles, LocalBoundingBox);
+			FConvexBuilder::Build(InVertices, Planes, FaceIndices, Vertices, LocalBoundingBox);
 			CHAOS_ENSURE(Planes.Num() == FaceIndices.Num());
 
 			// @chaos(todo): this only works with triangles. Fix that an we can run MergeFaces before calling this
-			CalculateVolumeAndCenterOfMass(SurfaceParticles, FaceIndices, Volume, CenterOfMass);
+			TParticles<FReal, 3> VertexParticles(CopyTemp(Vertices));
+			CalculateVolumeAndCenterOfMass(VertexParticles, FaceIndices, Volume, CenterOfMass);
 
-			FConvexBuilder::MergeFaces(Planes, FaceIndices, SurfaceParticles);
+			FConvexBuilder::MergeFaces(Planes, FaceIndices, Vertices);
 			CHAOS_ENSURE(Planes.Num() == FaceIndices.Num());
 
 			CreateStructureData(MoveTemp(FaceIndices));
 
 			SetMargin(InMargin);
-		}
+		}	
 
 	private:
 		void MovePlanesAndRebuild(FReal InDelta);
@@ -295,7 +296,7 @@ namespace Chaos
 
 		int32 NumVertices() const
 		{
-			return (int32)SurfaceParticles.Size();
+			return (int32)Vertices.Num();
 		}
 
 		// Get the plane at the specified index (e.g., indices from GetVertexPlanes)
@@ -307,7 +308,7 @@ namespace Chaos
 		// Get the vertex at the specified index (e.g., indices from GetPlaneVertices)
 		const FVec3& GetVertex(int32 VertexIndex) const
 		{
-			return SurfaceParticles.X(VertexIndex);
+			return Vertices[VertexIndex];
 		}
 
 
@@ -339,13 +340,13 @@ namespace Chaos
 		{
 			FReal MaxDot = TNumericLimits<FReal>::Lowest();
 			int32 MaxVIdx = INDEX_NONE;
-			const int32 NumVertices = SurfaceParticles.Size();
+			const int32 NumVertices = Vertices.Num();
 
 			if (ensure(NumVertices > 0))
 			{
 				for (int32 Idx = 0; Idx < NumVertices; ++Idx)
 				{
-					const FReal Dot = FVec3::DotProduct(SurfaceParticles.X(Idx), Direction);
+					const FReal Dot = FVec3::DotProduct(Vertices[Idx], Direction);
 					if (Dot > MaxDot)
 					{
 						MaxDot = Dot;
@@ -437,7 +438,7 @@ namespace Chaos
 				{
 					return GetMarginAdjustedVertex(SupportVertexIndex, NetMargin);
 				}
-				return SurfaceParticles.X(SupportVertexIndex);
+				return Vertices[SupportVertexIndex];
 			}
 			return FVec3(0);
 		}
@@ -467,7 +468,7 @@ namespace Chaos
 				}
 				else
 				{
-					VertexPosition = Scale * SurfaceParticles.X(SupportVertexIndex);
+					VertexPosition = Scale * Vertices[SupportVertexIndex];
 				}
 			}
 			return VertexPosition;
@@ -481,9 +482,9 @@ namespace Chaos
 			{
 				if (Thickness != 0.0f)
 				{
-					return SurfaceParticles.X(MaxVIdx) + Direction.GetUnsafeNormal() * Thickness;
+					return Vertices[MaxVIdx] + Direction.GetUnsafeNormal() * Thickness;
 				}
-				return SurfaceParticles.X(MaxVIdx);
+				return Vertices[MaxVIdx];
 			}
 			return FVec3(0);
 		}
@@ -493,9 +494,9 @@ namespace Chaos
 			return FString::Printf(TEXT("Convex"));
 		}
 
-		const TParticles<FReal, 3>& GetSurfaceParticles() const
+		const TArray<FVec3>& GetVertices() const
 		{
-			return SurfaceParticles;
+			return Vertices;
 		}
 
 		const TArray<TPlaneConcrete<FReal, 3>>& GetFaces() const
@@ -523,7 +524,12 @@ namespace Chaos
 		{
 			uint32 Result = LocalBoundingBox.GetTypeHash();
 
-			Result = HashCombine(Result, SurfaceParticles.GetTypeHash());
+			for (const FVec3& Vertex: Vertices)
+			{
+				Result = HashCombine(Result, ::GetTypeHash(Vertex[0]));
+				Result = HashCombine(Result, ::GetTypeHash(Vertex[1]));
+				Result = HashCombine(Result, ::GetTypeHash(Vertex[2]));
+			}
 
 			for(const TPlaneConcrete<FReal, 3>& Plane : Planes)
 			{
@@ -536,6 +542,7 @@ namespace Chaos
 		FORCEINLINE void SerializeImp(FArchive& Ar)
 		{
 			Ar.UsingCustomVersion(FExternalPhysicsCustomObjectVersion::GUID);
+			Ar.UsingCustomVersion(FUE5MainStreamObjectVersion::GUID);
 			FImplicitObject::SerializeImp(Ar);
 
 			if (Ar.CustomVer(FExternalPhysicsCustomObjectVersion::GUID) < FExternalPhysicsCustomObjectVersion::ConvexUsesTPlaneConcrete)
@@ -554,7 +561,23 @@ namespace Chaos
 				Ar << Planes;
 			}
 
-			Ar << SurfaceParticles;
+			if (Ar.CustomVer(FUE5MainStreamObjectVersion::GUID) < FUE5MainStreamObjectVersion::ConvexUsesVerticesArray)
+			{
+				TParticles<FReal, 3> TmpSurfaceParticles;
+				Ar << TmpSurfaceParticles;
+
+				const int32 NumVertices = (int32)TmpSurfaceParticles.Size();
+				Vertices.SetNum(NumVertices);
+				for (int32 VertexIndex = 0; VertexIndex < NumVertices; ++VertexIndex)
+				{
+					Vertices[VertexIndex] = TmpSurfaceParticles.X(VertexIndex);
+				}
+			}
+			else
+			{
+				Ar << Vertices;
+			}
+			
 			TBox<FReal,3>::SerializeAsAABB(Ar, LocalBoundingBox);
 
 			if (Ar.CustomVer(FExternalPhysicsCustomObjectVersion::GUID) >= FExternalPhysicsCustomObjectVersion::AddConvexCenterOfMassAndVolume)
@@ -567,8 +590,12 @@ namespace Chaos
 				// Rebuild convex in order to extract face indices.
 				// TODO: Make it so it can take SurfaceParticles as both input and output without breaking...
 				TArray<TArray<int32>> FaceIndices;
-				TParticles<FReal, 3> TempSurfaceParticles;
-				FConvexBuilder::Build(SurfaceParticles, Planes, FaceIndices, TempSurfaceParticles, LocalBoundingBox);
+				TArray<FVec3> TempVertices;
+				FConvexBuilder::Build(Vertices, Planes, FaceIndices, TempVertices, LocalBoundingBox);
+
+				// Copy vertices and move into particles.
+				TArray<FVec3> VerticesCopy = Vertices;
+				const TParticles<FReal, 3> SurfaceParticles(MoveTemp(VerticesCopy));
 				CalculateVolumeAndCenterOfMass(SurfaceParticles, FaceIndices, Volume, CenterOfMass);
 			}
 
@@ -586,7 +613,7 @@ namespace Chaos
 			{
 				// Generate the structure data from the planes and vertices
 				TArray<TArray<int32>> FaceIndices;
-				FConvexBuilder::BuildPlaneVertexIndices(Planes, SurfaceParticles, FaceIndices);
+				FConvexBuilder::BuildPlaneVertexIndices(Planes, Vertices, FaceIndices);
 				CreateStructureData(MoveTemp(FaceIndices));
 			}
 		}
@@ -604,18 +631,18 @@ namespace Chaos
 
 		virtual bool IsValidGeometry() const override
 		{
-			return (SurfaceParticles.Size() > 0 && Planes.Num() > 0);
+			return (Vertices.Num() > 0 && Planes.Num() > 0);
 		}
 
 		virtual bool IsPerformanceWarning() const override
 		{
-			return FConvexBuilder::IsPerformanceWarning(Planes.Num(), SurfaceParticles.Size());
+			return FConvexBuilder::IsPerformanceWarning(Planes.Num(), Vertices.Num());
 		}
 
 		virtual FString PerformanceWarningAndSimplifaction() override
 		{
 
-			FString PerformanceWarningString = FConvexBuilder::PerformanceWarningString(Planes.Num(), SurfaceParticles.Size());
+			FString PerformanceWarningString = FConvexBuilder::PerformanceWarningString(Planes.Num(), Vertices.Num());
 			if (FConvexBuilder::IsGeometryReductionEnabled())
 			{
 				PerformanceWarningString += ", [Simplifying]";
@@ -628,8 +655,8 @@ namespace Chaos
 		void SimplifyGeometry()
 		{
 			TArray<TArray<int32>> FaceIndices;
-			FConvexBuilder::Simplify(Planes, FaceIndices, SurfaceParticles, LocalBoundingBox);
-			FConvexBuilder::MergeFaces(Planes, FaceIndices, SurfaceParticles);
+			FConvexBuilder::Simplify(Planes, FaceIndices, Vertices, LocalBoundingBox);
+			FConvexBuilder::MergeFaces(Planes, FaceIndices, Vertices);
 			CreateStructureData(MoveTemp(FaceIndices));
 		}
 
@@ -640,7 +667,7 @@ namespace Chaos
 
 	private:
 		TArray<TPlaneConcrete<FReal, 3>> Planes;
-		TParticles<FReal, 3> SurfaceParticles;	//copy of the vertices that are just on the convex hull boundary
+		TArray<FVec3> Vertices; //copy of the vertices that are just on the convex hull boundary
 		TAABB<FReal, 3> LocalBoundingBox;
 		FConvexStructureData StructureData;
 		float Volume;
diff --git a/Engine/Source/Runtime/Experimental/Chaos/Public/Chaos/GeometryQueries.h b/Engine/Source/Runtime/Experimental/Chaos/Public/Chaos/GeometryQueries.h
--- a/Engine/Source/Runtime/Experimental/Chaos/Public/Chaos/GeometryQueries.h
+++ b/Engine/Source/Runtime/Experimental/Chaos/Public/Chaos/GeometryQueries.h
@@ -165,7 +165,7 @@ namespace Chaos
 				//todo: move this out of here
 				if (const auto Convex = TImplicitObjectScaled<FConvex>::AsScaled(InObject))
 				{
-					return Convex->GetUnscaledObject()->GetSurfaceParticles().Size() > 0;
+					return Convex->GetUnscaledObject()->NumVertices() > 0;
 				}				
 
 				return true;
diff --git a/Engine/Source/Runtime/Experimental/ChaosCore/Public/Chaos/Core.h b/Engine/Source/Runtime/Experimental/ChaosCore/Public/Chaos/Core.h
--- a/Engine/Source/Runtime/Experimental/ChaosCore/Public/Chaos/Core.h
+++ b/Engine/Source/Runtime/Experimental/ChaosCore/Public/Chaos/Core.h
@@ -30,5 +30,5 @@ namespace Chaos
 	using TVec3 = TVector<T, 3>;
 
 	/** Core version string for Chaos data. Any DDC builder dependent on Chaos for serialization should depend on this version */
-	static const FString ChaosVersionString = TEXT("8BF2C67F-2BD8-4C10-874B-F5DFA929A5AA");
+	static const FString ChaosVersionString = TEXT("BE819b6A-19B6-4A89-8120-9C995924D41D");
 }
