Merging

//UE5/Release-Engine-Staging/Engine/Source/Runtime/Experimental/Chaos/Private/Chaos/Joint/PBDJointSolverGaussSeidel.cpp

to //UE5/Main/Engine/Source/Runtime/Experimental/Chaos/Private/Chaos/Joint/PBDJointSolverGaussSeidel.cpp

diff --git a/Engine/Source/Runtime/Experimental/Chaos/Private/Chaos/Joint/PBDJointSolverGaussSeidel.cpp b/Engine/Source/Runtime/Experimental/Chaos/Private/Chaos/Joint/PBDJointSolverGaussSeidel.cpp
--- a/Engine/Source/Runtime/Experimental/Chaos/Private/Chaos/Joint/PBDJointSolverGaussSeidel.cpp
+++ b/Engine/Source/Runtime/Experimental/Chaos/Private/Chaos/Joint/PBDJointSolverGaussSeidel.cpp
@@ -224,6 +224,12 @@ namespace Chaos
 		const FPBDJointSolverSettings& SolverSettings,
 		const FPBDJointSettings& JointSettings)
 	{
+		if (InvMs[1] < SMALL_NUMBER)
+		{
+			// If child is kinematic, return. 
+			return;
+		}
+		
 		const FReal LinearProjection = FPBDJointUtilities::GetLinearProjection(SolverSettings, JointSettings);
 		const FReal AngularProjection = FPBDJointUtilities::GetAngularProjection(SolverSettings, JointSettings);
 
