Restore missing include

diff --git a/Engine/Source/Runtime/Experimental/ChaosCore/Public/Chaos/Core.h b/Engine/Source/Runtime/Experimental/ChaosCore/Public/Chaos/Core.h
--- a/Engine/Source/Runtime/Experimental/ChaosCore/Public/Chaos/Core.h
+++ b/Engine/Source/Runtime/Experimental/ChaosCore/Public/Chaos/Core.h
@@ -2,6 +2,7 @@
 
 #pragma once
 
+#include "Chaos/Real.h"
 #include "Chaos/Vector.h"
 #include "Chaos/Matrix.h"
 #include "Chaos/Rotation.h"
