Chaos - Convex vertces - prepare for future merge of back-ported format change from UE5 to UE4

diff --git a/Engine/Source/Runtime/Experimental/Chaos/Public/Chaos/Convex.h b/Engine/Source/Runtime/Experimental/Chaos/Public/Chaos/Convex.h
--- a/Engine/Source/Runtime/Experimental/Chaos/Public/Chaos/Convex.h
+++ b/Engine/Source/Runtime/Experimental/Chaos/Public/Chaos/Convex.h
@@ -11,6 +11,8 @@
 #include "ChaosCheck.h"
 #include "ChaosLog.h"
 #include "UObject/ReleaseObjectVersion.h"
+//#include "UObject/PhysicsObjectVersion.h"
+#include "UObject/UE5MainStreamObjectVersion.h"
 
 namespace Chaos
 {
@@ -538,6 +540,7 @@ namespace Chaos
 		{
 			Ar.UsingCustomVersion(FExternalPhysicsCustomObjectVersion::GUID);
 			Ar.UsingCustomVersion(FUE5MainStreamObjectVersion::GUID);
+			//Ar.UsingCustomVersion(FPhysicsObjectVersion::GUID);
 			FImplicitObject::SerializeImp(Ar);
 
 			if (Ar.CustomVer(FExternalPhysicsCustomObjectVersion::GUID) < FExternalPhysicsCustomObjectVersion::ConvexUsesTPlaneConcrete)
@@ -556,7 +559,19 @@ namespace Chaos
 				Ar << Planes;
 			}
 
-			if (Ar.CustomVer(FUE5MainStreamObjectVersion::GUID) < FUE5MainStreamObjectVersion::ConvexUsesVerticesArray)
+			// Do we use the old Particles array or the new Vertices array?
+			// Note: This change was back-ported from UE5 to UE4, so we need to check both 
+			// a UE4 and a UE5 object version. If either is >= ConvexUsesVerticesArray we have the new format.
+			//
+			// @todo(chaos): the following change was back-ported to UE4 which requires some ObjectVersion shenanigans.
+			// We should get a merge conflict here, in which case, the code coming from UE4 contains instructions on what to do.
+			// Alternatively ping Chris Caulfield if RM hasn't already done so.
+			//
+			bool bConvexVerticesNewFormatUE4 = false;
+			bool bConvexVerticesNewFormatUE5 = (Ar.CustomVer(FUE5MainStreamObjectVersion::GUID) >= FUE5MainStreamObjectVersion::ConvexUsesVerticesArray);
+			bool bConvexVerticesNewFormat = bConvexVerticesNewFormatUE4 || bConvexVerticesNewFormatUE5;
+
+			if (!bConvexVerticesNewFormat)
 			{
 				TParticles<FReal, 3> TmpSurfaceParticles;
 				Ar << TmpSurfaceParticles;
