Chaos - rotate DDC GUID after merge issue

diff --git a/Engine/Source/Runtime/Experimental/ChaosCore/Public/Chaos/Core.h b/Engine/Source/Runtime/Experimental/ChaosCore/Public/Chaos/Core.h
--- a/Engine/Source/Runtime/Experimental/ChaosCore/Public/Chaos/Core.h
+++ b/Engine/Source/Runtime/Experimental/ChaosCore/Public/Chaos/Core.h
@@ -31,6 +31,8 @@ namespace Chaos
 	template <typename T>
 	using TVec4 = TVector<T, 4>;
 
-	/** Core version string for Chaos data. Any DDC builder dependent on Chaos for serialization should depend on this version */
-	static const FString ChaosVersionString = TEXT("BE819b6A-19B6-4A89-8120-9C995924D41D");
+	// NOTE: if you get a merge conflict on the GUID, you must replace it with a new GUID - do not accept the source or target
+	// or you will likely get DDC version conflicts resulting in crashes during load.
+	// Core version string for Chaos data. Any DDC builder dependent on Chaos for serialization should depend on this version
+	static const FString ChaosVersionString = TEXT("F7030D07-09E0-4724-9B3E-8DBED3678C15");
 }
