Chaos - fix non-unity (DebugDrawQueue.h). Fix manually merged from Main.

diff --git a/Engine/Source/Runtime/Experimental/Chaos/Public/Chaos/DebugDrawQueue.h b/Engine/Source/Runtime/Experimental/Chaos/Public/Chaos/DebugDrawQueue.h
--- a/Engine/Source/Runtime/Experimental/Chaos/Public/Chaos/DebugDrawQueue.h
+++ b/Engine/Source/Runtime/Experimental/Chaos/Public/Chaos/DebugDrawQueue.h
@@ -7,6 +7,7 @@
 #include "Math/Vector.h"
 #include "Math/Color.h"
 #include "Math/Quat.h"
+#include "Misc/ScopeLock.h"
 #include "HAL/ThreadSafeBool.h"
 #include "HAL/IConsoleManager.h"
 
