ModelingTools: disable selection by default on AInternalToolFrameworkActor. Subclasses of this Actor are used for 3D gizmos and live mesh editing previews, they are not meant to be exposed in the actor/component details panel.
#rnx

diff --git a/Engine/Source/Runtime/Experimental/InteractiveToolsFramework/Public/InteractiveToolObjects.h b/Engine/Source/Runtime/Experimental/InteractiveToolsFramework/Public/InteractiveToolObjects.h
--- a/Engine/Source/Runtime/Experimental/InteractiveToolsFramework/Public/InteractiveToolObjects.h
+++ b/Engine/Source/Runtime/Experimental/InteractiveToolsFramework/Public/InteractiveToolObjects.h
@@ -22,6 +22,16 @@ private:
 
 public:
 
+	/** Controls whether this InternalToolFrameworkActor can be selected in the Editor. */
+	UPROPERTY()
+	bool bIsSelectableInEditor = false;
+
+#if WITH_EDITOR
+	virtual bool IsSelectable() const override
+	{
+		return bIsSelectableInEditor;
+	}
+#endif
 };
 
 
