Chaos Physics: When CCD disabled,  bounding boxes will now not expand with velocity

#[fyi] Michael.Lentine, Benn.Gallagher, Ori.Cohen

diff --git a/Engine/Source/Runtime/Experimental/Chaos/Private/Chaos/ChaosDebugDraw.cpp b/Engine/Source/Runtime/Experimental/Chaos/Private/Chaos/ChaosDebugDraw.cpp
--- a/Engine/Source/Runtime/Experimental/Chaos/Private/Chaos/ChaosDebugDraw.cpp
+++ b/Engine/Source/Runtime/Experimental/Chaos/Private/Chaos/ChaosDebugDraw.cpp
@@ -338,7 +338,7 @@ namespace Chaos
 			TConstGenericParticleHandle<FReal, 3> Particle = InParticle;
 
 			// This matches the calculation in SpatialAccelerationBroadPhase.h
-			const FReal Box1Thickness = ComputeBoundsThickness(Particle->V(), Dt, BoundsThickness, BoundsThicknessVelocityInflation).Size();
+			const FReal Box1Thickness = ComputeBoundsThickness(*InParticle, Dt, BoundsThickness, BoundsThicknessVelocityInflation).Size();
 			const FAABB3 Box = ComputeWorldSpaceBoundingBox<FReal>(*InParticle).ThickenSymmetrically(FVec3(Box1Thickness));
 
 			const FVec3 P = SpaceTransform.TransformPosition(Box.GetCenter());
diff --git a/Engine/Source/Runtime/Experimental/Chaos/Public/Chaos/BoundingVolumeUtilities.h b/Engine/Source/Runtime/Experimental/Chaos/Public/Chaos/BoundingVolumeUtilities.h
--- a/Engine/Source/Runtime/Experimental/Chaos/Public/Chaos/BoundingVolumeUtilities.h
+++ b/Engine/Source/Runtime/Experimental/Chaos/Public/Chaos/BoundingVolumeUtilities.h
@@ -33,18 +33,27 @@ inline FVec3 ComputeBoundsThickness(FVec3 Vel, FReal Dt, FReal BoundsThickness,
 
 inline FVec3 ComputeBoundsThickness(const TPBDRigidParticles<FReal, 3>& InParticles, FReal Dt, int32 BodyIndex, FReal BoundsThickness, FReal BoundsVelocityInflation)
 {
-	return ComputeBoundsThickness(InParticles.V(BodyIndex), Dt, BoundsThickness, BoundsVelocityInflation);
-}
-
-inline FVec3 ComputeBoundsThickness(const TKinematicGeometryParticles<FReal, 3>& InParticles, FReal Dt, int32 BodyIndex, FReal BoundsThickness, FReal BoundsVelocityInflation)
-{
-	return ComputeBoundsThickness(InParticles.V(BodyIndex), Dt, BoundsThickness, BoundsVelocityInflation);
+	return ComputeBoundsThickness(InParticles.V(BodyIndex), Dt, BoundsThickness, !InParticles.CCDEnabled(BodyIndex) ? FReal(0) : BoundsVelocityInflation);
 }
 
 template <typename THandle>
 FVec3 ComputeBoundsThickness(const THandle& PBDRigid, FReal Dt, FReal BoundsThickness, FReal BoundsVelocityInflation)
 {
-	return ComputeBoundsThickness(PBDRigid.V(), Dt, BoundsThickness, BoundsVelocityInflation);
+	FReal NewBoundsVelocityInflation = BoundsVelocityInflation;
+	auto* RigidBody = PBDRigid.CastToRigidParticle();
+	if (RigidBody && !RigidBody->CCDEnabled())
+	{
+		NewBoundsVelocityInflation = 0.0f;
+	}
+
+	FVec3 Vel{ 0 };
+	auto* KinematicBody = PBDRigid.CastToKinematicParticle();
+	if (KinematicBody)
+	{
+		Vel = KinematicBody->V();
+	}
+
+	return ComputeBoundsThickness(Vel, Dt, BoundsThickness, NewBoundsVelocityInflation);
 }
 
 template<class OBJECT_ARRAY>
diff --git a/Engine/Source/Runtime/Experimental/Chaos/Public/Chaos/PBDRigidsEvolutionGBF.h b/Engine/Source/Runtime/Experimental/Chaos/Public/Chaos/PBDRigidsEvolutionGBF.h
--- a/Engine/Source/Runtime/Experimental/Chaos/Public/Chaos/PBDRigidsEvolutionGBF.h
+++ b/Engine/Source/Runtime/Experimental/Chaos/Public/Chaos/PBDRigidsEvolutionGBF.h
@@ -226,7 +226,10 @@ namespace Chaos
 					{
 						const FAABB3& LocalBounds = Particle.LocalBounds();
 						FAABB3 WorldSpaceBounds = LocalBounds.TransformedAABB(FRigidTransform3(Particle.P(), Particle.Q()));
-						WorldSpaceBounds.ThickenSymmetrically(Particle.V() * Dt);
+						if (Particle.CCDEnabled())
+						{
+							WorldSpaceBounds.ThickenSymmetrically(Particle.V() * Dt);
+						}
 						Particle.SetWorldSpaceInflatedBounds(WorldSpaceBounds);
 					}
 				}
