Fix for uninitialized DamageThreshold warning.

diff --git a/Engine/Source/Runtime/Experimental/GeometryCollectionEngine/Private/GeometryCollection/GeometryCollectionObject.cpp b/Engine/Source/Runtime/Experimental/GeometryCollectionEngine/Private/GeometryCollection/GeometryCollectionObject.cpp
--- a/Engine/Source/Runtime/Experimental/GeometryCollectionEngine/Private/GeometryCollection/GeometryCollectionObject.cpp
+++ b/Engine/Source/Runtime/Experimental/GeometryCollectionEngine/Private/GeometryCollection/GeometryCollectionObject.cpp
@@ -89,6 +89,7 @@ FGeometryCollectionSizeSpecificData::FGeometryCollectionSizeSpecificData()
 	, CollisionObjectReductionPercentage(0.0f)
 	, CollisionParticlesFraction(1.0f)
 	, MaximumCollisionParticles(60)
+	, DamageThreshold(250.0)
 {
 }
 
