Chaos - merging various fixes that missed the main merge deadline

diff --git a/Engine/Source/Programs/HeadlessChaos/Private/HeadlessChaosTestConvex.cpp b/Engine/Source/Programs/HeadlessChaos/Private/HeadlessChaosTestConvex.cpp
--- a/Engine/Source/Programs/HeadlessChaos/Private/HeadlessChaosTestConvex.cpp
+++ b/Engine/Source/Programs/HeadlessChaos/Private/HeadlessChaosTestConvex.cpp
@@ -9,6 +9,7 @@
 #include "Chaos/Convex.h"
 #include "Chaos/ImplicitObjectScaled.h"
 #include "Chaos/Particles.h"
+#include "../Resource/TestGeometry2.h"
 
 namespace ChaosTest
 {
@@ -32,7 +33,7 @@ namespace ChaosTest
 		TAABB<FReal, 3> LocalBounds;
 
 		FConvexBuilder::Build(Particles, Planes, FaceVertices, SurfaceParticles, LocalBounds);
-		FConvexBuilder::MergeFaces(Planes, FaceVertices, SurfaceParticles);
+		FConvexBuilder::MergeFaces(Planes, FaceVertices, SurfaceParticles, 1.0f);
 
 		// Check that we have the right number of faces and particles
 		EXPECT_EQ(SurfaceParticles.Num(), 8);
@@ -87,6 +88,10 @@ namespace ChaosTest
 	// Check that the convex structure data is consistent (works for TBox and TConvex)
 	template<typename T_GEOM> void TestConvexStructureDataImpl(const T_GEOM& Convex)
 	{
+		// Note: This tolerance matches the one passed to FConvexBuilder::MergeFaces in the FConvex constructor, but it should be dependent on size
+		//const FReal Tolerance = 1.e-4f * Convex.BoundingBox().OriginRadius();
+		const FReal Tolerance = 1.0f;
+
 		// Check all per-plane data
 		for (int32 PlaneIndex = 0; PlaneIndex < Convex.NumPlanes(); ++PlaneIndex)
 		{
@@ -97,7 +102,7 @@ namespace ChaosTest
 				const int32 VertexIndex = Convex.GetPlaneVertex(PlaneIndex, PlaneVertexIndex);
 				const FVec3 Vertex = Convex.GetVertex(VertexIndex);
 				const FReal VertexDistance = FVec3::DotProduct(Plane.Normal(), Vertex - Plane.X());
-				EXPECT_NEAR(VertexDistance, 0.0f, 1.e-3f);
+				EXPECT_NEAR(VertexDistance, 0.0f, Tolerance);
 			}
 		}
 
@@ -111,7 +116,7 @@ namespace ChaosTest
 				const TPlaneConcrete<FReal, 3> Plane = Convex.GetPlane(PlaneIndex);
 				const FVec3 Vertex = Convex.GetVertex(VertexIndex);
 				const FReal VertexDistance = FVec3::DotProduct(Plane.Normal(), Vertex - Plane.X());
-				EXPECT_NEAR(VertexDistance, 0.0f, 1.e-3f);
+				EXPECT_NEAR(VertexDistance, 0.0f, Tolerance);
 			}
 		}
 	}
@@ -193,6 +198,89 @@ namespace ChaosTest
 		}
 	}
 
+	// Check the reverse mapping planes->vertices->planes is intact
+	template<typename T_STRUCTUREDATA>
+	void TestConvexStructureDataMapping(const T_STRUCTUREDATA& StructureData)
+	{
+		// For each plane, get the list of vertices that make its edges.
+		// Then check that we list of planes used by that vertex contains the original plane
+		for (int32 PlaneIndex = 0; PlaneIndex < StructureData.NumPlanes(); ++PlaneIndex)
+		{
+			for (int32 PlaneVertexIndex = 0; PlaneVertexIndex < StructureData.NumPlaneVertices(PlaneIndex); ++PlaneVertexIndex)
+			{
+				const int32 VertexIndex = StructureData.GetPlaneVertex(PlaneIndex, PlaneVertexIndex);
+
+				// Check that the plane's vertex has the plane in its list
+				bool bFoundPlane = false;
+				for (int32 VertexPlaneIndex = 0; VertexPlaneIndex < StructureData.NumVertexPlanes(VertexIndex); ++VertexPlaneIndex)
+				{
+					const int32 SearchPlaneIndex = StructureData.GetVertexPlane(VertexIndex, VertexPlaneIndex);
+					if (PlaneIndex == SearchPlaneIndex)
+					{
+						bFoundPlane = true;
+						break;
+					}
+				}
+				EXPECT_TRUE(bFoundPlane);
+			}
+		}
+	}
+
+	// Check that the structure data is good for convex shapes that have faces merged during construction
+	// This test uses the small index size in StructureData.
+	GTEST_TEST(ConvexStructureTests, TestSmallIndexStructureData)
+	{
+		FMath::RandInit(53799058);
+		const FReal Radius = 1000.0f;
+
+		const int32 NumVertices = TestGeometry2::RawVertexArray.Num() / 3;
+		TParticles<FReal, 3> Particles;
+		Particles.AddParticles(NumVertices);
+		for (int32 ParticleIndex = 0; ParticleIndex < NumVertices; ++ParticleIndex)
+		{
+			const FVec3 VertexPos = FVec3(
+				TestGeometry2::RawVertexArray[3 * ParticleIndex + 0],
+				TestGeometry2::RawVertexArray[3 * ParticleIndex + 1],
+				TestGeometry2::RawVertexArray[3 * ParticleIndex + 2]
+			);
+			Particles.X(ParticleIndex) = VertexPos;
+		}
+
+		FConvex Convex(Particles, 0.0f);
+
+		const FConvexStructureDataU8& StructureData = Convex.GetStructureData().Data8();
+		TestConvexStructureDataMapping(StructureData);
+		TestConvexStructureDataImpl(Convex);
+	}
+
+	// Check that the structure data is good for convex shapes that have faces merged during construction
+	// This test uses the large index size in StructureData.
+	GTEST_TEST(ConvexStructureTests, TestLargeIndexStructureData2)
+	{
+		FMath::RandInit(53799058);
+		const FReal Radius = 1000.0f;
+		const int32 NumVertices = 1000;
+
+		// Make a convex with points on a sphere.
+		TParticles<FReal, 3> Particles;
+		Particles.AddParticles(NumVertices);
+		for (int32 ParticleIndex = 0; ParticleIndex < NumVertices; ++ParticleIndex)
+		{
+			const FReal Theta = FMath::RandRange(-PI, PI);
+			const FReal Phi = FMath::RandRange(-0.5f * PI, 0.5f * PI);
+			const FVec3 VertexPos = Radius * FVec3(FMath::Cos(Theta), FMath::Sin(Theta), FMath::Sin(Phi));
+			Particles.X(ParticleIndex) = VertexPos;
+		}
+		FConvex Convex(Particles, 0.0f);
+
+		EXPECT_GT(Convex.NumVertices(), 800);
+		EXPECT_GT(Convex.NumPlanes(), 500);
+
+		const FConvexStructureDataS32& StructureData = Convex.GetStructureData().Data32();
+		TestConvexStructureDataMapping(StructureData);
+		TestConvexStructureDataImpl(Convex);
+	}
+
 	// Check that extremely small generated triangle don't trigger the normal check
 	GTEST_TEST(ConvexStructureTests, TestConvexFaceNormalCheck)
 	{
diff --git a/Engine/Source/Programs/HeadlessChaos/Private/HeadlessChaosTestImplicits.cpp b/Engine/Source/Programs/HeadlessChaos/Private/HeadlessChaosTestImplicits.cpp
--- a/Engine/Source/Programs/HeadlessChaos/Private/HeadlessChaosTestImplicits.cpp
+++ b/Engine/Source/Programs/HeadlessChaos/Private/HeadlessChaosTestImplicits.cpp
@@ -652,11 +652,84 @@ namespace ChaosTest {
 		}
 	}
 
+	TEST(ImplicitTests, TestImplicitConvex_PhiWithNormal_Penetrating)
+	{
+		const FVec3 Size = FVec3(500, 500, 100);
+		FImplicitConvex3 Convex = CreateConvexBox(Size, 10);
+		{
+			// Near point just inside the top face, near the forward edge
+			const FVec3 Point = FVec3(0.5f * Size.X, 0.0f, 0.5f * Size.Z) - FVec3(10, 0, 1);
+			FVec3 Normal;
+			FReal Phi = Convex.PhiWithNormal(Point, Normal);
+			EXPECT_NEAR(Normal.Z, 1.0f, 1.e-4f);
+			EXPECT_NEAR(Phi, -1.0f, 1.e-4f);
+		}
+		{
+			// Near point just inside the top face, near the forward edge
+			const FVec3 Point = FVec3(0.5f * Size.X, 0.0f, 0.5f * Size.Z) - FVec3(3, 0, 1);
+			FVec3 Normal;
+			FReal Phi = Convex.PhiWithNormal(Point, Normal);
+			EXPECT_NEAR(Normal.Z, 1.0f, 1.e-4f);
+			EXPECT_NEAR(Phi, -1.0f, 1.e-4f);
+		}
+		{
+			// Near point just inside the top face, near the forward edge
+			const FVec3 Point = FVec3(0.5f * Size.X, 0.0f, 0.5f * Size.Z) - FVec3(1, 0, 0.1);
+			FVec3 Normal;
+			FReal Phi = Convex.PhiWithNormal(Point, Normal);
+			EXPECT_NEAR(Normal.Z, 1.0f, 1.e-4f);
+			EXPECT_NEAR(Phi, -0.1f, 1.e-4f);
+		}
+	}
+
+	TEST(ImplicitTests, TestImplicitConvex_PhiWithNormal_Separated)
+	{
+		const FVec3 Size = FVec3(500, 500, 100);
+		FImplicitConvex3 Convex = CreateConvexBox(Size, 10);
+		{
+			const FVec3 Offset = FVec3(10, 0, 1);
+			const FVec3 Point = FVec3(0.5f * Size.X, 0.5f * Size.Y, 0.5f * Size.Z) + Offset;
+			FVec3 Normal;
+			FReal Phi = Convex.PhiWithNormal(Point, Normal);
+			const FVec3 ExpectedNormal = Offset.GetUnsafeNormal();
+			const FReal ExpectedPhi = Offset.Size();
+			EXPECT_NEAR(Normal.X, ExpectedNormal.X, 1.e-4f);
+			EXPECT_NEAR(Normal.Y, ExpectedNormal.Y, 1.e-4f);
+			EXPECT_NEAR(Normal.Z, ExpectedNormal.Z, 1.e-4f);
+			EXPECT_NEAR(Phi, ExpectedPhi, 1.e-4f);
+		}
+		{
+			const FVec3 Offset = FVec3(3, 2, 1);
+			const FVec3 Point = FVec3(0.5f * Size.X, 0.5f * Size.Y, 0.5f * Size.Z) + Offset;
+			FVec3 Normal;
+			FReal Phi = Convex.PhiWithNormal(Point, Normal);
+			const FVec3 ExpectedNormal = Offset.GetUnsafeNormal();
+			const FReal ExpectedPhi = Offset.Size();
+			EXPECT_NEAR(Normal.X, ExpectedNormal.X, 1.e-4f);
+			EXPECT_NEAR(Normal.Y, ExpectedNormal.Y, 1.e-4f);
+			EXPECT_NEAR(Normal.Z, ExpectedNormal.Z, 1.e-4f);
+			EXPECT_NEAR(Phi, ExpectedPhi, 1.e-4f);
+		}
+		{
+			const FVec3 Offset = FVec3(0, 1, 1);
+			const FVec3 Point = FVec3(0.5f * Size.X, 0.5f * Size.Y, 0.5f * Size.Z) + Offset;
+			FVec3 Normal;
+			FReal Phi = Convex.PhiWithNormal(Point, Normal);
+			const FVec3 ExpectedNormal = Offset.GetUnsafeNormal();
+			const FReal ExpectedPhi = Offset.Size();
+			EXPECT_NEAR(Normal.X, ExpectedNormal.X, 1.e-4f);
+			EXPECT_NEAR(Normal.Y, ExpectedNormal.Y, 1.e-4f);
+			EXPECT_NEAR(Normal.Z, ExpectedNormal.Z, 1.e-4f);
+			EXPECT_NEAR(Phi, ExpectedPhi, 1.e-4f);
+		}
+	}
+
+
 	// Check that PhiWithNormal works properly on Scaled Convex.
 	// There was a bug where scaled convexed would bias face selection based on the 
 	// scale, so a unit box scaled by 5 in the X would report the +X face as the 
 	// contact face for the position (0.4, 0.0, 4.8) even though the +Z face is closer.
-	TEST(ImplicitTests, TestImplicitScaledConvex)
+	TEST(ImplicitTests, TestImplicitScaledConvex_PhiWithNormal_Penetrating)
 	{
 		const FVec3 Size = FVec3(500, 500, 100);
 		const FVec3 Scale = FVec3(5, 5, 1);
@@ -665,30 +738,87 @@ namespace ChaosTest {
 
 		{
 			// Near point just inside the top face, near the forward edge
-			const FVec3 NearTopRight = FVec3(0.5f * ScaledSize.X, 0.0f, 0.5f * ScaledSize.Z) - FVec3(10, 0, 1);
+			const FVec3 Point = FVec3(0.5f * ScaledSize.X, 0.0f, 0.5f * ScaledSize.Z) - FVec3(10, 0, 1);
 			FVec3 Normal;
-			FReal Phi = ScaledConvex.PhiWithNormal(NearTopRight, Normal);
+			FReal Phi = ScaledConvex.PhiWithNormal(Point, Normal);
 			EXPECT_NEAR(Normal.Z, 1.0f, 1.e-4f);
 			EXPECT_NEAR(Phi, -1.0f, 1.e-4f);
 		}
 		{
 			// Near point just inside the top face, near the forward edge
-			const FVec3 NearTopRight = FVec3(0.5f * ScaledSize.X, 0.0f, 0.5f * ScaledSize.Z) - FVec3(3, 0, 1);
+			const FVec3 Point = FVec3(0.5f * ScaledSize.X, 0.0f, 0.5f * ScaledSize.Z) - FVec3(3, 0, 1);
 			FVec3 Normal;
-			FReal Phi = ScaledConvex.PhiWithNormal(NearTopRight, Normal);
+			FReal Phi = ScaledConvex.PhiWithNormal(Point, Normal);
 			EXPECT_NEAR(Normal.Z, 1.0f, 1.e-4f);
 			EXPECT_NEAR(Phi, -1.0f, 1.e-4f);
 		}
 		{
 			// Near point just inside the top face, near the forward edge
-			const FVec3 NearTopRight = FVec3(0.5f * ScaledSize.X, 0.0f, 0.5f * ScaledSize.Z) - FVec3(1, 0, 0.1);
+			const FVec3 Point = FVec3(0.5f * ScaledSize.X, 0.0f, 0.5f * ScaledSize.Z) - FVec3(1, 0, 0.1);
 			FVec3 Normal;
-			FReal Phi = ScaledConvex.PhiWithNormal(NearTopRight, Normal);
+			FReal Phi = ScaledConvex.PhiWithNormal(Point, Normal);
 			EXPECT_NEAR(Normal.Z, 1.0f, 1.e-4f);
 			EXPECT_NEAR(Phi, -0.1f, 1.e-4f);
 		}
 	}
 
+	TEST(ImplicitTests, TestImplicitScaledConvex_PhiWithNormal_Separated)
+	{
+		const FVec3 Size = FVec3(500, 500, 100);
+		const FVec3 Scale = FVec3(5, 5, 1);
+		const FVec3 ScaledSize = Scale * Size;
+		TImplicitObjectScaled<FImplicitConvex3> ScaledConvex = CreateScaledConvexBox(Size, Scale, 10);
+
+		{
+			// Near point just inside the top face, near the forward edge
+			const FVec3 Point = FVec3(0.5f * ScaledSize.X, 0.0f, 0.5f * ScaledSize.Z) + FVec3(-10, 0, 1);
+			FVec3 Normal;
+			FReal Phi = ScaledConvex.PhiWithNormal(Point, Normal);
+			EXPECT_NEAR(Normal.Z, 1.0f, 1.e-4f);
+			EXPECT_NEAR(Phi, 1.0f, 1.e-4f);
+		}
+		{
+			// Near point just inside the top face, near the forward edge
+			const FVec3 Point = FVec3(0.5f * ScaledSize.X, 0.0f, 0.5f * ScaledSize.Z) + FVec3(-3, 0, 1);
+			FVec3 Normal;
+			FReal Phi = ScaledConvex.PhiWithNormal(Point, Normal);
+			EXPECT_NEAR(Normal.Z, 1.0f, 1.e-4f);
+			EXPECT_NEAR(Phi, 1.0f, 1.e-4f);
+		}
+		{
+			// Near point just inside the top face, near the forward edge
+			const FVec3 Point = FVec3(0.5f * ScaledSize.X, 0.0f, 0.5f * ScaledSize.Z) + FVec3(-1, 0, 0.1);
+			FVec3 Normal;
+			FReal Phi = ScaledConvex.PhiWithNormal(Point, Normal);
+			EXPECT_NEAR(Normal.Z, 1.0f, 1.e-4f);
+			EXPECT_NEAR(Phi, 0.1f, 1.e-4f);
+		}
+		{
+			// Point outside the face edge 
+			const FVec3 Point = FVec3(0.5f * ScaledSize.X, 0.0f, 0.5f * ScaledSize.Z) + FVec3(1, 0, 1);
+			FVec3 Normal;
+			FReal Phi = ScaledConvex.PhiWithNormal(Point, Normal);
+			const FVec3 ExpectedNormal = FVec3(1, 0, 1).GetUnsafeNormal();
+			const FReal ExpectedPhi = FVec3(1, 0, 1).Size();
+			EXPECT_NEAR(Normal.X, ExpectedNormal.X, 1.e-4f);
+			EXPECT_NEAR(Normal.Y, ExpectedNormal.Y, 1.e-4f);
+			EXPECT_NEAR(Normal.Z, ExpectedNormal.Z, 1.e-4f);
+			EXPECT_NEAR(Phi, ExpectedPhi, 1.e-4f);
+		}
+		{
+			// Point outside the face corner 
+			const FVec3 Point = FVec3(0.5f * ScaledSize.X, 0.5f * ScaledSize.Y, 0.5f * ScaledSize.Z) + FVec3(3, 2, 1);
+			FVec3 Normal;
+			FReal Phi = ScaledConvex.PhiWithNormal(Point, Normal);
+			const FVec3 ExpectedNormal = FVec3(3, 2, 1).GetUnsafeNormal();
+			const FReal ExpectedPhi = FVec3(3, 2, 1).Size();
+			EXPECT_NEAR(Normal.X, ExpectedNormal.X, 1.e-4f);
+			EXPECT_NEAR(Normal.Y, ExpectedNormal.Y, 1.e-4f);
+			EXPECT_NEAR(Normal.Z, ExpectedNormal.Z, 1.e-4f);
+			EXPECT_NEAR(Phi, ExpectedPhi, 1.e-4f);
+		}
+	}
+
 
 	void ImplicitTransformed()
 	{
diff --git a/Engine/Source/Programs/HeadlessChaos/Private/HeadlessChaosTestUtility.cpp b/Engine/Source/Programs/HeadlessChaos/Private/HeadlessChaosTestUtility.cpp
--- a/Engine/Source/Programs/HeadlessChaos/Private/HeadlessChaosTestUtility.cpp
+++ b/Engine/Source/Programs/HeadlessChaos/Private/HeadlessChaosTestUtility.cpp
@@ -702,6 +702,28 @@ namespace ChaosTest {
 		}
 	}
 
+	FImplicitConvex3 CreateConvexBox(const FVec3& BoxSize, const FReal Margin)
+	{
+		const FVec3 HalfSize = 0.5f * BoxSize;
+
+		TArray<FVec3> BoxVerts =
+		{
+			FVec3(-HalfSize.X, -HalfSize.Y, -HalfSize.Z),
+			FVec3(-HalfSize.X,  HalfSize.Y, -HalfSize.Z),
+			FVec3(HalfSize.X,  HalfSize.Y, -HalfSize.Z),
+			FVec3(HalfSize.X, -HalfSize.Y, -HalfSize.Z),
+			FVec3(-HalfSize.X, -HalfSize.Y,  HalfSize.Z),
+			FVec3(-HalfSize.X,  HalfSize.Y,  HalfSize.Z),
+			FVec3(HalfSize.X,  HalfSize.Y,  HalfSize.Z),
+			FVec3(HalfSize.X, -HalfSize.Y,  HalfSize.Z),
+		};
+
+		TParticles<FReal, 3> BoxParticles(MoveTemp(BoxVerts));
+
+		return FImplicitConvex3(BoxParticles, Margin);
+	}
+
+
 	TImplicitObjectScaled<FImplicitConvex3> CreateScaledConvexBox(const FVec3& BoxSize, const FVec3 BoxScale, const FReal Margin)
 	{
 		const FVec3 HalfSize = 0.5f * BoxSize;
diff --git a/Engine/Source/Programs/HeadlessChaos/Public/HeadlessChaosTestUtility.h b/Engine/Source/Programs/HeadlessChaos/Public/HeadlessChaosTestUtility.h
--- a/Engine/Source/Programs/HeadlessChaos/Public/HeadlessChaosTestUtility.h
+++ b/Engine/Source/Programs/HeadlessChaos/Public/HeadlessChaosTestUtility.h
@@ -199,6 +199,6 @@ namespace ChaosTest {
 	}
 
 
+	extern FImplicitConvex3 CreateConvexBox(const FVec3& BoxSize, const FReal Margin);
 	extern TImplicitObjectScaled<FImplicitConvex3> CreateScaledConvexBox(const FVec3& BoxSize, const FVec3 BoxScale, const FReal Margin);
-
 }
diff --git a/Engine/Source/Runtime/Experimental/Chaos/Public/Chaos/CollisionConvexMesh.h b/Engine/Source/Runtime/Experimental/Chaos/Public/Chaos/CollisionConvexMesh.h
--- a/Engine/Source/Runtime/Experimental/Chaos/Public/Chaos/CollisionConvexMesh.h
+++ b/Engine/Source/Runtime/Experimental/Chaos/Public/Chaos/CollisionConvexMesh.h
@@ -522,7 +522,7 @@ namespace Chaos
 		}
 
 		// Convert multi-triangle faces to single n-gons
-		static void MergeFaces(TArray<TPlaneConcrete<FReal, 3>>& InOutPlanes, TArray<TArray<int32>>& InOutFaceVertexIndices, const TArray<FVec3>& Vertices)
+		static void MergeFaces(TArray<TPlaneConcrete<FReal, 3>>& InOutPlanes, TArray<TArray<int32>>& InOutFaceVertexIndices, const TArray<FVec3>& Vertices, float DistanceThreshold)
 		{
 			const FReal NormalThreshold = 1.e-4f;
 
@@ -537,19 +537,49 @@ namespace Chaos
 					const TPlaneConcrete<FReal, 3>& Plane1 = InOutPlanes[PlaneIndex1];
 					const TArray<int32>& Vertices1 = InOutFaceVertexIndices[PlaneIndex1];
 
+					// First similarity test: normals are close - this will reject all very dissimilar faces
 					const FReal PlaneNormalDot = FVec3::DotProduct(Plane0.Normal(), Plane1.Normal());
 					if (PlaneNormalDot > 1.0f - NormalThreshold)
 					{
-						// Merge the verts from the second plane into the first
-						for (int32 VertexIndex1 = 0; VertexIndex1 < Vertices1.Num(); ++VertexIndex1)
+						// Second similarity test: vertices of one plane are within threshold distance of the other. This is slower but more accurate
+						bool bWithinDistanceThreshold = true;
+						for (int32 Plane0VertexIndex : Vertices0)
 						{
-							Vertices0.AddUnique(Vertices1[VertexIndex1]);
+							const FVec3 Plane0Vertex = Vertices[Plane0VertexIndex];
+							const FReal Plane0VertexDistance = FMath::Abs(FVec3::DotProduct(Plane1.X() - Plane0Vertex, Plane1.Normal()));
+							if (Plane0VertexDistance > DistanceThreshold)
+							{
+								bWithinDistanceThreshold = false;
+								break;
+							}
 						}
+						if (bWithinDistanceThreshold)
+						{
+							for (int32 Plane1VertexIndex : Vertices1)
+							{
+								const FVec3 Plane1Vertex = Vertices[Plane1VertexIndex];
+								const FReal Plane1VertexDistance = FMath::Abs(FVec3::DotProduct(Plane0.X() - Plane1Vertex, Plane0.Normal()));
+								if (Plane1VertexDistance > DistanceThreshold)
+								{
+									bWithinDistanceThreshold = false;
+									break;
+								}
+							}
+						}
+
+						if (bWithinDistanceThreshold)
+						{
+							// Merge the verts from the second plane into the first
+							for (int32 VertexIndex1 = 0; VertexIndex1 < Vertices1.Num(); ++VertexIndex1)
+							{
+								Vertices0.AddUnique(Vertices1[VertexIndex1]);
+							}
 
-						// Erase the second plane
-						InOutPlanes.RemoveAtSwap(PlaneIndex1, 1, false);
-						InOutFaceVertexIndices.RemoveAtSwap(PlaneIndex1, 1, false);
-						--PlaneIndex1;
+							// Erase the second plane
+							InOutPlanes.RemoveAtSwap(PlaneIndex1, 1, false);
+							InOutFaceVertexIndices.RemoveAtSwap(PlaneIndex1, 1, false);
+							--PlaneIndex1;
+						}
 					}
 				}
 			}
diff --git a/Engine/Source/Runtime/Experimental/Chaos/Public/Chaos/Convex.h b/Engine/Source/Runtime/Experimental/Chaos/Public/Chaos/Convex.h
--- a/Engine/Source/Runtime/Experimental/Chaos/Public/Chaos/Convex.h
+++ b/Engine/Source/Runtime/Experimental/Chaos/Public/Chaos/Convex.h
@@ -79,11 +79,13 @@ namespace Chaos
 			FConvexBuilder::Build(InVertices, Planes, FaceIndices, Vertices, LocalBoundingBox);
 			CHAOS_ENSURE(Planes.Num() == FaceIndices.Num());
 
-			// @chaos(todo): this only works with triangles. Fix that an we can run MergeFaces before calling this
+			// @todo(chaos): this only works with triangles. Fix that an we can run MergeFaces before calling this
 			TParticles<FReal, 3> VertexParticles(CopyTemp(Vertices));
 			CalculateVolumeAndCenterOfMass(VertexParticles, FaceIndices, Volume, CenterOfMass);
 
-			FConvexBuilder::MergeFaces(Planes, FaceIndices, Vertices);
+			// @todo(chaos): DistanceTolerance should be based on size, or passed in
+			const FReal DistanceTolerance = 1.0f;
+			FConvexBuilder::MergeFaces(Planes, FaceIndices, Vertices, DistanceTolerance);
 			CHAOS_ENSURE(Planes.Num() == FaceIndices.Num());
 
 			CreateStructureData(MoveTemp(FaceIndices));
@@ -108,61 +110,20 @@ namespace Chaos
 		}
 
 		// Return the distance to the surface
-		virtual FReal PhiWithNormal(const FVec3& x, FVec3& Normal) const override
+		virtual FReal PhiWithNormal(const FVec3& X, FVec3& Normal) const override
 		{
-			float Phi = PhiWithNormalInternal(x, Normal);
-			if (Phi > 0)
-			{
-				//Outside convex, so test against bounding box - this is done to avoid
-				//inaccurate results given by PhiWithNormalInternal when x is far outside
-
-				FVec3 SnappedPosition, BoundingNormal;
-				const float BoundingPhi = LocalBoundingBox.PhiWithNormal(x, BoundingNormal);
-				if (BoundingPhi <= 0)
-				{
-					//Inside bounding box - snap to convex
-					SnappedPosition = x - Phi * Normal;
-				}
-				else
-				{
-					//Snap to bounding box, and test convex
-					SnappedPosition = x - BoundingPhi * BoundingNormal;
-					Phi = PhiWithNormalInternal(SnappedPosition, Normal);
-					SnappedPosition -= Phi * Normal;
-				}
-
-				//one final snap to ensure we're on the surface
-				Phi = PhiWithNormalInternal(SnappedPosition, Normal);
-				SnappedPosition -= Phi * Normal;
-
-				//Return phi/normal based on distance from original position to snapped position
-				const FVec3 Difference = x - SnappedPosition;
-				Phi = Difference.Size();
-				if (CHAOS_ENSURE(Phi > TNumericLimits<float>::Min())) //Phi shouldn't be 0 here since we only enter this block if x was outside the convex
-				{
-					Normal = (Difference) / Phi;
-				}
-				else
-				{
-					Normal = FVector::ForwardVector;
-				}
-			}
-			return Phi;
+			return PhiWithNormalInternal(X, Normal);
 		}
 
 		FReal PhiWithNormalScaled(const FVec3& X, const FVec3& Scale, const FVec3& InvScale, FVec3& Normal) const
 		{
-			const FVec3 UnscaledX = InvScale * X;
-			FVec3 UnscaledNormal;
-			const FReal ScaledPhi = PhiWithNormalScaledInternal(UnscaledX, Scale.GetAbs(), UnscaledNormal);
-			Normal = (Scale * UnscaledNormal).GetSafeNormal();
-			return ScaledPhi;
+			return PhiWithNormalScaledInternal(X, Scale, Normal);
 		}
 
 
 	private:
 		// Distance to the surface
-		FReal PhiWithNormalInternal(const FVec3& x, FVec3& Normal) const
+		FReal PhiWithNormalInternal(const FVec3& X, FVec3& Normal) const
 		{
 			const int32 NumPlanes = Planes.Num();
 			if (NumPlanes == 0)
@@ -176,7 +137,7 @@ namespace Chaos
 
 			for (int32 Idx = 0; Idx < NumPlanes; ++Idx)
 			{
-				const FReal Phi = Planes[Idx].SignedDistance(x);
+				const FReal Phi = Planes[Idx].SignedDistance(X);
 				if (Phi > MaxPhi)
 				{
 					MaxPhi = Phi;
@@ -184,12 +145,49 @@ namespace Chaos
 				}
 			}
 
-			return Planes[MaxPlane].PhiWithNormal(x, Normal);
+			FReal Phi = Planes[MaxPlane].PhiWithNormal(X, Normal);
+			if (Phi <= 0)
+			{
+				return Phi;
+			}
+
+			// If x is outside the convex mesh, we should find for the nearest point to triangles on the plane
+			const int32 PlaneVerticesNum = NumPlaneVertices(MaxPlane);
+			const FVec3 XOnPlane = X - Phi * Normal;
+			FReal ClosestDistance = TNumericLimits<FReal>::Max();
+			FVec3 ClosestPoint;
+			for (int32 Index = 0; Index < PlaneVerticesNum - 2; Index++)
+			{
+				const FVec3 A(GetVertex(GetPlaneVertex(MaxPlane, 0)));
+				const FVec3 B(GetVertex(GetPlaneVertex(MaxPlane, Index + 1)));
+				const FVec3 C(GetVertex(GetPlaneVertex(MaxPlane, Index + 2)));
+
+				const FVec3 Point = FindClosestPointOnTriangle(XOnPlane, A, B, C, X);
+				if (XOnPlane == X)
+				{
+					return Phi;
+				}
+
+				const FReal Distance = (Point - XOnPlane).Size();
+				if (Distance < ClosestDistance)
+				{
+					ClosestDistance = Distance;
+					ClosestPoint = Point;
+				}
+			}
+
+			const TVector<FReal, 3> Difference = X - ClosestPoint;
+			Phi = Difference.Size();
+			if (Phi > SMALL_NUMBER)
+			{
+				Normal = (Difference) / Phi;
+			}
+			return Phi;
 		}
 
 		// Distance from a point to the surface for use in the scaled version. When the convex
-		// is scaled, we need to correct the depth calculation to take into account the world-space scale
-		FReal PhiWithNormalScaledInternal(const FVec3& LocalX, const FVec3& ScaleAbs, FVec3& LocalNormal) const
+		// is scaled, we need to bias the depth calculation to take into account the world-space scale
+		FReal PhiWithNormalScaledInternal(const FVec3& X, const FVec3& Scale, FVec3& Normal) const
 		{
 			const int32 NumPlanes = Planes.Num();
 			if (NumPlanes == 0)
@@ -199,20 +197,60 @@ namespace Chaos
 			check(NumPlanes > 0);
 
 			FReal MaxPhi = TNumericLimits<FReal>::Lowest();
-			FVec3 MaxNormal = FVec3(0);
-
+			FVec3 MaxNormal = FVec3(0,0,1);
+			int32 MaxPlane = 0;
 			for (int32 Idx = 0; Idx < NumPlanes; ++Idx)
 			{
-				const FReal PhiScale = FVec3::DotProduct(ScaleAbs, Planes[Idx].Normal().GetAbs());
-				const FReal Phi = Planes[Idx].SignedDistance(LocalX) * PhiScale;
-				if (Phi > MaxPhi)
+				FVec3 PlaneNormal = (Planes[Idx].Normal() * Scale).GetUnsafeNormal();
+				FVec3 PlanePos = Planes[Idx].X() * Scale;
+				FReal PlaneDistance = FVec3::DotProduct(X - PlanePos, PlaneNormal);
+				if (PlaneDistance > MaxPhi)
 				{
-					MaxPhi = Phi;
-					MaxNormal = Planes[Idx].Normal();
+					MaxPhi = PlaneDistance;
+					MaxNormal = PlaneNormal;
+					MaxPlane = Idx;
 				}
 			}
 
-			LocalNormal = MaxNormal;
+			Normal = MaxNormal;
+
+			if (MaxPhi < 0)
+			{
+				return MaxPhi;
+			}
+
+			// If X is outside the convex mesh, we should find for the nearest point to triangles on the plane
+			const int32 PlaneVerticesNum = NumPlaneVertices(MaxPlane);
+			const FVec3 XOnPlane = X - MaxPhi * Normal;
+			FReal ClosestDistance = TNumericLimits<FReal>::Max();
+			FVec3 ClosestPoint;
+			for (int32 Index = 0; Index < PlaneVerticesNum - 2; Index++)
+			{
+				const FVec3 A(Scale * GetVertex(GetPlaneVertex(MaxPlane, 0)));
+				const FVec3 B(Scale * GetVertex(GetPlaneVertex(MaxPlane, Index + 1)));
+				const FVec3 C(Scale * GetVertex(GetPlaneVertex(MaxPlane, Index + 2)));
+
+				const FVec3 Point = FindClosestPointOnTriangle(XOnPlane, A, B, C, X);
+				if (XOnPlane == X)
+				{
+					return MaxPhi;
+				}
+
+				const FReal Distance = (Point - XOnPlane).Size();
+				if (Distance < ClosestDistance)
+				{
+					ClosestDistance = Distance;
+					ClosestPoint = Point;
+				}
+			}
+
+			const FVec3 Difference = X - ClosestPoint;
+			const FReal DifferenceLen = Difference.Size();
+			if (DifferenceLen > SMALL_NUMBER)
+			{
+				Normal = Difference / DifferenceLen;
+				MaxPhi = DifferenceLen;
+			}
 			return MaxPhi;
 		}
 
@@ -258,6 +296,9 @@ namespace Chaos
 		// Whether the structure data has been created for this convex (will eventually always be true)
 		bool HasStructureData() const { return StructureData.IsValid(); }
 
+		// The convex structure data (mainly exposed for testing)
+		const FConvexStructureData& GetStructureData() const { return StructureData; }
+
 		// Get the index of the plane that most opposes the normal
 		int32 GetMostOpposingPlane(const FVec3& Normal) const;
 
@@ -666,7 +707,11 @@ namespace Chaos
 		{
 			TArray<TArray<int32>> FaceIndices;
 			FConvexBuilder::Simplify(Planes, FaceIndices, Vertices, LocalBoundingBox);
-			FConvexBuilder::MergeFaces(Planes, FaceIndices, Vertices);
+
+			// @todo(chaos): DistanceTolerance should be based on size, or passed in
+			const FReal DistanceTolerance = 1.0f;
+			FConvexBuilder::MergeFaces(Planes, FaceIndices, Vertices, DistanceTolerance);
+
 			CreateStructureData(MoveTemp(FaceIndices));
 		}
 
diff --git a/Engine/Source/Runtime/Experimental/Chaos/Public/Chaos/ConvexStructureData.h b/Engine/Source/Runtime/Experimental/Chaos/Public/Chaos/ConvexStructureData.h
--- a/Engine/Source/Runtime/Experimental/Chaos/Public/Chaos/ConvexStructureData.h
+++ b/Engine/Source/Runtime/Experimental/Chaos/Public/Chaos/ConvexStructureData.h
@@ -48,6 +48,11 @@ namespace Chaos
 			return PlaneVertices.Num() > 0;
 		}
 
+		FORCEINLINE int32 NumVertices() const
+		{
+			return VertexPlanesOffsetCount.Num();
+		}
+
 		// The number of planes that use the specified vertex
 		FORCEINLINE int32 NumVertexPlanes(int32 VertexIndex) const
 		{
@@ -63,6 +68,11 @@ namespace Chaos
 			return (int32)VertexPlanes[VertexPlaneFlatArrayIndex];
 		}
 
+		FORCEINLINE int32 NumPlanes() const
+		{
+			return PlaneVerticesOffsetCount.Num();
+		}
+
 		// The number of vertices that make up the corners of the specified face
 		FORCEINLINE int32 NumPlaneVertices(int32 PlaneIndex) const
 		{
@@ -172,20 +182,25 @@ namespace Chaos
 		TArray<FIndex> VertexPlanes;
 	};
 
-
+	using FConvexStructureDataS32 = TConvexStructureDataImp<int32, int32>;
+	using FConvexStructureDataU8 = TConvexStructureDataImp<uint8, uint16>;
 
 	// Metadata for a convex shape used by the manifold generation system and anything
 	// else that can benefit from knowing which vertices are associated with the faces.
 	class CHAOS_API FConvexStructureData
 	{
 	private:
-		FORCEINLINE const TConvexStructureDataImp<int32, int32>& Data32() const { return static_cast<TConvexStructureDataImp<int32, int32>&>(*Data); }
-		FORCEINLINE TConvexStructureDataImp<int32, int32>& Data32() { return static_cast<TConvexStructureDataImp<int32, int32>&>(*Data); }
-
-		FORCEINLINE const TConvexStructureDataImp<uint8, uint16>& Data8() const { return static_cast<TConvexStructureDataImp<uint8, uint16>&>(*Data); }
-		FORCEINLINE TConvexStructureDataImp<uint8, uint16>& Data8() { return static_cast<TConvexStructureDataImp<uint8, uint16>&>(*Data); }
+		FORCEINLINE FConvexStructureDataS32& Data32() { return static_cast<FConvexStructureDataS32&>(*Data); }
+		FORCEINLINE FConvexStructureDataU8& Data8() { return static_cast<FConvexStructureDataU8&>(*Data); }
 
 	public:
+		enum class EIndexType : int8
+		{
+			None,
+			S32,
+			U8,
+		};
+
 		FConvexStructureData()
 			: Data(nullptr)
 			, IndexType(EIndexType::None)
@@ -202,6 +217,25 @@ namespace Chaos
 			return (Data != nullptr);
 		}
 
+		FORCEINLINE EIndexType GetIndexType() const
+		{
+			return IndexType;
+		}
+
+		// Only exposed for unit tests
+		FORCEINLINE const FConvexStructureDataS32& Data32() const
+		{
+			checkSlow(IndexType == EIndexType::S32);
+			return static_cast<FConvexStructureDataS32&>(*Data);
+		}
+
+		// Only exposed for unit tests
+		FORCEINLINE const FConvexStructureDataU8& Data8() const
+		{
+			checkSlow(IndexType == EIndexType::U8);
+			return static_cast<FConvexStructureDataU8&>(*Data);
+		}
+
 		// The number of planes that use the specified vertex
 		FORCEINLINE int32 NumVertexPlanes(int32 VertexIndex) const
 		{
@@ -222,7 +256,7 @@ namespace Chaos
 		// Get the plane index (in the outer convex container) of one of the planes that uses the specified vertex
 		FORCEINLINE int32 GetVertexPlane(int32 VertexIndex, int32 VertexPlaneIndex) const
 		{
-			check(IsValid());
+			checkSlow(IsValid());
 			if (IndexType == EIndexType::S32)
 			{
 				return Data32().GetVertexPlane(VertexIndex, VertexPlaneIndex);
@@ -253,7 +287,7 @@ namespace Chaos
 		// Get the vertex index (in the outer convex container) of one of the vertices making up the corners of the specified face
 		FORCEINLINE int32 GetPlaneVertex(int32 PlaneIndex, int32 PlaneVertexIndex) const
 		{
-			check(IsValid());
+			checkSlow(IsValid());
 			if (IndexType == EIndexType::S32)
 			{
 				return Data32().GetPlaneVertex(PlaneIndex, PlaneVertexIndex);
@@ -323,13 +357,6 @@ namespace Chaos
 		}
 
 	private:
-		enum class EIndexType : int8
-		{
-			None,
-			S32,
-			U8,
-		};
-
 		// Determine the minimum index size we need for the specified convex size
 		EIndexType GetRequiredIndexType(int32 NumPlanes, int32 NumVerts)
 		{
diff --git a/Engine/Source/Runtime/Experimental/ChaosCore/Public/Chaos/Core.h b/Engine/Source/Runtime/Experimental/ChaosCore/Public/Chaos/Core.h
--- a/Engine/Source/Runtime/Experimental/ChaosCore/Public/Chaos/Core.h
+++ b/Engine/Source/Runtime/Experimental/ChaosCore/Public/Chaos/Core.h
@@ -34,5 +34,5 @@ namespace Chaos
 	// NOTE: if you get a merge conflict on the GUID, you must replace it with a new GUID - do not accept the source or target
 	// or you will likely get DDC version conflicts resulting in crashes during load.
 	// Core version string for Chaos data. Any DDC builder dependent on Chaos for serialization should depend on this version
-	static const FString ChaosVersionString = TEXT("F7030D07-09E0-4724-9B3E-8DBED3678C15");
+	static const FString ChaosVersionString = TEXT("5D50D1DC-7B1F-423C-9243-9DAD394603D1");
 }
diff --git a/Engine/Source/Runtime/PhysicsCore/Private/ChaosEngineInterface.cpp b/Engine/Source/Runtime/PhysicsCore/Private/ChaosEngineInterface.cpp
--- a/Engine/Source/Runtime/PhysicsCore/Private/ChaosEngineInterface.cpp
+++ b/Engine/Source/Runtime/PhysicsCore/Private/ChaosEngineInterface.cpp
@@ -915,41 +915,42 @@ FPhysicsConstraintHandle FChaosEngineInterface::CreateConstraint(const FPhysicsA
 				bSwapped = true;
 				ValidParticle = InActorRef2;
 			}
+			if(ValidParticle->GetSolverBase())
+			{
+				FChaosScene* Scene = FChaosEngineInterface::GetCurrentScene(ValidParticle);
+
+				// Create kinematic actor to attach to joint
+				FPhysicsActorHandle KinematicEndPoint;
+				FActorCreationParams Params;
+				Params.bSimulatePhysics = false;
+				Params.bQueryOnly = false;
+				Params.Scene = Scene;
+				Params.bStatic = false;
+				Params.InitialTM = FTransform::Identity;
+				FChaosEngineInterface::CreateActor(Params, KinematicEndPoint);
+
+				// Chaos requires our particles have geometry.
+				auto Sphere = MakeUnique<Chaos::FImplicitSphere3>(FVector(0, 0, 0), 0);
+				KinematicEndPoint->GetGameThreadAPI().SetGeometry(MoveTemp(Sphere));
+				KinematicEndPoint->GetGameThreadAPI().SetUserData(nullptr);
 
-			FChaosScene* Scene = FChaosEngineInterface::GetCurrentScene(ValidParticle);
-
-			// Create kinematic actor to attach to joint
-			FPhysicsActorHandle KinematicEndPoint;
-			FActorCreationParams Params;
-			Params.bSimulatePhysics = false;
-			Params.bQueryOnly = false;
-			Params.Scene = Scene;
-			Params.bStatic = false;
-			Params.InitialTM = FTransform::Identity;
-			FChaosEngineInterface::CreateActor(Params, KinematicEndPoint);
-
-			// Chaos requires our particles have geometry.
-			auto Sphere = MakeUnique<Chaos::FImplicitSphere3>(FVector(0, 0, 0), 0);
-			KinematicEndPoint->GetGameThreadAPI().SetGeometry(MoveTemp(Sphere));
-			KinematicEndPoint->GetGameThreadAPI().SetUserData(nullptr);
+				auto* JointConstraint = new Chaos::FJointConstraint();
+				JointConstraint->SetKinematicEndPoint(KinematicEndPoint, Scene->GetSolver());
+				ConstraintRef.Constraint = JointConstraint;
 
-			auto* JointConstraint = new Chaos::FJointConstraint();
-			JointConstraint->SetKinematicEndPoint(KinematicEndPoint, Scene->GetSolver());
-			ConstraintRef.Constraint = JointConstraint;
+				JointConstraint->SetParticleProxies({ KinematicEndPoint, ValidParticle });
 
-			JointConstraint->SetParticleProxies({ KinematicEndPoint, ValidParticle });
+				Chaos::FJointConstraint::FTransformPair TransformPair = { InLocalFrame2, InLocalFrame1};
+				if (bSwapped)
+				{
+					Swap(TransformPair[0], TransformPair[1]);
+				}
+				JointConstraint->SetJointTransforms(TransformPair);
 
-			Chaos::FJointConstraint::FTransformPair TransformPair = { InLocalFrame2, InLocalFrame1};
-			if (bSwapped)
-			{
-				Swap(TransformPair[0], TransformPair[1]);
+				Chaos::FPhysicsSolver* Solver = ValidParticle->GetSolver<Chaos::FPhysicsSolver>();
+				checkSlow(Solver == KinematicEndPoint->GetSolver<Chaos::FPhysicsSolver>());
+				Solver->RegisterObject(JointConstraint);
 			}
-			JointConstraint->SetJointTransforms(TransformPair);
-
-			Chaos::FPhysicsSolver* Solver = ValidParticle->GetSolver<Chaos::FPhysicsSolver>();
-			checkSlow(Solver == KinematicEndPoint->GetSolver<Chaos::FPhysicsSolver>());
-			Solver->RegisterObject(JointConstraint);
-
 		}
 	}
 	return ConstraintRef;
