Geometry Collection -- Convex Hull utility class for managing cached convex hulls as a GeometryCollection managed array.

diff --git a/Engine/Source/Runtime/Experimental/Chaos/Private/GeometryCollection/GeometryCollection.cpp b/Engine/Source/Runtime/Experimental/Chaos/Private/GeometryCollection/GeometryCollection.cpp
--- a/Engine/Source/Runtime/Experimental/Chaos/Private/GeometryCollection/GeometryCollection.cpp
+++ b/Engine/Source/Runtime/Experimental/Chaos/Private/GeometryCollection/GeometryCollection.cpp
@@ -10,6 +10,7 @@
 #include "GeometryCollection/GeometryCollectionUtility.h"
 #include "GeometryCollection/GeometryCollectionProximityUtility.h"
 #include "GeometryCollection/GeometryCollectionClusteringUtility.h"
+#include "GeometryCollection/GeometryCollectionConvexUtility.h"
 
 #include <iostream>
 #include <fstream>
@@ -76,6 +77,7 @@ void FGeometryCollection::Construct()
 
 	// Material Group
 	AddExternalAttribute<FGeometryCollectionSection>("Sections", FGeometryCollection::MaterialGroup, Sections, FacesDependency);
+
 }
 
 
@@ -91,6 +93,8 @@ void FGeometryCollection::SetDefaults(FName Group, uint32 StartSize, uint32 NumE
 			StatusFlags[Idx] = 0;
 			InitialDynamicState[Idx] = static_cast<int32>(Chaos::EObjectStateType::Uninitialized);
 		}
+
+		FGeometryCollectionConvexUtility::SetDefaults(this, Group, StartSize, NumElements);
 	}
 }
 
@@ -205,7 +209,6 @@ int32 FGeometryCollection::AppendGeometry(const FGeometryCollection & Element, i
 			VertexCount[GeometryIndex + tdx] = ElementVertexCount[tdx];
 			TransformIndex[GeometryIndex + tdx] = BoneMaps[VertexStart[GeometryIndex + tdx]];
 			TransformToGeometryIndex[TransformIndex[GeometryIndex + tdx]] = GeometryIndex + tdx;
-
 		}
 	}
 	else // Element input failed to create a geometry group
@@ -468,6 +471,9 @@ void FGeometryCollection::RemoveElements(const FName & Group, const TArray<int32
 
 			RemoveGeometryElements(GeometryIndices);
 
+			// Find convex hulls connected to transform
+			FGeometryCollectionConvexUtility::RemoveConvexHulls(this, SortedDeletionList);
+
 			Super::RemoveElements(Group, SortedDeletionList);
 		}
 		else if (Group == FGeometryCollection::GeometryGroup)
diff --git a/Engine/Source/Runtime/Experimental/Chaos/Private/GeometryCollection/GeometryCollectionConvexUtility.cpp b/Engine/Source/Runtime/Experimental/Chaos/Private/GeometryCollection/GeometryCollectionConvexUtility.cpp
new file mode 100644
--- /dev/null
+++ b/Engine/Source/Runtime/Experimental/Chaos/Private/GeometryCollection/GeometryCollectionConvexUtility.cpp
@@ -0,0 +1,134 @@
+// Copyright Epic Games, Inc. All Rights Reserved.
+
+#include "GeometryCollection/GeometryCollectionConvexUtility.h"
+#include "Chaos/Convex.h"
+#include "GeometryCollection/GeometryCollection.h"
+
+
+FGeometryCollectionConvexUtility::FGeometryCollectionConvexData FGeometryCollectionConvexUtility::GetValidConvexHullData(FGeometryCollection* GeometryCollection)
+{
+	check (GeometryCollection)
+
+	if (!GeometryCollection->HasGroup("Convex"))
+	{
+		GeometryCollection->AddGroup("Convex");
+	}
+
+	if (!GeometryCollection->HasAttribute("TransformToConvexIndex", FTransformCollection::TransformGroup))
+	{
+		FManagedArrayCollection::FConstructionParameters ConvexDependency("Convex");
+		TManagedArray<int32>& IndexArray = GeometryCollection->AddAttribute<int32>("TransformToConvexIndex", FTransformCollection::TransformGroup, ConvexDependency);
+		IndexArray.Fill(INDEX_NONE);
+	}
+
+	if (!GeometryCollection->HasAttribute("ConvexHull", "Convex"))
+	{
+		GeometryCollection->AddAttribute<TUniquePtr<Chaos::FConvex>>("ConvexHull", "Convex");
+	}
+
+	// Check for correct population. All (and only) rigid nodes should have a convex associated.
+	const TManagedArray<int32>& SimulationType = GeometryCollection->GetAttribute<int32>("SimulationType", FTransformCollection::TransformGroup);
+	const TManagedArray<int32>& TransformToGeometryIndex = GeometryCollection->GetAttribute<int32>("TransformToGeometryIndex", FTransformCollection::TransformGroup);
+	TManagedArray<int32>& TransformToConvexIndex = GeometryCollection->GetAttribute<int32>("TransformToConvexIndex", FTransformCollection::TransformGroup);
+	TManagedArray<TUniquePtr<Chaos::FConvex>>& ConvexHull = GeometryCollection->GetAttribute<TUniquePtr<Chaos::FConvex>>("ConvexHull", "Convex");
+
+	TArray<int32> ProduceConvexHulls;
+	ProduceConvexHulls.Reserve(SimulationType.Num());
+	TArray<int32> EliminateConvexHulls;
+	EliminateConvexHulls.Reserve(SimulationType.Num());
+
+	for (int32 Idx = 0; Idx < SimulationType.Num(); ++Idx)
+	{
+		if ((SimulationType[Idx] == FGeometryCollection::ESimulationTypes::FST_Rigid) && (TransformToConvexIndex[Idx] == INDEX_NONE))
+		{
+			ProduceConvexHulls.Add(Idx);
+			
+		}
+		else if ((SimulationType[Idx] != FGeometryCollection::ESimulationTypes::FST_Rigid) && (TransformToConvexIndex[Idx] > INDEX_NONE))
+		{
+			EliminateConvexHulls.Add(Idx);
+		}
+	}
+
+	if (EliminateConvexHulls.Num())
+	{
+		GeometryCollection->RemoveElements("Convex", EliminateConvexHulls);
+		for (int32 Idx : EliminateConvexHulls)
+		{
+			TransformToConvexIndex[Idx] = INDEX_NONE;
+		}
+	}
+
+	if (ProduceConvexHulls.Num())
+	{
+		int32 NewConvexIndexStart = GeometryCollection->AddElements(ProduceConvexHulls.Num(), "Convex");
+		for (int32 Idx = 0; Idx < ProduceConvexHulls.Num(); ++Idx)
+		{
+			int32 GeometryIdx = TransformToGeometryIndex[ProduceConvexHulls[Idx]];
+			ConvexHull[NewConvexIndexStart + Idx] = FindConvexHull(GeometryCollection, GeometryIdx);
+			TransformToConvexIndex[ProduceConvexHulls[Idx]] = NewConvexIndexStart + Idx;
+		}
+	}
+
+	return { TransformToConvexIndex, ConvexHull };
+	
+}
+
+
+TUniquePtr<Chaos::FConvex> FGeometryCollectionConvexUtility::FindConvexHull(const FGeometryCollection* GeometryCollection, int32 GeometryIndex)
+{
+	check(GeometryCollection);
+
+	int32 VertexCount = GeometryCollection->VertexCount[GeometryIndex];
+	int32 VertexStart = GeometryCollection->VertexStart[GeometryIndex];
+
+	TArray<Chaos::FVec3> Vertices;
+	Vertices.SetNum(VertexCount);
+	for (int32 VertexIndex = 0; VertexIndex < VertexCount; ++VertexIndex)
+	{
+		Vertices[VertexIndex] = GeometryCollection->Vertex[VertexStart+VertexIndex];
+	}
+
+	return MakeUnique<Chaos::FConvex>(Vertices, 0.0f);
+}
+
+
+void FGeometryCollectionConvexUtility::RemoveConvexHulls(FGeometryCollection* GeometryCollection, const TArray<int32>& SortedTransformDeletes)
+{
+	if (GeometryCollection->HasGroup("Convex") && GeometryCollection->HasAttribute("TransformToConvexIndex", FTransformCollection::TransformGroup))
+	{
+		TManagedArray<int32>& TransformToConvexIndex = GeometryCollection->GetAttribute<int32>("TransformToConvexIndex", FTransformCollection::TransformGroup);
+		TArray<int32> ConvexIndices;
+		for (int32 TransformIdx : SortedTransformDeletes)
+		{
+			if (TransformToConvexIndex[TransformIdx] > INDEX_NONE)
+			{
+				ConvexIndices.Add(TransformToConvexIndex[TransformIdx]);
+				TransformToConvexIndex[TransformIdx] = INDEX_NONE;
+			}
+		}
+
+		if (ConvexIndices.Num())
+		{
+			ConvexIndices.Sort();
+			GeometryCollection->RemoveElements("Convex", ConvexIndices);
+		}
+	}
+}
+
+
+void FGeometryCollectionConvexUtility::SetDefaults(FGeometryCollection* GeometryCollection, FName Group, uint32 StartSize, uint32 NumElements)
+{
+	if (Group == FTransformCollection::TransformGroup)
+	{
+		if (GeometryCollection->HasAttribute("TransformToConvexIndex", FTransformCollection::TransformGroup))
+		{
+			TManagedArray<int32>& TransformToConvexIndex = GeometryCollection->GetAttribute<int32>("TransformToConvexIndex", FTransformCollection::TransformGroup);
+
+			for (uint32 Idx = StartSize; Idx < StartSize + NumElements; ++Idx)
+			{
+				TransformToConvexIndex[Idx] = INDEX_NONE;
+			}
+		}
+	}
+}
\ No newline at end of file
diff --git a/Engine/Source/Runtime/Experimental/Chaos/Public/GeometryCollection/GeometryCollectionConvexUtility.h b/Engine/Source/Runtime/Experimental/Chaos/Public/GeometryCollection/GeometryCollectionConvexUtility.h
new file mode 100644
--- /dev/null
+++ b/Engine/Source/Runtime/Experimental/Chaos/Public/GeometryCollection/GeometryCollectionConvexUtility.h
@@ -0,0 +1,38 @@
+// Copyright Epic Games, Inc. All Rights Reserved.
+
+#pragma once
+
+#include "CoreMinimal.h"
+#include "GeometryCollection/ManagedArrayCollection.h"
+
+class FGeometryCollection;
+
+namespace Chaos
+{
+	class FConvex;
+}
+
+class CHAOS_API FGeometryCollectionConvexUtility
+{
+public:
+
+	struct FGeometryCollectionConvexData
+	{
+		TManagedArray<int32>& TransformToConvexIndex;
+		TManagedArray<TUniquePtr<Chaos::FConvex>>& ConvexHull;
+	};
+
+	/** Ensure that convex hull data exists for the Geometry Collection and construct it if not (or if some data is missing. */
+	static FGeometryCollectionConvexData GetValidConvexHullData(FGeometryCollection* GeometryCollection);
+
+	/** Returns the convex hull of the vertices contained in the specified geometry. */
+	static TUniquePtr<Chaos::FConvex> FindConvexHull(const FGeometryCollection* GeometryCollection, int32 GeometryIndex);
+
+	/** Delete the convex hulls pointed at by the transform indices provided. */
+	static void RemoveConvexHulls(FGeometryCollection* GeometryCollection, const TArray<int32>& SortedTransformDeletes);
+
+	/** Set default values for convex hull related managed arrays. */
+	static void SetDefaults(FGeometryCollection* GeometryCollection, FName Group, uint32 StartSize, uint32 NumElements);
+
+};
+
diff --git a/Engine/Source/Runtime/Experimental/Chaos/Public/GeometryCollection/ManagedArray.h b/Engine/Source/Runtime/Experimental/Chaos/Public/GeometryCollection/ManagedArray.h
--- a/Engine/Source/Runtime/Experimental/Chaos/Public/GeometryCollection/ManagedArray.h
+++ b/Engine/Source/Runtime/Experimental/Chaos/Public/GeometryCollection/ManagedArray.h
@@ -580,7 +580,14 @@ public:
 			{
 				ensure(RemapVal < MaskSize);
 				this->operator[](Index) -= Offsets[RemapVal];
-				ensure(-1 <= this->operator[](Index) && this->operator[](Index) < FinalSize);
+
+				// #todo Reindexing is currently incorrect in a few cases where it leaves dangling indices that
+				// should be set to INDEX_NONE.
+				// This ensure is currently suppressed to handle the dangling index problem specifically 
+				// in the ConvexGroup case (GeometryCollectionConvexUtility::RemoveConvexHulls) but should be reinstated 
+				// once we have a better solution in place for Reindexing.
+				//ensure(-1 <= this->operator[](Index) && this->operator[](Index) < FinalSize);
+				ensure(-1 <= this->operator[](Index));
 			}
 		}
 	}
diff --git a/Engine/Source/Runtime/Experimental/Chaos/Public/GeometryCollection/ManagedArrayTypeValues.inl b/Engine/Source/Runtime/Experimental/Chaos/Public/GeometryCollection/ManagedArrayTypeValues.inl
--- a/Engine/Source/Runtime/Experimental/Chaos/Public/GeometryCollection/ManagedArrayTypeValues.inl
+++ b/Engine/Source/Runtime/Experimental/Chaos/Public/GeometryCollection/ManagedArrayTypeValues.inl
@@ -39,6 +39,7 @@ MANAGED_ARRAY_TYPE(TUniquePtr<Chaos::TGeometryParticleFloat3>, TGeometryParticle
 MANAGED_ARRAY_TYPE(Chaos::ThreadSafeSharedPtr_FImplicitObject, FImplicitObject3ThreadSafeSharedPointer)
 MANAGED_ARRAY_TYPE(Chaos::NotThreadSafeSharedPtr_FImplicitObject, FImplicitObject3SharedPointer)
 MANAGED_ARRAY_TYPE(Chaos::TPBDRigidClusteredParticleHandleFloat3*, TPBDRigidClusteredParticleHandle3fPtr)
+MANAGED_ARRAY_TYPE(TUniquePtr<Chaos::FConvex>, FConvexUniquePtr)
 
 // NOTE: new types must be added at the bottom to keep serialization from breaking
 
diff --git a/Engine/Source/Runtime/Experimental/Chaos/Public/GeometryCollection/ManagedArrayTypes.h b/Engine/Source/Runtime/Experimental/Chaos/Public/GeometryCollection/ManagedArrayTypes.h
--- a/Engine/Source/Runtime/Experimental/Chaos/Public/GeometryCollection/ManagedArrayTypes.h
+++ b/Engine/Source/Runtime/Experimental/Chaos/Public/GeometryCollection/ManagedArrayTypes.h
@@ -15,6 +15,7 @@
 #include "Chaos/BVHParticles.h"
 #include "Chaos/ParticleHandle.h"
 #include "Chaos/ParticleHandleFwd.h"
+#include "Chaos/Convex.h"
 
 inline FArchive& operator<<(FArchive& Ar, TArray<FVector>*& ValueIn)
 {
