Set Chaos::MaxDirtyElements to MAX_int32 in editor builds to avoid costly ReoptimizeTree calls during various operations.

diff --git a/Engine/Source/Runtime/Experimental/Chaos/Private/Chaos/BoundingVolume.cpp b/Engine/Source/Runtime/Experimental/Chaos/Private/Chaos/BoundingVolume.cpp
--- a/Engine/Source/Runtime/Experimental/Chaos/Private/Chaos/BoundingVolume.cpp
+++ b/Engine/Source/Runtime/Experimental/Chaos/Private/Chaos/BoundingVolume.cpp
@@ -15,7 +15,12 @@ FAutoConsoleVariableRef FBoundingVolumeCVars::CVarFilterFarBodies(
 
 namespace Chaos
 {
+#if WITH_EDITOR
+	CHAOS_API int32 MaxDirtyElements = MAX_int32;
+#else
 	CHAOS_API int32 MaxDirtyElements = 10000;
+#endif
+
 	FAutoConsoleVariableRef CVarMaxDirtyElements(
 	    TEXT("p.MaxDirtyElements"),
 	    MaxDirtyElements,
