Chaos : GeometryCollection
-Removed binding of unique indices between the game thread Particle and the physics thread ParticleHandle

diff --git a/Engine/Source/Runtime/Experimental/Chaos/Private/PhysicsProxy/GeometryCollectionPhysicsProxy.cpp b/Engine/Source/Runtime/Experimental/Chaos/Private/PhysicsProxy/GeometryCollectionPhysicsProxy.cpp
--- a/Engine/Source/Runtime/Experimental/Chaos/Private/PhysicsProxy/GeometryCollectionPhysicsProxy.cpp
+++ b/Engine/Source/Runtime/Experimental/Chaos/Private/PhysicsProxy/GeometryCollectionPhysicsProxy.cpp
@@ -657,7 +657,6 @@ void FGeometryCollectionPhysicsProxy::InitializeBodiesPT(Chaos::TPBDRigidsSolver
 
 				// We're on the physics thread here but we've already set up the GT particles and we're just linking here
 				Handle->GTGeometryParticle() = GTParticles[Idx].Get();
-				Handle->SetUniqueIdx(GTParticles[Idx]->UniqueIdx());
 
 				check(SolverParticleHandles[Idx]->GetParticleType() == Handle->GetParticleType());
 				RigidsSolver->GetEvolution()->CreateParticle(Handle);
