ChaosVehicles: Off behaviour of vehicle suspension compressing at speed, calculate compression speed by a different method. Downward compression causing vehicle to catch on the terrain when travelling at high speeds.
Relates to FROST-1129 but doesn't comeplety solve.

diff --git a/Engine/Source/Runtime/Experimental/ChaosVehicles/ChaosVehiclesCore/Private/SuspensionSystem.cpp b/Engine/Source/Runtime/Experimental/ChaosVehicles/ChaosVehiclesCore/Private/SuspensionSystem.cpp
--- a/Engine/Source/Runtime/Experimental/ChaosVehicles/ChaosVehiclesCore/Private/SuspensionSystem.cpp
+++ b/Engine/Source/Runtime/Experimental/ChaosVehicles/ChaosVehiclesCore/Private/SuspensionSystem.cpp
@@ -2,8 +2,69 @@
 
 #include "SuspensionSystem.h"
 
+static bool GNewSuspensionSim = true; // TODO: remove once sure new calculation is correct
+FAutoConsoleVariableRef CVarChaosVehiclesABTestSuspension(TEXT("p.Vehicle.NewSuspensionSim"), GNewSuspensionSim, TEXT("A/B Test Suspension Simulation."));
+
 namespace Chaos
 {
 
+	float FSimpleSuspensionSim::GetSpringLength()
+	{
+		if (Setup().SuspensionSmoothing)
+		{
+			// Trying smoothing the suspension movement out - looks Sooo much better when wheel traveling over pile of bricks
+			// The digital up and down of the wheels is slowed/smoothed out
+			float NewValue = SpringDisplacement - Setup().MaxLength;
+
+			if (AveragingNum < Setup().SuspensionSmoothing)
+			{
+				AveragingNum++;
+			}
+
+			AveragingLength[AveragingCount++] = NewValue;
+
+			if (AveragingCount >= Setup().SuspensionSmoothing)
+			{
+				AveragingCount = 0;
+			}
+
+			float Total = 0.0f;
+			for (int i = 0; i < AveragingNum; i++)
+			{
+				Total += AveragingLength[i];
+			}
+			float Average = Total / AveragingNum;
+
+			return Average;
+		}
+		else
+		{
+			return  (SpringDisplacement - Setup().MaxLength);
+		}
+	}
+
+	void FSimpleSuspensionSim::Simulate(float DeltaTime)
+	{
+		if (GNewSuspensionSim)
+		{
+			float Damping = (SpringDisplacement < LastDisplacement) ? Setup().CompressionDamping : Setup().ReboundDamping;
+			float SpringSpeed = (LastDisplacement - SpringDisplacement) / DeltaTime;
+
+			const float StiffnessForce = SpringDisplacement * Setup().SpringRate;
+			const float DampingForce = SpringSpeed * Damping;
+			SuspensionForce = StiffnessForce - DampingForce;
+			LastDisplacement = SpringDisplacement;
+		}
+		else
+		{
+			float Damping = (DisplacementInput < LastDisplacement) ? Setup().CompressionDamping : Setup().ReboundDamping;
+
+			const float StiffnessForce = SpringDisplacement * Setup().SpringRate;
+			const float DampingForce = LocalVelocity.Z * Damping;
+			SuspensionForce = StiffnessForce - DampingForce;
+			LastDisplacement = DisplacementInput;
+		}
+	}
 
 } // namespace Chaos
+
diff --git a/Engine/Source/Runtime/Experimental/ChaosVehicles/ChaosVehiclesCore/Public/SuspensionSystem.h b/Engine/Source/Runtime/Experimental/ChaosVehicles/ChaosVehiclesCore/Public/SuspensionSystem.h
--- a/Engine/Source/Runtime/Experimental/ChaosVehicles/ChaosVehiclesCore/Public/SuspensionSystem.h
+++ b/Engine/Source/Runtime/Experimental/ChaosVehicles/ChaosVehiclesCore/Public/SuspensionSystem.h
@@ -163,40 +163,7 @@ namespace Chaos
 
 // Outputs
 
-		float GetSpringLength()
-		{
-			if (Setup().SuspensionSmoothing)
-			{
-				// Trying smoothing the suspension movement out - looks Sooo much better when wheel traveling over pile of bricks
-				// The digital up and down of the wheels is slowed/smoothed out
-				float NewValue = SpringDisplacement - Setup().MaxLength;
-
-				if (AveragingNum < Setup().SuspensionSmoothing)
-				{
-					AveragingNum++;
-				}
-
-				AveragingLength[AveragingCount++] = NewValue;
-
-				if (AveragingCount >= Setup().SuspensionSmoothing)
-				{
-					AveragingCount = 0;
-				}
-
-				float Total = 0.0f;
-				for (int i = 0; i < AveragingNum; i++)
-				{
-					Total += AveragingLength[i];
-				}
-				float Average = Total / AveragingNum;
-
-				return Average;
-			}
-			else
-			{
-				return  (SpringDisplacement - Setup().MaxLength);
-			}
-		}
+		float GetSpringLength();
 
 		float GetSuspensionForce() const
 		{
@@ -221,15 +188,7 @@ namespace Chaos
 
 // Simulation
 
-		void Simulate(float DeltaTime)
-		{
-			float Damping = (DisplacementInput < LastDisplacement) ? Setup().CompressionDamping : Setup().ReboundDamping;
-
-			const float StiffnessForce = SpringDisplacement * Setup().SpringRate;
-			const float DampingForce = LocalVelocity.Z * Damping;
-			SuspensionForce = StiffnessForce - DampingForce;
-			LastDisplacement = DisplacementInput;
-		}
+		void Simulate(float DeltaTime);
 
 	protected:
 
