Optimized triangle mesh collision updates to avoid doing unecessary tree rebuilds, should be almost twice as fast now.

diff --git a/Engine/Source/Runtime/Experimental/Chaos/Private/Chaos/TriangleMeshImplicitObject.cpp b/Engine/Source/Runtime/Experimental/Chaos/Private/Chaos/TriangleMeshImplicitObject.cpp
--- a/Engine/Source/Runtime/Experimental/Chaos/Private/Chaos/TriangleMeshImplicitObject.cpp
+++ b/Engine/Source/Runtime/Experimental/Chaos/Private/Chaos/TriangleMeshImplicitObject.cpp
@@ -898,7 +898,7 @@ TUniquePtr<FTriangleMeshImplicitObject> FTriangleMeshImplicitObject::CopySlowImp
 		ExternalVertexIndexMapCopy = MakeUnique<TArray<int32>>(*ExternalVertexIndexMap.Get());
 	}
 
-	return MakeUnique<FTriangleMeshImplicitObject>(MoveTemp(ParticlesCopy), MoveTemp(ElementsCopy), MoveTemp(MaterialIndicesCopy), MoveTemp(ExternalFaceIndexMapCopy), MoveTemp(ExternalVertexIndexMapCopy), bCullsBackFaceRaycast);
+	return TUniquePtr<FTriangleMeshImplicitObject>(new FTriangleMeshImplicitObject(MoveTemp(ParticlesCopy), MoveTemp(ElementsCopy), MoveTemp(MaterialIndicesCopy), BVH, MoveTemp(ExternalFaceIndexMapCopy), MoveTemp(ExternalVertexIndexMapCopy), bCullsBackFaceRaycast));
 }
 
 TUniquePtr<FTriangleMeshImplicitObject> FTriangleMeshImplicitObject::CopySlow() const
diff --git a/Engine/Source/Runtime/Experimental/Chaos/Public/Chaos/AABBTree.h b/Engine/Source/Runtime/Experimental/Chaos/Public/Chaos/AABBTree.h
--- a/Engine/Source/Runtime/Experimental/Chaos/Public/Chaos/AABBTree.h
+++ b/Engine/Source/Runtime/Experimental/Chaos/Public/Chaos/AABBTree.h
@@ -395,6 +395,11 @@ public:
 
 	virtual ~TAABBTree() {}
 
+	void CopyFrom(const TAABBTree<TPayloadType, TLeafType, T, bMutable>& Other)
+	{
+		(*this) = Other;
+	}
+
 	virtual TUniquePtr<ISpatialAcceleration<TPayloadType, T, 3>> Copy() const override
 	{
 		return TUniquePtr<ISpatialAcceleration<TPayloadType, T, 3>>(new TAABBTree<TPayloadType, TLeafType, T, bMutable>(*this));
diff --git a/Engine/Source/Runtime/Experimental/Chaos/Public/Chaos/TriangleMeshImplicitObject.h b/Engine/Source/Runtime/Experimental/Chaos/Public/Chaos/TriangleMeshImplicitObject.h
--- a/Engine/Source/Runtime/Experimental/Chaos/Public/Chaos/TriangleMeshImplicitObject.h
+++ b/Engine/Source/Runtime/Experimental/Chaos/Public/Chaos/TriangleMeshImplicitObject.h
@@ -300,6 +300,25 @@ namespace Chaos
 		//using BVHType = TBoundingVolume<int32, T, 3>;
 		using BVHType = TAABBTree<int32, TAABBTreeLeafArray<int32, FReal, /*bComputeBounds=*/false>, FReal, /*bMutable=*/false>;
 
+		// Initialising constructor privately declared for use in CopySlow to copy the underlying BVH
+		template <typename IdxType>
+		FTriangleMeshImplicitObject(TParticles<FReal, 3>&& Particles, TArray<TVec3<IdxType>>&& Elements, TArray<uint16>&& InMaterialIndices, const BVHType& InBvhToCopy, TUniquePtr<TArray<int32>>&& InExternalFaceIndexMap = nullptr, TUniquePtr<TArray<int32>>&& InExternalVertexIndexMap = nullptr, const bool bInCullsBackFaceRaycast = false)
+			: FImplicitObject(EImplicitObject::HasBoundingBox | EImplicitObject::DisableCollisions, ImplicitObjectType::TriangleMesh)
+			, MParticles(MoveTemp(Particles))
+			, MElements(MoveTemp(Elements))
+			, MLocalBoundingBox(MParticles.X(0), MParticles.X(0))
+			, MaterialIndices(MoveTemp(InMaterialIndices))
+			, ExternalFaceIndexMap(MoveTemp(InExternalFaceIndexMap))
+			, ExternalVertexIndexMap(MoveTemp(InExternalVertexIndexMap))
+			, bCullsBackFaceRaycast(bInCullsBackFaceRaycast)
+		{
+			for(uint32 Idx = 1; Idx < MParticles.Size(); ++Idx)
+			{
+				MLocalBoundingBox.GrowToInclude(MParticles.X(Idx));
+			}
+			BVH.CopyFrom(InBvhToCopy);
+		}
+
 		template<typename InStorageType, typename InRealType>
 		friend struct FBvEntry;
 
