Fix GeometryCollection making duplicate entries in acceleration structure.

diff --git a/Engine/Source/Runtime/Experimental/Chaos/Private/PhysicsProxy/GeometryCollectionPhysicsProxy.cpp b/Engine/Source/Runtime/Experimental/Chaos/Private/PhysicsProxy/GeometryCollectionPhysicsProxy.cpp
--- a/Engine/Source/Runtime/Experimental/Chaos/Private/PhysicsProxy/GeometryCollectionPhysicsProxy.cpp
+++ b/Engine/Source/Runtime/Experimental/Chaos/Private/PhysicsProxy/GeometryCollectionPhysicsProxy.cpp
@@ -629,17 +629,25 @@ void FGeometryCollectionPhysicsProxy::InitializeBodiesPT(Chaos::TPBDRigidsSolver
 		int NumRigids = 0;
 		BaseParticleIndex = NumRigids;
 
+		// Gather unique indices from GT to pass into PT handle creation
+		TArray<Chaos::FUniqueIdx> UniqueIndices;
+		UniqueIndices.Reserve(SimulatableParticles.Num());
+
 		// Count geometry collection leaf node particles to add
 		int NumSimulatedParticles = 0;
 		for (int32 Idx = 0; Idx < SimulatableParticles.Num(); ++Idx)
 		{
 			NumSimulatedParticles += SimulatableParticles[Idx];
-			NumRigids += (SimulatableParticles[Idx] && !RestCollection->IsClustered(Idx));
+			if (SimulatableParticles[Idx] && !RestCollection->IsClustered(Idx))
+			{
+				NumRigids++;
+				UniqueIndices.Add(GTParticles[Idx]->UniqueIdx());
+			}
 		}
 
 		// Add entries into simulation array
 		RigidsSolver->GetEvolution()->ReserveParticles(NumSimulatedParticles);
-		TArray<Chaos::TPBDGeometryCollectionParticleHandle<float, 3>*> Handles = RigidsSolver->GetEvolution()->CreateGeometryCollectionParticles(NumRigids);
+		TArray<Chaos::TPBDGeometryCollectionParticleHandle<float, 3>*> Handles = RigidsSolver->GetEvolution()->CreateGeometryCollectionParticles(NumRigids, UniqueIndices.GetData());
 
 		int32 NextIdx = 0;
 		for (int32 Idx = 0; Idx < SimulatableParticles.Num(); ++Idx)
