Chaos
- sweeps now use reduced or no margins where possible
- explicit Sphere-Convex collision path
- explicit Capsule-Convex collision path
- ImplicitObject SupportCore uses passed-in margin only so shapes can be used with different margins
- Added wrapper object fro use in GJKPenetration that allows us to select whether to use the core shape or outer shape

diff --git a/Engine/Plugins/Experimental/ChaosCloth/Source/Chaos/Private/ChaosCloth/ChaosClothingSimulation.cpp b/Engine/Plugins/Experimental/ChaosCloth/Source/Chaos/Private/ChaosCloth/ChaosClothingSimulation.cpp
--- a/Engine/Plugins/Experimental/ChaosCloth/Source/Chaos/Private/ChaosCloth/ChaosClothingSimulation.cpp
+++ b/Engine/Plugins/Experimental/ChaosCloth/Source/Chaos/Private/ChaosCloth/ChaosClothingSimulation.cpp
@@ -15,6 +15,7 @@
 #if WITH_EDITOR || CHAOS_DEBUG_DRAW
 #include "Chaos/Sphere.h"
 #include "Chaos/TaperedCylinder.h"
+#include "Chaos/Capsule.h"
 #include "Chaos/Convex.h"
 #include "Chaos/PBDSphericalConstraint.h"
 #include "Chaos/PBDAnimDriveConstraint.h"
diff --git a/Engine/Plugins/Experimental/ChaosCloth/Source/Chaos/Private/ChaosCloth/ChaosClothingSimulationCollider.cpp b/Engine/Plugins/Experimental/ChaosCloth/Source/Chaos/Private/ChaosCloth/ChaosClothingSimulationCollider.cpp
--- a/Engine/Plugins/Experimental/ChaosCloth/Source/Chaos/Private/ChaosCloth/ChaosClothingSimulationCollider.cpp
+++ b/Engine/Plugins/Experimental/ChaosCloth/Source/Chaos/Private/ChaosCloth/ChaosClothingSimulationCollider.cpp
@@ -3,6 +3,7 @@
 #include "ChaosCloth/ChaosClothingSimulationSolver.h"
 #include "ChaosCloth/ChaosClothingSimulationCloth.h"
 #include "ChaosCloth/ChaosClothPrivate.h"
+#include "Chaos/Capsule.h"
 #include "Chaos/Plane.h"
 #include "Chaos/Sphere.h"
 #include "Chaos/TaperedCylinder.h"
diff --git a/Engine/Source/Programs/HeadlessChaos/Private/HeadlessChaosTestGJKDistance.cpp b/Engine/Source/Programs/HeadlessChaos/Private/HeadlessChaosTestGJKDistance.cpp
--- a/Engine/Source/Programs/HeadlessChaos/Private/HeadlessChaosTestGJKDistance.cpp
+++ b/Engine/Source/Programs/HeadlessChaos/Private/HeadlessChaosTestGJKDistance.cpp
@@ -19,20 +19,21 @@ namespace ChaosTest
 		FVec3 NearestA = { 0,0,0 };
 		FVec3 NearestB = { 0,0,0 };
 		FReal Distance = 0;
+		FVec3 Normal = {0, 0, 1};
 		// Fail - overlapping
 		{
 			TSphere<FReal, 3> A(FVec3(12, 0, 0), 5);
 			TSphere<FReal, 3> B(FVec3(4, 0, 0), 2);
-			bool bSuccess = GJKDistance<FReal>(A, B, FRigidTransform3(FVec3(2, 0, 0), FRotation3::FromIdentity()), Distance, NearestA, NearestB);
-			EXPECT_FALSE(bSuccess);
+			EGJKDistanceResult Result = GJKDistance<FReal>(A, B, FRigidTransform3(FVec3(2, 0, 0), FRotation3::FromIdentity()), Distance, NearestA, NearestB, Normal);
+			EXPECT_NE(Result, EGJKDistanceResult::Separated);
 		}
 
 		// Success - not overlapping
 		{
 			TSphere<FReal, 3> A(FVec3(12, 0, 0), 5);
 			TSphere<FReal, 3> B(FVec3(4, 0, 0), 2);
-			bool bSuccess = GJKDistance<FReal>(A, B, FRigidTransform3::Identity, Distance, NearestA, NearestB);
-			EXPECT_TRUE(bSuccess);
+			EGJKDistanceResult Result = GJKDistance<FReal>(A, B, FRigidTransform3::Identity, Distance, NearestA, NearestB, Normal);
+			EXPECT_EQ(Result, EGJKDistanceResult::Separated);
 			EXPECT_NEAR(Distance, (FReal)1, Tolerance);
 			EXPECT_NEAR(NearestA.X, (FReal)7, Tolerance);
 			EXPECT_NEAR(NearestA.Y, (FReal)0, Tolerance);
@@ -47,8 +48,8 @@ namespace ChaosTest
 			TSphere<FReal, 3> A(FVec3(0, 0, 0), 2);
 			TSphere<FReal, 3> B(FVec3(0, 0, 0), 2);
 			FVec3 BPos = FVec3(3, 3, 0);
-			bool bSuccess = GJKDistance<FReal>(A, B, FRigidTransform3(BPos, FRotation3::FromIdentity()), Distance, NearestA, NearestB);
-			EXPECT_TRUE(bSuccess);
+			EGJKDistanceResult Result = GJKDistance<FReal>(A, B, FRigidTransform3(BPos, FRotation3::FromIdentity()), Distance, NearestA, NearestB, Normal);
+			EXPECT_EQ(Result, EGJKDistanceResult::Separated);
 			FVec3 CenterDelta = (B.GetCenter() + BPos) - A.GetCenter();
 			FVec3 CenterDir = CenterDelta.GetSafeNormal();
 			EXPECT_NEAR(Distance, CenterDelta.Size() - (A.GetRadius() + B.GetRadius()), Tolerance);
@@ -65,8 +66,8 @@ namespace ChaosTest
 			TSphere<FReal, 3> A(FVec3(12, 0, 0), 5);
 			TSphere<FReal, 3> B(FVec3(4, 0, 0), 2);
 			FVec3 BPos = FVec3(0.99, 0, 0);
-			bool bSuccess = GJKDistance<FReal>(A, B, FRigidTransform3(BPos, FRotation3::FromIdentity()), Distance, NearestA, NearestB);
-			EXPECT_TRUE(bSuccess);
+			EGJKDistanceResult Result = GJKDistance<FReal>(A, B, FRigidTransform3(BPos, FRotation3::FromIdentity()), Distance, NearestA, NearestB, Normal);
+			EXPECT_EQ(Result, EGJKDistanceResult::Separated);
 			EXPECT_NEAR(Distance, (FReal)1 - BPos.X, Tolerance);
 			EXPECT_NEAR(NearestA.X, (FReal)7, Tolerance);
 			EXPECT_NEAR(NearestA.Y, (FReal)0, Tolerance);
@@ -89,21 +90,22 @@ namespace ChaosTest
 		FVec3 NearestA = { 0,0,0 };
 		FVec3 NearestB = { 0,0,0 };
 		FReal Distance = 0;
+		FVec3 Normal = { 0, 0, 1 };
 
 		// Fail - overlapping
 		{
 			FAABB3 A(FVec3(5, -2, -2), FVec3(8, 2, 2));
 			TSphere<FReal, 3> B(FVec3(2, 0, 0), 2);
-			bool bSuccess = GJKDistance<FReal>(A, B, FRigidTransform3(FVec3(2, 0, 0), FRotation3::FromIdentity()), Distance, NearestA, NearestB);
-			EXPECT_FALSE(bSuccess);
+			EGJKDistanceResult Result = GJKDistance<FReal>(A, B, FRigidTransform3(FVec3(2, 0, 0), FRotation3::FromIdentity()), Distance, NearestA, NearestB, Normal);
+			EXPECT_NE(Result, EGJKDistanceResult::Separated);
 		}
 
 		// Success - not overlapping - mid-face near point
 		{
 			FAABB3 A(FVec3(5, -2, -2), FVec3(8, 2, 2));
 			TSphere<FReal, 3> B(FVec3(2, 0, 0), 2);
-			bool bSuccess = GJKDistance<FReal>(A, B, FRigidTransform3::Identity, Distance, NearestA, NearestB);
-			EXPECT_TRUE(bSuccess);
+			EGJKDistanceResult Result = GJKDistance<FReal>(A, B, FRigidTransform3::Identity, Distance, NearestA, NearestB, Normal);
+			EXPECT_EQ(Result, EGJKDistanceResult::Separated);
 			EXPECT_NEAR(Distance, (FReal)1, Tolerance);
 			EXPECT_NEAR(NearestA.X, (FReal)5, Tolerance);
 			EXPECT_NEAR(NearestA.Y, (FReal)0, Tolerance);
@@ -116,8 +118,8 @@ namespace ChaosTest
 		{
 			FAABB3 A(FVec3(5, -2, -2), FVec3(8, 2, 2));
 			TSphere<FReal, 3> B(FVec3(2, 0, 0), 2);
-			bool bSuccess = GJKDistance<FReal>(B, A, FRigidTransform3::Identity, Distance, NearestB, NearestA);
-			EXPECT_TRUE(bSuccess);
+			EGJKDistanceResult Result = GJKDistance<FReal>(B, A, FRigidTransform3::Identity, Distance, NearestB, NearestA, Normal);
+			EXPECT_EQ(Result, EGJKDistanceResult::Separated);
 			EXPECT_NEAR(Distance, (FReal)1, Tolerance);
 			EXPECT_NEAR(NearestA.X, (FReal)5, Tolerance);
 			EXPECT_NEAR(NearestA.Y, (FReal)0, Tolerance);
@@ -131,11 +133,11 @@ namespace ChaosTest
 		{
 			FAABB3 A(FVec3(5, 2, 2), FVec3(8, 4, 4));
 			TSphere<FReal, 3> B(FVec3(2, 0, 0), 2);
-			bool bSuccess = GJKDistance<FReal>(A, B, FRigidTransform3::Identity, Distance, NearestA, NearestB);
+			EGJKDistanceResult Result = GJKDistance<FReal>(A, B, FRigidTransform3::Identity, Distance, NearestA, NearestB, Normal);
 			FVec3 NearPointOnA = A.Min();
 			FVec3 SphereNearPointDir = (NearPointOnA - B.GetCenter()).GetSafeNormal();
 			FVec3 NearPointOnB = B.GetCenter() + SphereNearPointDir * B.GetRadius();
-			EXPECT_TRUE(bSuccess);
+			EXPECT_EQ(Result, EGJKDistanceResult::Separated);
 			EXPECT_NEAR(Distance, (NearPointOnA - NearPointOnB).Size(), Tolerance);
 			EXPECT_NEAR(NearestA.X, NearPointOnA.X, Tolerance);
 			EXPECT_NEAR(NearestA.Y, NearPointOnA.Y, Tolerance);
@@ -148,11 +150,11 @@ namespace ChaosTest
 		{
 			FAABB3 A(FVec3(5, 2, 2), FVec3(8, 4, 4));
 			TSphere<FReal, 3> B(FVec3(2, 0, 0), 2);
-			bool bSuccess = GJKDistance<FReal>(B, A, FRigidTransform3::Identity, Distance, NearestB, NearestA);
+			EGJKDistanceResult Result = GJKDistance<FReal>(B, A, FRigidTransform3::Identity, Distance, NearestB, NearestA, Normal);
 			FVec3 NearPointOnA = A.Min();
 			FVec3 SphereNearPointDir = (NearPointOnA - B.GetCenter()).GetSafeNormal();
 			FVec3 NearPointOnB = B.GetCenter() + SphereNearPointDir * B.GetRadius();
-			EXPECT_TRUE(bSuccess);
+			EXPECT_EQ(Result, EGJKDistanceResult::Separated);
 			EXPECT_NEAR(Distance, (NearPointOnA - NearPointOnB).Size(), Tolerance);
 			EXPECT_NEAR(NearestA.X, NearPointOnA.X, Tolerance);
 			EXPECT_NEAR(NearestA.Y, NearPointOnA.Y, Tolerance);
@@ -166,14 +168,14 @@ namespace ChaosTest
 		{
 			FAABB3 A(FVec3(-2, -2, -2), FVec3(4, 4, 4));
 			TSphere<FReal, 3> B(FVec3(0, 0, 0), 2);
-			FRigidTransform3 BToATm = FRigidTransform3(FVec3(8, 0, 0), FRotation3::FromAxisAngle(FVec3(0, 1, 0), FMath::DegreesToRadians(45)));	// Rotation won't affect sphere
-			bool bSuccess = GJKDistance<FReal>(A, B, BToATm, Distance, NearestA, NearestB);
+			FRigidTransform3 BToATm = FRigidTransform3(FVec3(8, 0, 0), FRotation3::FromAxisAngle(FVec3(0, 1, 0), FMath::DegreesToRadians(45)));	// Rotation won't affect contact depth, but does affect local contact position
+			EGJKDistanceResult Result = GJKDistance<FReal>(A, B, BToATm, Distance, NearestA, NearestB, Normal);
 			FVec3 NearPointOnA = FVec3(4, 0, 0);
 			FVec3 BPos = BToATm.TransformPositionNoScale(B.GetCenter());
 			FVec3 NearPointDir = (NearPointOnA - BPos).GetSafeNormal();
 			FVec3 NearPointOnB = BPos + NearPointDir * B.GetRadius();
 			FVec3 NearPointOnBLocal = BToATm.InverseTransformPositionNoScale(NearPointOnB);
-			EXPECT_TRUE(bSuccess);
+			EXPECT_EQ(Result, EGJKDistanceResult::Separated);
 			EXPECT_NEAR(Distance, (NearPointOnA - NearPointOnB).Size(), Tolerance);
 			EXPECT_NEAR(NearestA.X, NearPointOnA.X, Tolerance);
 			EXPECT_NEAR(NearestA.Y, NearPointOnA.Y, Tolerance);
@@ -186,14 +188,14 @@ namespace ChaosTest
 		{
 			FAABB3 A(FVec3(-2, -2, -2), FVec3(4, 4, 4));
 			TSphere<FReal, 3> B(FVec3(0, 0, 0), 2);
-			FRigidTransform3 BToATm = FRigidTransform3(FVec3(-8, 0, 0), FRotation3::FromAxisAngle(FVec3(0, 1, 0), FMath::DegreesToRadians(45)));	// Rotation will affect box
-			bool bSuccess = GJKDistance<FReal>(B, A, BToATm, Distance, NearestB, NearestA);
+			FRigidTransform3 BToATm = FRigidTransform3(FVec3(-8, 0, 0), FRotation3::FromAxisAngle(FVec3(0, 1, 0), FMath::DegreesToRadians(45)));
+			EGJKDistanceResult Result = GJKDistance<FReal>(B, A, BToATm, Distance, NearestB, NearestA, Normal);
 			FVec3 NearPointOnA = FVec3(4, 0, 4);
 			FVec3 BPos = BToATm.InverseTransformPositionNoScale(B.GetCenter());
 			FVec3 NearPointDir = (NearPointOnA - BPos).GetSafeNormal();
 			FVec3 NearPointOnB = BPos + NearPointDir * B.GetRadius();
 			FVec3 NearPointOnBLocal = BToATm.TransformPositionNoScale(NearPointOnB);
-			EXPECT_TRUE(bSuccess);
+			EXPECT_EQ(Result, EGJKDistanceResult::Separated);
 			EXPECT_NEAR(Distance, (NearPointOnA - NearPointOnB).Size(), Tolerance);
 			EXPECT_NEAR(NearestA.X, NearPointOnA.X, Tolerance);
 			EXPECT_NEAR(NearestA.Y, NearPointOnA.Y, Tolerance);
@@ -211,11 +213,11 @@ namespace ChaosTest
 			bool bOverlap = GJKIntersection<FReal>(A, B, FRigidTransform3::Identity);
 			EXPECT_FALSE(bOverlap);
 
-			bool bSuccess = GJKDistance<FReal>(A, B, FRigidTransform3::Identity, Distance, NearestA, NearestB);
+			EGJKDistanceResult Result = GJKDistance<FReal>(A, B, FRigidTransform3::Identity, Distance, NearestA, NearestB, Normal);
 			FVec3 NearPointOnA = FVec3(5, 0, 2);
 			FVec3 NearPointDir = (NearPointOnA - B.GetCenter()).GetSafeNormal();
 			FVec3 NearPointOnB = B.GetCenter() + NearPointDir * B.GetRadius();
-			EXPECT_TRUE(bSuccess);
+			EXPECT_EQ(Result, EGJKDistanceResult::Separated);
 			EXPECT_NEAR(Distance, (NearPointOnA - NearPointOnB).Size(), Tolerance);
 			EXPECT_NEAR(NearestA.X, NearPointOnA.X, Tolerance);
 			EXPECT_NEAR(NearestA.Y, NearPointOnA.Y, Tolerance);
@@ -236,23 +238,24 @@ namespace ChaosTest
 		FVec3 NearestA = { 0,0,0 };
 		FVec3 NearestB = { 0,0,0 };
 		FReal Distance = 0;
+		FVec3 Normal = { 0, 0, 1 };
 
 		// Fail - overlapping
 		{
 			FAABB3 A(FVec3(5, -2, -2), FVec3(8, 2, 2));
 			TCapsule<FReal> B(FVec3(2, -2, 0), FVec3(2, 2, 0), 2);
-			bool bSuccess = GJKDistance<FReal>(A, B, FRigidTransform3(FVec3(2, 0, 0), FRotation3::FromIdentity()), Distance, NearestA, NearestB);
-			EXPECT_FALSE(bSuccess);
+			EGJKDistanceResult Result = GJKDistance<FReal>(A, B, FRigidTransform3(FVec3(2, 0, 0), FRotation3::FromIdentity()), Distance, NearestA, NearestB, Normal);
+			EXPECT_NE(Result, EGJKDistanceResult::Separated);
 		}
 
 		// Success - not overlapping, capsule axis parallel to nearest face (near points on cylinder and box face)
 		{
 			FAABB3 A(FVec3(5, -2, -2), FVec3(8, 2, 2));
 			TCapsule<FReal> B(FVec3(2, 0, -1), FVec3(2, 0, 2), 2);
-			bool bSuccess = GJKDistance<FReal>(A, B, FRigidTransform3::Identity, Distance, NearestA, NearestB);
+			EGJKDistanceResult Result = GJKDistance<FReal>(A, B, FRigidTransform3::Identity, Distance, NearestA, NearestB, Normal);
 
 			const FReal Tolerance = (FReal)2e-3;
-			EXPECT_TRUE(bSuccess);
+			EXPECT_EQ(Result, EGJKDistanceResult::Separated);
 			EXPECT_NEAR(Distance, (FReal)1, Tolerance);
 			EXPECT_NEAR(NearestA.X, (FReal)5, Tolerance);
 			EXPECT_NEAR(NearestA.Y, (FReal)0, Tolerance);
@@ -268,13 +271,13 @@ namespace ChaosTest
 		{
 			FAABB3 A(FVec3(5, -2, -2), FVec3(8, 2, 2));
 			TCapsule<FReal> B(FVec3(-2, 0, 3), FVec3(2, 0, -3), 2);
-			bool bSuccess = GJKDistance<FReal>(A, B, FRigidTransform3::Identity, Distance, NearestA, NearestB);
+			EGJKDistanceResult Result = GJKDistance<FReal>(A, B, FRigidTransform3::Identity, Distance, NearestA, NearestB, Normal);
 			FVec3 ExpectedNearestA = FVec3(5, 0, -2);
 			FVec3 ExpectedDir = (ExpectedNearestA - B.GetX2()).GetSafeNormal();
 			FVec3 ExpectedNearestB = B.GetX2() + ExpectedDir * B.GetRadius();
 
 			const FReal Tolerance = (FReal)2e-3;
-			EXPECT_TRUE(bSuccess);
+			EXPECT_EQ(Result, EGJKDistanceResult::Separated);
 			EXPECT_NEAR(Distance, (ExpectedNearestB - ExpectedNearestA).Size(), Tolerance);
 			EXPECT_NEAR(NearestA.X, (FReal)ExpectedNearestA.X, Tolerance);
 			EXPECT_NEAR(NearestA.Y, (FReal)ExpectedNearestA.Y, Tolerance);
@@ -288,12 +291,12 @@ namespace ChaosTest
 		{
 			TCapsule<FReal> A(FVec3(4, 0, -1), FVec3(4, 0, -7), 1);
 			FAABB3 B(FVec3(-2, -2, -2), FVec3(2, 2, 2));
-			bool bSuccess = GJKDistance<FReal>(A, B, FRigidTransform3::Identity, Distance, NearestA, NearestB);
+			EGJKDistanceResult Result = GJKDistance<FReal>(A, B, FRigidTransform3::Identity, Distance, NearestA, NearestB, Normal);
 			FVec3 ExpectedNearestA = FVec3(3, 0, (FReal)-1.5);
 			FVec3 ExpectedNearestB = FVec3(2, 0, (FReal)-1.5);
 
 			const FReal Tolerance = (FReal)2e-3;
-			EXPECT_TRUE(bSuccess);
+			EXPECT_EQ(Result, EGJKDistanceResult::Separated);
 			EXPECT_NEAR(Distance, (FReal)1, Tolerance);
 			EXPECT_NEAR(NearestA.X, (FReal)ExpectedNearestA.X, Tolerance);
 			EXPECT_NEAR(NearestA.Y, (FReal)ExpectedNearestA.Y, Tolerance);
@@ -310,12 +313,12 @@ namespace ChaosTest
 			TCapsule<FReal> A(FVec3(0, 0, -3), FVec3(0, 0, 3), 1);
 			FAABB3 B(FVec3(-2, -2, -2), FVec3(2, 2, 2));
 			FRigidTransform3 BToA = FRigidTransform3(FVec3(-4, 0, 4), FRotation3::FromIdentity());
-			bool bSuccess = GJKDistance<FReal>(A, B, BToA, Distance, NearestA, NearestB);
+			EGJKDistanceResult Result = GJKDistance<FReal>(A, B, BToA, Distance, NearestA, NearestB, Normal);
 			FVec3 ExpectedNearestA = FVec3(-1, 0, (FReal)2);
 			FVec3 ExpectedNearestB = FVec3(2, 0, (FReal)-2);
 
 			const FReal Tolerance = (FReal)2e-3;
-			EXPECT_TRUE(bSuccess);
+			EXPECT_EQ(Result, EGJKDistanceResult::Separated);
 			EXPECT_NEAR(Distance, (FReal)1, Tolerance);
 			EXPECT_NEAR(NearestA.X, (FReal)ExpectedNearestA.X, Tolerance);
 			EXPECT_NEAR(NearestA.Y, (FReal)ExpectedNearestA.Y, Tolerance);
@@ -339,6 +342,7 @@ namespace ChaosTest
 		FVec3 NearestA = { 0,0,0 };
 		FVec3 NearestB = { 0,0,0 };
 		FReal Distance = 0;
+		FVec3 Normal = { 0, 0, 1 };
 
 		// Capsule-box takes number of iterations at the moment (we can improve that with a better the choice of Initial V)
 		// so test that we still get an approximate answer with less iterations
@@ -347,13 +351,13 @@ namespace ChaosTest
 			TCapsule<FReal> B(FVec3(-2, 0, 3), FVec3(2, 0, -3), 2);
 			FReal Epsilon = (FReal)1e-6;
 			int32 MaxIts = 5;
-			bool bSuccess = GJKDistance<FReal>(A, B, FRigidTransform3::Identity, Distance, NearestA, NearestB, Epsilon, MaxIts);
+			EGJKDistanceResult Result = GJKDistance<FReal>(A, B, FRigidTransform3::Identity, Distance, NearestA, NearestB, Normal, Epsilon, MaxIts);
 			FVec3 ExpectedNearestA = FVec3(5, 0, -2);
 			FVec3 ExpectedDir = (ExpectedNearestA - B.GetX2()).GetSafeNormal();
 			FVec3 ExpectedNearestB = B.GetX2() + ExpectedDir * B.GetRadius();
 
 			const FReal Tolerance = (FReal)0.3;
-			EXPECT_TRUE(bSuccess);
+			EXPECT_EQ(Result, EGJKDistanceResult::Separated);
 			EXPECT_NEAR(Distance, (ExpectedNearestB - ExpectedNearestA).Size(), Tolerance);
 			EXPECT_NEAR(NearestA.X, (FReal)ExpectedNearestA.X, Tolerance);
 			EXPECT_NEAR(NearestA.Y, (FReal)ExpectedNearestA.Y, Tolerance);
diff --git a/Engine/Source/Programs/HeadlessChaos/Private/HeadlessChaosTestUtility.cpp b/Engine/Source/Programs/HeadlessChaos/Private/HeadlessChaosTestUtility.cpp
--- a/Engine/Source/Programs/HeadlessChaos/Private/HeadlessChaosTestUtility.cpp
+++ b/Engine/Source/Programs/HeadlessChaos/Private/HeadlessChaosTestUtility.cpp
@@ -579,8 +579,8 @@ namespace ChaosTest {
 		InParticles.Q() = InParticles.R();
 
 		// TODO: Change this error prone API to set bounds more automatically. This is easy to forget
-		InParticles.SetLocalBounds(TAABB<FReal, 3>(Cube[0], Cube[7]));
-		InParticles.SetWorldSpaceInflatedBounds(TAABB<FReal, 3>(Cube[0], Cube[7]));
+		InParticles.SetLocalBounds(FAABB3(Cube[0], Cube[7]));
+		InParticles.SetWorldSpaceInflatedBounds(FAABB3(Cube[0], Cube[7]));
 		InParticles.SetHasBounds(true);
 
 		InParticles.M() = 1.0;
@@ -702,11 +702,9 @@ namespace ChaosTest {
 		}
 	}
 
-	FImplicitConvex3 CreateConvexBox(const FVec3& BoxSize, const FReal Margin)
+	TArray<FVec3> MakeBoxVerts(const FVec3& HalfSize)
 	{
-		const FVec3 HalfSize = 0.5f * BoxSize;
-
-		TArray<FVec3> BoxVerts =
+		return
 		{
 			FVec3(-HalfSize.X, -HalfSize.Y, -HalfSize.Z),
 			FVec3(-HalfSize.X,  HalfSize.Y, -HalfSize.Z),
@@ -717,29 +715,26 @@ namespace ChaosTest {
 			FVec3(HalfSize.X,  HalfSize.Y,  HalfSize.Z),
 			FVec3(HalfSize.X, -HalfSize.Y,  HalfSize.Z),
 		};
+	}
 
-		return FImplicitConvex3(BoxVerts, Margin);
+	FImplicitConvex3 CreateConvexBox(const FVec3& BoxSize, const FReal Margin)
+	{
+		const FVec3 HalfSize = 0.5f * BoxSize;
+		return FImplicitConvex3(MakeBoxVerts(HalfSize), Margin);
 	}
 
+	TImplicitObjectInstanced<FImplicitConvex3> CreateInstancedConvexBox(const FVec3& BoxSize, const FReal Margin)
+	{
+		const FVec3 HalfSize = 0.5f * BoxSize;
+		TSharedPtr<FImplicitConvex3, ESPMode::ThreadSafe> BoxConvex = MakeShared<FImplicitConvex3, ESPMode::ThreadSafe>(MakeBoxVerts(HalfSize), 0.0f);
+		return TImplicitObjectInstanced<FImplicitConvex3>(BoxConvex, Margin);
+	}
 
 	TImplicitObjectScaled<FImplicitConvex3> CreateScaledConvexBox(const FVec3& BoxSize, const FVec3 BoxScale, const FReal Margin)
 	{
 		const FVec3 HalfSize = 0.5f * BoxSize;
-
-		TArray<FVec3> BoxVerts =
-		{
-			FVec3(-HalfSize.X, -HalfSize.Y, -HalfSize.Z),
-			FVec3(-HalfSize.X,  HalfSize.Y, -HalfSize.Z),
-			FVec3(HalfSize.X,  HalfSize.Y, -HalfSize.Z),
-			FVec3(HalfSize.X, -HalfSize.Y, -HalfSize.Z),
-			FVec3(-HalfSize.X, -HalfSize.Y,  HalfSize.Z),
-			FVec3(-HalfSize.X,  HalfSize.Y,  HalfSize.Z),
-			FVec3(HalfSize.X,  HalfSize.Y,  HalfSize.Z),
-			FVec3(HalfSize.X, -HalfSize.Y,  HalfSize.Z),
-		};
-		
-		TSharedPtr<FImplicitConvex3, ESPMode::ThreadSafe> BoxConvex = MakeShared<FImplicitConvex3, ESPMode::ThreadSafe>(BoxVerts, 0.0f);
-
+		TSharedPtr<FImplicitConvex3, ESPMode::ThreadSafe> BoxConvex = MakeShared<FImplicitConvex3, ESPMode::ThreadSafe>(MakeBoxVerts(HalfSize), 0.0f);
 		return TImplicitObjectScaled<FImplicitConvex3>(BoxConvex, BoxScale, Margin);
 	}
+
 }
diff --git a/Engine/Source/Programs/HeadlessChaos/Public/HeadlessChaosTestUtility.h b/Engine/Source/Programs/HeadlessChaos/Public/HeadlessChaosTestUtility.h
--- a/Engine/Source/Programs/HeadlessChaos/Public/HeadlessChaosTestUtility.h
+++ b/Engine/Source/Programs/HeadlessChaos/Public/HeadlessChaosTestUtility.h
@@ -200,5 +200,6 @@ namespace ChaosTest {
 
 
 	extern FImplicitConvex3 CreateConvexBox(const FVec3& BoxSize, const FReal Margin);
+	extern TImplicitObjectInstanced<FImplicitConvex3> CreateInstancedConvexBox(const FVec3& BoxSize, const FReal Margin);
 	extern TImplicitObjectScaled<FImplicitConvex3> CreateScaledConvexBox(const FVec3& BoxSize, const FVec3 BoxScale, const FReal Margin);
 }
diff --git a/Engine/Source/Runtime/Experimental/Chaos/Private/Chaos/Collision/CapsuleConvexContactPoint.cpp b/Engine/Source/Runtime/Experimental/Chaos/Private/Chaos/Collision/CapsuleConvexContactPoint.cpp
new file mode 100644
--- /dev/null
+++ b/Engine/Source/Runtime/Experimental/Chaos/Private/Chaos/Collision/CapsuleConvexContactPoint.cpp
@@ -0,0 +1,111 @@
+// Copyright Epic Games, Inc. All Rights Reserved.
+#include "Chaos/Collision/CapsuleConvexContactPoint.h"
+#include "Chaos/Collision/ContactPoint.h"
+#include "Chaos/Capsule.h"
+#include "Chaos/Convex.h"
+#include "Chaos/GJK.h"
+#include "Chaos/GJKShape.h"
+#include "Chaos/ImplicitObjectScaled.h"
+
+namespace Chaos
+{
+	// Use GJK (point to convex) to calculate separation.
+	// Fall back to plane testing if penetrating by more than Radius.
+	template<typename T_CONVEX>
+	FContactPoint ConvexCapsuleContactPointImpl(const T_CONVEX& Convex, const FRigidTransform3& ConvexTransform, const FImplicitCapsule3& Capsule, const FRigidTransform3& CapsuleTransform, FReal CullDistance)
+	{
+		FContactPoint ContactPoint;
+		if (Convex.NumPlanes() > 0)
+		{
+			const FRigidTransform3 CapsuleToConvexTransform = CapsuleTransform.GetRelativeTransform(ConvexTransform);
+			FVec3 PosConvex, PosCapsuleInConvex, NormalConvex;
+			FReal Penetration;
+			int32 VertexIndexA, VertexIndexB;
+
+			// Run GJK to find separating distance if available
+			// NOTE: Capsule is treated as a line (its core shape), Convex margin is ignored so we are using the outer non-shrunken hull.
+			// @todo(chaos): use GJKDistance and SAT when that fails rather that GJK/EPA (but this requires an edge list in the convex)
+			bool bHaveResult = GJKPenetration<true>(MakeGJKShape(Convex), MakeGJKCoreShape(Capsule), CapsuleToConvexTransform, Penetration, PosConvex, PosCapsuleInConvex, NormalConvex, VertexIndexA, VertexIndexB);
+			
+			// Build the contact point
+			if (bHaveResult)
+			{
+				const FVec3 PosCapsule = CapsuleToConvexTransform.InverseTransformPosition(PosCapsuleInConvex);
+				const FReal Phi = -Penetration;
+
+				ContactPoint.ShapeContactPoints[0] = PosConvex;
+				ContactPoint.ShapeContactPoints[1] = PosCapsule;
+				ContactPoint.ShapeContactNormal = -NormalConvex;
+				ContactPoint.ShapeMargins[0] = 0.0f;
+				ContactPoint.ShapeMargins[1] = 0.0f;
+				ContactPoint.ContactNormalOwnerIndex = 0;
+
+				ContactPoint.Location = ConvexTransform.TransformPosition(ContactPoint.ShapeContactPoints[0]);
+				ContactPoint.Normal = ConvexTransform.TransformVectorNoScale(-NormalConvex);
+				ContactPoint.Phi = Phi;
+			}
+		}
+
+		return ContactPoint;
+	}
+
+	template<typename T_CONVEX>
+	FContactPoint CapsuleConvexContactPointImpl(const FImplicitCapsule3& Capsule, const FRigidTransform3& CapsuleTransform, const T_CONVEX& Convex, const FRigidTransform3& ConvexTransform, FReal CullDistance)
+	{
+		return ConvexCapsuleContactPointImpl(Convex, ConvexTransform, Capsule, CapsuleTransform, CullDistance).SwapShapes();
+	}
+
+
+
+	FContactPoint CapsuleConvexContactPoint(const FImplicitCapsule3& Capsule, const FRigidTransform3& CapsuleTransform, const FImplicitConvex3& Convex, const FRigidTransform3& ConvexTransform, FReal CullDistance)
+	{
+		return CapsuleConvexContactPointImpl(Capsule, CapsuleTransform, Convex, ConvexTransform, CullDistance);
+	}
+
+	FContactPoint CapsuleConvexContactPoint(const FImplicitCapsule3& Capsule, const FRigidTransform3& CapsuleTransform, const TImplicitObjectInstanced<FImplicitConvex3>& Convex, const FRigidTransform3& ConvexTransform, FReal CullDistance)
+	{
+		return CapsuleConvexContactPointImpl(Capsule, CapsuleTransform, Convex, ConvexTransform, CullDistance);
+	}
+
+	FContactPoint CapsuleConvexContactPoint(const FImplicitCapsule3& Capsule, const FRigidTransform3& CapsuleTransform, const TImplicitObjectScaled<FImplicitConvex3>& Convex, const FRigidTransform3& ConvexTransform, FReal CullDistance)
+	{
+		return CapsuleConvexContactPointImpl(Capsule, CapsuleTransform, Convex, ConvexTransform, CullDistance);
+	}
+
+
+
+	FContactPoint ConvexCapsuleContactPoint(const FImplicitConvex3& Convex, const FRigidTransform3& ConvexTransform, const FImplicitCapsule3& Capsule, const FRigidTransform3& CapsuleTransform, FReal CullDistance)
+	{
+		return ConvexCapsuleContactPointImpl(Convex, ConvexTransform, Capsule, CapsuleTransform, CullDistance);
+	}
+
+	FContactPoint ConvexCapsuleContactPoint(const TImplicitObjectInstanced<FImplicitConvex3>& Convex, const FRigidTransform3& ConvexTransform, const FImplicitCapsule3& Capsule, const FRigidTransform3& CapsuleTransform, FReal CullDistance)
+	{
+		return ConvexCapsuleContactPointImpl(Convex, ConvexTransform, Capsule, CapsuleTransform, CullDistance);
+	}
+
+	FContactPoint ConvexCapsuleContactPoint(const TImplicitObjectScaled<FImplicitConvex3>& Convex, const FRigidTransform3& ConvexTransform, const FImplicitCapsule3& Capsule, const FRigidTransform3& CapsuleTransform, FReal CullDistance)
+	{
+		return ConvexCapsuleContactPointImpl(Convex, ConvexTransform, Capsule, CapsuleTransform, CullDistance);
+	}
+
+
+
+	FContactPoint CapsuleConvexContactPoint(const FImplicitCapsule3& Capsule, const FRigidTransform3& CapsuleTransform, const FImplicitObject& Object, const FRigidTransform3& ConvexTransform, FReal CullDistance)
+	{
+		if (const TImplicitObjectInstanced<FImplicitConvex3>* InstancedConvex = Object.template GetObject<const TImplicitObjectInstanced<FImplicitConvex3>>())
+		{
+			return CapsuleConvexContactPoint(Capsule, CapsuleTransform, *InstancedConvex, ConvexTransform, CullDistance);
+		}
+		else if (const TImplicitObjectScaled<FImplicitConvex3>* ScaledConvex = Object.template GetObject<const TImplicitObjectScaled<FImplicitConvex3>>())
+		{
+			return CapsuleConvexContactPoint(Capsule, CapsuleTransform, *ScaledConvex, ConvexTransform, CullDistance);
+		}
+		else if (const FImplicitConvex3* Convex = Object.template GetObject<const FImplicitConvex3>())
+		{
+			return CapsuleConvexContactPoint(Capsule, CapsuleTransform, *Convex, ConvexTransform, CullDistance);
+		}
+		return FContactPoint();
+	}
+
+}
diff --git a/Engine/Source/Runtime/Experimental/Chaos/Private/Chaos/Collision/SphereConvexContactPoint.cpp b/Engine/Source/Runtime/Experimental/Chaos/Private/Chaos/Collision/SphereConvexContactPoint.cpp
new file mode 100644
--- /dev/null
+++ b/Engine/Source/Runtime/Experimental/Chaos/Private/Chaos/Collision/SphereConvexContactPoint.cpp
@@ -0,0 +1,153 @@
+// Copyright Epic Games, Inc. All Rights Reserved.
+
+#include "Chaos/Collision/SphereConvexContactPoint.h"
+#include "Chaos/Collision/ContactPoint.h"
+#include "Chaos/Convex.h"
+#include "Chaos/GJK.h"
+#include "Chaos/GJKShape.h"
+#include "Chaos/ImplicitObjectScaled.h"
+#include "Chaos/Sphere.h"
+
+namespace Chaos
+{
+	// Calculate the shortest vector for the point to depenetrate the convex.
+	// Returns true unless there are no planes in the convex
+	// Note: this may be called with small positive separations of order epsilon passed to GJK,
+	// but the result is increasingly inaccurate as distance increases.
+	template<typename T_CONVEX>
+	bool ConvexPointPenetrationVector(const T_CONVEX& Convex, const FVec3& X, FVec3& OutNormal, FReal& OutPhi)
+	{
+		FReal MaxPhi = TNumericLimits<FReal>::Lowest();
+		int32 MaxPlaneIndex = INDEX_NONE;
+
+		const int32 NumPlanes = Convex.NumPlanes();
+		for (int32 PlaneIndex = 0; PlaneIndex < NumPlanes; ++PlaneIndex)
+		{
+			const FReal Phi = Convex.GetPlane(PlaneIndex).SignedDistance(X);
+			if (Phi > MaxPhi)
+			{
+				MaxPhi = Phi;
+				MaxPlaneIndex = PlaneIndex;
+			}
+		}
+
+		if (MaxPlaneIndex != INDEX_NONE)
+		{
+			OutPhi = MaxPhi;
+			OutNormal = Convex.GetPlane(MaxPlaneIndex).Normal();
+			return true;
+		}
+
+		return false;
+	}
+
+
+	// Use GJK (point to convex) to calculate separation.
+	// Fall back to plane testing if penetrating by more than Radius.
+	template<typename T_CONVEX>
+	FContactPoint ConvexSphereContactPointImpl(const T_CONVEX& Convex, const FRigidTransform3& ConvexTransform, const FImplicitSphere3& Sphere, const FRigidTransform3& SphereTransform, FReal CullDistance)
+	{
+		FContactPoint ContactPoint;
+		if (Convex.NumPlanes() > 0)
+		{
+			const FRigidTransform3 SphereToConvexTransform = SphereTransform.GetRelativeTransform(ConvexTransform);
+			FVec3 PosConvex, PosSphere, NormalConvex;
+			FReal Phi;
+
+			// Run GJK to find separating distance if available
+			// NOTE: Sphere is treated as a point (its core shape), Convex margin is ignored so we are using the outer non-shrunken hull.
+			const EGJKDistanceResult GjkResult = GJKDistance(MakeGJKShape(Convex), MakeGJKCoreShape(Sphere), SphereToConvexTransform, Phi, PosConvex, PosSphere, NormalConvex);
+			bool bHaveResult = (GjkResult != EGJKDistanceResult::DeepContact);
+
+			// If GJK failed, core shapes are penetrating so find the minimum penetration vector
+			if (!bHaveResult)
+			{
+				const FVec3 SpherePosConvex = SphereToConvexTransform.TransformPosition(Sphere.GetCenter());
+
+				FReal PointPhi;
+				bHaveResult = ConvexPointPenetrationVector(Convex, SpherePosConvex, NormalConvex, PointPhi);
+				if (bHaveResult)
+				{
+					PosConvex = SpherePosConvex - PointPhi * NormalConvex;
+					PosSphere = SphereToConvexTransform.InverseTransformVector(-Sphere.GetRadius() * NormalConvex);
+					Phi = PointPhi - Sphere.GetRadius();
+				}
+			}
+
+			// Build the contact point
+			if (bHaveResult)
+			{
+				ContactPoint.ShapeContactPoints[0] = PosConvex;
+				ContactPoint.ShapeContactPoints[1] = PosSphere;
+				ContactPoint.ShapeContactNormal = -NormalConvex;
+				ContactPoint.ShapeMargins[0] = 0.0f;
+				ContactPoint.ShapeMargins[1] = 0.0f;
+				ContactPoint.ContactNormalOwnerIndex = 0;
+
+				ContactPoint.Location = ConvexTransform.TransformPosition(ContactPoint.ShapeContactPoints[0]);
+				ContactPoint.Normal = ConvexTransform.TransformVectorNoScale(-NormalConvex);
+				ContactPoint.Phi = Phi;
+			}
+		}
+
+		return ContactPoint;
+	}
+
+	template<typename T_CONVEX>
+	FContactPoint SphereConvexContactPointImpl(const FImplicitSphere3& Sphere, const FRigidTransform3& SphereTransform, const T_CONVEX& Convex, const FRigidTransform3& ConvexTransform, FReal CullDistance)
+	{
+		return ConvexSphereContactPointImpl(Convex, ConvexTransform, Sphere, SphereTransform, CullDistance).SwapShapes();
+	}
+
+	FContactPoint SphereConvexContactPoint(const FImplicitSphere3& Sphere, const FRigidTransform3& SphereTransform, const FImplicitConvex3& Convex, const FRigidTransform3& ConvexTransform, FReal CullDistance)
+	{
+		return SphereConvexContactPointImpl(Sphere, SphereTransform, Convex, ConvexTransform, CullDistance);
+	}
+
+	FContactPoint SphereConvexContactPoint(const FImplicitSphere3& Sphere, const FRigidTransform3& SphereTransform, const TImplicitObjectInstanced<FImplicitConvex3>& Convex, const FRigidTransform3& ConvexTransform, FReal CullDistance)
+	{
+		return SphereConvexContactPointImpl(Sphere, SphereTransform, Convex, ConvexTransform, CullDistance);
+	}
+
+	FContactPoint SphereConvexContactPoint(const FImplicitSphere3& Sphere, const FRigidTransform3& SphereTransform, const TImplicitObjectScaled<FImplicitConvex3>& Convex, const FRigidTransform3& ConvexTransform, FReal CullDistance)
+	{
+		return SphereConvexContactPointImpl(Sphere, SphereTransform, Convex, ConvexTransform, CullDistance);
+	}
+
+
+
+	FContactPoint ConvexSphereContactPoint(const FImplicitConvex3& Convex, const FRigidTransform3& ConvexTransform, const FImplicitSphere3& Sphere, const FRigidTransform3& SphereTransform, FReal CullDistance)
+	{
+		return ConvexSphereContactPointImpl(Convex, ConvexTransform, Sphere, SphereTransform, CullDistance);
+	}
+
+	FContactPoint ConvexSphereContactPoint(const TImplicitObjectInstanced<FImplicitConvex3>& Convex, const FRigidTransform3& ConvexTransform, const FImplicitSphere3& Sphere, const FRigidTransform3& SphereTransform, FReal CullDistance)
+	{
+		return ConvexSphereContactPointImpl(Convex, ConvexTransform, Sphere, SphereTransform, CullDistance);
+	}
+
+	FContactPoint ConvexSphereContactPoint(const TImplicitObjectScaled<FImplicitConvex3>& Convex, const FRigidTransform3& ConvexTransform, const FImplicitSphere3& Sphere, const FRigidTransform3& SphereTransform, FReal CullDistance)
+	{
+		return ConvexSphereContactPointImpl(Convex, ConvexTransform, Sphere, SphereTransform, CullDistance);
+	}
+
+
+
+	FContactPoint SphereConvexContactPoint(const FImplicitSphere3& Sphere, const FRigidTransform3& SphereTransform, const FImplicitObject& Object, const FRigidTransform3& ConvexTransform, FReal CullDistance)
+	{
+		if (const TImplicitObjectInstanced<FImplicitConvex3>* InstancedConvex = Object.template GetObject<const TImplicitObjectInstanced<FImplicitConvex3>>())
+		{
+			return SphereConvexContactPoint(Sphere, SphereTransform, *InstancedConvex, ConvexTransform, CullDistance);
+		}
+		else if (const TImplicitObjectScaled<FImplicitConvex3>* ScaledConvex = Object.template GetObject<const TImplicitObjectScaled<FImplicitConvex3>>())
+		{
+			return SphereConvexContactPoint(Sphere, SphereTransform, *ScaledConvex, ConvexTransform, CullDistance);
+		}
+		else if (const FImplicitConvex3* Convex = Object.template GetObject<const FImplicitConvex3>())
+		{
+			return SphereConvexContactPoint(Sphere, SphereTransform, *Convex, ConvexTransform, CullDistance);
+		}
+		return FContactPoint();
+	}
+
+}
diff --git a/Engine/Source/Runtime/Experimental/Chaos/Private/Chaos/CollisionOneShotManifolds.cpp b/Engine/Source/Runtime/Experimental/Chaos/Private/Chaos/CollisionOneShotManifolds.cpp
--- a/Engine/Source/Runtime/Experimental/Chaos/Private/Chaos/CollisionOneShotManifolds.cpp
+++ b/Engine/Source/Runtime/Experimental/Chaos/Private/Chaos/CollisionOneShotManifolds.cpp
@@ -539,6 +539,10 @@ namespace Chaos
 
 			// Find the deepest penetration. This is used to determine the planes and points to use for the manifold
 			FContactPoint GJKContactPoint = GenericConvexConvexContactPoint(Convex1, Convex1Transform, Convex2, Convex2Transform, CullDistance, Constraint.Manifold.RestitutionPadding);
+			if (GJKContactPoint.Phi > CullDistance)
+			{
+				return;
+			}
 
 			// @todo(chaos): get the vertex index from GJK and use to to get the plane
 			const FVec3 SeparationDirectionLocalConvex1 = Convex1Transform.InverseTransformVectorNoScale(GJKContactPoint.Normal);
diff --git a/Engine/Source/Runtime/Experimental/Chaos/Private/Chaos/CollisionResolution.cpp b/Engine/Source/Runtime/Experimental/Chaos/Private/Chaos/CollisionResolution.cpp
--- a/Engine/Source/Runtime/Experimental/Chaos/Private/Chaos/CollisionResolution.cpp
+++ b/Engine/Source/Runtime/Experimental/Chaos/Private/Chaos/CollisionResolution.cpp
@@ -7,8 +7,10 @@
 #include "Chaos/Capsule.h"
 #include "Chaos/CollisionResolutionTypes.h"
 #include "Chaos/CollisionResolutionUtil.h"
+#include "Chaos/Collision/CapsuleConvexContactPoint.h"
 #include "Chaos/Collision/CollisionContext.h"
 #include "Chaos/Collision/PBDCollisionConstraint.h"
+#include "Chaos/Collision/SphereConvexContactPoint.h"
 #include "Chaos/Convex.h"
 #include "Chaos/Defines.h"
 #include "Chaos/HeightField.h"
@@ -18,6 +20,7 @@
 #include "Chaos/Levelset.h"
 #include "Chaos/Pair.h"
 #include "Chaos/ParticleHandle.h"
+#include "Chaos/Particle/ParticleUtilities.h"
 #include "Chaos/PBDCollisionConstraints.h"
 #include "Chaos/Sphere.h"
 #include "Chaos/Transform.h"
@@ -34,6 +37,9 @@ DECLARE_CYCLE_STAT(TEXT("Collisions::GJK"), STAT_Collisions_GJK, STATGROUP_Chaos
 #define SCOPE_CYCLE_COUNTER_GJK()
 #endif
 
+// @todo(chaos): clean up the contact creation time rejection to avoid extra transforms
+#define CHAOS_COLLISION_CREATE_BOUNDSCHECK 1
+
 DECLARE_CYCLE_STAT(TEXT("Collisions::ConstructConstraints"), STAT_Collisions_ConstructConstraints, STATGROUP_ChaosCollision);
 DECLARE_CYCLE_STAT(TEXT("Collisions::ConstructConstraintsInternal"), STAT_Collisions_ConstructConstraintsInternal, STATGROUP_ChaosCollision);
 DECLARE_CYCLE_STAT(TEXT("Collisions::FindAllIntersectingClusteredObjects"), STAT_Collisions_FindAllIntersectingClusteredObjects, STATGROUP_ChaosCollision);
@@ -231,8 +237,17 @@ namespace Chaos
 			}
 		}
 
-		void UpdateContactPoint(FRigidBodyPointContactConstraint& Constraint, const FContactPoint& ContactPoint, const FReal Dt)
+		void UpdateContactPoint(FRigidBodyPointContactConstraint& Constraint, const FContactPoint& ContactPoint, const FReal CullDistance, const FReal Dt)
 		{
+			// Permanently disable contacts beyond the CullDistance
+#if CHAOS_COLLISION_CREATE_BOUNDSCHECK
+			if (ContactPoint.Phi > CullDistance)
+			{
+				Constraint.SetDisabled(true);
+				return;
+			}
+#endif
+
 			// Ignore points that have not been initialized - if there is no detectable contact 
 			// point within reasonable range despite passing the AABB tests
 			if (ContactPoint.IsSet())
@@ -241,6 +256,7 @@ namespace Chaos
 			}
 		}
 
+
 		template <typename GeometryA, typename GeometryB>
 		FContactPoint GJKContactPoint2(const GeometryA& A, const GeometryB& B, const FRigidTransform3& ATM, const FRigidTransform3& BToATM, const FVec3& InitialDir, const FReal ShapePadding)
 		{
@@ -508,7 +524,7 @@ namespace Chaos
 			template<typename T_ImplicitA, typename T_ImplicitB>
 			static void UpdateConvexConvexConstraint(const T_ImplicitA& A, const FRigidTransform3& ATM, const T_ImplicitB& B, const FRigidTransform3& BTM, const FReal CullDistance, const FReal Dt, const FReal ShapePadding, FRigidBodyPointContactConstraint& Constraint)
 			{
-				UpdateContactPoint(Constraint, GJKContactPoint(A, ATM, B, BTM, FVec3(1, 0, 0), ShapePadding), Dt);
+				UpdateContactPoint(Constraint, GJKContactPoint(A, ATM, B, BTM, FVec3(1, 0, 0), ShapePadding), CullDistance, Dt);
 			}
 		};
 
@@ -533,7 +549,7 @@ namespace Chaos
 				}
 				else
 				{
-					UpdateContactPoint(Constraint, GJKContactPoint(A, ATM, B, BTM, FVec3(1, 0, 0), ShapePadding), Dt);
+					UpdateContactPoint(Constraint, GJKContactPoint(A, ATM, B, BTM, FVec3(1, 0, 0), ShapePadding), CullDistance, Dt);
 				}
 			}
 		};
@@ -627,7 +643,7 @@ namespace Chaos
 			}
 			else
 			{
-				UpdateContactPoint(Constraint, BoxBoxContactPoint(Box1, Box2, Box1Transform, Box2Transform, CullDistance, Constraint.Manifold.RestitutionPadding), Dt);
+				UpdateContactPoint(Constraint, BoxBoxContactPoint(Box1, Box2, Box1Transform, Box2Transform, CullDistance, Constraint.Manifold.RestitutionPadding), CullDistance, Dt);
 			}
 		}
 
@@ -663,7 +679,7 @@ namespace Chaos
 		void UpdateBoxHeightFieldConstraint(const FImplicitBox3& A, const FRigidTransform3& ATransform, const FHeightField& B, const FRigidTransform3& BTransform, const FReal CullDistance, const FReal Dt, FRigidBodyPointContactConstraint& Constraint)
 		{
 			// @todo(chaos): restitutionpadding
-			UpdateContactPoint(Constraint, BoxHeightFieldContactPoint(A, ATransform, B, BTransform, CullDistance, 0.0f), Dt);
+			UpdateContactPoint(Constraint, BoxHeightFieldContactPoint(A, ATransform, B, BTransform, CullDistance, 0.0f), CullDistance, Dt);
 		}
 
 		template<typename T_TRAITS>
@@ -761,7 +777,7 @@ namespace Chaos
 		{
 			// @todo(chaos): one-shot manifold
 			// @todo(chaos): restitutionpadding
-			UpdateContactPoint(Constraint, BoxPlaneContactPoint(A, B, ATransform, BTransform, CullDistance, 0.0f), Dt);
+			UpdateContactPoint(Constraint, BoxPlaneContactPoint(A, B, ATransform, BTransform, CullDistance, 0.0f), CullDistance, Dt);
 		}
 
 
@@ -797,7 +813,7 @@ namespace Chaos
 		void UpdateBoxTriangleMeshConstraint(const FImplicitBox3& Box0, const FRigidTransform3& WorldTransform0, const TriMeshType& TriangleMesh1, const FRigidTransform3& WorldTransform1, const FReal CullDistance, const FReal Dt, FRigidBodyPointContactConstraint& Constraint)
 		{
 			// @toto(chaos): restitutionpadding
-			UpdateContactPoint(Constraint, BoxTriangleMeshContactPoint(Box0, WorldTransform0, TriangleMesh1, WorldTransform1, CullDistance, 0.0f), Dt);
+			UpdateContactPoint(Constraint, BoxTriangleMeshContactPoint(Box0, WorldTransform0, TriangleMesh1, WorldTransform1, CullDistance, 0.0f), CullDistance, Dt);
 		}
 
 		template<typename T_TRAITS>
@@ -855,13 +871,13 @@ namespace Chaos
 			const FVec3 Direction = Center1 - Center2;
 			const FReal Size = Direction.Size();
 			const FVec3 Normal = Size > SMALL_NUMBER ? Direction / Size : FVec3(0, 0, 1);
-			const FReal NewPhi = Size - (R1 + R1);
+			const FReal NewPhi = Size - (R1 + R2);
 
 			// @todo(chaos): margin
 			Result.ShapeMargins[0] = 0.0f;
 			Result.ShapeMargins[1] = 0.0f;
-			Result.ShapeContactPoints[0] = Sphere1Transform.InverseTransformVector(-R1 * Normal);
-			Result.ShapeContactPoints[1] = Sphere2Transform.InverseTransformVector(R2 * Normal);
+			Result.ShapeContactPoints[0] = Sphere1.GetCenter() - Sphere1Transform.InverseTransformVector(R1 * Normal);
+			Result.ShapeContactPoints[1] = Sphere2.GetCenter() + Sphere2Transform.InverseTransformVector(R2 * Normal);
 			Result.ShapeContactNormal = Sphere2Transform.InverseTransformVector(Normal);
 
 			Result.Phi = NewPhi;
@@ -873,7 +889,7 @@ namespace Chaos
 
 		void UpdateSphereSphereConstraint(const TSphere<FReal, 3>& Sphere1, const FRigidTransform3& Sphere1Transform, const TSphere<FReal, 3>& Sphere2, const FRigidTransform3& Sphere2Transform, const FReal CullDistance, const FReal Dt, FRigidBodyPointContactConstraint& Constraint)
 		{
-			UpdateContactPoint(Constraint, SphereSphereContactPoint(Sphere1, Sphere1Transform, Sphere2, Sphere2Transform, CullDistance, Constraint.Manifold.RestitutionPadding), Dt);
+			UpdateContactPoint(Constraint, SphereSphereContactPoint(Sphere1, Sphere1Transform, Sphere2, Sphere2Transform, CullDistance, Constraint.Manifold.RestitutionPadding), CullDistance, Dt);
 		}
 
 		template<typename T_TRAITS>
@@ -907,13 +923,13 @@ namespace Chaos
 		void UpdateSphereHeightFieldConstraint(const TSphere<FReal, 3>& A, const FRigidTransform3& ATransform, const FHeightField& B, const FRigidTransform3& BTransform, const FReal CullDistance, const FReal Dt, FRigidBodyPointContactConstraint& Constraint)
 		{
 			// @todo(chaos): restitutionpadding
-			UpdateContactPoint(Constraint, SphereHeightFieldContactPoint(A, ATransform, B, BTransform, CullDistance, 0.0f), Dt);
+			UpdateContactPoint(Constraint, SphereHeightFieldContactPoint(A, ATransform, B, BTransform, CullDistance, 0.0f), CullDistance, Dt);
 		}
 
 		void UpdateSphereHeightFieldConstraintSwept(TGeometryParticleHandle<FReal, 3>* Particle0, const TSphere<FReal, 3>& A, const FRigidTransform3& ATransform, const FHeightField& B, const FRigidTransform3& BTransform, const FVec3& Dir, const FReal Length, const FReal CullDistance, const FReal Dt, FRigidBodySweptPointContactConstraint& Constraint)
 		{
 			FReal TOI = 1.0f;
-			UpdateContactPoint(Constraint, GJKImplicitSweptContactPoint<TSphere<FReal, 3>>(A, ATransform, B, BTransform, Dir, Length, CullDistance, TOI), Dt);
+			UpdateContactPoint(Constraint, GJKImplicitSweptContactPoint<TSphere<FReal, 3>>(A, ATransform, B, BTransform, Dir, Length, CullDistance, TOI), CullDistance, Dt);
 			SetSweptConstraintTOI(Particle0, TOI, Length, Dir, Constraint);
 		}
 
@@ -983,7 +999,7 @@ namespace Chaos
 
 		void UpdateSpherePlaneConstraint(const TSphere<FReal, 3>& Sphere, const FRigidTransform3& SphereTransform, const TPlane<FReal, 3>& Plane, const FRigidTransform3& PlaneTransform, const FReal CullDistance, const FReal Dt, FRigidBodyPointContactConstraint& Constraint)
 		{
-			UpdateContactPoint(Constraint, SpherePlaneContactPoint(Sphere, SphereTransform, Plane, PlaneTransform, CullDistance, Constraint.Manifold.RestitutionPadding), Dt);
+			UpdateContactPoint(Constraint, SpherePlaneContactPoint(Sphere, SphereTransform, Plane, PlaneTransform, CullDistance, Constraint.Manifold.RestitutionPadding), CullDistance, Dt);
 		}
 
 		template<typename T_TRAITS>
@@ -1039,7 +1055,7 @@ namespace Chaos
 
 		void UpdateSphereBoxConstraint(const TSphere<FReal, 3>& Sphere, const FRigidTransform3& SphereTransform, const FImplicitBox3& Box, const FRigidTransform3& BoxTransform, const FReal CullDistance, const FReal Dt, FRigidBodyPointContactConstraint& Constraint)
 		{
-			UpdateContactPoint(Constraint, SphereBoxContactPoint(Sphere, SphereTransform, Box, BoxTransform, CullDistance, Constraint.Manifold.RestitutionPadding), Dt);
+			UpdateContactPoint(Constraint, SphereBoxContactPoint(Sphere, SphereTransform, Box, BoxTransform, CullDistance, Constraint.Manifold.RestitutionPadding), CullDistance, Dt);
 		}
 
 		template<typename T_TRAITS>
@@ -1100,7 +1116,7 @@ namespace Chaos
 
 		void UpdateSphereCapsuleConstraint(const TSphere<FReal, 3>& A, const FRigidTransform3& ATransform, const TCapsule<FReal>& B, const FRigidTransform3& BTransform, const FReal CullDistance, const FReal Dt, FRigidBodyPointContactConstraint& Constraint)
 		{
-			UpdateContactPoint(Constraint, SphereCapsuleContactPoint(A, ATransform, B, BTransform, CullDistance, Constraint.Manifold.RestitutionPadding), Dt);
+			UpdateContactPoint(Constraint, SphereCapsuleContactPoint(A, ATransform, B, BTransform, CullDistance, Constraint.Manifold.RestitutionPadding), CullDistance, Dt);
 		}
 
 		template<typename T_TRAITS>
@@ -1121,6 +1137,33 @@ namespace Chaos
 			}
 		}
 
+		//
+		// Sphere-Convex
+		//
+
+		void UpdateSphereConvexConstraint(const TSphere<FReal, 3>& A, const FRigidTransform3& ATransform, const FImplicitObject3& B, const FRigidTransform3& BTransform, const FReal CullDistance, const FReal Dt, FRigidBodyPointContactConstraint& Constraint)
+		{
+			UpdateContactPoint(Constraint, SphereConvexContactPoint(A, ATransform, B, BTransform, CullDistance), CullDistance, Dt);
+		}
+
+		template<typename T_TRAITS>
+		void ConstructSphereConvexConstraints(TGeometryParticleHandle<FReal, 3>* Particle0, TGeometryParticleHandle<FReal, 3>* Particle1, const FImplicitObject* Implicit0, const FImplicitObject* Implicit1, const FRigidTransform3& LocalTransform0, const FRigidTransform3& LocalTransform1, const FReal CullDistance, const FReal Dt, const FCollisionContext& Context, FCollisionConstraintsArray& NewConstraints)
+		{
+			const TSphere<FReal, 3>* Object0 = Implicit0->template GetObject<const TSphere<FReal, 3> >();
+			const FImplicitObject* Object1 = Implicit1;
+			if (ensure(Object0 && Object1))
+			{
+				FRigidBodyPointContactConstraint Constraint = FRigidBodyPointContactConstraint(Particle0, Implicit0, nullptr, LocalTransform0, Particle1, Implicit1, nullptr, LocalTransform1, EContactShapesType::SphereConvex, false);	// No manifold
+				if (T_TRAITS::bImmediateUpdate)
+				{
+					FRigidTransform3 WorldTransform0 = LocalTransform0 * Collisions::GetTransform(Particle0);
+					FRigidTransform3 WorldTransform1 = LocalTransform1 * Collisions::GetTransform(Particle1);
+					UpdateSphereConvexConstraint(*Object0, WorldTransform0, *Object1, WorldTransform1, CullDistance, Dt, Constraint);
+				}
+				NewConstraints.Add(Constraint);
+			}
+		}
+
 		//
 		// Sphere-TriangleMesh
 		//
@@ -1151,14 +1194,14 @@ namespace Chaos
 		void UpdateSphereTriangleMeshConstraint(const TSphere<FReal, 3>& Sphere0, const FRigidTransform3& WorldTransform0, const TriMeshType& TriangleMesh1, const FRigidTransform3& WorldTransform1, const FReal CullDistance, const FReal Dt, FRigidBodyPointContactConstraint& Constraint)
 		{
 			// @todo(chaos): restitutionpadding
-			UpdateContactPoint(Constraint, SphereTriangleMeshContactPoint(Sphere0, WorldTransform0, TriangleMesh1, WorldTransform1, CullDistance, 0.0f), Dt);
+			UpdateContactPoint(Constraint, SphereTriangleMeshContactPoint(Sphere0, WorldTransform0, TriangleMesh1, WorldTransform1, CullDistance, 0.0f), CullDistance, Dt);
 		}
 
 		template<typename TriMeshType>
 		void UpdateSphereTriangleMeshConstraintSwept(TGeometryParticleHandle<FReal, 3>* Particle0, const TSphere<FReal, 3>& Sphere0, const FRigidTransform3& WorldTransform0, const TriMeshType& TriangleMesh1, const FRigidTransform3& WorldTransform1, const FVec3& Dir, const FReal Length, const FReal CullDistance, const FReal Dt, FRigidBodySweptPointContactConstraint& Constraint)
 		{
 			FReal TOI = 1.0f;
-			UpdateContactPoint(Constraint, SphereTriangleMeshSweptContactPoint(Sphere0, WorldTransform0, TriangleMesh1, WorldTransform1, Dir, Length, CullDistance, TOI), Dt);
+			UpdateContactPoint(Constraint, SphereTriangleMeshSweptContactPoint(Sphere0, WorldTransform0, TriangleMesh1, WorldTransform1, Dir, Length, CullDistance, TOI), CullDistance, Dt);
 			SetSweptConstraintTOI(Particle0, TOI, Length, Dir, Constraint);
 		}
 
@@ -1269,7 +1312,7 @@ namespace Chaos
 
 		void UpdateCapsuleCapsuleConstraint(const TCapsule<FReal>& A, const FRigidTransform3& ATransform, const TCapsule<FReal>& B, const FRigidTransform3& BTransform, const FReal CullDistance, const FReal Dt, FRigidBodyPointContactConstraint& Constraint)
 		{
-			UpdateContactPoint(Constraint, CapsuleCapsuleContactPoint(A, ATransform, B, BTransform, CullDistance, Constraint.Manifold.RestitutionPadding), Dt);
+			UpdateContactPoint(Constraint, CapsuleCapsuleContactPoint(A, ATransform, B, BTransform, CullDistance, Constraint.Manifold.RestitutionPadding), CullDistance, Dt);
 		}
 
 		template<typename T_TRAITS>
@@ -1312,7 +1355,7 @@ namespace Chaos
 			if (CapsuleAABB.Intersects(B.BoundingBox()))
 			{
 				const FVec3 InitialDir = ATransform.GetRotation().Inverse() * -Constraint.GetNormal();
-				UpdateContactPoint(Constraint, CapsuleBoxContactPoint(A, ATransform, B, BTransform, InitialDir, CullDistance, Constraint.Manifold.RestitutionPadding),Dt);
+				UpdateContactPoint(Constraint, CapsuleBoxContactPoint(A, ATransform, B, BTransform, InitialDir, CullDistance, Constraint.Manifold.RestitutionPadding), CullDistance, Dt);
 			}
 		}
 
@@ -1540,6 +1583,35 @@ namespace Chaos
 			}
 		}
 
+
+		//
+		// Capsule-Convex
+		//
+
+		void UpdateCapsuleConvexConstraint(const TCapsule<FReal>& A, const FRigidTransform3& ATransform, const FImplicitObject& B, const FRigidTransform3& BTransform, const FReal CullDistance, const FReal Dt, FRigidBodyPointContactConstraint& Constraint)
+		{
+			UpdateContactPoint(Constraint, CapsuleConvexContactPoint(A, ATransform, B, BTransform, CullDistance), CullDistance, Dt);
+		}
+
+		template<typename T_TRAITS>
+		void ConstructCapsuleConvexConstraints(TGeometryParticleHandle<FReal, 3>* Particle0, TGeometryParticleHandle<FReal, 3>* Particle1, const FImplicitObject* Implicit0, const FImplicitObject* Implicit1, const FRigidTransform3& LocalTransform0, const FRigidTransform3& LocalTransform1, const FReal CullDistance, const FReal Dt, const FCollisionContext& Context, FCollisionConstraintsArray& NewConstraints)
+		{
+			const TCapsule<FReal>* Object0 = Implicit0->template GetObject<const TCapsule<FReal>>();
+			const FImplicitObject* Object1 = Implicit1;
+			if (ensure(Object0 && Object1))
+			{
+				FRigidBodyPointContactConstraint Constraint = FRigidBodyPointContactConstraint(Particle0, Implicit0, nullptr, LocalTransform0, Particle1, Implicit1, nullptr, LocalTransform1, EContactShapesType::CapsuleConvex, false);
+				if (T_TRAITS::bImmediateUpdate)
+				{
+					FRigidTransform3 WorldTransform0 = LocalTransform0 * Collisions::GetTransform(Particle0);
+					FRigidTransform3 WorldTransform1 = LocalTransform1 * Collisions::GetTransform(Particle1);
+					UpdateCapsuleConvexConstraint(*Object0, WorldTransform0, *Object1, WorldTransform1, CullDistance, Dt, Constraint);
+				}
+				NewConstraints.Add(Constraint);
+			}
+		}
+
+
 		//
 		// Capsule-HeightField
 		//
@@ -1556,7 +1628,7 @@ namespace Chaos
 		{
 			// @todo(chaos): restitutionpadding
 			CONDITIONAL_SCOPE_CYCLE_COUNTER(STAT_Collisions_UpdateCapsuleHeightFieldConstraint, ConstraintsDetailedStats);
-			UpdateContactPoint(Constraint, CapsuleHeightFieldContactPoint(A, ATransform, B, BTransform, CullDistance, 0.0f), Dt);
+			UpdateContactPoint(Constraint, CapsuleHeightFieldContactPoint(A, ATransform, B, BTransform, CullDistance, 0.0f), CullDistance, Dt);
 		}
 
 
@@ -1564,7 +1636,7 @@ namespace Chaos
 		{
 			CONDITIONAL_SCOPE_CYCLE_COUNTER(STAT_Collisions_UpdateCapsuleHeightFieldConstraintSwept, ConstraintsDetailedStats);
 			FReal TOI = 1.0f;
-			UpdateContactPoint(Constraint, GJKImplicitSweptContactPoint<TCapsule<float> >(A, ATransform, B, BTransform, Dir, Length, CullDistance, TOI), Dt);
+			UpdateContactPoint(Constraint, GJKImplicitSweptContactPoint<TCapsule<float> >(A, ATransform, B, BTransform, Dir, Length, CullDistance, TOI), CullDistance, Dt);
 			SetSweptConstraintTOI(Particle0, TOI, Length, Dir, Constraint);
 		}
 
@@ -1638,7 +1710,7 @@ namespace Chaos
 		{
 			// @todo(chaos): restitutionpadding
 			CONDITIONAL_SCOPE_CYCLE_COUNTER(STAT_Collisions_UpdateCapsuleTriangleMeshConstraint, ConstraintsDetailedStats);
-			UpdateContactPoint(Constraint, CapsuleTriangleMeshContactPoint(Capsule0, WorldTransform0, TriangleMesh1, WorldTransform1, CullDistance, 0.0f), Dt);
+			UpdateContactPoint(Constraint, CapsuleTriangleMeshContactPoint(Capsule0, WorldTransform0, TriangleMesh1, WorldTransform1, CullDistance, 0.0f), CullDistance, Dt);
 		}
 
 		template <typename TriMeshType>
@@ -1646,7 +1718,7 @@ namespace Chaos
 		{
 			CONDITIONAL_SCOPE_CYCLE_COUNTER(STAT_Collisions_UpdateCapsuleTriangleMeshConstraint, ConstraintsDetailedStats);
 			FReal TOI = 1.0f;
-			UpdateContactPoint(Constraint, CapsuleTriangleMeshSweptContactPoint(Capsule0, WorldTransform0, TriangleMesh1, WorldTransform1, Dir, Length, CullDistance, TOI), Dt);
+			UpdateContactPoint(Constraint, CapsuleTriangleMeshSweptContactPoint(Capsule0, WorldTransform0, TriangleMesh1, WorldTransform1, Dir, Length, CullDistance, TOI), CullDistance, Dt);
 			SetSweptConstraintTOI(Particle0, TOI, Length, Dir, Constraint);
 		}
 
@@ -1732,7 +1804,7 @@ namespace Chaos
 		{
 			CONDITIONAL_SCOPE_CYCLE_COUNTER(STAT_Collisions_UpdateGenericConvexConvexConstraintSwept, ConstraintsDetailedStats);
 			FReal TOI = 1.0f;
-			UpdateContactPoint(Constraint, GenericConvexConvexContactPointSwept(Implicit0, WorldTransform0, Implicit1, WorldTransform1, Dir, Length, CullDistance, TOI), Dt);
+			UpdateContactPoint(Constraint, GenericConvexConvexContactPointSwept(Implicit0, WorldTransform0, Implicit1, WorldTransform1, Dir, Length, CullDistance, TOI), CullDistance, Dt);
 			SetSweptConstraintTOI(Particle0, TOI, Length, Dir, Constraint);
 		}
 
@@ -1780,7 +1852,7 @@ namespace Chaos
 		{
 			// @todo(chaos): restitutionpadding
 			CONDITIONAL_SCOPE_CYCLE_COUNTER(STAT_Collisions_UpdateConvexHeightFieldConstraint, ConstraintsDetailedStats);
-			UpdateContactPoint(Constraint, ConvexHeightFieldContactPoint(A, ATransform, B, BTransform, CullDistance, 0.0f), Dt);
+			UpdateContactPoint(Constraint, ConvexHeightFieldContactPoint(A, ATransform, B, BTransform, CullDistance, 0.0f), CullDistance, Dt);
 		}
 
 
@@ -1788,7 +1860,7 @@ namespace Chaos
 		{
 			CONDITIONAL_SCOPE_CYCLE_COUNTER(STAT_Collisions_UpdateConvexHeightFieldConstraintSwept, ConstraintsDetailedStats);
 			FReal TOI = 1.0f;
-			UpdateContactPoint(Constraint, GJKImplicitSweptContactPoint< FConvex >(A, ATransform, B, BTransform, Dir, Length, CullDistance, TOI), Dt);
+			UpdateContactPoint(Constraint, GJKImplicitSweptContactPoint< FConvex >(A, ATransform, B, BTransform, Dir, Length, CullDistance, TOI), CullDistance, Dt);
 			SetSweptConstraintTOI(Particle0, TOI, Length, Dir, Constraint);
 		}
 
@@ -1857,7 +1929,7 @@ namespace Chaos
 		{
 			// @todo(chaos): restitutionpadding
 			CONDITIONAL_SCOPE_CYCLE_COUNTER(STAT_Collisions_UpdateConvexTriangleMeshConstraint, ConstraintsDetailedStats);
-			UpdateContactPoint(Constraint, ConvexTriangleMeshContactPoint(Convex0, WorldTransform0, TriangleMesh1, WorldTransform1, CullDistance, 0.0f), Dt);
+			UpdateContactPoint(Constraint, ConvexTriangleMeshContactPoint(Convex0, WorldTransform0, TriangleMesh1, WorldTransform1, CullDistance, 0.0f), CullDistance, Dt);
 		}
 
 		// Sweeps convex against trimesh
@@ -1866,7 +1938,7 @@ namespace Chaos
 		{
 			CONDITIONAL_SCOPE_CYCLE_COUNTER(STAT_Collisions_UpdateConvexTriangleMeshConstraintSwept, ConstraintsDetailedStats);
 			FReal TOI = 1.0f;
-			UpdateContactPoint(Constraint, ConvexTriangleMeshSweptContactPoint(Convex0, WorldTransform0, TriangleMesh1, WorldTransform1, Dir, Length, CullDistance, TOI), Dt);
+			UpdateContactPoint(Constraint, ConvexTriangleMeshSweptContactPoint(Convex0, WorldTransform0, TriangleMesh1, WorldTransform1, Dir, Length, CullDistance, TOI), CullDistance, Dt);
 			SetSweptConstraintTOI(Particle0, TOI, Length, Dir, Constraint);
 		}
 
@@ -1978,7 +2050,7 @@ namespace Chaos
 					ContactPoint.ShapeContactNormal = WorldTransform1.InverseTransformVector(ContactPoint.Normal);
 				}
 
-				UpdateContactPoint(Constraint, ContactPoint, Dt);
+				UpdateContactPoint(Constraint, ContactPoint, CullDistance, Dt);
 			}
 		}
 
@@ -2084,8 +2156,7 @@ namespace Chaos
 				UpdateSphereBoxConstraint(*Implicit0.template GetObject<TSphere<FReal, 3>>(), WorldTransform0, *Implicit1.template GetObject<TBox<FReal, 3>>(), WorldTransform1, CullDistance, Dt, Constraint);
 				break;
 			case EContactShapesType::SphereConvex:
-				// MISSING CASE!!!
-				UpdateGenericConvexConvexConstraint(Implicit0, WorldTransform0, Implicit1, WorldTransform1, CullDistance, Dt, Constraint);
+				UpdateSphereConvexConstraint(*Implicit0.template GetObject<TSphere<FReal, 3>>(), WorldTransform0, Implicit1, WorldTransform1, CullDistance, Dt, Constraint);
 				break;
 			case EContactShapesType::SphereTriMesh:
 				if (const TImplicitObjectScaled<FTriangleMeshImplicitObject>* ScaledTriMesh = Implicit1.template GetObject<const TImplicitObjectScaled<FTriangleMeshImplicitObject> >())
@@ -2113,8 +2184,7 @@ namespace Chaos
 				UpdateCapsuleBoxConstraint(*Implicit0.template GetObject<TCapsule<FReal>>(), WorldTransform0, *Implicit1.template GetObject<TBox<FReal, 3>>(), WorldTransform1, CullDistance, Dt, Constraint);
 				break;
 			case EContactShapesType::CapsuleConvex:
-				// MISSING CASE!!!
-				UpdateGenericConvexConvexConstraint(Implicit0, WorldTransform0, Implicit1, WorldTransform1, CullDistance, Dt, Constraint);
+				UpdateCapsuleConvexConstraint(*Implicit0.template GetObject<TCapsule<FReal>>(), WorldTransform0, Implicit1, WorldTransform1, CullDistance, Dt, Constraint);
 				break;
 			case EContactShapesType::CapsuleTriMesh:
 				if (const TImplicitObjectScaled<FTriangleMeshImplicitObject>* ScaledTriMesh = Implicit1.template GetObject<const TImplicitObjectScaled<FTriangleMeshImplicitObject> >())
@@ -2136,7 +2206,6 @@ namespace Chaos
 				UpdateBoxBoxConstraint(*Implicit0.template GetObject<TBox<FReal, 3>>(), WorldTransform0, *Implicit1.template GetObject<TBox<FReal, 3>>(), WorldTransform1, CullDistance, Dt, Constraint);
 				break;
 			case EContactShapesType::BoxConvex:
-				// MISSING CASE!!!
 				UpdateGenericConvexConvexConstraint(Implicit0, WorldTransform0, Implicit1, WorldTransform1, CullDistance, Dt, Constraint);
 				break;
 			case EContactShapesType::BoxTriMesh:
@@ -2206,6 +2275,39 @@ namespace Chaos
 			FVec3 DirCCD(0.0f);
 			bool bUseCCD = UseCCD(Particle0, Particle1, Implicit0, DirCCD, LengthCCD);
 
+#if CHAOS_COLLISION_CREATE_BOUNDSCHECK
+			if (Implicit0->HasBoundingBox() && Implicit1->HasBoundingBox())
+			{
+				const FRigidTransform3 ParticleTransform0 = Collisions::GetTransform(Particle0);
+				const FRigidTransform3 ParticleTransform1 = Collisions::GetTransform(Particle1);
+				const FRigidTransform3 WorldTransform0 = LocalTransform0 * ParticleTransform0;
+				const FRigidTransform3 WorldTransform1 = LocalTransform1 * ParticleTransform1;
+				if (Chaos_Collision_NarrowPhase_SphereBoundsCheck)
+				{
+					const FReal R1 = Implicit0->BoundingBox().OriginRadius();
+					const FReal R2 = Implicit1->BoundingBox().OriginRadius();
+					const FReal SeparationSq = (WorldTransform1.GetTranslation() - WorldTransform0.GetTranslation()).SizeSquared();
+					const FReal CullDistanceSq = CullDistance * CullDistance;
+					if ((SeparationSq - FMath::Square(R1 + R2)) > CullDistanceSq)
+					{
+						return;
+					}
+				}
+
+				if (Chaos_Collision_NarrowPhase_AABBBoundsCheck)
+				{
+					const FRigidTransform3 Box2ToBox1TM = WorldTransform1.GetRelativeTransform(WorldTransform0);
+					const FAABB3 Box1 = Implicit0->BoundingBox();
+					const FAABB3 Box2In1 = Implicit1->BoundingBox().TransformedAABB(Box2ToBox1TM).Thicken(CullDistance);
+					if (!Box1.Intersects(Box2In1))
+					{
+						return;
+					}
+				}
+			}
+#endif
+
+
 			if (Implicit0Type == TBox<FReal, 3>::StaticType() && Implicit1Type == TBox<FReal, 3>::StaticType())
 			{
 				ConstructBoxBoxConstraints<T_TRAITS>(Particle0, Particle1, Implicit0, Implicit1, LocalTransform0, LocalTransform1, CullDistance, Dt, Context, NewConstraints);
@@ -2272,6 +2374,14 @@ namespace Chaos
 			{
 				ConstructSphereCapsuleConstraints<T_TRAITS>(Particle1, Particle0, Implicit1, Implicit0, LocalTransform1, LocalTransform0, CullDistance, Dt, Context, NewConstraints);
 			}
+			else if (Implicit0Type == TSphere<FReal, 3>::StaticType() && Implicit1Type == FImplicitConvex3::StaticType())
+			{
+				ConstructSphereConvexConstraints<T_TRAITS>(Particle0, Particle1, Implicit0, Implicit1, LocalTransform0, LocalTransform1, CullDistance, Dt, Context, NewConstraints);
+			}
+			else if (Implicit0Type == FImplicitConvex3::StaticType() && Implicit1Type == TSphere<FReal, 3>::StaticType())
+			{
+				ConstructSphereConvexConstraints<T_TRAITS>(Particle1, Particle0, Implicit1, Implicit0, LocalTransform1, LocalTransform0, CullDistance, Dt, Context, NewConstraints);
+			}
 			else if (Implicit0Type == TCapsule<FReal>::StaticType() && Implicit1Type == TCapsule<FReal>::StaticType())
 			{
 				ConstructCapsuleCapsuleConstraints<T_TRAITS>(Particle0, Particle1, Implicit0, Implicit1, LocalTransform0, LocalTransform1, CullDistance, Dt, Context, NewConstraints);
@@ -2284,6 +2394,14 @@ namespace Chaos
 			{
 				ConstructCapsuleBoxConstraints<T_TRAITS>(Particle1, Particle0, Implicit1, Implicit0, LocalTransform1, LocalTransform0, CullDistance, Dt, Context, NewConstraints);
 			}
+			else if (Implicit0Type == TCapsule<FReal>::StaticType() && Implicit1Type == FImplicitConvex3::StaticType())
+			{
+				ConstructCapsuleConvexConstraints<T_TRAITS>(Particle0, Particle1, Implicit0, Implicit1, LocalTransform0, LocalTransform1, CullDistance, Dt, Context, NewConstraints);
+			}
+			else if (Implicit0Type == FImplicitConvex3::StaticType() && Implicit1Type == TCapsule<FReal>::StaticType())
+			{
+				ConstructCapsuleConvexConstraints<T_TRAITS>(Particle1, Particle0, Implicit1, Implicit0, LocalTransform1, LocalTransform0, CullDistance, Dt, Context, NewConstraints);
+			}
 			else if (Implicit0Type == TCapsule<FReal>::StaticType() && Implicit1Type == FHeightField::StaticType())
 			{
 				if (bUseCCD)
diff --git a/Engine/Source/Runtime/Experimental/Chaos/Private/Chaos/Convex.cpp b/Engine/Source/Runtime/Experimental/Chaos/Private/Chaos/Convex.cpp
--- a/Engine/Source/Runtime/Experimental/Chaos/Private/Chaos/Convex.cpp
+++ b/Engine/Source/Runtime/Experimental/Chaos/Private/Chaos/Convex.cpp
@@ -1,10 +1,20 @@
 // Copyright Epic Games, Inc. All Rights Reserved.
 #include "Chaos/Convex.h"
+#include "Chaos/GJK.h"
+#include "Chaos/Sphere.h"
 
 //PRAGMA_DISABLE_OPTIMIZATION
 
 namespace Chaos
 {
+	bool FConvex::Raycast(const FVec3& StartPoint, const FVec3& Dir, const FReal Length, const FReal Thickness, FReal& OutTime, FVec3& OutPosition, FVec3& OutNormal, int32& OutFaceIndex) const
+	{
+		OutFaceIndex = INDEX_NONE;	//finding face is expensive, should be called directly by user
+		const FRigidTransform3 StartTM(StartPoint, FRotation3::FromIdentity());
+		const TSphere<FReal, 3> Sphere(FVec3(0), Thickness);
+		return GJKRaycast(*this, Sphere, StartTM, Dir, Length, OutTime, OutPosition, OutNormal);
+	}
+
 
 	int32 FConvex::FindMostOpposingFace(const FVec3& Position, const FVec3& UnitDir, int32 HintFaceIndex, FReal SearchDist) const
 	{
diff --git a/Engine/Source/Runtime/Experimental/Chaos/Public/Chaos/AABB.h b/Engine/Source/Runtime/Experimental/Chaos/Public/Chaos/AABB.h
--- a/Engine/Source/Runtime/Experimental/Chaos/Public/Chaos/AABB.h
+++ b/Engine/Source/Runtime/Experimental/Chaos/Public/Chaos/AABB.h
@@ -481,6 +481,7 @@ namespace Chaos
 		FORCEINLINE static T GetVolume(const TVector<T, 3>& Dim) { return Dim.Product(); }
 
 		FORCEINLINE T GetMargin() const { return 0; }
+		FORCEINLINE T GetRadius() const { return 0; }
 
 		FORCEINLINE static TAABB<T, d> EmptyAABB() { return TAABB<T, d>(TVector<T, d>(TNumericLimits<T>::Max()), TVector<T, d>(-TNumericLimits<T>::Max())); }
 		FORCEINLINE static TAABB<T, d> ZeroAABB() { return TAABB<T, d>(TVector<T, d>((T)0), TVector<T, d>((T)0)); }
diff --git a/Engine/Source/Runtime/Experimental/Chaos/Public/Chaos/Box.h b/Engine/Source/Runtime/Experimental/Chaos/Public/Chaos/Box.h
--- a/Engine/Source/Runtime/Experimental/Chaos/Public/Chaos/Box.h
+++ b/Engine/Source/Runtime/Experimental/Chaos/Public/Chaos/Box.h
@@ -94,6 +94,11 @@ namespace Chaos
 			return TUniquePtr<FImplicitObject>(new TBox<T,d>(*this));
 		}
 
+		float GetRadius() const
+		{
+			return 0.0f;
+		}
+
 		/**
 		 * Returns sample points centered about the origin.
 		 */
@@ -287,13 +292,14 @@ namespace Chaos
 		// Returns a position on the core shape excluding the margin
 		FORCEINLINE_DEBUGGABLE TVector<T, d> SupportCore(const TVector<T, d>& Direction, float InMargin) const
 		{
-			return AABB.SupportCore(Direction, InMargin + GetMargin());
+			return AABB.SupportCore(Direction, InMargin);
 		}
 
 		FORCEINLINE_DEBUGGABLE TVector<T, d> SupportCoreScaled(const TVector<T, d>& Direction, float InMargin, const TVector<T, d>& Scale) const
 		{
+			// @todo(chaos): Needs to operate in scaled space as margin is not non-uniform scalable
 			const FReal InvScale = 1.0f / Scale[0];
-			const FReal NetMargin = InvScale * (InMargin + GetMargin());
+			const FReal NetMargin = InvScale * InMargin;
 			return AABB.SupportCore(Direction * Scale, NetMargin) * Scale;
 		}
 
diff --git a/Engine/Source/Runtime/Experimental/Chaos/Public/Chaos/Capsule.h b/Engine/Source/Runtime/Experimental/Chaos/Public/Chaos/Capsule.h
--- a/Engine/Source/Runtime/Experimental/Chaos/Public/Chaos/Capsule.h
+++ b/Engine/Source/Runtime/Experimental/Chaos/Public/Chaos/Capsule.h
@@ -2,6 +2,8 @@
 #pragma once
 
 #include "Chaos/Cylinder.h"
+#include "Chaos/GJKShape.h"
+#include "Chaos/ImplicitFwd.h"
 #include "Chaos/ImplicitObject.h"
 #include "Chaos/ImplicitObjectUnion.h"
 #include "Chaos/Sphere.h"
@@ -70,6 +72,11 @@ namespace Chaos
 			return TCapsule<T>(X1, X2, Radius);
 		}
 
+		float GetRadius() const
+		{
+			return Margin;
+		}
+
 		/**
 		 * Returns sample points centered about the origin.
 		 *
@@ -286,11 +293,13 @@ namespace Chaos
 
 		FORCEINLINE TVector<T, 3> SupportCore(const TVector<T, 3>& Direction, float InMargin) const
 		{
+			// NOTE: Ignores InMargin, assumes Radius
 			return MSegment.SupportCore(Direction);
 		}
 
 		FORCEINLINE TVector<T, 3> SupportCoreScaled(const TVector<T, 3>& Direction, float InMargin, const TVector<T, 3>& Scale) const
 		{
+			// NOTE: Ignores InMargin, assumes Radius
 			return SupportCore(Scale * Direction, GetMargin()) * Scale;
 		}
 
@@ -330,7 +339,6 @@ namespace Chaos
 			return TUniquePtr<FImplicitObject>(new TCapsule<T>(*this));
 		}
 
-		const T GetRadius() const { return GetMargin(); }
 		T GetHeight() const { return MSegment.GetLength(); }
 		/** Returns the bottommost point on the capsule. */
 		const TVector<T, 3> GetOrigin() const { return GetX1() + GetAxis() * -GetRadius(); }
@@ -514,4 +522,5 @@ namespace Chaos
 			}
 		}
 	};
+
 } // namespace Chaos
diff --git a/Engine/Source/Runtime/Experimental/Chaos/Public/Chaos/Collision/CapsuleConvexContactPoint.h b/Engine/Source/Runtime/Experimental/Chaos/Public/Chaos/Collision/CapsuleConvexContactPoint.h
new file mode 100644
--- /dev/null
+++ b/Engine/Source/Runtime/Experimental/Chaos/Public/Chaos/Collision/CapsuleConvexContactPoint.h
@@ -0,0 +1,22 @@
+// Copyright Epic Games, Inc. All Rights Reserved.
+#pragma once
+
+#include "CoreMinimal.h"
+#include "Chaos/Core.h"
+#include "Chaos/ImplicitFwd.h"
+#include "Chaos/ImplicitObjectScaled.h"	// Cannot fwd declare scaled implicit
+
+namespace Chaos
+{
+	class FContactPoint;
+
+	extern CHAOS_API FContactPoint CapsuleConvexContactPoint(const FImplicitCapsule3& Capsule, const FRigidTransform3& CapsuleTransform, const FImplicitObject& Convex, const FRigidTransform3& ConvexTransform, FReal CullDistance);
+
+	extern CHAOS_API FContactPoint CapsuleConvexContactPoint(const FImplicitCapsule3& Capsule, const FRigidTransform3& CapsuleTransform, const FImplicitConvex3& Convex, const FRigidTransform3& ConvexTransform, FReal CullDistance);
+	extern CHAOS_API FContactPoint CapsuleConvexContactPoint(const FImplicitCapsule3& Capsule, const FRigidTransform3& CapsuleTransform, const TImplicitObjectInstanced<FImplicitConvex3>& Convex, const FRigidTransform3& ConvexTransform, FReal CullDistance);
+	extern CHAOS_API FContactPoint CapsuleConvexContactPoint(const FImplicitCapsule3& Capsule, const FRigidTransform3& CapsuleTransform, const TImplicitObjectScaled<FImplicitConvex3>& Convex, const FRigidTransform3& ConvexTransform, FReal CullDistance);
+
+	extern CHAOS_API FContactPoint ConvexCapsuleContactPoint(const FImplicitConvex3& Convex, const FRigidTransform3& ConvexTransform, const FImplicitCapsule3& Capsule, const FRigidTransform3& CapsuleTransform, FReal CullDistance);
+	extern CHAOS_API FContactPoint ConvexCapsuleContactPoint(const TImplicitObjectInstanced<FImplicitConvex3>& Convex, const FRigidTransform3& ConvexTransform, const FImplicitCapsule3& Capsule, const FRigidTransform3& CapsuleTransform, FReal CullDistance);
+	extern CHAOS_API FContactPoint ConvexCapsuleContactPoint(const TImplicitObjectScaled<FImplicitConvex3>& Convex, const FRigidTransform3& ConvexTransform, const FImplicitCapsule3& Capsule, const FRigidTransform3& CapsuleTransform, FReal CullDistance);
+}
diff --git a/Engine/Source/Runtime/Experimental/Chaos/Public/Chaos/Collision/ContactPoint.h b/Engine/Source/Runtime/Experimental/Chaos/Public/Chaos/Collision/ContactPoint.h
new file mode 100644
--- /dev/null
+++ b/Engine/Source/Runtime/Experimental/Chaos/Public/Chaos/Collision/ContactPoint.h
@@ -0,0 +1,51 @@
+// Copyright Epic Games, Inc.All Rights Reserved.
+#pragma once
+
+#include "CoreMinimal.h"
+#include "Templates/UnrealTemplate.h"
+
+#include "Chaos/Core.h"
+#include "Chaos/Vector.h"
+
+
+namespace Chaos
+{
+	// Data returned by the low-level collision functions
+	class CHAOS_API FContactPoint
+	{
+	public:
+		FVec3 ShapeContactPoints[2];	// Shape-space contact points on the two bodies, without the margin added (i.e., contact point on the core shape)
+		FVec3 ShapeContactNormal;		// Shape-space contact normal relative to the NormalOwner, but with direction that always goes from body 1 to body 0
+		FReal ShapeMargins[2];			// Margins used in collision detection
+		int32 ContactNormalOwnerIndex;	// The shape which owns the contact normal (usually the second body, but not always for manifolds)
+
+		// @todo(chaos): these do not need to be stored here (they can be derived from above)
+		FVec3 Location;					// World-space contact location
+		FVec3 Normal;					// World-space contact normal
+		FReal Phi;						// Contact separation (negative for overlap)
+
+		FContactPoint()
+			: ContactNormalOwnerIndex(1)
+			, Normal(1, 0, 0)
+			, Phi(TNumericLimits<FReal>::Max())
+		{
+		}
+
+		// Whether the contact point has been set up with contact data
+		bool IsSet() const { return (Phi != TNumericLimits<FReal>::Max()); }
+
+		// Switch the shape indices. For use when calling a collision detection method which takes shape types in the opposite order to what you want.
+		FContactPoint& SwapShapes()
+		{
+			if (IsSet())
+			{
+				Swap(ShapeContactPoints[0], ShapeContactPoints[1]);
+				Swap(ShapeMargins[0], ShapeMargins[1]);
+				ShapeContactNormal = -ShapeContactNormal;
+				ContactNormalOwnerIndex = (ContactNormalOwnerIndex == 0) ? 1 : 0;
+				Normal = -Normal;
+			}
+			return *this;
+		}
+	};
+}
\ No newline at end of file
diff --git a/Engine/Source/Runtime/Experimental/Chaos/Public/Chaos/Collision/PBDCollisionConstraint.h b/Engine/Source/Runtime/Experimental/Chaos/Public/Chaos/Collision/PBDCollisionConstraint.h
--- a/Engine/Source/Runtime/Experimental/Chaos/Public/Chaos/Collision/PBDCollisionConstraint.h
+++ b/Engine/Source/Runtime/Experimental/Chaos/Public/Chaos/Collision/PBDCollisionConstraint.h
@@ -5,6 +5,7 @@
 
 #include "Chaos/Core.h"
 #include "Chaos/CollisionResolutionTypes.h"
+#include "Chaos/Collision/ContactPoint.h"
 #include "Chaos/ParticleHandleFwd.h"
 #include "Chaos/Vector.h"
 #include "Chaos/BVHParticles.h"
@@ -14,28 +15,6 @@ namespace Chaos
 	class FImplicitObject;
 	class FPBDCollisionConstraintHandle;
 
-	// Data returned by the low-level collision functions
-	class CHAOS_API FContactPoint
-	{
-	public:
-		FVec3 ShapeContactPoints[2];	// Shape-space contact points on the two bodies, without the margin added (i.e., contact point on the core shape)
-		FVec3 ShapeContactNormal;		// Shape-space contact normal relative to the NormalOwner, but with direction that always goes from body 1 to body 0
-		FReal ShapeMargins[2];			// Margins used in collision detection
-		int32 ContactNormalOwnerIndex;	// The shape which owns the contact normal (usually the second body, but not always for manifolds)
-
-		// @todo(chaos): these do not need to be stored here (they can be derived from above)
-		FVec3 Location;					// World-space contact location
-		FVec3 Normal;					// World-space contact normal
-		FReal Phi;						// Contact separation (negative for overlap)
-
-		FContactPoint()
-			: ContactNormalOwnerIndex(1)
-			, Normal(1, 0, 0)
-			, Phi(TNumericLimits<FReal>::Max()) {}
-
-		bool IsSet() const { return (Phi != TNumericLimits<FReal>::Max()); }
-	};
-
 	class CHAOS_API FManifoldPoint
 	{
 	public:
diff --git a/Engine/Source/Runtime/Experimental/Chaos/Public/Chaos/Collision/SpatialAccelerationBroadPhase.h b/Engine/Source/Runtime/Experimental/Chaos/Public/Chaos/Collision/SpatialAccelerationBroadPhase.h
--- a/Engine/Source/Runtime/Experimental/Chaos/Public/Chaos/Collision/SpatialAccelerationBroadPhase.h
+++ b/Engine/Source/Runtime/Experimental/Chaos/Public/Chaos/Collision/SpatialAccelerationBroadPhase.h
@@ -313,7 +313,7 @@ namespace Chaos
 						SCOPE_CYCLE_COUNTER(STAT_Collisions_GenerateCollisions);
 						// Generate constraints for the potentially overlapping shape pairs. Also run collision detection to generate
 						// the contact position and normal (for contacts within CullDistance) for use in collision callbacks.
-						NarrowPhase.GenerateCollisions(NewConstraints, Dt, Particle1.Handle(), Particle2.Handle(), CullDistance);
+						NarrowPhase.GenerateCollisions(NewConstraints, Dt, Particle1.Handle(), Particle2.Handle(), CullDistance + BoundsThicknessVelocityInflation);
 					}
 				}
 
diff --git a/Engine/Source/Runtime/Experimental/Chaos/Public/Chaos/Collision/SphereConvexContactPoint.h b/Engine/Source/Runtime/Experimental/Chaos/Public/Chaos/Collision/SphereConvexContactPoint.h
new file mode 100644
--- /dev/null
+++ b/Engine/Source/Runtime/Experimental/Chaos/Public/Chaos/Collision/SphereConvexContactPoint.h
@@ -0,0 +1,22 @@
+// Copyright Epic Games, Inc. All Rights Reserved.
+#pragma once
+
+#include "CoreMinimal.h"
+#include "Chaos/Core.h"
+#include "Chaos/ImplicitFwd.h"
+#include "Chaos/ImplicitObjectScaled.h"	// Cannot fwd declare scaled implicit
+
+namespace Chaos
+{
+	class FContactPoint;
+
+	extern CHAOS_API FContactPoint SphereConvexContactPoint(const FImplicitSphere3& Sphere, const FRigidTransform3& SphereTransform, const FImplicitObject& Convex, const FRigidTransform3& ConvexTransform, FReal CullDistance);
+
+	extern CHAOS_API FContactPoint SphereConvexContactPoint(const FImplicitSphere3& Sphere, const FRigidTransform3& SphereTransform, const FImplicitConvex3& Convex, const FRigidTransform3& ConvexTransform, FReal CullDistance);
+	extern CHAOS_API FContactPoint SphereConvexContactPoint(const FImplicitSphere3& Sphere, const FRigidTransform3& SphereTransform, const TImplicitObjectInstanced<FImplicitConvex3>& Convex, const FRigidTransform3& ConvexTransform, FReal CullDistance);
+	extern CHAOS_API FContactPoint SphereConvexContactPoint(const FImplicitSphere3& Sphere, const FRigidTransform3& SphereTransform, const TImplicitObjectScaled<FImplicitConvex3>& Convex, const FRigidTransform3& ConvexTransform, FReal CullDistance);
+
+	extern CHAOS_API FContactPoint ConvexSphereContactPoint(const FImplicitConvex3& Convex, const FRigidTransform3& ConvexTransform, const FImplicitSphere3& Sphere, const FRigidTransform3& SphereTransform, FReal CullDistance);
+	extern CHAOS_API FContactPoint ConvexSphereContactPoint(const TImplicitObjectInstanced<FImplicitConvex3>& Convex, const FRigidTransform3& ConvexTransform, const FImplicitSphere3& Sphere, const FRigidTransform3& SphereTransform, FReal CullDistance);
+	extern CHAOS_API FContactPoint ConvexSphereContactPoint(const TImplicitObjectScaled<FImplicitConvex3>& Convex, const FRigidTransform3& ConvexTransform, const FImplicitSphere3& Sphere, const FRigidTransform3& SphereTransform, FReal CullDistance);
+}
diff --git a/Engine/Source/Runtime/Experimental/Chaos/Public/Chaos/Convex.h b/Engine/Source/Runtime/Experimental/Chaos/Public/Chaos/Convex.h
--- a/Engine/Source/Runtime/Experimental/Chaos/Public/Chaos/Convex.h
+++ b/Engine/Source/Runtime/Experimental/Chaos/Public/Chaos/Convex.h
@@ -7,7 +7,6 @@
 #include "Chaos/MassProperties.h"
 #include "CollisionConvexMesh.h"
 #include "ChaosArchive.h"
-#include "GJK.h"
 #include "ChaosCheck.h"
 #include "ChaosLog.h"
 #include "UObject/ReleaseObjectVersion.h"
@@ -104,6 +103,16 @@ namespace Chaos
 			return ImplicitObjectType::Convex;
 		}
 
+		float GetMargin() const
+		{
+			return Margin;
+		}
+
+		float GetRadius() const
+		{
+			return 0.0f;
+		}
+
 		virtual const TAABB<FReal, 3> BoundingBox() const override
 		{
 			return LocalBoundingBox;
@@ -261,13 +270,7 @@ namespace Chaos
 		 * and \c OutNormal should be.  The burden for detecting this case is deferred to the
 		 * caller. 
 		 */
-		virtual bool Raycast(const FVec3& StartPoint, const FVec3& Dir, const FReal Length, const FReal Thickness, FReal& OutTime, FVec3& OutPosition, FVec3& OutNormal, int32& OutFaceIndex) const override
-		{
-			OutFaceIndex = INDEX_NONE;	//finding face is expensive, should be called directly by user
-			const FRigidTransform3 StartTM(StartPoint, FRotation3::FromIdentity());
-			const TSphere<FReal, 3> Sphere(FVec3(0), Thickness);
-			return GJKRaycast(*this, Sphere, StartTM, Dir, Length, OutTime, OutPosition, OutNormal);
-		}
+		virtual bool Raycast(const FVec3& StartPoint, const FVec3& Dir, const FReal Length, const FReal Thickness, FReal& OutTime, FVec3& OutPosition, FVec3& OutNormal, int32& OutFaceIndex) const override;
 
 		virtual Pair<FVec3, bool> FindClosestIntersectionImp(const FVec3& StartPoint, const FVec3& EndPoint, const FReal Thickness) const override
 		{
@@ -470,11 +473,9 @@ namespace Chaos
 			const int32 SupportVertexIndex = GetSupportVertex(Direction);
 			if (SupportVertexIndex != INDEX_NONE)
 			{
-				// @chaos(todo): apply an upper limit to the margin to prevent an inverted shape
-				const float NetMargin = InMargin + GetMargin();
-				if (NetMargin > SMALL_NUMBER)
+				if (InMargin > SMALL_NUMBER)
 				{
-					return GetMarginAdjustedVertex(SupportVertexIndex, NetMargin);
+					return GetMarginAdjustedVertex(SupportVertexIndex, InMargin);
 				}
 				return Vertices[SupportVertexIndex];
 			}
@@ -486,23 +487,18 @@ namespace Chaos
 		FVec3 SupportCoreScaled(const FVec3& Direction, float InMargin, const FVec3& Scale) const
 		{
 			// Find the supporting vertex index
-			const FVec3 DirectionScaled = Scale * Direction;	// Not normalized
+			const FVec3 DirectionScaled = Scale * Direction;	// does not need to be normalized
 			const int32 SupportVertexIndex = GetSupportVertex(DirectionScaled);
 
 			// Adjust the vertex position based on margin
 			FVec3 VertexPosition = FVec3(0);
 			if (SupportVertexIndex != INDEX_NONE)
 			{
-				// Note: Shapes wrapped in a non-uniform scale should not have their own margin
-				// because we do not support non-uniformly scaled margins on convex (to do so would 
-				// require that we dupe the convex data for each scale).
-				const FReal InverseScale = 1.0f / Scale[0];
-				const float NetMargin = InMargin + InverseScale * GetMargin();
-
+				// Note: Shapes wrapped in a non-uniform scale should not have their own margin and we assume that here
 				// @chaos(todo): apply an upper limit to the margin to prevent a non-convex or null shape (also see comments in GetMarginAdjustedVertex)
-				if (NetMargin > SMALL_NUMBER)
+				if (InMargin > SMALL_NUMBER)
 				{
-					VertexPosition = GetMarginAdjustedVertexScaled(SupportVertexIndex, NetMargin, Scale);
+					VertexPosition = GetMarginAdjustedVertexScaled(SupportVertexIndex, InMargin, Scale);
 				}
 				else
 				{
diff --git a/Engine/Source/Runtime/Experimental/Chaos/Public/Chaos/ConvexStructureData.h b/Engine/Source/Runtime/Experimental/Chaos/Public/Chaos/ConvexStructureData.h
--- a/Engine/Source/Runtime/Experimental/Chaos/Public/Chaos/ConvexStructureData.h
+++ b/Engine/Source/Runtime/Experimental/Chaos/Public/Chaos/ConvexStructureData.h
@@ -6,7 +6,6 @@
 #include "Chaos/MassProperties.h"
 #include "CollisionConvexMesh.h"
 #include "ChaosArchive.h"
-#include "GJK.h"
 #include "ChaosCheck.h"
 #include "ChaosLog.h"
 #include "UObject/PhysicsObjectVersion.h"
diff --git a/Engine/Source/Runtime/Experimental/Chaos/Public/Chaos/GJK.h b/Engine/Source/Runtime/Experimental/Chaos/Public/Chaos/GJK.h
--- a/Engine/Source/Runtime/Experimental/Chaos/Public/Chaos/GJK.h
+++ b/Engine/Source/Runtime/Experimental/Chaos/Public/Chaos/GJK.h
@@ -1,16 +1,20 @@
 // Copyright Epic Games, Inc. All Rights Reserved.
 #pragma once
 
-#include "Chaos/Simplex.h"
-#include "Chaos/Capsule.h"
 #include "Chaos/Box.h"
-#include "Chaos/Sphere.h"
+#include "Chaos/Capsule.h"
 #include "Chaos/EPA.h"
+#include "Chaos/GJKShape.h"
+#include "Chaos/ImplicitObjectScaled.h"
+#include "Chaos/Simplex.h"
+#include "Chaos/Sphere.h"
+
 #include "ChaosCheck.h"
 #include "ChaosLog.h"
 
 namespace Chaos
 {
+
 	
 	/** Determines if two convex geometries overlap.
 	 @A The first geometry
@@ -49,9 +53,9 @@ namespace Chaos
 				break;	//if taking too long just stop. This should never happen
 			}
 			const TVector<T, 3> NegV = -V;
-			const TVector<T, 3> SupportA = A.SupportCore(NegV, 0.0f);
+			const TVector<T, 3> SupportA = A.SupportCore(NegV, A.GetMargin());
 			const TVector<T, 3> VInB = AToBRotation * V;
-			const TVector<T, 3> SupportBLocal = B.SupportCore(VInB, 0.0f);
+			const TVector<T, 3> SupportBLocal = B.SupportCore(VInB, B.GetMargin());
 			const TVector<T, 3> SupportB = BToATM.TransformPositionNoScale(SupportBLocal);
 			const TVector<T, 3> W = SupportA - SupportB;
 
@@ -117,7 +121,7 @@ namespace Chaos
 		auto SupportAFunc = [&A, &VertexIndexA](const TVec3<T>& V)
 		{
 			VertexIndexA = INDEX_NONE;
-			return A.SupportCore(V, 0.0f);
+			return A.SupportCore(V, A.GetMargin());
 		};
 
 		const TRotation<T, 3> AToBRotation = BToATM.GetRotation().Inverse();
@@ -126,7 +130,7 @@ namespace Chaos
 		{
 			VertexIndexB = INDEX_NONE;
 			const TVector<T, 3> VInB = AToBRotation * V;
-			const TVector<T, 3> SupportBLocal = B.SupportCore(VInB, 0.0f);
+			const TVector<T, 3> SupportBLocal = B.SupportCore(VInB, B.GetMargin());
 			return BToATM.TransformPositionNoScale(SupportBLocal);
 		};
 
@@ -436,8 +440,26 @@ namespace Chaos
 	{
 		ensure(FMath::IsNearlyEqual(RayDir.SizeSquared(), 1, KINDA_SMALL_NUMBER));
 		ensure(RayLength > 0);
-		const T ThicknessA = A.GetMargin();
-		const T ThicknessB = B.GetMargin();
+
+		// Margin selection logic: we only need a small margin for sweeps since we only move the sweeping object
+		// to the point where it just touches.
+		// Spheres and Capsules: always use the core shape and full "margin" because it represents the radius
+		// Sphere/Capsule versus OtherShape: no margin on other
+		// OtherShape versus OtherShape: use margin of the smaller shape, zero margin on the other
+		const T RadiusA = A.GetRadius();
+		const T RadiusB = B.GetRadius();
+		const bool bHasRadiusA = RadiusA > 0;
+		const bool bHasRadiusB = RadiusB > 0;
+
+		// The sweep margins if required. Only one can be non-zero (we keep the smaller one)
+		const T SweepMarginScale = 0.05f;
+		const bool bAIsSmallest = A.GetMargin() < B.GetMargin();
+		const T SweepMarginA = (bHasRadiusA || bHasRadiusB) ? 0.0f : (bAIsSmallest ? SweepMarginScale * A.GetMargin() : 0.0f);
+		const T SweepMarginB = (bHasRadiusA || bHasRadiusB) ? 0.0f : (bAIsSmallest ? 0.0f : SweepMarginScale * B.GetMargin());
+
+		// Net margin (note: both SweepMargins are zero if either Radius is non-zero, and only one SweepMargin can be non-zero)
+		const T MarginA = RadiusA + SweepMarginA;
+		const T MarginB = RadiusB + SweepMarginB;
 
 		const TVector<T, 3> StartPoint = StartTM.GetLocation();
 
@@ -446,22 +468,22 @@ namespace Chaos
 		TVector<T, 3> Bs[4] = { TVector<T,3>(0), TVector<T,3>(0), TVector<T,3>(0), TVector<T,3>(0) };
 
 		T Barycentric[4] = { -1,-1,-1,-1 };	//not needed, but compiler warns
-		const T Inflation = ThicknessA + ThicknessB;
+		const T Inflation = MarginA + MarginB;
 		const T Inflation2 = Inflation*Inflation + 1e-6;
 
 		FSimplex SimplexIDs;
 		const TRotation<T, 3> BToARotation = StartTM.GetRotation();
 		const TRotation<T, 3> AToBRotation = BToARotation.Inverse();
 
-		auto SupportAFunc = [&A](const TVec3<T>& V)
+		auto SupportAFunc = [&A, MarginA](const TVec3<T>& V)
 		{
-			return A.SupportCore(V, 0.0f);
+			return A.SupportCore(V, MarginA);
 		};
 
-		auto SupportBFunc = [&B, &AToBRotation, &BToARotation](const TVec3<T>& V)
+		auto SupportBFunc = [&B, MarginB, &AToBRotation, &BToARotation](const TVec3<T>& V)
 		{
 			const TVector<T, 3> VInB = AToBRotation * V;
-			const TVector<T, 3> SupportBLocal = B.SupportCore(VInB, 0.0f);
+			const TVector<T, 3> SupportBLocal = B.SupportCore(VInB, MarginB);
 			return BToARotation * SupportBLocal;
 		};
 
@@ -585,7 +607,7 @@ namespace Chaos
 			{
 				ClosestB += Bs[i] * Barycentric[i];
 			}
-			const TVector<T, 3> ClosestLocal = ClosestB - OutNormal * ThicknessB;
+			const TVector<T, 3> ClosestLocal = ClosestB - OutNormal * MarginB;
 
 			OutPosition = StartPoint + RayDir * Lambda + ClosestLocal;
 		}
@@ -619,8 +641,8 @@ namespace Chaos
 				const T InGJKPreDist = FMath::Sqrt(InGJKPreDist2);
 				OutNormal = V.GetUnsafeNormal();
 
-				const T Penetration = FMath::Clamp<T>(ThicknessA + ThicknessB - InGJKPreDist, 0, TNumericLimits<T>::Max());
-				const TVector<T, 3> ClosestLocal = ClosestB - OutNormal * ThicknessB;
+				const T Penetration = FMath::Clamp<T>(MarginA + MarginB - InGJKPreDist, 0, TNumericLimits<T>::Max());
+				const TVector<T, 3> ClosestLocal = ClosestB - OutNormal * MarginB;
 
 				OutPosition = StartPoint + ClosestLocal + OutNormal * Penetration;
 				OutTime = -Penetration;
@@ -647,7 +669,7 @@ namespace Chaos
 					auto SupportBAtOriginFunc = [&](const TVec3<T>& Dir)
 					{
 						const TVector<T, 3> DirInB = AToBRotation * Dir;
-						const TVector<T, 3> SupportBLocal = B.SupportCore(DirInB, 0.0f);
+						const TVector<T, 3> SupportBLocal = B.SupportCore(DirInB, MarginB);
 						return StartTM.TransformPositionNoScale(SupportBLocal);
 					};
 
@@ -669,7 +691,7 @@ namespace Chaos
 						//assume touching hit
 						OutTime = -Inflation;
 						OutNormal = MTD;
-						OutPosition = As[0] + OutNormal * ThicknessA;
+						OutPosition = As[0] + OutNormal * MarginA;
 					}
 				}
 				else
@@ -677,7 +699,7 @@ namespace Chaos
 					//didn't even go into gjk loop, touching hit
 					OutTime = -Inflation;
 					OutNormal = { 0,0,1 };
-					OutPosition = As[0] + OutNormal * ThicknessA;
+					OutPosition = As[0] + OutNormal * MarginA;
 				}
 			}
 		}
@@ -697,12 +719,12 @@ namespace Chaos
 	 * the actual separating vector and GJK will converge immediately).
 	 */
 	template <typename T, typename TGeometryA, typename TGeometryB>
-	TVector<T, 3> GJKDistanceInitialV(const TGeometryA& A, const TGeometryB& B, const TRigidTransform<T, 3>& BToATM)
+	TVector<T, 3> GJKDistanceInitialV(const TGeometryA& A, T MarginA, const TGeometryB& B, T MarginB, const TRigidTransform<T, 3>& BToATM)
 	{
-		const TVec3<T> V(1, 0, 0);
-		const TVector<T, 3> SupportA = A.Support(-V, 0);
+		const TVec3<T> V = -BToATM.GetTranslation();
+		const TVector<T, 3> SupportA = A.SupportCore(-V, MarginA);
 		const TVector<T, 3> VInB = BToATM.GetRotation().Inverse() * V;
-		const TVector<T, 3> SupportBLocal = B.Support(VInB, 0);
+		const TVector<T, 3> SupportBLocal = B.SupportCore(VInB, MarginB);
 		const TVector<T, 3> SupportB = BToATM.TransformPositionNoScale(SupportBLocal);
 		return SupportA - SupportB;
 	}
@@ -714,27 +736,21 @@ namespace Chaos
 	TVector<T, 3> GJKDistanceInitialV(const TSphere<T, 3>& A, const TSphere<T, 3>& B, const TRigidTransform<T, 3>& BToATM)
 	{
 		TVector<T, 3> Delta = A.GetCenter() - (B.GetCenter() + BToATM.GetTranslation());
-		T DeltaLen = Delta.Size();
-		T RadiusAB = A.GetRadius() + B.GetRadius();
-		if (DeltaLen > RadiusAB)
-		{
-			return Delta -  Delta * (RadiusAB / DeltaLen);
-		}
-		return TVector<T, 3>(0, 0, 0);
+		return Delta;
 	}
 
-	// Overloads for geometry types which don't have centroids.
-	template <typename T, typename TGeometryB>
-	TVector<T, 3> GJKDistanceInitialV(const FImplicitObject& A, const TGeometryB& B, const TRigidTransform<T, 3>& BToATM)
+	// Status of a call to GJKDistance
+	enum class EGJKDistanceResult
 	{
-		return -BToATM.GetTranslation();
-	}
+		// The shapes are separated by a positive amount and all outputs have valid values
+		Separated,
 
-	template <typename T, typename TGeometryA>
-	TVector<T, 3> GJKDistanceInitialV(TGeometryA A, const FImplicitObject& B, const TRigidTransform<T, 3>& BToATM)
-	{
-		return -BToATM.GetTranslation();
-	}
+		// The shapes are overlapping by less than the net margin and all outputs have valid values (with a negative separation)
+		Contact,
+
+		// The shapes are overlapping by more than the net margin and all outputs are invalid
+		DeepContact,
+	};
 
 	/**
 	 * Find the distance and nearest points on two convex geometries A and B.
@@ -751,10 +767,10 @@ namespace Chaos
 	 * @param OutNearestB if returns true, the near point on B in local-space, otherwise not modified.
 	 * @param Epsilon The algorithm terminates when the iterative distance reduction gets below this threshold.
 	 * @param MaxIts A limit on the number of iterations. Results may be approximate if this is too low.
-	 * @return true if we succeeded in calculating the distance, false otherwise (i.e., false if objects are overlapping).
+	 * @return EGJKDistanceResult - see comments on the enum
 	 */
 	template <typename T, typename TGeometryA, typename TGeometryB>
-	bool GJKDistance(const TGeometryA& A, const TGeometryB& B, const TRigidTransform<T, 3>& BToATM, T& OutDistance, TVector<T, 3>& OutNearestA, TVector<T, 3>& OutNearestB, const T Epsilon = (T)1e-6, const int32 MaxIts = 16)
+	EGJKDistanceResult GJKDistance(const TGeometryA& A, const TGeometryB& B, const TRigidTransform<T, 3>& BToATM, T& OutDistance, TVector<T, 3>& OutNearestA, TVector<T, 3>& OutNearestB, TVector<T, 3>& OutNormalA, const T Epsilon = (T)1e-3, const int32 MaxIts = 16)
 	{
 		check(A.IsConvex() && B.IsConvex());
 
@@ -763,10 +779,12 @@ namespace Chaos
 		T Barycentric[4] = { -1, -1, -1, -1 };
 
 		const TRotation<T, 3> AToBRotation = BToATM.GetRotation().Inverse();
+		const T AMargin = A.GetMargin();
+		const T BMargin = B.GetMargin();
 		T Mu = 0;
 
-		// Select an initial vector in A - B
-		TVector<T, 3> V = GJKDistanceInitialV(A, B, BToATM);
+		// Select an initial vector in Minkowski(A - B)
+		TVector<T, 3> V = GJKDistanceInitialV(A, AMargin, B, BMargin, BToATM);
 		T VLen = V.Size();
 
 		int32 It = 0;
@@ -775,9 +793,9 @@ namespace Chaos
 			// Find a new point in A-B that is closer to the origin
 			// NOTE: we do not use support thickness here. Thickness is used when separating objects
 			// so that GJK can find a solution, but that can be added in a later step.
-			const TVector<T, 3> SupportA = A.Support(-V, 0);
+			const TVector<T, 3> SupportA = A.SupportCore(-V, AMargin);
 			const TVector<T, 3> VInB = AToBRotation * V;
-			const TVector<T, 3> SupportBLocal = B.Support(VInB, 0);
+			const TVector<T, 3> SupportBLocal = B.SupportCore(VInB, BMargin);
 			const TVector<T, 3> SupportB = BToATM.TransformPositionNoScale(SupportBLocal);
 			const TVector<T, 3> W = SupportA - SupportB;
 
@@ -792,7 +810,6 @@ namespace Chaos
 				// case we (probably) have a solution but with an error larger than Epsilon (technically we could be missing
 				// the fact that we were going to eventually find the origin, but it'll be a close call so the approximation
 				// is still good enough).
-				OutDistance = VLen;
 				if (SimplexIDs.NumVerts == 0)
 				{
 					// Our initial guess of V was already the minimum separating vector
@@ -812,7 +829,14 @@ namespace Chaos
 						OutNearestB += Barycentric[WIndex] * SimplexB[WIndex];
 					}
 				}
-				return true;
+				const TVector<T, 3> NormalA = -V / VLen;
+				const TVector<T, 3> NormalB = VInB / VLen;
+				OutDistance = VLen - (AMargin + BMargin);
+				OutNearestA += AMargin * NormalA;
+				OutNearestB += BMargin * NormalB;
+				OutNormalA = NormalA;
+
+				return (OutDistance >= 0.0f) ? EGJKDistanceResult::Separated : EGJKDistanceResult::Contact;
 			}
 
 			// Add the new vertex to the simplex
@@ -827,10 +851,14 @@ namespace Chaos
 			VLen = V.Size();
 		}
 
-		// Our geometries overlap - we did not produce the near points (and didn't set distance, which is zero)
-		return false;
+		// Our geometries overlap - we did not set any outputs
+		return EGJKDistanceResult::DeepContact;
 	}
 
+
+
+
+
 	// Assumes objects are already intersecting, computes a minimum translation
 	// distance, deepest penetration positions on each body, and approximates
 	// a penetration normal and minimum translation distance.
diff --git a/Engine/Source/Runtime/Experimental/Chaos/Public/Chaos/GJKShape.h b/Engine/Source/Runtime/Experimental/Chaos/Public/Chaos/GJKShape.h
new file mode 100644
--- /dev/null
+++ b/Engine/Source/Runtime/Experimental/Chaos/Public/Chaos/GJKShape.h
@@ -0,0 +1,96 @@
+// Copyright Epic Games, Inc. All Rights Reserved.
+#pragma once
+
+namespace Chaos
+{
+	/**
+	* Helpers and Wrappers for use with GJK to select the appropriate margin and support function
+	* based on context. A different margin is used for sweeps and collisions,
+	* and margins are used or not depending on the shape pair type involved.
+	*/
+
+	// Wraps an FImplicitObject  and provides the API required for GJK, 
+	// treating the shape as if it has zero margin. This means spheres will 
+	// be spheres, convexes will be the outer hull, etc.
+	//
+	// See also TGJKCoreShape
+	//
+	// E.g., to use GJK of a sphere as a point against a marginless convex:
+	//		GJKDistance(TGJKCoreShape(MySphere), TGJKShape(MyConvex), ...);
+	//
+	template<typename T_SHAPE>
+	struct TGJKShape
+	{
+		using FImplicitObjectType = T_SHAPE;
+
+		TGJKShape(const FImplicitObjectType& InShape) : Shape(InShape) {}
+
+		FReal GetMargin() const
+		{
+			return 0.0f;
+		}
+
+		FVec3 SupportCore(const FVec3 Dir, FReal InMargin) const
+		{
+			return Shape.Support(Dir, InMargin);
+		}
+
+		bool IsConvex() const
+		{
+			return Shape.IsConvex();
+		}
+
+		const FImplicitObjectType& Shape;
+	};
+
+
+	// Wraps an FImplicitObject  and provides the API required for GJK, 
+	// treating the shape as if it has a reduced "core" shape with a 
+	// margin suitable for collision detection where significant overlaps are likely.
+	// This means spheres will be points, convexes will be rounded shrunken hulls, etc.
+	//
+	// See also TGJKShape
+	//
+	// E.g., to use GJK of a sphere as a point against a marginless convex:
+	//		GJKDistance(TGJKCoreShape(MySphere), TGJKShape(MyConvex), ...);
+	//
+	template<typename T_SHAPE>
+	struct TGJKCoreShape
+	{
+		using FImplicitObjectType = T_SHAPE;
+
+		TGJKCoreShape(const FImplicitObjectType& InShape) : Shape(InShape) {}
+
+		FReal GetMargin() const
+		{
+			return Shape.GetMargin();
+		}
+
+		FVec3 SupportCore(const FVec3 Dir, FReal InMargin) const
+		{
+			return Shape.SupportCore(Dir, InMargin);
+		}
+
+		bool IsConvex() const
+		{
+			return Shape.IsConvex();
+		}
+
+		const FImplicitObjectType& Shape;
+	};
+
+
+	// Utility for creating TGJKShape objects using template parameter deduction
+	template<typename T_SHAPE>
+	TGJKShape<T_SHAPE> MakeGJKShape(const T_SHAPE& InShape)
+	{
+		return TGJKShape<T_SHAPE>(InShape);
+	}
+
+	// Utility for creating TGJKCoreShape objects using template parameter deduction
+	template<typename T_SHAPE>
+	TGJKCoreShape<T_SHAPE> MakeGJKCoreShape(const T_SHAPE& InShape)
+	{
+		return TGJKCoreShape<T_SHAPE>(InShape);
+	}
+}
diff --git a/Engine/Source/Runtime/Experimental/Chaos/Public/Chaos/GeometryQueries.h b/Engine/Source/Runtime/Experimental/Chaos/Public/Chaos/GeometryQueries.h
--- a/Engine/Source/Runtime/Experimental/Chaos/Public/Chaos/GeometryQueries.h
+++ b/Engine/Source/Runtime/Experimental/Chaos/Public/Chaos/GeometryQueries.h
@@ -1,23 +1,23 @@
 // Copyright Epic Games, Inc. All Rights Reserved.
 #pragma once
 
+#include "Chaos/Box.h"
+#include "Chaos/Capsule.h"
 #include "Chaos/CastingUtilities.h"
+#include "Chaos/Convex.h"
+#include "Chaos/GJK.h"
+#include "Chaos/HeightField.h"
 #include "Chaos/ImplicitObject.h"
+#include "Chaos/ImplicitObjectScaled.h"
+#include "Chaos/Levelset.h"
 #include "Chaos/Plane.h"
+#include "Chaos/Sphere.h"
 #include "Chaos/Transform.h"
 #include "Chaos/TriangleMeshImplicitObject.h"
-#include "Chaos/HeightField.h"
-#include "Chaos/Convex.h"
-#include "Chaos/Capsule.h"
-#include "ImplicitObjectScaled.h"
-#include "Chaos/Box.h"
-#include "Chaos/Sphere.h"
-#include "Chaos/Levelset.h"
 
 #include "ChaosArchive.h"
 #include <algorithm>
 #include <utility>
-#include "GJK.h"
 
 namespace Chaos
 {
@@ -179,7 +179,12 @@ namespace Chaos
 			}
 
 			const FVec3 Offset = ATM.GetLocation() - BTM.GetLocation();
-			bResult = Utilities::CastHelper(A, BToATM, [&](const auto& ADowncast, const auto& BToAFullTM){ return GJKRaycast2(ADowncast, B, BToAFullTM, LocalDir, Length, OutTime, LocalPosition, LocalNormal, Thickness, bComputeMTD, Offset, Thickness); });
+			bResult = Utilities::CastHelperNoUnwrap(A, BToATM, 
+				[&](const auto& ADowncast, const auto& BToAFullTM)
+				{
+					return GJKRaycast2(ADowncast, B, BToAFullTM, LocalDir, Length, OutTime, LocalPosition, LocalNormal, Thickness, bComputeMTD, Offset, Thickness);
+				});
+
 			if (AType == ImplicitObjectType::Convex)
 			{
 				//todo: find face index
diff --git a/Engine/Source/Runtime/Experimental/Chaos/Public/Chaos/ImplicitObjectScaled.h b/Engine/Source/Runtime/Experimental/Chaos/Public/Chaos/ImplicitObjectScaled.h
--- a/Engine/Source/Runtime/Experimental/Chaos/Public/Chaos/ImplicitObjectScaled.h
+++ b/Engine/Source/Runtime/Experimental/Chaos/Public/Chaos/ImplicitObjectScaled.h
@@ -3,6 +3,7 @@
 
 #include "Chaos/Box.h"
 #include "Chaos/Convex.h"
+#include "Chaos/ImplicitFwd.h"
 #include "Chaos/ImplicitObject.h"
 #include "Chaos/Transform.h"
 #include "ChaosArchive.h"
@@ -56,6 +57,18 @@ public:
 		SetMargin(OuterMargin + MObject->GetMargin());
 	}
 
+	TImplicitObjectInstanced(TImplicitObjectInstanced<TConcrete>&& Other)
+		: FImplicitObject(EImplicitObject::HasBoundingBox, Other.MObject->GetType() | ImplicitObjectType::IsInstanced)
+		, MObject(MoveTemp(Other.MObject))
+		, OuterMargin(Other.OuterMargin)
+	{
+		ensureMsgf((IsScaled(MObject->GetType()) == false), TEXT("Scaled objects should not contain each other."));
+		ensureMsgf((IsInstanced(MObject->GetType()) == false), TEXT("Scaled objects should not contain instances."));
+		this->bIsConvex = Other.MObject->IsConvex();
+		this->bDoCollide = Other.MObject->GetDoCollide();
+		SetMargin(Other.GetMargin());
+	}
+
 	static constexpr EImplicitObjectType StaticType()
 	{
 		return TConcrete::StaticType() | ImplicitObjectType::IsInstanced;
@@ -66,6 +79,11 @@ public:
 		return MObject.Get();
 	}
 
+	float GetRadius() const
+	{
+		return MObject->GetRadius();
+	}
+
 	bool GetDoCollide() const
 	{
 		return MObject->GetDoCollide();
@@ -112,7 +130,7 @@ public:
 	// The support position from the specified direction, if the shape is reduced by the margin
 	FORCEINLINE TVector<T, d> SupportCore(const TVector<T, d>& Direction, float InMargin) const
 	{
-		return MObject->SupportCore(Direction, InMargin + GetMargin());
+		return MObject->SupportCore(Direction, InMargin);
 	}
 
 	// Returns a winding order multiplier used in the manifold clipping and required when we have negative scales
@@ -395,6 +413,11 @@ public:
 		return MObject.Get();
 	}
 
+	float GetRadius() const
+	{
+		return (MObject->GetRadius() > 0.0f) ? Margin : 0.0f;
+	}
+
 	virtual T PhiWithNormal(const TVector<T, d>& X, TVector<T, d>& Normal) const override
 	{
 		// @todo(chaos): support scaled PhiWithNormal on all types
@@ -672,7 +695,7 @@ public:
 
 	FORCEINLINE_DEBUGGABLE TVector<T, d> SupportCore(const TVector<T, d>& Direction, float InMargin) const
 	{
-		return MObject->SupportCoreScaled(Direction, InMargin + OuterMargin, MScale);
+		return MObject->SupportCoreScaled(Direction, InMargin, MScale);
 	}
 
 	// Returns a winding order multiplier used in the manifold clipping and required when we have negative scales
diff --git a/Engine/Source/Runtime/Experimental/Chaos/Public/Chaos/ImplicitObjectTransformed.h b/Engine/Source/Runtime/Experimental/Chaos/Public/Chaos/ImplicitObjectTransformed.h
--- a/Engine/Source/Runtime/Experimental/Chaos/Public/Chaos/ImplicitObjectTransformed.h
+++ b/Engine/Source/Runtime/Experimental/Chaos/Public/Chaos/ImplicitObjectTransformed.h
@@ -95,6 +95,18 @@ public:
 		return MObject.Get();
 	}
 
+	float GetMargin() const
+	{
+		// If the inner shape is quadratic, we have no margin
+		return (MObject->GetRadius() > 0.0f) ? 0.0f : Margin;
+	}
+
+	float GetRadius() const
+	{
+		// If the inner shape is quadratic, so are we
+		return (MObject->GetRadius() > 0.0f) ? Margin : 0.0f;
+	}
+
 	bool GetDoCollide() const
 	{
 		return MObject->GetDoCollide();
diff --git a/Engine/Source/Runtime/Experimental/Chaos/Public/Chaos/Plane.h b/Engine/Source/Runtime/Experimental/Chaos/Public/Chaos/Plane.h
--- a/Engine/Source/Runtime/Experimental/Chaos/Public/Chaos/Plane.h
+++ b/Engine/Source/Runtime/Experimental/Chaos/Public/Chaos/Plane.h
@@ -163,6 +163,11 @@ class TPlane final : public FImplicitObject
 		return ImplicitObjectType::Plane;
 	}
 
+	float GetRadius() const
+	{
+		return 0.0f;
+	}
+
 	/**
 	 * Phi is positive on the side of the normal, and negative otherwise.
 	 */
diff --git a/Engine/Source/Runtime/Experimental/Chaos/Public/Chaos/Sphere.h b/Engine/Source/Runtime/Experimental/Chaos/Public/Chaos/Sphere.h
--- a/Engine/Source/Runtime/Experimental/Chaos/Public/Chaos/Sphere.h
+++ b/Engine/Source/Runtime/Experimental/Chaos/Public/Chaos/Sphere.h
@@ -2,6 +2,8 @@
 #pragma once
 
 #include "Chaos/Box.h"
+#include "Chaos/GJKShape.h"
+#include "Chaos/ImplicitFwd.h"
 #include "Chaos/ImplicitObject.h"
 #include "ChaosArchive.h"
 
@@ -65,6 +67,11 @@ namespace Chaos
 			return ImplicitObjectType::Sphere; 
 		}
 
+		float GetRadius() const
+		{
+			return Margin;
+		}
+
 		virtual T PhiWithNormal(const TVector<T, d>& InSamplePoint, TVector<T, d>& OutNormal) const override
 		{
 			OutNormal = InSamplePoint - Center;
@@ -191,11 +198,13 @@ namespace Chaos
 
 		FORCEINLINE const TVector<T, d>& SupportCore(const TVector<T, d>& Direction, float InMargin) const
 		{
+			// Note: ignores InMargin, assumed Radius
 			return Center;
 		}
 
 		FORCEINLINE TVector<T, d> SupportCoreScaled(const TVector<T, d>& Direction, float InMargin, const TVector<T, d>& Scale) const
 		{
+			// Note: ignores InMargin, assumed Radius
 			return Center * Scale;
 		}
 
@@ -238,11 +247,6 @@ namespace Chaos
 			return Center; 
 		}
 
-		T GetRadius() const 
-		{ 
-			return GetMargin(); 
-		}
-
 		virtual FString ToString() const
 		{
 			return FString::Printf(TEXT("TSphere Center:%s, Radius:%f"), *Center.ToString(), GetRadius());
diff --git a/Engine/Source/Runtime/Experimental/Chaos/Public/Chaos/Triangle.h b/Engine/Source/Runtime/Experimental/Chaos/Public/Chaos/Triangle.h
--- a/Engine/Source/Runtime/Experimental/Chaos/Public/Chaos/Triangle.h
+++ b/Engine/Source/Runtime/Experimental/Chaos/Public/Chaos/Triangle.h
@@ -79,6 +79,7 @@ namespace Chaos
 
 		FORCEINLINE_DEBUGGABLE TVec3<T> SupportCore(const TVec3<T>& Direction, float InMargin) const
 		{
+			// Note: assumes margin == 0
 			const float DotA = TVec3<T>::DotProduct(A, Direction);
 			const float DotB = TVec3<T>::DotProduct(B, Direction);
 			const float DotC = TVec3<T>::DotProduct(C, Direction);
@@ -97,11 +98,12 @@ namespace Chaos
 
 		FORCEINLINE TVector<T, 3> SupportCoreScaled(const TVector<T, 3>& Direction, float InMargin, const TVector<T, 3>& Scale) const
 		{
-			// No margin support in triangles (they are zero thickness so cannot have an internal margin)
+			// Note: ignores InMargin, assumed 0 (triangles cannot have a margin as they are zero thickness)
 			return SupportCore(Direction * Scale, 0.0f) * Scale;
 		}
 
 		FORCEINLINE T GetMargin() const { return 0; }
+		FORCEINLINE T GetRadius() const { return 0; }
 
 		FORCEINLINE bool Raycast(const TVec3<T>& StartPoint, const TVec3<T>& Dir, const T Length, const T Thickness, T& OutTime, TVec3<T>& OutPosition, TVec3<T>& OutNormal, int32& OutFaceIndex) const
 		{
diff --git a/Engine/Source/Runtime/Experimental/Chaos/Public/Chaos/TriangleMeshImplicitObject.h b/Engine/Source/Runtime/Experimental/Chaos/Public/Chaos/TriangleMeshImplicitObject.h
--- a/Engine/Source/Runtime/Experimental/Chaos/Public/Chaos/TriangleMeshImplicitObject.h
+++ b/Engine/Source/Runtime/Experimental/Chaos/Public/Chaos/TriangleMeshImplicitObject.h
@@ -127,6 +127,11 @@ namespace Chaos
 		FTriangleMeshImplicitObject(FTriangleMeshImplicitObject&& Other) = delete;
 		virtual ~FTriangleMeshImplicitObject();
 
+		float GetRadius() const
+		{
+			return 0.0f;
+		}
+
 		virtual FReal PhiWithNormal(const FVec3& x, FVec3& Normal) const;
 
 		virtual bool Raycast(const FVec3& StartPoint, const FVec3& Dir, const FReal Length, const FReal Thickness, FReal& OutTime, FVec3& OutPosition, FVec3& OutNormal, int32& OutFaceIndex) const override;
