Fix non-unity (CollisionOneShotManifolds.cpp)

diff --git a/Engine/Source/Runtime/Experimental/Chaos/Private/Chaos/CollisionOneShotManifolds.cpp b/Engine/Source/Runtime/Experimental/Chaos/Private/Chaos/CollisionOneShotManifolds.cpp
--- a/Engine/Source/Runtime/Experimental/Chaos/Private/Chaos/CollisionOneShotManifolds.cpp
+++ b/Engine/Source/Runtime/Experimental/Chaos/Private/Chaos/CollisionOneShotManifolds.cpp
@@ -10,6 +10,8 @@
 #include "Chaos/Transform.h"
 #include "Chaos/Utilities.h"
 
+#include "HAL/IConsoleManager.h"
+
 //PRAGMA_DISABLE_OPTIMIZATION
 
 namespace Chaos
