Fix static analysis (null deref in CollisionResolution.cpp)

diff --git a/Engine/Source/Runtime/Experimental/Chaos/Private/Chaos/CollisionResolution.cpp b/Engine/Source/Runtime/Experimental/Chaos/Private/Chaos/CollisionResolution.cpp
--- a/Engine/Source/Runtime/Experimental/Chaos/Private/Chaos/CollisionResolution.cpp
+++ b/Engine/Source/Runtime/Experimental/Chaos/Private/Chaos/CollisionResolution.cpp
@@ -2276,38 +2276,40 @@ namespace Chaos
 			bool bUseCCD = UseCCD(Particle0, Particle1, Implicit0, DirCCD, LengthCCD);
 
 #if CHAOS_COLLISION_CREATE_BOUNDSCHECK
-			if (Implicit0->HasBoundingBox() && Implicit1->HasBoundingBox())
+			if ((Implicit0 != nullptr) && (Implicit1 != nullptr))
 			{
-				const FRigidTransform3 ParticleTransform0 = Collisions::GetTransform(Particle0);
-				const FRigidTransform3 ParticleTransform1 = Collisions::GetTransform(Particle1);
-				const FRigidTransform3 WorldTransform0 = LocalTransform0 * ParticleTransform0;
-				const FRigidTransform3 WorldTransform1 = LocalTransform1 * ParticleTransform1;
-				if (Chaos_Collision_NarrowPhase_SphereBoundsCheck)
+				if (Implicit0->HasBoundingBox() && Implicit1->HasBoundingBox())
 				{
-					const FReal R1 = Implicit0->BoundingBox().OriginRadius();
-					const FReal R2 = Implicit1->BoundingBox().OriginRadius();
-					const FReal SeparationSq = (WorldTransform1.GetTranslation() - WorldTransform0.GetTranslation()).SizeSquared();
-					const FReal CullDistanceSq = CullDistance * CullDistance;
-					if ((SeparationSq - FMath::Square(R1 + R2)) > CullDistanceSq)
+					const FRigidTransform3 ParticleTransform0 = Collisions::GetTransform(Particle0);
+					const FRigidTransform3 ParticleTransform1 = Collisions::GetTransform(Particle1);
+					const FRigidTransform3 WorldTransform0 = LocalTransform0 * ParticleTransform0;
+					const FRigidTransform3 WorldTransform1 = LocalTransform1 * ParticleTransform1;
+					if (Chaos_Collision_NarrowPhase_SphereBoundsCheck)
 					{
-						return;
+						const FReal R1 = Implicit0->BoundingBox().OriginRadius();
+						const FReal R2 = Implicit1->BoundingBox().OriginRadius();
+						const FReal SeparationSq = (WorldTransform1.GetTranslation() - WorldTransform0.GetTranslation()).SizeSquared();
+						const FReal CullDistanceSq = CullDistance * CullDistance;
+						if ((SeparationSq - FMath::Square(R1 + R2)) > CullDistanceSq)
+						{
+							return;
+						}
 					}
-				}
 
-				if (Chaos_Collision_NarrowPhase_AABBBoundsCheck)
-				{
-					const FRigidTransform3 Box2ToBox1TM = WorldTransform1.GetRelativeTransform(WorldTransform0);
-					const FAABB3 Box1 = Implicit0->BoundingBox();
-					const FAABB3 Box2In1 = Implicit1->BoundingBox().TransformedAABB(Box2ToBox1TM).Thicken(CullDistance);
-					if (!Box1.Intersects(Box2In1))
+					if (Chaos_Collision_NarrowPhase_AABBBoundsCheck)
 					{
-						return;
+						const FRigidTransform3 Box2ToBox1TM = WorldTransform1.GetRelativeTransform(WorldTransform0);
+						const FAABB3 Box1 = Implicit0->BoundingBox();
+						const FAABB3 Box2In1 = Implicit1->BoundingBox().TransformedAABB(Box2ToBox1TM).Thicken(CullDistance);
+						if (!Box1.Intersects(Box2In1))
+						{
+							return;
+						}
 					}
 				}
 			}
 #endif
 
-
 			if (Implicit0Type == TBox<FReal, 3>::StaticType() && Implicit1Type == TBox<FReal, 3>::StaticType())
 			{
 				ConstructBoxBoxConstraints<T_TRAITS>(Particle0, Particle1, Implicit0, Implicit1, LocalTransform0, LocalTransform1, CullDistance, Dt, Context, NewConstraints);
