Chaos - Replace float with FReal

diff --git a/Engine/Source/Runtime/Experimental/Chaos/Public/Chaos/Box.h b/Engine/Source/Runtime/Experimental/Chaos/Public/Chaos/Box.h
--- a/Engine/Source/Runtime/Experimental/Chaos/Public/Chaos/Box.h
+++ b/Engine/Source/Runtime/Experimental/Chaos/Public/Chaos/Box.h
@@ -94,7 +94,7 @@ namespace Chaos
 			return TUniquePtr<FImplicitObject>(new TBox<T,d>(*this));
 		}
 
-		float GetRadius() const
+		FReal GetRadius() const
 		{
 			return 0.0f;
 		}
diff --git a/Engine/Source/Runtime/Experimental/Chaos/Public/Chaos/Capsule.h b/Engine/Source/Runtime/Experimental/Chaos/Public/Chaos/Capsule.h
--- a/Engine/Source/Runtime/Experimental/Chaos/Public/Chaos/Capsule.h
+++ b/Engine/Source/Runtime/Experimental/Chaos/Public/Chaos/Capsule.h
@@ -72,7 +72,7 @@ namespace Chaos
 			return TCapsule<T>(X1, X2, Radius);
 		}
 
-		float GetRadius() const
+		FReal GetRadius() const
 		{
 			return Margin;
 		}
diff --git a/Engine/Source/Runtime/Experimental/Chaos/Public/Chaos/Convex.h b/Engine/Source/Runtime/Experimental/Chaos/Public/Chaos/Convex.h
--- a/Engine/Source/Runtime/Experimental/Chaos/Public/Chaos/Convex.h
+++ b/Engine/Source/Runtime/Experimental/Chaos/Public/Chaos/Convex.h
@@ -108,7 +108,7 @@ namespace Chaos
 			return Margin;
 		}
 
-		float GetRadius() const
+		FReal GetRadius() const
 		{
 			return 0.0f;
 		}
diff --git a/Engine/Source/Runtime/Experimental/Chaos/Public/Chaos/ImplicitObjectScaled.h b/Engine/Source/Runtime/Experimental/Chaos/Public/Chaos/ImplicitObjectScaled.h
--- a/Engine/Source/Runtime/Experimental/Chaos/Public/Chaos/ImplicitObjectScaled.h
+++ b/Engine/Source/Runtime/Experimental/Chaos/Public/Chaos/ImplicitObjectScaled.h
@@ -79,7 +79,7 @@ public:
 		return MObject.Get();
 	}
 
-	float GetRadius() const
+	FReal GetRadius() const
 	{
 		return MObject->GetRadius();
 	}
@@ -413,7 +413,7 @@ public:
 		return MObject.Get();
 	}
 
-	float GetRadius() const
+	FReal GetRadius() const
 	{
 		return (MObject->GetRadius() > 0.0f) ? Margin : 0.0f;
 	}
diff --git a/Engine/Source/Runtime/Experimental/Chaos/Public/Chaos/ImplicitObjectTransformed.h b/Engine/Source/Runtime/Experimental/Chaos/Public/Chaos/ImplicitObjectTransformed.h
--- a/Engine/Source/Runtime/Experimental/Chaos/Public/Chaos/ImplicitObjectTransformed.h
+++ b/Engine/Source/Runtime/Experimental/Chaos/Public/Chaos/ImplicitObjectTransformed.h
@@ -101,7 +101,7 @@ public:
 		return (MObject->GetRadius() > 0.0f) ? 0.0f : Margin;
 	}
 
-	float GetRadius() const
+	FReal GetRadius() const
 	{
 		// If the inner shape is quadratic, so are we
 		return (MObject->GetRadius() > 0.0f) ? Margin : 0.0f;
diff --git a/Engine/Source/Runtime/Experimental/Chaos/Public/Chaos/Plane.h b/Engine/Source/Runtime/Experimental/Chaos/Public/Chaos/Plane.h
--- a/Engine/Source/Runtime/Experimental/Chaos/Public/Chaos/Plane.h
+++ b/Engine/Source/Runtime/Experimental/Chaos/Public/Chaos/Plane.h
@@ -163,7 +163,7 @@ class TPlane final : public FImplicitObject
 		return ImplicitObjectType::Plane;
 	}
 
-	float GetRadius() const
+	FReal GetRadius() const
 	{
 		return 0.0f;
 	}
diff --git a/Engine/Source/Runtime/Experimental/Chaos/Public/Chaos/Sphere.h b/Engine/Source/Runtime/Experimental/Chaos/Public/Chaos/Sphere.h
--- a/Engine/Source/Runtime/Experimental/Chaos/Public/Chaos/Sphere.h
+++ b/Engine/Source/Runtime/Experimental/Chaos/Public/Chaos/Sphere.h
@@ -67,7 +67,7 @@ namespace Chaos
 			return ImplicitObjectType::Sphere; 
 		}
 
-		float GetRadius() const
+		FReal GetRadius() const
 		{
 			return Margin;
 		}
diff --git a/Engine/Source/Runtime/Experimental/Chaos/Public/Chaos/TriangleMeshImplicitObject.h b/Engine/Source/Runtime/Experimental/Chaos/Public/Chaos/TriangleMeshImplicitObject.h
--- a/Engine/Source/Runtime/Experimental/Chaos/Public/Chaos/TriangleMeshImplicitObject.h
+++ b/Engine/Source/Runtime/Experimental/Chaos/Public/Chaos/TriangleMeshImplicitObject.h
@@ -127,7 +127,7 @@ namespace Chaos
 		FTriangleMeshImplicitObject(FTriangleMeshImplicitObject&& Other) = delete;
 		virtual ~FTriangleMeshImplicitObject();
 
-		float GetRadius() const
+		FReal GetRadius() const
 		{
 			return 0.0f;
 		}
