Fix CIS - by removing alternate suspension option

diff --git a/Engine/Source/Runtime/Experimental/ChaosVehicles/ChaosVehiclesCore/Private/SuspensionSystem.cpp b/Engine/Source/Runtime/Experimental/ChaosVehicles/ChaosVehiclesCore/Private/SuspensionSystem.cpp
--- a/Engine/Source/Runtime/Experimental/ChaosVehicles/ChaosVehiclesCore/Private/SuspensionSystem.cpp
+++ b/Engine/Source/Runtime/Experimental/ChaosVehicles/ChaosVehiclesCore/Private/SuspensionSystem.cpp
@@ -2,9 +2,6 @@
 
 #include "SuspensionSystem.h"
 
-static bool GNewSuspensionSim = true; // TODO: remove once sure new calculation is correct
-FAutoConsoleVariableRef CVarChaosVehiclesABTestSuspension(TEXT("p.Vehicle.NewSuspensionSim"), GNewSuspensionSim, TEXT("A/B Test Suspension Simulation."));
-
 namespace Chaos
 {
 
@@ -45,25 +42,13 @@ namespace Chaos
 
 	void FSimpleSuspensionSim::Simulate(float DeltaTime)
 	{
-		if (GNewSuspensionSim)
-		{
-			float Damping = (SpringDisplacement < LastDisplacement) ? Setup().CompressionDamping : Setup().ReboundDamping;
-			float SpringSpeed = (LastDisplacement - SpringDisplacement) / DeltaTime;
+		float Damping = (SpringDisplacement < LastDisplacement) ? Setup().CompressionDamping : Setup().ReboundDamping;
+		float SpringSpeed = (LastDisplacement - SpringDisplacement) / DeltaTime;
 
-			const float StiffnessForce = SpringDisplacement * Setup().SpringRate;
-			const float DampingForce = SpringSpeed * Damping;
-			SuspensionForce = StiffnessForce - DampingForce;
-			LastDisplacement = SpringDisplacement;
-		}
-		else
-		{
-			float Damping = (DisplacementInput < LastDisplacement) ? Setup().CompressionDamping : Setup().ReboundDamping;
-
-			const float StiffnessForce = SpringDisplacement * Setup().SpringRate;
-			const float DampingForce = LocalVelocity.Z * Damping;
-			SuspensionForce = StiffnessForce - DampingForce;
-			LastDisplacement = DisplacementInput;
-		}
+		const float StiffnessForce = SpringDisplacement * Setup().SpringRate;
+		const float DampingForce = SpringSpeed * Damping;
+		SuspensionForce = StiffnessForce - DampingForce;
+		LastDisplacement = SpringDisplacement;
 	}
 
 } // namespace Chaos
