Chaos - remove old dead TOI collision code

diff --git a/Engine/Source/Runtime/Experimental/Chaos/Private/Chaos/CollisionResolution.cpp b/Engine/Source/Runtime/Experimental/Chaos/Private/Chaos/CollisionResolution.cpp
--- a/Engine/Source/Runtime/Experimental/Chaos/Private/Chaos/CollisionResolution.cpp
+++ b/Engine/Source/Runtime/Experimental/Chaos/Private/Chaos/CollisionResolution.cpp
@@ -2831,241 +2831,6 @@ namespace Chaos
 			}
 		}
 
-		template <typename GeometryA, typename GeometryB>
-		bool GetPairTOIHackImpl(const FImplicitObject& A, const FRigidTransform3& AStartTransform, const GeometryB& B, const FRigidTransform3& BTransform, const FVec3& Dir, const FReal Length, FReal& OutTOI, FVec3& OutNormal, FReal& OutPhi)
-		{
-			FReal CullDistance = 0.0f;
-
-			FReal TOI = 1.0f;
-			FContactPoint Contact = GJKImplicitSweptContactPoint<GeometryA>(A, AStartTransform, B, BTransform, Dir, Length, CullDistance, TOI);
-			if (Contact.Phi < OutPhi)
-			{
-				OutPhi = Contact.Phi;
-				OutNormal = Contact.Normal;
-				OutTOI = TOI;
-				return true;
-			}
-			return false;
-		}
-
-		template <typename GeometryA>
-		bool GetPairTOIHackTriMeshImpl(const FImplicitObject& A, const FRigidTransform3& AStartTransform, const FImplicitObject* B, const FRigidTransform3& BTransform, const FVec3& Dir, const FReal Length, FReal& OutTOI, FVec3& OutNormal, FReal& OutPhi)
-		{
-			FContactPoint Contact;
-			FReal TOI = 0.0f;
-			FReal CullDistance = 0.0f;
-
-			if (const TImplicitObjectScaled<FTriangleMeshImplicitObject>* ScaledTriangleMesh = B->template GetObject<const TImplicitObjectScaled<FTriangleMeshImplicitObject>>())
-			{
-				Contact = GJKImplicitScaledTriMeshSweptContactPoint<GeometryA>(A, AStartTransform, *ScaledTriangleMesh, BTransform, Dir, Length, CullDistance, TOI);
-			}
-			else if (const FTriangleMeshImplicitObject* TriangleMesh = B->template GetObject<const FTriangleMeshImplicitObject>())
-			{
-				Contact = GJKImplicitSweptContactPoint<GeometryA>(A, AStartTransform, *TriangleMesh, BTransform, Dir, Length, CullDistance, TOI);
-			}
-
-			if (Contact.Phi < OutPhi)
-			{
-				OutPhi = Contact.Phi;
-				OutNormal = Contact.Normal;
-				OutTOI = TOI;
-				return true;
-			}
-			return false;
-		}
-
-		bool GetPairTOIHackConvexConvexImpl(const FImplicitObject& A, const FRigidTransform3& AStartTransform, const FImplicitObject& B, const FRigidTransform3& BTransform, const FVec3& Dir, const FReal Length, FReal& OutTOI, FVec3& OutNormal, FReal& OutPhi)
-		{
-			FReal CullDistance = 0.0f;
-			FReal TOI = 0.0f;
-
-			FContactPoint Contact = GenericConvexConvexContactPointSwept(A, AStartTransform, B, BTransform, Dir, Length, CullDistance, TOI);
-			if (Contact.Phi < OutPhi)
-			{
-				OutPhi = Contact.Phi;
-				OutNormal = Contact.Normal;
-				OutTOI = TOI;
-				return true;
-			}
-			return false;
-		}
-
-
-
-		bool GetPairTOIHack(const TPBDRigidParticleHandle<FReal, 3>* Particle0, const TGeometryParticleHandle<FReal, 3>* Particle1, const FImplicitObject* Implicit0, const FImplicitObject* Implicit1, const FRigidTransform3& StartTransform0, const FRigidTransform3& Transform1, FReal& OutTOI, FVec3& OutNormal, FReal& OutPhi)
-		{
-			// It's impossible to access an null Implicit0 or Implicit1 as we check types.
-			CA_ASSUME(Implicit0);
-			CA_ASSUME(Implicit1);
-
-			if (!Implicit0 || !Implicit1)
-			{
-				return false;
-			}
-
-			// @todo(chaos): We use GetInnerType here because TriMeshes are left with their "Instanced" wrapper, unlike all other instanced implicits. Should we strip the instance on Tri Mesh too?
-			EImplicitObjectType Implicit0Type = Implicit0 ? GetInnerType(Implicit0->GetType()) : ImplicitObjectType::Unknown;
-			EImplicitObjectType Implicit1Type = Implicit1 ? GetInnerType(Implicit1->GetType()) : ImplicitObjectType::Unknown;
-			bool bIsConvex0 = Implicit0 && Implicit0->IsConvex();
-			bool bIsConvex1 = Implicit1 && Implicit1->IsConvex();
-
-			FVec3 Dir = Particle0->P() - Particle0->X();
-			FReal Length = Dir.Size();
-			if (Length < KINDA_SMALL_NUMBER)
-			{
-				return false;
-			}
-			Dir /= Length;
-
-			if (Implicit0Type == TBox<FReal, 3>::StaticType() && Implicit1Type == FHeightField::StaticType())
-			{
-				return GetPairTOIHackImpl<TBox<FReal, 3>>(*Implicit0->template GetObject<TBox<FReal, 3>>(), StartTransform0, *Implicit1->template GetObject<FHeightField>(), Transform1, Dir, Length, OutTOI, OutNormal, OutPhi);
-			}
-			else if (Implicit0Type == TSphere<FReal, 3>::StaticType() && Implicit1Type == FHeightField::StaticType())
-			{
-				return GetPairTOIHackImpl<TSphere<FReal, 3>>(*Implicit0->template GetObject<TSphere<FReal, 3>>(), StartTransform0, *Implicit1->template GetObject<FHeightField>(), Transform1, Dir, Length, OutTOI, OutNormal, OutPhi);
-			}
-			else if (Implicit0Type == TCapsule<FReal>::StaticType() && Implicit1Type == FHeightField::StaticType())
-			{
-				return GetPairTOIHackImpl<TCapsule<FReal>>(*Implicit0->template GetObject<TCapsule<FReal>>(), StartTransform0, *Implicit1->template GetObject<FHeightField>(), Transform1, Dir, Length, OutTOI, OutNormal, OutPhi);
-			}
-			else if (Implicit0Type == TBox<FReal, 3>::StaticType() && Implicit1Type == FTriangleMeshImplicitObject::StaticType())
-			{
-				return GetPairTOIHackTriMeshImpl<TBox<FReal, 3>>(*Implicit0->template GetObject<TBox<FReal, 3>>(), StartTransform0, Implicit1, Transform1, Dir, Length, OutTOI, OutNormal, OutPhi);
-			}
-			else if (Implicit0Type == TSphere<FReal, 3>::StaticType() && Implicit1Type == FTriangleMeshImplicitObject::StaticType())
-			{
-				return GetPairTOIHackTriMeshImpl<TSphere<FReal, 3>>(*Implicit0->template GetObject<TSphere<FReal, 3>>(), StartTransform0, Implicit1, Transform1, Dir, Length, OutTOI, OutNormal, OutPhi);
-			}
-			if (Implicit0Type == TCapsule<FReal>::StaticType() && Implicit1Type == FTriangleMeshImplicitObject::StaticType())
-			{
-				return GetPairTOIHackTriMeshImpl<TCapsule<FReal>>(*Implicit0->template GetObject<TCapsule<FReal>>(), StartTransform0, Implicit1, Transform1, Dir, Length, OutTOI, OutNormal, OutPhi);
-			}
-			else if (bIsConvex0 && Implicit1Type == FHeightField::StaticType())
-			{
-				return GetPairTOIHackImpl<FConvex>(*Implicit0, StartTransform0, *Implicit1->template GetObject<FHeightField>(), Transform1, Dir, Length, OutTOI, OutNormal, OutPhi);
-			}
-			else if (bIsConvex0 && Implicit1Type == FTriangleMeshImplicitObject::StaticType())
-			{
-				return GetPairTOIHackTriMeshImpl<FConvex>(*Implicit0, StartTransform0, Implicit1, Transform1, Dir, Length, OutTOI, OutNormal, OutPhi);
-			}
-			else if (bIsConvex0 && bIsConvex1)
-			{
-				return GetPairTOIHackConvexConvexImpl(*Implicit0, StartTransform0, *Implicit1, Transform1, Dir, Length, OutTOI, OutNormal, OutPhi);
-			}
-
-			return false;
-		}
-
-
-		bool GetTOIHackImpl(const TPBDRigidParticleHandle<FReal, 3>* Particle0, const TGeometryParticleHandle<FReal, 3>* Particle1, const FImplicitObject* Implicit0, const FImplicitObject* Implicit1, const FRigidTransform3& Transform0, const FRigidTransform3& Transform1, FReal& OutTOI, FVec3& OutNormal, FReal& OutPhi)
-		{
-			EImplicitObjectType Implicit0Type = Implicit0 ? GetInnerType(Implicit0->GetType()) : ImplicitObjectType::Unknown;
-			EImplicitObjectType Implicit1Type = Implicit1 ? GetInnerType(Implicit1->GetType()) : ImplicitObjectType::Unknown;
-
-			if (!Implicit0 || !Implicit1)
-			{
-				return false;
-			}
-
-			EImplicitObjectType Implicit0OuterType = Implicit0->GetType();
-			EImplicitObjectType Implicit1OuterType = Implicit1->GetType();
-
-			// Handle transform wrapper shape
-			if ((Implicit0OuterType == TImplicitObjectTransformed<FReal, 3>::StaticType()) && (Implicit1OuterType == TImplicitObjectTransformed<FReal, 3>::StaticType()))
-			{
-				const TImplicitObjectTransformed<FReal, 3>* TransformedImplicit0 = Implicit0->template GetObject<const TImplicitObjectTransformed<FReal, 3>>();
-				const TImplicitObjectTransformed<FReal, 3>* TransformedImplicit1 = Implicit1->template GetObject<const TImplicitObjectTransformed<FReal, 3>>();
-				TRigidTransform<FReal, 3> TransformedTransform0 = TransformedImplicit0->GetTransform() * Transform0;
-				TRigidTransform<FReal, 3> TransformedTransform1 = TransformedImplicit1->GetTransform() * Transform1;
-				return GetTOIHackImpl(Particle0, Particle1, TransformedImplicit0->GetTransformedObject(), TransformedImplicit1->GetTransformedObject(), TransformedTransform0, TransformedTransform1, OutTOI, OutNormal, OutPhi);
-			}
-			else if (Implicit0OuterType == TImplicitObjectTransformed<FReal, 3>::StaticType())
-			{
-				const TImplicitObjectTransformed<FReal, 3>* TransformedImplicit0 = Implicit0->template GetObject<const TImplicitObjectTransformed<FReal, 3>>();
-				TRigidTransform<FReal, 3> TransformedTransform0 = TransformedImplicit0->GetTransform() * Transform0;
-				return GetTOIHackImpl(Particle0, Particle1, TransformedImplicit0->GetTransformedObject(), Implicit1, TransformedTransform0, Transform1, OutTOI, OutNormal, OutPhi);
-			}
-			else if (Implicit1OuterType == TImplicitObjectTransformed<FReal, 3>::StaticType())
-			{
-				const TImplicitObjectTransformed<FReal, 3>* TransformedImplicit1 = Implicit1->template GetObject<const TImplicitObjectTransformed<FReal, 3>>();
-				TRigidTransform<FReal, 3> TransformedTransform1 = TransformedImplicit1->GetTransform() * Transform1;
-				return GetTOIHackImpl(Particle0, Particle1, Implicit0, TransformedImplicit1->GetTransformedObject(), Transform0, TransformedTransform1, OutTOI, OutNormal, OutPhi);
-			}
-			// Handle Instanced shapes
-			// NOTE: Tri Meshes are handled differently. We should probably do something about this...
-			if (((uint32)Implicit0OuterType & ImplicitObjectType::IsInstanced) || ((uint32)Implicit1OuterType & ImplicitObjectType::IsInstanced))
-			{
-				const FImplicitObject* InnerImplicit0 = Implicit0;
-				const FImplicitObject* InnerImplicit1 = Implicit1;
-				if ((uint32)Implicit0OuterType & ImplicitObjectType::IsInstanced)
-				{
-					InnerImplicit0 = GetInstancedImplicit(Implicit0);
-				}
-				if ((uint32)Implicit1OuterType & ImplicitObjectType::IsInstanced)
-				{
-					InnerImplicit1 = GetInstancedImplicit(Implicit1);
-				}
-				if (InnerImplicit0 && InnerImplicit1)
-				{
-					return GetTOIHackImpl(Particle0, Particle1, InnerImplicit0, InnerImplicit1, Transform0, Transform1, OutTOI, OutNormal, OutPhi);
-				}
-				else if (InnerImplicit0 && !InnerImplicit1)
-				{
-					return GetTOIHackImpl(Particle0, Particle1, InnerImplicit0, Implicit1, Transform0, Transform1, OutTOI, OutNormal, OutPhi);
-				}
-				else if (!InnerImplicit0 && InnerImplicit1)
-				{
-					return GetTOIHackImpl(Particle0, Particle1, Implicit0, InnerImplicit1, Transform0, Transform1, OutTOI, OutNormal, OutPhi);
-				}
-			}
-
-			// Handle Unions
-			if (Implicit0OuterType == FImplicitObjectUnion::StaticType())
-			{
-				bool bHit = false;
-				const FImplicitObjectUnion* Union0 = Implicit0->template GetObject<FImplicitObjectUnion>();
-				for (const auto& Child0 : Union0->GetObjects())
-				{
-					if (GetTOIHackImpl(Particle0, Particle1, Child0.Get(), Implicit1, Transform0, Transform1, OutTOI, OutNormal, OutPhi))
-					{
-						bHit = true;
-					}
-				}
-				return bHit;
-			}
-			if (Implicit1OuterType == FImplicitObjectUnion::StaticType())
-			{
-				bool bHit = false;
-				const FImplicitObjectUnion* Union1 = Implicit1->template GetObject<FImplicitObjectUnion>();
-				for (const auto& Child1 : Union1->GetObjects())
-				{
-					if (GetTOIHackImpl(Particle0, Particle1, Implicit0, Child1.Get(), Transform0, Transform1, OutTOI, OutNormal, OutPhi))
-					{
-						bHit = true;
-					}
-				}
-				return bHit;
-			}
-
-			// Check shape pair filtering if enable
-			if (!DoCollide(Implicit0Type, const_cast<TPBDRigidParticleHandle<FReal, 3>*>(Particle0)->GetImplicitShape(Implicit0), Implicit1Type, const_cast<TGeometryParticleHandle<FReal, 3>*>(Particle1)->GetImplicitShape(Implicit1)))
-			{
-				return false;
-			}
-
-			// If we get here, we have a pair of concrete shapes (i.e., no wrappers or containers)
-			// Create a constraint for the shape pair
-			return GetPairTOIHack(Particle0, Particle1, Implicit0, Implicit1, Transform0, Transform1, OutTOI, OutNormal, OutPhi);
-		}
-
-		bool GetTOIHack(const TPBDRigidParticleHandle<FReal, 3>* Particle0, const TGeometryParticleHandle<FReal, 3>* Particle1, FReal& OutTOI, FVec3& OutNormal, FReal& OutPhi)
-		{
-			OutPhi = 0.0f;	// Cull Distance
-
-			FRigidTransform3 Particle0StartTransform = FRigidTransform3(Particle0->X(), Particle0->Q());	// Start Pos, Final Rot
-			return GetTOIHackImpl(Particle0, Particle1, Particle0->Geometry().Get(), Particle1->Geometry().Get(), Particle0StartTransform, Collisions::GetTransform(Particle1), OutTOI, OutNormal, OutPhi);
-		}
 
 		template void UpdateLevelsetLevelsetConstraint<ECollisionUpdateType::Any>(const FRigidTransform3& WorldTransform0, const FRigidTransform3& WorldTransform1, const float CullDistance, const FReal Dt, FRigidBodyPointContactConstraint& Constraint);
 		template void UpdateLevelsetLevelsetConstraint<ECollisionUpdateType::Deepest>(const FRigidTransform3& WorldTransform0, const FRigidTransform3& WorldTransform1, const float CullDistance, const FReal Dt, FRigidBodyPointContactConstraint& Constraint);
@@ -3075,4 +2840,5 @@ namespace Chaos
 
 	} // Collisions
 
+
 } // Chaos
diff --git a/Engine/Source/Runtime/Experimental/Chaos/Private/Chaos/PBDRigidsEvolutionGBF.cpp b/Engine/Source/Runtime/Experimental/Chaos/Private/Chaos/PBDRigidsEvolutionGBF.cpp
--- a/Engine/Source/Runtime/Experimental/Chaos/Private/Chaos/PBDRigidsEvolutionGBF.cpp
+++ b/Engine/Source/Runtime/Experimental/Chaos/Private/Chaos/PBDRigidsEvolutionGBF.cpp
@@ -54,12 +54,6 @@ FAutoConsoleVariableRef CVarDisableCulledContacts(TEXT("p.CollisionDisableCulled
 float BoundsThicknessVelocityMultiplier = 2.0f;	// @todo(chaos): more to FChaosSolverConfiguration
 FAutoConsoleVariableRef CVarBoundsThicknessVelocityMultiplier(TEXT("p.CollisionBoundsVelocityInflation"), BoundsThicknessVelocityMultiplier, TEXT("Collision velocity inflation for speculatibe contact generation.[def:2.0]"));
 
-float HackCCD_EnableThreshold = -1.f;
-FAutoConsoleVariableRef CVarHackCCDVelThreshold(TEXT("p.Chaos.CCD.EnableThreshold"), HackCCD_EnableThreshold, TEXT("If distance moved is greater than this times the minimum object dimension, use CCD"));
-
-float HackCCD_DepthThreshold = 0.05f;
-FAutoConsoleVariableRef CVarHackCCDDepthThreshold(TEXT("p.Chaos.CCD.DepthThreshold"), HackCCD_DepthThreshold, TEXT("When returning to TOI, leave this much contact depth (as a fraction of MinBounds)"));
-
 float SmoothedPositionLerpRate = 0.1f;
 FAutoConsoleVariableRef CVarSmoothedPositionLerpRate(TEXT("p.Chaos.SmoothedPositionLerpRate"), SmoothedPositionLerpRate, TEXT("The interpolation rate for the smoothed position calculation. Used for sleeping."));
 
@@ -79,7 +73,6 @@ DECLARE_CYCLE_STAT(TEXT("TPBDRigidsEvolutionGBF::DetectCollisions"), STAT_Evolut
 DECLARE_CYCLE_STAT(TEXT("TPBDRigidsEvolutionGBF::PostDetectCollisionsCallback"), STAT_Evolution_PostDetectCollisionsCallback, STATGROUP_Chaos);
 DECLARE_CYCLE_STAT(TEXT("TPBDRigidsEvolutionGBF::UpdateConstraintPositionBasedState"), STAT_Evolution_UpdateConstraintPositionBasedState, STATGROUP_Chaos);
 DECLARE_CYCLE_STAT(TEXT("TPBDRigidsEvolutionGBF::ComputeIntermediateSpatialAcceleration"), STAT_Evolution_ComputeIntermediateSpatialAcceleration, STATGROUP_Chaos);
-DECLARE_CYCLE_STAT(TEXT("TPBDRigidsEvolutionGBF::CCDHack"), STAT_Evolution_CCDHack, STATGROUP_Chaos);
 DECLARE_CYCLE_STAT(TEXT("TPBDRigidsEvolutionGBF::CreateConstraintGraph"), STAT_Evolution_CreateConstraintGraph, STATGROUP_Chaos);
 DECLARE_CYCLE_STAT(TEXT("TPBDRigidsEvolutionGBF::CreateIslands"), STAT_Evolution_CreateIslands, STATGROUP_Chaos);
 DECLARE_CYCLE_STAT(TEXT("TPBDRigidsEvolutionGBF::PreApplyCallback"), STAT_Evolution_PreApplyCallback, STATGROUP_Chaos);
@@ -153,114 +146,6 @@ void TPBDRigidsEvolutionGBF<Traits>::Advance(const FReal Dt,const FReal MaxStepD
 	}
 }
 
-//
-//
-//
-// BEGIN TOI
-//
-//
-//
-
-namespace Collisions
-{
-	extern bool GetTOIHack(const TPBDRigidParticleHandle<FReal, 3>* Particle0, const FGeometryParticleHandle* Particle1, FReal& OutTOI, FVec3& OutNormal, FReal& OutPhi);
-}
-
-void PostMoveToTOIHack(FReal Dt, TPBDRigidParticleHandle<FReal, 3>* Particle)
-{
-	// Update bounds, velocity etc
-	FPerParticlePBDUpdateFromDeltaPosition VelUpdateRule;
-	VelUpdateRule.Apply(Particle, Dt);
-}
-
-void MoveToTOIPairHack(FReal Dt, TPBDRigidParticleHandle<FReal, 3>* Particle1, const FGeometryParticleHandle* Particle2)
-{
-	FReal TOI = 0.0f;
-	FReal Phi = 0.0f;
-	FVec3 Normal = FVec3(0);
-	if (Collisions::GetTOIHack(Particle1, Particle2, TOI, Normal, Phi))
-	{
-		FReal Depth = -Phi;
-		FReal MinBounds = Particle1->Geometry()->BoundingBox().Extents().Min();
-		FReal MaxDepth = MinBounds * HackCCD_DepthThreshold;
-
-		//UE_LOG(LogChaos, Warning, TEXT("MoveTOIHAck: TOI %f; Depth %f"), TOI, Depth);
-
-		if ((Depth > MaxDepth) && (TOI > KINDA_SMALL_NUMBER) && (TOI < 1.0f))
-		{
-			// Move the particle to just after the TOI so we still have a collision to resolve but won't get tunneling
-			// The time after TOI we want is given by the ratio of max acceptable depth to depth at T = 1
-			FReal ExtraT = (1.0f - TOI) * MaxDepth / Depth;
-			FReal FinalT = TOI + ExtraT;
-			FVec3 FinalP = FVec3::Lerp(Particle1->X(), Particle1->P(), FinalT);
-
-			//UE_LOG(LogChaos, Warning, TEXT("MoveTOIHack: FinalTOI %f; Correction %f"), FinalT, (FinalP - Particle1->P()).Size());
-		
-			Particle1->SetP(FinalP);
-			PostMoveToTOIHack(Dt, Particle1);
-		}
-	}
-}
-
-
-void MoveToTOIHack(FReal Dt, TTransientPBDRigidParticleHandle<FReal, 3>& Particle, const ISpatialAcceleration<FAccelerationStructureHandle, FReal, 3>* SpatialAcceleration)
-{
-	if (const auto AABBTree = SpatialAcceleration->template As<TAABBTree<FAccelerationStructureHandle, TAABBTreeLeafArray<FAccelerationStructureHandle>>>())
-	{
-		MoveToTOIHackImpl(Dt, Particle, AABBTree);
-	}
-	else if (const auto BV = SpatialAcceleration->template As<TBoundingVolume<FAccelerationStructureHandle>>())
-	{
-		MoveToTOIHackImpl(Dt, Particle, BV);
-	}
-	else if (const auto AABBTreeBV = SpatialAcceleration->template As<TAABBTree<FAccelerationStructureHandle, TBoundingVolume<FAccelerationStructureHandle>>>())
-	{
-		MoveToTOIHackImpl(Dt, Particle, AABBTreeBV);
-	}
-	else if (const auto Collection = SpatialAcceleration->template As<ISpatialAccelerationCollection<FAccelerationStructureHandle, FReal, 3>>())
-	{
-		Collection->CallMoveToTOIHack(Dt, Particle);
-	}
-}
-
-bool RequiresCCDHack(FReal Dt, const TTransientPBDRigidParticleHandle<FReal, 3>& Particle)
-{
-	if (Particle.HasBounds() && (Particle.ObjectState() == EObjectStateType::Dynamic))
-	{
-		FReal Dist2 = Particle.V().SizeSquared() * Dt * Dt;
-		FReal MinBounds = Particle.Geometry()->BoundingBox().Extents().Min();
-		FReal MaxDepth = MinBounds * HackCCD_EnableThreshold;
-		if (Dist2 > MaxDepth * MaxDepth)
-		{
-			//UE_LOG(LogChaos, Warning, TEXT("MoveTOIHack: Enabled at DR = %f / %f"), FMath::Sqrt(Dist2), MaxDepth);
-
-			return true;
-		}
-	}
-	return false;
-}
-
-void CCDHack(const FReal Dt, TParticleView<TPBDRigidParticles<FReal, 3>>& ParticlesView, const ISpatialAcceleration<FAccelerationStructureHandle, FReal, 3>* SpatialAcceleration)
-{
-	if (HackCCD_EnableThreshold > 0)
-	{
-		for (auto& Particle : ParticlesView)
-		{
-			if (RequiresCCDHack(Dt, Particle))
-			{
-				MoveToTOIHack(Dt, Particle, SpatialAcceleration);
-			}
-		}
-	}
-}
-
-//
-//
-//
-// END TOI
-//
-//
-//
 
 template <typename Traits>
 void TPBDRigidsEvolutionGBF<Traits>::AdvanceOneTimeStep(const FReal Dt,const FSubStepInfo& SubStepInfo)
@@ -329,11 +214,6 @@ void TPBDRigidsEvolutionGBF<Traits>::AdvanceOneTimeStepImpl(const FReal Dt,const
 		Base::ComputeIntermediateSpatialAcceleration();
 	}
 
-	{
-		SCOPE_CYCLE_COUNTER(STAT_Evolution_CCDHack);
-		CCDHack(Dt, Particles.GetActiveParticlesView(), InternalAcceleration);
-	}
-
 	{
 		SCOPE_CYCLE_COUNTER(STAT_Evolution_DetectCollisions);
 		CollisionDetector.GetBroadPhase().SetSpatialAcceleration(InternalAcceleration);
diff --git a/Engine/Source/Runtime/Experimental/Chaos/Public/Chaos/Collision/SpatialAccelerationBroadPhase.h b/Engine/Source/Runtime/Experimental/Chaos/Public/Chaos/Collision/SpatialAccelerationBroadPhase.h
--- a/Engine/Source/Runtime/Experimental/Chaos/Public/Chaos/Collision/SpatialAccelerationBroadPhase.h
+++ b/Engine/Source/Runtime/Experimental/Chaos/Public/Chaos/Collision/SpatialAccelerationBroadPhase.h
@@ -338,43 +338,4 @@ namespace Chaos
 
 		FIgnoreCollisionManager IgnoreCollisionManager;
 	};
-
-
-
-
-	extern void MoveToTOIPairHack(FReal Dt, FPBDRigidParticleHandle* Particle1, const FGeometryParticleHandle* Particle2);
-
-	template<typename T_SPATIALACCELERATION>
-	void MoveToTOIHackImpl(FReal Dt, TTransientPBDRigidParticleHandle<FReal, 3>& Particle1, const T_SPATIALACCELERATION* SpatialAcceleration)
-	{
-		TArray<FAccelerationStructureHandle> PotentialIntersections;
-		const FAABB3 Box1 = ComputeWorldSpaceBoundingBox(Particle1);
-		FSimOverlapVisitor OverlapVisitor(PotentialIntersections);
-		SpatialAcceleration->Overlap(Box1, OverlapVisitor);
-
-		const int32 NumPotentials = PotentialIntersections.Num();
-		for (int32 i = 0; i < NumPotentials; ++i)
-		{
-			auto& Particle2 = *PotentialIntersections[i].GetGeometryParticleHandle_PhysicsThread();
-
-			// We only care about static/kinematics
-			if (Particle2.CastToRigidParticle() && (Particle2.CastToRigidParticle()->ObjectState() == EObjectStateType::Dynamic))
-			{
-				continue;
-			}
-			if (!Particle1.Geometry() && !Particle2.Geometry())
-			{
-				continue;
-			}
-
-			const bool bBody2Bounded = HasBoundingBox(Particle2);
-			if (Particle1.Handle() == Particle2.Handle())
-			{
-				continue;
-			}
-
-			MoveToTOIPairHack(Dt, Particle1.Handle(), &Particle2);
-		}
-	}
-
 }
diff --git a/Engine/Source/Runtime/Experimental/Chaos/Public/Chaos/ISpatialAccelerationCollection.h b/Engine/Source/Runtime/Experimental/Chaos/Public/Chaos/ISpatialAccelerationCollection.h
--- a/Engine/Source/Runtime/Experimental/Chaos/Public/Chaos/ISpatialAccelerationCollection.h
+++ b/Engine/Source/Runtime/Experimental/Chaos/Public/Chaos/ISpatialAccelerationCollection.h
@@ -30,7 +30,6 @@ public:
 
 	/** This is kind of a hack to avoid virtuals. We simply route calls into templated functions */
 	virtual void PBDComputeConstraintsLowLevel(T Dt, FSpatialAccelerationBroadPhase& BroadPhase, FNarrowPhase& NarrowPhase, FAsyncCollisionReceiver& Receiver, CollisionStats::FStatData& StatData, IResimCacheBase* ResimCache) const = 0;
-	virtual void CallMoveToTOIHack(FReal Dt, TTransientPBDRigidParticleHandle<FReal, 3>& Particle) const = 0;
 	virtual TArray<FSpatialAccelerationIdx> GetAllSpatialIndices() const = 0;
 
 	bool IsBucketActive(uint8 BucketIdx) const
diff --git a/Engine/Source/Runtime/Experimental/Chaos/Public/Chaos/SpatialAccelerationCollection.h b/Engine/Source/Runtime/Experimental/Chaos/Public/Chaos/SpatialAccelerationCollection.h
--- a/Engine/Source/Runtime/Experimental/Chaos/Public/Chaos/SpatialAccelerationCollection.h
+++ b/Engine/Source/Runtime/Experimental/Chaos/Public/Chaos/SpatialAccelerationCollection.h
@@ -306,20 +306,6 @@ typename TEnableIf<!TIsSame<typename SpatialAccelerationCollection::TPayloadType
 {
 }
 
-template<typename T_SPATIALACCELERATION>
-void MoveToTOIHackImpl(FReal Dt, TTransientPBDRigidParticleHandle<FReal, 3>& Particle1, const T_SPATIALACCELERATION* SpatialAcceleration);
-
-template <typename SpatialAccelerationCollection>
-typename TEnableIf<TIsSame<typename SpatialAccelerationCollection::TPayloadType, TAccelerationStructureHandle<FReal, 3>>::Value, void>::Type CallMoveToTOIHack_Helper(FReal Dt, const SpatialAccelerationCollection& Accel, TTransientPBDRigidParticleHandle<FReal, 3>& Particle)
-{
-	MoveToTOIHackImpl(Dt, Particle, &Accel);
-}
-
-template <typename SpatialAccelerationCollection>
-typename TEnableIf<!TIsSame<typename SpatialAccelerationCollection::TPayloadType, TAccelerationStructureHandle<FReal, 3>>::Value, void>::Type CallMoveToTOIHack_Helper(FReal Dt, const SpatialAccelerationCollection& Accel, TTransientPBDRigidParticleHandle<FReal, 3>& Particle)
-{
-}
-
 
 
 template <typename ... TSpatialAccelerationTypes>
@@ -492,11 +478,6 @@ public:
 		PBDComputeConstraintsLowLevel_Helper(Dt, *this, BroadPhase, NarrowPhase, Receiver, StatData, ResimCache);
 	}
 
-	virtual void CallMoveToTOIHack(FReal Dt, TTransientPBDRigidParticleHandle<FReal, 3>& Particle) const override
-	{
-		CallMoveToTOIHack_Helper(Dt, *this, Particle);
-	}
-
 	virtual void Serialize(FChaosArchive& Ar)
 	{
 		//todo: let user serialize out bucket params
