GeometryCollectionComponent -- proper initialization on embedded data struct for default constructor.

diff --git a/Engine/Source/Runtime/Experimental/GeometryCollectionEngine/Public/GeometryCollection/GeometryCollectionObject.h b/Engine/Source/Runtime/Experimental/GeometryCollectionEngine/Public/GeometryCollection/GeometryCollectionObject.h
--- a/Engine/Source/Runtime/Experimental/GeometryCollectionEngine/Public/GeometryCollection/GeometryCollectionObject.h
+++ b/Engine/Source/Runtime/Experimental/GeometryCollectionEngine/Public/GeometryCollection/GeometryCollectionObject.h
@@ -39,7 +39,13 @@ struct GEOMETRYCOLLECTIONENGINE_API FGeometryCollectionEmbeddedExemplar
 {
 	GENERATED_BODY()
 
-	FGeometryCollectionEmbeddedExemplar() {};
+	FGeometryCollectionEmbeddedExemplar() 
+		: StaticMeshExemplar(FString(TEXT("None")))
+		, StartCullDistance(0.0f)
+		, EndCullDistance(0.0f)
+		, InstanceCount(0)
+	{ };
+	
 	FGeometryCollectionEmbeddedExemplar(FSoftObjectPath NewExemplar)
 		: StaticMeshExemplar(NewExemplar)
 		, StartCullDistance(0.0f)
