Chaos - fix cull distance calculation

diff --git a/Engine/Source/Programs/HeadlessChaos/Private/HeadlessChaosTestSolverProxies.cpp b/Engine/Source/Programs/HeadlessChaos/Private/HeadlessChaosTestSolverProxies.cpp
--- a/Engine/Source/Programs/HeadlessChaos/Private/HeadlessChaosTestSolverProxies.cpp
+++ b/Engine/Source/Programs/HeadlessChaos/Private/HeadlessChaosTestSolverProxies.cpp
@@ -63,7 +63,7 @@ namespace ChaosTest {
 	}
 
 	template<typename Traits>
-	void SingleParticleProxyWakeEventPropergationTest()
+	void SingleParticleProxyWakeEventPropagationTest()
 	{
 		using namespace Chaos;
 		auto Sphere = TSharedPtr<FImplicitObject, ESPMode::ThreadSafe>(new TSphere<FReal, 3>(FVec3(0), 10));
@@ -122,6 +122,6 @@ namespace ChaosTest {
 	TYPED_TEST(AllTraits, SingleParticleProxyTests)
 	{
 		ChaosTest::SingleParticleProxySingleThreadTest<TypeParam>();
-		ChaosTest::SingleParticleProxyWakeEventPropergationTest<TypeParam>();
+		ChaosTest::SingleParticleProxyWakeEventPropagationTest<TypeParam>();
 	}
 }
diff --git a/Engine/Source/Runtime/Experimental/Chaos/Public/Chaos/Collision/SpatialAccelerationBroadPhase.h b/Engine/Source/Runtime/Experimental/Chaos/Public/Chaos/Collision/SpatialAccelerationBroadPhase.h
--- a/Engine/Source/Runtime/Experimental/Chaos/Public/Chaos/Collision/SpatialAccelerationBroadPhase.h
+++ b/Engine/Source/Runtime/Experimental/Chaos/Public/Chaos/Collision/SpatialAccelerationBroadPhase.h
@@ -311,9 +311,17 @@ namespace Chaos
 					
 					{
 						SCOPE_CYCLE_COUNTER(STAT_Collisions_GenerateCollisions);
+
+						// We move the bodies during contact resolution and it may be in any direction
+						// @todo(chaos): this expansion can be very large for some objects - we may want to consider extending only along
+						// the velocity direction if CCD is not enabled for either object.
+						const FReal CullDistance1 = ComputeBoundsThickness(Particle1, Dt, BoundsThickness, BoundsThicknessVelocityInflation).Size();
+						const FReal CullDistance2 = ComputeBoundsThickness(Particle2, Dt, BoundsThickness, BoundsThicknessVelocityInflation).Size();
+						const FReal NetCullDistance = CullDistance + CullDistance1 + CullDistance2;
+
 						// Generate constraints for the potentially overlapping shape pairs. Also run collision detection to generate
 						// the contact position and normal (for contacts within CullDistance) for use in collision callbacks.
-						NarrowPhase.GenerateCollisions(NewConstraints, Dt, Particle1.Handle(), Particle2.Handle(), CullDistance + BoundsThicknessVelocityInflation);
+						NarrowPhase.GenerateCollisions(NewConstraints, Dt, Particle1.Handle(), Particle2.Handle(), NetCullDistance);
 					}
 				}
 
