Chaos - enable contact manifolds for all shape types

diff --git a/Engine/Source/Runtime/Experimental/Chaos/Private/Chaos/CollisionResolution.cpp b/Engine/Source/Runtime/Experimental/Chaos/Private/Chaos/CollisionResolution.cpp
--- a/Engine/Source/Runtime/Experimental/Chaos/Private/Chaos/CollisionResolution.cpp
+++ b/Engine/Source/Runtime/Experimental/Chaos/Private/Chaos/CollisionResolution.cpp
@@ -899,7 +899,7 @@ namespace Chaos
 			const TSphere<FReal, 3>* Object1 = Implicit1->template GetObject<const TSphere<FReal, 3> >();
 			if (ensure(Object0 && Object1))
 			{
-				FRigidBodyPointContactConstraint Constraint = FRigidBodyPointContactConstraint(Particle0, Implicit0, nullptr, LocalTransform0, Particle1, Implicit1, nullptr, LocalTransform1, EContactShapesType::SphereSphere, false);	// No manifold
+				FRigidBodyPointContactConstraint Constraint = FRigidBodyPointContactConstraint(Particle0, Implicit0, nullptr, LocalTransform0, Particle1, Implicit1, nullptr, LocalTransform1, EContactShapesType::SphereSphere, Context.bAllowManifolds);
 				if (T_TRAITS::bImmediateUpdate)
 				{
 					FRigidTransform3 WorldTransform0 = LocalTransform0 * Collisions::GetTransform(Particle0);
@@ -941,7 +941,7 @@ namespace Chaos
 			const FHeightField* Object1 = Implicit1->template GetObject<const FHeightField >();
 			if (ensure(Object0 && Object1))
 			{
-				FRigidBodyPointContactConstraint Constraint = FRigidBodyPointContactConstraint(Particle0, Implicit0, nullptr, LocalTransform0, Particle1, Implicit1, nullptr, LocalTransform1, EContactShapesType::SphereHeightField, false);
+				FRigidBodyPointContactConstraint Constraint = FRigidBodyPointContactConstraint(Particle0, Implicit0, nullptr, LocalTransform0, Particle1, Implicit1, nullptr, LocalTransform1, EContactShapesType::SphereHeightField, Context.bAllowManifolds);
 				if (T_TRAITS::bImmediateUpdate)
 				{
 					FRigidTransform3 WorldTransform0 = LocalTransform0 * Collisions::GetTransform(Particle0);
@@ -1010,7 +1010,7 @@ namespace Chaos
 			const TPlane<FReal, 3>* Object1 = Implicit1->template GetObject<const TPlane<FReal, 3> >();
 			if (ensure(Object0 && Object1))
 			{
-				FRigidBodyPointContactConstraint Constraint = FRigidBodyPointContactConstraint(Particle0, Implicit0, nullptr, LocalTransform0, Particle1, Implicit1, nullptr, LocalTransform1, EContactShapesType::SpherePlane, false);
+				FRigidBodyPointContactConstraint Constraint = FRigidBodyPointContactConstraint(Particle0, Implicit0, nullptr, LocalTransform0, Particle1, Implicit1, nullptr, LocalTransform1, EContactShapesType::SpherePlane, Context.bAllowManifolds);
 				if (T_TRAITS::bImmediateUpdate)
 				{
 					FRigidTransform3 WorldTransform0 = LocalTransform0 * Collisions::GetTransform(Particle0);
@@ -1066,7 +1066,7 @@ namespace Chaos
 			const TBox<FReal, 3>* Object1 = Implicit1->template GetObject<const TBox<FReal, 3> >();
 			if (ensure(Object0 && Object1))
 			{
-				FRigidBodyPointContactConstraint Constraint = FRigidBodyPointContactConstraint(Particle0, Implicit0, nullptr, LocalTransform0, Particle1, Implicit1, nullptr, LocalTransform1, EContactShapesType::SphereBox, false);
+				FRigidBodyPointContactConstraint Constraint = FRigidBodyPointContactConstraint(Particle0, Implicit0, nullptr, LocalTransform0, Particle1, Implicit1, nullptr, LocalTransform1, EContactShapesType::SphereBox, Context.bAllowManifolds);
 				if (T_TRAITS::bImmediateUpdate)
 				{
 					FRigidTransform3 WorldTransform0 = LocalTransform0 * Collisions::GetTransform(Particle0);
@@ -1126,7 +1126,7 @@ namespace Chaos
 			const TCapsule<FReal>* Object1 = Implicit1->template GetObject<const TCapsule<FReal> >();
 			if (ensure(Object0 && Object1))
 			{
-				FRigidBodyPointContactConstraint Constraint = FRigidBodyPointContactConstraint(Particle0, Implicit0, nullptr, LocalTransform0, Particle1, Implicit1, nullptr, LocalTransform1, EContactShapesType::SphereCapsule, false);	// No manifold
+				FRigidBodyPointContactConstraint Constraint = FRigidBodyPointContactConstraint(Particle0, Implicit0, nullptr, LocalTransform0, Particle1, Implicit1, nullptr, LocalTransform1, EContactShapesType::SphereCapsule, Context.bAllowManifolds);
 				if (T_TRAITS::bImmediateUpdate)
 				{
 					FRigidTransform3 WorldTransform0 = LocalTransform0 * Collisions::GetTransform(Particle0);
@@ -1153,7 +1153,7 @@ namespace Chaos
 			const FImplicitObject* Object1 = Implicit1;
 			if (ensure(Object0 && Object1))
 			{
-				FRigidBodyPointContactConstraint Constraint = FRigidBodyPointContactConstraint(Particle0, Implicit0, nullptr, LocalTransform0, Particle1, Implicit1, nullptr, LocalTransform1, EContactShapesType::SphereConvex, false);	// No manifold
+				FRigidBodyPointContactConstraint Constraint = FRigidBodyPointContactConstraint(Particle0, Implicit0, nullptr, LocalTransform0, Particle1, Implicit1, nullptr, LocalTransform1, EContactShapesType::SphereConvex, Context.bAllowManifolds);
 				if (T_TRAITS::bImmediateUpdate)
 				{
 					FRigidTransform3 WorldTransform0 = LocalTransform0 * Collisions::GetTransform(Particle0);
@@ -1213,7 +1213,7 @@ namespace Chaos
 			{
 				if (const TImplicitObjectScaled<FTriangleMeshImplicitObject>* ScaledTriangleMesh = Implicit1->template GetObject<const TImplicitObjectScaled<FTriangleMeshImplicitObject>>())
 				{
-					FRigidBodyPointContactConstraint Constraint = FRigidBodyPointContactConstraint(Particle0, Implicit0, nullptr, LocalTransform0, Particle1, Implicit1, nullptr, LocalTransform1, EContactShapesType::SphereTriMesh, false);
+					FRigidBodyPointContactConstraint Constraint = FRigidBodyPointContactConstraint(Particle0, Implicit0, nullptr, LocalTransform0, Particle1, Implicit1, nullptr, LocalTransform1, EContactShapesType::SphereTriMesh, Context.bAllowManifolds);
 					if (T_TRAITS::bImmediateUpdate)
 					{
 						FRigidTransform3 WorldTransform0 = LocalTransform0 * Collisions::GetTransform(Particle0);
@@ -1224,7 +1224,7 @@ namespace Chaos
 				}
 				else if (const FTriangleMeshImplicitObject* TriangleMesh = Implicit1->template GetObject<const FTriangleMeshImplicitObject>())
 				{
-					FRigidBodyPointContactConstraint Constraint = FRigidBodyPointContactConstraint(Particle0, Implicit0, nullptr, LocalTransform0, Particle1, Implicit1, nullptr, LocalTransform1, EContactShapesType::SphereTriMesh, false);
+					FRigidBodyPointContactConstraint Constraint = FRigidBodyPointContactConstraint(Particle0, Implicit0, nullptr, LocalTransform0, Particle1, Implicit1, nullptr, LocalTransform1, EContactShapesType::SphereTriMesh, Context.bAllowManifolds);
 					if (T_TRAITS::bImmediateUpdate)
 					{
 						FRigidTransform3 WorldTransform0 = LocalTransform0 * Collisions::GetTransform(Particle0);
@@ -1322,7 +1322,7 @@ namespace Chaos
 			const TCapsule<FReal>* Object1 = Implicit1->template GetObject<const TCapsule<FReal> >();
 			if (ensure(Object0 && Object1))
 			{
-				FRigidBodyPointContactConstraint Constraint = FRigidBodyPointContactConstraint(Particle0, Implicit0, nullptr, LocalTransform0, Particle1, Implicit1, nullptr, LocalTransform1, EContactShapesType::CapsuleCapsule, false);
+				FRigidBodyPointContactConstraint Constraint = FRigidBodyPointContactConstraint(Particle0, Implicit0, nullptr, LocalTransform0, Particle1, Implicit1, nullptr, LocalTransform1, EContactShapesType::CapsuleCapsule, Context.bAllowManifolds);
 				if (T_TRAITS::bImmediateUpdate)
 				{
 					FRigidTransform3 WorldTransform0 = LocalTransform0 * Collisions::GetTransform(Particle0);
@@ -1572,7 +1572,7 @@ namespace Chaos
 			const TBox<FReal, 3>* Object1 = Implicit1->template GetObject<const TBox<FReal, 3>>();
 			if (ensure(Object0 && Object1))
 			{
-				FRigidBodyPointContactConstraint Constraint = FRigidBodyPointContactConstraint(Particle0, Implicit0, nullptr, LocalTransform0, Particle1, Implicit1, nullptr, LocalTransform1, EContactShapesType::CapsuleBox, false);
+				FRigidBodyPointContactConstraint Constraint = FRigidBodyPointContactConstraint(Particle0, Implicit0, nullptr, LocalTransform0, Particle1, Implicit1, nullptr, LocalTransform1, EContactShapesType::CapsuleBox, Context.bAllowManifolds);
 				if (T_TRAITS::bImmediateUpdate)
 				{
 					FRigidTransform3 WorldTransform0 = LocalTransform0 * Collisions::GetTransform(Particle0);
@@ -1600,7 +1600,7 @@ namespace Chaos
 			const FImplicitObject* Object1 = Implicit1;
 			if (ensure(Object0 && Object1))
 			{
-				FRigidBodyPointContactConstraint Constraint = FRigidBodyPointContactConstraint(Particle0, Implicit0, nullptr, LocalTransform0, Particle1, Implicit1, nullptr, LocalTransform1, EContactShapesType::CapsuleConvex, false);
+				FRigidBodyPointContactConstraint Constraint = FRigidBodyPointContactConstraint(Particle0, Implicit0, nullptr, LocalTransform0, Particle1, Implicit1, nullptr, LocalTransform1, EContactShapesType::CapsuleConvex, Context.bAllowManifolds);
 				if (T_TRAITS::bImmediateUpdate)
 				{
 					FRigidTransform3 WorldTransform0 = LocalTransform0 * Collisions::GetTransform(Particle0);
@@ -1649,7 +1649,7 @@ namespace Chaos
 			const FHeightField* Object1 = Implicit1->template GetObject<const FHeightField >();
 			if (ensure(Object0 && Object1))
 			{
-				FRigidBodyPointContactConstraint Constraint = FRigidBodyPointContactConstraint(Particle0, Implicit0, nullptr, LocalTransform0, Particle1, Implicit1, nullptr, LocalTransform1, EContactShapesType::CapsuleHeightField, false);
+				FRigidBodyPointContactConstraint Constraint = FRigidBodyPointContactConstraint(Particle0, Implicit0, nullptr, LocalTransform0, Particle1, Implicit1, nullptr, LocalTransform1, EContactShapesType::CapsuleHeightField, Context.bAllowManifolds);
 				if (T_TRAITS::bImmediateUpdate)
 				{
 					FRigidTransform3 WorldTransform0 = LocalTransform0 * Collisions::GetTransform(Particle0);
@@ -1732,7 +1732,7 @@ namespace Chaos
 			{
 				if (const TImplicitObjectScaled<FTriangleMeshImplicitObject>* ScaledTriangleMesh = Implicit1->template GetObject<const TImplicitObjectScaled<FTriangleMeshImplicitObject>>())
 				{
-					FRigidBodyPointContactConstraint Constraint = FRigidBodyPointContactConstraint(Particle0, Implicit0, nullptr, LocalTransform0, Particle1, Implicit1, nullptr, LocalTransform1, EContactShapesType::CapsuleTriMesh, false);
+					FRigidBodyPointContactConstraint Constraint = FRigidBodyPointContactConstraint(Particle0, Implicit0, nullptr, LocalTransform0, Particle1, Implicit1, nullptr, LocalTransform1, EContactShapesType::CapsuleTriMesh, Context.bAllowManifolds);
 					if (T_TRAITS::bImmediateUpdate)
 					{
 						FRigidTransform3 WorldTransform0 = LocalTransform0 * Collisions::GetTransform(Particle0);
@@ -1743,7 +1743,7 @@ namespace Chaos
 				}
 				else if (const FTriangleMeshImplicitObject* TriangleMesh = Implicit1->template GetObject<const FTriangleMeshImplicitObject>())
 				{
-					FRigidBodyPointContactConstraint Constraint = FRigidBodyPointContactConstraint(Particle0, Implicit0, nullptr, LocalTransform0, Particle1, Implicit1, nullptr, LocalTransform1, EContactShapesType::CapsuleTriMesh, false);
+					FRigidBodyPointContactConstraint Constraint = FRigidBodyPointContactConstraint(Particle0, Implicit0, nullptr, LocalTransform0, Particle1, Implicit1, nullptr, LocalTransform1, EContactShapesType::CapsuleTriMesh, Context.bAllowManifolds);
 					if (T_TRAITS::bImmediateUpdate)
 					{
 						FRigidTransform3 WorldTransform0 = LocalTransform0 * Collisions::GetTransform(Particle0);
@@ -1871,7 +1871,7 @@ namespace Chaos
 			const FHeightField* Object1 = Implicit1->template GetObject<const FHeightField >();
 			if (ensure(Implicit0->IsConvex() && Object1))
 			{
-				FRigidBodyPointContactConstraint Constraint = FRigidBodyPointContactConstraint(Particle0, Implicit0, nullptr, LocalTransform0, Particle1, Implicit1, nullptr, LocalTransform1, EContactShapesType::ConvexHeightField, false);
+				FRigidBodyPointContactConstraint Constraint = FRigidBodyPointContactConstraint(Particle0, Implicit0, nullptr, LocalTransform0, Particle1, Implicit1, nullptr, LocalTransform1, EContactShapesType::ConvexHeightField, Context.bAllowManifolds);
 				if (T_TRAITS::bImmediateUpdate)
 				{
 					FRigidTransform3 WorldTransform0 = LocalTransform0 * Collisions::GetTransform(Particle0);
@@ -1950,7 +1950,7 @@ namespace Chaos
 			{
 				if (const TImplicitObjectScaled<FTriangleMeshImplicitObject>* ScaledTriangleMesh = Implicit1->template GetObject<const TImplicitObjectScaled<FTriangleMeshImplicitObject>>())
 				{
-					FRigidBodyPointContactConstraint Constraint = FRigidBodyPointContactConstraint(Particle0, Implicit0, nullptr, LocalTransform0, Particle1, Implicit1, nullptr, LocalTransform0, EContactShapesType::ConvexTriMesh, false);
+					FRigidBodyPointContactConstraint Constraint = FRigidBodyPointContactConstraint(Particle0, Implicit0, nullptr, LocalTransform0, Particle1, Implicit1, nullptr, LocalTransform0, EContactShapesType::ConvexTriMesh, Context.bAllowManifolds);
 					if (T_TRAITS::bImmediateUpdate)
 					{
 						FRigidTransform3 WorldTransform0 = LocalTransform0 * Collisions::GetTransform(Particle0);
@@ -1961,7 +1961,7 @@ namespace Chaos
 				}
 				else if (const FTriangleMeshImplicitObject* TriangleMesh = Implicit1->template GetObject<const FTriangleMeshImplicitObject>())
 				{
-					FRigidBodyPointContactConstraint Constraint = FRigidBodyPointContactConstraint(Particle0, Implicit0, nullptr, LocalTransform0, Particle1, Implicit1, nullptr, LocalTransform1, EContactShapesType::ConvexTriMesh, false);
+					FRigidBodyPointContactConstraint Constraint = FRigidBodyPointContactConstraint(Particle0, Implicit0, nullptr, LocalTransform0, Particle1, Implicit1, nullptr, LocalTransform1, EContactShapesType::ConvexTriMesh, Context.bAllowManifolds);
 					if (T_TRAITS::bImmediateUpdate)
 					{
 						FRigidTransform3 WorldTransform0 = LocalTransform0 * Collisions::GetTransform(Particle0);
@@ -2064,7 +2064,7 @@ namespace Chaos
 		template<typename T_TRAITS>
 		void ConstructLevelsetLevelsetConstraints(TGeometryParticleHandle<FReal, 3>* Particle0, TGeometryParticleHandle<FReal, 3>* Particle1, const FImplicitObject* Implicit0, const TBVHParticles<FReal, 3>* Simplicial0, const FImplicitObject* Implicit1, const TBVHParticles<FReal, 3>* Simplicial1, const FRigidTransform3& LocalTransform0, const FRigidTransform3& LocalTransform1, const FReal CullDistance, const FReal Dt, const FCollisionContext& Context, FCollisionConstraintsArray& NewConstraints)
 		{
-			FRigidBodyPointContactConstraint Constraint = FRigidBodyPointContactConstraint(Particle0, Implicit0, Simplicial0, LocalTransform0, Particle1, Implicit1, Simplicial1, LocalTransform1, EContactShapesType::LevelSetLevelSet, true);
+			FRigidBodyPointContactConstraint Constraint = FRigidBodyPointContactConstraint(Particle0, Implicit0, Simplicial0, LocalTransform0, Particle1, Implicit1, Simplicial1, LocalTransform1, EContactShapesType::LevelSetLevelSet, Context.bAllowManifolds);
 
 			bool bIsParticleDynamic0 = Particle0->CastToRigidParticle() && Particle0->ObjectState() == EObjectStateType::Dynamic;
 			int32 P0NumCollisionParticles = Simplicial0 ? Simplicial0->Size() : Particle0->CastToRigidParticle()->CollisionParticlesSize();
