Chaos - disable debugdraw defaults

diff --git a/Engine/Source/Runtime/Experimental/Chaos/Private/PBDRigidsSolver.cpp b/Engine/Source/Runtime/Experimental/Chaos/Private/PBDRigidsSolver.cpp
--- a/Engine/Source/Runtime/Experimental/Chaos/Private/PBDRigidsSolver.cpp
+++ b/Engine/Source/Runtime/Experimental/Chaos/Private/PBDRigidsSolver.cpp
@@ -32,7 +32,7 @@ DEFINE_LOG_CATEGORY_STATIC(LogPBDRigidsSolver, Log, All);
 #if CHAOS_DEBUG_DRAW
 
 // Must be 0 when checked in...
-#define CHAOS_SOLVER_ENABLE_DEBUG_DRAW 1
+#define CHAOS_SOLVER_ENABLE_DEBUG_DRAW 0
 
 int32 ChaosSolverDebugDrawShapes = CHAOS_SOLVER_ENABLE_DEBUG_DRAW;
 int32 ChaosSolverDebugDrawCollisions = CHAOS_SOLVER_ENABLE_DEBUG_DRAW;
