GeometryCollection -- include StaticMesh.h to correct non-Unity build

diff --git a/Engine/Source/Runtime/Experimental/GeometryCollectionEngine/Private/GeometryCollection/GeometryCollectionObject.cpp b/Engine/Source/Runtime/Experimental/GeometryCollectionEngine/Private/GeometryCollection/GeometryCollectionObject.cpp
--- a/Engine/Source/Runtime/Experimental/GeometryCollectionEngine/Private/GeometryCollection/GeometryCollectionObject.cpp
+++ b/Engine/Source/Runtime/Experimental/GeometryCollectionEngine/Private/GeometryCollection/GeometryCollectionObject.cpp
@@ -15,7 +15,7 @@
 #include "Materials/MaterialInstance.h"
 #include "ProfilingDebugging/CookStats.h"
 #include "EngineUtils.h"
-
+#include "Engine/StaticMesh.h"
 
 #if WITH_EDITOR
 #include "GeometryCollection/DerivedDataGeometryCollectionCooker.h"
