GeometryCollection -- updated GetComponent to use undeprecated method.

diff --git a/Engine/Source/Runtime/Experimental/GeometryCollectionEngine/Private/GeometryCollection/GeometryCollectionComponent.cpp b/Engine/Source/Runtime/Experimental/GeometryCollectionEngine/Private/GeometryCollection/GeometryCollectionComponent.cpp
--- a/Engine/Source/Runtime/Experimental/GeometryCollectionEngine/Private/GeometryCollection/GeometryCollectionComponent.cpp
+++ b/Engine/Source/Runtime/Experimental/GeometryCollectionEngine/Private/GeometryCollection/GeometryCollectionComponent.cpp
@@ -2885,7 +2885,8 @@ bool UGeometryCollectionComponent::IsEmbeddedGeometryValid() const
 void UGeometryCollectionComponent::ClearEmbeddedGeometry()
 {
 	AActor* OwningActor = GetOwner();
-	TArray<UActorComponent*> TargetComponents = OwningActor->GetComponentsByClass(UInstancedStaticMeshComponent::StaticClass());
+	TArray<UActorComponent*> TargetComponents;
+	OwningActor->GetComponents(TargetComponents, false);
 	
 	for (UActorComponent* EmbeddedGeometryComponent : TargetComponents)
 	{
