Chaos : Clustering
Remove active cluster particle array support, its not actually used by anything.

#okforgethub public, internal

diff --git a/Engine/Source/Programs/HeadlessChaos/Private/GeometryCollection/GeometryCollectionTestClustering.cpp b/Engine/Source/Programs/HeadlessChaos/Private/GeometryCollection/GeometryCollectionTestClustering.cpp
--- a/Engine/Source/Programs/HeadlessChaos/Private/GeometryCollection/GeometryCollectionTestClustering.cpp
+++ b/Engine/Source/Programs/HeadlessChaos/Private/GeometryCollection/GeometryCollectionTestClustering.cpp
@@ -36,23 +36,26 @@ namespace GeometryCollectionTest
 {
 	using namespace ChaosTest;
 
-	bool ClusterMapContains(const Chaos::TPBDRigidClustering<FPBDRigidsEvolution, FPBDCollisionConstraints>::FClusterMap& ClusterMap, const FPBDRigidParticleHandle* Key, TArray<FPBDRigidParticleHandle*> Elements)
+	bool ClusterMapContains(const Chaos::TPBDRigidClustering<FPBDRigidsEvolution, FPBDCollisionConstraints>::FClusterMap& ClusterMap, const FPBDRigidParticleHandle* InKey, TArray<FPBDRigidParticleHandle*> Elements)
 	{
 		if (ClusterMap.Num())
 		{
-			if(ClusterMap.Contains(Key))
+			if (const Chaos::TPBDRigidClusteredParticleHandle<FReal, 3>* Key = InKey->CastToClustered())
 			{
-				if(ClusterMap[Key].Num() == Elements.Num())
+				if (ClusterMap.Contains(Key))
 				{
-					for(FPBDRigidParticleHandle* Element : Elements)
+					if (ClusterMap[Key].Num() == Elements.Num())
 					{
-						if(!ClusterMap[Key].Contains(Element))
+						for (FPBDRigidParticleHandle* Element : Elements)
 						{
-							return false;
+							if (!ClusterMap[Key].Contains(Element))
+							{
+								return false;
+							}
 						}
-					}
 
-					return true;
+						return true;
+					}
 				}
 			}
 		}
diff --git a/Engine/Source/Programs/HeadlessChaos/Private/HeadlessChaosTestClustering.cpp b/Engine/Source/Programs/HeadlessChaos/Private/HeadlessChaosTestClustering.cpp
--- a/Engine/Source/Programs/HeadlessChaos/Private/HeadlessChaosTestClustering.cpp
+++ b/Engine/Source/Programs/HeadlessChaos/Private/HeadlessChaosTestClustering.cpp
@@ -261,11 +261,12 @@ namespace ChaosTest {
 		//EXPECT_TRUE(Evolution.GetParticles().Disabled(RootClusterHandle) == false);	//one of the connected pieces should re-use this
 		EXPECT_TRUE(ClusterHandlesDup[2]->Disabled() == false);	//this cluster is on its own and should be enabled 
 		
-		EXPECT_EQ(Evolution.GetActiveClusteredArray().Num(), 3);	//there should only be 3 pieces
+		
+		EXPECT_EQ(Evolution.GetRigidClustering().GetTopLevelClusterParents().Num(), 3);	//there should only be 3 pieces
 		for (uint32 BoxID : BoxIDs)
 		{
 			EXPECT_TRUE(ClusteredParticles.Disabled(BoxID));	//no boxes should be active yet
-			EXPECT_TRUE(Evolution.GetActiveClusteredArray().Contains(ClusteredParticles.Handle(BoxID)) == false);
+			EXPECT_TRUE(Evolution.GetRigidClustering().GetTopLevelClusterParents().Contains(ClusteredParticles.Handle(BoxID)) == false);
 			for (int32 Island = 0; Island < Evolution.NumIslands(); ++Island)
 			{
 				EXPECT_TRUE(Evolution.GetIslandParticles(Island).Contains(ClusteredParticles.Handle(BoxID)) == false);
@@ -281,16 +282,17 @@ namespace ChaosTest {
 		//EXPECT_TRUE(Evolution.GetParticles().Disabled(RootClusterHandle) == false);	//one of the connected pieces should re-use this
 		EXPECT_TRUE(ClusterHandlesDup[2]->Disabled() == false);	//this cluster is on its own and should be enabled 
 
-		EXPECT_EQ(Evolution.GetActiveClusteredArray().Num(), 3);	//there should only be 3 pieces
+		EXPECT_EQ(Evolution.GetRigidClustering().GetTopLevelClusterParents().Num(), 3);	//there should only be 3 pieces
 		for (uint32 BoxID : BoxIDs)
 		{
 			EXPECT_TRUE(ClusteredParticles.Disabled(BoxID));	//no boxes should be active yet
-			EXPECT_TRUE(Evolution.GetActiveClusteredArray().Contains(ClusteredParticles.Handle(BoxID)) == false);
+			EXPECT_TRUE(Evolution.GetRigidClustering().GetTopLevelClusterParents().Contains(ClusteredParticles.Handle(BoxID)) == false);
 			for (int32 Island = 0; Island < Evolution.NumIslands(); ++Island)
 			{
 				EXPECT_TRUE(Evolution.GetIslandParticles(Island).Contains(ClusteredParticles.Handle(BoxID)) == false);
 			}
 		}
+		
 	}
 }
 
diff --git a/Engine/Source/Runtime/Experimental/Chaos/Private/Chaos/PBDRigidClustering.cpp b/Engine/Source/Runtime/Experimental/Chaos/Private/Chaos/PBDRigidClustering.cpp
--- a/Engine/Source/Runtime/Experimental/Chaos/Private/Chaos/PBDRigidClustering.cpp
+++ b/Engine/Source/Runtime/Experimental/Chaos/Private/Chaos/PBDRigidClustering.cpp
@@ -1259,20 +1259,22 @@ namespace Chaos
 			//
 			{
 				SCOPE_CYCLE_COUNTER(STAT_UpdateDirtyImpulses);
-				const auto& ActiveClusteredArray = MEvolution.GetActiveClusteredArray();
-				for (const auto& ActiveCluster : ActiveClusteredArray)
+				for (const auto& ActiveCluster : GetTopLevelClusterParents() )
 				{
-					if (ActiveCluster->ClusterIds().NumChildren > 0) //active index is a cluster
+					if (!ActiveCluster->Disabled())
 					{
-						TArray<FPBDRigidParticleHandle*>& ParentToChildren = MChildren[ActiveCluster];
-						for (FPBDRigidParticleHandle* Child : ParentToChildren)
+						if (ActiveCluster->ClusterIds().NumChildren > 0) //active index is a cluster
 						{
-							if (FPBDRigidClusteredParticleHandle* ClusteredChild = Child->CastToClustered())
+							TArray<FRigidHandle>& ParentToChildren = MChildren[ActiveCluster];
+							for (FRigidHandle Child : ParentToChildren)
 							{
-								if (ClusteredChild->Strain() <= 0.f)
+								if (FClusterHandle ClusteredChild = Child->CastToClustered())
 								{
-									ClusteredChild->CollisionImpulse() = FLT_MAX;
-									MCollisionImpulseArrayDirty = true;
+									if (ClusteredChild->Strain() <= 0.f)
+									{
+										ClusteredChild->CollisionImpulse() = FLT_MAX;
+										MCollisionImpulseArrayDirty = true;
+									}
 								}
 							}
 						}
@@ -1355,7 +1357,7 @@ namespace Chaos
 			if (ClusteredParticle->ClusterIds().NumChildren)
 			{
 				AllActivatedChildren.Add(
-					ClusteredParticle, 
+					ClusteredParticle,
 					ReleaseClusterParticles(ClusteredParticle, ExternalStrainMap));
 			}
 			else
@@ -1394,9 +1396,9 @@ namespace Chaos
 		if (FPBDRigidClusteredParticleHandle* ClusteredCurrentNode = CurrentNode->CastToClustered())
 		{
 			FReal ChildrenStrains = (FReal)0.;
-			if (MChildren.Contains(CurrentNode))
+			if (MChildren.Contains(ClusteredCurrentNode))
 			{
-				for (FPBDRigidParticleHandle* Child : MChildren[CurrentNode])
+				for (FRigidHandle Child : MChildren[ClusteredCurrentNode])
 				{
 					ChildrenStrains += PromoteStrains(Child);
 				}
@@ -1414,46 +1416,52 @@ namespace Chaos
 	DECLARE_CYCLE_STAT(TEXT("TPBDRigidClustering<>::UpdateKinematicProperties()"), STAT_UpdateKinematicProperties, STATGROUP_Chaos);
 	template<class T_FPBDRigidsEvolution, class T_FPBDCollisionConstraint>
 	void TPBDRigidClustering<T_FPBDRigidsEvolution, T_FPBDCollisionConstraint>::UpdateKinematicProperties(
-		Chaos::FPBDRigidParticleHandle* Parent)
+		Chaos::FPBDRigidParticleHandle* InParent)
 	{
 		SCOPE_CYCLE_COUNTER(STAT_UpdateKinematicProperties);
 
 		EObjectStateType ObjectState = EObjectStateType::Dynamic;
-		check(Parent != nullptr);
-		if (MChildren.Contains(Parent) && MChildren[Parent].Num())
+		check(InParent != nullptr);
+		if (FClusterHandle ClusteredCurrentNode = InParent->CastToClustered())
 		{
-			// TQueue is a linked list, which has no preallocator.
-			TQueue<Chaos::FPBDRigidParticleHandle*> Queue;
-			for (Chaos::FPBDRigidParticleHandle* Child : MChildren[Parent])
+			if (MChildren.Contains(ClusteredCurrentNode) && MChildren[ClusteredCurrentNode].Num())
 			{
-				Queue.Enqueue(Child);
-			}
+				// TQueue is a linked list, which has no preallocator.
+				TQueue<Chaos::FPBDRigidParticleHandle*> Queue;
+				for (Chaos::FPBDRigidParticleHandle* Child : MChildren[ClusteredCurrentNode])
+				{
+					Queue.Enqueue(Child);
+				}
 
-			Chaos::FPBDRigidParticleHandle* CurrentHandle;
-			while (Queue.Dequeue(CurrentHandle) && ObjectState == EObjectStateType::Dynamic)
-			{
-				// @question : Maybe we should just store the leaf node bodies in a
-				// map, that will require Memory(n*log(n))
-				if (MChildren.Contains(CurrentHandle))
+				Chaos::FPBDRigidParticleHandle* CurrentHandle;
+				while (Queue.Dequeue(CurrentHandle) && ObjectState == EObjectStateType::Dynamic)
 				{
-					for (Chaos::FPBDRigidParticleHandle* Child : MChildren[CurrentHandle])
+					if (FClusterHandle CurrentClusterHandle = CurrentHandle->CastToClustered())
 					{
-						Queue.Enqueue(Child);
+						// @question : Maybe we should just store the leaf node bodies in a
+						// map, that will require Memory(n*log(n))
+						if (MChildren.Contains(CurrentClusterHandle))
+						{
+							for( Chaos::FPBDRigidParticleHandle* Child : MChildren[CurrentClusterHandle])
+							{
+								Queue.Enqueue(Child);
+							}
+						}
 					}
-				}
 
-				const EObjectStateType CurrState = CurrentHandle->ObjectState();
-				if (CurrState == EObjectStateType::Kinematic)
-				{
-					ObjectState = EObjectStateType::Kinematic;
-				}
-				else if (CurrState == EObjectStateType::Static)
-				{
-					ObjectState = EObjectStateType::Static;
+					const EObjectStateType CurrState = CurrentHandle->ObjectState();
+					if (CurrState == EObjectStateType::Kinematic)
+					{
+						ObjectState = EObjectStateType::Kinematic;
+					}
+					else if (CurrState == EObjectStateType::Static)
+					{
+						ObjectState = EObjectStateType::Static;
+					}
 				}
-			}
 
-			Parent->SetObjectState(ObjectState);
+				ClusteredCurrentNode->SetObjectState(ObjectState);
+			}
 		}
 	}
 
@@ -1932,12 +1940,12 @@ namespace Chaos
 				}
 			};
 
-			if (const TArray<FPBDRigidParticleHandle*>* ChildrenPtr = MParentToChildren.Find(ConstrainedParticles[0]->CastToRigidParticle()))
+			if (const TArray<FPBDRigidParticleHandle*>* ChildrenPtr = MParentToChildren.Find(ConstrainedParticles[0]->CastToClustered()))
 			{
 				ComputeStrainLambda(ConstrainedParticles[0]->CastToClustered(), *ChildrenPtr);
 			}
 
-			if (const TArray<FPBDRigidParticleHandle*>* ChildrenPtr = MParentToChildren.Find(ConstrainedParticles[1]->CastToRigidParticle()))
+			if (const TArray<FPBDRigidParticleHandle*>* ChildrenPtr = MParentToChildren.Find(ConstrainedParticles[1]->CastToClustered()))
 			{
 				ComputeStrainLambda(ConstrainedParticles[1]->CastToClustered(), *ChildrenPtr);
 			}
@@ -2232,7 +2240,9 @@ namespace Chaos
 	DECLARE_CYCLE_STAT(TEXT("TPBDRigidClustering<>::UpdateConnectivityGraphUsingDelaunayTriangulation"), STAT_UpdateConnectivityGraphUsingDelaunayTriangulation, STATGROUP_Chaos);
 
 	template<class T_FPBDRigidsEvolution, class T_FPBDCollisionConstraint>
-	void TPBDRigidClustering<T_FPBDRigidsEvolution, T_FPBDCollisionConstraint>::UpdateConnectivityGraphUsingDelaunayTriangulation(Chaos::FPBDRigidClusteredParticleHandle* Parent, const FClusterCreationParameters<FReal>& Parameters)
+	void TPBDRigidClustering<T_FPBDRigidsEvolution, T_FPBDCollisionConstraint>::UpdateConnectivityGraphUsingDelaunayTriangulation(
+		Chaos::FPBDRigidClusteredParticleHandle* Parent, 
+		const FClusterCreationParameters<FReal>& Parameters)
 	{
 		SCOPE_CYCLE_COUNTER(STAT_UpdateConnectivityGraphUsingDelaunayTriangulation);
 
diff --git a/Engine/Source/Runtime/Experimental/Chaos/Private/PhysicsProxy/GeometryCollectionPhysicsProxy.cpp b/Engine/Source/Runtime/Experimental/Chaos/Private/PhysicsProxy/GeometryCollectionPhysicsProxy.cpp
--- a/Engine/Source/Runtime/Experimental/Chaos/Private/PhysicsProxy/GeometryCollectionPhysicsProxy.cpp
+++ b/Engine/Source/Runtime/Experimental/Chaos/Private/PhysicsProxy/GeometryCollectionPhysicsProxy.cpp
@@ -1238,10 +1238,9 @@ void FGeometryCollectionPhysicsProxy::GetRelevantParticleHandles(
 				Handles.Add(ClusterHandle);
 				if (ClusterHandle->ClusterIds().NumChildren)
 				{
-					Chaos::TPBDRigidParticleHandle<float, 3>* RigidHandle = ClusterHandle->CastToRigidParticle();
-					if (ClusterMap.Contains(RigidHandle))
+					if (ClusterMap.Contains(ClusterHandle))
 					{
-						for (Chaos::TPBDRigidParticleHandle<float, 3> * Child : ClusterMap[RigidHandle])
+						for (Chaos::TPBDRigidParticleHandle<float, 3> * Child : ClusterMap[ClusterHandle])
 						{
 							Handles.Add(Child);
 						}
diff --git a/Engine/Source/Runtime/Experimental/Chaos/Private/PhysicsProxy/PerSolverFieldSystem.cpp b/Engine/Source/Runtime/Experimental/Chaos/Private/PhysicsProxy/PerSolverFieldSystem.cpp
--- a/Engine/Source/Runtime/Experimental/Chaos/Private/PhysicsProxy/PerSolverFieldSystem.cpp
+++ b/Engine/Source/Runtime/Experimental/Chaos/Private/PhysicsProxy/PerSolverFieldSystem.cpp
@@ -211,10 +211,10 @@ void FPerSolverFieldSystem::GetRelevantParticleHandles(
 			const auto* Clustered = Handle->CastToClustered();
 			if (Clustered && Clustered->ClusterIds().NumChildren)
 			{
-				Chaos::TPBDRigidParticleHandle<float, 3>* RigidHandle = (*It).Handle()->CastToRigidParticle();
-				if (ClusterMap.Contains(RigidHandle))
+				Chaos::TPBDRigidClusteredParticleHandle<float, 3>* ClusterHandle = (*It).Handle()->CastToClustered();
+				if (ClusterMap.Contains(ClusterHandle))
 				{
-					for (Chaos::TPBDRigidParticleHandle<float, 3> * Child : ClusterMap[RigidHandle])
+					for (Chaos::TPBDRigidParticleHandle<float, 3> * Child : ClusterMap[ClusterHandle])
 					{
 						Handles.Add(Child);
 					}
diff --git a/Engine/Source/Runtime/Experimental/Chaos/Public/Chaos/PBDRigidClustering.h b/Engine/Source/Runtime/Experimental/Chaos/Public/Chaos/PBDRigidClustering.h
--- a/Engine/Source/Runtime/Experimental/Chaos/Public/Chaos/PBDRigidClustering.h
+++ b/Engine/Source/Runtime/Experimental/Chaos/Public/Chaos/PBDRigidClustering.h
@@ -42,7 +42,15 @@ class CHAOS_API TPBDRigidClustering
 	typedef typename T_FPBDCollisionConstraint::FPointContactConstraint FPointContactConstraint;
 public:
 	/** Parent to children */
-	typedef TMap<FPBDRigidParticleHandle*, TArray<FPBDRigidParticleHandle*> > FClusterMap;
+
+	typedef FPBDRigidParticleHandle*					FRigidHandle;
+	typedef TArray<FRigidHandle>						FRigidHandleArray;
+
+	typedef FPBDRigidClusteredParticleHandle*			FClusterHandle;
+
+	typedef TPair<FClusterHandle, FRigidHandleArray>	FClusterMapEntry;
+	typedef TMap<FClusterHandle, FRigidHandleArray>	    FClusterMap;
+
 
 	using FCollisionConstraintHandle = FPBDCollisionConstraintHandle;
 
@@ -184,8 +192,8 @@ public:
 	*    active id, see the GetActiveClusterIndex to find the active cluster.
 	*    INDEX_NONE represents a non-clustered body.
 	*/
-	TArrayCollectionArray<ClusterId>&       GetClusterIdsArray() { return MParticles.ClusterIdsArray(); }
-	const TArrayCollectionArray<ClusterId>& GetClusterIdsArray() const { return MParticles.ClusterIdsArray(); }
+	TArrayCollectionArray<ClusterId>&					GetClusterIdsArray() { return MParticles.ClusterIdsArray(); }
+	const TArrayCollectionArray<ClusterId>&             GetClusterIdsArray() const { return MParticles.ClusterIdsArray(); }
 
 	/*
 	*  GetInternalClusterArray
diff --git a/Engine/Source/Runtime/Experimental/Chaos/Public/Chaos/PBDRigidsEvolution.h b/Engine/Source/Runtime/Experimental/Chaos/Public/Chaos/PBDRigidsEvolution.h
--- a/Engine/Source/Runtime/Experimental/Chaos/Public/Chaos/PBDRigidsEvolution.h
+++ b/Engine/Source/Runtime/Experimental/Chaos/Public/Chaos/PBDRigidsEvolution.h
@@ -499,8 +499,6 @@ public:
 		}
 	}
 
-	//TEMP: this is only needed while clustering continues to use indices directly
-	const auto& GetActiveClusteredArray() const { return Particles.GetActiveClusteredArray(); }
 	const auto& GetNonDisabledClusteredArray() const { return Particles.GetNonDisabledClusteredArray(); }
 
 	CHAOS_API TSerializablePtr<FChaosPhysicsMaterial> GetPhysicsMaterial(const FGeometryParticleHandle* Particle) const { return Particle->AuxilaryValue(PhysicsMaterials); }
diff --git a/Engine/Source/Runtime/Experimental/Chaos/Public/Chaos/PBDRigidsEvolutionGBF.h b/Engine/Source/Runtime/Experimental/Chaos/Public/Chaos/PBDRigidsEvolutionGBF.h
--- a/Engine/Source/Runtime/Experimental/Chaos/Public/Chaos/PBDRigidsEvolutionGBF.h
+++ b/Engine/Source/Runtime/Experimental/Chaos/Public/Chaos/PBDRigidsEvolutionGBF.h
@@ -61,7 +61,6 @@ namespace Chaos
 		using Base::CreateClusteredParticles;
 		using Base::EnableParticle;
 		using Base::DisableParticles;
-		using Base::GetActiveClusteredArray;
 		using Base::NumIslands;
 		using Base::GetNonDisabledClusteredArray;
 		using Base::DisableParticle;
diff --git a/Engine/Source/Runtime/Experimental/Chaos/Public/Chaos/PBDRigidsSOAs.h b/Engine/Source/Runtime/Experimental/Chaos/Public/Chaos/PBDRigidsSOAs.h
--- a/Engine/Source/Runtime/Experimental/Chaos/Public/Chaos/PBDRigidsSOAs.h
+++ b/Engine/Source/Runtime/Experimental/Chaos/Public/Chaos/PBDRigidsSOAs.h
@@ -192,7 +192,6 @@ public:
 		if (!Params.bStartSleeping)
 		{
 			InsertToMapAndArray(reinterpret_cast<TArray<TPBDRigidParticleHandle<T,d>*>&>(NewClustered), ActiveParticlesToIndex, ActiveParticlesArray);
-			InsertToMapAndArray(NewClustered, ActiveClusteredToIndex, ActiveClusteredArray);
 		}
 
 		UpdateViews();
@@ -233,12 +232,6 @@ public:
 				{
 					RemoveFromMapAndArray(PBDRigidClustered,
 						NonDisabledClusteredToIndex, NonDisabledClusteredArray);
-
-					if (Particle->ObjectState() == EObjectStateType::Dynamic)
-					{
-						RemoveFromMapAndArray(PBDRigidClustered,
-							ActiveClusteredToIndex, ActiveClusteredArray);
-					}
 				}
 			}
 			else
@@ -289,11 +282,6 @@ public:
 				{
 					RemoveFromMapAndArray(PBDRigidClustered, 
 						NonDisabledClusteredToIndex, NonDisabledClusteredArray);
-					if (Particle->ObjectState() == EObjectStateType::Dynamic)
-					{
-						RemoveFromMapAndArray(PBDRigidClustered, 
-							ActiveClusteredToIndex, ActiveClusteredArray);
-					}
 				}
 			}
 			else
@@ -339,12 +327,6 @@ public:
 				{
 					InsertToMapAndArray(PBDRigidClustered, 
 						NonDisabledClusteredToIndex, NonDisabledClusteredArray);
-					if (!PBDRigid->Sleeping() && Particle->ObjectState() == EObjectStateType::Dynamic)
-					{
-						// Clustered, enabled, (dynamic, !sleeping)
-						InsertToMapAndArray(PBDRigidClustered, 
-							ActiveClusteredToIndex, ActiveClusteredArray);
-					}
 				}
 			}
 			else
@@ -393,12 +375,6 @@ public:
 							bGeometryCollectionDirty = true;
 							return;
 						}
-						else
-						{
-							// Clustered, non geometry collection:
-							InsertToMapAndArray(PBDRigidClustered, 
-								ActiveClusteredToIndex, ActiveClusteredArray);
-						}
 					}
 					else
 					{
@@ -462,11 +438,6 @@ public:
 						{
 							bGeometryCollectionDirty = true;
 						}
-						else
-						{
-							RemoveFromMapAndArray(PBDRigidClustered, 
-								ActiveClusteredToIndex, ActiveClusteredArray);
-						}
 					}
 					else
 					{
@@ -580,7 +551,6 @@ public:
 		Ar << GeometryCollectionParticles;
 
 		SerializeMapAndArray(Ar, ActiveParticlesToIndex, ActiveParticlesArray);
-		//SerializeMapAndArray(Ar, ActiveClusteredToIndex, ActiveClusteredArray);
 		//SerializeMapAndArray(Ar, NonDisabledClusteredToIndex, NonDisabledClusteredArray);
 
 		//todo: update deterministic ID
@@ -745,8 +715,6 @@ public:
 		bGeometryCollectionDirty = false;
 	}
 
-	//TEMP: only needed while clustering code continues to use direct indices
-	const auto& GetActiveClusteredArray() const { return ActiveClusteredArray; }
 	const auto& GetNonDisabledClusteredArray() const { return NonDisabledClusteredArray; }
 
 	const auto& GetClusteredParticles() const { return *ClusteredParticles; }
@@ -976,8 +944,6 @@ private:
 	//Utility structures for maintaining an Active particles view
 	TMap<TPBDRigidParticleHandle<T, d>*, int32> ActiveParticlesToIndex;
 	TArray<TPBDRigidParticleHandle<T, d>*> ActiveParticlesArray;
-	TMap<TPBDRigidClusteredParticleHandle<T, d>*, int32> ActiveClusteredToIndex;
-	TArray<TPBDRigidClusteredParticleHandle<T, d>*> ActiveClusteredArray;
 	TArray<TPBDRigidParticleHandle<T,d>*> TransientDirtyArray;
 	TMap<TPBDRigidParticleHandle<T,d>*,int32> TransientDirtyToIndex;
 
diff --git a/Engine/Source/Runtime/Experimental/GeometryCollectionEngine/Private/GeometryCollection/GeometryCollectionComponent.cpp b/Engine/Source/Runtime/Experimental/GeometryCollectionEngine/Private/GeometryCollection/GeometryCollectionComponent.cpp
--- a/Engine/Source/Runtime/Experimental/GeometryCollectionEngine/Private/GeometryCollection/GeometryCollectionComponent.cpp
+++ b/Engine/Source/Runtime/Experimental/GeometryCollectionEngine/Private/GeometryCollection/GeometryCollectionComponent.cpp
@@ -1070,11 +1070,11 @@ void UGeometryCollectionComponent::UpdateRepData()
 	}
 }
 
-void SetHierarchyStrain(Chaos::TPBDRigidClusteredParticleHandle<Chaos::FReal, 3>* P, TMap<Chaos::TPBDRigidParticleHandle<Chaos::FReal, 3>*, TArray<Chaos::TPBDRigidParticleHandle<Chaos::FReal, 3>*>>& Map, float Strain)
+void SetHierarchyStrain(Chaos::TPBDRigidClusteredParticleHandle<Chaos::FReal, 3>* P, TMap<Chaos::TPBDRigidClusteredParticleHandle<Chaos::FReal, 3>*, TArray<Chaos::TPBDRigidParticleHandle<Chaos::FReal, 3>*>>& Map, float Strain)
 {
 	TArray<Chaos::TPBDRigidParticleHandle<Chaos::FReal, 3>*>* Children = Map.Find(P);
 
-	if(Children)
+	if(Children)	
 	{
 		for(Chaos::TPBDRigidParticleHandle<Chaos::FReal, 3> * ChildP : (*Children))
 		{
