Chaos Physics - Add the ability to disable/enable rigid bodies from a skeletal simulation

diff --git a/Engine/Source/Runtime/Core/Public/UObject/UE5MainStreamObjectVersion.h b/Engine/Source/Runtime/Core/Public/UObject/UE5MainStreamObjectVersion.h
--- a/Engine/Source/Runtime/Core/Public/UObject/UE5MainStreamObjectVersion.h
+++ b/Engine/Source/Runtime/Core/Public/UObject/UE5MainStreamObjectVersion.h
@@ -94,6 +94,9 @@ struct CORE_API FUE5MainStreamObjectVersion
 		// Serialize HLOD info in WorldPartitionActorDesc
 		WorldPartitionActorDescSerializeHLODInfo,
 
+		// Expose particle Disabled flag to the game thread
+		AddDisabledFlag,
+
 		// -----<new versions can be added above this line>-------------------------------------------------
 		VersionPlusOne,
 		LatestVersion = VersionPlusOne - 1
diff --git a/Engine/Source/Runtime/Engine/Classes/Components/SkeletalMeshComponent.h b/Engine/Source/Runtime/Engine/Classes/Components/SkeletalMeshComponent.h
--- a/Engine/Source/Runtime/Engine/Classes/Components/SkeletalMeshComponent.h
+++ b/Engine/Source/Runtime/Engine/Classes/Components/SkeletalMeshComponent.h
@@ -1957,7 +1957,6 @@ public:
 	/** Find the root body index*/
 	int32 FindRootBodyIndex() const;
 
-
 	/** Terminate physics on all bodies below the named bone, effectively disabling collision forever. If you terminate, you won't be able to re-init later. */
 	UFUNCTION(BlueprintCallable, Category="Components|SkeletalMesh")
 	void TermBodiesBelow(FName ParentBoneName);
@@ -1997,6 +1996,14 @@ public:
 	UFUNCTION(BlueprintCallable, Category="Components|SkeletalMesh")
 	void SetEnablePhysicsBlending(bool bNewBlendPhysics);
 
+	/** [WARNING: Chaos Only] 
+	* Set all of the bones below passed in bone to be disabled or not for the associated physics solver 
+	* Bodies will not be colliding or be part of the physics simulation. 
+	* This is different from SetAllBodiesBelowSimulatePhysics that changes bodies to Kinematic/simulated 	
+	*/
+	UFUNCTION(BlueprintCallable, Category = "Components|SkeletalMesh")
+	void SetAllBodiesBelowPhysicsDisabled(const FName& InBoneName, bool bDisabled, bool bIncludeSelf = true);
+
 	/** Set all of the bones below passed in bone to be simulated */
 	UFUNCTION(BlueprintCallable, Category="Components|SkeletalMesh")
 	void SetAllBodiesBelowSimulatePhysics(const FName& InBoneName, bool bNewSimulate, bool bIncludeSelf = true );
diff --git a/Engine/Source/Runtime/Engine/Classes/PhysicsEngine/BodyInstance.h b/Engine/Source/Runtime/Engine/Classes/PhysicsEngine/BodyInstance.h
--- a/Engine/Source/Runtime/Engine/Classes/PhysicsEngine/BodyInstance.h
+++ b/Engine/Source/Runtime/Engine/Classes/PhysicsEngine/BodyInstance.h
@@ -777,6 +777,13 @@ public:
 	/** Enable/disable Continuous Collidion Detection feature */
 	void SetUseCCD(bool bInUseCCD);
 
+	/** Disable/Re-Enable this body in the solver,  when disable, the body won't be part of the simulation ( regardless if it's dynamic or kinematic ) and no collision will occur 
+	* this can be used for performance control situation for example
+	*/
+	void SetPhysicsDisabled(bool bSetDisabled);
+
+	bool IsPhysicsDisabled() const;
+
 private:
 
 	/** Struct of body instance delegates that are rarely bound so that we can only allocate memory if one is actually being used. */
diff --git a/Engine/Source/Runtime/Engine/Private/PhysicsEngine/BodyInstance.cpp b/Engine/Source/Runtime/Engine/Private/PhysicsEngine/BodyInstance.cpp
--- a/Engine/Source/Runtime/Engine/Private/PhysicsEngine/BodyInstance.cpp
+++ b/Engine/Source/Runtime/Engine/Private/PhysicsEngine/BodyInstance.cpp
@@ -3568,6 +3568,26 @@ void FBodyInstance::SetUseCCD(bool bInUseCCD)
 	}
 }
 
+void FBodyInstance::SetPhysicsDisabled(bool bSetDisabled)
+{
+	FPhysicsCommand::ExecuteWrite(ActorHandle, [&](const FPhysicsActorHandle& Actor)
+	{
+		if (FPhysicsInterface::IsValid(Actor) && FPhysicsInterface::IsRigidBody(Actor))
+		{
+			FPhysicsInterface::SetDisabled(Actor, bSetDisabled);
+		}
+	});
+}
+
+bool FBodyInstance::IsPhysicsDisabled() const
+{
+	bool bIsDisabled = false;
+	FPhysicsCommand::ExecuteRead(ActorHandle, [&](const FPhysicsActorHandle& Actor)
+	{
+		bIsDisabled = FPhysicsInterface::IsDisabled(Actor);
+	});
+	return bIsDisabled;
+}
 
 void FBodyInstance::AddRadialImpulseToBody(const FVector& Origin, float Radius, float Strength, uint8 Falloff, bool bVelChange)
 {
diff --git a/Engine/Source/Runtime/Engine/Private/PhysicsEngine/PhysAnim.cpp b/Engine/Source/Runtime/Engine/Private/PhysicsEngine/PhysAnim.cpp
--- a/Engine/Source/Runtime/Engine/Private/PhysicsEngine/PhysAnim.cpp
+++ b/Engine/Source/Runtime/Engine/Private/PhysicsEngine/PhysAnim.cpp
@@ -235,6 +235,12 @@ void USkeletalMeshComponent::PerformBlendPhysicsBones(const TArray<FBoneIndexTyp
 
 					float UsePhysWeight = (bBlendPhysics)? 1.f : PhysicsAssetBodyInstance->PhysicsBlendWeight;
 
+					// if the body instance is disabled, then we want to use the animation transform and ignore the physics one
+					if (PhysicsAssetBodyInstance->IsPhysicsDisabled())
+					{
+						UsePhysWeight = 0.0f;
+					}
+
 					// Find this bones parent matrix.
 					FTransform ParentWorldTM;
 
diff --git a/Engine/Source/Runtime/Engine/Private/PhysicsEngine/PhysicsInterfacePhysX.cpp b/Engine/Source/Runtime/Engine/Private/PhysicsEngine/PhysicsInterfacePhysX.cpp
--- a/Engine/Source/Runtime/Engine/Private/PhysicsEngine/PhysicsInterfacePhysX.cpp
+++ b/Engine/Source/Runtime/Engine/Private/PhysicsEngine/PhysicsInterfacePhysX.cpp
@@ -1170,6 +1170,19 @@ bool FPhysicsInterface_PhysX::IsCcdEnabled(const FPhysicsActorHandle_PhysX& InAc
 	return GetRigidBodyFlag(InActorHandle, PxRigidBodyFlag::eENABLE_CCD);
 }
 
+bool FPhysicsInterface_PhysX::IsDisabled(const FPhysicsActorHandle_PhysX& InHandle)
+{
+	// @todo(chaos) todo - see SetDisabled() bewlow
+	return false;
+}
+
+void FPhysicsInterface_PhysX::SetDisabled(const FPhysicsActorHandle_PhysX& InHandle, bool bSetDisabled)
+{
+	// @todo(chaos) todo - the closest way may be to set the PxRigidBodyFlag::eDISABLE_SIMULATION, but this may have implication on what you can or not call on the body after that 
+	ensureMsgf(false, TEXT("SetDisabled is unsupported by the PhysX interface"));
+}
+
+
 bool FPhysicsInterface_PhysX::IsInScene(const FPhysicsActorHandle_PhysX& InActorHandle)
 {
 	if(PxRigidActor* Actor = GetPxRigidActor_AssumesLocked(InActorHandle))
diff --git a/Engine/Source/Runtime/Engine/Private/SkeletalMeshComponentPhysics.cpp b/Engine/Source/Runtime/Engine/Private/SkeletalMeshComponentPhysics.cpp
--- a/Engine/Source/Runtime/Engine/Private/SkeletalMeshComponentPhysics.cpp
+++ b/Engine/Source/Runtime/Engine/Private/SkeletalMeshComponentPhysics.cpp
@@ -1093,6 +1093,13 @@ void USkeletalMeshComponent::SetAllBodiesNotifyRigidBodyCollision(bool bNewNotif
 	}
 }
 
+void USkeletalMeshComponent::SetAllBodiesBelowPhysicsDisabled(const FName& InBoneName, bool bDisabled, bool bIncludeSelf)
+{
+	int32 NumBodiesFound = ForEachBodyBelow(InBoneName, bIncludeSelf, /*bSkipCustomPhysicsType=*/ false, [bDisabled](FBodyInstance* BI)
+	{
+		BI->SetPhysicsDisabled(bDisabled);
+	});
+}
 
 void USkeletalMeshComponent::SetAllBodiesBelowSimulatePhysics( const FName& InBoneName, bool bNewSimulate, bool bIncludeSelf )
 {
diff --git a/Engine/Source/Runtime/Engine/Public/Physics/PhysicsInterfacePhysX.h b/Engine/Source/Runtime/Engine/Public/Physics/PhysicsInterfacePhysX.h
--- a/Engine/Source/Runtime/Engine/Public/Physics/PhysicsInterfacePhysX.h
+++ b/Engine/Source/Runtime/Engine/Public/Physics/PhysicsInterfacePhysX.h
@@ -255,7 +255,9 @@ struct ENGINE_API FPhysicsInterface_PhysX : public FGenericPhysicsInterface, pub
 	static bool IsKinematic_AssumesLocked(const FPhysicsActorHandle_PhysX& InHandle);
 	static bool IsSleeping(const FPhysicsActorHandle_PhysX& InHandle);
 	static bool IsCcdEnabled(const FPhysicsActorHandle_PhysX& InHandle);
-	static bool IsInScene(const FPhysicsActorHandle_PhysX& InHandle);
+	static bool IsCcdEnabled(const FPhysicsActorHandle_PhysX& InHandle);
+	static bool IsDisabled(const FPhysicsActorHandle_PhysX& InHandle);
+	static void SetDisabled(const FPhysicsActorHandle_PhysX& InHandle, bool bSetDisabled);
 	static FPhysScene* GetCurrentScene(const FPhysicsActorHandle_PhysX& InHandle);
 	static bool CanSimulate_AssumesLocked(const FPhysicsActorHandle_PhysX& InHandle);
 	static float GetMass_AssumesLocked(const FPhysicsActorHandle_PhysX& InHandle);
diff --git a/Engine/Source/Runtime/Experimental/Chaos/Private/PhysicsProxy/SingleParticlePhysicsProxy.cpp b/Engine/Source/Runtime/Experimental/Chaos/Private/PhysicsProxy/SingleParticlePhysicsProxy.cpp
--- a/Engine/Source/Runtime/Experimental/Chaos/Private/PhysicsProxy/SingleParticlePhysicsProxy.cpp
+++ b/Engine/Source/Runtime/Experimental/Chaos/Private/PhysicsProxy/SingleParticlePhysicsProxy.cpp
@@ -111,6 +111,19 @@ void PushToPhysicsStateImp(const Chaos::FDirtyPropertiesManager& Manager, Chaos:
 			{
 				Evolution.SetParticleObjectState(RigidHandle,NewData->ObjectState());
 
+				if (RigidHandle->Disabled() != NewData->Disabled())
+				{
+					if (NewData->Disabled())
+					{
+						Evolution.DisableParticle(Handle);
+					}
+					else
+					{
+						Evolution.EnableParticle(Handle, nullptr);
+					}
+				}
+				Evolution.SetParticleObjectState(RigidHandle, NewData->ObjectState());
+
 				RigidHandle->SetDynamicMisc(*NewData);
 			}
 		}
diff --git a/Engine/Source/Runtime/Experimental/Chaos/Public/Chaos/PBDRigidsEvolution.h b/Engine/Source/Runtime/Experimental/Chaos/Public/Chaos/PBDRigidsEvolution.h
--- a/Engine/Source/Runtime/Experimental/Chaos/Public/Chaos/PBDRigidsEvolution.h
+++ b/Engine/Source/Runtime/Experimental/Chaos/Public/Chaos/PBDRigidsEvolution.h
@@ -327,8 +327,9 @@ public:
 	{
 		RemoveParticleFromAccelerationStructure(*Particle);
 		Particles.DisableParticle(Particle);
-		ConstraintGraph.DisableParticle(Particle);
+		// disable constraint before particles has particles need to be in the constraint graph for constraints to be removed
 		DisableConstraints(TSet<FGeometryParticleHandle*>({ Particle }));
+		ConstraintGraph.DisableParticle(Particle);
 	}
 
 	CHAOS_API void FlushExternalAccelerationQueue(FAccelerationStructure& Acceleration,FPendingSpatialDataQueue& ExternalQueue);
diff --git a/Engine/Source/Runtime/Experimental/Chaos/Public/Chaos/ParticleDirtyFlags.h b/Engine/Source/Runtime/Experimental/Chaos/Public/Chaos/ParticleDirtyFlags.h
--- a/Engine/Source/Runtime/Experimental/Chaos/Public/Chaos/ParticleDirtyFlags.h
+++ b/Engine/Source/Runtime/Experimental/Chaos/Public/Chaos/ParticleDirtyFlags.h
@@ -274,6 +274,11 @@ public:
 		{
 			Ar << MCollisionConstraintFlag;
 		}
+
+		if (Ar.CustomVer(FUE5MainStreamObjectVersion::GUID) >= FUE5MainStreamObjectVersion::AddDisabledFlag)
+		{
+			Ar << bDisabled;
+		}
 	}
 
 	template <typename TOther>
@@ -288,6 +293,7 @@ public:
 		SetOneWayInteraction(Other.OneWayInteraction());
 		SetCollisionConstraintFlag(Other.CollisionConstraintFlag());
 		SetCCDEnabled(Other.CCDEnabled());
+		SetDisabled(Other.Disabled());
 	}
 
 	template <typename TOther>
@@ -301,7 +307,8 @@ public:
 			&& ResimType() == Other.ResimType()
 			&& OneWayInteraction() == Other.OneWayInteraction() 
 			&& CollisionConstraintFlag() == Other.CollisionConstraintFlag()
-			&& CCDEnabled() == Other.CCDEnabled();
+			&& CCDEnabled() == Other.CCDEnabled()
+			&& Disabled() == Other.Disabled();
 	}
 
 	bool operator==(const FParticleDynamicMisc& Other) const
@@ -324,6 +331,9 @@ public:
 	bool CCDEnabled() const { return bCCDEnabled; }
 	void SetCCDEnabled(bool bInCCDEnabled) { bCCDEnabled = bInCCDEnabled; }
 
+	bool Disabled() const { return bDisabled; }
+	void SetDisabled(bool bInDisabled) { bDisabled = bInDisabled; }
+
 	int32 CollisionGroup() const { return MCollisionGroup; }
 	void SetCollisionGroup(int32 InGroup){ MCollisionGroup = InGroup; }
 
@@ -348,6 +358,7 @@ private:
 	uint32 MCollisionConstraintFlag = 0;
 
 	bool bCCDEnabled;
+	bool bDisabled;
 };
 
 inline FChaosArchive& operator<<(FChaosArchive& Ar,FParticleDynamicMisc& Data)
diff --git a/Engine/Source/Runtime/Experimental/Chaos/Public/Chaos/ParticleHandle.h b/Engine/Source/Runtime/Experimental/Chaos/Public/Chaos/ParticleHandle.h
--- a/Engine/Source/Runtime/Experimental/Chaos/Public/Chaos/ParticleHandle.h
+++ b/Engine/Source/Runtime/Experimental/Chaos/Public/Chaos/ParticleHandle.h
@@ -805,6 +805,7 @@ public:
 		SetCollisionGroup(DynamicMisc.CollisionGroup());
 		SetGravityEnabled(DynamicMisc.GravityEnabled());
 		SetCCDEnabled(DynamicMisc.CCDEnabled());
+		SetDisabled(DynamicMisc.Disabled());
 		SetResimType(DynamicMisc.ResimType());
 		SetOneWayInteraction(DynamicMisc.OneWayInteraction());
 		AddCollisionConstraintFlag( (Chaos::ECollisionConstraintFlags)DynamicMisc.CollisionConstraintFlag() );
@@ -2113,6 +2114,13 @@ public:
 		MMiscData.Modify(true, MDirtyFlags, Proxy, [bInCCDEnabled](auto& Data) {Data.SetCCDEnabled(bInCCDEnabled); });
 	}
 
+	bool Disabled() const { return MMiscData.Read().Disabled(); }
+
+	void SetDisabled(bool bInDisabled)
+	{
+		MMiscData.Modify(true, MDirtyFlags, Proxy, [bInDisabled](auto& Data) {Data.SetDisabled(bInDisabled); });
+	}
+
 	//todo: remove this
 	bool IsInitialized() const { return MInitialized; }
 	void SetInitialized(const bool InInitialized)
diff --git a/Engine/Source/Runtime/Experimental/Chaos/Public/PhysicsProxy/SingleParticlePhysicsProxy.h b/Engine/Source/Runtime/Experimental/Chaos/Public/PhysicsProxy/SingleParticlePhysicsProxy.h
--- a/Engine/Source/Runtime/Experimental/Chaos/Public/PhysicsProxy/SingleParticlePhysicsProxy.h
+++ b/Engine/Source/Runtime/Experimental/Chaos/Public/PhysicsProxy/SingleParticlePhysicsProxy.h
@@ -896,6 +896,26 @@ public:
 		return false;
 	}
 
+	void SetDisabled(bool bDisable)
+	{
+		VerifyContext();
+		if (auto Rigid = GetParticle_LowLevel()->CastToRigidParticle())
+		{
+			Rigid->SetDisabled(bDisable);
+		}
+	}
+
+	bool Disabled() const
+	{
+		VerifyContext();
+		if (auto Rigid = GetParticle_LowLevel()->CastToRigidParticle())
+		{
+			return Rigid->Disabled();
+		}
+
+		return false;
+	}
+
 };
 
 static_assert(sizeof(FRigidBodyHandle_External) == sizeof(FSingleParticlePhysicsProxy), "Derived types only used to constrain API, all data lives in base class ");
diff --git a/Engine/Source/Runtime/PhysicsCore/Private/ChaosEngineInterface.cpp b/Engine/Source/Runtime/PhysicsCore/Private/ChaosEngineInterface.cpp
--- a/Engine/Source/Runtime/PhysicsCore/Private/ChaosEngineInterface.cpp
+++ b/Engine/Source/Runtime/PhysicsCore/Private/ChaosEngineInterface.cpp
@@ -341,6 +341,16 @@ void FChaosEngineInterface::RemoveDisabledCollisionsFor_AssumesLocked(TArray< FP
 	}
 }
 
+void FChaosEngineInterface::SetDisabled(const FPhysicsActorHandle& InPhysicsActor, bool bSetDisabled)
+{
+	InPhysicsActor->GetGameThreadAPI().SetDisabled(bSetDisabled);
+}
+
+bool FChaosEngineInterface::IsDisabled(const FPhysicsActorHandle& InPhysicsActor)
+{
+	return InPhysicsActor->GetGameThreadAPI().Disabled();
+}
+
 void FChaosEngineInterface::SetActorUserData_AssumesLocked(FPhysicsActorHandle& InActorReference,FPhysicsUserData* InUserData)
 {
 	InActorReference->GetGameThreadAPI().SetUserData(InUserData);
diff --git a/Engine/Source/Runtime/PhysicsCore/Public/Chaos/ChaosEngineInterface.h b/Engine/Source/Runtime/PhysicsCore/Public/Chaos/ChaosEngineInterface.h
--- a/Engine/Source/Runtime/PhysicsCore/Public/Chaos/ChaosEngineInterface.h
+++ b/Engine/Source/Runtime/PhysicsCore/Public/Chaos/ChaosEngineInterface.h
@@ -370,6 +370,9 @@ public:
 	static void AddDisabledCollisionsFor_AssumesLocked(const TMap<FPhysicsActorHandle, TArray< FPhysicsActorHandle > >& InMap);
 	static void RemoveDisabledCollisionsFor_AssumesLocked(TArray< FPhysicsActorHandle > & InPhysicsActors);
 
+	static void SetDisabled(const FPhysicsActorHandle& InPhysicsActor, bool bSetDisabled);
+	static bool IsDisabled(const FPhysicsActorHandle& InPhysicsActor);
+
 	static void SetActorUserData_AssumesLocked(FPhysicsActorHandle& InActorReference,FPhysicsUserData* InUserData);
 
 	static bool IsRigidBody(const FPhysicsActorHandle& InActorReference);
