Chaos Physics - temporarily disable the sync of Disabled particle flag in PushToPhysicsStateImp as this is currently having some side effects with clustering code
Will provide a better fix soon

diff --git a/Engine/Source/Runtime/Experimental/Chaos/Private/PhysicsProxy/SingleParticlePhysicsProxy.cpp b/Engine/Source/Runtime/Experimental/Chaos/Private/PhysicsProxy/SingleParticlePhysicsProxy.cpp
--- a/Engine/Source/Runtime/Experimental/Chaos/Private/PhysicsProxy/SingleParticlePhysicsProxy.cpp
+++ b/Engine/Source/Runtime/Experimental/Chaos/Private/PhysicsProxy/SingleParticlePhysicsProxy.cpp
@@ -111,7 +111,8 @@ void PushToPhysicsStateImp(const Chaos::FDirtyPropertiesManager& Manager, Chaos:
 			{
 				Evolution.SetParticleObjectState(RigidHandle,NewData->ObjectState());
 
-				if (RigidHandle->Disabled() != NewData->Disabled())
+				// @todo(chaos) re-enable this below when the GT/PT two way sync for this flag is fixed - see similar comment in ParticleHandle.h
+				/*if (RigidHandle->Disabled() != NewData->Disabled())
 				{
 					if (NewData->Disabled())
 					{
@@ -121,7 +122,7 @@ void PushToPhysicsStateImp(const Chaos::FDirtyPropertiesManager& Manager, Chaos:
 					{
 						Evolution.EnableParticle(Handle, nullptr);
 					}
-				}
+				}*/
 				Evolution.SetParticleObjectState(RigidHandle, NewData->ObjectState());
 
 				RigidHandle->SetDynamicMisc(*NewData);
diff --git a/Engine/Source/Runtime/Experimental/Chaos/Public/Chaos/ParticleHandle.h b/Engine/Source/Runtime/Experimental/Chaos/Public/Chaos/ParticleHandle.h
--- a/Engine/Source/Runtime/Experimental/Chaos/Public/Chaos/ParticleHandle.h
+++ b/Engine/Source/Runtime/Experimental/Chaos/Public/Chaos/ParticleHandle.h
@@ -805,7 +805,8 @@ public:
 		SetCollisionGroup(DynamicMisc.CollisionGroup());
 		SetGravityEnabled(DynamicMisc.GravityEnabled());
 		SetCCDEnabled(DynamicMisc.CCDEnabled());
-		SetDisabled(DynamicMisc.Disabled());
+		// @todo(chaos) re-enable this below when the GT/PT two way sync for this flag is fixed - see similar comment in SingleParticlePhysicsProxy.h
+		//SetDisabled(DynamicMisc.Disabled()); 
 		SetResimType(DynamicMisc.ResimType());
 		SetOneWayInteraction(DynamicMisc.OneWayInteraction());
 		AddCollisionConstraintFlag( (Chaos::ECollisionConstraintFlags)DynamicMisc.CollisionConstraintFlag() );
