Fix uninitialized variable that was causing the sync of Disabled particle in PushToPhysicsStateImp to have side effect with clustering unit tests
revert changes from 15531322 and make sure SetDisabled is called on the rigid particle

diff --git a/Engine/Source/Runtime/Experimental/Chaos/Private/PhysicsProxy/SingleParticlePhysicsProxy.cpp b/Engine/Source/Runtime/Experimental/Chaos/Private/PhysicsProxy/SingleParticlePhysicsProxy.cpp
--- a/Engine/Source/Runtime/Experimental/Chaos/Private/PhysicsProxy/SingleParticlePhysicsProxy.cpp
+++ b/Engine/Source/Runtime/Experimental/Chaos/Private/PhysicsProxy/SingleParticlePhysicsProxy.cpp
@@ -112,7 +112,7 @@ void PushToPhysicsStateImp(const Chaos::FDirtyPropertiesManager& Manager, Chaos:
 				Evolution.SetParticleObjectState(RigidHandle,NewData->ObjectState());
 
 				// @todo(chaos) re-enable this below when the GT/PT two way sync for this flag is fixed - see similar comment in ParticleHandle.h
-				/*if (RigidHandle->Disabled() != NewData->Disabled())
+				if (RigidHandle->Disabled() != NewData->Disabled())
 				{
 					if (NewData->Disabled())
 					{
@@ -122,7 +122,7 @@ void PushToPhysicsStateImp(const Chaos::FDirtyPropertiesManager& Manager, Chaos:
 					{
 						Evolution.EnableParticle(Handle, nullptr);
 					}
-				}*/
+				}
 				Evolution.SetParticleObjectState(RigidHandle, NewData->ObjectState());
 
 				RigidHandle->SetDynamicMisc(*NewData);
diff --git a/Engine/Source/Runtime/Experimental/Chaos/Public/Chaos/ParticleHandle.h b/Engine/Source/Runtime/Experimental/Chaos/Public/Chaos/ParticleHandle.h
--- a/Engine/Source/Runtime/Experimental/Chaos/Public/Chaos/ParticleHandle.h
+++ b/Engine/Source/Runtime/Experimental/Chaos/Public/Chaos/ParticleHandle.h
@@ -80,6 +80,7 @@ void PBDRigidParticleDefaultConstruct(FConcrete& Concrete, const TPBDRigidPartic
 	Concrete.SetAngularEtherDrag(0.f);
 	Concrete.SetGravityEnabled(Params.bGravityEnabled);
 	Concrete.SetCCDEnabled(Params.bCCDEnabled);
+	Concrete.SetDisabled(Params.bDisabled);
 	Concrete.SetResimType(EResimType::FullResim);
 }
 
@@ -805,8 +806,7 @@ public:
 		SetCollisionGroup(DynamicMisc.CollisionGroup());
 		SetGravityEnabled(DynamicMisc.GravityEnabled());
 		SetCCDEnabled(DynamicMisc.CCDEnabled());
-		// @todo(chaos) re-enable this below when the GT/PT two way sync for this flag is fixed - see similar comment in SingleParticlePhysicsProxy.h
-		//SetDisabled(DynamicMisc.Disabled()); 
+		SetDisabled(DynamicMisc.Disabled()); 
 		SetResimType(DynamicMisc.ResimType());
 		SetOneWayInteraction(DynamicMisc.OneWayInteraction());
 		AddCollisionConstraintFlag( (Chaos::ECollisionConstraintFlags)DynamicMisc.CollisionConstraintFlag() );
