remove comment that shoudl have been removed in 15533556

diff --git a/Engine/Source/Runtime/Experimental/Chaos/Private/PhysicsProxy/SingleParticlePhysicsProxy.cpp b/Engine/Source/Runtime/Experimental/Chaos/Private/PhysicsProxy/SingleParticlePhysicsProxy.cpp
--- a/Engine/Source/Runtime/Experimental/Chaos/Private/PhysicsProxy/SingleParticlePhysicsProxy.cpp
+++ b/Engine/Source/Runtime/Experimental/Chaos/Private/PhysicsProxy/SingleParticlePhysicsProxy.cpp
@@ -111,7 +111,6 @@ void PushToPhysicsStateImp(const Chaos::FDirtyPropertiesManager& Manager, Chaos:
 			{
 				Evolution.SetParticleObjectState(RigidHandle,NewData->ObjectState());
 
-				// @todo(chaos) re-enable this below when the GT/PT two way sync for this flag is fixed - see similar comment in ParticleHandle.h
 				if (RigidHandle->Disabled() != NewData->Disabled())
 				{
 					if (NewData->Disabled())
