Fixed "Simulation Generates Hit Events" failing to work correctly with skeletal meshes. Previously this flag was fixed at the component level so any shape/body in a skeletal mesh would notify a collision even if that body hadn't been set to notify. Now the shape flags are properly queried during the processing of collision events before passing back to the game thread.

diff --git a/Engine/Source/Runtime/Experimental/Chaos/Private/EventDefaults.cpp b/Engine/Source/Runtime/Experimental/Chaos/Private/EventDefaults.cpp
--- a/Engine/Source/Runtime/Experimental/Chaos/Private/EventDefaults.cpp
+++ b/Engine/Source/Runtime/Experimental/Chaos/Private/EventDefaults.cpp
@@ -13,10 +13,10 @@
 #include "PhysicsProxy/GeometryCollectionPhysicsProxy.h"
 #include "PhysicsProxy/PerSolverFieldSystem.h"
 #include "ChaosSolversModule.h"
+#include "Chaos/CollisionFilterData.h"
 
 namespace Chaos
 {
-
 	template <typename Traits>
 	void TEventDefaults<Traits>::RegisterSystemEvents(TEventManager<Traits>& EventManager)
 	{
@@ -98,6 +98,21 @@ namespace Chaos
 										// presently when a rigidbody or kinematic hits static geometry then Body1 is null
 										TKinematicGeometryParticleHandle<float, 3>* Body1 = Particle1->CastToKinematicParticle();
 
+										const FImplicitObject* Implicit0 = Constraint.Manifold.Implicit[0];
+										const FImplicitObject* Implicit1 = Constraint.Manifold.Implicit[1];
+
+										const FPerShapeData* Shape0 = Particle0->GetImplicitShape(Implicit0);
+										const FPerShapeData* Shape1 = Particle0->GetImplicitShape(Implicit0);
+
+										const FCollisionFilterData& Filter0 = Shape0->GetSimData();
+										const FCollisionFilterData& Filter1 = Shape1->GetSimData();
+
+										if(!Filter0.HasFlag(EFilterFlags::ContactNotify) && !Filter1.HasFlag(EFilterFlags::ContactNotify))
+										{
+											// No need to notify - engine didn't request notifications for either shape.
+											continue;
+										}
+
 										if (!Constraint.AccumulatedImpulse.IsZero() && Body0)
 										{
 											if (ensure(!Constraint.GetLocation().ContainsNaN() &&
diff --git a/Engine/Source/Runtime/Experimental/Chaos/Public/Chaos/CollisionFilterData.h b/Engine/Source/Runtime/Experimental/Chaos/Public/Chaos/CollisionFilterData.h
--- a/Engine/Source/Runtime/Experimental/Chaos/Public/Chaos/CollisionFilterData.h
+++ b/Engine/Source/Runtime/Experimental/Chaos/Public/Chaos/CollisionFilterData.h
@@ -4,6 +4,24 @@
 #include "ChaosArchive.h"
 #include "UObject/ExternalPhysicsCustomObjectVersion.h"
 
+// Supported flags for Chaos filter, shares some flags with EPhysXFilterDataFlags
+// Shared flags should be kept in sync until unified. @see EPhysXFilterDataFlags
+// #TODO unify filter builder with this flag list
+
+namespace Chaos
+{
+	enum class EFilterFlags : uint8
+	{
+		SimpleCollision			= 0b00000001,	// The shape is used for simple collision
+		ComplexCollision		= 0b00000010,	// The shape is used for complex (trimesh) collision
+		CCD						= 0b00000100,	// Unused - present for compatibility. CCD handled per-particle in Chaos
+		ContactNotify			= 0b00001000,	// Whether collisions with this shape should be reported back to the game thread
+		StaticShape				= 0b00010000,	// Unused - present for compatibility
+		ModifyContacts			= 0b00100000,	// Unused - present for compatibility, whether to allow contact modification, handled in Chaos callbacks now
+		KinematicKinematicPairs	= 0b01000000,	// Unused - present for compatibility, whether to generate KK pairs, Chaos never generates KK pairs
+	};
+}
+
 struct CHAOS_API FCollisionFilterData
 {
 	uint32 Word0;
@@ -15,6 +33,12 @@ struct CHAOS_API FCollisionFilterData
 	{
 		Word0 = Word1 = Word2 = Word3 = 0;
 	}
+
+	bool HasFlag(Chaos::EFilterFlags InFlag) const
+	{
+		const uint32 FilterFlags = (Word3 & 0xFFFFFF);
+		return FilterFlags & static_cast<uint32>(InFlag);
+	}
 };
 
 inline Chaos::FChaosArchive& operator<<(Chaos::FChaosArchive& Ar, FCollisionFilterData& Filter)
