CIS fix, check for null shapes in filtering - clusters may not have valid shapes

diff --git a/Engine/Source/Runtime/Experimental/Chaos/Private/EventDefaults.cpp b/Engine/Source/Runtime/Experimental/Chaos/Private/EventDefaults.cpp
--- a/Engine/Source/Runtime/Experimental/Chaos/Private/EventDefaults.cpp
+++ b/Engine/Source/Runtime/Experimental/Chaos/Private/EventDefaults.cpp
@@ -102,12 +102,13 @@ namespace Chaos
 										const FImplicitObject* Implicit1 = Constraint.Manifold.Implicit[1];
 
 										const FPerShapeData* Shape0 = Particle0->GetImplicitShape(Implicit0);
-										const FPerShapeData* Shape1 = Particle0->GetImplicitShape(Implicit0);
+										const FPerShapeData* Shape1 = Particle1->GetImplicitShape(Implicit1);
 
-										const FCollisionFilterData& Filter0 = Shape0->GetSimData();
-										const FCollisionFilterData& Filter1 = Shape1->GetSimData();
+										// If we don't have a filter - allow the notify, otherwise obey the filter flag
+										const bool bFilter0Notify = Shape0 ? Shape0->GetSimData().HasFlag(EFilterFlags::ContactNotify) : true;
+										const bool bFilter1Notify = Shape1 ? Shape1->GetSimData().HasFlag(EFilterFlags::ContactNotify) : true;
 
-										if(!Filter0.HasFlag(EFilterFlags::ContactNotify) && !Filter1.HasFlag(EFilterFlags::ContactNotify))
+										if(!bFilter0Notify && !bFilter1Notify)
 										{
 											// No need to notify - engine didn't request notifications for either shape.
 											continue;
