Chaos - collision optimizations
- do not calculate collision data required for static friction until it is needed

diff --git a/Engine/Source/Runtime/Experimental/Chaos/Private/Chaos/ChaosDebugDraw.cpp b/Engine/Source/Runtime/Experimental/Chaos/Private/Chaos/ChaosDebugDraw.cpp
--- a/Engine/Source/Runtime/Experimental/Chaos/Private/Chaos/ChaosDebugDraw.cpp
+++ b/Engine/Source/Runtime/Experimental/Chaos/Private/Chaos/ChaosDebugDraw.cpp
@@ -451,7 +451,6 @@ namespace Chaos
 						const FVec3 PointLocation = PointTransform.TransformPosition(ManifoldPoint.CoMContactPoints[ContactPointOwner]) - ManifoldPoint.ContactPoint.ShapeMargins[ContactPointOwner] * PlaneNormal;
 						const FVec3 PlaneLocation = PlaneTransform.TransformPosition(ManifoldPoint.CoMContactPoints[ContactPlaneOwner]) + ManifoldPoint.ContactPoint.ShapeMargins[ContactPlaneOwner] * PlaneNormal;
 						const FVec3 PointPlaneLocation = PointLocation - FVec3::DotProduct(PointLocation - PlaneLocation, PlaneNormal) * PlaneNormal;
-						const FVec3 OldPointPlaneLocation = PlaneTransform.TransformPosition(ManifoldPoint.PrevCoMContactPoints[ContactPlaneOwner]) + ManifoldPoint.ContactPoint.ShapeMargins[ContactPlaneOwner] * PlaneNormal;
 
 						// Dynamic friction, restitution = red
 						// Static friction, no restitution = green
@@ -478,7 +477,6 @@ namespace Chaos
 						const FVec3 WorldPointLocation = SpaceTransform.TransformPosition(PointLocation);
 						const FVec3 WorldPlaneLocation = SpaceTransform.TransformPosition(PlaneLocation);
 						const FVec3 WorldPointPlaneLocation = SpaceTransform.TransformPosition(PointPlaneLocation);
-						const FVec3 WorldOldPointPlaneLocation = SpaceTransform.TransformPosition(OldPointPlaneLocation);
 						const FVec3 WorldPlaneNormal = SpaceTransform.TransformVectorNoScale(PlaneNormal);
 
 						// Pushout
@@ -500,9 +498,6 @@ namespace Chaos
 							FDebugDrawQueue::GetInstance().DrawDebugLine(WorldPointPlaneLocation, WorldPointPlaneLocation + Settings.DrawScale * Settings.ImpulseScale * ManifoldPoint.NetPushOutImpulseTangent * WorldTangent, Color, false, KINDA_SMALL_NUMBER, Settings.DrawPriority, Settings.LineThickness);
 						}
 
-						// Static friction error
-						FDebugDrawQueue::GetInstance().DrawDebugLine(WorldPointPlaneLocation, WorldOldPointPlaneLocation, FColor::White, false, KINDA_SMALL_NUMBER, Settings.DrawPriority, Settings.LineThickness);
-
 						// Manifold plane and normal
 						DrawCollisionImpl(WorldPlaneLocation, WorldPlaneNormal, ManifoldPoint.ContactPoint.Phi, ManifoldPoint.NetImpulse, DiscColor, NormalColor, ImpulseColor, ColorScale, Settings);
 
diff --git a/Engine/Source/Runtime/Experimental/Chaos/Private/Chaos/Collision/CollisionSolver.cpp b/Engine/Source/Runtime/Experimental/Chaos/Private/Chaos/Collision/CollisionSolver.cpp
--- a/Engine/Source/Runtime/Experimental/Chaos/Private/Chaos/Collision/CollisionSolver.cpp
+++ b/Engine/Source/Runtime/Experimental/Chaos/Private/Chaos/Collision/CollisionSolver.cpp
@@ -261,7 +261,7 @@ namespace Chaos
 		// This modifies the in/out CoM position and rotations.
 		void ApplyManifoldPushOutCorrection(
 			const FReal Stiffness,
-			const FCollisionContact& Contact,
+			const FRigidBodyPointContactConstraint& Constraint,
 			const TGenericParticleHandle<FReal, 3> Particle0,
 			const TGenericParticleHandle<FReal, 3> Particle1,
 			const FContactIterationParameters& IterationParameters,
@@ -278,8 +278,8 @@ namespace Chaos
 			FVec3& W1,
 			FManifoldPoint& ManifoldPoint)
 		{
-			TPBDRigidParticleHandle<FReal, 3>* PBDRigid0 = Particle0->CastToRigidParticle();
-			TPBDRigidParticleHandle<FReal, 3>* PBDRigid1 = Particle1->CastToRigidParticle();
+			const TPBDRigidParticleHandle<FReal, 3>* PBDRigid0 = Particle0->CastToRigidParticle();
+			const TPBDRigidParticleHandle<FReal, 3>* PBDRigid1 = Particle1->CastToRigidParticle();
 
 			const FReal Margin0 = ManifoldPoint.ContactPoint.ShapeMargins[0];
 			const FReal Margin1 = ManifoldPoint.ContactPoint.ShapeMargins[1];
@@ -290,8 +290,14 @@ namespace Chaos
 			const FVec3 ContactNormal = PlaneQ * ManifoldPoint.CoMContactNormal;
 
 			const bool bApplyStaticFriction = (ManifoldPoint.bInsideStaticFrictionCone && ManifoldPoint.bPotentialRestingContact && Chaos_Manifold_PushOut_StaticFriction);
-			FVec3 LocalContactPoint0 = (bApplyStaticFriction) ? ManifoldPoint.PrevCoMContactPoints[0] : ManifoldPoint.CoMContactPoints[0];
-			FVec3 LocalContactPoint1 = (bApplyStaticFriction) ? ManifoldPoint.PrevCoMContactPoints[1] : ManifoldPoint.CoMContactPoints[1];
+			FVec3 LocalContactPoint0 = ManifoldPoint.CoMContactPoints[0];
+			FVec3 LocalContactPoint1 = ManifoldPoint.CoMContactPoints[1];
+
+			// If we are applying static friction, we are attempting to move the contact points at their intitial positions back together
+			if (bApplyStaticFriction)
+			{
+				Constraint.CalculatePrevCoMContactPoints(Particle0, Particle1, ManifoldPoint, IterationParameters.Dt, LocalContactPoint0, LocalContactPoint1);
+			}
 
 			// We could push out to the PBD distance that would give an implicit velocity equal to -(1+e).Vin
 			// but the values involved are not very stable from frame to frame, so instead we actually pull
@@ -320,8 +326,8 @@ namespace Chaos
 			}
 
 			// Calculate joint-space mass matrix (J.M.Jt)
-			const FMatrix33 InvI0 = bIsRigidDynamic0 ? Utilities::ComputeWorldSpaceInertia(Q0, PBDRigid0->InvI()) * Contact.InvInertiaScale0 : FMatrix33(0);
-			const FMatrix33 InvI1 = bIsRigidDynamic1 ? Utilities::ComputeWorldSpaceInertia(Q1, PBDRigid1->InvI()) * Contact.InvInertiaScale1 : FMatrix33(0);
+			const FMatrix33 InvI0 = bIsRigidDynamic0 ? Utilities::ComputeWorldSpaceInertia(Q0, PBDRigid0->InvI()) * Constraint.Manifold.InvInertiaScale0 : FMatrix33(0);
+			const FMatrix33 InvI1 = bIsRigidDynamic1 ? Utilities::ComputeWorldSpaceInertia(Q1, PBDRigid1->InvI()) * Constraint.Manifold.InvInertiaScale1 : FMatrix33(0);
 			const FMatrix33 ContactMassInv =
 				(bIsRigidDynamic0 ? ComputeFactorMatrix3(RelativeContactPoint0, InvI0, PBDRigid0->InvM()) : FMatrix33(0)) +
 				(bIsRigidDynamic1 ? ComputeFactorMatrix3(RelativeContactPoint1, InvI1, PBDRigid1->InvM()) : FMatrix33(0));
@@ -688,7 +694,7 @@ namespace Chaos
 					{
 						ApplyManifoldPushOutCorrection(
 							Stiffness,
-							Constraint.Manifold,
+							Constraint,
 							Particle0,
 							Particle1,
 							IterationParameters,
diff --git a/Engine/Source/Runtime/Experimental/Chaos/Private/Chaos/Collision/PBDCollisionConstraint.cpp b/Engine/Source/Runtime/Experimental/Chaos/Private/Chaos/Collision/PBDCollisionConstraint.cpp
--- a/Engine/Source/Runtime/Experimental/Chaos/Private/Chaos/Collision/PBDCollisionConstraint.cpp
+++ b/Engine/Source/Runtime/Experimental/Chaos/Private/Chaos/Collision/PBDCollisionConstraint.cpp
@@ -227,6 +227,29 @@ namespace Chaos
 		ManifoldPoint.CoMContactPoints[1] = CoMContactPoint1;
 		ManifoldPoint.CoMContactNormal = CoMContactNormal;
 
+		// Calculate the initial contact velocity for use in restitution
+		// Should we use Previous V and W?
+		const FRigidTransform3 WorldTransform0 = FParticleUtilities::GetActorWorldTransform(Particle0);
+		const FRigidTransform3 WorldTransform1 = FParticleUtilities::GetActorWorldTransform(Particle1);
+		const FVec3 WorldContactVel0 = Particle0->V() + FVec3::CrossProduct(Particle0->W(), WorldTransform0.GetRotation() * LocalContactPoint0);
+		const FVec3 WorldContactVel1 = Particle1->V() + FVec3::CrossProduct(Particle1->W(), WorldTransform1.GetRotation() * LocalContactPoint1);
+		const FVec3 WorldContactNormal = (ManifoldPoint.ContactPoint.ContactNormalOwnerIndex == 0) ? WorldTransform0.GetRotation() * LocalContactNormal : WorldTransform1.GetRotation() * LocalContactNormal;
+		const FReal WorldContactVelNorm = FVec3::DotProduct(WorldContactVel0 - WorldContactVel1, WorldContactNormal);
+		ManifoldPoint.InitialContactVelocity = WorldContactVelNorm;
+
+		// Store the initial penetration depth for use with restitution with PBD
+		// NOTE: This is incorrect if we are updating the point each iteration (which we are for incremental manifolds but not one-shots - see UpdateManifoldPoint)
+		ManifoldPoint.InitialPhi = ManifoldPoint.ContactPoint.Phi;
+	}
+
+	void FRigidBodyPointContactConstraint::CalculatePrevCoMContactPoints(
+		const TConstGenericParticleHandle<FReal, 3> Particle0,
+		const TConstGenericParticleHandle<FReal, 3> Particle1,
+		FManifoldPoint& ManifoldPoint,
+		FReal Dt,
+		FVec3& OutPrevCoMContactPoint0,
+		FVec3& OutPrevCoMContactPoint1) const
+	{
 		// Recalculate the previous local-space contact position on the plane owner. This is used by static friction where
 		// we try to move the contact points back to their previous relative positions.
 		auto CalculatePrevCoMTransform = [Dt](const TConstGenericParticleHandle<FReal, 3>& ParticleHandle) -> FRigidTransform3
@@ -243,52 +266,26 @@ namespace Chaos
 
 		const FRigidTransform3 PrevCoMTransform0 = CalculatePrevCoMTransform(Particle0);
 		const FRigidTransform3 PrevCoMTransform1 = CalculatePrevCoMTransform(Particle1);
-		UpdatePrevCoMContactPoints(ManifoldPoint, PrevCoMTransform0.GetTranslation(), PrevCoMTransform0.GetRotation(), PrevCoMTransform1.GetTranslation(), PrevCoMTransform1.GetRotation());
 
-		// Calculate the initial contact velocity for use in restitution
-		const FRotation3& PrevPlaneRotation = (ManifoldPoint.ContactPoint.ContactNormalOwnerIndex == 0) ? PrevCoMTransform0.GetRotation() : PrevCoMTransform1.GetRotation();
-		const FVec3 PrevWorldContactNormal = PrevPlaneRotation * ManifoldPoint.CoMContactNormal;
-		const FVec3 PrevWorldContactOffset0 = PrevCoMTransform0.GetRotation() * ManifoldPoint.PrevCoMContactPoints[0] - ManifoldPoint.ContactPoint.ShapeMargins[0] * PrevWorldContactNormal;
-		const FVec3 PrevWorldContactOffset1 = PrevCoMTransform1.GetRotation() * ManifoldPoint.PrevCoMContactPoints[1] + ManifoldPoint.ContactPoint.ShapeMargins[1] * PrevWorldContactNormal;
-		const FVec3 PrevWorldContactVel0 = Particle0->PreV() + FVec3::CrossProduct(Particle0->PreW(), PrevWorldContactOffset0);
-		const FVec3 PrevWorldContactVel1 = Particle1->PreV() + FVec3::CrossProduct(Particle1->PreW(), PrevWorldContactOffset1);
-		const FReal PrevWorldContactVelNorm = FVec3::DotProduct(PrevWorldContactVel0 - PrevWorldContactVel1, PrevWorldContactNormal);
-		ManifoldPoint.InitialContactVelocity = PrevWorldContactVelNorm;
-
-		// Store the initial penetration depth for use with restitution with PBD
-		// NOTE: This is incorrect if we are updating the point each iteration (which we are for incremental manifolds but not one-shots - see UpdateManifoldPoint)
-		ManifoldPoint.InitialPhi = ManifoldPoint.ContactPoint.Phi;
-	}
-
-	// Recalculate the previous local-space contact position on the plane owner. This is used by static friction where
-	// we try to move the contact points back to their previous relative positions.
-	// NOTE: These com-relative contact positions are without margins applied, which is reflected in the Phi calculation.
-	void FRigidBodyPointContactConstraint::UpdatePrevCoMContactPoints(
-		FManifoldPoint& ManifoldPoint,
-		const FVec3& XCoM0,
-		const FRotation3& RCoM0,
-		const FVec3& XCoM1,
-		const FRotation3& RCoM1)
-	{
-		const FVec3 PrevWorldContactLocation0 = XCoM0 + RCoM0 * ManifoldPoint.CoMContactPoints[0];
-		const FVec3 PrevWorldContactLocation1 = XCoM1 + RCoM1 * ManifoldPoint.CoMContactPoints[1];
+		const FVec3 PrevWorldContactLocation0 = PrevCoMTransform0.GetTranslation() + PrevCoMTransform0.GetRotation() * ManifoldPoint.CoMContactPoints[0];
+		const FVec3 PrevWorldContactLocation1 = PrevCoMTransform1.GetTranslation() + PrevCoMTransform1.GetRotation() * ManifoldPoint.CoMContactPoints[1];
 		FVec3 PrevCoMContactPoint0 = ManifoldPoint.CoMContactPoints[0];
 		FVec3 PrevCoMContactPoint1 = ManifoldPoint.CoMContactPoints[1];
 		if (ManifoldPoint.ContactPoint.ContactNormalOwnerIndex == 0)
 		{
-			const FVec3 PrevWorldContactNormal = RCoM0 * ManifoldPoint.CoMContactNormal;
+			const FVec3 PrevWorldContactNormal = PrevCoMTransform0.GetRotation() * ManifoldPoint.CoMContactNormal;
 			const FReal PrevPhi = FVec3::DotProduct(PrevWorldContactLocation0 - PrevWorldContactLocation1, PrevWorldContactNormal);
-			PrevCoMContactPoint0 = RCoM0.Inverse() * (PrevWorldContactLocation1 + PrevPhi * PrevWorldContactNormal - XCoM0);
+			PrevCoMContactPoint0 = PrevCoMTransform0.GetRotation().Inverse() * (PrevWorldContactLocation1 + PrevPhi * PrevWorldContactNormal - PrevCoMTransform0.GetTranslation());
 		}
 		else
 		{
-			const FVec3 PrevWorldContactNormal = RCoM1 * ManifoldPoint.CoMContactNormal;
+			const FVec3 PrevWorldContactNormal = PrevCoMTransform1.GetRotation() * ManifoldPoint.CoMContactNormal;
 			const FReal PrevPhi = FVec3::DotProduct(PrevWorldContactLocation0 - PrevWorldContactLocation1, PrevWorldContactNormal);
-			PrevCoMContactPoint1 = RCoM1.Inverse() * (PrevWorldContactLocation0 - PrevPhi * PrevWorldContactNormal - XCoM1);
+			PrevCoMContactPoint1 = PrevCoMTransform1.GetRotation().Inverse() * (PrevWorldContactLocation0 - PrevPhi * PrevWorldContactNormal - PrevCoMTransform1.GetTranslation());
 		}
-		ManifoldPoint.PrevCoMContactPoints[0] = PrevCoMContactPoint0;
-		ManifoldPoint.PrevCoMContactPoints[1] = PrevCoMContactPoint1;
 
+		OutPrevCoMContactPoint0 = PrevCoMContactPoint0;
+		OutPrevCoMContactPoint1 = PrevCoMContactPoint1;
 	}
 
 	int32 FRigidBodyPointContactConstraint::AddManifoldPoint(const FContactPoint& ContactPoint, const FReal Dt)
diff --git a/Engine/Source/Runtime/Experimental/Chaos/Public/Chaos/Collision/PBDCollisionConstraint.h b/Engine/Source/Runtime/Experimental/Chaos/Public/Chaos/Collision/PBDCollisionConstraint.h
--- a/Engine/Source/Runtime/Experimental/Chaos/Public/Chaos/Collision/PBDCollisionConstraint.h
+++ b/Engine/Source/Runtime/Experimental/Chaos/Public/Chaos/Collision/PBDCollisionConstraint.h
@@ -14,13 +14,13 @@ namespace Chaos
 {
 	class FImplicitObject;
 	class FPBDCollisionConstraintHandle;
+	template<typename T, int d> class TConstGenericParticleHandle;
 
 	class CHAOS_API FManifoldPoint
 	{
 	public:
 		FManifoldPoint() 
 			: CoMContactPoints{ FVec3(0), FVec3(0) }
-			, PrevCoMContactPoints{ FVec3(0), FVec3(0) }
 			, CoMContactNormal(0)
 			, NetImpulse(0)
 			, NetPushOut(0)
@@ -37,7 +37,6 @@ namespace Chaos
 		FManifoldPoint(const FContactPoint& InContactPoint) 
 			: ContactPoint(InContactPoint)
 			, CoMContactPoints{ FVec3(0), FVec3(0) }
-			, PrevCoMContactPoints{ FVec3(0), FVec3(0) }
 			, CoMContactNormal(0)
 			, NetImpulse(0)
 			, NetPushOut(0)
@@ -54,7 +53,6 @@ namespace Chaos
 		// @todo(chaos): Normal and plane owner should be per manifold, not per manifold-point when we are not using incremental manifolds any more
 		FContactPoint ContactPoint;			// Shape-space data from low-level collision detection
 		FVec3 CoMContactPoints[2];			// CoM-space contact points on the two bodies core shapes (not including margin)
-		FVec3 PrevCoMContactPoints[2];		// CoM-space contact points on the two bodies core shapes at previous transforms (used for static friction)
 		FVec3 CoMContactNormal;				// CoM-space contact normal relative to ContactNormalOwner body	
 		FVec3 NetImpulse;					// Total impulse applied by this contact point
 		FVec3 NetPushOut;					// Total pushout applied at this contact point
@@ -283,12 +281,13 @@ namespace Chaos
 			const FRotation3& Q0,
 			const FVec3& P1,
 			const FRotation3& Q1);
-		void UpdatePrevCoMContactPoints(
+		void CalculatePrevCoMContactPoints(
+			const TConstGenericParticleHandle<FReal, 3> Particle0,
+			const TConstGenericParticleHandle<FReal, 3> Particle1,
 			FManifoldPoint& ManifoldPoint,
-			const FVec3& XCoM0,
-			const FRotation3& RCoM0,
-			const FVec3& XCoM1,
-			const FRotation3& RCoM1);
+			FReal Dt,
+			FVec3& OutPrevCoMContactPoint0,
+			FVec3& OutPrevCoMContactPoint1) const;
 
 		void AddIncrementalManifoldContact(const FContactPoint& ContactPoint, const FReal Dt);
 		void AddOneshotManifoldContact(const FContactPoint& ContactPoint, const FReal Dt);
@@ -321,7 +320,8 @@ namespace Chaos
 		void UpdateManifoldPoint(int32 ManifoldPointIndex, const FContactPoint& ContactPoint, const FReal Dt);
 		void SetActiveContactPoint(const FContactPoint& ContactPoint);
 
-		// @todo(chaos): inline array
+		// @todo(chaos): Switching this to inline allocator does not help, but maybe a physics scratchpad would
+		//TArray<FManifoldPoint, TInlineAllocator<4>> ManifoldPoints;
 		TArray<FManifoldPoint> ManifoldPoints;
 		bool bUseManifold;
 	};
diff --git a/Engine/Source/Runtime/Experimental/Chaos/Public/Chaos/ParticleHandle.h b/Engine/Source/Runtime/Experimental/Chaos/Public/Chaos/ParticleHandle.h
--- a/Engine/Source/Runtime/Experimental/Chaos/Public/Chaos/ParticleHandle.h
+++ b/Engine/Source/Runtime/Experimental/Chaos/Public/Chaos/ParticleHandle.h
@@ -19,6 +19,7 @@ namespace Chaos
 {
 	class FConstraintHandle;
 
+
 template <typename T, int d>
 struct TGeometryParticleParameters
 {
@@ -1487,8 +1488,10 @@ class TConstGenericParticleHandle
 {
 public:
 	TConstGenericParticleHandle(const TGeometryParticleHandle<T, d>* InHandle) : Imp(const_cast<TGeometryParticleHandle<T, d>*>(InHandle)) {}
+	TConstGenericParticleHandle(const TGenericParticleHandle<T, d> InHandle) : Imp(InHandle->Handle()) {}
 
 	const TGenericParticleHandleHandleImp<T, d>* operator->() const { return &Imp; }
+	const TGenericParticleHandleHandleImp<T, d>* Get() const { return &Imp; }
 
 private:
 	TGenericParticleHandleHandleImp<T, d> Imp;
