GeometryCollectionPhysicsProxy
  -- Correctly recorded physical material on the gamethread particle so that LineTrace etc. returns the correct Physical Material on hit.

diff --git a/Engine/Source/Runtime/Experimental/Chaos/Private/PhysicsProxy/GeometryCollectionPhysicsProxy.cpp b/Engine/Source/Runtime/Experimental/Chaos/Private/PhysicsProxy/GeometryCollectionPhysicsProxy.cpp
--- a/Engine/Source/Runtime/Experimental/Chaos/Private/PhysicsProxy/GeometryCollectionPhysicsProxy.cpp
+++ b/Engine/Source/Runtime/Experimental/Chaos/Private/PhysicsProxy/GeometryCollectionPhysicsProxy.cpp
@@ -484,6 +484,8 @@ void FGeometryCollectionPhysicsProxy::Initialize(Chaos::TPBDRigidsEvolutionBase<
 				Chaos::FPerShapeData* Shape = Shapes[ShapeIndex].Get();
 				Shape->SetSimData(SimFilter);
 				Shape->SetQueryData(QueryFilter);
+				Shape->SetProxy(this);
+				Shape->SetMaterial(Parameters.PhysicalMaterialHandle);
 			}
 		}
 	}
