GeometryCollectionComponent
-- corrected a bug that prevented collision events from notifying.

diff --git a/Engine/Source/Runtime/Experimental/GeometryCollectionEngine/Private/GeometryCollection/GeometryCollectionComponent.cpp b/Engine/Source/Runtime/Experimental/GeometryCollectionEngine/Private/GeometryCollection/GeometryCollectionComponent.cpp
--- a/Engine/Source/Runtime/Experimental/GeometryCollectionEngine/Private/GeometryCollection/GeometryCollectionComponent.cpp
+++ b/Engine/Source/Runtime/Experimental/GeometryCollectionEngine/Private/GeometryCollection/GeometryCollectionComponent.cpp
@@ -2059,6 +2059,12 @@ void UGeometryCollectionComponent::OnCreatePhysicsState()
 			// Enable for complex and simple (no dual representation currently like other meshes)
 			InitialQueryFilter.Word3 |= (EPDF_SimpleCollision | EPDF_ComplexCollision);
 			InitialSimFilter.Word3 |= (EPDF_SimpleCollision | EPDF_ComplexCollision);
+			
+			if (bNotifyCollisions)
+			{
+				InitialQueryFilter.Word3 |= EPDF_ContactNotify;
+				InitialSimFilter.Word3 |= EPDF_ContactNotify;
+			}
 
  			PhysicsProxy = new FGeometryCollectionPhysicsProxy(this, *DynamicCollection, SimulationParameters, InitialQueryFilter, InitialSimFilter, InitFunc, CacheSyncFunc, FinalSyncFunc);
 			FPhysScene_Chaos* Scene = GetInnerChaosScene();
