CIS failure: compiler warning
#[fyi] Benn.Gallagher

diff --git a/Engine/Source/Runtime/Experimental/Chaos/Private/Chaos/CollisionResolution.cpp b/Engine/Source/Runtime/Experimental/Chaos/Private/Chaos/CollisionResolution.cpp
--- a/Engine/Source/Runtime/Experimental/Chaos/Private/Chaos/CollisionResolution.cpp
+++ b/Engine/Source/Runtime/Experimental/Chaos/Private/Chaos/CollisionResolution.cpp
@@ -1829,7 +1829,7 @@ namespace Chaos
 		void ConstructGenericConvexConvexConstraintsSwept(TGeometryParticleHandle<FReal, 3>* Particle0, TGeometryParticleHandle<FReal, 3>* Particle1, const FImplicitObject* Implicit0, const FImplicitObject* Implicit1, const FRigidTransform3& LocalTransform0, const FRigidTransform3& LocalTransform1, const FReal CullDistance, const FReal Dt, const FVec3& Dir, FReal Length, FCollisionConstraintsArray& NewConstraints)
 		{
 			CONDITIONAL_SCOPE_CYCLE_COUNTER(STAT_Collisions_ConstructGenericConvexConvexConstraintsSwept, ConstraintsDetailedStats);
-			TGenericParticleHandle<FReal, 3> RigidParticle = TGenericParticleHandle<FReal, 3>(Particle0);
+			FGenericParticleHandle RigidParticle = TGenericParticleHandle<FReal, 3>(Particle0);
 			FRigidBodySweptPointContactConstraint Constraint = FRigidBodySweptPointContactConstraint(Particle0, Implicit0, nullptr, LocalTransform0, Particle1, Implicit1, nullptr, LocalTransform1, EContactShapesType::GenericConvexConvex);
 			// Note: This is unusual but we are using a mix of the previous and current transform
 			// This is due to how CCD only rewinds the position
diff --git a/Engine/Source/Runtime/Experimental/Chaos/Private/Chaos/PBDCollisionConstraintsContact.cpp b/Engine/Source/Runtime/Experimental/Chaos/Private/Chaos/PBDCollisionConstraintsContact.cpp
--- a/Engine/Source/Runtime/Experimental/Chaos/Private/Chaos/PBDCollisionConstraintsContact.cpp
+++ b/Engine/Source/Runtime/Experimental/Chaos/Private/Chaos/PBDCollisionConstraintsContact.cpp
@@ -61,7 +61,7 @@ namespace Chaos
 
 		void UpdateSwept(FRigidBodySweptPointContactConstraint& Constraint, const FReal CullDistance, const FReal Dt)
 		{
-			TGenericParticleHandle<FReal, 3> Particle0 = TGenericParticleHandle<FReal, 3>(Constraint.Particle[0]);
+			FGenericParticleHandle Particle0 = TGenericParticleHandle<FReal, 3>(Constraint.Particle[0]);
 			// Note: This is unusual but we are using a mix of the previous and current transform
 			// This is due to how CCD only rewinds the position
 			const FRigidTransform3 TransformXQ0(Particle0->X(), Particle0->Q());
