Chaos Physics - Second pass at disabling/enabling rigid body in simulation
- Handlke re-enable properly
- Change PBDJointCOnstraint tro properly bail from the apply method if any of the particle is disabled ( as well as make sur ethis code is shared between methods )
- Changed the Disable/Enable Particle function of the PBDRigidsEvolution to not update the ConstraintRule ( as this would set the particles of the constraints to be null and stopping from being enabled again )
- Make sure particle Disabled flag is accounted for in PBDRigidsEvaolutionGBF ( to avoid gravity from applying for example )

diff --git a/Engine/Source/Runtime/Engine/Private/SkeletalMeshComponentPhysics.cpp b/Engine/Source/Runtime/Engine/Private/SkeletalMeshComponentPhysics.cpp
--- a/Engine/Source/Runtime/Engine/Private/SkeletalMeshComponentPhysics.cpp
+++ b/Engine/Source/Runtime/Engine/Private/SkeletalMeshComponentPhysics.cpp
@@ -1095,9 +1095,15 @@ void USkeletalMeshComponent::SetAllBodiesNotifyRigidBodyCollision(bool bNewNotif
 
 void USkeletalMeshComponent::SetAllBodiesBelowPhysicsDisabled(const FName& InBoneName, bool bDisabled, bool bIncludeSelf)
 {
-	int32 NumBodiesFound = ForEachBodyBelow(InBoneName, bIncludeSelf, /*bSkipCustomPhysicsType=*/ false, [bDisabled](FBodyInstance* BI)
+	int32 NumBodiesFound = ForEachBodyBelow(InBoneName, bIncludeSelf, /*bSkipCustomPhysicsType=*/ false, [bDisabled, this](FBodyInstance* BI)
 	{
 		BI->SetPhysicsDisabled(bDisabled);
+		if (bDisabled == false)
+		{
+			FTransform BoneTransform(GetBoneMatrix(BI->InstanceBoneIndex));
+			// if we re-enable it, let's make sure the body transform is up to date 
+			BI->SetBodyTransform(BoneTransform, ETeleportType::TeleportPhysics, true);
+		}
 	});
 }
 
diff --git a/Engine/Source/Runtime/Experimental/Chaos/Private/Chaos/PBDConstraintGraph.cpp b/Engine/Source/Runtime/Experimental/Chaos/Private/Chaos/PBDConstraintGraph.cpp
--- a/Engine/Source/Runtime/Experimental/Chaos/Private/Chaos/PBDConstraintGraph.cpp
+++ b/Engine/Source/Runtime/Experimental/Chaos/Private/Chaos/PBDConstraintGraph.cpp
@@ -245,6 +245,7 @@ void FPBDConstraintGraph::AddConstraint(const uint32 InContainerId, FConstraintH
 	// Must have at least one constrained particle
 	check((ConstrainedParticles[0]) || (ConstrainedParticles[1]));
 
+	// todo(chaos) : we shoudl probably have a list of indices of removed edges, so we can reuse them instead of adding to the end ( see RemoveConstraint logic )
 	const int32 NewEdgeIndex = Edges.Num();
 	FGraphEdge NewEdge;
 	NewEdge.Data = { InContainerId, InConstraintHandle };
@@ -290,27 +291,53 @@ void FPBDConstraintGraph::RemoveConstraint(const uint32 InContainerId, FConstrai
 	int32* PNodeIndex1 = (ConstrainedParticles[1]) ? ParticleToNodeIndex.Find(ConstrainedParticles[1]) : nullptr;
 	if (ensure(PNodeIndex0 || PNodeIndex1) )
 	{
-		if ( InContainerId < (uint32)Edges.Num() )
+		// we need to find the edge index that matches this constraint
+		// we could go through the entire edge list and find the matching InContainerId, 
+		// but each node has an edge array that will contain it 
+		int32 EdgeIndex = INDEX_NONE;
+		if (PNodeIndex0)
+		{
+			EdgeIndex = GetConstraintIndexFromNodeAndConstraintHandle(Nodes[*PNodeIndex0], InConstraintHandle);
+		}
+		if (PNodeIndex1 && EdgeIndex == INDEX_NONE)
+		{
+			EdgeIndex = GetConstraintIndexFromNodeAndConstraintHandle(Nodes[*PNodeIndex1], InConstraintHandle);
+		}
+		// we should have an valid index at this point
+		if (EdgeIndex != INDEX_NONE)
 		{
 			int32 Index0 = INDEX_NONE, Index1 = INDEX_NONE;
-			if (PNodeIndex0) 
+			if (PNodeIndex0)
 			{
-				Nodes[*PNodeIndex0].Edges.Remove(InContainerId);
+				Nodes[*PNodeIndex0].Edges.Remove(EdgeIndex);
 				Index0 = *PNodeIndex0;
 			}
 			if (PNodeIndex1)
 			{
-				Nodes[*PNodeIndex1].Edges.Remove(InContainerId);
+				Nodes[*PNodeIndex1].Edges.Remove(EdgeIndex);
 				Index1 = *PNodeIndex1;
 			}
-			if (Edges[InContainerId].FirstNode == Index0 && Edges[InContainerId].SecondNode == Index1)
+			if (Edges[EdgeIndex].FirstNode == Index0 && Edges[EdgeIndex].SecondNode == Index1)
 			{
-				Edges[InContainerId] = FGraphEdge();
+				Edges[EdgeIndex] = FGraphEdge();
 			}
 		}
 	}
 }
 
+int32 FPBDConstraintGraph::GetConstraintIndexFromNodeAndConstraintHandle(const FGraphNode& Node, const FConstraintHandle* ConstraintHandle)
+{
+	for (const int32 NodeEdgeIndex : Node.Edges)
+	{
+		const FConstraintData& EdgeData = Edges[NodeEdgeIndex].Data;
+		if (EdgeData.GetConstraintHandle() == ConstraintHandle)
+		{
+			return NodeEdgeIndex;
+		}
+	}
+	return INDEX_NONE;
+}
+
 const typename FPBDConstraintGraph::FConstraintData& FPBDConstraintGraph::GetConstraintData(int32 ConstraintDataIndex) const
 {
 	return Edges[ConstraintDataIndex].Data;
@@ -366,6 +393,7 @@ void FPBDConstraintGraph::ComputeIslands(const TParticleView<FPBDRigidParticles>
 
 	for (auto& Particle : PBDRigids)
 	{
+		ensure(!Particle.Disabled());
 		auto* ParticleHandle = Particle.Handle();
 		int32 Idx = ParticleToNodeIndex[ParticleHandle];  // ryan - FAILS!
 		// selective reset of islands, don't reset if has been visited due to being edge connected to earlier processed node
@@ -614,18 +642,17 @@ bool FPBDConstraintGraph::ComputeIsland(const int32 InNode, const int32 Island,
 		FPBDRigidParticleHandle* RigidHandle = Node.Particle->CastToRigidParticle();
 		const bool isRigidDynamic = RigidHandle && RigidHandle->ObjectState() != EObjectStateType::Kinematic;  //??
 
+		if (RigidHandle && RigidHandle->Disabled())
+		{
+			continue;
+		}
+
 		ParticlesInIsland.Add(Node.Particle);
 		if (isRigidDynamic == false)
 		{
 			continue;
 		}
 
-		// @todo(ccaulfield): we don't handle enable/disable properly so this breaks
-		//if (RigidHandle->Disabled())
-		//{
-		//	continue;
-		//}
-
 		Node.Island = Island;
 		Visited[NodeIndex] = VisitToken;
 
diff --git a/Engine/Source/Runtime/Experimental/Chaos/Private/Chaos/PBDJointConstraints.cpp b/Engine/Source/Runtime/Experimental/Chaos/Private/Chaos/PBDJointConstraints.cpp
--- a/Engine/Source/Runtime/Experimental/Chaos/Private/Chaos/PBDJointConstraints.cpp
+++ b/Engine/Source/Runtime/Experimental/Chaos/Private/Chaos/PBDJointConstraints.cpp
@@ -456,7 +456,6 @@ namespace Chaos
 		}
 	}
 
-
 	void FPBDJointConstraints::SortConstraints()
 	{
 		// Sort constraints so that constraints with lower level (closer to a kinematic joint) are first
@@ -1234,16 +1233,46 @@ namespace Chaos
 		return true;
 	}
 
+	bool FPBDJointConstraints::CanEvaluate(const int32 ConstraintIndex) const
+	{
+		if (!IsConstraintEnabled(ConstraintIndex))
+		{
+			return false;
+		}
+
+		int32 Index0, Index1;
+		GetConstrainedParticleIndices(ConstraintIndex, Index0, Index1);
+		const TGenericParticleHandle<FReal, 3> Particle0 = TGenericParticleHandle<FReal, 3>(ConstraintParticles[ConstraintIndex][Index0]);
+		const TGenericParticleHandle<FReal, 3> Particle1 = TGenericParticleHandle<FReal, 3>(ConstraintParticles[ConstraintIndex][Index1]);
+
+		// check for valid and enabled particles
+		if (Particle0->Handle() == nullptr || Particle0->Disabled()
+			|| Particle1->Handle() == nullptr || Particle0->Disabled())
+		{
+			return false;
+		}
+
+		// check valid particle and solver state
+		const FJointSolverGaussSeidel& Solver = ConstraintSolvers[ConstraintIndex];
+		if ((Particle0->Sleeping() && Particle1->Sleeping())
+			|| (Particle0->IsKinematic() && Particle1->Sleeping())
+			|| (Particle0->Sleeping() && Particle1->IsKinematic())
+			|| (FMath::IsNearlyZero(Solver.InvM(0)) && FMath::IsNearlyZero(Solver.InvM(1))))
+		{
+			return false;
+		}
+		return true;
+	}
+
 	// This position solver iterates over each of the inner constraints (position, twist, swing) and solves them independently.
 	// This will converge slowly in some cases, particularly where resolving angular constraints violates position constraints and vice versa.
 	bool FPBDJointConstraints::ApplySingle(const FReal Dt, const int32 ConstraintIndex, const int32 NumPairIts, const int32 It, const int32 NumIts)
 	{
-		if (!IsConstraintEnabled(ConstraintIndex))
+		if (!CanEvaluate(ConstraintIndex))
 		{
 			return false;
 		}
 
-
 		const TVector<TGeometryParticleHandle<FReal, 3>*, 2>& Constraint = ConstraintParticles[ConstraintIndex];
 		UE_LOG(LogChaosJoint, VeryVerbose, TEXT("Solve Joint Constraint %d %s %s (dt = %f; it = %d / %d)"), ConstraintIndex, *Constraint[0]->ToString(), *Constraint[1]->ToString(), Dt, It, NumIts);
 
@@ -1255,15 +1284,6 @@ namespace Chaos
 		FGenericParticleHandle Particle0 = FGenericParticleHandle(ConstraintParticles[ConstraintIndex][Index0]);
 		FGenericParticleHandle Particle1 = FGenericParticleHandle(ConstraintParticles[ConstraintIndex][Index1]);
 
-		if ((Particle0->Sleeping() && Particle1->Sleeping())
-			|| (Particle0->IsKinematic() && Particle1->Sleeping()) 
-			|| (Particle0->Sleeping() && Particle1->IsKinematic())
-			|| (FMath::IsNearlyZero(Solver.InvM(0)) && FMath::IsNearlyZero(Solver.InvM(1))))
-		{
-			return false;
-		}
-
-
 		const FVec3 P0 = FParticleUtilities::GetCoMWorldPosition(Particle0);
 		const FRotation3 Q0 = FParticleUtilities::GetCoMWorldRotation(Particle0);
 		const FVec3 P1 = FParticleUtilities::GetCoMWorldPosition(Particle1);
@@ -1321,7 +1341,7 @@ namespace Chaos
 
 	bool FPBDJointConstraints::ApplyPushOutSingle(const FReal Dt, const int32 ConstraintIndex, const int32 NumPairIts, const int32 It, const int32 NumIts)
 	{
-		if (!IsConstraintEnabled(ConstraintIndex))
+		if (!CanEvaluate(ConstraintIndex))
 		{
 			return false;
 		}
@@ -1337,14 +1357,6 @@ namespace Chaos
 		FGenericParticleHandle Particle0 = FGenericParticleHandle(ConstraintParticles[ConstraintIndex][Index0]);
 		FGenericParticleHandle Particle1 = FGenericParticleHandle(ConstraintParticles[ConstraintIndex][Index1]);
 
-		if ((Particle0->Sleeping() && Particle1->Sleeping())
-			|| (Particle0->IsKinematic() && Particle1->Sleeping())
-			|| (Particle0->Sleeping() && Particle1->IsKinematic()) 
-			|| (FMath::IsNearlyZero(Solver.InvM(0)) && FMath::IsNearlyZero(Solver.InvM(1))))
-		{
-			return false;
-		}
-
 		const FVec3 P0 = FParticleUtilities::GetCoMWorldPosition(Particle0);
 		const FRotation3 Q0 = FParticleUtilities::GetCoMWorldRotation(Particle0);
 		const FVec3 P1 = FParticleUtilities::GetCoMWorldPosition(Particle1);
diff --git a/Engine/Source/Runtime/Experimental/Chaos/Private/PhysicsProxy/SingleParticlePhysicsProxy.cpp b/Engine/Source/Runtime/Experimental/Chaos/Private/PhysicsProxy/SingleParticlePhysicsProxy.cpp
--- a/Engine/Source/Runtime/Experimental/Chaos/Private/PhysicsProxy/SingleParticlePhysicsProxy.cpp
+++ b/Engine/Source/Runtime/Experimental/Chaos/Private/PhysicsProxy/SingleParticlePhysicsProxy.cpp
@@ -122,7 +122,6 @@ void PushToPhysicsStateImp(const Chaos::FDirtyPropertiesManager& Manager, Chaos:
 						Evolution.EnableParticle(Handle, nullptr);
 					}
 				}
-				Evolution.SetParticleObjectState(RigidHandle, NewData->ObjectState());
 
 				RigidHandle->SetDynamicMisc(*NewData);
 
diff --git a/Engine/Source/Runtime/Experimental/Chaos/Public/Chaos/PBDConstraintGraph.h b/Engine/Source/Runtime/Experimental/Chaos/Public/Chaos/PBDConstraintGraph.h
--- a/Engine/Source/Runtime/Experimental/Chaos/Public/Chaos/PBDConstraintGraph.h
+++ b/Engine/Source/Runtime/Experimental/Chaos/Public/Chaos/PBDConstraintGraph.h
@@ -237,6 +237,8 @@ namespace Chaos
 		int32 GetNextNodeIndex();
 		const TArray<int32>& GetUpdatedNodes() const { return UpdatedNodes; }
 
+		int32 GetConstraintIndexFromNodeAndConstraintHandle(const FGraphNode& Node, const FConstraintHandle* ConstraintHandle);
+
 		TArray<FGraphNode> Nodes;
 		TArray<FGraphEdge> Edges;
 		TArray<FIslandData> IslandToData;
diff --git a/Engine/Source/Runtime/Experimental/Chaos/Public/Chaos/PBDJointConstraints.h b/Engine/Source/Runtime/Experimental/Chaos/Public/Chaos/PBDJointConstraints.h
--- a/Engine/Source/Runtime/Experimental/Chaos/Public/Chaos/PBDJointConstraints.h
+++ b/Engine/Source/Runtime/Experimental/Chaos/Public/Chaos/PBDJointConstraints.h
@@ -117,7 +117,7 @@ namespace Chaos
 		void RemoveConstraints(const TSet<TGeometryParticleHandle<FReal, 3>*>& RemovedParticles) {}
 
 		/*
-		* Disable the constraints attached to the input particles. 
+		* Disable the constraints attached to the input particles. This will set the constrained Particle elements to nullptr and set the Enable flag to false
 		*/
 		void DisableConstraints(const TSet<TGeometryParticleHandle<FReal, 3>*>& RemovedParticles);
 
@@ -242,6 +242,8 @@ namespace Chaos
 		void BatchConstraints();
 		void CheckBatches();
 
+		bool CanEvaluate(const int32 ConstraintIndex) const;
+
 		bool ApplyBatch(const FReal Dt, const int32 BatchIndex, const int32 NumPairIts, const int32 It, const int32 NumIts);
 		bool ApplySingle(const FReal Dt, const int32 ConstraintIndex, const int32 NumPairIts, const int32 It, const int32 NumIts);
 		bool ApplyPushOutSingle(const FReal Dt, const int32 ConstraintIndex, const int32 NumPairIts, const int32 It, const int32 NumIts);
diff --git a/Engine/Source/Runtime/Experimental/Chaos/Public/Chaos/PBDRigidsEvolution.h b/Engine/Source/Runtime/Experimental/Chaos/Public/Chaos/PBDRigidsEvolution.h
--- a/Engine/Source/Runtime/Experimental/Chaos/Public/Chaos/PBDRigidsEvolution.h
+++ b/Engine/Source/Runtime/Experimental/Chaos/Public/Chaos/PBDRigidsEvolution.h
@@ -320,6 +320,8 @@ public:
 	{
 		Particles.EnableParticle(Particle);
 		ConstraintGraph.EnableParticle(Particle, ParentParticle);
+		// enable constraint after disabling the particle in the graph because of dependencies 
+		AddConstraintsToConstraintGraph(Particle->ParticleConstraints());
 		DirtyParticle(*Particle);
 	}
 
@@ -327,8 +329,8 @@ public:
 	{
 		RemoveParticleFromAccelerationStructure(*Particle);
 		Particles.DisableParticle(Particle);
-		// disable constraint before particles has particles need to be in the constraint graph for constraints to be removed
-		DisableConstraints(TSet<FGeometryParticleHandle*>({ Particle }));
+		// disable constraint before disabling the particle in the graph because of dependencies 
+		RemoveConstraintsFromConstraintGraph(Particle->ParticleConstraints());
 		ConstraintGraph.DisableParticle(Particle);
 	}
 
@@ -340,9 +342,13 @@ public:
 		{
 			RemoveParticleFromAccelerationStructure(*Particle);
 			Particles.DisableParticle(Particle);
+			RemoveConstraintsFromConstraintGraph(Particle->ParticleConstraints());
+		}
+		// make sure we do this has a second loop to make sure constraints have valid particles when removed from the constraint graph in the first loop
+		for (FGeometryParticleHandle* Particle : ParticlesIn)
+		{
 			ConstraintGraph.DisableParticle(Particle);
 		}
-		DisableConstraints(ParticlesIn);
 	}
 
 	template <bool bPersistent>
@@ -422,7 +428,7 @@ public:
 		Particle->SetObjectStateLowLevel(ObjectState);
 		Particles.SetDynamicParticleSOA(Particle);
 
-		if (InitialState == EObjectStateType::Sleeping && InitialState != ObjectState)
+		if (InitialState == EObjectStateType::Sleeping && InitialState != ObjectState && !Particle->Disabled())
 		{
 			if (Particle->Island() != INDEX_NONE)
 			{
@@ -472,19 +478,39 @@ public:
 		}*/
 	}
 
-	CHAOS_API void DisableConstraints(const TSet<FGeometryParticleHandle*>& RemovedParticles)
+	/** remove a list of constraints from the constraint graph (see AddConstraintsToConstraintGraph) */
+	CHAOS_API void RemoveConstraintsFromConstraintGraph(const TArray<FConstraintHandle*>& Constraints)
 	{
-		for (FGeometryParticleHandle* ParticleHandle : RemovedParticles)
+		for (FConstraintHandle* BaseConstraintHandle : Constraints)
 		{
-			for (FConstraintHandle* BaseConstraintHandle : ParticleHandle->ParticleConstraints())
+			if (FPBDJointConstraintHandle* ConstraintHandle = BaseConstraintHandle->As<FPBDJointConstraintHandle>())
 			{
-				if (FPBDJointConstraintHandle* ConstraintHandle = BaseConstraintHandle->As<FPBDJointConstraintHandle>())
-				{
-					ConstraintGraph.RemoveConstraint(ConstraintHandle->GetConstraintIndex(), ConstraintHandle, ConstraintHandle->GetConstrainedParticles());
-				}
+				ConstraintGraph.RemoveConstraint(ConstraintHandle->GetConstraintIndex(), ConstraintHandle, ConstraintHandle->GetConstrainedParticles());
 			}
 		}
+	}
 
+	/** Add a list of constraints to the constraint graph (see RemoveConstraintsFromConstraintGraph) */
+	CHAOS_API void AddConstraintsToConstraintGraph(const TArray<FConstraintHandle*>& Constraints)
+	{
+		for (FConstraintHandle* BaseConstraintHandle : Constraints)
+		{
+			if (FPBDJointConstraintHandle* ConstraintHandle = BaseConstraintHandle->As<FPBDJointConstraintHandle>())
+			{
+				ConstraintGraph.AddConstraint(ConstraintHandle->GetConstraintIndex(), ConstraintHandle, ConstraintHandle->GetConstrainedParticles());
+			}
+		}
+	}
+
+	/** Disable constraints from a set of particles to be removed (or destroyed) 
+	* this will set the constraints to Enbaled = false and set their respective bodies handles to nullptr
+	*/
+	CHAOS_API void DisableConstraints(const TSet<FGeometryParticleHandle*>& RemovedParticles)
+	{
+		for (FGeometryParticleHandle* ParticleHandle : RemovedParticles)
+		{
+			RemoveConstraintsFromConstraintGraph(ParticleHandle->ParticleConstraints());
+		}
 
 		for (FPBDConstraintGraphRule* ConstraintRule : ConstraintRules)
 		{
diff --git a/Engine/Source/Runtime/Experimental/Chaos/Public/Chaos/PBDRigidsSOAs.h b/Engine/Source/Runtime/Experimental/Chaos/Public/Chaos/PBDRigidsSOAs.h
--- a/Engine/Source/Runtime/Experimental/Chaos/Public/Chaos/PBDRigidsSOAs.h
+++ b/Engine/Source/Runtime/Experimental/Chaos/Public/Chaos/PBDRigidsSOAs.h
@@ -287,7 +287,7 @@ public:
 			}
 			else
 			{
-				Particle->MoveToSOA(*DynamicDisabledParticles);
+				SetDynamicParticleSOA(PBDRigid);
 			}
 
 			// All active particles RIGID particles
@@ -474,33 +474,39 @@ public:
 	{
 		const EObjectStateType State = Particle->ObjectState();
 
-
-		if (Particle->ObjectState() != EObjectStateType::Dynamic)
+		if (Particle->Disabled())
 		{
+			Particle->MoveToSOA(*DynamicDisabledParticles);
 			RemoveFromMapAndArray(Particle->CastToRigidParticle(), ActiveParticlesToIndex, ActiveParticlesArray);
 		}
 		else
 		{
-			InsertToMapAndArray(Particle->CastToRigidParticle(), ActiveParticlesToIndex, ActiveParticlesArray);
-		}
+			if (Particle->ObjectState() != EObjectStateType::Dynamic)
+			{
+				RemoveFromMapAndArray(Particle->CastToRigidParticle(), ActiveParticlesToIndex, ActiveParticlesArray);
+			}
+			else
+			{
+				InsertToMapAndArray(Particle->CastToRigidParticle(), ActiveParticlesToIndex, ActiveParticlesArray);
+			}
 
-		// Move to appropriate dynamic SOA
-		switch (State)
-		{
-		case EObjectStateType::Kinematic:
-			Particle->MoveToSOA(*DynamicKinematicParticles);
-			break;
+			// Move to appropriate dynamic SOA
+			switch (State)
+			{
+			case EObjectStateType::Kinematic:
+				Particle->MoveToSOA(*DynamicKinematicParticles);
+				break;
 
-		case EObjectStateType::Dynamic:
-			Particle->MoveToSOA(*DynamicParticles);
-			break;
+			case EObjectStateType::Dynamic:
+				Particle->MoveToSOA(*DynamicParticles);
+				break;
 
-		default:
-			// TODO: Special SOAs for sleeping and static particles?
-			Particle->MoveToSOA(*DynamicParticles);
-			break;
+			default:
+				// TODO: Special SOAs for sleeping and static particles?
+				Particle->MoveToSOA(*DynamicParticles);
+				break;
+			}
 		}
-
 		UpdateViews();
 	}
 
