Use a map to avoid the O(n^2) handling of unique edges when building up connectivity information on a chaos triangle mesh

#rnx

diff --git a/Engine/Source/Runtime/Experimental/Chaos/Private/Chaos/TriangleMesh.cpp b/Engine/Source/Runtime/Experimental/Chaos/Private/Chaos/TriangleMesh.cpp
--- a/Engine/Source/Runtime/Experimental/Chaos/Private/Chaos/TriangleMesh.cpp
+++ b/Engine/Source/Runtime/Experimental/Chaos/Private/Chaos/TriangleMesh.cpp
@@ -631,11 +631,12 @@ FSegmentMesh& FTriangleMesh::GetSegmentMesh()
 		return MSegmentMesh;
 	}
 
-	// XXX - Unfortunately, TSet is not a tree, it's a hash set.  This exposes
-	// us to the possibility we'll see hash collisions, and that's not something
-	// we should allow.  So we use a TArray instead.
+	// Array of ordered edges; other mappings will refer to edges via their index in this array
 	TArray<TVec2<int32>> UniqueEdges;
 	UniqueEdges.Reserve(MElements.Num() * 3);
+	// Map to accelerate checking if an edge is already in the UniqueEdges array
+	TMap<TVector<int32, 2>, int32> EdgeToIdx;
+	EdgeToIdx.Reserve(MElements.Num() * 3);
 
 	MEdgeToFaces.Reset();
 	MEdgeToFaces.Reserve(MElements.Num() * 3); // over estimate
@@ -649,7 +650,18 @@ FSegmentMesh& FTriangleMesh::GetSegmentMesh()
 		{
 			TVec2<int32> Edge(Tri[j], Tri[(j + 1) % 3]);
 
-			const int32 EdgeIdx = UniqueEdges.AddUnique(GetOrdered(Edge));
+			int32 EdgeIdx;
+			TVector<int32, 2> OrderedEdge = GetOrdered(Edge);
+			int32* FoundEdgeIdx = EdgeToIdx.Find(OrderedEdge);
+			if (!FoundEdgeIdx)
+			{
+				EdgeIdx = UniqueEdges.Add(OrderedEdge);
+				EdgeToIdx.Add(OrderedEdge, EdgeIdx);
+			}
+			else
+			{
+				EdgeIdx = *FoundEdgeIdx;
+			}
 			EdgeIds[j] = EdgeIdx;
 
 			// Track which faces are shared by edges.
