Fix asset supporet for old levels

diff --git a/Engine/Source/Runtime/Experimental/FieldSystem/Source/FieldSystemEngine/Public/Field/FieldSystemComponent.h b/Engine/Source/Runtime/Experimental/FieldSystem/Source/FieldSystemEngine/Public/Field/FieldSystemComponent.h
--- a/Engine/Source/Runtime/Experimental/FieldSystem/Source/FieldSystemEngine/Public/Field/FieldSystemComponent.h
+++ b/Engine/Source/Runtime/Experimental/FieldSystem/Source/FieldSystemEngine/Public/Field/FieldSystemComponent.h
@@ -39,7 +39,7 @@ public:
 	FORCEINLINE const UFieldSystem* GetFieldSystem() const { return FieldSystem; }
 
 	/** Field system asset to be used to store the construction fields */
-	UPROPERTY(EditAnywhere, BlueprintReadOnly, Category = "Field", meta = (ToolTip = "Field system asset to be used to store the construction fields"))
+	UPROPERTY(EditAnywhere, AdvancedDisplay, BlueprintReadOnly, Category = "Field", meta = (ToolTip = "Field system asset to be used to store the construction fields. This asset is not required anymore and will be deprecated soon."))
 	TObjectPtr<UFieldSystem> FieldSystem;
 
 	/** If enabled the field will be pushed to the world fields and will be available to materials and niagara */
@@ -51,7 +51,7 @@ public:
 	bool bIsChaosField;
 
 	/** List of solvers this field will affect. An empty list makes this field affect all solvers. */
-	UPROPERTY(EditAnywhere, Category = "Field", meta = (EditCondition = "bIsChaosField == true", ToolTip = "List of chaos solvers that will used the field"))
+	UPROPERTY(EditAnywhere, Category = "Field", meta = (EditCondition = "bIsChaosField == true", ToolTip = "List of chaos solvers that will use the field"))
 	TArray<TSoftObjectPtr<AChaosSolverActor>> SupportedSolvers;
 
 	/** List of all the construction command */
diff --git a/Engine/Source/Runtime/Experimental/GeometryCollectionEngine/Private/GeometryCollection/GeometryCollectionComponent.cpp b/Engine/Source/Runtime/Experimental/GeometryCollectionEngine/Private/GeometryCollection/GeometryCollectionComponent.cpp
--- a/Engine/Source/Runtime/Experimental/GeometryCollectionEngine/Private/GeometryCollection/GeometryCollectionComponent.cpp
+++ b/Engine/Source/Runtime/Experimental/GeometryCollectionEngine/Private/GeometryCollection/GeometryCollectionComponent.cpp
@@ -2661,10 +2661,26 @@ void UGeometryCollectionComponent::GetInitializationCommands(TArray<FFieldSystem
 		{
 			if (FieldSystemActor->GetFieldSystemComponent())
 			{
-				for (int32 CommandIndex = 0, NumCommands = FieldSystemActor->GetFieldSystemComponent()->ConstructionCommands.GetNumCommands();
-					CommandIndex < NumCommands; ++CommandIndex)
+				const int32 NumCommands = FieldSystemActor->GetFieldSystemComponent()->ConstructionCommands.GetNumCommands();
+				if (NumCommands > 0)
 				{
-					CombinedCommmands.Add(FieldSystemActor->GetFieldSystemComponent()->ConstructionCommands.BuildFieldCommand(CommandIndex));
+					for (int32 CommandIndex = 0; CommandIndex < NumCommands; ++CommandIndex)
+					{
+						CombinedCommmands.Add(FieldSystemActor->GetFieldSystemComponent()->ConstructionCommands.BuildFieldCommand(CommandIndex));
+					}
+				}
+				// Legacy path : only there for old levels. New ones will have the commands directly stored onto the component
+				else if (FieldSystemActor->GetFieldSystemComponent()->GetFieldSystem())
+				{
+					for (const FFieldSystemCommand& Command : FieldSystemActor->GetFieldSystemComponent()->GetFieldSystem()->Commands)
+					{
+						FFieldSystemCommand NewCommand = { Command.TargetAttribute, Command.RootNode->NewCopy() };
+						for (auto& Elem : Command.MetaData)
+						{
+							NewCommand.MetaData.Add(Elem.Key, TUniquePtr<FFieldSystemMetaData>(Elem.Value->NewCopy()));
+						}
+						CombinedCommmands.Add(NewCommand);
+					}
 				}
 			}
 		}
