Chaos physics - fix static analysis error ( actually a bug too )

diff --git a/Engine/Source/Runtime/Experimental/Chaos/Private/Chaos/PBDJointConstraints.cpp b/Engine/Source/Runtime/Experimental/Chaos/Private/Chaos/PBDJointConstraints.cpp
--- a/Engine/Source/Runtime/Experimental/Chaos/Private/Chaos/PBDJointConstraints.cpp
+++ b/Engine/Source/Runtime/Experimental/Chaos/Private/Chaos/PBDJointConstraints.cpp
@@ -1247,7 +1247,7 @@ namespace Chaos
 
 		// check for valid and enabled particles
 		if (Particle0->Handle() == nullptr || Particle0->Disabled()
-			|| Particle1->Handle() == nullptr || Particle0->Disabled())
+			|| Particle1->Handle() == nullptr || Particle1->Disabled())
 		{
 			return false;
 		}
