Chaos - Scaled convex fixes
- OneShotManifold edge collision fix
- Add non-template base class to ImplicitScaled and Instanced (part of a postponed refactor)

diff --git a/Engine/Source/Runtime/Experimental/Chaos/Private/Chaos/Collision/PBDCollisionConstraint.cpp b/Engine/Source/Runtime/Experimental/Chaos/Private/Chaos/Collision/PBDCollisionConstraint.cpp
--- a/Engine/Source/Runtime/Experimental/Chaos/Private/Chaos/Collision/PBDCollisionConstraint.cpp
+++ b/Engine/Source/Runtime/Experimental/Chaos/Private/Chaos/Collision/PBDCollisionConstraint.cpp
@@ -215,9 +215,9 @@ namespace Chaos
 		check(ManifoldPoint.ContactPoint.ContactNormalOwnerIndex >= 0);
 		check(ManifoldPoint.ContactPoint.ContactNormalOwnerIndex < 2);
 		const FRigidTransform3& PlaneTransform = ImplicitTransform[ManifoldPoint.ContactPoint.ContactNormalOwnerIndex];
-		const FVec3 LocalContactPoint0 = ImplicitTransform[0].TransformPosition(ManifoldPoint.ContactPoint.ShapeContactPoints[0]);	// Particle Space on body 0
-		const FVec3 LocalContactPoint1 = ImplicitTransform[1].TransformPosition(ManifoldPoint.ContactPoint.ShapeContactPoints[1]);	// Particle Space on body 1
-		const FVec3 LocalContactNormal = PlaneTransform.TransformVector(ManifoldPoint.ContactPoint.ShapeContactNormal);				// Particle Space on Plane owner
+		const FVec3 LocalContactPoint0 = ImplicitTransform[0].TransformPositionNoScale(ManifoldPoint.ContactPoint.ShapeContactPoints[0]);	// Particle Space on body 0
+		const FVec3 LocalContactPoint1 = ImplicitTransform[1].TransformPositionNoScale(ManifoldPoint.ContactPoint.ShapeContactPoints[1]);	// Particle Space on body 1
+		const FVec3 LocalContactNormal = PlaneTransform.TransformNormalNoScale(ManifoldPoint.ContactPoint.ShapeContactNormal);				// Particle Space on Plane owner
 
 		const FRotation3& PlaneRotationOfMass = (ManifoldPoint.ContactPoint.ContactNormalOwnerIndex == 0) ? Particle0->RotationOfMass() : Particle1->RotationOfMass();
 		const FVec3 CoMContactPoint0 = Particle0->RotationOfMass().Inverse() * (LocalContactPoint0 - Particle0->CenterOfMass());	// CoM Space on Body 0
diff --git a/Engine/Source/Runtime/Experimental/Chaos/Private/Chaos/CollisionOneShotManifolds.cpp b/Engine/Source/Runtime/Experimental/Chaos/Private/Chaos/CollisionOneShotManifolds.cpp
--- a/Engine/Source/Runtime/Experimental/Chaos/Private/Chaos/CollisionOneShotManifolds.cpp
+++ b/Engine/Source/Runtime/Experimental/Chaos/Private/Chaos/CollisionOneShotManifolds.cpp
@@ -16,7 +16,7 @@
 
 namespace Chaos
 {
-	float Chaos_Collision_Manifold_PlaneContactNormalEpsilon = 0.001f;
+	float Chaos_Collision_Manifold_PlaneContactNormalEpsilon = 0.01f;
 	FAutoConsoleVariableRef CVarChaos_Manifold_PlaneContactNormalEpsilon(TEXT("p.Chaos.Collision.Manifold.PlaneContactNormalEpsilon"), Chaos_Collision_Manifold_PlaneContactNormalEpsilon, TEXT("Normal tolerance used to distinguish face contacts from edge-edge contacts"));
 
 	// @todo(chaos): TEMP - use convex-convex collisio detection for box-box until TBox::GetClosestEdgePosition is implemented for that path (without plane hint)
@@ -569,17 +569,56 @@ namespace Chaos
 			const FReal PlaneContactNormalEpsilon = Chaos_Collision_Manifold_PlaneContactNormalEpsilon;
 			const bool bIsPlaneContact = FMath::IsNearlyEqual(BestPlaneDotNormalConvex1, 1.0f, PlaneContactNormalEpsilon) || FMath::IsNearlyEqual(BestPlaneDotNormalConvex2, 1.0f, PlaneContactNormalEpsilon);
 
-			// For edge-edge contacts, we find the edges involved and project the contact onto the edges
+			// For edge-edge contacts, we find the edges involved and generate contact from the closest points on the edges
 			if (!bIsPlaneContact)
 			{
-				FVec3 ShapeEdgePos1 = Convex1.GetClosestEdgePosition(MostOpposingPlaneIndexConvex1, GJKContactPoint.ShapeContactPoints[0]);
-				FVec3 ShapeEdgePos2 = Convex2.GetClosestEdgePosition(MostOpposingPlaneIndexConvex2, GJKContactPoint.ShapeContactPoints[1]);
-				FVec3 EdgePos1 = Convex1Transform.TransformPosition(ShapeEdgePos1);
-				FVec3 EdgePos2 = Convex2Transform.TransformPosition(ShapeEdgePos2);
-				FReal EdgePhi = FVec3::DotProduct(EdgePos1 - EdgePos2, GJKContactPoint.Normal);
+				// Get the edge vertices on each convex closest to the GJK contact
+				int32 EdgeVertexIndex1A, EdgeVertexIndex1B;
+				int32 EdgeVertexIndex2A, EdgeVertexIndex2B;
+				bool bFoundEdge1 = Convex1.GetClosestEdgeVertices(MostOpposingPlaneIndexConvex1, GJKContactPoint.ShapeContactPoints[0], EdgeVertexIndex1A, EdgeVertexIndex1B);
+				bool bFoundEdge2 = Convex2.GetClosestEdgeVertices(MostOpposingPlaneIndexConvex2, GJKContactPoint.ShapeContactPoints[1], EdgeVertexIndex2A, EdgeVertexIndex2B);
+				
+				// Sanity check - should only happen with null convexes
+				if (!bFoundEdge1 || !bFoundEdge2)
+				{
+					return;
+				}
+
+				// Get the world-space edge vertices and find tge closest point
+				const FVec3 EdgeVertex1A = Convex1Transform.TransformPosition(Convex1.GetVertex(EdgeVertexIndex1A));
+				const FVec3 EdgeVertex1B = Convex1Transform.TransformPosition(Convex1.GetVertex(EdgeVertexIndex1B));
+				const FVec3 EdgeVertex2A = Convex2Transform.TransformPosition(Convex2.GetVertex(EdgeVertexIndex2A));
+				const FVec3 EdgeVertex2B = Convex2Transform.TransformPosition(Convex2.GetVertex(EdgeVertexIndex2B));
+				FReal S, T;
+				FVec3 EdgePos1, EdgePos2;
+				Utilities::NearestPointsOnLineSegments(EdgeVertex1A, EdgeVertex1B, EdgeVertex2A, EdgeVertex2B, S, T, EdgePos1, EdgePos2);
+
+				// Generate the normal from the edges (make sure we don't change the sign)
+				FVec3 EdgeNormal = FVec3::CrossProduct(EdgeVertex1B - EdgeVertex1A, EdgeVertex2B - EdgeVertex2A);
+				const FReal EdgeNormalSq = EdgeNormal.SizeSquared();
+				if (EdgeNormalSq > SMALL_NUMBER)
+				{
+					EdgeNormal = EdgeNormal * FMath::InvSqrt(EdgeNormalSq);
+					if (FVec3::DotProduct(EdgeNormal, GJKContactPoint.Normal) < 0.0f)
+					{
+						EdgeNormal = -EdgeNormal;
+					}
+				}
+				else
+				{
+					EdgeNormal = GJKContactPoint.Normal;
+				}
+
+				// Shape-space contact data
+				const FReal EdgePhi = FVec3::DotProduct(EdgePos1 - EdgePos2, EdgeNormal);
+				const FVec3 ShapeEdgePos1 = Convex1Transform.InverseTransformPositionNoScale(EdgePos1);
+				const FVec3 ShapeEdgePos2 = Convex2Transform.InverseTransformPositionNoScale(EdgePos2);
+				const FVec3 ShapeEdgeNormal2 = Convex2Transform.InverseTransformVectorNoScale(EdgeNormal);
 
+				// Update the contact with the edge-edge contact
 				GJKContactPoint.ShapeContactPoints[0] = ShapeEdgePos1;
 				GJKContactPoint.ShapeContactPoints[1] = ShapeEdgePos2;
+				GJKContactPoint.ShapeContactNormal = ShapeEdgeNormal2;
 				GJKContactPoint.Phi = EdgePhi;
 				GJKContactPoint.Location = 0.5f * (EdgePos1 + EdgePos2);
 
@@ -668,7 +707,7 @@ namespace Chaos
 				ContactPoint.ShapeMargins[1] = 0.0f;
 				ContactPoint.ShapeContactPoints[0] = ReferenceFaceConvex1 ? PointProjectedOntoReferenceFace : ClippedPointInOtherCoordinates;
 				ContactPoint.ShapeContactPoints[1] = ReferenceFaceConvex1 ? ClippedPointInOtherCoordinates : PointProjectedOntoReferenceFace;
-				ContactPoint.ShapeContactNormal = RefSeparationDirection;
+				ContactPoint.ShapeContactNormal = ReferenceFaceConvex1 ? -RefPlaneNormal : RefPlaneNormal;
 				ContactPoint.ContactNormalOwnerIndex = ReferenceFaceConvex1 ? 0 : 1;
 
 				ContactPoint.Location = RefConvexTM->TransformPositionNoScale(PointProjectedOntoReferenceFace);
diff --git a/Engine/Source/Runtime/Experimental/Chaos/Private/Chaos/Convex.cpp b/Engine/Source/Runtime/Experimental/Chaos/Private/Chaos/Convex.cpp
--- a/Engine/Source/Runtime/Experimental/Chaos/Private/Chaos/Convex.cpp
+++ b/Engine/Source/Runtime/Experimental/Chaos/Private/Chaos/Convex.cpp
@@ -153,6 +153,42 @@ namespace Chaos
 		return ClosestEdgePosition;
 	}
 
+	bool FConvex::GetClosestEdgeVertices(int32 PlaneIndex, const FVec3& Position, int32& OutVertexIndex0, int32& OutVertexIndex1) const
+	{
+		OutVertexIndex0 = INDEX_NONE;
+		OutVertexIndex1 = INDEX_NONE;
+
+		FReal ClosestDistanceSq = FLT_MAX;
+		const int32 PlaneVerticesNum = NumPlaneVertices(PlaneIndex);
+		if (PlaneVerticesNum > 0)
+		{
+			int32 VertexIndex0 = GetPlaneVertex(PlaneIndex, PlaneVerticesNum - 1);
+			FVec3 P0 = GetVertex(VertexIndex0);
+
+			for (int32 PlaneVertexIndex = 0; PlaneVertexIndex < PlaneVerticesNum; ++PlaneVertexIndex)
+			{
+				const int32 VertexIndex1 = GetPlaneVertex(PlaneIndex, PlaneVertexIndex);
+				const FVec3 P1 = GetVertex(VertexIndex1);
+
+				const FVec3 EdgePosition = FMath::ClosestPointOnLine(P0, P1, Position);
+				const FReal EdgeDistanceSq = (EdgePosition - Position).SizeSquared();
+
+				if (EdgeDistanceSq < ClosestDistanceSq)
+				{
+					OutVertexIndex0 = VertexIndex0;
+					OutVertexIndex1 = VertexIndex1;
+					ClosestDistanceSq = EdgeDistanceSq;
+				}
+
+				VertexIndex0 = VertexIndex1;
+				P0 = P1;
+			}
+			return true;
+		}
+		return false;
+	}
+
+
 	int32 FConvex::NumVertexPlanes(int32 VertexIndex) const
 	{
 		if (StructureData.IsValid())
diff --git a/Engine/Source/Runtime/Experimental/Chaos/Public/Chaos/AABB.h b/Engine/Source/Runtime/Experimental/Chaos/Public/Chaos/AABB.h
--- a/Engine/Source/Runtime/Experimental/Chaos/Public/Chaos/AABB.h
+++ b/Engine/Source/Runtime/Experimental/Chaos/Public/Chaos/AABB.h
@@ -141,11 +141,11 @@ namespace Chaos
 		}
 
 
-		FORCEINLINE T PhiWithNormal(const TVector<T, d>& x, TVector<T, d>& Normal) const 
+		FORCEINLINE T PhiWithNormal(const TVector<T, d>& X, TVector<T, d>& Normal) const 
 		{
-			const TVector<T, d> MaxDists = x - MMax;
-			const TVector<T, d> MinDists = MMin - x;
-			if (x <= MMax && x >= MMin)
+			const TVector<T, d> MaxDists = X - MMax;
+			const TVector<T, d> MinDists = MMin - X;
+			if (X <= MMax && X >= MMin)
 			{
 				const Pair<T, int32> MaxAndAxis = TVector<T, d>::MaxAndAxis(MinDists, MaxDists);
 				Normal = MaxDists[MaxAndAxis.Second] > MinDists[MaxAndAxis.Second] ? TVector<T, d>::AxisVector(MaxAndAxis.Second) : -TVector<T, d>::AxisVector(MaxAndAxis.Second);
diff --git a/Engine/Source/Runtime/Experimental/Chaos/Public/Chaos/Box.h b/Engine/Source/Runtime/Experimental/Chaos/Public/Chaos/Box.h
--- a/Engine/Source/Runtime/Experimental/Chaos/Public/Chaos/Box.h
+++ b/Engine/Source/Runtime/Experimental/Chaos/Public/Chaos/Box.h
@@ -156,6 +156,11 @@ namespace Chaos
 			return AABB.PhiWithNormal(Pos, Normal);
 		}
 
+		virtual T PhiWithNormalScaled(const TVector<T, d>& Pos, const TVector<T, d>& Scale, TVector<T, d>& Normal) const override
+		{
+			return TAABB<T, d>(Scale * AABB.Min(), Scale * AABB.Max()).PhiWithNormal(Pos, Normal);
+		}
+
 		static FORCEINLINE bool RaycastFast(const TVector<T,d>& InMin, const TVector<T,d>& InMax, const TVector<T, d>& StartPoint, const TVector<T, d>& Dir, const TVector<T, d>& InvDir, const bool* bParallel, const T Length, const T InvLength, T& OutTime, TVector<T, d>& OutPosition)
 		{
 			return TAABB<T, d>(InMin, InMax).RaycastFast(StartPoint, Dir, InvDir, bParallel, Length, InvLength, OutTime, OutPosition);
@@ -240,6 +245,48 @@ namespace Chaos
 			return ClosestEdgePosition;
 		}
 
+		bool GetClosestEdgeVertices(int32 PlaneIndexHint, const FVec3& Position, int32& OutVertexIndex0, int32& OutVertexIndex1) const
+		{
+			if (PlaneIndexHint >= 0)
+			{
+				FReal ClosestDistanceSq = FLT_MAX;
+
+				int32 PlaneVerticesNum = NumPlaneVertices(PlaneIndexHint);
+				if (PlaneVerticesNum > 0)
+				{
+					int32 VertexIndex0 = GetPlaneVertex(PlaneIndexHint, PlaneVerticesNum - 1);
+					FVec3 P0 = GetVertex(VertexIndex0);
+
+					for (int32 PlaneVertexIndex = 0; PlaneVertexIndex < PlaneVerticesNum; ++PlaneVertexIndex)
+					{
+						const int32 VertexIndex1 = GetPlaneVertex(PlaneIndexHint, PlaneVertexIndex);
+						const FVec3 P1 = GetVertex(VertexIndex1);
+
+						const TVector<T, d> EdgePosition = FMath::ClosestPointOnLine(P0, P1, Position);
+						const FReal EdgeDistanceSq = (EdgePosition - Position).SizeSquared();
+
+						if (EdgeDistanceSq < ClosestDistanceSq)
+						{
+							OutVertexIndex0 = VertexIndex0;
+							OutVertexIndex1 = VertexIndex1;
+							ClosestDistanceSq = EdgeDistanceSq;
+						}
+
+						VertexIndex0 = VertexIndex1;
+						P0 = P1;
+					}
+					return true;
+				}
+			}
+			else
+			{
+				// @todo(chaos)
+				check(false);
+			}
+			return false;
+		}
+
+
 		// The number of planes that use the specified vertex
 		int32 NumVertexPlanes(int32 VertexIndex) const
 		{
diff --git a/Engine/Source/Runtime/Experimental/Chaos/Public/Chaos/Convex.h b/Engine/Source/Runtime/Experimental/Chaos/Public/Chaos/Convex.h
--- a/Engine/Source/Runtime/Experimental/Chaos/Public/Chaos/Convex.h
+++ b/Engine/Source/Runtime/Experimental/Chaos/Public/Chaos/Convex.h
@@ -126,7 +126,7 @@ namespace Chaos
 			return PhiWithNormalInternal(X, Normal);
 		}
 
-		FReal PhiWithNormalScaled(const FVec3& X, const FVec3& Scale, const FVec3& InvScale, FVec3& Normal) const
+		virtual FReal PhiWithNormalScaled(const FVec3& X, const FVec3& Scale, FVec3& Normal) const override
 		{
 			return PhiWithNormalScaledInternal(X, Scale, Normal);
 		}
@@ -313,6 +313,7 @@ namespace Chaos
 		// Get the nearest point on an edge of the specified face
 		FVec3 GetClosestEdgePosition(int32 PlaneIndex, const FVec3& Position) const;
 
+		bool GetClosestEdgeVertices(int32 PlaneIndex, const FVec3& Position, int32& OutVertexIndex0, int32& OutVertexIndex1) const;
 
 		// The number of planes that use the specified vertex
 		int32 NumVertexPlanes(int32 VertexIndex) const;
@@ -511,6 +512,7 @@ namespace Chaos
 		}
 
 		// Return support point on the shape
+		// @todo(chaos): do we need to support thickness?
 		FORCEINLINE FVec3 Support(const FVec3& Direction, const FReal Thickness) const
 		{
 			const int32 MaxVIdx = GetSupportVertex(Direction);
@@ -525,6 +527,16 @@ namespace Chaos
 			return FVec3(0);
 		}
 
+		FORCEINLINE FVec3 SupportScaled(const FVec3& Direction, const FReal Thickness, const FVec3& Scale) const
+		{
+			FVec3 SupportPoint = Support(Direction * Scale, 0.0f) * Scale;
+			if (Thickness > 0.0f)
+			{
+				SupportPoint += Thickness * Direction.GetSafeNormal();
+			}
+			return SupportPoint;
+		}
+
 		virtual FString ToString() const
 		{
 			return FString::Printf(TEXT("Convex"));
diff --git a/Engine/Source/Runtime/Experimental/Chaos/Public/Chaos/ImplicitObject.h b/Engine/Source/Runtime/Experimental/Chaos/Public/Chaos/ImplicitObject.h
--- a/Engine/Source/Runtime/Experimental/Chaos/Public/Chaos/ImplicitObject.h
+++ b/Engine/Source/Runtime/Experimental/Chaos/Public/Chaos/ImplicitObject.h
@@ -199,8 +199,25 @@ public:
 
 	// Explicitly non-virtual.  Must cast to derived types to target their implementation.
 	FVec3 Normal(const FVec3& x) const;
+
+	// Find the closest point on the surface, and return the separating distance and axis
 	virtual FReal PhiWithNormal(const FVec3& x, FVec3& Normal) const = 0;
-	virtual const class TAABB<FReal, 3> BoundingBox() const;
+
+	// Find the closest point on the surface, and return the separating distance and axis
+	virtual FReal PhiWithNormalScaled(const FVec3& Pos, const FVec3& Scale, FVec3& Normal) const
+	{
+		// @todo(chaos): implement for all derived types - this is not a valid solution
+		const FVec3 UnscaledX = Pos / Scale;
+		FVec3 UnscaledNormal;
+		const FReal UnscaledPhi = PhiWithNormal(UnscaledX, UnscaledNormal);
+		Normal = Scale * UnscaledNormal;
+		const FReal ScaleFactor = Normal.SafeNormalize();
+		const FReal ScaledPhi = UnscaledPhi * ScaleFactor;
+		return ScaledPhi;
+	}
+
+	virtual const FAABB3 BoundingBox() const;
+
 	bool HasBoundingBox() const { return bHasBoundingBox; }
 
 	bool IsConvex() const { return bIsConvex; }
diff --git a/Engine/Source/Runtime/Experimental/Chaos/Public/Chaos/ImplicitObjectScaled.h b/Engine/Source/Runtime/Experimental/Chaos/Public/Chaos/ImplicitObjectScaled.h
--- a/Engine/Source/Runtime/Experimental/Chaos/Public/Chaos/ImplicitObjectScaled.h
+++ b/Engine/Source/Runtime/Experimental/Chaos/Public/Chaos/ImplicitObjectScaled.h
@@ -17,8 +17,33 @@ namespace Chaos
 
 struct FMTDInfo;
 
+class FImplicitObjectInstanced : public FImplicitObject
+{
+public:
+	FImplicitObjectInstanced(int32 Flags, EImplicitObjectType InType)
+        : FImplicitObject(Flags, InType | ImplicitObjectType::IsInstanced)
+		, OuterMargin(0)
+	{
+	}
+
+	virtual TSerializablePtr<FImplicitObject> GetInnerObject() const
+	{
+		return TSerializablePtr<FImplicitObject>();
+	}
+
+	// Returns a winding order multiplier used in the manifold clipping and required when we have negative scales
+	FORCEINLINE float GetWindingOrder() const
+	{
+		return 1.0f;
+	}
+
+	
+protected:
+	FReal OuterMargin;
+};
+	
 template <typename TConcrete>
-class TImplicitObjectInstanced final : public FImplicitObject
+class TImplicitObjectInstanced final : public FImplicitObjectInstanced
 {
 public:
 	using T = typename TConcrete::TType;
@@ -31,41 +56,42 @@ public:
 
 	//needed for serialization
 	TImplicitObjectInstanced()
-		: FImplicitObject(EImplicitObject::HasBoundingBox,StaticType())
-		, OuterMargin(0)
-	{}
+		: FImplicitObjectInstanced(EImplicitObject::HasBoundingBox,StaticType())
+	{
+		this->OuterMargin = 0;
+	}
 
 	TImplicitObjectInstanced(const ObjectType&& Object, const FReal InMargin = 0)
-		: FImplicitObject(EImplicitObject::HasBoundingBox, Object->GetType() | ImplicitObjectType::IsInstanced)
+		: FImplicitObjectInstanced(EImplicitObject::HasBoundingBox, Object->GetType())
 		, MObject(MoveTemp(Object))
-		, OuterMargin(InMargin)
 	{
 		ensure(IsInstanced(MObject->GetType()) == false);	//cannot have an instance of an instance
 		this->bIsConvex = MObject->IsConvex();
 		this->bDoCollide = MObject->GetDoCollide();
+		this->OuterMargin = InMargin;
 		SetMargin(OuterMargin + MObject->GetMargin());
 	}
 
 	TImplicitObjectInstanced(const ObjectType& Object, const FReal InMargin = 0)
-		: FImplicitObject(EImplicitObject::HasBoundingBox,Object->GetType() | ImplicitObjectType::IsInstanced)
+		: FImplicitObjectInstanced(EImplicitObject::HasBoundingBox,Object->GetType())
 		, MObject(Object)
-		, OuterMargin(InMargin)
 	{
 		ensure(IsInstanced(MObject->GetType()) == false);	//cannot have an instance of an instance
 		this->bIsConvex = MObject->IsConvex();
 		this->bDoCollide = MObject->GetDoCollide();
+		this->OuterMargin = InMargin;
 		SetMargin(OuterMargin + MObject->GetMargin());
 	}
 
 	TImplicitObjectInstanced(TImplicitObjectInstanced<TConcrete>&& Other)
-		: FImplicitObject(EImplicitObject::HasBoundingBox, Other.MObject->GetType() | ImplicitObjectType::IsInstanced)
+		: FImplicitObjectInstanced(EImplicitObject::HasBoundingBox, Other.MObject->GetType())
 		, MObject(MoveTemp(Other.MObject))
-		, OuterMargin(Other.OuterMargin)
 	{
 		ensureMsgf((IsScaled(MObject->GetType()) == false), TEXT("Scaled objects should not contain each other."));
 		ensureMsgf((IsInstanced(MObject->GetType()) == false), TEXT("Scaled objects should not contain instances."));
 		this->bIsConvex = Other.MObject->IsConvex();
 		this->bDoCollide = Other.MObject->GetDoCollide();
+		this->OuterMargin = Other.OuterMargin;
 		SetMargin(Other.GetMargin());
 	}
 
@@ -74,6 +100,11 @@ public:
 		return TConcrete::StaticType() | ImplicitObjectType::IsInstanced;
 	}
 
+	virtual TSerializablePtr<FImplicitObject> GetInnerObject() const override
+	{
+		return MakeSerializable(MObject);
+	}
+	
 	const TConcrete* GetInstancedObject() const
 	{
 		return MObject.Get();
@@ -127,16 +158,16 @@ public:
 		return MObject->Support(Direction, Thickness); 
 	}
 
-	// The support position from the specified direction, if the shape is reduced by the margin
-	FORCEINLINE TVector<T, d> SupportCore(const TVector<T, d>& Direction, float InMargin) const
+	// this shouldn't be called, but is required until we remove the explicit function implementations in CollisionResolution.cpp
+	FORCEINLINE TVector<T, d> SupportScaled(const TVector<T, d>& Direction, const T Thickness, const FVec3& Scale) const
 	{
-		return MObject->SupportCore(Direction, InMargin);
+		return MObject->SupportScaled(Direction, Thickness, Scale);
 	}
 
-	// Returns a winding order multiplier used in the manifold clipping and required when we have negative scales
-	FORCEINLINE float GetWindingOrder() const
+	// The support position from the specified direction, if the shape is reduced by the margin
+	FORCEINLINE TVector<T, d> SupportCore(const TVector<T, d>& Direction, float InMargin) const
 	{
-		return 1.0f;
+		return MObject->SupportCore(Direction, InMargin);
 	}
 
 	virtual const TAABB<T, d> BoundingBox() const override 
@@ -206,6 +237,10 @@ public:
 		return MObject->GetClosestEdgePosition(PlaneIndex, Position);
 	}
 
+	bool GetClosestEdgeVertices(int32 PlaneIndexHint, const FVec3& Position, int32& OutVertexIndex0, int32& OutVertexIndex1) const
+	{
+		return MObject->GetClosestEdgeVertices(PlaneIndexHint, Position, OutVertexIndex0, OutVertexIndex1);
+	}
 
 	// The number of planes that use the specified vertex
 	int32 NumVertexPlanes(int32 VertexIndex) const
@@ -255,7 +290,6 @@ public:
 
 protected:
 	ObjectType MObject;
-	FReal OuterMargin;
 
 	static TImplicitObjectInstanced<TConcrete>* CopyHelper(const TImplicitObjectInstanced<TConcrete>* Obj)
 	{
@@ -263,9 +297,66 @@ protected:
 	}
 };
 
+class FImplicitObjectScaled : public FImplicitObject
+{
+public:
+	FImplicitObjectScaled(int32 Flags, EImplicitObjectType InType)
+		: FImplicitObject(Flags, InType | ImplicitObjectType::IsScaled)
+		, MScale(1)
+		, MInvScale(1)
+		, OuterMargin(0.0f)
+		, MLocalBoundingBox(FAABB3::EmptyAABB())
+	{
+	}
+
+	virtual TSerializablePtr<FImplicitObject> GetInnerObject() const
+	{
+		return TSerializablePtr<FImplicitObject>();
+	}
+
+	// Returns a winding order multiplier used in the manifold clipping and required when we have negative scales
+	FORCEINLINE float GetWindingOrder() const
+	{
+		const FVec3 SignVector = MScale.GetSignVector();
+		return SignVector.X * SignVector.Y * SignVector.Z;
+	}
+
+	const FVec3& GetScale() const
+	{
+		return MScale;
+	}
+
+	const FVec3& GetInvScale() const
+	{
+		return MInvScale;
+	}
+	
+	virtual const FAABB3 BoundingBox() const override
+	{
+		return MLocalBoundingBox;
+	}
+
+	const FReal GetVolume() const
+	{
+		// TODO: More precise volume!
+		return BoundingBox().GetVolume();
+	}
+
+	const FMatrix33 GetInertiaTensor(const FReal Mass) const
+	{
+		// TODO: More precise inertia!
+		return BoundingBox().GetInertiaTensor(Mass);
+	}
+
+protected:
+	FVec3 MScale;
+	FVec3 MInvScale;
+	FReal OuterMargin;	//Allows us to inflate the instance before the scale is applied. This is useful when sweeps need to apply a non scale on a geometry with uniform thickness
+	FAABB3 MLocalBoundingBox;
+};
 
 template<typename TConcrete, bool bInstanced = true>
-class TImplicitObjectScaled final : public FImplicitObject
+class TImplicitObjectScaled final : public FImplicitObjectScaled
 {
 public:
 	using T = typename TConcrete::TType;
@@ -277,10 +368,9 @@ public:
 	using FImplicitObject::GetTypeName;
 
 	TImplicitObjectScaled(ObjectType Object, const TVector<T, d>& Scale, T InMargin = 0)
-	    : FImplicitObject(EImplicitObject::HasBoundingBox, Object->GetType() | ImplicitObjectType::IsScaled)
+	    : FImplicitObjectScaled(EImplicitObject::HasBoundingBox, Object->GetType())
 	    , MObject(MoveTemp(Object))
 		, MSharedPtrForRefCount(nullptr)
-		, OuterMargin(InMargin)
 	{
 		ensureMsgf((IsScaled(MObject->GetType()) == false), TEXT("Scaled objects should not contain each other."));
 		ensureMsgf((IsInstanced(MObject->GetType()) == false), TEXT("Scaled objects should not contain instances."));
@@ -294,14 +384,14 @@ public:
 		}
 		this->bIsConvex = MObject->IsConvex();
 		this->bDoCollide = MObject->GetDoCollide();
+		this->OuterMargin = InMargin;
 		SetScale(Scale);
 	}
 
 	TImplicitObjectScaled(TSharedPtr<TConcrete, ESPMode::ThreadSafe> Object, const TVector<T, d>& Scale, T InMargin = 0)
-	    : FImplicitObject(EImplicitObject::HasBoundingBox, Object->GetType() | ImplicitObjectType::IsScaled)
+	    : FImplicitObjectScaled(EImplicitObject::HasBoundingBox, Object->GetType())
 	    , MObject(MakeSerializable<TConcrete, ESPMode::ThreadSafe>(Object))
 		, MSharedPtrForRefCount(Object)
-		, OuterMargin(InMargin)
 	{
 		ensureMsgf((IsScaled(MObject->GetType()) == false), TEXT("Scaled objects should not contain each other."));
 		ensureMsgf((IsInstanced(MObject->GetType()) == false), TEXT("Scaled objects should not contain instances."));
@@ -315,37 +405,39 @@ public:
 		}
 		this->bIsConvex = MObject->IsConvex();
 		this->bDoCollide = MObject->GetDoCollide();
+		this->OuterMargin = InMargin;
 		SetScale(Scale);
 	}
 
 	TImplicitObjectScaled(ObjectType Object, TUniquePtr<Chaos::FImplicitObject> &&ObjectOwner, const TVector<T, d>& Scale, T InMargin = 0)
-	    : FImplicitObject(EImplicitObject::HasBoundingBox, Object->GetType() | ImplicitObjectType::IsScaled)
+	    : FImplicitObjectScaled(EImplicitObject::HasBoundingBox, Object->GetType())
 	    , MObject(Object)
 		, MSharedPtrForRefCount(nullptr)
-		, OuterMargin(InMargin)
 	{
 		ensureMsgf((IsScaled(MObject->GetType(true)) == false), TEXT("Scaled objects should not contain each other."));
 		ensureMsgf((IsInstanced(MObject->GetType(true)) == false), TEXT("Scaled objects should not contain instances."));
 		this->bIsConvex = Object->IsConvex();
 		this->bDoCollide = MObject->GetDoCollide();
+		this->OuterMargin = InMargin;
 		SetScale(Scale);
 	}
 
 	TImplicitObjectScaled(const TImplicitObjectScaled<TConcrete, bInstanced>& Other) = delete;
 	TImplicitObjectScaled(TImplicitObjectScaled<TConcrete, bInstanced>&& Other)
-	    : FImplicitObject(EImplicitObject::HasBoundingBox, Other.MObject->GetType() | ImplicitObjectType::IsScaled)
+	    : FImplicitObjectScaled(EImplicitObject::HasBoundingBox, Other.MObject->GetType() | ImplicitObjectType::IsScaled)
 	    , MObject(MoveTemp(Other.MObject))
 		, MSharedPtrForRefCount(MoveTemp(Other.MSharedPtrForRefCount))
-	    , MScale(Other.MScale)
-		, MInvScale(Other.MInvScale)
-		, OuterMargin(Other.OuterMargin)
-	    , MLocalBoundingBox(MoveTemp(Other.MLocalBoundingBox))
 	{
 		ensureMsgf((IsScaled(MObject->GetType()) == false), TEXT("Scaled objects should not contain each other."));
 		ensureMsgf((IsInstanced(MObject->GetType()) == false), TEXT("Scaled objects should not contain instances."));
 		this->bIsConvex = Other.MObject->IsConvex();
 		this->bDoCollide = Other.MObject->GetDoCollide();
-		SetMargin(Other.GetMargin());
+		this->OuterMargin = Other.OuterMargin;
+		this->MScale = Other.MScale;
+        this->MInvScale = Other.MInvScale;
+        this->OuterMargin = Other.OuterMargin;
+        this->MLocalBoundingBox = Other.MLocalBoundingBox;
+        SetMargin(Other.GetMargin());
 	}
 	~TImplicitObjectScaled() {}
 
@@ -408,6 +500,11 @@ public:
 		}
 	}
 
+	virtual TSerializablePtr<FImplicitObject> GetInnerObject() const override
+	{
+		return MakeSerializable(MObject);
+	}
+
 	const TConcrete* GetUnscaledObject() const
 	{
 		return MObject.Get();
@@ -417,24 +514,10 @@ public:
 	{
 		return (MObject->GetRadius() > 0.0f) ? Margin : 0.0f;
 	}
-
+	
 	virtual T PhiWithNormal(const TVector<T, d>& X, TVector<T, d>& Normal) const override
 	{
-		// @todo(chaos): support scaled PhiWithNormal on all types
-		if (const FImplicitConvex3* Convex = MObject->template GetObject<FImplicitConvex3>())
-		{
-			return Convex->PhiWithNormalScaled(X, MScale, MInvScale, Normal);
-		}
-		else
-		{
-			const TVector<T, d> UnscaledX = MInvScale * X;
-			TVector<T, d> UnscaledNormal;
-			const T UnscaledPhi = MObject->PhiWithNormal(UnscaledX, UnscaledNormal);
-			Normal = MScale * UnscaledNormal;
-			const T ScaleFactor = Normal.SafeNormalize();
-			const T ScaledPhi = UnscaledPhi * ScaleFactor;
-			return ScaledPhi;
-		}
+		return MObject->PhiWithNormalScaled(X, MScale, Normal);
 	}
 
 	virtual bool Raycast(const TVector<T, d>& StartPoint, const TVector<T, d>& Dir, const T Length, const T Thickness, T& OutTime, TVector<T, d>& OutPosition, TVector<T, d>& OutNormal, int32& OutFaceIndex) const override
@@ -537,13 +620,13 @@ public:
 	// Get the index of the plane that most opposes the normal
 	int32 GetMostOpposingPlane(const FVec3& Normal) const
 	{
-		return MObject->GetMostOpposingPlane(GetInverseScaledNormal(Normal));
+		return MObject->GetMostOpposingPlane(GetUnscaledNormal(Normal));
 	}
 
 	// Get the index of the plane that most opposes the normal, assuming it passes through the specified vertex
 	int32 GetMostOpposingPlaneWithVertex(int32 VertexIndex, const FVec3& Normal) const
 	{
-		return MObject->GetMostOpposingPlane(VertexIndex, GetInverseScaledNormal(Normal));
+		return MObject->GetMostOpposingPlane(VertexIndex, GetUnscaledNormal(Normal));
 	}
 
 	// Get the nearest point on an edge of the specified face
@@ -552,6 +635,11 @@ public:
 		return MObject->GetClosestEdgePosition(PlaneIndex, MInvScale * Position) * MScale;
 	}
 
+	bool GetClosestEdgeVertices(int32 PlaneIndex, const FVec3& Position, int32& OutVertexIndex0, int32& OutVertexIndex1) const
+	{
+		return MObject->GetClosestEdgeVertices(PlaneIndex, MInvScale * Position, OutVertexIndex0, OutVertexIndex1);
+	}
+
 	// The number of planes that use the specified vertex
 	int32 NumVertexPlanes(int32 VertexIndex) const
 	{
@@ -590,7 +678,7 @@ public:
 	const TPlaneConcrete<FReal, 3> GetPlane(int32 FaceIndex) const
 	{
 		const TPlaneConcrete<FReal, 3> InnerPlane = MObject->GetPlane(FaceIndex);
-		return TPlaneConcrete<FReal, 3>(MScale * InnerPlane.X(), GetScaledNormal(InnerPlane.Normal()));
+		return TPlaneConcrete<FReal, 3>::MakeScaledUnsafe(InnerPlane, MScale);	// "Unsafe" means scale has no zeros
 	}
 
 	// Get the vertex at the specified index (e.g., indices from GetPlaneVertex)
@@ -698,19 +786,10 @@ public:
 		return MObject->SupportCoreScaled(Direction, InMargin, MScale);
 	}
 
-	// Returns a winding order multiplier used in the manifold clipping and required when we have negative scales
-	FORCEINLINE float GetWindingOrder() const
-	{
-		const FVec3 SignVector = MScale.GetSignVector();
-		return SignVector.X * SignVector.Y * SignVector.Z;
-	}
-
-	const TVector<T, d>& GetScale() const { return MScale; }
-	const TVector<T, d>& GetInvScale() const { return MInvScale; }
-	void SetScale(const TVector<T, d>& Scale)
+	void SetScale(const FVec3& Scale)
 	{
-		constexpr T MinMagnitude = 1e-6;
-		for (int Axis = 0; Axis < d; ++Axis)
+		constexpr FReal MinMagnitude = 1e-6;
+		for (int Axis = 0; Axis < 3; ++Axis)
 		{
 			if (!CHAOS_ENSURE(FMath::Abs(Scale[Axis]) >= MinMagnitude))
 			{
@@ -727,20 +806,6 @@ public:
 		UpdateBounds();
 	}
 
-	virtual const TAABB<T, d> BoundingBox() const override { return MLocalBoundingBox; }
-
-	const FReal GetVolume() const
-	{
-		// TODO: More precise volume!
-		return BoundingBox().GetVolume();
-	}
-
-	const FMatrix33 GetInertiaTensor(const FReal Mass) const
-	{
-		// TODO: More precise inertia!
-		return BoundingBox().GetInertiaTensor(Mass);
-	}
-
 	const FVec3 GetCenterOfMass() const
 	{
 		// TODO: I'm not sure this is correct in all cases
@@ -792,20 +857,13 @@ public:
 private:
 	ObjectType MObject;
 	TSharedPtr<TConcrete, ESPMode::ThreadSafe> MSharedPtrForRefCount; // Temporary solution to force ref counting on trianglemesh from body setup.
-	TVector<T, d> MScale;
-	TVector<T, d> MInvScale;
-	T OuterMargin;	//Allows us to inflate the instance before the scale is applied. This is useful when sweeps need to apply a non scale on a geometry with uniform thickness
-	TAABB<T, d> MLocalBoundingBox;
 
 	//needed for serialization
 	TImplicitObjectScaled()
-	: FImplicitObject(EImplicitObject::HasBoundingBox, StaticType())
-	, OuterMargin(0)
+	: FImplicitObjectScaled(EImplicitObject::HasBoundingBox, StaticType())
 	{}
 	friend FImplicitObject;	//needed for serialization
 
-	friend class FImplicitObjectScaled;
-
 	static TImplicitObjectScaled<TConcrete, true>* CopyHelper(const TImplicitObjectScaled<TConcrete, true>* Obj)
 	{
 		return new TImplicitObjectScaled<TConcrete, true>(Obj->MObject, Obj->MScale, Obj->OuterMargin);
@@ -816,8 +874,8 @@ private:
 		return new TImplicitObjectScaled<TConcrete, false>(Obj->MObject->Copy(), Obj->MScale, Obj->OuterMargin);
 	}
 
-	// Convert a normal in the scaled object space into a normal in the inner object space.
-	FVec3 GetInverseScaledNormal(const TVector<T, d>& OuterNormal) const
+	// Convert a normal in the inner unscaled object space into a normal in the outer scaled object space.
+	FVec3 GetScaledNormal(const TVector<T, d>& OuterNormal) const
 	{
 		const TVector<T, d> UnscaledDirDenorm = MInvScale * OuterNormal;
 		const float LengthScale = UnscaledDirDenorm.Size();
@@ -828,8 +886,8 @@ private:
 		return UnscaledDir;
 	}
 
-	// Convert a normal in the inner object space (unscaled) into a normal in the outer scaled object space
-	FVec3 GetScaledNormal(const TVector<T, d>& InnerNormal) const
+	// Convert a normal in the outer scaled object space into a normal in the inner unscaled object space
+	FVec3 GetUnscaledNormal(const TVector<T, d>& InnerNormal) const
 	{
 		const TVector<T, d> ScaledDirDenorm = MScale * InnerNormal;
 		const float LengthScale = ScaledDirDenorm.Size();
diff --git a/Engine/Source/Runtime/Experimental/Chaos/Public/Chaos/ParticleHandleFwd.h b/Engine/Source/Runtime/Experimental/Chaos/Public/Chaos/ParticleHandleFwd.h
--- a/Engine/Source/Runtime/Experimental/Chaos/Public/Chaos/ParticleHandleFwd.h
+++ b/Engine/Source/Runtime/Experimental/Chaos/Public/Chaos/ParticleHandleFwd.h
@@ -91,6 +91,7 @@ namespace Chaos
 	using TTransientPBDGeometryCollectionParticleHandle = TPBDGeometryCollectionParticleHandleImp<T, d, false>;
 	
 	class FGenericParticleHandle;
+	class FConstGenericParticleHandle;
 
 	// TParticleIterator
 
diff --git a/Engine/Source/Runtime/Experimental/Chaos/Public/Chaos/Plane.h b/Engine/Source/Runtime/Experimental/Chaos/Public/Chaos/Plane.h
--- a/Engine/Source/Runtime/Experimental/Chaos/Public/Chaos/Plane.h
+++ b/Engine/Source/Runtime/Experimental/Chaos/Public/Chaos/Plane.h
@@ -12,6 +12,83 @@ template <typename T, int d>
 class TPlaneConcrete
 {
 public:
+
+	// Scale the plane and assume that any of the scale components could be zero
+	static TPlaneConcrete<T, d> MakeScaledSafe(const TPlaneConcrete<T, d>& Plane, const TVector<T, d>& Scale)
+	{
+		const FVec3 ScaledX = Plane.MX * Scale;
+		
+		// If all 3 scale components are non-zero we can just inverse-scale the normal
+		// If 1 scale component is zero, the normal will point in that direction of the zero scale
+		// If 2 scale components are zero, the normal will be zero along the non-zero scale direction
+		// If 3 scale components are zero, the normal will be unchanged
+		const int32 ZeroX = FMath::IsNearlyZero(Scale.X) ? 1 : 0;
+		const int32 ZeroY = FMath::IsNearlyZero(Scale.Y) ? 1 : 0;
+		const int32 ZeroZ = FMath::IsNearlyZero(Scale.Z) ? 1 : 0;
+		const int32 NumZeros = ZeroX + ZeroY + ZeroZ;
+		FVec3 ScaledN;
+		if (NumZeros == 0)
+		{
+			// All 3 scale components non-zero
+			ScaledN = FVec3(Plane.MNormal.X / Scale.X, Plane.MNormal.Y / Scale.Y, Plane.MNormal.Z / Scale.Z);
+		}
+		else if (NumZeros == 1)
+		{
+			// Exactly one Scale component is zero
+			ScaledN = FVec3(
+				(ZeroX) ? 1.0f : 0.0f,
+				(ZeroY) ? 1.0f : 0.0f,
+				(ZeroZ) ? 1.0f : 0.0f);
+		}
+		else if (NumZeros == 2)
+		{
+			// Exactly two Scale components is zero
+			ScaledN = FVec3(
+				(ZeroX) ? Plane.MNormal.X : 0.0f,
+				(ZeroY) ? Plane.MNormal.Y : 0.0f,
+				(ZeroZ) ? Plane.MNormal.Z : 0.0f);
+		}
+		else // (NumZeros == 3)
+		{
+			// All 3 scale components are zero
+			ScaledN = Plane.MNormal;
+		}
+
+		// Even after all the above, we may still get a zero normal (e.g., we scale N=(1,0,0) by S=(0,1,0))
+		const FReal ScaleN2 = ScaledN.SizeSquared();
+		if (ScaleN2 > SMALL_NUMBER)
+		{
+			ScaledN = ScaledN * FMath::InvSqrt(ScaleN2);
+		}
+		else
+		{
+			ScaledN = Plane.MNormal;
+		}
+		
+		return TPlaneConcrete<FReal, 3>(ScaledX, ScaledN);
+	}
+
+	// Scale the plane and assume that none of the scale components are zero
+	static TPlaneConcrete<T, d> MakeScaledUnsafe(const TPlaneConcrete<T, d>& Plane, const TVector<T, d>& Scale)
+	{
+		const FVec3 ScaledX = Plane.MX * Scale;
+		FVec3 ScaledN = Plane.MNormal / Scale;
+
+		// We don't handle zero scales, but we could still end up with a small normal
+		const FReal ScaleN2 = ScaledN.SizeSquared();
+		if (ScaleN2 > SMALL_NUMBER)
+		{
+			ScaledN =  ScaledN * FMath::InvSqrt(ScaleN2);
+		}
+		else
+		{
+			ScaledN = Plane.MNormal;
+		}
+
+		return TPlaneConcrete<FReal, 3>(ScaledX, ScaledN);
+	}
+
+
 	TPlaneConcrete() = default;
 	TPlaneConcrete(const TVector<T, d>& InX, const TVector<T, d>& InNormal)
 	    : MX(InX)
@@ -110,7 +187,7 @@ public:
 	const TVector<T, d>& X() const { return MX; }
 	const TVector<T, d>& Normal() const { return MNormal; }
 	const TVector<T, d>& Normal(const TVector<T, d>&) const { return MNormal; }
-	
+
 	FORCEINLINE void Serialize(FArchive& Ar)
 	{
 		Ar << MX << MNormal;
diff --git a/Engine/Source/Runtime/Experimental/Chaos/Public/Chaos/Serializable.h b/Engine/Source/Runtime/Experimental/Chaos/Public/Chaos/Serializable.h
--- a/Engine/Source/Runtime/Experimental/Chaos/Public/Chaos/Serializable.h
+++ b/Engine/Source/Runtime/Experimental/Chaos/Public/Chaos/Serializable.h
@@ -63,6 +63,12 @@ TSerializablePtr<Ret> MakeSerializable(const TUniquePtr<T>& Unique)
 	return TSerializablePtr<Ret>(Unique);
 }
 
+template <typename T>
+TSerializablePtr<T> MakeSerializable(const TSerializablePtr<T>& P)
+{
+	return P;
+}
+
 template <typename T>
 TSerializablePtr<T> MakeSerializable(const TUniquePtr<T>&& Unique) = delete;
 
diff --git a/Engine/Source/Runtime/Experimental/Chaos/Public/Chaos/Transform.h b/Engine/Source/Runtime/Experimental/Chaos/Public/Chaos/Transform.h
--- a/Engine/Source/Runtime/Experimental/Chaos/Public/Chaos/Transform.h
+++ b/Engine/Source/Runtime/Experimental/Chaos/Public/Chaos/Transform.h
@@ -147,6 +147,36 @@ class TRigidTransform<float, 3> : public FTransform
 	{
 		return ToMatrixNoScale().Inverse();
 	}
+
+	// Get the transform which maps from Other to This, ignoring the scale on both.
+	TRigidTransform<float, 3> GetRelativeTransformNoScale(const TRigidTransform<float, 3>& Other) const
+	{
+		// @todo(chaos): optimize
+		TRigidTransform<float, 3> ThisNoScale(GetTranslation(), GetRotation());
+		TRigidTransform<float, 3> OtherNoScale(Other.GetTranslation(), Other.GetRotation());
+		return ThisNoScale.GetRelativeTransform(OtherNoScale);
+	}
+
+	TVector<float, 3> TransformNormalNoScale(const TVector<float, 3>& Normal) const
+	{
+		return TransformVectorNoScale(Normal);
+	}
+
+	// Transform the normal when scale may be non-unitary. Assumes no scale components are zero.
+	TVector<float, 3> TransformNormalUnsafe(const TVector<float, 3>& Normal) const
+	{
+		const TVector<float, 3> RotatedNormal = TransformNormalNoScale(Normal);
+		const TVector<float, 3> ScaledNormal = RotatedNormal / GetScale3D();
+		const float ScaledNormal2 = ScaledNormal.SizeSquared();
+		if (ScaledNormal2 > SMALL_NUMBER)
+		{
+			return ScaledNormal * FMath::InvSqrt(ScaledNormal2);
+		}
+		else
+		{
+			return RotatedNormal;
+		}
+	}
 };
 }
 
diff --git a/Engine/Source/Runtime/Experimental/Chaos/Public/Chaos/Utilities.h b/Engine/Source/Runtime/Experimental/Chaos/Public/Chaos/Utilities.h
--- a/Engine/Source/Runtime/Experimental/Chaos/Public/Chaos/Utilities.h
+++ b/Engine/Source/Runtime/Experimental/Chaos/Public/Chaos/Utilities.h
@@ -490,7 +490,7 @@ namespace Chaos
 			const FReal EpsilonSq = Epsilon * Epsilon;
 			const FVec3 D1 = Q1 - P1;
 			const FVec3 D2 = Q2 - P2;
-			const FVec3 R = P2 - Q2;
+			const FVec3 R = P1 - P2;
 			const FReal A = FVec3::DotProduct(D1, D1);
 			const FReal B = FVec3::DotProduct(D1, D2);
 			const FReal C = FVec3::DotProduct(D1, R);
@@ -516,6 +516,7 @@ namespace Chaos
 			}
 			else
 			{
+				// Non-degenrate case - we have two lines
 				const FReal Denom = A * E - B * B;
 				if (Denom != 0.0f)
 				{
