HeadlessChaos : remove ensure check that should have been removed in the previous revision

diff --git a/Engine/Source/Runtime/Experimental/Chaos/Private/Chaos/PBDConstraintGraph.cpp b/Engine/Source/Runtime/Experimental/Chaos/Private/Chaos/PBDConstraintGraph.cpp
--- a/Engine/Source/Runtime/Experimental/Chaos/Private/Chaos/PBDConstraintGraph.cpp
+++ b/Engine/Source/Runtime/Experimental/Chaos/Private/Chaos/PBDConstraintGraph.cpp
@@ -393,7 +393,6 @@ void FPBDConstraintGraph::ComputeIslands(const TParticleView<FPBDRigidParticles>
 
 	for (auto& Particle : PBDRigids)
 	{
-		ensure(!Particle.Disabled());
 		auto* ParticleHandle = Particle.Handle();
 		int32 Idx = ParticleToNodeIndex[ParticleHandle];  // ryan - FAILS!
 		// selective reset of islands, don't reset if has been visited due to being edge connected to earlier processed node
