GeometryCollectionProximityUtility
-- proximity is only calculated on Rigids now.
-- Added ability to propogate rigid proximity to ancestral clusters.

diff --git a/Engine/Source/Runtime/Experimental/Chaos/Private/GeometryCollection/GeometryCollectionProximityUtility.cpp b/Engine/Source/Runtime/Experimental/Chaos/Private/GeometryCollection/GeometryCollectionProximityUtility.cpp
--- a/Engine/Source/Runtime/Experimental/Chaos/Private/GeometryCollection/GeometryCollectionProximityUtility.cpp
+++ b/Engine/Source/Runtime/Experimental/Chaos/Private/GeometryCollection/GeometryCollectionProximityUtility.cpp
@@ -7,6 +7,7 @@
 #include "GeometryCollection/GeometryCollectionAlgo.h"
 #include "Math/GenericOctreePublic.h"
 #include "Math/GenericOctree.h"
+#include "GeometryCollection/GeometryCollectionClusteringUtility.h"
 
 #include <chrono>
 
@@ -127,7 +128,7 @@ void FGeometryCollectionProximityUtility::UpdateProximity(FGeometryCollection* G
 
 	//
 	// Create a FaceTransformDataArray for fast <FaceIndex, TransformIndex. lookup
-	// It only contains faces for GEOMETRY && !CLUSTERED
+	// It only contains faces for GEOMETRY && RIGID, ie. not embedded geometry and not clustered
 	//
 	int32 NumFaces = GeometryCollection->NumElements(FGeometryCollection::FacesGroup);
 	for (int32 IdxFace = 0; IdxFace < NumFaces; ++IdxFace)
@@ -136,7 +137,7 @@ void FGeometryCollectionProximityUtility::UpdateProximity(FGeometryCollection* G
 
 		//		UE_LOG(LogChaosProximity, Log, TEXT("IdxFace = %d, TransformIndex = %d>"), IdxFace, TransformIndex);
 
-		if (GeometryCollection->IsGeometry(TransformIndex) && !GeometryCollection->IsClustered(TransformIndex))
+		if (GeometryCollection->IsGeometry(TransformIndex) && GeometryCollection->IsRigid(TransformIndex))
 		{
 			//			UE_LOG(LogChaosProximity, Log, TEXT("ADDING TO FACETRANSFORMDATAARRAY"));
 
@@ -183,6 +184,7 @@ void FGeometryCollectionProximityUtility::UpdateProximity(FGeometryCollection* G
 #endif
 
 	// Build reverse map between TransformIdx and GeometryGroup index
+	// #todo We already have this? TransformToGeometryIndex
 	EnterProgressFrame();
 	TMap<int32, int32> GeometryGroupIndexMap;
 	int32 NumGeometries = GeometryCollection->NumElements(FGeometryCollection::GeometryGroup);
@@ -573,3 +575,35 @@ void FGeometryCollectionProximityUtility::UpdateProximity(FGeometryCollection* G
 	std::chrono::duration<double> elapsed = finish - start;
 	UE_LOG(LogChaosProximity, Log, TEXT("Elapsed Time = %fs>"), elapsed.count());
 }
+
+
+void CollectTopNodeConnections(FGeometryCollection* GeometryCollection, int32 Index, TSet<int32>& OutConnections, int32 OperatingLevel)
+{
+	const TManagedArray<int32>& Level = GeometryCollection->GetAttribute<int32>("Level", FGeometryCollection::TransformGroup);
+	if (OperatingLevel < 0)
+	{
+		OperatingLevel = Level[Index];
+	}
+	
+	const TManagedArray<TSet<int32>>& Children = GeometryCollection->Children;
+	if (GeometryCollection->IsRigid(Index)) // leaf node
+	{
+		const TManagedArray<TSet<int32>>& Proximity = GeometryCollection->GetAttribute<TSet<int32>>("Proximity", FGeometryCollection::GeometryGroup);
+		const TManagedArray<int32>& GeometryToTransformIndex = GeometryCollection->TransformIndex;
+		const TManagedArray<int32>& TransformToGeometryIndex = GeometryCollection->TransformToGeometryIndex;
+
+
+		for (int32 Neighbor : Proximity[TransformToGeometryIndex[Index]])
+		{
+			int32 NeighborTransformIndex = GeometryToTransformIndex[Neighbor];
+			OutConnections.Add(FGeometryCollectionClusteringUtility::GetParentOfBoneAtSpecifiedLevel(GeometryCollection, NeighborTransformIndex, OperatingLevel));
+		}
+	}
+	else
+	{
+		for (int32 ChildIndex : Children[Index])
+		{
+			CollectTopNodeConnections(GeometryCollection, ChildIndex, OutConnections, OperatingLevel);
+		}
+	}
+}
diff --git a/Engine/Source/Runtime/Experimental/Chaos/Public/GeometryCollection/GeometryCollectionProximityUtility.h b/Engine/Source/Runtime/Experimental/Chaos/Public/GeometryCollection/GeometryCollectionProximityUtility.h
--- a/Engine/Source/Runtime/Experimental/Chaos/Public/GeometryCollection/GeometryCollectionProximityUtility.h
+++ b/Engine/Source/Runtime/Experimental/Chaos/Public/GeometryCollection/GeometryCollectionProximityUtility.h
@@ -81,5 +81,8 @@ public:
 
 	static bool IsPointInsideOfTriangle(const FVector& P, const FVector& Vertex0, const FVector& Vertex1, const FVector& Vertex2, float Threshold);
 
+	// For any transform index, get the connected indices at the operating level. If OperatingLevel is -1, the operating level will be determined by the level of the Indexed node.
+	static void LevelConnectivity(FGeometryCollection* GeometryCollection, int32 Index, TSet<int32>& OutConnections, int32 OperatingLevel = -1);
+
 private:
 };
