Chaos - change some debug draw defaults

diff --git a/Engine/Source/Runtime/Experimental/Chaos/Private/PBDRigidsSolver.cpp b/Engine/Source/Runtime/Experimental/Chaos/Private/PBDRigidsSolver.cpp
--- a/Engine/Source/Runtime/Experimental/Chaos/Private/PBDRigidsSolver.cpp
+++ b/Engine/Source/Runtime/Experimental/Chaos/Private/PBDRigidsSolver.cpp
@@ -60,9 +60,9 @@ Chaos::DebugDraw::FChaosDebugDrawSettings ChaosSolverDebugDebugDrawSettings(
 	/* FontScale =			*/ 1.5f,
 	/* ShapeThicknesScale = */ 1.0f,
 	/* PointSize =			*/ 5.0f,
-	/* VelScale =			*/ 1.0f,
+	/* VelScale =			*/ 0.0f,
 	/* AngVelScale =		*/ 0.0f,
-	/* ImpulseScale =		*/ 0.05f,
+	/* ImpulseScale =		*/ 0.0f,
 	/* DrawPriority =		*/ 10.0f
 );
 FAutoConsoleVariableRef CVarChaosSolverArrowSize(TEXT("p.Chaos.Solver.DebugDraw.ArrowSize"), ChaosSolverDebugDebugDrawSettings.ArrowSize, TEXT("ArrowSize."));
