Chaos - fix incremental manifold restitution that I broke recently

diff --git a/Engine/Source/Runtime/Experimental/Chaos/Private/Chaos/Collision/PBDCollisionConstraint.cpp b/Engine/Source/Runtime/Experimental/Chaos/Private/Chaos/Collision/PBDCollisionConstraint.cpp
--- a/Engine/Source/Runtime/Experimental/Chaos/Private/Chaos/Collision/PBDCollisionConstraint.cpp
+++ b/Engine/Source/Runtime/Experimental/Chaos/Private/Chaos/Collision/PBDCollisionConstraint.cpp
@@ -211,7 +211,7 @@ namespace Chaos
 		// Calculate and store the data required for static friction and restitution in PushOut
 		//	- contact point on second shape at previous location (to enforce static friction)
 		//	- initial and previous separations (for restitution)
-		// @todo(chaos): we shouldn't have to recalculate local space positions - that were available in collision update
+		// @todo(chaos): we shouldn't have to recalculate local space positions - they were available in collision update
 		check(ManifoldPoint.ContactPoint.ContactNormalOwnerIndex >= 0);
 		check(ManifoldPoint.ContactPoint.ContactNormalOwnerIndex < 2);
 		const FRigidTransform3& PlaneTransform = ImplicitTransform[ManifoldPoint.ContactPoint.ContactNormalOwnerIndex];
@@ -228,11 +228,13 @@ namespace Chaos
 		ManifoldPoint.CoMContactNormal = CoMContactNormal;
 
 		// Calculate the initial contact velocity for use in restitution
-		// Should we use Previous V and W?
+		// We use PreV and PreW to support incremental manifold generation. In this case, manifold points
+		// can be added after we have already run some solver iterations, which gives us an incorrect initial
+		// velocity if we just use V and W (one-shots will work the same either way since V=PreV on first pass)
 		const FRigidTransform3 WorldTransform0 = FParticleUtilities::GetActorWorldTransform(Particle0);
 		const FRigidTransform3 WorldTransform1 = FParticleUtilities::GetActorWorldTransform(Particle1);
-		const FVec3 WorldContactVel0 = Particle0->V() + FVec3::CrossProduct(Particle0->W(), WorldTransform0.GetRotation() * LocalContactPoint0);
-		const FVec3 WorldContactVel1 = Particle1->V() + FVec3::CrossProduct(Particle1->W(), WorldTransform1.GetRotation() * LocalContactPoint1);
+		const FVec3 WorldContactVel0 = Particle0->PreV() + FVec3::CrossProduct(Particle0->PreW(), WorldTransform0.GetRotation() * LocalContactPoint0);
+		const FVec3 WorldContactVel1 = Particle1->PreV() + FVec3::CrossProduct(Particle1->PreW(), WorldTransform1.GetRotation() * LocalContactPoint1);
 		const FVec3 WorldContactNormal = (ManifoldPoint.ContactPoint.ContactNormalOwnerIndex == 0) ? WorldTransform0.GetRotation() * LocalContactNormal : WorldTransform1.GetRotation() * LocalContactNormal;
 		const FReal WorldContactVelNorm = FVec3::DotProduct(WorldContactVel0 - WorldContactVel1, WorldContactNormal);
 		ManifoldPoint.InitialContactVelocity = WorldContactVelNorm;
