CIS fix - bad merge of MassProperties.h

diff --git a/Engine/Source/Runtime/Experimental/Chaos/Public/Chaos/MassProperties.h b/Engine/Source/Runtime/Experimental/Chaos/Public/Chaos/MassProperties.h
--- a/Engine/Source/Runtime/Experimental/Chaos/Public/Chaos/MassProperties.h
+++ b/Engine/Source/Runtime/Experimental/Chaos/Public/Chaos/MassProperties.h
@@ -52,7 +52,6 @@ namespace Chaos
 	// Combine a list of transformed inertia tensors into a single inertia. 
 	// NOTE: If there is more than one item in the list, the output may be non-diagonaly and will have a zero rotation.
 	// If there is only 1 item in the list it will return it directly, so the rotation of mass may be non-zero.
-	template<class T, int d>
 	FMassProperties CHAOS_API CombineWorldSpace(const TArray<FMassProperties>& MPArray);
 
 
