Fracture Tool
-- Fixed indexing bug in voronoi initialization.
-- Improved dangling cluster removal

diff --git a/Engine/Plugins/Experimental/ChaosEditor/Source/FractureEditor/Private/FractureToolAutoCluster.cpp b/Engine/Plugins/Experimental/ChaosEditor/Source/FractureEditor/Private/FractureToolAutoCluster.cpp
--- a/Engine/Plugins/Experimental/ChaosEditor/Source/FractureEditor/Private/FractureToolAutoCluster.cpp
+++ b/Engine/Plugins/Experimental/ChaosEditor/Source/FractureEditor/Private/FractureToolAutoCluster.cpp
@@ -130,6 +130,7 @@ void FVoronoiPartitioner::KMeansPartition(int32 InPartitionCount)
 		{
 			break;
 		}
+		UE_LOG(LogTemp, Warning, TEXT("Iter %d"), Iteration);
 	}
 }
 
@@ -295,9 +296,9 @@ void FVoronoiPartitioner::InitializePartitions()
 
 	PartitionCount = FMath::Min(PartitionCount, TransformIndices.Num());
 
-	// Set initial partition centers as regular selects from the vertex set
+	// Set initial partition centers as selects from the vertex set
 	PartitionCenters.SetNum(PartitionCount);
-	int32 TransformStride = FMath::CeilToInt(float(TransformIndices.Num()) / float(PartitionCount));
+	int32 TransformStride = FMath::Min(1,FMath::FloorToInt(float(TransformIndices.Num()) / float(PartitionCount)));
 	for (int32 PartitionIndex = 0; PartitionIndex < PartitionCount; ++PartitionIndex)
 	{
 		PartitionCenters[PartitionIndex] = Centroids[TransformStride * PartitionIndex];
diff --git a/Engine/Source/Runtime/Experimental/Chaos/Private/GeometryCollection/GeometryCollectionClusteringUtility.cpp b/Engine/Source/Runtime/Experimental/Chaos/Private/GeometryCollection/GeometryCollectionClusteringUtility.cpp
--- a/Engine/Source/Runtime/Experimental/Chaos/Private/GeometryCollection/GeometryCollectionClusteringUtility.cpp
+++ b/Engine/Source/Runtime/Experimental/Chaos/Private/GeometryCollection/GeometryCollectionClusteringUtility.cpp
@@ -858,7 +858,6 @@ void FGeometryCollectionClusteringUtility::RemoveDanglingClusters(FGeometryColle
 {
 	check(GeometryCollection);
 
-	const TManagedArray<TSet<int32>>& Children = GeometryCollection->Children;
 	const TManagedArray<FTransform>& Transforms = GeometryCollection->Transform;
 	const TManagedArray<int32>& SimulationType = GeometryCollection->SimulationType;
 
@@ -868,7 +867,9 @@ void FGeometryCollectionClusteringUtility::RemoveDanglingClusters(FGeometryColle
 	{
 		if(GeometryCollection->IsClustered(Idx))
 		{
-			if (Children[Idx].Num() == 0)
+			TArray<int32> LeafBones;
+			GetLeafBones(GeometryCollection, Idx, true, LeafBones);
+			if (LeafBones.Num() == 0)
 			{
 				DeletionList.Add(Idx);
 			}
@@ -877,6 +878,7 @@ void FGeometryCollectionClusteringUtility::RemoveDanglingClusters(FGeometryColle
 
 	if (DeletionList.Num())
 	{
+		DeletionList.Sort();
 		GeometryCollection->RemoveElements(FGeometryCollection::TransformGroup, DeletionList);
 	}
 }
