Speed up TTriangleMesh<T>::GetVertexImportanceOrdering by skipping all pre-sorting before the sort-by-distance-from-center, and switching that from StableSort to just Sort.  Any pre-sorting would have only mattered for points where the floating point distances were exactly the same, which seems very unlikely to happen with any reliability ... for almost all of the points (except co-incident points) I think the sort by distance must have fully determined the order.
#rnx

diff --git a/Engine/Source/Runtime/Experimental/Chaos/Private/Chaos/TriangleMesh.cpp b/Engine/Source/Runtime/Experimental/Chaos/Private/Chaos/TriangleMesh.cpp
--- a/Engine/Source/Runtime/Experimental/Chaos/Private/Chaos/TriangleMesh.cpp
+++ b/Engine/Source/Runtime/Experimental/Chaos/Private/Chaos/TriangleMesh.cpp
@@ -995,62 +995,6 @@ TArray<int32> FTriangleMesh::GetVertexImportanceOrdering(
 		return PointOrder;
 	}
 
-	// A linear ordering biases towards the order in which the vertices were
-	// authored, which is likely to be topologically adjacent.  Randomize the
-	// initial ordering.
-	FRandomStream Rand(NumPoints);
-	for (int32 i = 0; i < NumPoints; i++)
-	{
-		const int32 j = Rand.RandRange(0, NumPoints - 1);
-		Swap(PointOrder[i], PointOrder[j]);
-	}
-
-	// Find particles with no connectivity and send them to the back of the
-	// list.  We penalize free points, but we don't exclude them.  It's
-	// possible they were added for extra resolution.
-	TArray<uint8> Rank;
-	Rank.AddUninitialized(NumPoints);
-	AscendingPredicate<uint8> AscendingRankPred(Rank, Offset); // low to high
-	bool FoundFreeParticle = false;
-	for (int i = 0; i < NumPoints; i++)
-	{
-		const int32 Idx = PointOrder[i];
-		const TSet<uint32>* Neighbors = MPointToNeighborsMap.Find(Idx);
-		const bool IsFree = Neighbors == nullptr || Neighbors->Num() == 0;
-		Rank[Idx - Offset] = IsFree ? 1 : 0;
-		FoundFreeParticle |= IsFree;
-	}
-	if (FoundFreeParticle)
-	{
-		StableSort(&PointOrder[0], NumPoints, AscendingRankPred);
-	}
-
-	// Sort the pointOrder array by pointCurvatures so that points attached
-	// to edges with the highest curvatures come first.
-	if (PointCurvatures.Num() > 0)
-	{
-		// Curvature is measured by the angle between face normals.  0.0 means
-		// coplanar, angles approaching M_PI are more creased.  So, sort from
-		// high to low.
-		check(PointCurvatures.Num() == MNumIndices);
-
-		// PointCurvatures is sized to the index range of the mesh.  That may
-		// not include all free particles.  If the DescendingPredicate gets an
-		// index that is out of bounds of the curvature array, it'll use
-		// -FLT_MAX, which will put free particles at the end.  In order to get
-		// PointCurvatures to line up with PointOrder indices, offset them by
-		// -MStartIdx when not RestrictToLocalIndexRange.
-
-		// PointCurvatures[0] always corresponds to Points[MStartIdx]
-		DescendingPredicate<FReal> curvaturePred(PointCurvatures, MStartIdx); // high to low
-
-		// The indexing scheme used for sorting is a little complicated.  The pointOrder
-		// array contains point indices.  The sorting binary predicate is handed 2 of
-		// those indices, which we use to look up values we want to sort by; curvature
-		// in this case.  So, the sort data array needs to be in the original ordering.
-		StableSort(&PointOrder[0], PointOrder.Num(), curvaturePred);
-	}
-
 	// Move the points to the origin to avoid floating point aliasing far away
 	// from the origin.
 	FAABB3 Bbox(Points[0], Points[0]);
@@ -1082,7 +1026,7 @@ TArray<int32> FTriangleMesh::GetVertexImportanceOrdering(
 		Dist[i] = (LocalPoints[i] - LocalCenter).SizeSquared();
 	}
 	DescendingPredicate<FReal> DescendingDistPred(Dist); // high to low
-	StableSort(&PointOrder[0], NumPoints, DescendingDistPred);
+	Sort(&PointOrder[0], NumPoints, DescendingDistPred);
 
 	// If all points are coincident, return early.
 	const FReal MaxBBoxDim = LocalBBox.Extents().Max();
@@ -1106,6 +1050,8 @@ TArray<int32> FTriangleMesh::GetVertexImportanceOrdering(
 		CoincidentVertices->Reserve(64); // a guess
 	}
 	int32 NumCoincident = 0;
+	TArray<uint8> Rank;
+	AscendingPredicate<uint8> AscendingRankPred(Rank, Offset); // low to high
 	{
 		const int64 Resolution = static_cast<int64>(floor(MaxBBoxDim / 0.01));
 		const FReal CellSize = MaxBBoxDim / Resolution;
@@ -1198,9 +1144,9 @@ TArray<int32> FTriangleMesh::GetVertexImportanceOrdering(
 		// gets coarser, stable sort will get more and more selective about
 		// which particles get promoted.
 		//
-		// Since the initial ordering was biased by curvature and distance from
+		// Since the initial ordering was biased by distance from
 		// the center, each rank is similarly ordered. That is, the first vertex
-		// to land in a cell will be the most distant, and the highest curvature.
+		// to land in a cell will be the most distant.
 		StableSort(&PointOrder[0], NumPoints - NumCoincident, AscendingRankPred);
 	} // end for
 
