Merge from UE4

Chaos - particle counters

diff --git a/Engine/Source/Runtime/Experimental/Chaos/Private/PBDRigidsSolver.cpp b/Engine/Source/Runtime/Experimental/Chaos/Private/PBDRigidsSolver.cpp
--- a/Engine/Source/Runtime/Experimental/Chaos/Private/PBDRigidsSolver.cpp
+++ b/Engine/Source/Runtime/Experimental/Chaos/Private/PBDRigidsSolver.cpp
@@ -28,6 +28,19 @@
 
 DEFINE_LOG_CATEGORY_STATIC(LogPBDRigidsSolver, Log, All);
 
+// Stat Counters
+DECLARE_DWORD_ACCUMULATOR_STAT(TEXT("NumDisabledParticles"), STAT_ChaosCounter_NumDisabledParticles, STATGROUP_ChaosCounters);
+DECLARE_DWORD_ACCUMULATOR_STAT(TEXT("NumParticles"), STAT_ChaosCounter_NumParticles, STATGROUP_ChaosCounters);
+DECLARE_DWORD_ACCUMULATOR_STAT(TEXT("NumDynamicParticles"), STAT_ChaosCounter_NumDynamicParticles, STATGROUP_ChaosCounters);
+DECLARE_DWORD_ACCUMULATOR_STAT(TEXT("NumActiveDynamicParticles"), STAT_ChaosCounter_NumActiveDynamicParticles, STATGROUP_ChaosCounters);
+DECLARE_DWORD_ACCUMULATOR_STAT(TEXT("NumKinematicParticles"), STAT_ChaosCounter_NumKinematicParticles, STATGROUP_ChaosCounters);
+DECLARE_DWORD_ACCUMULATOR_STAT(TEXT("NumStaticParticles"), STAT_ChaosCounter_NumStaticParticles, STATGROUP_ChaosCounters);
+DECLARE_DWORD_ACCUMULATOR_STAT(TEXT("NumGeomCollParticles"), STAT_ChaosCounter_NumGeometryCollectionParticles, STATGROUP_ChaosCounters);
+DECLARE_DWORD_ACCUMULATOR_STAT(TEXT("NumIslands"), STAT_ChaosCounter_NumIslands, STATGROUP_ChaosCounters);
+DECLARE_DWORD_ACCUMULATOR_STAT(TEXT("NumContacts"), STAT_ChaosCounter_NumContacts, STATGROUP_ChaosCounters);
+DECLARE_DWORD_ACCUMULATOR_STAT(TEXT("NumJoints"), STAT_ChaosCounter_NumJoints, STATGROUP_ChaosCounters);
+
+
 // DebugDraw CVars
 #if CHAOS_DEBUG_DRAW
 
@@ -254,6 +267,23 @@ namespace Chaos
 #endif
 			}
 
+			{
+				// Particle counts
+				SET_DWORD_STAT(STAT_ChaosCounter_NumDisabledParticles, MSolver->GetEvolution()->GetParticles().GetAllParticlesView().Num() - MSolver->GetEvolution()->GetParticles().GetNonDisabledView().Num());
+				SET_DWORD_STAT(STAT_ChaosCounter_NumParticles, MSolver->GetEvolution()->GetParticles().GetNonDisabledView().Num());
+				SET_DWORD_STAT(STAT_ChaosCounter_NumDynamicParticles, MSolver->GetEvolution()->GetParticles().GetNonDisabledDynamicView().Num());
+				SET_DWORD_STAT(STAT_ChaosCounter_NumActiveDynamicParticles, MSolver->GetEvolution()->GetParticles().GetActiveParticlesView().Num());
+				SET_DWORD_STAT(STAT_ChaosCounter_NumKinematicParticles, MSolver->GetEvolution()->GetParticles().GetActiveKinematicParticlesView().Num());
+				SET_DWORD_STAT(STAT_ChaosCounter_NumStaticParticles, MSolver->GetEvolution()->GetParticles().GetActiveStaticParticlesView().Num());
+				SET_DWORD_STAT(STAT_ChaosCounter_NumGeometryCollectionParticles, MSolver->GetEvolution()->GetParticles().GetGeometryCollectionParticles().Size());
+
+				// Constraint counts
+				SET_DWORD_STAT(STAT_ChaosCounter_NumIslands, MSolver->GetEvolution()->GetConstraintGraph().NumIslands());
+				SET_DWORD_STAT(STAT_ChaosCounter_NumContacts, MSolver->NumCollisionConstraints());
+				SET_DWORD_STAT(STAT_ChaosCounter_NumJoints, MSolver->NumJointConstraints());
+			}
+
+
 			{
 				SCOPE_CYCLE_COUNTER(STAT_EventDataGathering);
 				{
@@ -1154,6 +1184,19 @@ namespace Chaos
 		PullPhysicsStateForEachDirtyProxy_External([](auto){});
 	}
 
+	template <typename Traits>
+	int32 TPBDRigidsSolver<Traits>::NumJointConstraints() const
+	{
+		return JointConstraints.NumConstraints();
+	}
+	
+	template <typename Traits>
+	int32 TPBDRigidsSolver<Traits>::NumCollisionConstraints() const
+	{
+		return GetEvolution()->GetCollisionConstraints().NumConstraints();
+	}
+
+
 	template <typename Traits>
 	void TPBDRigidsSolver<Traits>::PostTickDebugDraw(FReal Dt) const
 	{
diff --git a/Engine/Source/Runtime/Experimental/Chaos/Public/ChaosStats.h b/Engine/Source/Runtime/Experimental/Chaos/Public/ChaosStats.h
--- a/Engine/Source/Runtime/Experimental/Chaos/Public/ChaosStats.h
+++ b/Engine/Source/Runtime/Experimental/Chaos/Public/ChaosStats.h
@@ -13,6 +13,8 @@ DECLARE_STATS_GROUP(TEXT("ChaosEngine"), STATGROUP_ChaosEngine, STATCAT_Advanced
 DECLARE_STATS_GROUP(TEXT("ChaosCollision"), STATGROUP_ChaosCollision, STATCAT_Advanced);
 DECLARE_STATS_GROUP(TEXT("ChaosJoint"), STATGROUP_ChaosJoint, STATCAT_Advanced);
 DECLARE_STATS_GROUP(TEXT("ChaosMinEvolution"), STATGROUP_ChaosMinEvolution, STATCAT_Advanced);
+DECLARE_STATS_GROUP(TEXT("ChaosCounters"), STATGROUP_ChaosCounters, STATCAT_Advanced);
+
 
 DECLARE_CYCLE_STAT_EXTERN(TEXT("Physics Tick"), STAT_ChaosTick, STATGROUP_Chaos, CHAOS_API);
 DECLARE_CYCLE_STAT_EXTERN(TEXT("Physics Advance"), STAT_PhysicsAdvance, STATGROUP_Chaos, CHAOS_API);
diff --git a/Engine/Source/Runtime/Experimental/Chaos/Public/PBDRigidsSolver.h b/Engine/Source/Runtime/Experimental/Chaos/Public/PBDRigidsSolver.h
--- a/Engine/Source/Runtime/Experimental/Chaos/Public/PBDRigidsSolver.h
+++ b/Engine/Source/Runtime/Experimental/Chaos/Public/PBDRigidsSolver.h
@@ -183,6 +183,8 @@ namespace Chaos
 		/**/
 		FDirtyParticlesBuffer* GetDirtyParticlesBuffer() const { return MDirtyParticlesBuffer.Get(); }
 
+		int32 NumJointConstraints() const;
+		int32 NumCollisionConstraints() const;
 
 		//Make friend with unit test code so we can verify some behavior
 		template <typename TSolver>
diff --git a/Engine/Source/Runtime/PhysicsCore/Private/ChaosScene.cpp b/Engine/Source/Runtime/PhysicsCore/Private/ChaosScene.cpp
--- a/Engine/Source/Runtime/PhysicsCore/Private/ChaosScene.cpp
+++ b/Engine/Source/Runtime/PhysicsCore/Private/ChaosScene.cpp
@@ -38,6 +38,9 @@
 DECLARE_CYCLE_STAT(TEXT("Update Kinematics On Deferred SkelMeshes"),STAT_UpdateKinematicsOnDeferredSkelMeshesChaos,STATGROUP_Physics);
 CSV_DEFINE_CATEGORY(ChaosPhysics,true);
 
+// Stat Counters
+DECLARE_DWORD_ACCUMULATOR_STAT(TEXT("NumDirtyAABBTreeElements"), STAT_ChaosCounter_NumDirtyAABBTreeElements, STATGROUP_ChaosCounters);
+
 TAutoConsoleVariable<int32> CVar_ChaosSimulationEnable(TEXT("P.Chaos.Simulation.Enable"),1,TEXT("Enable / disable chaos simulation. If disabled, physics will not tick."));
 TAutoConsoleVariable<int32> CVar_ApplyProjectSettings(TEXT("p.Chaos.Simulation.ApplySolverProjectSettings"), 1, TEXT("Whether to apply the solver project settings on spawning a solver"));
 
@@ -463,6 +466,7 @@ void FChaosScene::EndFrame()
 
 	int32 DirtyElements = DirtyElementCount(GetSpacialAcceleration()->AsChecked<SpatialAccelerationCollection>());
 	CSV_CUSTOM_STAT(ChaosPhysics,AABBTreeDirtyElementCount,DirtyElements,ECsvCustomStatOp::Set);
+	SET_DWORD_STAT(STAT_ChaosCounter_NumDirtyAABBTreeElements, DirtyElements);
 
 	check(IsCompletionEventComplete())
 	//check(PhysicsTickTask->IsComplete());
