Fix GeometryCollection cache not playing with anchors.

diff --git a/Engine/Source/Runtime/Experimental/Chaos/Private/PhysicsProxy/GeometryCollectionPhysicsProxy.cpp b/Engine/Source/Runtime/Experimental/Chaos/Private/PhysicsProxy/GeometryCollectionPhysicsProxy.cpp
--- a/Engine/Source/Runtime/Experimental/Chaos/Private/PhysicsProxy/GeometryCollectionPhysicsProxy.cpp
+++ b/Engine/Source/Runtime/Experimental/Chaos/Private/PhysicsProxy/GeometryCollectionPhysicsProxy.cpp
@@ -753,7 +753,7 @@ void FGeometryCollectionPhysicsProxy::InitializeBodiesPT(Chaos::TPBDRigidsSolver
 			Commands.Add(Cmd);
 		}
 		Parameters.InitializationCommands.Empty();
-		FieldParameterUpdateCallback(RigidsSolver);
+		FieldParameterUpdateCallback(RigidsSolver, false);
 
 		if (Parameters.InitialVelocityType == EInitialVelocityTypeEnum::Chaos_Initial_Velocity_User_Defined)
 		{
@@ -1301,59 +1301,6 @@ void FGeometryCollectionPhysicsProxy::GetRelevantParticleHandles(
 #endif
 }
 
-
-void FGeometryCollectionPhysicsProxy::PushKinematicStateToSolver()
-{
-	FGeometryDynamicCollection& Collection = GameThreadCollection;
-	if (Collection.Transform.Num())
-	{
-		TManagedArray<int32>& DynamicState = Collection.DynamicState;
-
-		for (int32 TransformIndex = 0; TransformIndex < DynamicState.Num(); TransformIndex++)
-		{
-			Chaos::TPBDRigidClusteredParticleHandle<float, 3>* Handle = SolverParticleHandles[TransformIndex];
-			if (!Handle)
-			{
-				continue;
-			}
-	
-			if (DynamicState[TransformIndex] == (int32)EObjectStateTypeEnum::Chaos_Object_Dynamic
-				&& (Handle->ObjectState() == Chaos::EObjectStateType::Kinematic || Handle->ObjectState() == Chaos::EObjectStateType::Static)
-				&& FLT_EPSILON < Handle->M())
-			{
-				Handle->SetObjectState(Chaos::EObjectStateType::Dynamic);
-				if (Parameters.InitialVelocityType == EInitialVelocityTypeEnum::Chaos_Initial_Velocity_User_Defined)
-				{
-					Handle->SetV(Collection.InitialLinearVelocity[TransformIndex]);
-					Handle->SetW(Collection.InitialAngularVelocity[TransformIndex]);
-				}
-			}
-			else if ((DynamicState[TransformIndex] == (int32)EObjectStateTypeEnum::Chaos_Object_Kinematic)
-				&& (Handle->ObjectState() == Chaos::EObjectStateType::Dynamic)
-				&& FLT_EPSILON < Handle->M())
-			{
-				Handle->SetObjectState(Chaos::EObjectStateType::Kinematic);
-			}
-			else if ((DynamicState[TransformIndex] == (int32)EObjectStateTypeEnum::Chaos_Object_Static)
-				&& (Handle->ObjectState() == Chaos::EObjectStateType::Dynamic)
-				&& FLT_EPSILON < Handle->M())
-			{
-				Handle->SetObjectState(Chaos::EObjectStateType::Static);
-			}
-			else if ((DynamicState[TransformIndex] == (int32)EObjectStateTypeEnum::Chaos_Object_Sleeping)
-				&& (Handle->ObjectState() == Chaos::EObjectStateType::Dynamic))
-			{
-				Handle->SetObjectState(Chaos::EObjectStateType::Sleeping);
-			}
-			else if ((DynamicState[TransformIndex] == (int32)EObjectStateTypeEnum::Chaos_Object_Dynamic)
-				&& (Handle->ObjectState() == Chaos::EObjectStateType::Sleeping))
-			{
-				Handle->SetObjectState(Chaos::EObjectStateType::Dynamic);
-			}
-		}
-	}
-}
-
 int32 FGeometryCollectionPhysicsProxy::CalculateHierarchyLevel(const FGeometryDynamicCollection& GeometryCollection, int32 TransformIndex) const
 {
 	int32 Level = 0;
@@ -2514,11 +2461,12 @@ void BuildSimulationData(Chaos::FErrorReporter& ErrorReporter, FGeometryCollecti
 //==============================================================================
 
 template <typename Traits>
-void FGeometryCollectionPhysicsProxy::FieldParameterUpdateCallback(Chaos::TPBDRigidsSolver<Traits>* RigidSolver)
+void FGeometryCollectionPhysicsProxy::FieldParameterUpdateCallback(Chaos::TPBDRigidsSolver<Traits>* RigidSolver, const bool bUpdateViews)
 {
 	SCOPE_CYCLE_COUNTER(STAT_ParamUpdateField_Object);
 
-	FGeometryDynamicCollection& Collection = GameThreadCollection;
+	// We are updating the Collection from the InitializeBodiesPT, so we need the PT collection
+	FGeometryDynamicCollection& Collection = PhysicsThreadCollection;
 	Chaos::FPBDPositionConstraints PositionTarget;
 	TMap<int32, int32> TargetedParticles;
 
@@ -2563,18 +2511,29 @@ void FGeometryCollectionPhysicsProxy::FieldParameterUpdateCallback(Chaos::TPBDRi
 
 						SCOPE_CYCLE_COUNTER(STAT_ParamUpdateField_DynamicState);
 						{
+							bool bHasStateChanged = false;
 							InitDynamicStateResults(ParticleHandles, FieldContext, LocalResults);
-
+							
 							static_cast<const FFieldNode<int32>*>(FieldCommand.RootNode.Get())->Evaluate(FieldContext, ResultsView);
 							for (const FFieldContextIndex& Index : FieldContext.GetEvaluatedSamples())
 							{
 								Chaos::TPBDRigidParticleHandle<float, 3>* RigidHandle = ParticleHandles[Index.Sample]->CastToRigidParticle();
 								if (RigidHandle)
 								{
-									Collection.DynamicState[HandleToTransformGroupIndex[RigidHandle]] = ResultsView[Index.Result];
+									// Update of the handles object state. No need to update 
+									// the initial velocities since it is done after this function call in InitializeBodiesPT
+									const int8 ResultState = ResultsView[Index.Result];
+									bHasStateChanged |= ReportDynamicStateResult(RigidSolver, static_cast<Chaos::EObjectStateType>(ResultState), RigidHandle,
+										false, Chaos::TVector<float, 3>(0), false, Chaos::TVector<float, 3>(0));
+
+									// Update of the Collection dynamic state. It will be used just after to set the initial velocity
+									Collection.DynamicState[HandleToTransformGroupIndex[RigidHandle]] = ResultState;
 								}
 							}
-							PushKinematicStateToSolver();
+							if (bUpdateViews)
+							{
+								UpdateSolverParticlesState(RigidSolver, bHasStateChanged);
+							}
 						}
 						CommandsToRemove.Add(CommandIndex);
 					}
@@ -2707,7 +2666,7 @@ void FGeometryCollectionPhysicsProxy::FieldForcesUpdateCallback(Chaos::TPBDRigid
 	typename Chaos::TPBDRigidsSolver<Chaos::Traits>::FParticlesType& Particles);\
 	template void FGeometryCollectionPhysicsProxy::OnRemoveFromSolver(Chaos::TPBDRigidsSolver<Chaos::Traits> *RBDSolver);\
 	template void FGeometryCollectionPhysicsProxy::BufferPhysicsResults(Chaos::TPBDRigidsSolver<Chaos::Traits>* CurrentSolver,Chaos::FDirtyGeometryCollectionData& BufferData);\
-	template void FGeometryCollectionPhysicsProxy::FieldParameterUpdateCallback(Chaos::TPBDRigidsSolver<Chaos::Traits>* RigidSolver);\
+	template void FGeometryCollectionPhysicsProxy::FieldParameterUpdateCallback(Chaos::TPBDRigidsSolver<Chaos::Traits>* RigidSolver, const bool bUpdateViews);\
 	template void FGeometryCollectionPhysicsProxy::FieldForcesUpdateCallback(Chaos::TPBDRigidsSolver<Chaos::Traits>* RigidSolver);\
 	template void FGeometryCollectionPhysicsProxy::GetRelevantParticleHandles(\
 		TArray<Chaos::TGeometryParticleHandle<float,3>*>& Handles,\
diff --git a/Engine/Source/Runtime/Experimental/Chaos/Public/PhysicsProxy/GeometryCollectionPhysicsProxy.h b/Engine/Source/Runtime/Experimental/Chaos/Public/PhysicsProxy/GeometryCollectionPhysicsProxy.h
--- a/Engine/Source/Runtime/Experimental/Chaos/Public/PhysicsProxy/GeometryCollectionPhysicsProxy.h
+++ b/Engine/Source/Runtime/Experimental/Chaos/Public/PhysicsProxy/GeometryCollectionPhysicsProxy.h
@@ -171,7 +171,7 @@ public:
 	void FieldForcesUpdateCallback(Chaos::TPBDRigidsSolver<Traits>* RigidSolver);
 
 	template <typename Traits>
-	void FieldParameterUpdateCallback(Chaos::TPBDRigidsSolver<Traits>* RigidSolver);
+	void FieldParameterUpdateCallback(Chaos::TPBDRigidsSolver<Traits>* RigidSolver, const bool bUpdateViews = true);
 
 	void UpdateKinematicBodiesCallback(const FParticlesType& InParticles, const float InDt, const float InTime, FKinematicProxy& InKinematicProxy) {}
 	void StartFrameCallback(const float InDt, const float InTime) {}
@@ -351,7 +351,7 @@ private:
 	typename Chaos::TPBDRigidsSolver<Chaos::Traits>::FParticlesType& Particles);\
 	extern template void FGeometryCollectionPhysicsProxy::OnRemoveFromSolver(Chaos::TPBDRigidsSolver<Chaos::Traits> *RBDSolver);\
 	extern template void FGeometryCollectionPhysicsProxy::BufferPhysicsResults(Chaos::TPBDRigidsSolver<Chaos::Traits>* CurrentSolver,Chaos::FDirtyGeometryCollectionData& BufferData);\
-	extern template void FGeometryCollectionPhysicsProxy::FieldParameterUpdateCallback(Chaos::TPBDRigidsSolver<Chaos::Traits>* RigidSolver);\
+	extern template void FGeometryCollectionPhysicsProxy::FieldParameterUpdateCallback(Chaos::TPBDRigidsSolver<Chaos::Traits>* RigidSolver, const bool bUpdateViews);\
 	extern template void FGeometryCollectionPhysicsProxy::FieldForcesUpdateCallback(Chaos::TPBDRigidsSolver<Chaos::Traits>* RigidSolver);\
 	extern template void FGeometryCollectionPhysicsProxy::GetRelevantParticleHandles(\
 		TArray<Chaos::TGeometryParticleHandle<float,3>*>& Handles,\
