Chaos - fix scaled FindMostOpposingFace and GetOpposingPlane

diff --git a/Engine/Source/Programs/HeadlessChaos/Private/HeadlessChaosTestMostOpposing.cpp b/Engine/Source/Programs/HeadlessChaos/Private/HeadlessChaosTestMostOpposing.cpp
--- a/Engine/Source/Programs/HeadlessChaos/Private/HeadlessChaosTestMostOpposing.cpp
+++ b/Engine/Source/Programs/HeadlessChaos/Private/HeadlessChaosTestMostOpposing.cpp
@@ -3,6 +3,7 @@
 #include "HeadlessChaosTestMostOpposing.h"
 
 #include "HeadlessChaos.h"
+#include "HeadlessChaosTestUtility.h"
 #include "Chaos/TriangleMeshImplicitObject.h"
 #include "Chaos/ImplicitObjectScaled.h"
 #include "Chaos/Convex.h"
@@ -212,4 +213,54 @@ namespace ChaosTest
 			EXPECT_EQ(Scaled.FindMostOpposingFace(Position, FVec3(-1, 0, 0), FaceIndex, 0.01), 3);	//too far to care about other face
 		}
 	}
+
+
+	// Unscaled test to accompany TestConvexBox_Scaled. See comments on that function
+	// This test has always passed whereas TestConvexBox_Scaled used to fail.
+	GTEST_TEST(FindMostOpposingFaceTests, TestConvexBox_Unsccaled)
+	{
+		const FVec3 BoxScale = FVec3(10, 5, 1);
+		const FVec3 BoxSize = FVec3(100, 100, 100);
+		const FReal Margin = 0;
+		FImplicitConvex3 Box = CreateConvexBox(BoxSize * BoxScale, Margin);
+
+		// Position is not important in this test as all faces are within TestDistance of TestPosition
+		const FVec3 TestPos = FVec3(0, 0, 0);
+		const FReal TestDistance = 10000.0f;
+
+		// Pointing mostly left to select the right-facing face. If we transform this normal into unscaled space and
+		// call FindMostOpposingFace on the core convex, it would incorrectly select the up-facing face.
+		const FVec3 TestDir = FVec3(-1.0, 0, -0.5).GetSafeNormal();
+		int32 FaceIndex = Box.FindMostOpposingFace(TestPos, TestDir, INDEX_NONE, TestDistance);
+		const FVec3 FaceNormal = Box.GetPlane(FaceIndex).Normal();
+
+		EXPECT_NEAR(FaceNormal.X, 1.0f, KINDA_SMALL_NUMBER);
+		EXPECT_NEAR(FaceNormal.Z, 0.0f, KINDA_SMALL_NUMBER);
+	}
+
+	// FindMostOpposingFace on Scaled Convex Box
+	// NOTE: Box faces normals are not affected by scale though the face positions are
+	// This tests for a specific bug: TImplicitObjectScaled::FindMostOpposingFace was calling
+	// onto FConvex::FindMostOpposingFace with a corrected normal, but this does not work - we
+	// actually need to perform the most-opposing test in scaled space.
+	GTEST_TEST(FindMostOpposingFaceTests, TestConvexBox_Scaled)
+	{
+		const FVec3 BoxScale = FVec3(10, 5, 1);
+		const FVec3 BoxSize = FVec3(100, 100, 100);
+		const FReal Margin = 0;
+		TImplicitObjectScaled<FImplicitConvex3> Box = CreateScaledConvexBox(BoxSize, BoxScale, Margin);
+
+		// Position is not important in this test as all faces are within TestDistance of TestPosition
+		const FVec3 TestPos = FVec3(0, 0, 0);
+		const FReal TestDistance = 10000.0f;
+
+		// Pointing mostly left to select the right-facing face. If we transform this normal into unscaled space and
+		// call FindMostOpposingFace on the core convex, it would incorrectly select the up-facing face.
+		const FVec3 TestDir = FVec3(-1.0, 0, -0.5).GetSafeNormal();
+		int32 FaceIndex = Box.FindMostOpposingFace(TestPos, TestDir, INDEX_NONE, TestDistance);
+		const FVec3 FaceNormal = Box.GetPlane(FaceIndex).Normal();
+
+		EXPECT_NEAR(FaceNormal.X, 1.0f, KINDA_SMALL_NUMBER);
+		EXPECT_NEAR(FaceNormal.Z, 0.0f, KINDA_SMALL_NUMBER);
+	}
 }
\ No newline at end of file
diff --git a/Engine/Source/Runtime/Experimental/Chaos/Private/Chaos/Convex.cpp b/Engine/Source/Runtime/Experimental/Chaos/Private/Chaos/Convex.cpp
--- a/Engine/Source/Runtime/Experimental/Chaos/Private/Chaos/Convex.cpp
+++ b/Engine/Source/Runtime/Experimental/Chaos/Private/Chaos/Convex.cpp
@@ -18,8 +18,6 @@ namespace Chaos
 
 	int32 FConvex::FindMostOpposingFace(const FVec3& Position, const FVec3& UnitDir, int32 HintFaceIndex, FReal SearchDist) const
 	{
-		//NOTE: this approach assumes we never have conincident planes, which at the moment is not true.
-		//Need to make sure convex hull is cleaned up so that there are n-gon faces so this is unique
 		SearchDist = FMath::Max(SearchDist,FMath::Abs(BoundingBox().Extents().GetAbsMax()) * 1e-4f);
 		//todo: use hill climbing
 		int32 MostOpposingIdx = INDEX_NONE;
@@ -43,6 +41,35 @@ namespace Chaos
 		return MostOpposingIdx;
 	}
 
+	// @todo(chaos): dedupe from above. Only difference is use of TPlaneConcrete<FReal, 3>::MakeScaledUnsafe
+	int32 FConvex::FindMostOpposingFaceScaled(const FVec3& Position, const FVec3& UnitDir, int32 HintFaceIndex, FReal SearchDist, const FVec3& Scale) const
+	{
+		// Use of Scale.Max() is a bit dodgy, but the whole thing is fuzzy anyway
+		SearchDist = FMath::Max(Scale.Max() * SearchDist, FMath::Abs(BoundingBox().Extents().GetAbsMax()) * 1e-4f);
+
+		//todo: use hill climbing
+		int32 MostOpposingIdx = INDEX_NONE;
+		FReal MostOpposingDot = TNumericLimits<FReal>::Max();
+		for (int32 Idx = 0; Idx < Planes.Num(); ++Idx)
+		{
+			const TPlaneConcrete<FReal, 3> Plane = TPlaneConcrete<FReal, 3>::MakeScaledUnsafe(Planes[Idx], Scale);
+			const FReal Distance = Plane.SignedDistance(Position);
+			if (FMath::Abs(Distance) < SearchDist)
+			{
+				// TPlane has an override for Normal() that doesn't call PhiWithNormal().
+				const FReal Dot = FVec3::DotProduct(Plane.Normal(), UnitDir);
+				if (Dot < MostOpposingDot)
+				{
+					MostOpposingDot = Dot;
+					MostOpposingIdx = Idx;
+				}
+			}
+		}
+		CHAOS_ENSURE(MostOpposingIdx != INDEX_NONE);
+		return MostOpposingIdx;
+	}
+
+
 	int32 FConvex::FindClosestFaceAndVertices(const FVec3& Position, TArray<FVec3>& FaceVertices, FReal SearchDist) const
 	{
 		//
@@ -79,8 +106,6 @@ namespace Chaos
 
 	int32 FConvex::GetMostOpposingPlane(const FVec3& Normal) const
 	{
-		// @todo(chaos): this approach assumes we never have conincident planes, which at the moment is not true.
-		// Need to make sure convex hull is cleaned up so that there are n-gon faces so this is unique
 		// @todo(chaos): use hill climbing
 		int32 MostOpposingIdx = INDEX_NONE;
 		FReal MostOpposingDot = TNumericLimits<FReal>::Max();
@@ -98,6 +123,29 @@ namespace Chaos
 		return MostOpposingIdx;
 	}
 
+	int32 FConvex::GetMostOpposingPlaneScaled(const FVec3& Normal, const FVec3& Scale) const
+	{
+		// NOTE: We cannot just call the scale-less version like we can for box, even if we unscale the normal
+		// @todo(chaos): use hill climbing
+		int32 MostOpposingIdx = INDEX_NONE;
+		FReal MostOpposingDot = TNumericLimits<FReal>::Max();
+		for (int32 Idx = 0; Idx < Planes.Num(); ++Idx)
+		{
+			const TPlaneConcrete<FReal, 3>& Plane = Planes[Idx];
+			const FVec3 ScaledNormal = (Plane.Normal() * Scale).GetSafeNormal();
+			const FReal Dot = FVec3::DotProduct(ScaledNormal, Normal);
+			if (Dot < MostOpposingDot)
+			{
+				MostOpposingDot = Dot;
+				MostOpposingIdx = Idx;
+			}
+		}
+		CHAOS_ENSURE(MostOpposingIdx != INDEX_NONE);
+		return MostOpposingIdx;
+	}
+
+
+
 	int32 FConvex::GetMostOpposingPlaneWithVertex(int32 VertexIndex, const FVec3& Normal) const
 	{
 		const int32 VertexPlaneNum = NumVertexPlanes(VertexIndex);
diff --git a/Engine/Source/Runtime/Experimental/Chaos/Public/Chaos/Box.h b/Engine/Source/Runtime/Experimental/Chaos/Public/Chaos/Box.h
--- a/Engine/Source/Runtime/Experimental/Chaos/Public/Chaos/Box.h
+++ b/Engine/Source/Runtime/Experimental/Chaos/Public/Chaos/Box.h
@@ -201,6 +201,12 @@ namespace Chaos
 			return AxisIndex;
 		}
 
+		int32 GetMostOpposingPlaneScaled(const TVector<T, d>& Normal, const TVector<T, d>& Scale) const
+		{
+			// Scale does not affect the face normals of a box
+			return GetMostOpposingPlane(Normal);
+		}
+
 		// Get the index of the plane that most opposes the normal (VertexIndex is ignored)
 		int32 GetMostOpposingPlaneWithVertex(int32 VertexIndex, const TVector<T, d>& Normal) const
 		{
diff --git a/Engine/Source/Runtime/Experimental/Chaos/Public/Chaos/Convex.h b/Engine/Source/Runtime/Experimental/Chaos/Public/Chaos/Convex.h
--- a/Engine/Source/Runtime/Experimental/Chaos/Public/Chaos/Convex.h
+++ b/Engine/Source/Runtime/Experimental/Chaos/Public/Chaos/Convex.h
@@ -307,6 +307,9 @@ namespace Chaos
 		// Get the index of the plane that most opposes the normal
 		int32 GetMostOpposingPlane(const FVec3& Normal) const;
 
+		// Get the index of the plane that most opposes the normal
+		int32 GetMostOpposingPlaneScaled(const FVec3& Normal, const FVec3& Scale) const;
+
 		// Get the index of the plane that most opposes the normal, assuming it passes through the specified vertex
 		int32 GetMostOpposingPlaneWithVertex(int32 VertexIndex, const FVec3& Normal) const;
 
@@ -352,6 +355,8 @@ namespace Chaos
 
 		virtual int32 FindMostOpposingFace(const FVec3& Position, const FVec3& UnitDir, int32 HintFaceIndex, FReal SearchDist) const override;
 
+		virtual int32 FindMostOpposingFaceScaled(const FVec3& Position, const FVec3& UnitDir, int32 HintFaceIndex, FReal SearchDist, const FVec3& Scale) const override;
+
 		FVec3 FindGeometryOpposingNormal(const FVec3& DenormDir, int32 FaceIndex, const FVec3& OriginalNormal) const
 		{
 			// For convexes, this function must be called with a face index.
diff --git a/Engine/Source/Runtime/Experimental/Chaos/Public/Chaos/ImplicitObject.h b/Engine/Source/Runtime/Experimental/Chaos/Public/Chaos/ImplicitObject.h
--- a/Engine/Source/Runtime/Experimental/Chaos/Public/Chaos/ImplicitObject.h
+++ b/Engine/Source/Runtime/Experimental/Chaos/Public/Chaos/ImplicitObject.h
@@ -269,6 +269,12 @@ public:
 		return INDEX_NONE;
 	}
 
+	virtual int32 FindMostOpposingFaceScaled(const FVec3& Position, const FVec3& UnitDir, int32 HintFaceIndex, FReal SearchDist, const FVec3& Scale) const
+	{
+		//Many objects have no concept of a face
+		return INDEX_NONE;
+	}
+
 
 	/** Finds the first intersecting face at given position
 	@param Position - local position to search around (for example a point on the surface of a convex hull)
diff --git a/Engine/Source/Runtime/Experimental/Chaos/Public/Chaos/ImplicitObjectScaled.h b/Engine/Source/Runtime/Experimental/Chaos/Public/Chaos/ImplicitObjectScaled.h
--- a/Engine/Source/Runtime/Experimental/Chaos/Public/Chaos/ImplicitObjectScaled.h
+++ b/Engine/Source/Runtime/Experimental/Chaos/Public/Chaos/ImplicitObjectScaled.h
@@ -624,7 +624,7 @@ public:
 	// Get the index of the plane that most opposes the normal
 	int32 GetMostOpposingPlane(const FVec3& Normal) const
 	{
-		return MObject->GetMostOpposingPlane(GetUnscaledNormal(Normal));
+		return MObject->GetMostOpposingPlaneScaled(Normal, MScale);
 	}
 
 	// Get the index of the plane that most opposes the normal, assuming it passes through the specified vertex
@@ -695,18 +695,8 @@ public:
 
 	virtual int32 FindMostOpposingFace(const TVector<T, d>& Position, const TVector<T, d>& UnitDir, int32 HintFaceIndex, T SearchDist) const override
 	{
-		//ensure(OuterMargin == 0);	//not supported: do we care?
 		ensure(FMath::IsNearlyEqual(UnitDir.SizeSquared(), 1, KINDA_SMALL_NUMBER));
-
-		const TVector<T, d> UnscaledPosition = MInvScale * Position;
-		const TVector<T, d> UnscaledDirDenorm = MInvScale * UnitDir;
-		const float LengthScale = UnscaledDirDenorm.Size();
-		const TVector<T, d> UnscaledDir
-			= ensure(LengthScale > TNumericLimits<T>::Min())
-			? UnscaledDirDenorm / LengthScale
-			: TVector<T, d>(0.f, 0.f, 1.f);
-		const T UnscaledSearchDist = SearchDist * MScale.Max();	//this is not quite right since it's no longer a sphere, but the whole thing is fuzzy anyway
-		return MObject->FindMostOpposingFace(UnscaledPosition, UnscaledDir, HintFaceIndex, UnscaledSearchDist);
+		return MObject->FindMostOpposingFaceScaled(Position, UnitDir, HintFaceIndex, SearchDist, MScale);
 	}
 
 	virtual TVector<T, 3> FindGeometryOpposingNormal(const TVector<T, d>& DenormDir, int32 HintFaceIndex, const TVector<T, d>& OriginalNormal) const override
