Update velocity when applying joint constraints
Merge UE5/Release-Engine-Staging [at] 15630841 to UE5/Main
This represents UE4/Main [at] 15601601

diff --git a/Engine/Source/Runtime/Experimental/Chaos/Private/Chaos/PBDJointConstraints.cpp b/Engine/Source/Runtime/Experimental/Chaos/Private/Chaos/PBDJointConstraints.cpp
--- a/Engine/Source/Runtime/Experimental/Chaos/Private/Chaos/PBDJointConstraints.cpp
+++ b/Engine/Source/Runtime/Experimental/Chaos/Private/Chaos/PBDJointConstraints.cpp
@@ -721,8 +721,8 @@ namespace Chaos
 				const FTransformPair& JointFrames = ConstraintFrames[JointIndex];
 				int32 Index0, Index1;
 				GetConstrainedParticleIndices(JointIndex, Index0, Index1);
-				FGenericParticleHandle Particle0 = FGenericParticleHandle(ConstraintParticles[JointIndex][Index0]);
-				FGenericParticleHandle Particle1 = FGenericParticleHandle(ConstraintParticles[JointIndex][Index1]);
+				TGenericParticleHandle<FReal, 3> Particle0 = TGenericParticleHandle<FReal, 3>(ConstraintParticles[JointIndex][Index0]);
+				TGenericParticleHandle<FReal, 3> Particle1 = TGenericParticleHandle<FReal, 3>(ConstraintParticles[JointIndex][Index1]);
 
 				JointState.Init(
 					Settings,
@@ -752,8 +752,8 @@ namespace Chaos
 
 				int32 Index0, Index1;
 				GetConstrainedParticleIndices(JointIndex, Index0, Index1);
-				FGenericParticleHandle Particle0 = FGenericParticleHandle(ConstraintParticles[JointIndex][Index0]);
-				FGenericParticleHandle Particle1 = FGenericParticleHandle(ConstraintParticles[JointIndex][Index1]);
+				TGenericParticleHandle<FReal, 3> Particle0 = TGenericParticleHandle<FReal, 3>(ConstraintParticles[JointIndex][Index0]);
+				TGenericParticleHandle<FReal, 3> Particle1 = TGenericParticleHandle<FReal, 3>(ConstraintParticles[JointIndex][Index1]);
 
 				Solver.Init(
 					Dt,
@@ -816,8 +816,8 @@ namespace Chaos
 	{
 		int32 Index0, Index1;
 		GetConstrainedParticleIndices(ConstraintIndex, Index0, Index1);
-		FGenericParticleHandle Particle0 = FGenericParticleHandle(ConstraintParticles[ConstraintIndex][Index0]);
-		FGenericParticleHandle Particle1 = FGenericParticleHandle(ConstraintParticles[ConstraintIndex][Index1]);
+		TGenericParticleHandle<FReal, 3> Particle0 = TGenericParticleHandle<FReal, 3>(ConstraintParticles[ConstraintIndex][Index0]);
+		TGenericParticleHandle<FReal, 3> Particle1 = TGenericParticleHandle<FReal, 3>(ConstraintParticles[ConstraintIndex][Index1]);
 		const FVec3 P0 = FParticleUtilities::GetCoMWorldPosition(Particle0);
 		const FRotation3 Q0 = FParticleUtilities::GetCoMWorldRotation(Particle0);
 		const FVec3 P1 = FParticleUtilities::GetCoMWorldPosition(Particle1);
@@ -1062,8 +1062,8 @@ namespace Chaos
 
 		int32 Index0, Index1;
 		GetConstrainedParticleIndices(JointIndex, Index0, Index1);
-		FGenericParticleHandle Particle0 = FGenericParticleHandle(ConstraintParticles[JointIndex][Index0]);
-		FGenericParticleHandle Particle1 = FGenericParticleHandle(ConstraintParticles[JointIndex][Index1]);
+		TGenericParticleHandle<FReal, 3> Particle0 = TGenericParticleHandle<FReal, 3>(ConstraintParticles[JointIndex][Index0]);
+		TGenericParticleHandle<FReal, 3> Particle1 = TGenericParticleHandle<FReal, 3>(ConstraintParticles[JointIndex][Index1]);
 
 		JointState.Update(
 			FParticleUtilities::GetCoMWorldPosition(Particle0),
@@ -1281,8 +1281,8 @@ namespace Chaos
 
 		int32 Index0, Index1;
 		GetConstrainedParticleIndices(ConstraintIndex, Index0, Index1);
-		FGenericParticleHandle Particle0 = FGenericParticleHandle(ConstraintParticles[ConstraintIndex][Index0]);
-		FGenericParticleHandle Particle1 = FGenericParticleHandle(ConstraintParticles[ConstraintIndex][Index1]);
+		TGenericParticleHandle<FReal, 3> Particle0 = TGenericParticleHandle<FReal, 3>(ConstraintParticles[ConstraintIndex][Index0]);
+		TGenericParticleHandle<FReal, 3> Particle1 = TGenericParticleHandle<FReal, 3>(ConstraintParticles[ConstraintIndex][Index1]);
 
 		const FVec3 P0 = FParticleUtilities::GetCoMWorldPosition(Particle0);
 		const FRotation3 Q0 = FParticleUtilities::GetCoMWorldRotation(Particle0);
@@ -1354,8 +1354,8 @@ namespace Chaos
 
 		int32 Index0, Index1;
 		GetConstrainedParticleIndices(ConstraintIndex, Index0, Index1);
-		FGenericParticleHandle Particle0 = FGenericParticleHandle(ConstraintParticles[ConstraintIndex][Index0]);
-		FGenericParticleHandle Particle1 = FGenericParticleHandle(ConstraintParticles[ConstraintIndex][Index1]);
+		TGenericParticleHandle<FReal, 3> Particle0 = TGenericParticleHandle<FReal, 3>(ConstraintParticles[ConstraintIndex][Index0]);
+		TGenericParticleHandle<FReal, 3> Particle1 = TGenericParticleHandle<FReal, 3>(ConstraintParticles[ConstraintIndex][Index1]);
 
 		const FVec3 P0 = FParticleUtilities::GetCoMWorldPosition(Particle0);
 		const FRotation3 Q0 = FParticleUtilities::GetCoMWorldRotation(Particle0);
