Remove degenerate triangles from geometry collections before sending to Nanite builder.

diff --git a/Engine/Source/Runtime/Experimental/GeometryCollectionEngine/Private/GeometryCollection/GeometryCollectionObject.cpp b/Engine/Source/Runtime/Experimental/GeometryCollectionEngine/Private/GeometryCollection/GeometryCollectionObject.cpp
--- a/Engine/Source/Runtime/Experimental/GeometryCollectionEngine/Private/GeometryCollection/GeometryCollectionObject.cpp
+++ b/Engine/Source/Runtime/Experimental/GeometryCollectionEngine/Private/GeometryCollection/GeometryCollectionObject.cpp
@@ -657,7 +657,7 @@ TUniquePtr<FGeometryCollectionNaniteData> UGeometryCollection::CreateNaniteData(
 
 		// TODO: Respect multiple materials like in FGeometryCollectionConversion::AppendStaticMesh
 
-		uint32 DestFaceStart = MaterialIndices.Num();
+		int32 DestFaceStart = MaterialIndices.Num();
 		MaterialIndices.Reserve(DestFaceStart + FaceCount);
 		BuildIndices.Reserve((DestFaceStart + FaceCount) * 3);
 		for (int32 FaceIndex = 0; FaceIndex < FaceCount; ++FaceIndex)
@@ -667,10 +667,20 @@ TUniquePtr<FGeometryCollectionNaniteData> UGeometryCollection::CreateNaniteData(
 				continue;
 			}
 
-			const FIntVector FaceIndices = IndicesArray[FaceStart + FaceIndex];
-			BuildIndices.Add(FaceIndices.X - VertexStart + DestVertexStart);
-			BuildIndices.Add(FaceIndices.Y - VertexStart + DestVertexStart);
-			BuildIndices.Add(FaceIndices.Z - VertexStart + DestVertexStart);
+			FIntVector FaceIndices = IndicesArray[FaceStart + FaceIndex];
+			FaceIndices = FaceIndices + FIntVector( DestVertexStart - VertexStart );
+
+			// Remove degenerates
+			if( BuildVertices[ FaceIndices[0] ].Position == BuildVertices[ FaceIndices[1] ].Position ||
+				BuildVertices[ FaceIndices[1] ].Position == BuildVertices[ FaceIndices[2] ].Position ||
+				BuildVertices[ FaceIndices[2] ].Position == BuildVertices[ FaceIndices[0] ].Position )
+			{
+				continue;
+			}
+
+			BuildIndices.Add(FaceIndices.X);
+			BuildIndices.Add(FaceIndices.Y);
+			BuildIndices.Add(FaceIndices.Z);
 
 			const int32 MaterialIndex = MaterialIDArray[FaceStart + FaceIndex];
 			MaterialIndices.Add(MaterialIndex);
