Chaos - fix restitution when actor and com locations are different

diff --git a/Engine/Source/Runtime/Experimental/Chaos/Private/Chaos/Collision/PBDCollisionConstraint.cpp b/Engine/Source/Runtime/Experimental/Chaos/Private/Chaos/Collision/PBDCollisionConstraint.cpp
--- a/Engine/Source/Runtime/Experimental/Chaos/Private/Chaos/Collision/PBDCollisionConstraint.cpp
+++ b/Engine/Source/Runtime/Experimental/Chaos/Private/Chaos/Collision/PBDCollisionConstraint.cpp
@@ -231,11 +231,11 @@ namespace Chaos
 		// We use PreV and PreW to support incremental manifold generation. In this case, manifold points
 		// can be added after we have already run some solver iterations, which gives us an incorrect initial
 		// velocity if we just use V and W (one-shots will work the same either way since V=PreV on first pass)
-		const FRigidTransform3 WorldTransform0 = FParticleUtilities::GetActorWorldTransform(Particle0);
-		const FRigidTransform3 WorldTransform1 = FParticleUtilities::GetActorWorldTransform(Particle1);
-		const FVec3 WorldContactVel0 = Particle0->PreV() + FVec3::CrossProduct(Particle0->PreW(), WorldTransform0.GetRotation() * LocalContactPoint0);
-		const FVec3 WorldContactVel1 = Particle1->PreV() + FVec3::CrossProduct(Particle1->PreW(), WorldTransform1.GetRotation() * LocalContactPoint1);
-		const FVec3 WorldContactNormal = (ManifoldPoint.ContactPoint.ContactNormalOwnerIndex == 0) ? WorldTransform0.GetRotation() * LocalContactNormal : WorldTransform1.GetRotation() * LocalContactNormal;
+		const FRigidTransform3 WorldCoMTransform0 = FParticleUtilities::GetCoMWorldTransform(Particle0);
+		const FRigidTransform3 WorldCoMTransform1 = FParticleUtilities::GetCoMWorldTransform(Particle1);
+		const FVec3 WorldContactVel0 = Particle0->PreV() + FVec3::CrossProduct(Particle0->PreW(), WorldCoMTransform0.GetRotation() * CoMContactPoint0);
+		const FVec3 WorldContactVel1 = Particle1->PreV() + FVec3::CrossProduct(Particle1->PreW(), WorldCoMTransform1.GetRotation() * CoMContactPoint1);
+		const FVec3 WorldContactNormal = (ManifoldPoint.ContactPoint.ContactNormalOwnerIndex == 0) ? WorldCoMTransform0.GetRotation() * CoMContactNormal : WorldCoMTransform1.GetRotation() * CoMContactNormal;
 		const FReal WorldContactVelNorm = FVec3::DotProduct(WorldContactVel0 - WorldContactVel1, WorldContactNormal);
 		ManifoldPoint.InitialContactVelocity = WorldContactVelNorm;
 
