Removing unnecessary forward declaration.

diff --git a/Engine/Source/Runtime/Experimental/InteractiveToolsFramework/Public/ToolContextInterfaces.h b/Engine/Source/Runtime/Experimental/InteractiveToolsFramework/Public/ToolContextInterfaces.h
--- a/Engine/Source/Runtime/Experimental/InteractiveToolsFramework/Public/ToolContextInterfaces.h
+++ b/Engine/Source/Runtime/Experimental/InteractiveToolsFramework/Public/ToolContextInterfaces.h
@@ -26,10 +26,6 @@ class UInteractiveToolStorableSelection;
 struct FTypedElementHandle;
 class UTypedElementSelectionSet;
 
-#if WITH_EDITOR
-class HHitProxy;
-#endif
-
 /**
  * FToolBuilderState is a bucket of state information that a ToolBuilder might need
  * to construct a Tool. This information comes from a level above the Tools framework,
