Chaos - fix inertiia of geometry collision level set particles

diff --git a/Engine/Source/Runtime/Experimental/Chaos/Private/PhysicsProxy/GeometryCollectionPhysicsProxy.cpp b/Engine/Source/Runtime/Experimental/Chaos/Private/PhysicsProxy/GeometryCollectionPhysicsProxy.cpp
--- a/Engine/Source/Runtime/Experimental/Chaos/Private/PhysicsProxy/GeometryCollectionPhysicsProxy.cpp
+++ b/Engine/Source/Runtime/Experimental/Chaos/Private/PhysicsProxy/GeometryCollectionPhysicsProxy.cpp
@@ -1976,7 +1976,8 @@ void FGeometryCollectionPhysicsProxy::InitializeSharedCollisionStructures(
 
 			if (InertiaComputationNeeded[GeometryIndex])
 			{
-				CalculateInertiaAndRotationOfMass(MassSpaceParticles, TriMesh->GetSurfaceElements(), Density_i, MassProperties.CenterOfMass, MassProperties.InertiaTensor, MassProperties.RotationOfMass);
+				// Note: particles already in CoM space, so passing in zero as CoM
+				CalculateInertiaAndRotationOfMass(MassSpaceParticles, TriMesh->GetSurfaceElements(), Density_i, FVec3(0), MassProperties.InertiaTensor, MassProperties.RotationOfMass);
 				CollectionInertiaTensor[TransformGroupIndex] = FVec3(MassProperties.InertiaTensor.M[0][0], MassProperties.InertiaTensor.M[1][1], MassProperties.InertiaTensor.M[2][2]);
 #if false
 				CollectionMassToLocal[TransformGroupIndex] = FTransform(MassProperties.RotationOfMass, MassProperties.CenterOfMass);
