fix AddFaces re-targeting FaceStart before the reorder call now that reorder correctly updates FaceStart

#rnx

diff --git a/Engine/Source/Runtime/Experimental/Chaos/Private/GeometryCollection/GeometryCollectionAlgo.cpp b/Engine/Source/Runtime/Experimental/Chaos/Private/GeometryCollection/GeometryCollectionAlgo.cpp
--- a/Engine/Source/Runtime/Experimental/Chaos/Private/GeometryCollection/GeometryCollectionAlgo.cpp
+++ b/Engine/Source/Runtime/Experimental/Chaos/Private/GeometryCollection/GeometryCollectionAlgo.cpp
@@ -745,7 +745,6 @@ namespace GeometryCollectionAlgo
 				NewSpaceIdx++;
 			}
 			GeometryCollection->FaceCount[GeometryIdx] += NumToAdd;
-			GeometryCollection->FaceStart[GeometryIdx] = NewStart;
 		}
 		GeometryCollection->ReorderElements(FGeometryCollection::FacesGroup, ShiftedOrder);
 	}
