Chaos - debug draw

diff --git a/Engine/Source/Runtime/Engine/Private/PhysicsEngine/ImmediatePhysics/ImmediatePhysicsChaos/ImmediatePhysicsSimulation_Chaos.cpp b/Engine/Source/Runtime/Engine/Private/PhysicsEngine/ImmediatePhysics/ImmediatePhysicsChaos/ImmediatePhysicsSimulation_Chaos.cpp
--- a/Engine/Source/Runtime/Engine/Private/PhysicsEngine/ImmediatePhysics/ImmediatePhysicsChaos/ImmediatePhysicsSimulation_Chaos.cpp
+++ b/Engine/Source/Runtime/Engine/Private/PhysicsEngine/ImmediatePhysics/ImmediatePhysicsChaos/ImmediatePhysicsSimulation_Chaos.cpp
@@ -162,25 +162,29 @@ FAutoConsoleVariableRef CVarChaosImmPhysDebugDrawJointFeatures(TEXT("p.Chaos.Imm
 FAutoConsoleVariableRef CVarChaosImmPhysDebugDrawSimulationSpace(TEXT("p.Chaos.ImmPhys.DebugDrawSimulationSpace"), ChaosImmediate_DebugDrawSimulationSpace, TEXT("Draw the simulation frame of reference, acceleration and velocity."));
 
 Chaos::DebugDraw::FChaosDebugDrawSettings ChaosImmPhysDebugDebugDrawSettings(
-	/* ArrowSize =			*/ 1.5f,
-	/* BodyAxisLen =		*/ 4.0f,
-	/* ContactLen =			*/ 4.0f,
-	/* ContactWidth =		*/ 2.0f,
-	/* ContactPhiWidth =	*/ 0.0f,
-	/* ContactOwnerWidth =	*/ 0.0f,
-	/* ConstraintAxisLen =	*/ 5.0f,
-	/* JointComSize =		*/ 2.0f,
-	/* LineThickness =		*/ 0.15f,
-	/* DrawScale =			*/ 1.0f,
-	/* FontHeight =			*/ 10.0f,
-	/* FontScale =			*/ 1.5f,
-	/* ShapeThicknesScale = */ 1.0f,
-	/* PointSize =			*/ 2.0f,
-	/* VelScale =			*/ 0.0f,
-	/* AngVelScale =		*/ 0.0f,
-	/* ImpulseScale =		*/ 0.0f,
-	/* DrawPriority =		*/ 10.0f
-);
+	/* ArrowSize =					*/ 1.5f,
+	/* BodyAxisLen =				*/ 4.0f,
+	/* ContactLen =					*/ 4.0f,
+	/* ContactWidth =				*/ 2.0f,
+	/* ContactPhiWidth =			*/ 0.0f,
+	/* ContactOwnerWidth =			*/ 0.0f,
+	/* ConstraintAxisLen =			*/ 5.0f,
+	/* JointComSize =				*/ 2.0f,
+	/* LineThickness =				*/ 0.15f,
+	/* DrawScale =					*/ 1.0f,
+	/* FontHeight =					*/ 10.0f,
+	/* FontScale =					*/ 1.5f,
+	/* ShapeThicknesScale =			*/ 1.0f,
+	/* PointSize =					*/ 2.0f,
+	/* VelScale =					*/ 0.0f,
+	/* AngVelScale =				*/ 0.0f,
+	/* ImpulseScale =				*/ 0.0f,
+	/* InertiaScale =				*/ 0.0f,
+	/* DrawPriority =				*/ 10.0f,
+	/* bShowSimple =				*/ true,
+	/* bShowComplex =				*/ false,
+	/* bInShowLevelSetCollision =	*/ false
+	);
 
 FAutoConsoleVariableRef CVarChaosImmPhysArrowSize(TEXT("p.Chaos.ImmPhys.DebugDraw.ArrowSize"), ChaosImmPhysDebugDebugDrawSettings.ArrowSize, TEXT("ArrowSize."));
 FAutoConsoleVariableRef CVarChaosImmPhysBodyAxisLen(TEXT("p.Chaos.ImmPhys.DebugDraw.BodyAxisLen"), ChaosImmPhysDebugDebugDrawSettings.BodyAxisLen, TEXT("BodyAxisLen."));
diff --git a/Engine/Source/Runtime/Experimental/Chaos/Private/Chaos/AABB.cpp b/Engine/Source/Runtime/Experimental/Chaos/Private/Chaos/AABB.cpp
--- a/Engine/Source/Runtime/Experimental/Chaos/Private/Chaos/AABB.cpp
+++ b/Engine/Source/Runtime/Experimental/Chaos/Private/Chaos/AABB.cpp
@@ -250,6 +250,24 @@ TAABB<T, d> TAABB<T, d>::TransformedAABB(const FTransform& SpaceTransform) const
 		return TransformedAABBHelper<T, d>(*this, SpaceTransform);
 	}
 }
+
+template<typename T, int d>
+inline TAABB<T, d> TAABB<T, d>::InverseTransformedAABB(const TRigidTransform<T, 3>& SpaceTransform) const
+{
+	TVector<T, d> CurrentExtents = Extents();
+	int32 Idx = 0;
+	const TVector<T, d> MinToNewSpace = SpaceTransform.InverseTransformPosition(Min());
+	TAABB<T, d> NewAABB(MinToNewSpace, MinToNewSpace);
+	NewAABB.GrowToInclude(SpaceTransform.InverseTransformPosition(Max()));
+
+	for (int32 j = 0; j < d; ++j)
+	{
+		NewAABB.GrowToInclude(SpaceTransform.InverseTransformPosition(Min() + TVector<T, d>::AxisVector(j) * CurrentExtents));
+		NewAABB.GrowToInclude(SpaceTransform.InverseTransformPosition(Max() - TVector<T, d>::AxisVector(j) * CurrentExtents));
+	}
+
+	return NewAABB;
+}
 }
 
 template class Chaos::TAABB<Chaos::FReal, 3>;
diff --git a/Engine/Source/Runtime/Experimental/Chaos/Private/Chaos/ChaosDebugDraw.cpp b/Engine/Source/Runtime/Experimental/Chaos/Private/Chaos/ChaosDebugDraw.cpp
--- a/Engine/Source/Runtime/Experimental/Chaos/Private/Chaos/ChaosDebugDraw.cpp
+++ b/Engine/Source/Runtime/Experimental/Chaos/Private/Chaos/ChaosDebugDraw.cpp
@@ -3,6 +3,8 @@
 #include "Chaos/Box.h"
 #include "Chaos/Capsule.h"
 #include "Chaos/Convex.h"
+#include "Chaos/HeightField.h"
+#include "Chaos/Triangle.h"
 #include "Chaos/TriangleMesh.h"
 #include "Chaos/DebugDrawQueue.h"
 #include "Chaos/Evolution/SimulationSpace.h"
@@ -42,24 +44,28 @@ namespace Chaos
 		// if the user does not specify settings in the debug draw call.
 		// See PBDRigidsColver.cpp and ImmediatePhysicsSimulation_Chaos.cpp for example.
 		FChaosDebugDrawSettings ChaosDefaultDebugDebugDrawSettings(
-			/* ArrowSize =			*/ 1.5f,
-			/* BodyAxisLen =		*/ 4.0f,
-			/* ContactLen =			*/ 4.0f,
-			/* ContactWidth =		*/ 2.0f,
-			/* ContactPhiWidth =	*/ 0.0f,
-			/* ContactOwnerWidth =	*/ 0.0f,
-			/* ConstraintAxisLen =	*/ 5.0f,
-			/* JointComSize =		*/ 2.0f,
-			/* LineThickness =		*/ 0.15f,
-			/* DrawScale =			*/ 1.0f,
-			/* FontHeight =			*/ 10.0f,
-			/* FontScale =			*/ 1.5f,
-			/* ShapeThicknesScale = */ 1.0f,
-			/* PointSize =			*/ 2.0f,
-			/* VelScale =			*/ 0.0f,
-			/* AngVelScale =		*/ 0.0f,
-			/* ImpulseScale =		*/ 0.0f,
-			/* DrawPriority =		*/ 10.0f
+			/* ArrowSize =					*/ 1.5f,
+			/* BodyAxisLen =				*/ 4.0f,
+			/* ContactLen =					*/ 4.0f,
+			/* ContactWidth =				*/ 2.0f,
+			/* ContactPhiWidth =			*/ 0.0f,
+			/* ContactOwnerWidth =			*/ 0.0f,
+			/* ConstraintAxisLen =			*/ 5.0f,
+			/* JointComSize =				*/ 2.0f,
+			/* LineThickness =				*/ 0.15f,
+			/* DrawScale =					*/ 1.0f,
+			/* FontHeight =					*/ 10.0f,
+			/* FontScale =					*/ 1.5f,
+			/* ShapeThicknesScale =			*/ 1.0f,
+			/* PointSize =					*/ 2.0f,
+			/* VelScale =					*/ 0.0f,
+			/* AngVelScale =				*/ 0.0f,
+			/* ImpulseScale =				*/ 0.0f,
+			/* InertiaScale =				*/ 0.0f,
+			/* DrawPriority =				*/ 10.0f,
+			/* bShowSimple =				*/ true,
+			/* bShowComplex =				*/ false,
+			/* bInShowLevelSetCollision =	*/ false
 		);
 
 		const FChaosDebugDrawSettings& GetChaosDebugDrawSettings(const FChaosDebugDrawSettings* InSettings)
@@ -82,6 +88,63 @@ namespace Chaos
 			DrawShapesImpl(nullptr, ShapeTransform, Shape, Color, GetChaosDebugDrawSettings(Settings));
 		}
 
+		void DrawShapesHeightFieldImpl(const TGeometryParticleHandle<FReal, 3>* Particle, const FRigidTransform3& ShapeTransform, const FHeightField* Shape, const FColor& Color, const FChaosDebugDrawSettings& Settings)
+		{
+			const FVec3& WorldQueryCenter = FDebugDrawQueue::GetInstance().GetCenterOfInterest();
+			const FReal WorldQueryRadius = FDebugDrawQueue::GetInstance().GetRadiusOfInterest();
+			const FAABB3 WorldQueryBounds = FAABB3(WorldQueryCenter - FVec3(WorldQueryRadius), WorldQueryCenter + FVec3(WorldQueryRadius));
+			const FAABB3 LocalQueryBounds = WorldQueryBounds.InverseTransformedAABB(ShapeTransform);
+
+			Shape->VisitTriangles(LocalQueryBounds, [&](const TTriangle<FReal>& Tri)
+				{
+					FVec3 A = ShapeTransform.TransformPosition(Tri[0]);
+					FVec3 B = ShapeTransform.TransformPosition(Tri[1]);
+					FVec3 C = ShapeTransform.TransformPosition(Tri[2]);
+					FDebugDrawQueue::GetInstance().DrawDebugLine(A, B, Color, false, -1.f, 0, Settings.ShapeThicknesScale * Settings.LineThickness);
+					FDebugDrawQueue::GetInstance().DrawDebugLine(B, C, Color, false, -1.f, 0, Settings.ShapeThicknesScale * Settings.LineThickness);
+					FDebugDrawQueue::GetInstance().DrawDebugLine(C, A, Color, false, -1.f, 0, Settings.ShapeThicknesScale * Settings.LineThickness);
+				});
+		}
+
+		void DrawShapesTriangleMeshImpl(const TGeometryParticleHandle<FReal, 3>* Particle, const FRigidTransform3& ShapeTransform, const FTriangleMeshImplicitObject* Shape, const FColor& Color, const FChaosDebugDrawSettings& Settings)
+		{
+			const FVec3& WorldQueryCenter = FDebugDrawQueue::GetInstance().GetCenterOfInterest();
+			const FReal WorldQueryRadius = FDebugDrawQueue::GetInstance().GetRadiusOfInterest();
+			const FAABB3 WorldQueryBounds = FAABB3(WorldQueryCenter - FVec3(WorldQueryRadius), WorldQueryCenter + FVec3(WorldQueryRadius));
+			const FAABB3 LocalQueryBounds = WorldQueryBounds.InverseTransformedAABB(ShapeTransform);
+
+			Shape->VisitTriangles(LocalQueryBounds, [&](const TTriangle<FReal>& Tri)
+			{
+				FVec3 A = ShapeTransform.TransformPosition(Tri[0]);
+				FVec3 B = ShapeTransform.TransformPosition(Tri[1]);
+				FVec3 C = ShapeTransform.TransformPosition(Tri[2]);
+				FDebugDrawQueue::GetInstance().DrawDebugLine(A, B, Color, false, -1.f, 0, Settings.ShapeThicknesScale * Settings.LineThickness);
+				FDebugDrawQueue::GetInstance().DrawDebugLine(B, C, Color, false, -1.f, 0, Settings.ShapeThicknesScale * Settings.LineThickness);
+				FDebugDrawQueue::GetInstance().DrawDebugLine(C, A, Color, false, -1.f, 0, Settings.ShapeThicknesScale * Settings.LineThickness);
+			});
+		}
+
+		void DrawShapesLevelSetImpl(const TGeometryParticleHandle<FReal, 3>* Particle, const FRigidTransform3& ShapeTransform, const FLevelSet* Shape, const FColor& Color, const FChaosDebugDrawSettings& Settings)
+		{
+			if (!Settings.bShowLevelSetCollision)
+			{
+				return;
+			}
+
+			if (const TPBDRigidParticleHandle<FReal, 3>* Rigid = Particle->CastToRigidParticle())
+			{
+				const TUniquePtr<TBVHParticles<FReal, 3>>& CollisionParticles = Rigid->CollisionParticles();
+				if (CollisionParticles != nullptr)
+				{
+					for (int32 ParticleIndex = 0; ParticleIndex < (int32)CollisionParticles->Size(); ++ParticleIndex)
+					{
+						const FVec3 P = ShapeTransform.TransformPosition(CollisionParticles->X(ParticleIndex));
+						FDebugDrawQueue::GetInstance().DrawDebugPoint(P, Color, false, -1.f, 0, Settings.PointSize);
+					}
+				}
+			}
+		}
+
 		template <bool bInstanced>
 		void DrawShapesScaledImpl(const FGeometryParticleHandle* Particle, const FRigidTransform3& ShapeTransform, const FImplicitObject* Shape, const FColor& Color, const FChaosDebugDrawSettings& Settings)
 		{
@@ -121,9 +184,19 @@ namespace Chaos
 			case ImplicitObjectType::Cylinder:
 				break;
 			case ImplicitObjectType::TriangleMesh:
+			{
+				const TImplicitObjectScaled<FTriangleMeshImplicitObject, bInstanced>* Scaled = Shape->template GetObject<TImplicitObjectScaled<FTriangleMeshImplicitObject, bInstanced>>();
+				ScaleTM.SetScale3D(Scaled->GetScale());
+				DrawShapesImpl(Particle, ScaleTM * ShapeTransform, Scaled->GetUnscaledObject(), Color, Settings);
 				break;
+			}
 			case ImplicitObjectType::HeightField:
+			{
+				const TImplicitObjectScaled<FHeightField, bInstanced>* Scaled = Shape->template GetObject<TImplicitObjectScaled<FHeightField, bInstanced>>();
+				ScaleTM.SetScale3D(Scaled->GetScale());
+				DrawShapesImpl(Particle, ScaleTM * ShapeTransform, Scaled->GetUnscaledObject(), Color, Settings);
 				break;
+			}
 			default:
 				break;
 			}
@@ -165,9 +238,17 @@ namespace Chaos
 			case ImplicitObjectType::Cylinder:
 				break;
 			case ImplicitObjectType::TriangleMesh:
+			{
+				const TImplicitObjectInstanced<FTriangleMeshImplicitObject>* Scaled = Shape->template GetObject<TImplicitObjectInstanced<FTriangleMeshImplicitObject>>();
+				DrawShapesImpl(Particle, ShapeTransform, Scaled->GetInstancedObject(), Color, Settings);
 				break;
+			}
 			case ImplicitObjectType::HeightField:
+			{
+				const TImplicitObjectInstanced<FHeightField>* Scaled = Shape->template GetObject<TImplicitObjectInstanced<FHeightField>>();
+				DrawShapesImpl(Particle, ShapeTransform, Scaled->GetInstancedObject(), Color, Settings);
 				break;
+			}
 			default:
 				break;
 			}
@@ -178,7 +259,7 @@ namespace Chaos
 			const EImplicitObjectType PackedType = Shape->GetType(); // Type includes scaling and instancing data
 			const EImplicitObjectType InnerType = GetInnerType(Shape->GetType());
 
-			// For scaled shapes, we must unpack the scaled type first
+			// Unwrap the wrapper/aggregating shapes
 			if (IsScaled(PackedType))
 			{
 				if (IsInstanced(PackedType))
@@ -196,15 +277,64 @@ namespace Chaos
 				DrawShapesInstancedImpl(Particle, ShapeTransform, Shape, Color, Settings);
 				return;
 			}
+			else if (InnerType == ImplicitObjectType::Transformed)
+			{
+				const TImplicitObjectTransformed<FReal, 3>* Transformed = Shape->template GetObject<TImplicitObjectTransformed<FReal, 3>>();
+				FRigidTransform3 TransformedTransform = FRigidTransform3(ShapeTransform.TransformPosition(Transformed->GetTransform().GetLocation()), ShapeTransform.GetRotation() * Transformed->GetTransform().GetRotation());
+				DrawShapesImpl(Particle, TransformedTransform, Transformed->GetTransformedObject(), Color, Settings);
+				return;
+			}
+			else if (InnerType == ImplicitObjectType::Union)
+			{
+				const FImplicitObjectUnion* Union = Shape->template GetObject<FImplicitObjectUnion>();
+				for (auto& UnionShape : Union->GetObjects())
+				{
+					DrawShapesImpl(Particle, ShapeTransform, UnionShape.Get(), Color, Settings);
+				}
+				return;
+			}
+			else if (InnerType == ImplicitObjectType::UnionClustered)
+			{
+				const FImplicitObjectUnionClustered* Union = Shape->template GetObject<FImplicitObjectUnionClustered>();
+				for (auto& UnionShape : Union->GetObjects())
+				{
+					DrawShapesImpl(Particle, ShapeTransform, UnionShape.Get(), Color, Settings);
+				}
+				return;
+			}
+
+
+			// Whether we should show meshes and non-mesh shapes
+			bool bShowMeshes = Settings.bShowComplexCollision;
+			bool bShowNonMeshes = Settings.bShowSimpleCollision;
+			const FPerShapeData* ShapeData = Particle->GetImplicitShape(Shape);
+			if (ShapeData != nullptr)
+			{
+				bShowMeshes = (Settings.bShowComplexCollision && (ShapeData->GetCollisionTraceType() != EChaosCollisionTraceFlag::Chaos_CTF_UseSimpleAsComplex))
+					|| (Settings.bShowSimpleCollision && (ShapeData->GetCollisionTraceType() == EChaosCollisionTraceFlag::Chaos_CTF_UseComplexAsSimple));
+				bShowNonMeshes = (Settings.bShowSimpleCollision && (ShapeData->GetCollisionTraceType() != EChaosCollisionTraceFlag::Chaos_CTF_UseComplexAsSimple))
+					|| (Settings.bShowComplexCollision && (ShapeData->GetCollisionTraceType() == EChaosCollisionTraceFlag::Chaos_CTF_UseSimpleAsComplex));
+			}
+
+			// Quit if we don't want to show this shape
+			const bool bIsMesh = (InnerType == ImplicitObjectType::TriangleMesh);
+			if (bIsMesh && !bShowMeshes)
+			{
+				return;
+			}
+			else if (!bIsMesh && !bShowNonMeshes)
+			{
+				return;
+			}
 
-			// @todo(ccaulfield): handle scale throughout
+			// If we get here, we have an actual shape to render
 			switch (InnerType)
 			{
 			case ImplicitObjectType::Sphere:
 			{
 				const TSphere<FReal, 3>* Sphere = Shape->template GetObject<TSphere<FReal, 3>>();
 				const FVec3 P = ShapeTransform.TransformPosition(Sphere->GetCenter());
-				FDebugDrawQueue::GetInstance().DrawDebugSphere(P, Sphere->GetRadius(), 20, Color, false, KINDA_SMALL_NUMBER, Settings.DrawPriority, Settings.ShapeThicknesScale * Settings.LineThickness);
+				FDebugDrawQueue::GetInstance().DrawDebugSphere(P, Sphere->GetRadius(), 8, Color, false, KINDA_SMALL_NUMBER, Settings.DrawPriority, Settings.ShapeThicknesScale * Settings.LineThickness);
 				break;
 			}
 			case ImplicitObjectType::Box:
@@ -224,26 +354,13 @@ namespace Chaos
 				FDebugDrawQueue::GetInstance().DrawDebugCapsule(P, (FReal)0.5 * Capsule->GetHeight() + Capsule->GetRadius(), Capsule->GetRadius(), Q, Color, false, KINDA_SMALL_NUMBER, Settings.DrawPriority, Settings.ShapeThicknesScale * Settings.LineThickness);
 				break;
 			}
-			case ImplicitObjectType::Transformed:
-			{
-				const TImplicitObjectTransformed<FReal, 3>* Transformed = Shape->template GetObject<TImplicitObjectTransformed<FReal, 3>>();
-				FRigidTransform3 TransformedTransform = FRigidTransform3(ShapeTransform.TransformPosition(Transformed->GetTransform().GetLocation()), ShapeTransform.GetRotation() * Transformed->GetTransform().GetRotation());
-				DrawShapesImpl(Particle, TransformedTransform, Transformed->GetTransformedObject(), Color, Settings);
-				break;
-			}
-			case ImplicitObjectType::Union:
+			case ImplicitObjectType::LevelSet:
 			{
-				const FImplicitObjectUnion* Union = Shape->template GetObject<FImplicitObjectUnion>();
-				for (auto& UnionShape : Union->GetObjects())
-				{
-					DrawShapesImpl(Particle, ShapeTransform, UnionShape.Get(), Color, Settings);
-				}
+				const FLevelSet* LevelSet = Shape->template GetObject<FLevelSet>();
+				DrawShapesLevelSetImpl(Particle, ShapeTransform, LevelSet, Color, Settings);
 				break;
 			}
-			case ImplicitObjectType::LevelSet:
-				break;
-			case ImplicitObjectType::Unknown:
-				break;
+			break;
 			case ImplicitObjectType::Convex:
 			{
 				if (const FConvex* Convex = Shape->template GetObject<FConvex>())
@@ -271,9 +388,17 @@ namespace Chaos
 			case ImplicitObjectType::Cylinder:
 				break;
 			case ImplicitObjectType::TriangleMesh:
+			{
+				const FTriangleMeshImplicitObject* TriangleMesh = Shape->template GetObject<FTriangleMeshImplicitObject>();
+				DrawShapesTriangleMeshImpl(Particle, ShapeTransform, TriangleMesh, Color, Settings);
 				break;
+			}
 			case ImplicitObjectType::HeightField:
+			{
+				const FHeightField* HeightField = Shape->template GetObject<FHeightField>();
+				DrawShapesHeightFieldImpl(Particle, ShapeTransform, HeightField, Color, Settings);
 				break;
+			}
 			default:
 				break;
 			}
@@ -395,6 +520,15 @@ namespace Chaos
 			{
 				FDebugDrawQueue::GetInstance().DrawDebugLine(PCOM, PCOM + SpaceTransform.TransformVector(Particle->W()) * Settings.AngVelScale, Green, false, KINDA_SMALL_NUMBER, Settings.DrawPriority, Settings.LineThickness);
 			}
+
+			if (Settings.InertiaScale > 0.0f)
+			{
+				if (const TPBDRigidParticleHandle<FReal, 3>* Rigid = InParticle->CastToRigidParticle())
+				{
+					const FVec3 EquivalentBoxSize = Utilities::BoxSizeFromInertia(Rigid->I().GetDiagonal(), Rigid->M());
+					FDebugDrawQueue::GetInstance().DrawDebugBox(PCOM, 0.5f * Settings.InertiaScale * EquivalentBoxSize, QCOM, FColor::Magenta, false, 0.0f, 0, Settings.LineThickness);
+				}
+			}
 		}
 
 		void DrawCollisionImpl(const FVec3& Location, const FVec3& Normal, FReal Phi, const FVec3& Impulse, const FColor& DiscColor, const FColor& NormalColor, const FColor& ImpulseColor, FRealSingle ColorScale, const FChaosDebugDrawSettings& Settings)
@@ -721,7 +855,7 @@ namespace Chaos
 									{
 										IslandAABB.GrowToInclude(CollisionHandle->GetConstrainedParticles()[1]->WorldSpaceInflatedBounds());
 									}
-									DrawGraphCollision(SpaceTransform, *CollisionHandle, IslandIndex, LevelIndex, ColorIndex, Settings);
+									//DrawGraphCollision(SpaceTransform, *CollisionHandle, IslandIndex, LevelIndex, ColorIndex, Settings);
 								}
 							}
 
diff --git a/Engine/Source/Runtime/Experimental/Chaos/Private/Chaos/HeightField.cpp b/Engine/Source/Runtime/Experimental/Chaos/Private/Chaos/HeightField.cpp
--- a/Engine/Source/Runtime/Experimental/Chaos/Private/Chaos/HeightField.cpp
+++ b/Engine/Source/Runtime/Experimental/Chaos/Private/Chaos/HeightField.cpp
@@ -1569,6 +1569,28 @@ namespace Chaos
 		return SweepGeomImp(QueryGeom, StartTM, Dir, Length, OutTime, OutPosition, OutNormal, OutFaceIndex, Thickness, bComputeMTD);
 	}
 
+	void FHeightField::VisitTriangles(const FAABB3& QueryBounds, const TFunction<void(const FTriangle& Triangle)>& Visitor) const
+	{
+		FBounds2D FlatQueryBounds;
+		FlatQueryBounds.Min = TVector<FReal, 2>(QueryBounds.Min()[0], QueryBounds.Min()[1]);
+		FlatQueryBounds.Max = TVector<FReal, 2>(QueryBounds.Max()[0], QueryBounds.Max()[1]);
+
+		TArray<TVector<int32, 2>> Intersections;
+		FVec3 Points[4];
+
+		GetGridIntersections(FlatQueryBounds, Intersections);
+
+		bool bOverlaps = false;
+		for (const TVector<int32, 2>&Cell : Intersections)
+		{
+			const int32 SingleIndex = Cell[1] * (GeomData.NumCols) + Cell[0];
+			GeomData.GetPointsScaled(SingleIndex, Points);
+
+			Visitor(FTriangle(Points[0], Points[1], Points[3]));
+			Visitor(FTriangle(Points[0], Points[3], Points[2]));
+		}
+	}
+
 	int32 FHeightField::FindMostOpposingFace(const FVec3& Position, const FVec3& UnitDir, int32 HintFaceIndex, FReal SearchDist) const
 	{
 		const FReal SearchDist2 = SearchDist * SearchDist;
diff --git a/Engine/Source/Runtime/Experimental/Chaos/Private/Chaos/TriangleMeshImplicitObject.cpp b/Engine/Source/Runtime/Experimental/Chaos/Private/Chaos/TriangleMeshImplicitObject.cpp
--- a/Engine/Source/Runtime/Experimental/Chaos/Private/Chaos/TriangleMeshImplicitObject.cpp
+++ b/Engine/Source/Runtime/Experimental/Chaos/Private/Chaos/TriangleMeshImplicitObject.cpp
@@ -496,6 +496,32 @@ bool FTriangleMeshImplicitObject::Overlap(const FVec3& Point, const FReal Thickn
 	}
 }
 
+void FTriangleMeshImplicitObject::VisitTriangles(const FAABB3& QueryBounds, const TFunction<void(const FTriangle& Triangle)>& Visitor) const
+{
+	const TArray<int32> PotentialIntersections = BVH.FindAllIntersections(QueryBounds);
+
+	auto TriangleProducer = [&](const auto& Elements)
+	{
+		for (int32 TriIdx : PotentialIntersections)
+		{
+			TVec3<FReal> A, B, C;
+			TransformVertsHelper(FVec3(1), TriIdx, MParticles, Elements, A, B, C);
+
+			Visitor(FTriangle(A, B, C));
+		}
+	};
+
+	if (MElements.RequiresLargeIndices())
+	{
+		return TriangleProducer(MElements.GetLargeIndexBuffer());
+	}
+	else
+	{
+		return TriangleProducer(MElements.GetSmallIndexBuffer());
+	}
+}
+
+
 template <typename QueryGeomType>
 bool FTriangleMeshImplicitObject::OverlapGeomImp(const QueryGeomType& QueryGeom, const FRigidTransform3& QueryTM, const FReal Thickness, FMTDInfo* OutMTD, FVec3 TriMeshScale) const
 {
diff --git a/Engine/Source/Runtime/Experimental/Chaos/Private/PBDRigidsSolver.cpp b/Engine/Source/Runtime/Experimental/Chaos/Private/PBDRigidsSolver.cpp
--- a/Engine/Source/Runtime/Experimental/Chaos/Private/PBDRigidsSolver.cpp
+++ b/Engine/Source/Runtime/Experimental/Chaos/Private/PBDRigidsSolver.cpp
@@ -24,6 +24,8 @@
 #include "Chaos/PullPhysicsDataImp.h"
 #include "Chaos/PhysicsSolverBaseImpl.h"
 
+#include "ProfilingDebugging/CsvProfiler.h"
+
 //PRAGMA_DISABLE_OPTIMIZATION
 
 DEFINE_LOG_CATEGORY_STATIC(LogPBDRigidsSolver, Log, All);
@@ -40,6 +42,8 @@ DECLARE_DWORD_ACCUMULATOR_STAT(TEXT("NumIslands"), STAT_ChaosCounter_NumIslands,
 DECLARE_DWORD_ACCUMULATOR_STAT(TEXT("NumContacts"), STAT_ChaosCounter_NumContacts, STATGROUP_ChaosCounters);
 DECLARE_DWORD_ACCUMULATOR_STAT(TEXT("NumJoints"), STAT_ChaosCounter_NumJoints, STATGROUP_ChaosCounters);
 
+CSV_DEFINE_CATEGORY(ChaosCounters, true);
+
 
 // DebugDraw CVars
 #if CHAOS_DEBUG_DRAW
@@ -52,31 +56,41 @@ int32 ChaosSolverDebugDrawCollisions = CHAOS_SOLVER_ENABLE_DEBUG_DRAW;
 int32 ChaosSolverDebugDrawBounds = 0;
 int32 ChaosSolverDrawTransforms = 0;
 int32 ChaosSolverDrawIslands = 0;
+int32 ChaosSolverDrawShapesShowStatic = 1;
+int32 ChaosSolverDrawShapesShowKinematic = 1;
+int32 ChaosSolverDrawShapesShowDynamic = 1;
 FAutoConsoleVariableRef CVarChaosSolverDrawShapes(TEXT("p.Chaos.Solver.DebugDrawShapes"), ChaosSolverDebugDrawShapes, TEXT("Draw Shapes (0 = never; 1 = end of frame)."));
 FAutoConsoleVariableRef CVarChaosSolverDrawCollisions(TEXT("p.Chaos.Solver.DebugDrawCollisions"), ChaosSolverDebugDrawCollisions, TEXT("Draw Collisions (0 = never; 1 = end of frame)."));
 FAutoConsoleVariableRef CVarChaosSolverDrawBounds(TEXT("p.Chaos.Solver.DebugDrawBounds"), ChaosSolverDebugDrawBounds, TEXT("Draw bounding volumes inside the broadphase (0 = never; 1 = end of frame)."));
 FAutoConsoleVariableRef CVarChaosSolverDrawTransforms(TEXT("p.Chaos.Solver.DebugDrawTransforms"), ChaosSolverDrawTransforms, TEXT("Draw particle transforms (0 = never; 1 = end of frame)."));
 FAutoConsoleVariableRef CVarChaosSolverDrawIslands(TEXT("p.Chaos.Solver.DebugDrawIslands"), ChaosSolverDrawIslands, TEXT("Draw solver islands (0 = never; 1 = end of frame)."));
+FAutoConsoleVariableRef CVarChaosSolverDrawShapesShapesStatic(TEXT("p.Chaos.Solver.DebugDraw.ShowStatics"), ChaosSolverDrawShapesShowStatic, TEXT("If DebugDrawShapes is enabled, whether to show static objects"));
+FAutoConsoleVariableRef CVarChaosSolverDrawShapesShapesKinematic(TEXT("p.Chaos.Solver.DebugDraw.ShowKinematics"), ChaosSolverDrawShapesShowKinematic, TEXT("If DebugDrawShapes is enabled, whether to show kinematic objects"));
+FAutoConsoleVariableRef CVarChaosSolverDrawShapesShapesDynamic(TEXT("p.Chaos.Solver.DebugDraw.ShowDynamics"), ChaosSolverDrawShapesShowDynamic, TEXT("If DebugDrawShapes is enabled, whether to show dynamic objects"));
 
 Chaos::DebugDraw::FChaosDebugDrawSettings ChaosSolverDebugDebugDrawSettings(
-	/* ArrowSize =			*/ 10.0f,
-	/* BodyAxisLen =		*/ 30.0f,
-	/* ContactLen =			*/ 30.0f,
-	/* ContactWidth =		*/ 6.0f,
-	/* ContactPhiWidth =	*/ 0.0f,
-	/* ContactOwnerWidth =	*/ 0.0f,
-	/* ConstraintAxisLen =	*/ 30.0f,
-	/* JointComSize =		*/ 2.0f,
-	/* LineThickness =		*/ 1.0f,
-	/* DrawScale =			*/ 1.0f,
-	/* FontHeight =			*/ 10.0f,
-	/* FontScale =			*/ 1.5f,
-	/* ShapeThicknesScale = */ 1.0f,
-	/* PointSize =			*/ 5.0f,
-	/* VelScale =			*/ 0.0f,
-	/* AngVelScale =		*/ 0.0f,
-	/* ImpulseScale =		*/ 0.0f,
-	/* DrawPriority =		*/ 10.0f
+	/* ArrowSize =					*/ 10.0f,
+	/* BodyAxisLen =				*/ 30.0f,
+	/* ContactLen =					*/ 30.0f,
+	/* ContactWidth =				*/ 6.0f,
+	/* ContactPhiWidth =			*/ 0.0f,
+	/* ContactOwnerWidth =			*/ 0.0f,
+	/* ConstraintAxisLen =			*/ 30.0f,
+	/* JointComSize =				*/ 2.0f,
+	/* LineThickness =				*/ 1.0f,
+	/* DrawScale =					*/ 1.0f,
+	/* FontHeight =					*/ 10.0f,
+	/* FontScale =					*/ 1.5f,
+	/* ShapeThicknesScale =			*/ 1.0f,
+	/* PointSize =					*/ 5.0f,
+	/* VelScale =					*/ 0.0f,
+	/* AngVelScale =				*/ 0.0f,
+	/* ImpulseScale =				*/ 0.0f,
+	/* InertiaScale =				*/ 1.0f,
+	/* DrawPriority =				*/ 10.0f,
+	/* bShowSimple =				*/ true,
+	/* bShowComplex =				*/ false,
+	/* bInShowLevelSetCollision =	*/ true
 );
 FAutoConsoleVariableRef CVarChaosSolverArrowSize(TEXT("p.Chaos.Solver.DebugDraw.ArrowSize"), ChaosSolverDebugDebugDrawSettings.ArrowSize, TEXT("ArrowSize."));
 FAutoConsoleVariableRef CVarChaosSolverBodyAxisLen(TEXT("p.Chaos.Solver.DebugDraw.BodyAxisLen"), ChaosSolverDebugDebugDrawSettings.BodyAxisLen, TEXT("BodyAxisLen."));
@@ -91,7 +105,11 @@ FAutoConsoleVariableRef CVarChaosSolverPointSize(TEXT("p.Chaos.Solver.DebugDraw.
 FAutoConsoleVariableRef CVarChaosSolverVelScale(TEXT("p.Chaos.Solver.DebugDraw.VelScale"), ChaosSolverDebugDebugDrawSettings.VelScale, TEXT("If >0 show velocity when drawing particle transforms."));
 FAutoConsoleVariableRef CVarChaosSolverAngVelScale(TEXT("p.Chaos.Solver.DebugDraw.AngVelScale"), ChaosSolverDebugDebugDrawSettings.AngVelScale, TEXT("If >0 show angular velocity when drawing particle transforms."));
 FAutoConsoleVariableRef CVarChaosSolverImpulseScale(TEXT("p.Chaos.Solver.DebugDraw.ImpulseScale"), ChaosSolverDebugDebugDrawSettings.ImpulseScale, TEXT("If >0 show impulses when drawing collisions."));
+FAutoConsoleVariableRef CVarChaosSolverInertiaScale(TEXT("p.Chaos.Solver.DebugDraw.InertiaScale"), ChaosSolverDebugDebugDrawSettings.InertiaScale, TEXT("When DebugDrawTransforms is enabled, show the mass-normalized inertia matrix scaled by this amount."));
 FAutoConsoleVariableRef CVarChaosSolverScale(TEXT("p.Chaos.Solver.DebugDraw.Scale"), ChaosSolverDebugDebugDrawSettings.DrawScale, TEXT("Scale applied to all Chaos Debug Draw line lengths etc."));
+FAutoConsoleVariableRef CVarChaosSolverShowSimple(TEXT("p.Chaos.Solver.DebugDraw.ShowSimple"), ChaosSolverDebugDebugDrawSettings.bShowSimpleCollision, TEXT("Whether to show simple collision is shape drawing is enabled"));
+FAutoConsoleVariableRef CVarChaosSolverShowComplex(TEXT("p.Chaos.Solver.DebugDraw.ShowComplex"), ChaosSolverDebugDebugDrawSettings.bShowComplexCollision, TEXT("Whether to show complex collision is shape drawing is enabled"));
+FAutoConsoleVariableRef CVarChaosSolverShowLevelSet(TEXT("p.Chaos.Solver.DebugDraw.ShowLevelSet"), ChaosSolverDebugDebugDrawSettings.bShowLevelSetCollision, TEXT("Whether to show levelset collision is shape drawing is enabled"));
 
 #endif
 
@@ -280,6 +298,21 @@ namespace Chaos
 				SET_DWORD_STAT(STAT_ChaosCounter_NumIslands, MSolver->GetEvolution()->GetConstraintGraph().NumIslands());
 				SET_DWORD_STAT(STAT_ChaosCounter_NumContacts, MSolver->NumCollisionConstraints());
 				SET_DWORD_STAT(STAT_ChaosCounter_NumJoints, MSolver->NumJointConstraints());
+
+#if CSV_PROFILER
+				// Particle counts
+				CSV_CUSTOM_STAT(ChaosCounters, NumDisabledParticles, MSolver->GetEvolution()->GetParticles().GetAllParticlesView().Num() - MSolver->GetEvolution()->GetParticles().GetNonDisabledView().Num(), ECsvCustomStatOp::Accumulate);
+				CSV_CUSTOM_STAT(ChaosCounters, NumParticles, MSolver->GetEvolution()->GetParticles().GetNonDisabledView().Num(), ECsvCustomStatOp::Accumulate);
+				CSV_CUSTOM_STAT(ChaosCounters, NumDynamicParticles, MSolver->GetEvolution()->GetParticles().GetNonDisabledDynamicView().Num(), ECsvCustomStatOp::Accumulate);
+				CSV_CUSTOM_STAT(ChaosCounters, NumKinematicParticles, MSolver->GetEvolution()->GetParticles().GetActiveKinematicParticlesView().Num(), ECsvCustomStatOp::Accumulate);
+				CSV_CUSTOM_STAT(ChaosCounters, NumStaticParticles, MSolver->GetEvolution()->GetParticles().GetActiveStaticParticlesView().Num(), ECsvCustomStatOp::Accumulate);
+				CSV_CUSTOM_STAT(ChaosCounters, NumGeometryCollectionParticles, (int)MSolver->GetEvolution()->GetParticles().GetGeometryCollectionParticles().Size(), ECsvCustomStatOp::Accumulate);
+
+				// Constraint counts
+				CSV_CUSTOM_STAT(ChaosCounters, NumIslands, MSolver->GetEvolution()->GetConstraintGraph().NumIslands(), ECsvCustomStatOp::Accumulate);
+				CSV_CUSTOM_STAT(ChaosCounters, NumContacts, MSolver->NumCollisionConstraints(), ECsvCustomStatOp::Accumulate);
+				CSV_CUSTOM_STAT(ChaosCounters, NumJoints, MSolver->NumJointConstraints(), ECsvCustomStatOp::Accumulate);
+#endif
 			}
 
 
@@ -1185,9 +1218,18 @@ namespace Chaos
 		QUICK_SCOPE_CYCLE_COUNTER(SolverDebugDraw);
 		if (ChaosSolverDebugDrawShapes == 1)
 		{
-			DebugDraw::DrawParticleShapes(FRigidTransform3(), Particles.GetActiveStaticParticlesView(), FColor(128, 0, 0), &ChaosSolverDebugDebugDrawSettings);
-			DebugDraw::DrawParticleShapes(FRigidTransform3(), Particles.GetActiveKinematicParticlesView(), FColor(64, 32, 0), &ChaosSolverDebugDebugDrawSettings);
-			DebugDraw::DrawParticleShapes(FRigidTransform3(), Particles.GetNonDisabledDynamicView(), FColor(255, 255, 0), &ChaosSolverDebugDebugDrawSettings);
+			if (ChaosSolverDrawShapesShowStatic)
+			{
+				DebugDraw::DrawParticleShapes(FRigidTransform3(), Particles.GetActiveStaticParticlesView(), FColor(128, 0, 0), &ChaosSolverDebugDebugDrawSettings);
+			}
+			if (ChaosSolverDrawShapesShowKinematic)
+			{
+				DebugDraw::DrawParticleShapes(FRigidTransform3(), Particles.GetActiveKinematicParticlesView(), FColor(64, 32, 0), &ChaosSolverDebugDebugDrawSettings);
+			}
+			if (ChaosSolverDrawShapesShowDynamic)
+			{
+				DebugDraw::DrawParticleShapes(FRigidTransform3(), Particles.GetNonDisabledDynamicView(), FColor(255, 255, 0), &ChaosSolverDebugDebugDrawSettings);
+			}
 		}
 		if (ChaosSolverDebugDrawCollisions == 1) 
 		{
diff --git a/Engine/Source/Runtime/Experimental/Chaos/Public/Chaos/AABB.h b/Engine/Source/Runtime/Experimental/Chaos/Public/Chaos/AABB.h
--- a/Engine/Source/Runtime/Experimental/Chaos/Public/Chaos/AABB.h
+++ b/Engine/Source/Runtime/Experimental/Chaos/Public/Chaos/AABB.h
@@ -90,6 +90,8 @@ namespace Chaos
 		CHAOS_API TAABB<T, d> TransformedAABB(const FMatrix&) const;
 		CHAOS_API TAABB<T, d> TransformedAABB(const Chaos::PMatrix<FReal, 4, 4>&) const;
 
+		CHAOS_API TAABB<T, d> InverseTransformedAABB(const Chaos::TRigidTransform<T, 3>&) const;
+
 		FORCEINLINE bool Intersects(const TAABB<T, d>& Other) const
 		{
 			for (int32 i = 0; i < d; ++i)
diff --git a/Engine/Source/Runtime/Experimental/Chaos/Public/Chaos/ChaosDebugDraw.h b/Engine/Source/Runtime/Experimental/Chaos/Public/Chaos/ChaosDebugDraw.h
--- a/Engine/Source/Runtime/Experimental/Chaos/Public/Chaos/ChaosDebugDraw.h
+++ b/Engine/Source/Runtime/Experimental/Chaos/Public/Chaos/ChaosDebugDraw.h
@@ -34,8 +34,12 @@ namespace Chaos
 				FRealSingle InVelScale,
 				FRealSingle InAngVelScale,
 				FRealSingle InImpulseScale,
-				int InDrawPriority
-			)
+				float InInertiaScale,
+				int InDrawPriority,
+				bool bInShowSimpleCollision,
+				bool bInShowComplexCollision,
+				bool bInShowLevelSetCollision
+				)
 				: ArrowSize(InArrowSize)
 				, BodyAxisLen(InBodyAxisLen)
 				, ContactLen(InContactLen)
@@ -53,7 +57,11 @@ namespace Chaos
 				, VelScale(InVelScale)
 				, AngVelScale(InAngVelScale)
 				, ImpulseScale(InImpulseScale)
+				, InertiaScale(InInertiaScale)
 				, DrawPriority(InDrawPriority)
+				, bShowSimpleCollision(bInShowSimpleCollision)
+				, bShowComplexCollision(bInShowComplexCollision)
+				, bShowLevelSetCollision(bInShowLevelSetCollision)
 			{}
 
 			FRealSingle ArrowSize;
@@ -73,7 +81,11 @@ namespace Chaos
 			FRealSingle VelScale;
 			FRealSingle AngVelScale;
 			FRealSingle ImpulseScale;
+			FRealSingle InertiaScale;
 			int DrawPriority;
+			bool bShowSimpleCollision;
+			bool bShowComplexCollision;
+			bool bShowLevelSetCollision;
 		};
 
 
diff --git a/Engine/Source/Runtime/Experimental/Chaos/Public/Chaos/DebugDrawQueue.h b/Engine/Source/Runtime/Experimental/Chaos/Public/Chaos/DebugDrawQueue.h
--- a/Engine/Source/Runtime/Experimental/Chaos/Public/Chaos/DebugDrawQueue.h
+++ b/Engine/Source/Runtime/Experimental/Chaos/Public/Chaos/DebugDrawQueue.h
@@ -388,6 +388,16 @@ public:
 		RadiusOfInterest = InRadius;
 	}
 
+	const FVector& GetCenterOfInterest() const
+	{
+		return CenterOfInterest;
+	}
+
+	FReal GetRadiusOfInterest() const
+	{
+		return RadiusOfInterest;
+	}
+
 	void SetConsumerActive(void* Consumer, bool bConsumerActive);
 
 	static FDebugDrawQueue& GetInstance()
diff --git a/Engine/Source/Runtime/Experimental/Chaos/Public/Chaos/HeightField.h b/Engine/Source/Runtime/Experimental/Chaos/Public/Chaos/HeightField.h
--- a/Engine/Source/Runtime/Experimental/Chaos/Public/Chaos/HeightField.h
+++ b/Engine/Source/Runtime/Experimental/Chaos/Public/Chaos/HeightField.h
@@ -17,6 +17,7 @@ namespace Chaos
 {
 	class FHeightfieldRaycastVisitor;
 	class FConvex;
+	class FTriangle;
 	struct FMTDInfo;
 }
 
@@ -87,6 +88,8 @@ namespace Chaos
 		virtual int32 FindMostOpposingFace(const FVec3& Position, const FVec3& UnitDir, int32 HintFaceIndex, FReal SearchDist) const override;
 		virtual FVec3 FindGeometryOpposingNormal(const FVec3& DenormDir, int32 FaceIndex, const FVec3& OriginalNormal) const override;
 
+		void VisitTriangles(const FAABB3& InQueryBounds, const TFunction<void(const FTriangle& Triangle)>& Visitor) const;
+
 		struct FClosestFaceData
 		{
 			int32 FaceIndex = INDEX_NONE;
diff --git a/Engine/Source/Runtime/Experimental/Chaos/Public/Chaos/ParticleHandle.h b/Engine/Source/Runtime/Experimental/Chaos/Public/Chaos/ParticleHandle.h
--- a/Engine/Source/Runtime/Experimental/Chaos/Public/Chaos/ParticleHandle.h
+++ b/Engine/Source/Runtime/Experimental/Chaos/Public/Chaos/ParticleHandle.h
@@ -518,7 +518,7 @@ public:
 		GeometryParticles->SetHandle(ParticleIdx, this);
 	}
 
-	const FPerShapeData* GetImplicitShape(const FImplicitObject* InObject)
+	const FPerShapeData* GetImplicitShape(const FImplicitObject* InObject) const
 	{
 		return GeometryParticles->GetImplicitShape(ParticleIdx, InObject);
 	}
diff --git a/Engine/Source/Runtime/Experimental/Chaos/Public/Chaos/Triangle.h b/Engine/Source/Runtime/Experimental/Chaos/Public/Chaos/Triangle.h
--- a/Engine/Source/Runtime/Experimental/Chaos/Public/Chaos/Triangle.h
+++ b/Engine/Source/Runtime/Experimental/Chaos/Public/Chaos/Triangle.h
@@ -3,6 +3,7 @@
 #pragma once
 
 #include "Chaos/Core.h"
+#include "Chaos/GJK.h"
 #include "ImplicitObject.h"
 #include "Plane.h"
 
diff --git a/Engine/Source/Runtime/Experimental/Chaos/Public/Chaos/TriangleMeshImplicitObject.h b/Engine/Source/Runtime/Experimental/Chaos/Public/Chaos/TriangleMeshImplicitObject.h
--- a/Engine/Source/Runtime/Experimental/Chaos/Public/Chaos/TriangleMeshImplicitObject.h
+++ b/Engine/Source/Runtime/Experimental/Chaos/Public/Chaos/TriangleMeshImplicitObject.h
@@ -21,8 +21,8 @@ namespace Chaos
 	extern CHAOS_API bool TriMeshPerPolySupport;
 
 	class FCapsule;
-
 	class FConvex;
+	class FTriangle;
 	struct FMTDInfo;
 
 	class CHAOS_API FTrimeshIndexBuffer
@@ -290,6 +290,8 @@ namespace Chaos
 
 		void UpdateVertices(const TArray<FVector>& Positions);
 
+		void VisitTriangles(const FAABB3& InQueryBounds, const TFunction<void(const FTriangle& Triangle)>& Visitor) const;
+
 	private:
 		void RebuildBV();
 
diff --git a/Engine/Source/Runtime/Experimental/Chaos/Public/Chaos/Utilities.h b/Engine/Source/Runtime/Experimental/Chaos/Public/Chaos/Utilities.h
--- a/Engine/Source/Runtime/Experimental/Chaos/Public/Chaos/Utilities.h
+++ b/Engine/Source/Runtime/Experimental/Chaos/Public/Chaos/Utilities.h
@@ -463,6 +463,23 @@ namespace Chaos
 			return ScaledInertia;
 		}
 
+		// Compute the box size that would generate the given (diagonal) inertia
+		inline FVec3 BoxSizeFromInertia(const FVec3& Inertia, const FReal Mass)
+		{
+			// System of 3 equations in X^2, Y^2, Z^2
+			// Inertia.X = 1/12 M (Size.Y^2 + Size.Z^2)
+			// Inertia.Y = 1/12 M (Size.Z^2 + Size.X^2)
+			// Inertia.Z = 1/12 M (Size.X^2 + Size.Y^2)
+			if (Mass > 0)
+			{
+				const FVec3 S = Inertia * 12.0f / Mass;
+				const FMatrix33 R = FMatrix33(-0.5f, 0.5f, 0.5f, 0.5f, -0.5f, 0.5f, 0.5f, 0.5f, -0.5f);
+				const FVec3 XYZSq = R * S;
+				return FVec3(FMath::Sqrt(XYZSq.X), FMath::Sqrt(XYZSq.Y), FMath::Sqrt(XYZSq.Z));
+			}
+			return FVec3(0);
+		}
+
 		// Replacement for FMath::Wrap that works for integers and returns a value in [Begin, End).
 		// Note: this implementation uses a loop to bring the value into range - it should not be used if the value is much larger than the range.
 		inline int32 WrapIndex(int32 V, int32 Begin, int32 End)
