Chaos - fix a couple more scaled collision issues

diff --git a/Engine/Source/Runtime/Experimental/Chaos/Private/Chaos/Convex.cpp b/Engine/Source/Runtime/Experimental/Chaos/Private/Chaos/Convex.cpp
--- a/Engine/Source/Runtime/Experimental/Chaos/Private/Chaos/Convex.cpp
+++ b/Engine/Source/Runtime/Experimental/Chaos/Private/Chaos/Convex.cpp
@@ -132,7 +132,7 @@ namespace Chaos
 		for (int32 Idx = 0; Idx < Planes.Num(); ++Idx)
 		{
 			const TPlaneConcrete<FReal, 3>& Plane = Planes[Idx];
-			const FVec3 ScaledNormal = (Plane.Normal() * Scale).GetSafeNormal();
+			const FVec3 ScaledNormal = (Plane.Normal() / Scale).GetSafeNormal();
 			const FReal Dot = FVec3::DotProduct(ScaledNormal, Normal);
 			if (Dot < MostOpposingDot)
 			{
diff --git a/Engine/Source/Runtime/Experimental/Chaos/Public/Chaos/ImplicitObjectScaled.h b/Engine/Source/Runtime/Experimental/Chaos/Public/Chaos/ImplicitObjectScaled.h
--- a/Engine/Source/Runtime/Experimental/Chaos/Public/Chaos/ImplicitObjectScaled.h
+++ b/Engine/Source/Runtime/Experimental/Chaos/Public/Chaos/ImplicitObjectScaled.h
@@ -701,12 +701,13 @@ public:
 
 	virtual TVector<T, 3> FindGeometryOpposingNormal(const TVector<T, d>& DenormDir, int32 HintFaceIndex, const TVector<T, d>& OriginalNormal) const override
 	{
-		//ensure(OuterMargin == 0);	//not supported: do we care?
+		// @todo(chaos): we need a virtual FindGeometryOpposingNormal. Some types (Convex, Box) can just return the
+		// normal from the face index without needing calculating the unscaled normals.
 		ensure(FMath::IsNearlyEqual(OriginalNormal.SizeSquared(), 1, KINDA_SMALL_NUMBER));
 
 		// Get unscaled dir and normal
-		const TVector<T, 3> LocalDenormDir = DenormDir * MInvScale;
-		const TVector<T, 3> LocalOriginalNormalDenorm = OriginalNormal * MInvScale;
+		const TVector<T, 3> LocalDenormDir = DenormDir * MScale;
+		const TVector<T, 3> LocalOriginalNormalDenorm = OriginalNormal * MScale;
 		const T NormalLengthScale = LocalOriginalNormalDenorm.Size();
 		const TVector<T, 3> LocalOriginalNormal
 			= ensure(NormalLengthScale > SMALL_NUMBER)
@@ -715,10 +716,10 @@ public:
 
 		// Compute final normal
 		const TVector<T, d> LocalNormal = MObject->FindGeometryOpposingNormal(LocalDenormDir, HintFaceIndex, LocalOriginalNormal);
-		TVector<T, d> Normal = LocalNormal;
+		TVector<T, d> Normal = LocalNormal * MInvScale;
 		if (CHAOS_ENSURE(Normal.SafeNormalize(TNumericLimits<T>::Min())) == 0)
 		{
-			Normal = TVector<T,3>(0,0,1);
+			Normal = TVector<T, 3>(0, 0, 1);
 		}
 
 		return Normal;
