GeometryCollectionComponent
-- wrapped PostEditChangeChainProperty in WITH_EDITOR to prevent build error for runtime.

diff --git a/Engine/Source/Runtime/Experimental/GeometryCollectionEngine/Private/GeometryCollection/GeometryCollectionComponent.cpp b/Engine/Source/Runtime/Experimental/GeometryCollectionEngine/Private/GeometryCollection/GeometryCollectionComponent.cpp
--- a/Engine/Source/Runtime/Experimental/GeometryCollectionEngine/Private/GeometryCollection/GeometryCollectionComponent.cpp
+++ b/Engine/Source/Runtime/Experimental/GeometryCollectionEngine/Private/GeometryCollection/GeometryCollectionComponent.cpp
@@ -858,9 +858,10 @@ void UGeometryCollectionComponent::InitializeComponent()
 	}
 }
 
-void UGeometryCollectionComponent::PostEditChangeProperty(FPropertyChangedEvent& PropertyChangedEvent)
+#if WITH_EDITOR
+void UGeometryCollectionComponent::PostEditChangeChainProperty(FPropertyChangedChainEvent& PropertyChangedEvent)
 {
-	Super::PostEditChangeProperty(PropertyChangedEvent);
+	Super::PostEditChangeChainProperty(PropertyChangedEvent);
 
 	if (PropertyChangedEvent.Property && PropertyChangedEvent.Property->GetFName() == GET_MEMBER_NAME_CHECKED(UGeometryCollectionComponent, bShowBoneColors))
 	{
@@ -870,6 +871,7 @@ void UGeometryCollectionComponent::PostEditChangeProperty(FPropertyChangedEvent&
 		MarkRenderDynamicDataDirty();
 	}
 }
+#endif
 
 static void DispatchGeometryCollectionBreakEvent(const FChaosBreakEvent& Event)
 {
diff --git a/Engine/Source/Runtime/Experimental/GeometryCollectionEngine/Public/GeometryCollection/GeometryCollectionComponent.h b/Engine/Source/Runtime/Experimental/GeometryCollectionEngine/Public/GeometryCollection/GeometryCollectionComponent.h
--- a/Engine/Source/Runtime/Experimental/GeometryCollectionEngine/Public/GeometryCollection/GeometryCollectionComponent.h
+++ b/Engine/Source/Runtime/Experimental/GeometryCollectionEngine/Public/GeometryCollection/GeometryCollectionComponent.h
@@ -341,7 +341,10 @@ public:
 	virtual void EndPlay(const EEndPlayReason::Type ReasonEnd) override;
 	virtual void GetLifetimeReplicatedProps(TArray<FLifetimeProperty>& OutLifetimeProps) const override;
 	virtual void InitializeComponent() override;
-	virtual void PostEditChangeProperty(FPropertyChangedEvent& PropertyChangedEvent) override;
+
+#if WITH_EDITOR
+	virtual void PostEditChangeChainProperty(FPropertyChangedChainEvent& PropertyChangedEvent) override;
+#endif
 	//~ Begin UActorComponent Interface. 
 
 
