Chaos Physics: Increase default error margin (CVAR: p.Chaos.CCD.AllowedDepthBoundsScale) for Shape cast, contact gen inconsistencies for CCD

diff --git a/Engine/Source/Runtime/Experimental/Chaos/Private/Chaos/CollisionResolution.cpp b/Engine/Source/Runtime/Experimental/Chaos/Private/Chaos/CollisionResolution.cpp
--- a/Engine/Source/Runtime/Experimental/Chaos/Private/Chaos/CollisionResolution.cpp
+++ b/Engine/Source/Runtime/Experimental/Chaos/Private/Chaos/CollisionResolution.cpp
@@ -79,7 +79,7 @@ DECLARE_CYCLE_STAT(TEXT("Collisions::UpdateLevelsetLevelsetConstraint"), STAT_Up
 Chaos::FRealSingle CCDEnableThresholdBoundsScale = 0.4f;
 FAutoConsoleVariableRef  CVarCCDEnableThresholdBoundsScale(TEXT("p.Chaos.CCD.EnableThresholdBoundsScale"), CCDEnableThresholdBoundsScale , TEXT("CCD is used when object position is changing > smallest bound's extent * BoundsScale. 0 will always Use CCD. Values < 0 disables CCD."));
 
-Chaos::FRealSingle CCDAllowedDepthBoundsScale = 0.05f;
+Chaos::FRealSingle CCDAllowedDepthBoundsScale = 0.1f;
 FAutoConsoleVariableRef CVarCCDAllowedDepthBoundsScale(TEXT("p.Chaos.CCD.AllowedDepthBoundsScale"), CCDAllowedDepthBoundsScale, TEXT("When rolling back to TOI, allow (smallest bound's extent) * AllowedDepthBoundsScale, instead of rolling back to exact TOI w/ penetration = 0."));
 
 int32 ConstraintsDetailedStats = 0;
