Chaos : Destruction
Re-exposed the connection method enums to the GeometryCollectionObject.
#okforgethub public, internal

diff --git a/Engine/Plugins/Experimental/GeometryCollectionPlugin/Source/GeometryCollectionEditor/Private/GeometryCollection/ActorFactoryGeometryCollection.cpp b/Engine/Plugins/Experimental/GeometryCollectionPlugin/Source/GeometryCollectionEditor/Private/GeometryCollection/ActorFactoryGeometryCollection.cpp
--- a/Engine/Plugins/Experimental/GeometryCollectionPlugin/Source/GeometryCollectionEditor/Private/GeometryCollection/ActorFactoryGeometryCollection.cpp
+++ b/Engine/Plugins/Experimental/GeometryCollectionPlugin/Source/GeometryCollectionEditor/Private/GeometryCollection/ActorFactoryGeometryCollection.cpp
@@ -54,7 +54,6 @@ void UActorFactoryGeometryCollection::PostSpawnActor(UObject* Asset, AActor* New
 	NewGeometryCollectionActor->GetGeometryCollectionComponent()->ClusterGroupIndex = GeometryCollection->ClusterGroupIndex;
 	NewGeometryCollectionActor->GetGeometryCollectionComponent()->MaxClusterLevel = GeometryCollection->MaxClusterLevel;
 	NewGeometryCollectionActor->GetGeometryCollectionComponent()->DamageThreshold = GeometryCollection->DamageThreshold;
-	NewGeometryCollectionActor->GetGeometryCollectionComponent()->ClusterConnectionType = GeometryCollection->ClusterConnectionType;
 
 	// Init Component
 	NewGeometryCollectionActor->GetGeometryCollectionComponent()->RegisterComponent();
diff --git a/Engine/Source/Runtime/Experimental/ChaosSolverEngine/Public/Chaos/ChaosSolverActor.h b/Engine/Source/Runtime/Experimental/ChaosSolverEngine/Public/Chaos/ChaosSolverActor.h
--- a/Engine/Source/Runtime/Experimental/ChaosSolverEngine/Public/Chaos/ChaosSolverActor.h
+++ b/Engine/Source/Runtime/Experimental/ChaosSolverEngine/Public/Chaos/ChaosSolverActor.h
@@ -31,10 +31,10 @@ namespace Chaos
 UENUM()
 enum class EClusterConnectionTypeEnum : uint8
 {
-	Chaos_PointImplicit = Chaos::FClusterCreationParameters::PointImplicit UMETA(Hidden),
+	Chaos_PointImplicit = Chaos::FClusterCreationParameters::PointImplicit ,
 	Chaos_DelaunayTriangulation = Chaos::FClusterCreationParameters::DelaunayTriangulation UMETA(Hidden),
-	Chaos_MinimalSpanningSubsetDelaunayTriangulation = Chaos::FClusterCreationParameters::MinimalSpanningSubsetDelaunayTriangulation UMETA(Hidden),
-	Chaos_PointImplicitAugmentedWithMinimalDelaunay = Chaos::FClusterCreationParameters::PointImplicitAugmentedWithMinimalDelaunay UMETA(Hidden),
+	Chaos_MinimalSpanningSubsetDelaunayTriangulation = Chaos::FClusterCreationParameters::MinimalSpanningSubsetDelaunayTriangulation ,
+	Chaos_PointImplicitAugmentedWithMinimalDelaunay = Chaos::FClusterCreationParameters::PointImplicitAugmentedWithMinimalDelaunay ,
 	Chaos_None = Chaos::FClusterCreationParameters::None UMETA(Hidden),
 	//
 	Chaos_EClsuterCreationParameters_Max UMETA(Hidden)
diff --git a/Engine/Source/Runtime/Experimental/GeometryCollectionEngine/Private/GeometryCollection/GeometryCollectionComponent.cpp b/Engine/Source/Runtime/Experimental/GeometryCollectionEngine/Private/GeometryCollection/GeometryCollectionComponent.cpp
--- a/Engine/Source/Runtime/Experimental/GeometryCollectionEngine/Private/GeometryCollection/GeometryCollectionComponent.cpp
+++ b/Engine/Source/Runtime/Experimental/GeometryCollectionEngine/Private/GeometryCollection/GeometryCollectionComponent.cpp
@@ -219,7 +219,7 @@ UGeometryCollectionComponent::UGeometryCollectionComponent(const FObjectInitiali
 	, MaxClusterLevel(100)
 	, DamageThreshold({250.0})
 	, bUseSizeSpecificDamageThreshold(false)
-	, ClusterConnectionType(EClusterConnectionTypeEnum::Chaos_PointImplicit)
+	, ClusterConnectionType_DEPRECATED(EClusterConnectionTypeEnum::Chaos_MinimalSpanningSubsetDelaunayTriangulation)
 	, CollisionGroup(0)
 	, CollisionSampleFraction(1.0)
 	, InitialVelocityType(EInitialVelocityTypeEnum::Chaos_Initial_Velocity_User_Defined)
@@ -1825,10 +1825,12 @@ void UGeometryCollectionComponent::OnCreatePhysicsState()
 #if !(UE_BUILD_SHIPPING || UE_BUILD_TEST)
 				SimulationParameters.Name = GetPathName();
 #endif
+				EClusterConnectionTypeEnum ClusterCollectionType = ClusterConnectionType_DEPRECATED;
 				if (RestCollection)
 				{
 					RestCollection->GetSharedSimulationParams(SimulationParameters.Shared);
 					SimulationParameters.RestCollection = RestCollection->GetGeometryCollection().Get();
+					ClusterCollectionType = RestCollection->ClusterConnectionType;
 				}
 				SimulationParameters.Simulating = Simulating;
 				SimulationParameters.EnableClustering = EnableClustering;
@@ -1836,7 +1838,7 @@ void UGeometryCollectionComponent::OnCreatePhysicsState()
 				SimulationParameters.MaxClusterLevel = MaxClusterLevel;
 				SimulationParameters.bUseSizeSpecificDamageThresholds = bUseSizeSpecificDamageThreshold;
 				SimulationParameters.DamageThreshold = DamageThreshold;
-				SimulationParameters.ClusterConnectionMethod = (Chaos::FClusterCreationParameters::EConnectionMethod)ClusterConnectionType;
+				SimulationParameters.ClusterConnectionMethod = (Chaos::FClusterCreationParameters::EConnectionMethod)ClusterCollectionType;
 				SimulationParameters.CollisionGroup = CollisionGroup;
 				SimulationParameters.CollisionSampleFraction = CollisionSampleFraction;
 				SimulationParameters.InitialVelocityType = InitialVelocityType;
diff --git a/Engine/Source/Runtime/Experimental/GeometryCollectionEngine/Public/GeometryCollection/GeometryCollectionComponent.h b/Engine/Source/Runtime/Experimental/GeometryCollectionEngine/Public/GeometryCollection/GeometryCollectionComponent.h
--- a/Engine/Source/Runtime/Experimental/GeometryCollectionEngine/Public/GeometryCollection/GeometryCollectionComponent.h
+++ b/Engine/Source/Runtime/Experimental/GeometryCollectionEngine/Public/GeometryCollection/GeometryCollectionComponent.h
@@ -466,8 +466,8 @@ public:
 	bool bUseSizeSpecificDamageThreshold;
 
 	/** */
-	UPROPERTY(EditAnywhere, BlueprintReadWrite, Category = "ChaosPhysics|Clustering")
-	EClusterConnectionTypeEnum ClusterConnectionType;
+	UPROPERTY(meta = (DeprecatedProperty, DeprecationMessage = "Connection types are defined on the asset now."))
+	EClusterConnectionTypeEnum ClusterConnectionType_DEPRECATED;
 
 	/** */
 	UPROPERTY(EditAnywhere, BlueprintReadWrite, Category = "ChaosPhysics|Collisions")
