Removed some extraneous comments on geometry collection enum states that were being picked up for tooltips. Added display name for default state.

diff --git a/Engine/Source/Runtime/Experimental/Chaos/Public/GeometryCollection/GeometryCollectionSimulationTypes.h b/Engine/Source/Runtime/Experimental/Chaos/Public/GeometryCollection/GeometryCollectionSimulationTypes.h
--- a/Engine/Source/Runtime/Experimental/Chaos/Public/GeometryCollection/GeometryCollectionSimulationTypes.h
+++ b/Engine/Source/Runtime/Experimental/Chaos/Public/GeometryCollection/GeometryCollectionSimulationTypes.h
@@ -29,14 +29,14 @@ enum class EImplicitTypeEnum : uint8
 UENUM(BlueprintType)
 enum class EObjectStateTypeEnum : uint8
 {
-	Chaos_NONE = 0 UMETA(Hidden),
-	Chaos_Object_Sleeping  = 1 /*Chaos::EObjectStateType::Sleeping*/   UMETA(DisplayName = "Sleeping"),
-	Chaos_Object_Kinematic = 2 /*Chaos::EObjectStateType::Kinematic*/  UMETA(DisplayName = "Kinematic"),
-	Chaos_Object_Static = 3    /*Chaos::EObjectStateType::Static*/     UMETA(DisplayName = "Static"),
-	Chaos_Object_Dynamic   = 4 /*Chaos::EObjectStateType::Dynamic*/    UMETA(DisplayName = "Dynamic"),
-	Chaos_Object_UserDefined     = 100                                 UMETA(DisplayName = "User Defined"),
+	Chaos_NONE = 0 UMETA(Hidden, DisplayName = "None"),
+	Chaos_Object_Sleeping  = 1 UMETA(DisplayName = "Sleeping"),
+	Chaos_Object_Kinematic = 2 UMETA(DisplayName = "Kinematic"),
+	Chaos_Object_Static = 3    UMETA(DisplayName = "Static"),
+	Chaos_Object_Dynamic = 4 UMETA(DisplayName = "Dynamic"),
+	Chaos_Object_UserDefined = 100 UMETA(DisplayName = "User Defined"),
 	//
-	Chaos_Max                UMETA(Hidden)
+	Chaos_Max UMETA(Hidden)
 };
 
 UENUM(BlueprintType)
