Chaos - clean up stats

diff --git a/Engine/Source/Runtime/Experimental/Chaos/Private/PBDRigidsSolver.cpp b/Engine/Source/Runtime/Experimental/Chaos/Private/PBDRigidsSolver.cpp
--- a/Engine/Source/Runtime/Experimental/Chaos/Private/PBDRigidsSolver.cpp
+++ b/Engine/Source/Runtime/Experimental/Chaos/Private/PBDRigidsSolver.cpp
@@ -284,38 +284,6 @@ namespace Chaos
 #endif
 			}
 
-			{
-				// Particle counts
-				SET_DWORD_STAT(STAT_ChaosCounter_NumDisabledParticles, MSolver->GetEvolution()->GetParticles().GetAllParticlesView().Num() - MSolver->GetEvolution()->GetParticles().GetNonDisabledView().Num());
-				SET_DWORD_STAT(STAT_ChaosCounter_NumParticles, MSolver->GetEvolution()->GetParticles().GetNonDisabledView().Num());
-				SET_DWORD_STAT(STAT_ChaosCounter_NumDynamicParticles, MSolver->GetEvolution()->GetParticles().GetNonDisabledDynamicView().Num());
-				SET_DWORD_STAT(STAT_ChaosCounter_NumActiveDynamicParticles, MSolver->GetEvolution()->GetParticles().GetActiveParticlesView().Num());
-				SET_DWORD_STAT(STAT_ChaosCounter_NumKinematicParticles, MSolver->GetEvolution()->GetParticles().GetActiveKinematicParticlesView().Num());
-				SET_DWORD_STAT(STAT_ChaosCounter_NumStaticParticles, MSolver->GetEvolution()->GetParticles().GetActiveStaticParticlesView().Num());
-				SET_DWORD_STAT(STAT_ChaosCounter_NumGeometryCollectionParticles, MSolver->GetEvolution()->GetParticles().GetGeometryCollectionParticles().Size());
-
-				// Constraint counts
-				SET_DWORD_STAT(STAT_ChaosCounter_NumIslands, MSolver->GetEvolution()->GetConstraintGraph().NumIslands());
-				SET_DWORD_STAT(STAT_ChaosCounter_NumContacts, MSolver->NumCollisionConstraints());
-				SET_DWORD_STAT(STAT_ChaosCounter_NumJoints, MSolver->NumJointConstraints());
-
-#if CSV_PROFILER
-				// Particle counts
-				CSV_CUSTOM_STAT(ChaosCounters, NumDisabledParticles, MSolver->GetEvolution()->GetParticles().GetAllParticlesView().Num() - MSolver->GetEvolution()->GetParticles().GetNonDisabledView().Num(), ECsvCustomStatOp::Accumulate);
-				CSV_CUSTOM_STAT(ChaosCounters, NumParticles, MSolver->GetEvolution()->GetParticles().GetNonDisabledView().Num(), ECsvCustomStatOp::Accumulate);
-				CSV_CUSTOM_STAT(ChaosCounters, NumDynamicParticles, MSolver->GetEvolution()->GetParticles().GetNonDisabledDynamicView().Num(), ECsvCustomStatOp::Accumulate);
-				CSV_CUSTOM_STAT(ChaosCounters, NumKinematicParticles, MSolver->GetEvolution()->GetParticles().GetActiveKinematicParticlesView().Num(), ECsvCustomStatOp::Accumulate);
-				CSV_CUSTOM_STAT(ChaosCounters, NumStaticParticles, MSolver->GetEvolution()->GetParticles().GetActiveStaticParticlesView().Num(), ECsvCustomStatOp::Accumulate);
-				CSV_CUSTOM_STAT(ChaosCounters, NumGeometryCollectionParticles, (int)MSolver->GetEvolution()->GetParticles().GetGeometryCollectionParticles().Size(), ECsvCustomStatOp::Accumulate);
-
-				// Constraint counts
-				CSV_CUSTOM_STAT(ChaosCounters, NumIslands, MSolver->GetEvolution()->GetConstraintGraph().NumIslands(), ECsvCustomStatOp::Accumulate);
-				CSV_CUSTOM_STAT(ChaosCounters, NumContacts, MSolver->NumCollisionConstraints(), ECsvCustomStatOp::Accumulate);
-				CSV_CUSTOM_STAT(ChaosCounters, NumJoints, MSolver->NumJointConstraints(), ECsvCustomStatOp::Accumulate);
-#endif
-			}
-
-
 			{
 				SCOPE_CYCLE_COUNTER(STAT_EventDataGathering);
 				{
@@ -344,6 +312,8 @@ namespace Chaos
 			MSolver->GetCurrentFrame()++;
 			MSolver->PostTickDebugDraw(MDeltaTime);
 
+			MSolver->UpdateStatCounters();
+
 			//Editor ticks with 0 dt. We don't want to buffer any dirty data from this since it won't be consumed
 			//TODO: handle this more gracefully
 			if(MDeltaTime > 0)
@@ -1211,6 +1181,38 @@ namespace Chaos
 		return GetEvolution()->GetCollisionConstraints().NumConstraints();
 	}
 
+	template <typename Traits>
+	void FPBDRigidsSolver<Traits>::UpdateStatCounters() const
+	{
+		// Particle counts
+		SET_DWORD_STAT(STAT_ChaosCounter_NumDisabledParticles, GetEvolution()->GetParticles().GetAllParticlesView().Num() - GetEvolution()->GetParticles().GetNonDisabledView().Num());
+		SET_DWORD_STAT(STAT_ChaosCounter_NumParticles, GetEvolution()->GetParticles().GetNonDisabledView().Num());
+		SET_DWORD_STAT(STAT_ChaosCounter_NumDynamicParticles, GetEvolution()->GetParticles().GetNonDisabledDynamicView().Num());
+		SET_DWORD_STAT(STAT_ChaosCounter_NumActiveDynamicParticles, GetEvolution()->GetParticles().GetActiveParticlesView().Num());
+		SET_DWORD_STAT(STAT_ChaosCounter_NumKinematicParticles, GetEvolution()->GetParticles().GetActiveKinematicParticlesView().Num());
+		SET_DWORD_STAT(STAT_ChaosCounter_NumStaticParticles, GetEvolution()->GetParticles().GetActiveStaticParticlesView().Num());
+		SET_DWORD_STAT(STAT_ChaosCounter_NumGeometryCollectionParticles, GetEvolution()->GetParticles().GetGeometryCollectionParticles().Size());
+
+		// Constraint counts
+		SET_DWORD_STAT(STAT_ChaosCounter_NumIslands, GetEvolution()->GetConstraintGraph().NumIslands());
+		SET_DWORD_STAT(STAT_ChaosCounter_NumContacts, NumCollisionConstraints());
+		SET_DWORD_STAT(STAT_ChaosCounter_NumJoints, NumJointConstraints());
+
+#if CSV_PROFILER
+		// Particle counts
+		CSV_CUSTOM_STAT(ChaosCounters, NumDisabledParticles, GetEvolution()->GetParticles().GetAllParticlesView().Num() - GetEvolution()->GetParticles().GetNonDisabledView().Num(), ECsvCustomStatOp::Accumulate);
+		CSV_CUSTOM_STAT(ChaosCounters, NumParticles, GetEvolution()->GetParticles().GetNonDisabledView().Num(), ECsvCustomStatOp::Accumulate);
+		CSV_CUSTOM_STAT(ChaosCounters, NumDynamicParticles, GetEvolution()->GetParticles().GetNonDisabledDynamicView().Num(), ECsvCustomStatOp::Accumulate);
+		CSV_CUSTOM_STAT(ChaosCounters, NumKinematicParticles, GetEvolution()->GetParticles().GetActiveKinematicParticlesView().Num(), ECsvCustomStatOp::Accumulate);
+		CSV_CUSTOM_STAT(ChaosCounters, NumStaticParticles, GetEvolution()->GetParticles().GetActiveStaticParticlesView().Num(), ECsvCustomStatOp::Accumulate);
+		CSV_CUSTOM_STAT(ChaosCounters, NumGeometryCollectionParticles, (int)GetEvolution()->GetParticles().GetGeometryCollectionParticles().Size(), ECsvCustomStatOp::Accumulate);
+
+		// Constraint counts
+		CSV_CUSTOM_STAT(ChaosCounters, NumIslands, GetEvolution()->GetConstraintGraph().NumIslands(), ECsvCustomStatOp::Accumulate);
+		CSV_CUSTOM_STAT(ChaosCounters, NumContacts, NumCollisionConstraints(), ECsvCustomStatOp::Accumulate);
+		CSV_CUSTOM_STAT(ChaosCounters, NumJoints, NumJointConstraints(), ECsvCustomStatOp::Accumulate);
+#endif
+	}
 
 	void FPBDRigidsSolver::PostTickDebugDraw(FReal Dt) const
 	{
diff --git a/Engine/Source/Runtime/Experimental/Chaos/Public/PBDRigidsSolver.h b/Engine/Source/Runtime/Experimental/Chaos/Public/PBDRigidsSolver.h
--- a/Engine/Source/Runtime/Experimental/Chaos/Public/PBDRigidsSolver.h
+++ b/Engine/Source/Runtime/Experimental/Chaos/Public/PBDRigidsSolver.h
@@ -312,6 +312,9 @@ namespace Chaos
 		/** Update the particles forces based on field evaluation */
 		void FieldForcesUpdateCallback();
 
+		// Update the counter in Stats and the CSV profiler
+		void UpdateStatCounters() const;
+
 	private:
 
 		/**/
