Fix build

diff --git a/Engine/Source/Runtime/Experimental/Chaos/Private/PBDRigidsSolver.cpp b/Engine/Source/Runtime/Experimental/Chaos/Private/PBDRigidsSolver.cpp
--- a/Engine/Source/Runtime/Experimental/Chaos/Private/PBDRigidsSolver.cpp
+++ b/Engine/Source/Runtime/Experimental/Chaos/Private/PBDRigidsSolver.cpp
@@ -1181,8 +1181,7 @@ namespace Chaos
 		return GetEvolution()->GetCollisionConstraints().NumConstraints();
 	}
 
-	template <typename Traits>
-	void FPBDRigidsSolver<Traits>::UpdateStatCounters() const
+	void FPBDRigidsSolver::UpdateStatCounters() const
 	{
 		// Particle counts
 		SET_DWORD_STAT(STAT_ChaosCounter_NumDisabledParticles, GetEvolution()->GetParticles().GetAllParticlesView().Num() - GetEvolution()->GetParticles().GetNonDisabledView().Num());
