Trivial fix of redundant code

diff --git a/Engine/Source/Runtime/Experimental/Chaos/Private/Chaos/CollisionResolution.cpp b/Engine/Source/Runtime/Experimental/Chaos/Private/Chaos/CollisionResolution.cpp
--- a/Engine/Source/Runtime/Experimental/Chaos/Private/Chaos/CollisionResolution.cpp
+++ b/Engine/Source/Runtime/Experimental/Chaos/Private/Chaos/CollisionResolution.cpp
@@ -2927,7 +2927,7 @@ namespace Chaos
 				TPBDRigidParticleHandle<FReal, 3>* RigidParticle0 = Particle0->CastToRigidParticle(), *RigidParticle1 = Particle1->CastToRigidParticle();
 				if(RigidParticle0 && RigidParticle1)
 				{
-					bDeferUpdate |= (RigidParticle0->ObjectState() == EObjectStateType::Sleeping) || (RigidParticle0->ObjectState() == EObjectStateType::Sleeping);
+					bDeferUpdate |= (RigidParticle0->ObjectState() == EObjectStateType::Sleeping) || (RigidParticle1->ObjectState() == EObjectStateType::Sleeping);
 				}
 			}
 			if (bDeferUpdate)
