Add init velocvity for the GC apply physics field

diff --git a/Engine/Source/Runtime/Experimental/Chaos/Private/PhysicsProxy/GeometryCollectionPhysicsProxy.cpp b/Engine/Source/Runtime/Experimental/Chaos/Private/PhysicsProxy/GeometryCollectionPhysicsProxy.cpp
--- a/Engine/Source/Runtime/Experimental/Chaos/Private/PhysicsProxy/GeometryCollectionPhysicsProxy.cpp
+++ b/Engine/Source/Runtime/Experimental/Chaos/Private/PhysicsProxy/GeometryCollectionPhysicsProxy.cpp
@@ -2504,21 +2504,30 @@ void FGeometryCollectionPhysicsProxy::FieldParameterUpdateCallback(Chaos::FPBDRi
 						{
 							bool bHasStateChanged = false;
 							InitDynamicStateResults(ParticleHandles, FieldContext, LocalResults);
-							
+
 							static_cast<const FFieldNode<int32>*>(FieldCommand.RootNode.Get())->Evaluate(FieldContext, ResultsView);
 							for (const FFieldContextIndex& Index : FieldContext.GetEvaluatedSamples())
 							{
 								Chaos::TPBDRigidParticleHandle<float, 3>* RigidHandle = ParticleHandles[Index.Sample]->CastToRigidParticle();
 								if (RigidHandle)
 								{
+									const int8 ResultState = ResultsView[Index.Result];  
+									const int32 TransformIndex = HandleToTransformGroupIndex[RigidHandle];
+
 									// Update of the handles object state. No need to update 
 									// the initial velocities since it is done after this function call in InitializeBodiesPT
-									const int8 ResultState = ResultsView[Index.Result];
-									bHasStateChanged |= ReportDynamicStateResult(RigidSolver, static_cast<Chaos::EObjectStateType>(ResultState), RigidHandle,
-										false, Chaos::TVector<float, 3>(0), false, Chaos::TVector<float, 3>(0));
-
+									if (bUpdateViews && (Parameters.InitialVelocityType == EInitialVelocityTypeEnum::Chaos_Initial_Velocity_User_Defined))
+									{
+										bHasStateChanged |= ReportDynamicStateResult(RigidSolver, static_cast<Chaos::EObjectStateType>(ResultState), RigidHandle,
+											true, Collection.InitialLinearVelocity[TransformIndex], true, Collection.InitialAngularVelocity[TransformIndex]);
+									}
+									else
+									{
+										bHasStateChanged |= ReportDynamicStateResult(RigidSolver, static_cast<Chaos::EObjectStateType>(ResultState), RigidHandle,
+											false, Chaos::TVector<float, 3>(0), false, Chaos::TVector<float, 3>(0));
+									}
 									// Update of the Collection dynamic state. It will be used just after to set the initial velocity
-									Collection.DynamicState[HandleToTransformGroupIndex[RigidHandle]] = ResultState;
+									Collection.DynamicState[TransformIndex] = ResultState;
 								}
 							}
 							if (bUpdateViews)
diff --git a/Engine/Source/Runtime/Experimental/Chaos/Public/PhysicsProxy/GeometryCollectionPhysicsProxy.h b/Engine/Source/Runtime/Experimental/Chaos/Public/PhysicsProxy/GeometryCollectionPhysicsProxy.h
--- a/Engine/Source/Runtime/Experimental/Chaos/Public/PhysicsProxy/GeometryCollectionPhysicsProxy.h
+++ b/Engine/Source/Runtime/Experimental/Chaos/Public/PhysicsProxy/GeometryCollectionPhysicsProxy.h
@@ -245,8 +245,6 @@ protected:
 		const TArray<int32>& ChildTransformGroupIndices,
 		const Chaos::FClusterCreationParameters & Parameters);
 
-	void PushKinematicStateToSolver();
-
 	/** 
 	 * Traverses the parents of \p TransformIndex in \p GeometryCollection, counting
 	 * the number of levels until the next parent is \c INDEX_NONE.
