Fix non-unity build errors in ToolTarget.

#rnx

diff --git a/Engine/Source/Runtime/Experimental/InteractiveToolsFramework/Public/ToolTargets/ToolTarget.h b/Engine/Source/Runtime/Experimental/InteractiveToolsFramework/Public/ToolTargets/ToolTarget.h
--- a/Engine/Source/Runtime/Experimental/InteractiveToolsFramework/Public/ToolTargets/ToolTarget.h
+++ b/Engine/Source/Runtime/Experimental/InteractiveToolsFramework/Public/ToolTargets/ToolTarget.h
@@ -59,7 +59,7 @@ public:
  * Base class for factories of tool targets, which let a tool manager build targets
  * out of inputs without knowing anything itself.
  */
-UCLASS(Transient)
+UCLASS(Transient, Abstract)
 class INTERACTIVETOOLSFRAMEWORK_API UToolTargetFactory : public UObject
 {
 	GENERATED_BODY()
