Fix UE4 reference in comments

diff --git a/Engine/Source/Runtime/Experimental/GeometryCollectionEngine/Private/GeometryCollection/StaticMeshSimulationComponent.cpp b/Engine/Source/Runtime/Experimental/GeometryCollectionEngine/Private/GeometryCollection/StaticMeshSimulationComponent.cpp
--- a/Engine/Source/Runtime/Experimental/GeometryCollectionEngine/Private/GeometryCollection/StaticMeshSimulationComponent.cpp
+++ b/Engine/Source/Runtime/Experimental/GeometryCollectionEngine/Private/GeometryCollection/StaticMeshSimulationComponent.cpp
@@ -60,7 +60,7 @@ void UStaticMeshSimulationComponent::TickComponent(float DeltaTime, enum ELevelT
 	Super::TickComponent(DeltaTime, TickType, ThisTickFunction);
 
 #if 0	//todo: remove
-	// for kinematic objects, we assume that UE4 can and will move them, so we need to pass the new data to the phys solver
+	// for kinematic objects, we assume that Unreal can and will move them, so we need to pass the new data to the phys solver
 	if ((ObjectType == EObjectStateTypeEnum::Chaos_Object_Kinematic) && Simulating)
 	{
 		FChaosSolversModule* ChaosModule = FModuleManager::Get().GetModulePtr<FChaosSolversModule>("ChaosSolvers");
