Remove UE4 reference from GeometryCollection debug OBJ output writer.

diff --git a/Engine/Source/Runtime/Experimental/Chaos/Private/GeometryCollection/GeometryCollection.cpp b/Engine/Source/Runtime/Experimental/Chaos/Private/GeometryCollection/GeometryCollection.cpp
--- a/Engine/Source/Runtime/Experimental/Chaos/Private/GeometryCollection/GeometryCollection.cpp
+++ b/Engine/Source/Runtime/Experimental/Chaos/Private/GeometryCollection/GeometryCollection.cpp
@@ -1386,7 +1386,7 @@ void FGeometryCollection::WriteDataToOBJFile(const FString &Name, const FString
 
 		DataFile.open(string(TCHAR_TO_UTF8(*FullPath)));
 
-		DataFile << "# File exported from UE4" << endl;
+		DataFile << "# File exported from Unreal Engine" << endl;
 		DataFile << "# " << NumVertices << " points" << endl;
 		DataFile << "# " << NumVertices * 3 << " vertices" << endl;
 		DataFile << "# " << NumFaces << " primitives" << endl;
