Make bIsPropertySetEnabled a TransientToolProperty so that it doesn't get saved with the property set, so disabling a property set in one tool doesn't end up disabling it in all tools.

#rnx

diff --git a/Engine/Source/Runtime/Experimental/InteractiveToolsFramework/Public/InteractiveTool.h b/Engine/Source/Runtime/Experimental/InteractiveToolsFramework/Public/InteractiveTool.h
--- a/Engine/Source/Runtime/Experimental/InteractiveToolsFramework/Public/InteractiveTool.h
+++ b/Engine/Source/Runtime/Experimental/InteractiveToolsFramework/Public/InteractiveTool.h
@@ -226,7 +226,8 @@ protected:
 	UPROPERTY()
 	TObjectPtr<UInteractiveToolPropertySet> CachedProperties = nullptr;
 
-	UPROPERTY()
+	// Controls whether a property set is shown in the UI.  Transient so that disabling a PropertySet in one tool doesn't disable it in others.
+	UPROPERTY(meta=(TransientToolProperty))
 	bool bIsPropertySetEnabled = true;
 
 	friend class UInteractiveTool;	// so that tool can enable/disable
