Chaos - fix scaled tri mesh sweeps (scaled convex fix knock-on bug)

diff --git a/Engine/Source/Runtime/Experimental/Chaos/Public/Chaos/ImplicitObject.h b/Engine/Source/Runtime/Experimental/Chaos/Public/Chaos/ImplicitObject.h
--- a/Engine/Source/Runtime/Experimental/Chaos/Public/Chaos/ImplicitObject.h
+++ b/Engine/Source/Runtime/Experimental/Chaos/Public/Chaos/ImplicitObject.h
@@ -271,8 +271,18 @@ public:
 
 	virtual int32 FindMostOpposingFaceScaled(const FVec3& Position, const FVec3& UnitDir, int32 HintFaceIndex, FReal SearchDist, const FVec3& Scale) const
 	{
-		//Many objects have no concept of a face
-		return INDEX_NONE;
+		// For now, this is the implementation that used to be in FImplicitObjectScaled to use as a default.
+		// @todo(chaos): implement FindMostOpposingFaceScaled for all types that can be wrapped in ImplicitObjectScaled
+		// NOTE: this isn't strictly correct. The distance check does not account for non-uniforms scales, but also the algorithm is finding the
+		// face which has the most opposing normal. The scaled-space normal for a face i is Ns_i = (N_i / S) / |(N_i / S)|
+		// We want to find i with the minimum F_i = (D . Ns_i), where D is the unit direction we are interested in.
+		// F_i = D . (N_i / S) / |(N_i / S)|
+		// Below we are effectively testing F_i = ((D / S) . N_i) and ignoring the |(N_i / S)| which is dependent on i and therefore cannot 
+		// be ignored in the minimum check.
+		const FReal UnscaledSearchDist = SearchDist * Scale.Max();	//this is not quite right since it's no longer a sphere, but the whole thing is fuzzy anyway
+		const FVec3 UnscaledPosition = Position / Scale;
+		const FVec3 UnscaledNormal = GetInnerUnscaledNormal(UnitDir, Scale);
+		return FindMostOpposingFace(UnscaledPosition, UnscaledNormal, HintFaceIndex, UnscaledSearchDist);
 	}
 
 
@@ -350,6 +360,29 @@ public:
 
 protected:
 
+	// Safely scale a normal - used with both scale and inverse scale
+	FVec3 ScaledNormalHelper(const FVec3& Normal, const FVec3& Scale) const
+	{
+		const FVec3 ScaledNormal = Scale * Normal;
+		const float ScaledNormalLen = ScaledNormal.Size();
+		return ensure(ScaledNormalLen > TNumericLimits<FReal>::Min())
+			? ScaledNormal / ScaledNormalLen
+			: FVec3(0.f, 0.f, 1.f);
+	}
+
+	// Convert a normal in the inner unscaled object space into a normal in the outer scaled object space.
+	// Note: INverse scale is passed in, not the scale
+	FVec3 GetOuterScaledNormal(const FVec3& InnerNormal, const FVec3& Scale) const
+	{
+		return ScaledNormalHelper(InnerNormal, FVec3(1.0f / Scale.X, 1.0f / Scale.Y, 1.0f / Scale.Z));
+	}
+
+	// Convert a normal in the outer scaled object space into a normal in the inner unscaled object space
+	FVec3 GetInnerUnscaledNormal(const FVec3& OuterNormal, const FVec3& Scale) const
+	{
+		return ScaledNormalHelper(OuterNormal, Scale);
+	}
+
 	// Not all derived types support a margin, and for some it represents some other
 	// property (the radius of a sphere for example), so the setter should only be exposed 
 	// in a derived class if at all (and it may want to change the size of the core shape as well)
diff --git a/Engine/Source/Runtime/Experimental/Chaos/Public/Chaos/ImplicitObjectScaled.h b/Engine/Source/Runtime/Experimental/Chaos/Public/Chaos/ImplicitObjectScaled.h
--- a/Engine/Source/Runtime/Experimental/Chaos/Public/Chaos/ImplicitObjectScaled.h
+++ b/Engine/Source/Runtime/Experimental/Chaos/Public/Chaos/ImplicitObjectScaled.h
@@ -630,7 +630,7 @@ public:
 	// Get the index of the plane that most opposes the normal, assuming it passes through the specified vertex
 	int32 GetMostOpposingPlaneWithVertex(int32 VertexIndex, const FVec3& Normal) const
 	{
-		return MObject->GetMostOpposingPlane(VertexIndex, GetUnscaledNormal(Normal));
+		return MObject->GetMostOpposingPlane(VertexIndex, GetInnerUnscaledNormal(Normal));
 	}
 
 	// Get the nearest point on an edge of the specified face
@@ -877,30 +877,6 @@ private:
 		return new TImplicitObjectScaled<TConcrete, false>(Obj->MObject->Copy(), Obj->MScale, Obj->OuterMargin);
 	}
 
-	// Convert a normal in the inner unscaled object space into a normal in the outer scaled object space.
-	FVec3 GetScaledNormal(const TVector<T, d>& OuterNormal) const
-	{
-		const TVector<T, d> UnscaledDirDenorm = MInvScale * OuterNormal;
-		const T LengthScale = UnscaledDirDenorm.Size();
-		const TVector<T, d> UnscaledDir
-			= ensure(LengthScale > TNumericLimits<T>::Min())
-			? UnscaledDirDenorm / LengthScale
-			: TVector<T, d>(0.f, 0.f, 1.f);
-		return UnscaledDir;
-	}
-
-	// Convert a normal in the outer scaled object space into a normal in the inner unscaled object space
-	FVec3 GetUnscaledNormal(const TVector<T, d>& InnerNormal) const
-	{
-		const TVector<T, d> ScaledDirDenorm = MScale * InnerNormal;
-		const T LengthScale = ScaledDirDenorm.Size();
-		const TVector<T, d> ScaledDir
-			= ensure(LengthScale > TNumericLimits<T>::Min())
-			? ScaledDirDenorm / LengthScale
-			: TVector<T, d>(0.f, 0.f, 1.f);
-		return ScaledDir;
-	}
-
 	void UpdateBounds()
 	{
 		const TAABB<T, d> UnscaledBounds = MObject->BoundingBox();
