Chaos : Clustering
-Added particle deletion for when geometry collection actors are removed from a cluster. This is a fast path that assumes the collection will not be activating any of the children.
(DestroyClusterParticle)
-Added a PendingDestroyGeometryCollectionProxy
-Added IsObjectDeleting to avoid any processing of the collections with in a torn down state.
-Removed unused path for Move bodies on deactivation.
-Moved deletion of the geometry collection proxy from the command in UnregisterObject to DestroyPendingProxies_Internal. Appears that other proxies are now being released here and it fixes a invalid access crash post OnRemoveFromSolver.
-Added a check for timestamp in PullPhysicsDataImp, was crashing on a null Timestamp.
#okforgethub public, internal

diff --git a/Engine/Source/Runtime/Experimental/Chaos/Private/Chaos/PBDRigidClustering.cpp b/Engine/Source/Runtime/Experimental/Chaos/Private/Chaos/PBDRigidClustering.cpp
--- a/Engine/Source/Runtime/Experimental/Chaos/Private/Chaos/PBDRigidClustering.cpp
+++ b/Engine/Source/Runtime/Experimental/Chaos/Private/Chaos/PBDRigidClustering.cpp
@@ -64,9 +64,6 @@ namespace Chaos
 	int32 MinCleanedPointsBeforeRemovingInternals = 10;
 	FAutoConsoleVariableRef CVarMinCleanedPointsBeforeRemovingInternals(TEXT("p.MinCleanedPointsBeforeRemovingInternals"), MinCleanedPointsBeforeRemovingInternals, TEXT("If we only have this many clean points, don't bother removing internal points as the object is likely very small"));
 
-	int32 MoveClustersWhenDeactivated = 0;
-	FAutoConsoleVariableRef CVarMoveClustersWhenDeactivated(TEXT("p.MoveClustersWhenDeactivated"), MoveClustersWhenDeactivated, TEXT("If clusters should be moved when deactivated."));
-
 	int32 DeactivateClusterChildren = 0;
 	FAutoConsoleVariableRef CVarDeactivateClusterChildren(TEXT("p.DeactivateClusterChildren"), DeactivateClusterChildren, TEXT("If children should be decativated when broken and put into another cluster."));
 
@@ -530,7 +527,7 @@ namespace Chaos
 		NewParticle->SetCollisionGroup(INT_MAX);
 		TopLevelClusterParents.Add(NewParticle);
 		NewParticle->SetInternalCluster(true);
-		NewParticle->SetClusterId(ClusterId(Parent, Children.Num()));
+		NewParticle->SetClusterId(ClusterId(nullptr, Children.Num()));
 		for (auto& Constituent : Children) MEvolution.DoInternalParticleInitilization(Constituent, NewParticle);
 
 		//
@@ -1206,6 +1203,7 @@ namespace Chaos
 		return ActivatedBodies;
 	}
 
+
 	DECLARE_CYCLE_STAT(TEXT("TPBDRigidClustering<>::AdvanceClustering"), STAT_AdvanceClustering, STATGROUP_Chaos);
 	DECLARE_CYCLE_STAT(TEXT("TPBDRigidClustering<>::Update Impulse from Strain"), STAT_UpdateImpulseStrain, STATGROUP_Chaos);
 	DECLARE_CYCLE_STAT(TEXT("TPBDRigidClustering<>::Update Dirty Impulses"), STAT_UpdateDirtyImpulses, STATGROUP_Chaos);
@@ -1973,14 +1971,6 @@ namespace Chaos
 	{
 		// #note: we don't recursively descend to the children
 		MEvolution.DisableParticle(ClusteredParticle);
-
-		if (MoveClustersWhenDeactivated)
-		{
-			ClusteredParticle->P() -= FVector(0.f, 0.f, -10000.f); // HACK : Move them away to prevent reactivation. 
-			ClusteredParticle->X() -= FVector(0.f, 0.f, -10000.f); // HACK : Move them away to prevent reactivation. 
-			ClusteredParticle->V() = FVector(0.f);            // HACK : Move them away to prevent reactivation.
-		}
-
 		TopLevelClusterParents.Remove(ClusteredParticle);
 		GetChildrenMap().Remove(ClusteredParticle);
 		ClusteredParticle->ClusterIds() = ClusterId();
@@ -2007,6 +1997,78 @@ namespace Chaos
 		}
 	}
 
+	template<class T_FPBDRigidsEvolution, class T_FPBDCollisionConstraint>
+	FPBDRigidClusteredParticleHandle*
+	TPBDRigidClustering<T_FPBDRigidsEvolution, T_FPBDCollisionConstraint>::DestroyClusterParticle(
+		FPBDRigidClusteredParticleHandle* ClusteredParticle,
+		const FClusterDestoryParameters& Parameters)
+	{
+		FClusterHandle ParentParticle = nullptr;
+
+		// detach connections to thie parent from the children
+		if (MChildren.Contains(ClusteredParticle))
+		{
+			for (FRigidHandle Child : MChildren[ClusteredParticle])
+			{
+				if (FClusterHandle ClusteredChild = Child->CastToClustered())
+				{
+					ClusteredChild->ClusterIds() = ClusterId();
+					ClusteredChild->ClusterGroupIndex() = 0;
+				}
+			}
+
+			MChildren.Remove(ClusteredParticle);
+		}
+
+		// disable within the solver
+		if (!ClusteredParticle->Disabled())
+		{
+			MEvolution.DisableParticle(ClusteredParticle);
+			ensure(ClusteredParticle->ClusterIds().Id == nullptr);
+		}
+
+		// reset the structures
+		TopLevelClusterParents.Remove(ClusteredParticle);
+		MActiveRemovalIndices.Remove(ClusteredParticle);
+
+		// disconnect from the parents
+		if (ClusteredParticle->ClusterIds().Id)
+		{
+			ParentParticle = ClusteredParticle->ClusterIds().Id->CastToClustered();
+
+			ClusteredParticle->ClusterIds() = ClusterId();
+			ClusteredParticle->ClusterGroupIndex() = 0;
+
+			if (MChildren.Contains(ParentParticle))
+			{
+				FRigidHandleArray& Children = MChildren[ParentParticle];
+
+				// disconnect from your parents children list
+				Children.Remove(ClusteredParticle);
+
+				// disable internal parents that have lost all their children
+				if (!MChildren[ParentParticle].Num() && ParentParticle->InternalCluster())
+				{
+					DisableCluster(ClusteredParticle);
+				}
+			}
+		}
+
+		// remove internal parents that have no children. 
+		if (ClusteredParticle->InternalCluster())
+		{
+			MEvolution.DestroyParticle(ClusteredParticle);
+		}
+
+		if (Parameters.bReturnInternalOnly && ParentParticle && !ParentParticle->InternalCluster())
+		{
+			ParentParticle = nullptr;
+		}
+		return ParentParticle;
+
+	}
+
+
 	DECLARE_CYCLE_STAT(TEXT("TPBDRigidClustering<>::UpdateConnectivityGraphUsingPointImplicit"), STAT_UpdateConnectivityGraphUsingPointImplicit, STATGROUP_Chaos);
 	template<class T_FPBDRigidsEvolution, class T_FPBDCollisionConstraint>
 	void TPBDRigidClustering<T_FPBDRigidsEvolution, T_FPBDCollisionConstraint>::UpdateConnectivityGraphUsingPointImplicit(
diff --git a/Engine/Source/Runtime/Experimental/Chaos/Private/PBDRigidsSolver.cpp b/Engine/Source/Runtime/Experimental/Chaos/Private/PBDRigidsSolver.cpp
--- a/Engine/Source/Runtime/Experimental/Chaos/Private/PBDRigidsSolver.cpp
+++ b/Engine/Source/Runtime/Experimental/Chaos/Private/PBDRigidsSolver.cpp
@@ -553,7 +553,8 @@ namespace Chaos
 				GeometryCollectionPhysicsProxies_Internal.RemoveSingle(InProxy);
 				InProxy->SyncBeforeDestroy();
 				InProxy->OnRemoveFromSolver(this);
-				delete InProxy;
+				InProxy->ResetDirtyIdx();
+				PendingDestroyGeometryCollectionPhysicsProxy.Add(InProxy);
 			});
 	}
 
@@ -689,6 +690,19 @@ namespace Chaos
 			delete Proxy;
 		}
 		PendingDestroyPhysicsProxy.Reset();
+
+		bool bResetCollisionConstraints=false;
+		for (auto Proxy : PendingDestroyGeometryCollectionPhysicsProxy)
+		{
+			//ensure(Proxy->GetHandle_LowLevel() == nullptr);	//should have already cleared this out
+			MarshallingManager.GetCurrentPullData_Internal()->DirtyGeometryCollections.Reset();
+			bResetCollisionConstraints = true;
+			delete Proxy;
+		}
+		PendingDestroyGeometryCollectionPhysicsProxy.Reset();
+
+		if(bResetCollisionConstraints)
+			GetEvolution()->GetCollisionConstraints();
 	}
 
 	void FPBDRigidsSolver::AdvanceSolverBy(const FReal DeltaTime, const FSubStepInfo& SubStepInfo)
diff --git a/Engine/Source/Runtime/Experimental/Chaos/Private/PhysicsProxy/GeometryCollectionPhysicsProxy.cpp b/Engine/Source/Runtime/Experimental/Chaos/Private/PhysicsProxy/GeometryCollectionPhysicsProxy.cpp
--- a/Engine/Source/Runtime/Experimental/Chaos/Private/PhysicsProxy/GeometryCollectionPhysicsProxy.cpp
+++ b/Engine/Source/Runtime/Experimental/Chaos/Private/PhysicsProxy/GeometryCollectionPhysicsProxy.cpp
@@ -318,6 +318,7 @@ FGeometryCollectionPhysicsProxy::FGeometryCollectionPhysicsProxy(
 	, BaseParticleIndex(INDEX_NONE)
 	, IsObjectDynamic(false)
 	, IsObjectLoading(true)
+	, IsObjectDeleting(false)
 	, SimFilter(InSimFilter)
 	, QueryFilter(InQueryFilter)
 #if TODO_REIMPLEMENT_RIGID_CACHING
@@ -1331,31 +1332,57 @@ void FGeometryCollectionPhysicsProxy::InitializeRemoveOnFracture(FParticlesType&
 
 void FGeometryCollectionPhysicsProxy::OnRemoveFromSolver(Chaos::FPBDRigidsSolver *RBDSolver)
 {
-	const FGeometryDynamicCollection& DynamicCollection = PhysicsThreadCollection;
-
 	Chaos::FPBDRigidsEvolutionGBF* Evolution = RBDSolver->GetEvolution();
 
-	for (const FClusterHandle* Handle : SolverClusterHandles)
+	TSet< FClusterHandle* > ClustersToReuild;
+	for (int i = 0; i < SolverParticleHandles.Num(); i++)
 	{
-		RBDSolver->RemoveParticleToProxy(Handle);
-	}
-
-	for (FClusterHandle* Handle : SolverParticleHandles)
-	{	
-		if (Handle)
+		if (FClusterHandle* Handle = SolverParticleHandles[i])
 		{
-			if (Chaos::TPBDRigidClusteredParticleHandle<Chaos::FReal, 3>* Cluster = Handle->CastToClustered())
+			RBDSolver->RemoveParticleToProxy(Handle);
+			if (FClusterHandle* ParentCluster = Evolution->GetRigidClustering().DestroyClusterParticle(Handle))
 			{
-				Evolution->GetRigidClustering().GetTopLevelClusterParents().Remove(Cluster);
-				Evolution->GetRigidClustering().GetChildrenMap().Remove(Cluster);
-				Evolution->DestroyParticle(Cluster);
+				if (ParentCluster->InternalCluster())
+				{
+					ClustersToReuild.Add(ParentCluster);
+				}
 			}
-			else
+		}
+	}
+
+	for (int i = 0; i < SolverParticleHandles.Num(); i++)
+	{
+		if (FClusterHandle* Handle = SolverParticleHandles[i])
+		{
+			Evolution->DestroyParticle(Handle);
+		}
+	}
+
+	for (FClusterHandle* Cluster : ClustersToReuild)
+	{
+		ensure(Cluster->InternalCluster());
+		if (ensure(Evolution->GetRigidClustering().GetChildrenMap().Contains(Cluster)))
+		{
+			// copy cluster state for recreation
+			int32 ClusterGroupIndex = Cluster->ClusterGroupIndex();
+			TArray<FParticleHandle*> Children = Evolution->GetRigidClustering().GetChildrenMap()[Cluster];
+
+			// destroy the invalid cluster
+			FClusterHandle* NullHandle = Evolution->GetRigidClustering().DestroyClusterParticle(Cluster);
+			ensure(NullHandle == nullptr);
+
+			// create a new cluster if needed
+			if (Children.Num())
 			{
-				Evolution->DestroyParticle(Handle);
+				if (FClusterHandle* NewParticle = Evolution->GetRigidClustering().CreateClusterParticle(ClusterGroupIndex, MoveTemp(Children)))
+				{
+					NewParticle->SetInternalCluster(true);
+				}
 			}
-		}
+		}		
 	}
+
+	IsObjectDeleting = true;
 }
 
 void FGeometryCollectionPhysicsProxy::OnRemoveFromScene()
@@ -1414,7 +1441,7 @@ void FGeometryCollectionPhysicsProxy::BufferPhysicsResults(Chaos::FPBDRigidsSolv
 	 */
 	using namespace Chaos;
 	SCOPE_CYCLE_COUNTER(STAT_CacheResultGeomCollection);
-
+	if (IsObjectDeleting) return;
 	BufferData.SetProxy(*this);
 
 	IsObjectDynamic = false;
@@ -1626,6 +1653,8 @@ void FGeometryCollectionPhysicsProxy::FlipBuffer()
 // Called from FPhysScene_ChaosInterface::SyncBodies(), NOT the solver.
 bool FGeometryCollectionPhysicsProxy::PullFromPhysicsState(const Chaos::FDirtyGeometryCollectionData& BufferData, const int32 SolverSyncTimestamp)
 {
+	if(IsObjectDeleting) return false;
+
 	/**
 	 * CONTEXT: GAMETHREAD (Read Locked)
 	 * Perform a similar operation to Sync, but take the data from a gamethread-safe buffer. This will be called
diff --git a/Engine/Source/Runtime/Experimental/Chaos/Public/Chaos/PBDRigidClustering.h b/Engine/Source/Runtime/Experimental/Chaos/Public/Chaos/PBDRigidClustering.h
--- a/Engine/Source/Runtime/Experimental/Chaos/Public/Chaos/PBDRigidClustering.h
+++ b/Engine/Source/Runtime/Experimental/Chaos/Public/Chaos/PBDRigidClustering.h
@@ -33,6 +33,10 @@ public:
 	TArray<Chaos::TSerializablePtr<FImplicitObject>> GeometryPtrs;
 };
 
+struct CHAOS_API FClusterDestoryParameters {
+	bool bReturnInternalOnly : true;
+};
+
 /* 
 * PDBRigidClustering
 */
@@ -45,6 +49,7 @@ public:
 
 	typedef FPBDRigidParticleHandle*					FRigidHandle;
 	typedef TArray<FRigidHandle>						FRigidHandleArray;
+	typedef TSet<FPBDRigidParticleHandle>				FRigidHandleSet;
 
 	typedef FPBDRigidClusteredParticleHandle*			FClusterHandle;
 
@@ -129,6 +134,18 @@ public:
 	TSet<FPBDRigidParticleHandle*> ReleaseClusterParticles(
 		TArray<FPBDRigidParticleHandle*> ChildrenParticles);
 
+	/*
+	*  DestroyClusterParticle
+	*    Disable the cluster particle and remove from all internal clustering
+	*    structures. This will not activate children.
+	* 
+	*    Returns the active parent cluster that might need to be rebuilt because
+	*    its geometry might be pointing to deleted particle handels. 
+	*/
+	FPBDRigidClusteredParticleHandle* DestroyClusterParticle(
+		FPBDRigidClusteredParticleHandle* ClusteredParticle,
+		const FClusterDestoryParameters& Parameters = FClusterDestoryParameters());
+
 	//
 	// Operational 
 	//
diff --git a/Engine/Source/Runtime/Experimental/Chaos/Public/Chaos/PullPhysicsDataImp.h b/Engine/Source/Runtime/Experimental/Chaos/Public/Chaos/PullPhysicsDataImp.h
--- a/Engine/Source/Runtime/Experimental/Chaos/Public/Chaos/PullPhysicsDataImp.h
+++ b/Engine/Source/Runtime/Experimental/Chaos/Public/Chaos/PullPhysicsDataImp.h
@@ -28,7 +28,11 @@ public:
 
 	TProxy* GetProxy() const
 	{
-		return !Timestamp->bDeleted ? Proxy : nullptr;
+		if( Timestamp )
+		{
+			return !Timestamp->bDeleted ? Proxy : nullptr;
+		}
+		return nullptr;
 	}
 
 	const FProxyTimestamp* GetTimestamp() const { return Timestamp.Get(); }
diff --git a/Engine/Source/Runtime/Experimental/Chaos/Public/PBDRigidsSolver.h b/Engine/Source/Runtime/Experimental/Chaos/Public/PBDRigidsSolver.h
--- a/Engine/Source/Runtime/Experimental/Chaos/Public/PBDRigidsSolver.h
+++ b/Engine/Source/Runtime/Experimental/Chaos/Public/PBDRigidsSolver.h
@@ -377,6 +377,7 @@ namespace Chaos
 		THandleArray<FChaosPhysicsMaterialMask> SimMaterialMasks;
 
 		TArray<FSingleParticlePhysicsProxy*> PendingDestroyPhysicsProxy;
+		TArray<FGeometryCollectionPhysicsProxy*> PendingDestroyGeometryCollectionPhysicsProxy;
 
 		void ProcessSinglePushedData_Internal(FPushPhysicsData& PushData);
 		virtual void ProcessPushedData_Internal(FPushPhysicsData& PushData) override;
diff --git a/Engine/Source/Runtime/Experimental/Chaos/Public/PhysicsProxy/GeometryCollectionPhysicsProxy.h b/Engine/Source/Runtime/Experimental/Chaos/Public/PhysicsProxy/GeometryCollectionPhysicsProxy.h
--- a/Engine/Source/Runtime/Experimental/Chaos/Public/PhysicsProxy/GeometryCollectionPhysicsProxy.h
+++ b/Engine/Source/Runtime/Experimental/Chaos/Public/PhysicsProxy/GeometryCollectionPhysicsProxy.h
@@ -274,6 +274,7 @@ private:
 	//
 	bool IsObjectDynamic; // Records current dynamic state
 	bool IsObjectLoading; // Indicate when loaded
+	bool IsObjectDeleting; // Indicatge when pending deletion
 
 	TManagedArray<TUniquePtr<Chaos::FGeometryParticle>> GTParticles;
 	FCollisionFilterData SimFilter;
