Chaos - more debug draw

diff --git a/Engine/Source/Runtime/Experimental/Chaos/Private/Chaos/ChaosDebugDraw.cpp b/Engine/Source/Runtime/Experimental/Chaos/Private/Chaos/ChaosDebugDraw.cpp
--- a/Engine/Source/Runtime/Experimental/Chaos/Private/Chaos/ChaosDebugDraw.cpp
+++ b/Engine/Source/Runtime/Experimental/Chaos/Private/Chaos/ChaosDebugDraw.cpp
@@ -39,6 +39,11 @@ namespace Chaos
 		bool bChaosDebugDebugDrawInactiveContacts = true;
 		FAutoConsoleVariableRef CVarChaosDebugDrawInactiveContacts(TEXT("p.Chaos.DebugDraw.ShowInactiveContacts"), bChaosDebugDebugDrawInactiveContacts, TEXT("Whether to show inactive contacts (ones that contributed no impulses or pushout)"));
 
+		bool bChaosDebugDebugDrawCoreShapes = false;
+		bool bChaosDebugDebugDrawExactCoreShapes = false;
+		FAutoConsoleVariableRef CVarChaosDebugDrawCoreShapes(TEXT("p.Chaos.DebugDraw.ShowCoreShapes"), bChaosDebugDebugDrawCoreShapes, TEXT("Whether to show the core (margin-reduced) shape where applicable"));
+		FAutoConsoleVariableRef CVarChaosDebugDrawExactShapes(TEXT("p.Chaos.DebugDraw.ShowExactCoreShapes"), bChaosDebugDebugDrawExactCoreShapes, TEXT("Whether to show the exact core shape. NOTE: Extremely expensive and should only be used on a small scene with a couple convex shapes in it"));
+
 
 		// NOTE: These settings should never really be used - they are the fallback defaults
 		// if the user does not specify settings in the debug draw call.
@@ -81,11 +86,45 @@ namespace Chaos
 		//
 		//
 
-		void DrawShapesImpl(const FGeometryParticleHandle* Particle, const FRigidTransform3& ShapeTransform, const FImplicitObject* Shape, const FColor& Color, const FChaosDebugDrawSettings& Settings);
+		void DrawShapesImpl(const FGeometryParticleHandle* Particle, const FRigidTransform3& ShapeTransform, const FImplicitObject* Shape, const FReal Margin, const FColor& Color, const FChaosDebugDrawSettings& Settings);
 
 		void DrawShape(const FRigidTransform3& ShapeTransform, const FImplicitObject* Shape, const FColor& Color, const FChaosDebugDrawSettings* Settings)
 		{
-			DrawShapesImpl(nullptr, ShapeTransform, Shape, Color, GetChaosDebugDrawSettings(Settings));
+			DrawShapesImpl(nullptr, ShapeTransform, Shape, 0.0f, Color, GetChaosDebugDrawSettings(Settings));
+		}
+
+		void DrawShapesConvexImpl(const TGeometryParticleHandle<FReal, 3>* Particle, const FRigidTransform3& ShapeTransform, const FConvex* Shape, const FReal InMargin, const FColor& Color, const FChaosDebugDrawSettings& Settings)
+		{
+			if (Shape->HasStructureData())
+			{
+				const FReal Margin = InMargin + Shape->GetMargin();
+
+				for (int32 PlaneIndex = 0; PlaneIndex < Shape->GetFaces().Num(); ++PlaneIndex)
+				{
+					const int32 PlaneVerticesNum = Shape->NumPlaneVertices(PlaneIndex);
+					for (int32 PlaneVertexIndex = 0; PlaneVertexIndex < PlaneVerticesNum; ++PlaneVertexIndex)
+					{
+						const int32 VertexIndex0 = Shape->GetPlaneVertex(PlaneIndex, PlaneVertexIndex);
+						const int32 VertexIndex1 = Shape->GetPlaneVertex(PlaneIndex, Utilities::WrapIndex(PlaneVertexIndex + 1, 0, PlaneVerticesNum));
+
+						const FVec3 OuterP0 = ShapeTransform.TransformPosition(Shape->GetVertex(VertexIndex0));
+						const FVec3 OuterP1 = ShapeTransform.TransformPosition(Shape->GetVertex(VertexIndex1));
+
+						// Outer shape
+						FDebugDrawQueue::GetInstance().DrawDebugLine(OuterP0, OuterP1, Color, false, -1.f, Settings.DrawPriority, Settings.ShapeThicknesScale * Settings.LineThickness);
+
+						// Core shape and lines connecting core to outer
+						if (Margin > 0.0f)
+						{
+							const FReal LineThickness = 0.5f * Settings.ShapeThicknesScale * Settings.LineThickness;
+							const FVec3 InnerP0 = ShapeTransform.TransformPositionNoScale(Shape->GetMarginAdjustedVertexScaled(VertexIndex0, Margin, ShapeTransform.GetScale3D()));
+							const FVec3 InnerP1 = ShapeTransform.TransformPositionNoScale(Shape->GetMarginAdjustedVertexScaled(VertexIndex1, Margin, ShapeTransform.GetScale3D()));
+							FDebugDrawQueue::GetInstance().DrawDebugLine(InnerP0, InnerP1, FColor::Blue, false, -1.f, Settings.DrawPriority, LineThickness);
+							FDebugDrawQueue::GetInstance().DrawDebugLine(InnerP0, OuterP0, FColor::Black, false, -1.f, Settings.DrawPriority, LineThickness);
+						}
+					}
+				}
+			}
 		}
 
 		void DrawShapesHeightFieldImpl(const TGeometryParticleHandle<FReal, 3>* Particle, const FRigidTransform3& ShapeTransform, const FHeightField* Shape, const FColor& Color, const FChaosDebugDrawSettings& Settings)
@@ -146,7 +185,7 @@ namespace Chaos
 		}
 
 		template <bool bInstanced>
-		void DrawShapesScaledImpl(const FGeometryParticleHandle* Particle, const FRigidTransform3& ShapeTransform, const FImplicitObject* Shape, const FColor& Color, const FChaosDebugDrawSettings& Settings)
+		void DrawShapesScaledImpl(const FGeometryParticleHandle* Particle, const FRigidTransform3& ShapeTransform, const FImplicitObject* Shape, const FReal Margin, const FColor& Color, const FChaosDebugDrawSettings& Settings)
 		{
 			const EImplicitObjectType PackedType = Shape->GetType();
 			const EImplicitObjectType InnerType = GetInnerType(PackedType);
@@ -176,7 +215,7 @@ namespace Chaos
 			{
 				const TImplicitObjectScaled<FConvex, bInstanced>* Scaled = Shape->template GetObject<TImplicitObjectScaled<FConvex, bInstanced>>();
 				ScaleTM.SetScale3D(Scaled->GetScale());
-				DrawShapesImpl(Particle, ScaleTM * ShapeTransform, Scaled->GetUnscaledObject(), Color, Settings);
+				DrawShapesImpl(Particle, ScaleTM * ShapeTransform, Scaled->GetUnscaledObject(), Scaled->GetMargin(), Color, Settings);
 				break;
 			}
 			case ImplicitObjectType::TaperedCylinder:
@@ -187,14 +226,14 @@ namespace Chaos
 			{
 				const TImplicitObjectScaled<FTriangleMeshImplicitObject, bInstanced>* Scaled = Shape->template GetObject<TImplicitObjectScaled<FTriangleMeshImplicitObject, bInstanced>>();
 				ScaleTM.SetScale3D(Scaled->GetScale());
-				DrawShapesImpl(Particle, ScaleTM * ShapeTransform, Scaled->GetUnscaledObject(), Color, Settings);
+				DrawShapesImpl(Particle, ScaleTM * ShapeTransform, Scaled->GetUnscaledObject(), 0.0f, Color, Settings);
 				break;
 			}
 			case ImplicitObjectType::HeightField:
 			{
 				const TImplicitObjectScaled<FHeightField, bInstanced>* Scaled = Shape->template GetObject<TImplicitObjectScaled<FHeightField, bInstanced>>();
 				ScaleTM.SetScale3D(Scaled->GetScale());
-				DrawShapesImpl(Particle, ScaleTM * ShapeTransform, Scaled->GetUnscaledObject(), Color, Settings);
+				DrawShapesImpl(Particle, ScaleTM * ShapeTransform, Scaled->GetUnscaledObject(), 0.0f, Color, Settings);
 				break;
 			}
 			default:
@@ -202,7 +241,7 @@ namespace Chaos
 			}
 		}
 
-		void DrawShapesInstancedImpl(const FGeometryParticleHandle* Particle, const FRigidTransform3& ShapeTransform, const FImplicitObject* Shape, const FColor& Color, const FChaosDebugDrawSettings& Settings)
+		void DrawShapesInstancedImpl(const FGeometryParticleHandle* Particle, const FRigidTransform3& ShapeTransform, const FImplicitObject* Shape, const FReal Margin, const FColor& Color, const FChaosDebugDrawSettings& Settings)
 		{
 			const EImplicitObjectType PackedType = Shape->GetType();
 			const EImplicitObjectType InnerType = GetInnerType(PackedType);
@@ -230,7 +269,7 @@ namespace Chaos
 			case ImplicitObjectType::Convex:
 			{
 				const TImplicitObjectInstanced<FConvex>* Instanced = Shape->template GetObject<TImplicitObjectInstanced<FConvex>>();
-				DrawShapesImpl(Particle, ShapeTransform, Instanced->GetInstancedObject(), Color, Settings);
+				DrawShapesImpl(Particle, ShapeTransform, Instanced->GetInstancedObject(), Instanced->GetMargin(), Color, Settings);
 				break;
 			}
 			case ImplicitObjectType::TaperedCylinder:
@@ -240,13 +279,13 @@ namespace Chaos
 			case ImplicitObjectType::TriangleMesh:
 			{
 				const TImplicitObjectInstanced<FTriangleMeshImplicitObject>* Scaled = Shape->template GetObject<TImplicitObjectInstanced<FTriangleMeshImplicitObject>>();
-				DrawShapesImpl(Particle, ShapeTransform, Scaled->GetInstancedObject(), Color, Settings);
+				DrawShapesImpl(Particle, ShapeTransform, Scaled->GetInstancedObject(), 0.0f, Color, Settings);
 				break;
 			}
 			case ImplicitObjectType::HeightField:
 			{
 				const TImplicitObjectInstanced<FHeightField>* Scaled = Shape->template GetObject<TImplicitObjectInstanced<FHeightField>>();
-				DrawShapesImpl(Particle, ShapeTransform, Scaled->GetInstancedObject(), Color, Settings);
+				DrawShapesImpl(Particle, ShapeTransform, Scaled->GetInstancedObject(), 0.0f, Color, Settings);
 				break;
 			}
 			default:
@@ -254,7 +293,7 @@ namespace Chaos
 			}
 		}
 
-		void DrawShapesImpl(const FGeometryParticleHandle* Particle, const FRigidTransform3& ShapeTransform, const FImplicitObject* Shape, const FColor& Color, const FChaosDebugDrawSettings& Settings)
+		void DrawShapesImpl(const FGeometryParticleHandle* Particle, const FRigidTransform3& ShapeTransform, const FImplicitObject* Shape, const FReal Margin, const FColor& Color, const FChaosDebugDrawSettings& Settings)
 		{
 			const EImplicitObjectType PackedType = Shape->GetType(); // Type includes scaling and instancing data
 			const EImplicitObjectType InnerType = GetInnerType(Shape->GetType());
@@ -264,24 +303,24 @@ namespace Chaos
 			{
 				if (IsInstanced(PackedType))
 				{
-					DrawShapesScaledImpl<true>(Particle, ShapeTransform, Shape, Color, Settings);
+					DrawShapesScaledImpl<true>(Particle, ShapeTransform, Shape, Margin, Color, Settings);
 				}
 				else
 				{
-					DrawShapesScaledImpl<false>(Particle, ShapeTransform, Shape, Color, Settings);
+					DrawShapesScaledImpl<false>(Particle, ShapeTransform, Shape, Margin, Color, Settings);
 				}
 				return;
 			}
 			else if (IsInstanced(PackedType))
 			{
-				DrawShapesInstancedImpl(Particle, ShapeTransform, Shape, Color, Settings);
+				DrawShapesInstancedImpl(Particle, ShapeTransform, Shape, Margin, Color, Settings);
 				return;
 			}
 			else if (InnerType == ImplicitObjectType::Transformed)
 			{
 				const TImplicitObjectTransformed<FReal, 3>* Transformed = Shape->template GetObject<TImplicitObjectTransformed<FReal, 3>>();
 				FRigidTransform3 TransformedTransform = FRigidTransform3(ShapeTransform.TransformPosition(Transformed->GetTransform().GetLocation()), ShapeTransform.GetRotation() * Transformed->GetTransform().GetRotation());
-				DrawShapesImpl(Particle, TransformedTransform, Transformed->GetTransformedObject(), Color, Settings);
+				DrawShapesImpl(Particle, TransformedTransform, Transformed->GetTransformedObject(), Margin, Color, Settings);
 				return;
 			}
 			else if (InnerType == ImplicitObjectType::Union)
@@ -289,7 +328,7 @@ namespace Chaos
 				const FImplicitObjectUnion* Union = Shape->template GetObject<FImplicitObjectUnion>();
 				for (auto& UnionShape : Union->GetObjects())
 				{
-					DrawShapesImpl(Particle, ShapeTransform, UnionShape.Get(), Color, Settings);
+					DrawShapesImpl(Particle, ShapeTransform, UnionShape.Get(), Margin, Color, Settings);
 				}
 				return;
 			}
@@ -298,7 +337,7 @@ namespace Chaos
 				const FImplicitObjectUnionClustered* Union = Shape->template GetObject<FImplicitObjectUnionClustered>();
 				for (auto& UnionShape : Union->GetObjects())
 				{
-					DrawShapesImpl(Particle, ShapeTransform, UnionShape.Get(), Color, Settings);
+					DrawShapesImpl(Particle, ShapeTransform, UnionShape.Get(), Margin, Color, Settings);
 				}
 				return;
 			}
@@ -363,23 +402,24 @@ namespace Chaos
 			break;
 			case ImplicitObjectType::Convex:
 			{
-				if (const FConvex* Convex = Shape->template GetObject<FConvex>())
+				const FConvex* Convex = Shape->template GetObject<FConvex>();
+
+				const FReal NetMargin = bChaosDebugDebugDrawCoreShapes ? Margin + Convex->GetMargin() : 0.0f;
+				DrawShapesConvexImpl(Particle, ShapeTransform, Convex, NetMargin, Color, Settings);
+
+				// Generate the exact marging-reduced convex for comparison with the runtime approximation
+				// Warning: extremely expensive!
+				if (bChaosDebugDebugDrawExactCoreShapes)
 				{
-					if (Convex->HasStructureData())
+					const FRigidTransform3 ShrunkScaledTransform = FRigidTransform3(ShapeTransform.GetTranslation(), ShapeTransform.GetRotation());
+					TArray<FVec3> ScaledVerts = Convex->GetVertices();
+					for (FVec3& Vert : ScaledVerts)
 					{
-						for (int32 PlaneIndex = 0; PlaneIndex < Convex->GetFaces().Num(); ++PlaneIndex)
-						{
-							const int32 PlaneVerticesNum = Convex->NumPlaneVertices(PlaneIndex);
-							for (int32 PlaneVertexIndex = 0; PlaneVertexIndex < PlaneVerticesNum; ++PlaneVertexIndex)
-							{
-								const int32 VertexIndex0 = Convex->GetPlaneVertex(PlaneIndex, PlaneVertexIndex);
-								const int32 VertexIndex1 = Convex->GetPlaneVertex(PlaneIndex, Utilities::WrapIndex(PlaneVertexIndex + 1, 0, PlaneVerticesNum));
-								const FVec3 P0 = ShapeTransform.TransformPosition(Convex->GetVertex(VertexIndex0));
-								const FVec3 P1 = ShapeTransform.TransformPosition(Convex->GetVertex(VertexIndex1));
-								FDebugDrawQueue::GetInstance().DrawDebugLine(P0, P1, Color, false, -1.f, 0, Settings.ShapeThicknesScale * Settings.LineThickness);
-							}
-						}
+						Vert = Vert * ShapeTransform.GetScale3D();
 					}
+					FConvex ShrunkScaledConvex(ScaledVerts, 0.0f);
+					ShrunkScaledConvex.MovePlanesAndRebuild(-Margin);
+					DrawShapesConvexImpl(Particle, ShrunkScaledTransform, &ShrunkScaledConvex, 0.0f, FColor::Red, Settings);
 				}
 				break;
 			}
@@ -440,7 +480,7 @@ namespace Chaos
 				Color = FColor(InColor.R / 2, InColor.G / 2, InColor.B / 2, InColor.A);
 			}
 
-			DrawShapesImpl(Particle, FRigidTransform3(P, Q), Particle->Geometry().Get(), Color, Settings);
+			DrawShapesImpl(Particle, FRigidTransform3(P, Q), Particle->Geometry().Get(), 0.0f, Color, Settings);
 		}
 
 		void DrawParticleShapesImpl(const FRigidTransform3& SpaceTransform, const FGeometryParticle* Particle, const FColor& InColor, const FChaosDebugDrawSettings& Settings)
@@ -455,7 +495,7 @@ namespace Chaos
 				Color = FColor(InColor.R / 2, InColor.G / 2, InColor.B / 2, InColor.A);
 			}
 
-			DrawShapesImpl(Particle->Handle(), FRigidTransform3(P, Q), Particle->Geometry().Get(), Color, Settings);
+			DrawShapesImpl(Particle->Handle(), FRigidTransform3(P, Q), Particle->Geometry().Get(), 0.0f, Color, Settings);
 		}
 
 		void DrawParticleBoundsImpl(const FRigidTransform3& SpaceTransform, const FGeometryParticleHandle* InParticle, const FReal Dt, const FReal BoundsThickness, const FReal BoundsThicknessVelocityInflation, const FChaosDebugDrawSettings& Settings)
diff --git a/Engine/Source/Runtime/Experimental/Chaos/Private/Chaos/Convex.cpp b/Engine/Source/Runtime/Experimental/Chaos/Private/Chaos/Convex.cpp
--- a/Engine/Source/Runtime/Experimental/Chaos/Private/Chaos/Convex.cpp
+++ b/Engine/Source/Runtime/Experimental/Chaos/Private/Chaos/Convex.cpp
@@ -283,7 +283,6 @@ namespace Chaos
 	{
 		TArray<FPlane> NewPlanes;
 		TArray<FVec3> NewPoints;
-		TArray<TArray<int32>> NewPointPlanes;
 		NewPlanes.Reserve(Planes.Num());
 
 		// Move all the planes inwards
@@ -318,66 +317,32 @@ namespace Chaos
 						if (NewPointIndex == INDEX_NONE)
 						{
 							NewPointIndex = NewPoints.Add(PlanesPos);
-							NewPointPlanes.AddDefaulted();
 						}
-
-						// Keep track of planes that contribute to the point
-						NewPointPlanes[NewPointIndex].AddUnique(PlaneIndex0);
-						NewPointPlanes[NewPointIndex].AddUnique(PlaneIndex1);
-						NewPointPlanes[NewPointIndex].AddUnique(PlaneIndex2);
 					}
 				}
 			}
 		}
 
-		// Reject points outside the planes
+		// Reject points outside the planes to get down to a sensible number for the build step
 		const FReal PointPlaneTolerance = PointTolerance;
-		for (int32 PointIndex = 0; PointIndex < NewPoints.Num(); ++PointIndex)
+		int32 NumNewPoints = NewPoints.Num();
+		for (int32 PointIndex = 0; PointIndex < NumNewPoints; ++PointIndex)
 		{
-			for (int32 PlaneIndex = 0; PlaneIndex < NewPlanes.Num(); ++PlaneIndex)
+			const int32 NumNewPlanes = NewPlanes.Num();
+			for (int32 PlaneIndex = 0; PlaneIndex < NumNewPlanes; ++PlaneIndex)
 			{
 				const FReal PointPlaneDistance = NewPlanes[PlaneIndex].PlaneDot(NewPoints[PointIndex]);
 				if (PointPlaneDistance > PointPlaneTolerance)
 				{
 					NewPoints.RemoveAtSwap(PointIndex);
-					NewPointPlanes.RemoveAtSwap(PointIndex);
 					--PointIndex;
+					--NumNewPoints;
 					break;
 				}
 			}
 		}
 
-		// Reject planes that are not used
-		TArray<bool> NewPlaneUsed;
-		NewPlaneUsed.SetNumZeroed(NewPlanes.Num());
-		for (int32 PointIndex = 0; PointIndex < NewPoints.Num(); ++PointIndex)
-		{
-			for (int32 PointPlaneIndex = 0; PointPlaneIndex < NewPointPlanes[PointIndex].Num(); ++PointPlaneIndex)
-			{
-				const int32 PlaneIndex = NewPointPlanes[PointIndex][PointPlaneIndex];
-				NewPlaneUsed[PlaneIndex] = true;
-			}
-		}
-
-		for (int32 PlaneIndex = 0; PlaneIndex < NewPlanes.Num(); ++PlaneIndex)
-		{
-			if (!NewPlaneUsed[PlaneIndex])
-			{
-				NewPlanes.RemoveAtSwap(PlaneIndex);
-				NewPlaneUsed.RemoveAtSwap(PlaneIndex);
-				--PlaneIndex;
-			}
-		}
-
-		// Use the new planes
-		Planes.Reset(NewPlanes.Num());
-		for (int32 PlaneIndex = 0; PlaneIndex < NewPlanes.Num(); ++PlaneIndex)
-		{
-			Planes.Emplace(TPlaneConcrete<FReal, 3>(NewPlanes[PlaneIndex].GetOrigin(), NewPlanes[PlaneIndex].GetNormal()));
-		}
-
-
-		// Use the new surface points
-		Vertices = MoveTemp(NewPoints);
+		// Generate a new convex from the points
+		*this = FConvex(NewPoints, 0.0f);
 	}
 }
\ No newline at end of file
diff --git a/Engine/Source/Runtime/Experimental/Chaos/Public/Chaos/Convex.h b/Engine/Source/Runtime/Experimental/Chaos/Public/Chaos/Convex.h
--- a/Engine/Source/Runtime/Experimental/Chaos/Public/Chaos/Convex.h
+++ b/Engine/Source/Runtime/Experimental/Chaos/Public/Chaos/Convex.h
@@ -92,11 +92,37 @@ namespace Chaos
 			CreateStructureData(MoveTemp(FaceIndices));
 
 			SetMargin(InMargin);
-		}	
+		}
+
+		FConvex& operator=(const FConvex& Other) = delete;
+		
+		FConvex& operator=(FConvex&& Other)
+		{
+			// Base class assignment
+			// @todo(chaos): Base class needs protected assignment
+			Type = Other.Type;
+			CollisionType = Other.CollisionType;
+			Margin = Other.Margin;
+			bIsConvex = Other.bIsConvex;
+			bDoCollide = Other.bDoCollide;
+			bHasBoundingBox = Other.bHasBoundingBox;
+#if TRACK_CHAOS_GEOMETRY
+			bIsTracked = Other.bIsTracked;
+#endif
+			// This class assignment
+			Planes = MoveTemp(Other.Planes);
+			Vertices = MoveTemp(Other.Vertices);
+			LocalBoundingBox = MoveTemp(Other.LocalBoundingBox);
+			StructureData = MoveTemp(Other.StructureData);
+			Volume = MoveTemp(Other.Volume);
+			CenterOfMass = MoveTemp(Other.CenterOfMass);
+
+			return *this;
+		}
 
-	private:
 		void MovePlanesAndRebuild(FReal InDelta);
 
+	private:
 		void CreateStructureData(TArray<TArray<int32>>&& FaceIndices);
 
 	public:
@@ -401,6 +427,8 @@ namespace Chaos
 			return MaxVIdx;
 		}
 
+	public:
+
 		FVec3 GetMarginAdjustedVertex(int32 VertexIndex, FReal InMargin) const
 		{
 			// @chaos(todo): moving the vertices this way based on margin is only valid for small margins. If the margin
@@ -408,6 +436,11 @@ namespace Chaos
 			// This can be fixed with some extra data in the convex structure, but for now we accept the fact that large 
 			// margins on convexes with small faces can cause non-convex core shapes.
 
+			if (InMargin == 0.0f)
+			{
+				return GetVertex(InMargin);
+			}
+
 			// Get any 3 planes that contribute to this vertex
 			if (NumVertexPlanes(VertexIndex) >= 3)
 			{
@@ -435,6 +468,11 @@ namespace Chaos
 
 		FVec3 GetMarginAdjustedVertexScaled(int32 VertexIndex, FReal InMargin, const FVec3& Scale) const
 		{
+			if (InMargin == 0.0f)
+			{
+				return GetVertex(InMargin) * Scale;
+			}
+
 			// Get any 3 planes that contribute to this vertex
 			if (NumVertexPlanes(VertexIndex) >= 3)
 			{
@@ -447,9 +485,9 @@ namespace Chaos
 				const FVec3 NewPlaneX = Scale * GetVertex(VertexIndex);
 				const FVec3 NewPlaneNs[3] = 
 				{
-					(Scale * Planes[PlaneIndex0].Normal()).GetUnsafeNormal(),
-					(Scale * Planes[PlaneIndex1].Normal()).GetUnsafeNormal(),
-					(Scale * Planes[PlaneIndex2].Normal()).GetUnsafeNormal(),
+					(Planes[PlaneIndex0].Normal() / Scale).GetUnsafeNormal(),
+					(Planes[PlaneIndex1].Normal() / Scale).GetUnsafeNormal(),
+					(Planes[PlaneIndex2].Normal() / Scale).GetUnsafeNormal(),
 				};
 				FReal NewPlaneDs[3] = 
 				{
diff --git a/Engine/Source/Runtime/Experimental/Chaos/Public/Chaos/ConvexStructureData.h b/Engine/Source/Runtime/Experimental/Chaos/Public/Chaos/ConvexStructureData.h
--- a/Engine/Source/Runtime/Experimental/Chaos/Public/Chaos/ConvexStructureData.h
+++ b/Engine/Source/Runtime/Experimental/Chaos/Public/Chaos/ConvexStructureData.h
@@ -207,11 +207,32 @@ namespace Chaos
 		{
 		}
 
+		FConvexStructureData(const FConvexStructureData& Other) = delete;
+
+		FConvexStructureData(FConvexStructureData&& Other)
+		{
+			*this = MoveTemp(Other);
+		}
+
 		~FConvexStructureData()
 		{
 			DestroyDataContainer();
 		}
 
+		FConvexStructureData& operator=(const FConvexStructureData& Other) = delete;
+
+		FConvexStructureData& operator=(FConvexStructureData&& Other)
+		{
+			Data = Other.Data;
+			IndexType = Other.IndexType;
+
+			Other.Data = nullptr;
+			Other.IndexType = EIndexType::None;
+
+			return *this;
+		}
+
+
 		FORCEINLINE bool IsValid() const
 		{
 			return (Data != nullptr);
