Fix build issue with deprecated class

diff --git a/Engine/Source/Runtime/Experimental/Chaos/Public/BoneHierarchy.h b/Engine/Source/Runtime/Experimental/Chaos/Public/BoneHierarchy.h
--- a/Engine/Source/Runtime/Experimental/Chaos/Public/BoneHierarchy.h
+++ b/Engine/Source/Runtime/Experimental/Chaos/Public/BoneHierarchy.h
@@ -465,7 +465,7 @@ protected:
 	void SetAnimLocalClean(const int32 TransformIndex)
 	{ AnimDirty[TransformIndex] &= 0b1110; }
 protected:
-	friend class USkeletalMeshSimulationComponent;
+	//friend class USkeletalMeshSimulationComponent;
 
 	// Owner of all implicit shape groups.
 	TArray<TUniquePtr<FAnalyticImplicitGroup> > ImplicitGroups;
diff --git a/Engine/Source/Runtime/Experimental/Chaos/Public/PhysicsProxy/AnalyticImplicitGroup.h b/Engine/Source/Runtime/Experimental/Chaos/Public/PhysicsProxy/AnalyticImplicitGroup.h
--- a/Engine/Source/Runtime/Experimental/Chaos/Public/PhysicsProxy/AnalyticImplicitGroup.h
+++ b/Engine/Source/Runtime/Experimental/Chaos/Public/PhysicsProxy/AnalyticImplicitGroup.h
@@ -617,7 +617,7 @@ protected:
 	}
 
 protected:
-	friend class USkeletalMeshSimulationComponent;
+	//friend class USkeletalMeshSimulationComponent;
 	friend struct FPhysicsAssetSimulationUtil;
 
 	FName BoneName;
