Chaos Physics - add support for moving geometry collection
- account for component transform on a per frame basis
- updates kinematic particle accordingly
- add missing calls to the solver to support it

diff --git a/Engine/Source/Runtime/Experimental/Chaos/Private/PBDRigidsSolver.cpp b/Engine/Source/Runtime/Experimental/Chaos/Private/PBDRigidsSolver.cpp
--- a/Engine/Source/Runtime/Experimental/Chaos/Private/PBDRigidsSolver.cpp
+++ b/Engine/Source/Runtime/Experimental/Chaos/Private/PBDRigidsSolver.cpp
@@ -813,7 +813,8 @@ namespace Chaos
 			}
 			case EPhysicsProxyType::GeometryCollectionType:
 			{
-				// Not invalid but doesn't currently use the remote data process
+				auto Proxy = static_cast<FGeometryCollectionPhysicsProxy*>(Dirty.Proxy);
+				Proxy->PushStateOnGameThread(this);
 				break;
 			}
 			case EPhysicsProxyType::JointConstraintType:
@@ -913,6 +914,8 @@ namespace Chaos
 			
 				case EPhysicsProxyType::GeometryCollectionType:
 				{
+					auto Proxy = static_cast<FGeometryCollectionPhysicsProxy*>(Dirty.Proxy);
+					Proxy->PushToPhysicsState();
 					// Currently no push needed for geometry collections and they handle the particle creation internally
 					// #TODO This skips the rewind data push so GC will not be rewindable until resolved.
 					Dirty.Proxy->ResetDirtyIdx();
diff --git a/Engine/Source/Runtime/Experimental/Chaos/Private/PhysicsProxy/GeometryCollectionPhysicsProxy.cpp b/Engine/Source/Runtime/Experimental/Chaos/Private/PhysicsProxy/GeometryCollectionPhysicsProxy.cpp
--- a/Engine/Source/Runtime/Experimental/Chaos/Private/PhysicsProxy/GeometryCollectionPhysicsProxy.cpp
+++ b/Engine/Source/Runtime/Experimental/Chaos/Private/PhysicsProxy/GeometryCollectionPhysicsProxy.cpp
@@ -332,6 +332,7 @@ FGeometryCollectionPhysicsProxy::FGeometryCollectionPhysicsProxy(
 	, CollisionParticlesPerObjectFraction(CollisionParticlesPerObjectFractionDefault)
 
 	, GameThreadCollection(GameThreadCollectionIn)
+	, bIsPhysicsThreadWorldTransformDirty(false)
 {
 	// We rely on a guarded buffer.
 	check(BufferMode == Chaos::EMultiBufferMode::TripleGuarded);
@@ -372,6 +373,8 @@ void FGeometryCollectionPhysicsProxy::Initialize(Chaos::FPBDRigidsEvolutionBase
 	{
 		InitFunc(Parameters);
 	}
+	// compatibility requirement to make sure we at least initialize GameThreadPerFrameData properly
+	GameThreadPerFrameData.SetWorldTransform(Parameters.WorldTransform);
 
 	//
 	// Collision vertices down sampling validation.  
@@ -1432,6 +1435,59 @@ void FGeometryCollectionPhysicsProxy::BufferGameState()
 	//
 }
 
+
+void FGeometryCollectionPhysicsProxy::SetWorldTransform(const FTransform& WorldTransform)
+{
+	check(IsInGameThread());
+	GameThreadPerFrameData.SetWorldTransform(WorldTransform);
+}
+
+void FGeometryCollectionPhysicsProxy::PushStateOnGameThread(Chaos::FPBDRigidsSolver* InSolver)
+{
+	// CONTEXT: GAMETHREAD
+	// this is running on GAMETHREAD before the PhysicsThread code runs for this frame
+	bIsPhysicsThreadWorldTransformDirty = GameThreadPerFrameData.GetIsWorldTransformDirty();
+	if (bIsPhysicsThreadWorldTransformDirty)
+	{
+		Parameters.WorldTransform = GameThreadPerFrameData.GetWorldTransform();
+		GameThreadPerFrameData.ResetIsWorldTransformDirty();
+	}
+}
+
+void FGeometryCollectionPhysicsProxy::PushToPhysicsState()
+{
+	// CONTEXT: PHYSICSTHREAD
+	// because the aattached actor can be dynamic, we need to update the kinematic particles properly
+	if (bIsPhysicsThreadWorldTransformDirty)
+	{
+		const FTransform& ActorToWorld = Parameters.WorldTransform;
+
+		int32 NumTransformGroupElements = PhysicsThreadCollection.NumElements(FGeometryCollection::TransformGroup);
+		for (int32 TransformGroupIndex = 0; TransformGroupIndex < NumTransformGroupElements; ++TransformGroupIndex)
+		{
+			Chaos::FPBDRigidClusteredParticleHandle* Handle = SolverParticleHandles[TransformGroupIndex];
+			if (!Handle)
+			{
+				continue;
+			}
+			const Chaos::EObjectStateType ObjectState = Handle->ObjectState();
+			switch (ObjectState)
+			{
+			case Chaos::EObjectStateType::Kinematic:
+			{
+				FTransform LocalTransform = PhysicsThreadCollection.Transform[TransformGroupIndex];
+				FTransform WorldTransform = PhysicsThreadCollection.MassToLocal[TransformGroupIndex] * PhysicsThreadCollection.Transform[TransformGroupIndex] * ActorToWorld;;
+
+				Handle->SetX(WorldTransform.GetTranslation());
+				Handle->SetR(WorldTransform.GetRotation().GetNormalized());
+				Handle->SetP(Handle->X());
+				Handle->SetQ(Handle->R());
+			}
+			}
+		}
+	}
+}
+
 void FGeometryCollectionPhysicsProxy::BufferPhysicsResults(Chaos::FPBDRigidsSolver* CurrentSolver, Chaos::FDirtyGeometryCollectionData& BufferData)
 {
 	/**
diff --git a/Engine/Source/Runtime/Experimental/Chaos/Public/GeometryCollectionProxyData.h b/Engine/Source/Runtime/Experimental/Chaos/Public/GeometryCollectionProxyData.h
--- a/Engine/Source/Runtime/Experimental/Chaos/Public/GeometryCollectionProxyData.h
+++ b/Engine/Source/Runtime/Experimental/Chaos/Public/GeometryCollectionProxyData.h
@@ -89,6 +89,31 @@ public:
 };
 
 
+class FGeometryCollectioPerFrameData
+{
+public:
+	FGeometryCollectioPerFrameData()
+		: IsWorldTransformDirty(false) {}
+
+	const FTransform& GetWorldTransform() const { return WorldTransform; }
+
+	void SetWorldTransform(const FTransform& InWorldTransform)
+	{
+		if (!WorldTransform.Equals(InWorldTransform))
+		{
+			WorldTransform = InWorldTransform;
+			IsWorldTransformDirty = true;
+		}
+	}
+
+	bool GetIsWorldTransformDirty() const { return IsWorldTransformDirty; }
+	void ResetIsWorldTransformDirty() { IsWorldTransformDirty = false; }
+
+private:
+	FTransform WorldTransform;
+	bool IsWorldTransformDirty;
+};
+
 /**
  * Buffer structure for communicating simulation state between game and physics
  * threads.
diff --git a/Engine/Source/Runtime/Experimental/Chaos/Public/PhysicsProxy/GeometryCollectionPhysicsProxy.h b/Engine/Source/Runtime/Experimental/Chaos/Public/PhysicsProxy/GeometryCollectionPhysicsProxy.h
--- a/Engine/Source/Runtime/Experimental/Chaos/Public/PhysicsProxy/GeometryCollectionPhysicsProxy.h
+++ b/Engine/Source/Runtime/Experimental/Chaos/Public/PhysicsProxy/GeometryCollectionPhysicsProxy.h
@@ -124,6 +124,12 @@ public:
 	/** Push physics state into the \c PhysToGameInterchange. */
 	void BufferPhysicsResults(Chaos::FPBDRigidsSolver* CurrentSolver, Chaos::FDirtyGeometryCollectionData& BufferData);
 
+	/** Push data from the game thread to the physics thread */
+	void PushStateOnGameThread(Chaos::FPBDRigidsSolver* InSolver);
+
+	/** apply the state changes on the physics thread */
+	void PushToPhysicsState();
+	
 	/** Does nothing as \c BufferPhysicsResults() already did this. */
 	void FlipBuffer();
 	
@@ -174,6 +180,9 @@ public:
 
 	bool IsGTCollectionDirty() const { return GameThreadCollection.IsDirty(); }
 
+	// set the world transform ( this needs to be called on the game thread ) 
+	void SetWorldTransform(const FTransform& WorldTransform);
+
 	const TArray<FClusterHandle*> GetParticles() const
 	{
 		return SolverParticleHandles;
@@ -319,6 +328,10 @@ private:
 	FGeometryDynamicCollection PhysicsThreadCollection;
 	FGeometryDynamicCollection& GameThreadCollection;
 
+	// this data flows from Game thread to physics thread
+	FGeometryCollectioPerFrameData GameThreadPerFrameData;
+	bool bIsPhysicsThreadWorldTransformDirty;
+
 	// Currently this is using triple buffers for game-physics and 
 	// physics-game thread communication, but not for any reason other than this 
 	// is the only implementation we currently have of a guarded buffer - a buffer 
diff --git a/Engine/Source/Runtime/Experimental/GeometryCollectionEngine/Private/GeometryCollection/GeometryCollectionComponent.cpp b/Engine/Source/Runtime/Experimental/GeometryCollectionEngine/Private/GeometryCollection/GeometryCollectionComponent.cpp
--- a/Engine/Source/Runtime/Experimental/GeometryCollectionEngine/Private/GeometryCollection/GeometryCollectionComponent.cpp
+++ b/Engine/Source/Runtime/Experimental/GeometryCollectionEngine/Private/GeometryCollection/GeometryCollectionComponent.cpp
@@ -1666,6 +1666,18 @@ void UGeometryCollectionComponent::InitDynamicData(FGeometryCollectionDynamicDat
 	}
 }
 
+void UGeometryCollectionComponent::OnUpdateTransform(EUpdateTransformFlags UpdateTransformFlags, ETeleportType Teleport)
+{
+	Super::OnUpdateTransform(UpdateTransformFlags, Teleport);
+
+#if WITH_CHAOS
+	if (PhysicsProxy)
+	{
+		PhysicsProxy->SetWorldTransform(GetComponentTransform());
+	}
+#endif
+}
+
 void UGeometryCollectionComponent::TickComponent(float DeltaTime, enum ELevelTick TickType, FActorComponentTickFunction *ThisTickFunction)
 {
 	//UE_LOG(UGCC_LOG, Log, TEXT("GeometryCollectionComponent[%p]::TickComponent()"), this);
diff --git a/Engine/Source/Runtime/Experimental/GeometryCollectionEngine/Public/GeometryCollection/GeometryCollectionComponent.h b/Engine/Source/Runtime/Experimental/GeometryCollectionEngine/Public/GeometryCollection/GeometryCollectionComponent.h
--- a/Engine/Source/Runtime/Experimental/GeometryCollectionEngine/Public/GeometryCollection/GeometryCollectionComponent.h
+++ b/Engine/Source/Runtime/Experimental/GeometryCollectionEngine/Public/GeometryCollection/GeometryCollectionComponent.h
@@ -601,6 +601,8 @@ public:
 	virtual bool ShouldCreatePhysicsState() const override;
 	virtual bool HasValidPhysicsState() const override;
 
+	virtual void OnUpdateTransform(EUpdateTransformFlags UpdateTransformFlags, ETeleportType Teleport = ETeleportType::None) override;
+
 	// Mirrored from the proxy on a sync
 	//TManagedArray<int32> RigidBodyIds;
 	TArray<bool> DisabledFlags;
