GeometryCollection : Rendering
-Fix regression with GeometryCollection not rendering in Standalone Game
#okforgethub public, internal

diff --git a/Engine/Source/Runtime/Experimental/GeometryCollectionEngine/Private/GeometryCollection/GeometryCollectionComponent.cpp b/Engine/Source/Runtime/Experimental/GeometryCollectionEngine/Private/GeometryCollection/GeometryCollectionComponent.cpp
--- a/Engine/Source/Runtime/Experimental/GeometryCollectionEngine/Private/GeometryCollection/GeometryCollectionComponent.cpp
+++ b/Engine/Source/Runtime/Experimental/GeometryCollectionEngine/Private/GeometryCollection/GeometryCollectionComponent.cpp
@@ -463,83 +463,87 @@ FBoxSphereBounds UGeometryCollectionComponent::CalcBounds(const FTransform& Loca
 void UGeometryCollectionComponent::CreateRenderState_Concurrent(FRegisterComponentContext* Context)
 {
 	Super::CreateRenderState_Concurrent(Context);
+}
 
-	if (SceneProxy && RestCollection && RestCollection->HasVisibleGeometry())
-	{
-		FGeometryCollectionConstantData* const ConstantData = ::new FGeometryCollectionConstantData;
-		InitConstantData(ConstantData);
+FPrimitiveSceneProxy* UGeometryCollectionComponent::CreateSceneProxy()
+{
+	FPrimitiveSceneProxy* LocalSceneProxy = nullptr;
 
-		FGeometryCollectionDynamicData* const DynamicData = ::new FGeometryCollectionDynamicData;
-		InitDynamicData(DynamicData);
+	if (RestCollection)
+	{
+		// TODO: Abstract with a common helper
+		if (UseNanite(GetScene()->GetShaderPlatform()) &&
+			RestCollection->EnableNanite &&
+			RestCollection->NaniteData != nullptr &&
+			GGeometryCollectionNanite != 0)
+		{
+			LocalSceneProxy = new FNaniteGeometryCollectionSceneProxy(this);
+		}
+		else
+		{
+			LocalSceneProxy = new FGeometryCollectionSceneProxy(this);
+		}
 
-		if (SceneProxy->IsNaniteMesh())
+		if (RestCollection->HasVisibleGeometry())
 		{
-			FNaniteGeometryCollectionSceneProxy* const GeometryCollectionSceneProxy = static_cast<FNaniteGeometryCollectionSceneProxy*>(SceneProxy);
+			FGeometryCollectionConstantData* const ConstantData = ::new FGeometryCollectionConstantData;
+			InitConstantData(ConstantData);
 
-			// ...
+			FGeometryCollectionDynamicData* const DynamicData = ::new FGeometryCollectionDynamicData;
+			InitDynamicData(DynamicData);
 
-		#if GEOMETRYCOLLECTION_EDITOR_SELECTION
-			if (bIsTransformSelectionModeEnabled)
+			if (LocalSceneProxy->IsNaniteMesh())
 			{
+				FNaniteGeometryCollectionSceneProxy* const GeometryCollectionSceneProxy = static_cast<FNaniteGeometryCollectionSceneProxy*>(LocalSceneProxy);
+
 				// ...
-			}
-		#endif
 
-			ENQUEUE_RENDER_COMMAND(CreateRenderState)(
-				[GeometryCollectionSceneProxy, ConstantData, DynamicData](FRHICommandListImmediate& RHICmdList)
+#if GEOMETRYCOLLECTION_EDITOR_SELECTION
+				if (bIsTransformSelectionModeEnabled)
 				{
-					if (GeometryCollectionSceneProxy)
-					{
-						GeometryCollectionSceneProxy->SetConstantData_RenderThread(ConstantData);
-						GeometryCollectionSceneProxy->SetDynamicData_RenderThread(DynamicData);
-						GeometryCollectionSceneProxy->GetPrimitiveSceneInfo()->RequestGPUSceneUpdate();
-					}
+					// ...
 				}
-			);
-		}
-		else
-		{
-			FGeometryCollectionSceneProxy* const GeometryCollectionSceneProxy = static_cast<FGeometryCollectionSceneProxy*>(SceneProxy);
+#endif
 
-		#if GEOMETRYCOLLECTION_EDITOR_SELECTION
-			// Re-init subsections
-			if (bIsTransformSelectionModeEnabled)
-			{
-				GeometryCollectionSceneProxy->UseSubSections(true, false);  // Do not force reinit now, it'll be done in SetConstantData_RenderThread
+				ENQUEUE_RENDER_COMMAND(CreateRenderState)(
+					[GeometryCollectionSceneProxy, ConstantData, DynamicData](FRHICommandListImmediate& RHICmdList)
+					{
+						if (GeometryCollectionSceneProxy)
+						{
+							GeometryCollectionSceneProxy->SetConstantData_RenderThread(ConstantData);
+							GeometryCollectionSceneProxy->SetDynamicData_RenderThread(DynamicData);
+							GeometryCollectionSceneProxy->GetPrimitiveSceneInfo()->RequestGPUSceneUpdate();
+						}
+					}
+				);
 			}
-		#endif
+			else
+			{
+				FGeometryCollectionSceneProxy* const GeometryCollectionSceneProxy = static_cast<FGeometryCollectionSceneProxy*>(LocalSceneProxy);
 
-			ENQUEUE_RENDER_COMMAND(CreateRenderState)(
-				[GeometryCollectionSceneProxy, ConstantData, DynamicData](FRHICommandListImmediate& RHICmdList)
+#if GEOMETRYCOLLECTION_EDITOR_SELECTION
+				// Re-init subsections
+				if (bIsTransformSelectionModeEnabled)
 				{
-					if (GeometryCollectionSceneProxy)
+					GeometryCollectionSceneProxy->UseSubSections(true, false);  // Do not force reinit now, it'll be done in SetConstantData_RenderThread
+				}
+#endif
+
+				ENQUEUE_RENDER_COMMAND(CreateRenderState)(
+					[GeometryCollectionSceneProxy, ConstantData, DynamicData](FRHICommandListImmediate& RHICmdList)
 					{
-						GeometryCollectionSceneProxy->SetConstantData_RenderThread(ConstantData);
-						GeometryCollectionSceneProxy->SetDynamicData_RenderThread(DynamicData);
+						if (GeometryCollectionSceneProxy)
+						{
+							GeometryCollectionSceneProxy->SetConstantData_RenderThread(ConstantData);
+							GeometryCollectionSceneProxy->SetDynamicData_RenderThread(DynamicData);
+						}
 					}
-				}
-			);
+				);
+			}
 		}
 	}
-}
-
-FPrimitiveSceneProxy* UGeometryCollectionComponent::CreateSceneProxy()
-{
-	if (!RestCollection)
-	{
-		return nullptr;
-	}
-
-	// TODO: Abstract with a common helper
-	if (UseNanite(GetScene()->GetShaderPlatform()) &&
-		RestCollection->EnableNanite &&
-		RestCollection->NaniteData != nullptr &&
-		GGeometryCollectionNanite != 0)
-	{
-		return new FNaniteGeometryCollectionSceneProxy(this);
-	}
+	return LocalSceneProxy;
 
-	return new FGeometryCollectionSceneProxy(this);
 }
 
 bool UGeometryCollectionComponent::ShouldCreatePhysicsState() const
