Chaos : Secondary Particle Emission
-- Changed how particle to proxy mapping is implemented. Removed the map maintained on the solver. Each particle now maintains its own pointer to the associated proxy. Internal cluster particles are linked to multiple proxies (necessarily: they need this information to mark proxies dirty during the buffering process). Internal cluster particles have a set of proxy pointers that are used in situations that require it.

diff --git a/Engine/Plugins/Experimental/ChaosCaching/Source/ChaosCaching/Private/Chaos/Adapters/GeometryCollectionComponentCacheAdapter.cpp b/Engine/Plugins/Experimental/ChaosCaching/Source/ChaosCaching/Private/Chaos/Adapters/GeometryCollectionComponentCacheAdapter.cpp
--- a/Engine/Plugins/Experimental/ChaosCaching/Source/ChaosCaching/Private/Chaos/Adapters/GeometryCollectionComponentCacheAdapter.cpp
+++ b/Engine/Plugins/Experimental/ChaosCaching/Source/ChaosCaching/Private/Chaos/Adapters/GeometryCollectionComponentCacheAdapter.cpp
@@ -277,7 +277,6 @@ namespace Chaos
 
 							FBreakingData CachedBreak;
 							CachedBreak.Particle = Particle;
-							CachedBreak.ParticleProxy = Proxy;
 							CachedBreak.Location = Event->Location;
 							CachedBreak.Velocity = Event->Velocity;
 							CachedBreak.AngularVelocity = Event->AngularVelocity;
diff --git a/Engine/Plugins/Experimental/ChaosNiagara/Source/ChaosNiagara/Private/NiagaraDataInterfaceChaosDestruction.cpp b/Engine/Plugins/Experimental/ChaosNiagara/Source/ChaosNiagara/Private/NiagaraDataInterfaceChaosDestruction.cpp
--- a/Engine/Plugins/Experimental/ChaosNiagara/Source/ChaosNiagara/Private/NiagaraDataInterfaceChaosDestruction.cpp
+++ b/Engine/Plugins/Experimental/ChaosNiagara/Source/ChaosNiagara/Private/NiagaraDataInterfaceChaosDestruction.cpp
@@ -1259,7 +1259,7 @@ void UNiagaraDataInterfaceChaosDestruction::HandleBreakingEvents(const Chaos::FB
 #if INCLUDE_CHAOS
 			for (auto& Solver : Solvers)
 			{
-				GeometryCollectionComponent = Solver.PhysScene->GetOwningComponent<UGeometryCollectionComponent>(DataIn.ParticleProxy);
+				GeometryCollectionComponent = Solver.PhysScene->GetOwningComponent<UGeometryCollectionComponent>(DataIn.Particle->PhysicsProxy());
 				if (GeometryCollectionComponent)
 					break;
 			}
diff --git a/Engine/Source/Runtime/Engine/Private/PhysicsEngine/Experimental/PhysScene_Chaos.cpp b/Engine/Source/Runtime/Engine/Private/PhysicsEngine/Experimental/PhysScene_Chaos.cpp
--- a/Engine/Source/Runtime/Engine/Private/PhysicsEngine/Experimental/PhysScene_Chaos.cpp
+++ b/Engine/Source/Runtime/Engine/Private/PhysicsEngine/Experimental/PhysScene_Chaos.cpp
@@ -703,7 +703,7 @@ void FPhysScene_Chaos::HandleCollisionEvents(const Chaos::FCollisionEventData& E
 							int32 CollisionIdx = Chaos::FEventManager::DecodeCollisionIndex(EncodedCollisionIdx, bSwapOrder);
 
 							Chaos::FCollidingData const& CollisionDataItem = CollisionData[CollisionIdx];
-							IPhysicsProxyBase* const PhysicsProxy1 = bSwapOrder ? CollisionDataItem.ParticleProxy : CollisionDataItem.LevelsetProxy;
+							IPhysicsProxyBase* const PhysicsProxy1 = bSwapOrder ? CollisionDataItem.Particle->PhysicsProxy() : CollisionDataItem.Levelset->PhysicsProxy();
 
 							bool bNewEntry = false;
 							FCollisionNotifyInfo& NotifyInfo = GetPendingCollisionForContactPair(PhysicsProxy0, PhysicsProxy1, bNewEntry);
diff --git a/Engine/Source/Runtime/Experimental/Chaos/Private/Chaos/PBDRigidClustering.cpp b/Engine/Source/Runtime/Experimental/Chaos/Private/Chaos/PBDRigidClustering.cpp
--- a/Engine/Source/Runtime/Experimental/Chaos/Private/Chaos/PBDRigidClustering.cpp
+++ b/Engine/Source/Runtime/Experimental/Chaos/Private/Chaos/PBDRigidClustering.cpp
@@ -791,7 +791,6 @@ namespace Chaos
 						const int32 NewIdx = MAllClusterBreakings.Add(FBreakingData());
 						FBreakingData& ClusterBreak = MAllClusterBreakings[NewIdx];
 						ClusterBreak.Particle = Child;
-						ClusterBreak.ParticleProxy = nullptr;
 						ClusterBreak.Location = Child->X();
 						ClusterBreak.Velocity = Child->V();
 						ClusterBreak.AngularVelocity = Child->W();
@@ -1060,7 +1059,6 @@ namespace Chaos
 						const int32 NewIdx = MAllClusterBreakings.Add(FBreakingData());
 						FBreakingData& ClusterBreak = MAllClusterBreakings[NewIdx];
 						ClusterBreak.Particle = Child;
-						ClusterBreak.ParticleProxy = nullptr;
 						ClusterBreak.Location = Child->X();
 						ClusterBreak.Velocity = Child->V();
 						ClusterBreak.AngularVelocity = Child->W();
@@ -1370,7 +1368,6 @@ namespace Chaos
 							int32 NewIdx = MAllClusterBreakings.Add(FBreakingData());
 							FBreakingData& ClusterBreak = MAllClusterBreakings[NewIdx];
 							ClusterBreak.Particle = ClusteredParticle;
-							ClusterBreak.ParticleProxy = nullptr;
 							ClusterBreak.Location = ClusteredParticle->X();
 							ClusterBreak.Velocity = ClusteredParticle->V();
 							ClusterBreak.AngularVelocity = ClusteredParticle->W();
@@ -1989,7 +1986,6 @@ namespace Chaos
 			const int32 NewIdx = MAllClusterBreakings.Add(FBreakingData());
 			FBreakingData& ClusterBreak = MAllClusterBreakings[NewIdx];
 			ClusterBreak.Particle = Particle;
-			ClusterBreak.ParticleProxy = nullptr;
 			ClusterBreak.Location = Particle->X();
 			ClusterBreak.Velocity = Particle->V();
 			ClusterBreak.AngularVelocity = Particle->W();
diff --git a/Engine/Source/Runtime/Experimental/Chaos/Private/Chaos/PBDRigidsEvolutionGBF.cpp b/Engine/Source/Runtime/Experimental/Chaos/Private/Chaos/PBDRigidsEvolutionGBF.cpp
--- a/Engine/Source/Runtime/Experimental/Chaos/Private/Chaos/PBDRigidsEvolutionGBF.cpp
+++ b/Engine/Source/Runtime/Experimental/Chaos/Private/Chaos/PBDRigidsEvolutionGBF.cpp
@@ -454,7 +454,7 @@ FPBDRigidsEvolutionGBF::FPBDRigidsEvolutionGBF(FPBDRigidsSOAs& InParticles,THand
 
 	AddConstraintRule(&CollisionRule);
 
-	SetInternalParticleInitilizationFunction([](const FGeometryParticleHandle*, const FGeometryParticleHandle*) {});
+	SetInternalParticleInitilizationFunction([](const FGeometryParticleHandle*, FGeometryParticleHandle*) {});
 	NarrowPhase.GetContext().bFilteringEnabled = true;
 	NarrowPhase.GetContext().bDeferUpdate = true;
 	NarrowPhase.GetContext().bAllowManifolds = false;
diff --git a/Engine/Source/Runtime/Experimental/Chaos/Private/EventDefaults.cpp b/Engine/Source/Runtime/Experimental/Chaos/Private/EventDefaults.cpp
--- a/Engine/Source/Runtime/Experimental/Chaos/Private/EventDefaults.cpp
+++ b/Engine/Source/Runtime/Experimental/Chaos/Private/EventDefaults.cpp
@@ -67,64 +67,53 @@ namespace Chaos
 				const FReal MinDeltaVelocityForHitEvents = FChaosSolversModule::GetModule()->GetSettingsProvider().GetMinDeltaVelocityForHitEvents();
 				for (const Chaos::FPBDCollisionConstraintHandle * ContactHandle : CollisionRule.GetConstConstraintHandles())
 				{
+					if (NumValidCollisions >= CollisionRule.NumConstraints())
+					{
+						break;
+					}
+
 					if (ContactHandle->GetType() == FCollisionConstraintBase::FType::SinglePoint ||
 						ContactHandle->GetType() == FCollisionConstraintBase::FType::SinglePointSwept)
 					{
 						const FRigidBodyPointContactConstraint& Constraint = (ContactHandle->GetType() == FCollisionConstraintBase::FType::SinglePoint) ?
 							ContactHandle->GetPointContact() :
 							*ContactHandle->GetSweptPointContact().As<FRigidBodyPointContactConstraint>();
-
-						// Since Clustered GCs can be unioned the particleIndex representing the union 
-						// is not associated with a PhysicsProxy
-						if (const TSet<IPhysicsProxyBase*>* Proxies = Solver->GetProxies(Constraint.Particle[0]->Handle()))
+						
+						if (ensure(!Constraint.AccumulatedImpulse.ContainsNaN() && FMath::IsFinite(Constraint.GetPhi())))
 						{
-							for (IPhysicsProxyBase* Proxy : *Proxies)
+							FGeometryParticleHandle* Particle0 = Constraint.Particle[0];
+							FGeometryParticleHandle* Particle1 = Constraint.Particle[1];
+							FKinematicGeometryParticleHandle* Body0 = Particle0->CastToKinematicParticle();
+
+							// presently when a rigidbody or kinematic hits static geometry then Body1 is null
+							FKinematicGeometryParticleHandle* Body1 = Particle1->CastToKinematicParticle();
+
+							const FImplicitObject* Implicit0 = Constraint.Manifold.Implicit[0];
+							const FImplicitObject* Implicit1 = Constraint.Manifold.Implicit[1];
+
+							const FPerShapeData* Shape0 = Particle0->GetImplicitShape(Implicit0);
+							const FPerShapeData* Shape1 = Particle1->GetImplicitShape(Implicit1);
+
+							// If we don't have a filter - allow the notify, otherwise obey the filter flag
+							const bool bFilter0Notify = Shape0 ? Shape0->GetSimData().HasFlag(EFilterFlags::ContactNotify) : true;
+							const bool bFilter1Notify = Shape1 ? Shape1->GetSimData().HasFlag(EFilterFlags::ContactNotify) : true;
+
+							if(!bFilter0Notify && !bFilter1Notify)
 							{
-								if (NumValidCollisions >= CollisionRule.NumConstraints())
-								{
-									break;
-								}
+								// No need to notify - engine didn't request notifications for either shape.
+								continue;
+							}
 
-								if (Proxy != nullptr)
+							if (!Constraint.AccumulatedImpulse.IsZero() && Body0)
+							{
+								if (ensure(!Constraint.GetLocation().ContainsNaN() &&
+									!Constraint.GetNormal().ContainsNaN()) &&
+									!Body0->V().ContainsNaN() &&
+									!Body0->W().ContainsNaN() &&
+									(Body1 == nullptr || ((!Body1->V().ContainsNaN()) && !Body1->W().ContainsNaN())))
 								{
-									if (ensure(!Constraint.AccumulatedImpulse.ContainsNaN() && FMath::IsFinite(Constraint.GetPhi())))
-									{
-										FGeometryParticleHandle* Particle0 = Constraint.Particle[0];
-										FGeometryParticleHandle* Particle1 = Constraint.Particle[1];
-										FKinematicGeometryParticleHandle* Body0 = Particle0->CastToKinematicParticle();
-
-										// presently when a rigidbody or kinematic hits static geometry then Body1 is null
-										FKinematicGeometryParticleHandle* Body1 = Particle1->CastToKinematicParticle();
-
-										const FImplicitObject* Implicit0 = Constraint.Manifold.Implicit[0];
-										const FImplicitObject* Implicit1 = Constraint.Manifold.Implicit[1];
-
-										const FPerShapeData* Shape0 = Particle0->GetImplicitShape(Implicit0);
-										const FPerShapeData* Shape1 = Particle1->GetImplicitShape(Implicit1);
-
-										// If we don't have a filter - allow the notify, otherwise obey the filter flag
-										const bool bFilter0Notify = Shape0 ? Shape0->GetSimData().HasFlag(EFilterFlags::ContactNotify) : true;
-										const bool bFilter1Notify = Shape1 ? Shape1->GetSimData().HasFlag(EFilterFlags::ContactNotify) : true;
-
-										if(!bFilter0Notify && !bFilter1Notify)
-										{
-											// No need to notify - engine didn't request notifications for either shape.
-											continue;
-										}
-
-										if (!Constraint.AccumulatedImpulse.IsZero() && Body0)
-										{
-											if (ensure(!Constraint.GetLocation().ContainsNaN() &&
-												!Constraint.GetNormal().ContainsNaN()) &&
-												!Body0->V().ContainsNaN() &&
-												!Body0->W().ContainsNaN() &&
-												(Body1 == nullptr || ((!Body1->V().ContainsNaN()) && !Body1->W().ContainsNaN())))
-											{
-												ValidCollisionHandles[NumValidCollisions] = ContactHandle;
-												NumValidCollisions++;
-											}
-										}
-									}
+									ValidCollisionHandles[NumValidCollisions] = ContactHandle;
+									NumValidCollisions++;
 								}
 							}
 						}
@@ -152,10 +141,8 @@ namespace Chaos
 							Data.AccumulatedImpulse = Constraint.AccumulatedImpulse;
 							Data.Normal = Constraint.GetNormal();
 							Data.PenetrationDepth = Constraint.GetPhi();
-							Data.ParticleProxy = Solver->GetProxies(Particle0->Handle()) && Solver->GetProxies(Particle0->Handle())->Array().Num() ? 
-								Solver->GetProxies(Particle0->Handle())->Array().operator[](0) : nullptr; // @todo(chaos) : Iterate all proxies
-							Data.LevelsetProxy = Solver->GetProxies(Particle1->Handle()) && Solver->GetProxies(Particle0->Handle())->Array().Num() ? 
-								Solver->GetProxies(Particle1->Handle())->Array().operator[](0) : nullptr; // @todo(chaos) : Iterate all proxies
+							Data.Particle = Particle0;
+							Data.Levelset = Particle1;
 
 							if (FPBDRigidParticleHandle * Rigid0 = Particle0->CastToRigidParticle())
 							{
@@ -184,8 +171,8 @@ namespace Chaos
 								Data.Mass2 = PBDRigid1->M();
 							}
 
-							IPhysicsProxyBase* const PhysicsProxy = Data.ParticleProxy;
-							IPhysicsProxyBase* const OtherPhysicsProxy = Data.LevelsetProxy;
+							IPhysicsProxyBase* const PhysicsProxy = Particle0->PhysicsProxy();
+							IPhysicsProxyBase* const OtherPhysicsProxy = Particle1->PhysicsProxy();
 							//Data.Material1 = nullptr; // #todo: provide UPhysicalMaterial for Particle
 							//Data.Material2 = nullptr; // #todo: provide UPhysicalMaterial for Levelset
 
@@ -280,8 +267,6 @@ namespace Chaos
 							BreakingData.AngularVelocity = PBDRigid->W();
 							BreakingData.Mass = PBDRigid->M();
 							BreakingData.Particle = PBDRigid;
-							BreakingData.ParticleProxy = Solver->GetProxies(PBDRigid->Handle()) && Solver->GetProxies(PBDRigid->Handle())->Array().Num() ?
-								Solver->GetProxies(PBDRigid->Handle())->Array().operator[](0) : nullptr; // @todo(chaos) : Iterate all proxies
 							
 							if(PBDRigid->Geometry()->HasBoundingBox())
 							{
@@ -296,7 +281,7 @@ namespace Chaos
 								BreakingDataArrayItem = BreakingData;
 
 								// Add to AllBreakingIndicesByPhysicsProxy
-								AllBreakingIndicesByPhysicsProxy.FindOrAdd(BreakingData.ParticleProxy).Add(FEventManager::EncodeCollisionIndex(NewIdx, false));
+								AllBreakingIndicesByPhysicsProxy.FindOrAdd(BreakingData.Particle->PhysicsProxy()).Add(FEventManager::EncodeCollisionIndex(NewIdx, false));
 
 #if 0 // #todo
 								// If AllBreakingsArray[Idx].ParticleIndex is a cluster store an index for a mesh in this cluster
@@ -409,19 +394,13 @@ namespace Chaos
 			{
 				if(SleepData.Particle)
 				{
-					if (const TSet<IPhysicsProxyBase*>* Proxies = Solver->GetProxies(SleepData.Particle))
+					FGeometryParticle* Particle = SleepData.Particle->GTGeometryParticle();
+					if (Particle && SleepData.Particle->PhysicsProxy())
 					{
-						for (IPhysicsProxyBase* Proxy : *Proxies)
-						{
-							FGeometryParticle* Particle = SleepData.Particle->GTGeometryParticle();
-							if (Particle != nullptr && Proxy != nullptr)
-							{
-								int32 NewIdx = EventSleepDataArray.Add(FSleepingData());
-								FSleepingData& SleepingDataArrayItem = EventSleepDataArray[NewIdx];
-								SleepingDataArrayItem.Particle = Particle;
-								SleepingDataArrayItem.Sleeping = SleepData.Sleeping;
-							}
-						}
+						int32 NewIdx = EventSleepDataArray.Add(FSleepingData());
+						FSleepingData& SleepingDataArrayItem = EventSleepDataArray[NewIdx];
+						SleepingDataArrayItem.Particle = Particle;
+						SleepingDataArrayItem.Sleeping = SleepData.Sleeping;
 					}
 				}
 			}
diff --git a/Engine/Source/Runtime/Experimental/Chaos/Private/PBDRigidActiveParticlesBuffer.cpp b/Engine/Source/Runtime/Experimental/Chaos/Private/PBDRigidActiveParticlesBuffer.cpp
--- a/Engine/Source/Runtime/Experimental/Chaos/Private/PBDRigidActiveParticlesBuffer.cpp
+++ b/Engine/Source/Runtime/Experimental/Chaos/Private/PBDRigidActiveParticlesBuffer.cpp
@@ -33,25 +33,20 @@ namespace Chaos
 		{
 			if (ActiveParticle.Handle())	//can this be null?
 			{
-				if (const TSet<IPhysicsProxyBase*>* Proxies = Solver->GetProxies(ActiveParticle.Handle()))//can this be null?
+				IPhysicsProxyBase* Proxy = ActiveParticle.Handle()->PhysicsProxy();
+				
+				if(Proxy != nullptr)	//can this be null?
 				{
-					for(IPhysicsProxyBase* Proxy : *Proxies)
+					if(Proxy->GetType() == EPhysicsProxyType::SingleParticleProxy)
 					{
-						if(Proxy != nullptr)	//can this be null?
-						{
-							if(Proxy->GetType() == EPhysicsProxyType::SingleParticleProxy)
-							{
-								ensure(Proxies->Num() == 1);	//single rigid should only have one proxy
-								ActiveGameThreadParticles.Add(static_cast<FSingleParticlePhysicsProxy*>(Proxy));
-							}
-							else
-							{
-								//must be a geometry collection
-								PhysicsParticleProxies.Add(Proxy);
-							}
-						}
+						ActiveGameThreadParticles.Add(static_cast<FSingleParticlePhysicsProxy*>(Proxy));
 					}
-				}
+					else
+					{
+						//must be a geometry collection
+						PhysicsParticleProxies.Add(Proxy);
+					}
+				}	
 			}
 		}
 	}
diff --git a/Engine/Source/Runtime/Experimental/Chaos/Private/PBDRigidsSolver.cpp b/Engine/Source/Runtime/Experimental/Chaos/Private/PBDRigidsSolver.cpp
--- a/Engine/Source/Runtime/Experimental/Chaos/Private/PBDRigidsSolver.cpp
+++ b/Engine/Source/Runtime/Experimental/Chaos/Private/PBDRigidsSolver.cpp
@@ -362,14 +362,14 @@ namespace Chaos
 		MEvolution->AddConstraintRule(&SuspensionConstraintRule);
 
 		MEvolution->SetInternalParticleInitilizationFunction(
-			[this](const Chaos::FGeometryParticleHandle* OldParticle, const Chaos::FGeometryParticleHandle* NewParticle) {
-				if (const TSet<IPhysicsProxyBase*>* Proxies = GetProxies(OldParticle))
+			[this](const Chaos::FGeometryParticleHandle* OldParticle, Chaos::FGeometryParticleHandle* NewParticle) 
+			{
+				IPhysicsProxyBase* Proxy = const_cast<IPhysicsProxyBase*>(OldParticle->PhysicsProxy());
+				if (Chaos::FPBDRigidClusteredParticleHandle* NewClusteredParticle = NewParticle->CastToClustered())
 				{
-					for (IPhysicsProxyBase* Proxy : *Proxies)
-					{
-						this->AddParticleToProxy(NewParticle, Proxy);
-					}
+					NewClusteredParticle->AddPhysicsProxy(Proxy);
 				}
+				NewParticle->SetPhysicsProxy(Proxy);
 			});
 		
 		JointConstraints.SetUpdateVelocityInApplyConstraints(true);
@@ -451,33 +451,22 @@ namespace Chaos
 		{
 			UE_LOG(LogPBDRigidsSolver, Verbose, TEXT("FPBDRigidsSolver::UnregisterObject() ~ Dequeue"));
 
-				// Generally need to remove stale events for particles that no longer exist
-				GetEventManager()->template ClearEvents<FCollisionEventData>(EEventType::Collision, [Proxy]
-				(FCollisionEventData& EventDataInOut)
+			// Generally need to remove stale events for particles that no longer exist
+			GetEventManager()->template ClearEvents<FCollisionEventData>(EEventType::Collision, [Proxy]
+			(FCollisionEventData& EventDataInOut)
+			{
+				Chaos::FCollisionDataArray const& CollisionData = EventDataInOut.CollisionData.AllCollisionsArray;
+				if (CollisionData.Num() > 0)
 				{
-					Chaos::FCollisionDataArray const& CollisionData = EventDataInOut.CollisionData.AllCollisionsArray;
-					if (CollisionData.Num() > 0)
+					check(Proxy);
+					TArray<int32> const* const CollisionIndices = EventDataInOut.PhysicsProxyToCollisionIndices.PhysicsProxyToIndicesMap.Find(Proxy);
+					if (CollisionIndices)
 					{
-						check(Proxy);
-						TArray<int32> const* const CollisionIndices = EventDataInOut.PhysicsProxyToCollisionIndices.PhysicsProxyToIndicesMap.Find(Proxy);
-						if (CollisionIndices)
-						{
-							for (int32 EncodedCollisionIdx : *CollisionIndices)
-							{
-								bool bSwapOrder;
-								int32 CollisionIdx = Chaos::FEventManager::DecodeCollisionIndex(EncodedCollisionIdx, bSwapOrder);
-
-								// invalidate but don't delete from array, as this would mean we'd need to reindex PhysicsProxyToIndicesMap to maintain the other collisions lookup
-								Chaos::FCollidingData& CollisionDataItem = EventDataInOut.CollisionData.AllCollisionsArray[CollisionIdx];
-								CollisionDataItem.ParticleProxy = nullptr;
-								CollisionDataItem.LevelsetProxy = nullptr;
-							}
-
-							EventDataInOut.PhysicsProxyToCollisionIndices.PhysicsProxyToIndicesMap.Remove(Proxy);
-						}
+						EventDataInOut.PhysicsProxyToCollisionIndices.PhysicsProxyToIndicesMap.Remove(Proxy);
 					}
+				}
 
-				});
+			});
 
 			// Get the physics thread-handle from the proxy, and then delete the proxy.
 			//
@@ -499,12 +488,6 @@ namespace Chaos
 					RewindData->RemoveParticle(Handle->UniqueIdx());
 				}
   
-				if(LogCorruptMap)
-				{
-					UE_LOG(LogChaos, Warning, TEXT("UnregisterObject this:%llx, Handle:%llx &MParticleToProxy:%llx, MParticleToProxy.Num():%d"), this, Handle, &MParticleToProxy, MParticleToProxy.Num());
-				}
-				MParticleToProxy.Remove(Handle);
-  
 				// Use the handle to destroy the particle data
 				GetEvolution()->DestroyParticle(Handle);
 			}
@@ -545,11 +528,6 @@ namespace Chaos
 
 		EnqueueCommandImmediate([InProxy, this]()
 			{
-				const TArray<Chaos::FPBDRigidClusteredParticleHandle*>& ParticleHandles = InProxy->GetSolverParticleHandles();
-				for (const Chaos::FPBDRigidClusteredParticleHandle* ParticleHandle : ParticleHandles)
-				{
-					RemoveParticleToProxy(ParticleHandle);
-				}
 				GeometryCollectionPhysicsProxies_Internal.RemoveSingle(InProxy);
 				InProxy->SyncBeforeDestroy();
 				InProxy->OnRemoveFromSolver(this);
@@ -878,7 +856,7 @@ namespace Chaos
 			if(bIsNew)
 			{
 				auto Handle = Proxy->GetHandle_LowLevel();
-				AddParticleToProxy(Handle,Proxy);
+				Handle->SetPhysicsProxy(Proxy);
 				GetEvolution()->CreateParticle(Handle);
 				Proxy->SetInitialized(true);
 			}
@@ -1108,36 +1086,45 @@ namespace Chaos
 
 			for (Chaos::TPBDRigidParticleHandleImp<FReal, 3, false>& DirtyParticle : DirtyParticles)
 			{
-				if( const TSet<IPhysicsProxyBase*>* Proxies = GetProxies(DirtyParticle.Handle()))
+				IPhysicsProxyBase* Proxy = DirtyParticle.Handle()->PhysicsProxy();
+				if(Proxy != nullptr)
 				{
-					for (IPhysicsProxyBase* Proxy : *Proxies)
+					switch(DirtyParticle.GetParticleType())
 					{
-							if(Proxy != nullptr)
-						{
-								switch(DirtyParticle.GetParticleType())
+						case Chaos::EParticleType::Rigid:
+							{
+								PullData->DirtyRigids.AddDefaulted();
+								((FSingleParticlePhysicsProxy*)(Proxy))->BufferPhysicsResults(PullData->DirtyRigids.Last());
+								break;
+							}
+						case Chaos::EParticleType::Kinematic:
+						case Chaos::EParticleType::Static:
+							ensure(false);
+							break;
+						case Chaos::EParticleType::GeometryCollection:
+							ActiveGC.AddUnique((FGeometryCollectionPhysicsProxy*)(Proxy));
+							break;
+						case Chaos::EParticleType::Clustered:
+							if (auto ClusterParticle = DirtyParticle.CastToClustered())
 							{
-							case Chaos::EParticleType::Rigid:
+								if (ClusterParticle->InternalCluster())
 								{
-									PullData->DirtyRigids.AddDefaulted();
-									((FSingleParticlePhysicsProxy*)(Proxy))->BufferPhysicsResults(PullData->DirtyRigids.Last());
-									break;
+									const TSet<IPhysicsProxyBase*> Proxies = ClusterParticle->PhysicsProxies();
+									for (IPhysicsProxyBase* ClusterProxy : Proxies)
+									{
+										ActiveGC.AddUnique((FGeometryCollectionPhysicsProxy*)(ClusterProxy));
+									}
+								}
+								else
+								{
+									ActiveGC.AddUnique((FGeometryCollectionPhysicsProxy*)(Proxy));
 								}
-							case Chaos::EParticleType::Kinematic:
-							case Chaos::EParticleType::Static:
-								ensure(false);
-								break;
-							case Chaos::EParticleType::GeometryCollection:
-								ActiveGC.AddUnique((FGeometryCollectionPhysicsProxy*)(Proxy));
-								break;
-							case Chaos::EParticleType::Clustered:
-								ActiveGC.AddUnique((FGeometryCollectionPhysicsProxy*)(Proxy));
-								break;
-							default:
-								check(false);
 							}
-						}
+							break;
+						default:
+							check(false);
 					}
-				}
+				}	
 			}
 		}
 
diff --git a/Engine/Source/Runtime/Experimental/Chaos/Private/PhysicsProxy/GeometryCollectionPhysicsProxy.cpp b/Engine/Source/Runtime/Experimental/Chaos/Private/PhysicsProxy/GeometryCollectionPhysicsProxy.cpp
--- a/Engine/Source/Runtime/Experimental/Chaos/Private/PhysicsProxy/GeometryCollectionPhysicsProxy.cpp
+++ b/Engine/Source/Runtime/Experimental/Chaos/Private/PhysicsProxy/GeometryCollectionPhysicsProxy.cpp
@@ -662,7 +662,7 @@ void FGeometryCollectionPhysicsProxy::InitializeBodiesPT(Chaos::FPBDRigidsSolver
 
 				Chaos::TPBDGeometryCollectionParticleHandle<float, 3>* Handle = Handles[NextIdx++];
 
-				RigidsSolver->AddParticleToProxy(Handle, this);
+				Handle->SetPhysicsProxy(this);
 
 				SolverParticleHandles[Idx] = Handle;
 				HandleToTransformGroupIndex.Add(Handle, Idx);
@@ -877,7 +877,7 @@ void FGeometryCollectionPhysicsProxy::InitializeBodiesPT(Chaos::FPBDRigidsSolver
 					SolverClusterHandles[TransformGroupIndex] = Handle;
 					SolverParticleHandles[TransformGroupIndex] = Handle;
 					HandleToTransformGroupIndex.Add(Handle, TransformGroupIndex);
-					RigidsSolver->AddParticleToProxy(Handle, this);
+					Handle->SetPhysicsProxy(this);
 
 					RigidsSolver->GetEvolution()->DirtyParticle(*Handle);
 				}
@@ -1342,7 +1342,6 @@ void FGeometryCollectionPhysicsProxy::OnRemoveFromSolver(Chaos::FPBDRigidsSolver
 	{
 		if (FClusterHandle* Handle = SolverParticleHandles[i])
 		{
-			RBDSolver->RemoveParticleToProxy(Handle);
 			if (FClusterHandle* ParentCluster = Evolution->GetRigidClustering().DestroyClusterParticle(Handle))
 			{
 				if (ParentCluster->InternalCluster())
diff --git a/Engine/Source/Runtime/Experimental/Chaos/Public/Chaos/ExternalCollisionData.h b/Engine/Source/Runtime/Experimental/Chaos/Public/Chaos/ExternalCollisionData.h
--- a/Engine/Source/Runtime/Experimental/Chaos/Public/Chaos/ExternalCollisionData.h
+++ b/Engine/Source/Runtime/Experimental/Chaos/Public/Chaos/ExternalCollisionData.h
@@ -22,8 +22,8 @@ namespace Chaos
 			, Normal(FVec3((FReal)0.0))
 			, Velocity1(FVec3((FReal)0.0))
 			, Velocity2(FVec3((FReal)0.0))
-		    , DeltaVelocity1(FVec3((FReal)0.0))
-		    , DeltaVelocity2(FVec3((FReal)0.0))
+			, DeltaVelocity1(FVec3((FReal)0.0))
+			, DeltaVelocity2(FVec3((FReal)0.0))
 			, AngularVelocity1(FVec3((FReal)0.0))
 			, AngularVelocity2(FVec3((FReal)0.0))
 			, Mass1((FReal)0.0)
@@ -31,14 +31,12 @@ namespace Chaos
 			, PenetrationDepth((FReal)0.0)
 			, Particle(nullptr)
 			, Levelset(nullptr)
-			, ParticleProxy(nullptr)
-		    , LevelsetProxy(nullptr)
 		{}
 
 		FCollidingData(FVec3 InLocation, FVec3 InAccumulatedImpulse, FVec3 InNormal, FVec3 InVelocity1, FVec3 InVelocity2, FVec3 InDeltaVelocity1, FVec3 InDeltaVelocity2
-		, FVec3 InAngularVelocity1, FVec3 InAngularVelocity2, FReal InMass1, FReal InMass2, FReal InPenetrationDepth, FGeometryParticle* InParticle
-		, FGeometryParticle* InLevelset, IPhysicsProxyBase* InParticleProxy, IPhysicsProxyBase* InLevelsetProxy)
-		    : Location(InLocation)
+			, FVec3 InAngularVelocity1, FVec3 InAngularVelocity2, FReal InMass1,FReal InMass2,  FReal InPenetrationDepth, FGeometryParticleHandle* InParticle
+			, FGeometryParticleHandle* InLevelset)
+			: Location(InLocation)
 			, AccumulatedImpulse(InAccumulatedImpulse)
 			, Normal(InNormal)
 			, Velocity1(InVelocity1)
@@ -52,12 +50,8 @@ namespace Chaos
 			, PenetrationDepth(InPenetrationDepth)
 			, Particle(InParticle)
 			, Levelset(InLevelset)
-			, ParticleProxy(InParticleProxy)
-		    , LevelsetProxy(InLevelsetProxy)
 		{}
 
-		bool IsValid() { return (ParticleProxy && LevelsetProxy); }
-
 		FVec3 Location;
 		FVec3 AccumulatedImpulse;
 		FVec3 Normal;
@@ -66,10 +60,8 @@ namespace Chaos
 		FVec3 AngularVelocity1, AngularVelocity2;
 		FReal Mass1, Mass2;
 		FReal PenetrationDepth;
-		FGeometryParticle* Particle;
-		FGeometryParticle* Levelset;
-		IPhysicsProxyBase* ParticleProxy;
-		IPhysicsProxyBase* LevelsetProxy;
+		FGeometryParticleHandle* Particle;
+		FGeometryParticleHandle* Levelset;
 	};
 
 	/*
@@ -89,8 +81,6 @@ namespace Chaos
 			, Mass2((FReal)0.0)
 			, Particle(nullptr)
 			, Levelset(nullptr)
-		    , ParticleProxy(nullptr)
-		    , LevelsetProxy(nullptr)
 			, BoundingboxVolume((FReal)-1.0)
 			, BoundingboxExtentMin((FReal)-1.0)
 			, BoundingboxExtentMax((FReal)-1.0)
@@ -99,8 +89,7 @@ namespace Chaos
 
 		FCollidingDataExt(
 		    FVec3 InLocation, FVec3 InAccumulatedImpulse, FVec3 InNormal, FVec3 InVelocity1, FVec3 InVelocity2
-			, FVec3 InAngularVelocity1, FVec3 InAngularVelocity2, FReal InMass1, FReal InMass2, FGeometryParticle* InParticle
-			, FGeometryParticle* InLevelset, IPhysicsProxyBase* InParticleProxy, IPhysicsProxyBase* InLevelsetProxy
+			, FVec3 InAngularVelocity1, FVec3 InAngularVelocity2, FReal InMass1, FReal InMass2, FGeometryParticleHandle* InParticle, FGeometryParticleHandle* InLevelset
 			, FReal InBoundingboxVolume, FReal InBoundingboxExtentMin, FReal InBoundingboxExtentMax, int32 InSurfaceType)
 			: Location(InLocation)
 			, AccumulatedImpulse(InAccumulatedImpulse)
@@ -113,8 +102,6 @@ namespace Chaos
 			, Mass2(InMass2)
 			, Particle(InParticle)
 			, Levelset(InLevelset)
-		    , ParticleProxy(InParticleProxy)
-		    , LevelsetProxy(InLevelsetProxy)
 			, BoundingboxVolume(InBoundingboxVolume)
 			, BoundingboxExtentMin(InBoundingboxExtentMin)
 			, BoundingboxExtentMax(InBoundingboxExtentMax)
@@ -133,8 +120,6 @@ namespace Chaos
 			, Mass2(InCollisionData.Mass2)
 			, Particle(InCollisionData.Particle)
 			, Levelset(InCollisionData.Levelset)
-		    , ParticleProxy(InCollisionData.ParticleProxy)
-		    , LevelsetProxy(InCollisionData.LevelsetProxy)
 			, BoundingboxVolume((FReal)-1.0)
 			, BoundingboxExtentMin((FReal)-1.0)
 			, BoundingboxExtentMax((FReal)-1.0)
@@ -148,10 +133,8 @@ namespace Chaos
 		FVec3 Velocity1, Velocity2;
 		FVec3 AngularVelocity1, AngularVelocity2;
 		FReal Mass1, Mass2;
-		FGeometryParticle* Particle;
-		FGeometryParticle* Levelset;
-		IPhysicsProxyBase* ParticleProxy;
-		IPhysicsProxyBase* LevelsetProxy;
+		FGeometryParticleHandle* Particle;
+		FGeometryParticleHandle* Levelset;
 		FReal BoundingboxVolume;
 		FReal BoundingboxExtentMin, BoundingboxExtentMax;
 		int32 SurfaceType;
@@ -164,7 +147,6 @@ namespace Chaos
 	{
 		FBreakingData()
 			: Particle(nullptr)
-		    , ParticleProxy(nullptr)
 			, Location(FVec3((FReal)0.0))
 			, Velocity(FVec3((FReal)0.0))
 			, AngularVelocity(FVec3((FReal)0.0))
@@ -173,7 +155,6 @@ namespace Chaos
 		{}
 
 		FGeometryParticleHandle* Particle;
-		IPhysicsProxyBase* ParticleProxy;
 		FVec3 Location;
 		FVec3 Velocity;
 		FVec3 AngularVelocity;
@@ -192,7 +173,6 @@ namespace Chaos
 			, AngularVelocity(FVec3((FReal)0.0))
 			, Mass((FReal)0.0)
 		    , Particle(nullptr)
-		    , ParticleProxy(nullptr)
 			, BoundingboxVolume((FReal)-1.0)
 			, BoundingboxExtentMin((FReal)-1.0)
 			, BoundingboxExtentMax((FReal)-1.0)
@@ -204,7 +184,6 @@ namespace Chaos
 			, FVec3 InAngularVelocity
 			, FReal InMass
 			, FGeometryParticleHandle* InParticle
-			, IPhysicsProxyBase* InParticleProxy
 			, FReal InBoundingboxVolume
 			, FReal InBoundingboxExtentMin
 			, FReal InBoundingboxExtentMax
@@ -214,7 +193,6 @@ namespace Chaos
 			, AngularVelocity(InAngularVelocity)
 			, Mass(InMass)
 		    , Particle(InParticle)
-		    , ParticleProxy(InParticleProxy)
 			, BoundingboxVolume(InBoundingboxVolume)
 			, BoundingboxExtentMin(InBoundingboxExtentMin)
 			, BoundingboxExtentMax(InBoundingboxExtentMax)
@@ -227,7 +205,6 @@ namespace Chaos
 			, AngularVelocity(InBreakingData.AngularVelocity)
 			, Mass(InBreakingData.Mass)
 		    , Particle(InBreakingData.Particle)
-		    , ParticleProxy(InBreakingData.ParticleProxy)
 			, BoundingboxVolume((FReal)-1.0)
 			, BoundingboxExtentMin((FReal)-1.0)
 			, BoundingboxExtentMax((FReal)-1.0)
@@ -240,7 +217,6 @@ namespace Chaos
 		FVec3 AngularVelocity;
 		FReal Mass;
 		FGeometryParticleHandle* Particle;
-		IPhysicsProxyBase* ParticleProxy;
 		FReal BoundingboxVolume;
 		FReal BoundingboxExtentMin, BoundingboxExtentMax;
 		int32 SurfaceType;
@@ -268,18 +244,16 @@ namespace Chaos
 			, AngularVelocity(FVec3((FReal)0.0))
 			, Mass((FReal)0.0)
 			, Particle(nullptr)
-		    , ParticleProxy(nullptr)
 			, BoundingBox(FAABB3(FVec3((FReal)0.0), FVec3((FReal)0.0)))
 		{}
 
 		FTrailingData(FVec3 InLocation, FVec3 InVelocity, FVec3 InAngularVelocity, FReal InMass
-			, FGeometryParticleHandle* InParticle, IPhysicsProxyBase* InParticleProxy, Chaos::FAABB3& InBoundingBox)
+			, FGeometryParticleHandle* InParticle, Chaos::TAABB<FReal, 3>& InBoundingBox)
 			: Location(InLocation)
 			, Velocity(InVelocity)
 			, AngularVelocity(InAngularVelocity)
 			, Mass(InMass)
 			, Particle(InParticle)
-		    , ParticleProxy(InParticleProxy)
 			, BoundingBox(InBoundingBox)
 		{}
 
@@ -288,7 +262,6 @@ namespace Chaos
 		FVec3 AngularVelocity;
 		FReal Mass;
 		FGeometryParticleHandle* Particle;
-		IPhysicsProxyBase* ParticleProxy;
 		Chaos::FAABB3 BoundingBox;
 
 		friend inline uint32 GetTypeHash(const FTrailingData& Other)
@@ -324,7 +297,6 @@ namespace Chaos
 			, FVec3 InAngularVelocity
 			, FReal InMass
 			, FGeometryParticleHandle* InParticle
-			, IPhysicsProxyBase* InParticleProxy
 			, FReal InBoundingboxVolume
 			, FReal InBoundingboxExtentMin
 			, FReal InBoundingboxExtentMax
@@ -334,7 +306,6 @@ namespace Chaos
 			, AngularVelocity(InAngularVelocity)
 			, Mass(InMass)
 			, Particle(InParticle)
-		    , ParticleProxy(InParticleProxy)
 			, BoundingboxVolume(InBoundingboxVolume)
 			, BoundingboxExtentMin(InBoundingboxExtentMin)
 			, BoundingboxExtentMax(InBoundingboxExtentMax)
@@ -347,7 +318,6 @@ namespace Chaos
 			, AngularVelocity(InTrailingData.AngularVelocity)
 			, Mass(InTrailingData.Mass)
 			, Particle(InTrailingData.Particle)
-		    , ParticleProxy(InTrailingData.ParticleProxy)
 			, BoundingboxVolume((FReal)-1.0)
 			, BoundingboxExtentMin((FReal)-1.0)
 			, BoundingboxExtentMax((FReal)-1.0)
@@ -359,7 +329,6 @@ namespace Chaos
 		FVec3 AngularVelocity;
 		FReal Mass;
 		FGeometryParticleHandle* Particle;
-		IPhysicsProxyBase* ParticleProxy;
 		//	int32 ParticleIndexMesh; // If ParticleIndex points to a cluster then this index will point to an actual mesh in the cluster
 								 // It is important to be able to get extra data from the component
 		FReal BoundingboxVolume;
diff --git a/Engine/Source/Runtime/Experimental/Chaos/Public/Chaos/GeometryParticles.h b/Engine/Source/Runtime/Experimental/Chaos/Public/Chaos/GeometryParticles.h
--- a/Engine/Source/Runtime/Experimental/Chaos/Public/Chaos/GeometryParticles.h
+++ b/Engine/Source/Runtime/Experimental/Chaos/Public/Chaos/GeometryParticles.h
@@ -358,6 +358,7 @@ namespace Chaos
 			{
 				TArrayCollection::AddArray(&MGeometryParticleHandle);
 				TArrayCollection::AddArray(&MGeometryParticle);
+				TArrayCollection::AddArray(&MPhysicsProxy);
 			}
 
 		}
@@ -371,6 +372,7 @@ namespace Chaos
 			, MDynamicGeometry(MoveTemp(Other.MDynamicGeometry))
 			, MGeometryParticleHandle(MoveTemp(Other.MGeometryParticleHandle))
 			, MGeometryParticle(MoveTemp(Other.MGeometryParticle))
+			, MPhysicsProxy(MoveTemp(Other.MPhysicsProxy))
 			, MShapesArray(MoveTemp(Other.MShapesArray))
 			, ImplicitShapeMap(MoveTemp(Other.ImplicitShapeMap))
 			, MLocalBounds(MoveTemp(Other.MLocalBounds))
@@ -414,6 +416,7 @@ namespace Chaos
 			{
 				TArrayCollection::AddArray(&MGeometryParticleHandle);
 				TArrayCollection::AddArray(&MGeometryParticle);
+				TArrayCollection::AddArray(&MPhysicsProxy);
 			}
 		}
 
@@ -450,6 +453,7 @@ namespace Chaos
 			{
 				TArrayCollection::AddArray(&MGeometryParticleHandle);
 				TArrayCollection::AddArray(&MGeometryParticle);
+				TArrayCollection::AddArray(&MPhysicsProxy);
 			}
 		}
 
@@ -586,6 +590,13 @@ namespace Chaos
 		CHAOS_API TGeometryParticle<T, d>* GTGeometryParticle(const int32 Index) const { return MGeometryParticle[Index]; }
 		CHAOS_API TGeometryParticle<T, d>*& GTGeometryParticle(const int32 Index) { return MGeometryParticle[Index]; }
 
+		CHAOS_API const IPhysicsProxyBase* PhysicsProxy(const int32 Index) const { return MPhysicsProxy[Index];  }
+		CHAOS_API IPhysicsProxyBase* PhysicsProxy(const int32 Index) { return MPhysicsProxy[Index]; }
+		CHAOS_API void SetPhysicsProxy(const int32 Index, IPhysicsProxyBase* InPhysicsProxy)
+		{
+			MPhysicsProxy[Index] = InPhysicsProxy;
+		}
+
 		CHAOS_API FWeakParticleHandle& WeakParticleHandle(const int32 Index)
 		{
 			FWeakParticleHandle& WeakHandle = MWeakParticleHandle[Index];
@@ -739,6 +750,7 @@ public:
 		TArrayCollectionArray<TUniquePtr<FImplicitObject>> MDynamicGeometry;
 		TArrayCollectionArray<TSerializablePtr<TGeometryParticleHandle<T, d>>> MGeometryParticleHandle;
 		TArrayCollectionArray<TGeometryParticle<T, d>*> MGeometryParticle;
+		TArrayCollectionArray<IPhysicsProxyBase*> MPhysicsProxy;
 		TArrayCollectionArray<FShapesArray> MShapesArray;
 		TArrayCollectionArray<TMap<const FImplicitObject*, int32>> ImplicitShapeMap;
 		TArrayCollectionArray<TAABB<T,d>> MLocalBounds;
diff --git a/Engine/Source/Runtime/Experimental/Chaos/Public/Chaos/PBDRigidClusteredParticles.h b/Engine/Source/Runtime/Experimental/Chaos/Public/Chaos/PBDRigidClusteredParticles.h
--- a/Engine/Source/Runtime/Experimental/Chaos/Public/Chaos/PBDRigidClusteredParticles.h
+++ b/Engine/Source/Runtime/Experimental/Chaos/Public/Chaos/PBDRigidClusteredParticles.h
@@ -108,6 +108,7 @@ class TPBDRigidClusteredParticles : public TPBDRigidParticles<T, d>
 	, MInternalCluster(MoveTemp(Other.MInternalCluster))
 	, MMultiChildProxyId(MoveTemp(Other.MMultiChildProxyId))
 	, MMultiChildProxyData(MoveTemp(Other.MMultiChildProxyData))
+	, MPhysicsProxies(MoveTemp(Other.MPhysicsProxies))
 	, MCollisionImpulses(MoveTemp(Other.MCollisionImpulses))
 	, MStrains(MoveTemp(Other.MStrains))
 	, MConnectivityEdges(MoveTemp(Other.MConnectivityEdges))
@@ -138,6 +139,9 @@ class TPBDRigidClusteredParticles : public TPBDRigidParticles<T, d>
 	const auto& MultiChildProxyData(int32 Idx) const { return MMultiChildProxyData[Idx]; }
 	auto& MultiChildProxyData(int32 Idx) { return MMultiChildProxyData[Idx]; }
 
+	const auto& PhysicsProxies(int32 Idx) const { return MPhysicsProxies[Idx]; }
+	auto& PhysicsProxies(int32 Idx) { return MPhysicsProxies[Idx]; }
+
 	const auto& CollisionImpulses(int32 Idx) const { return MCollisionImpulses[Idx]; }
 	auto& CollisionImpulses(int32 Idx) { return MCollisionImpulses[Idx]; }
 	auto& CollisionImpulsesArray() { return MCollisionImpulses; }
@@ -187,6 +191,7 @@ class TPBDRigidClusteredParticles : public TPBDRigidParticles<T, d>
 		  TArrayCollection::AddArray(&MChildrenSpatial);
 		  TArrayCollection::AddArray(&MMultiChildProxyId);
 		  TArrayCollection::AddArray(&MMultiChildProxyData);
+		  TArrayCollection::AddArray(&MPhysicsProxies);
 		  TArrayCollection::AddArray(&MCollisionImpulses);
 		  TArrayCollection::AddArray(&MStrains);
 		  TArrayCollection::AddArray(&MConnectivityEdges);
@@ -200,6 +205,9 @@ class TPBDRigidClusteredParticles : public TPBDRigidParticles<T, d>
 	  TArrayCollectionArray<FMultiChildProxyId> MMultiChildProxyId;
 	  TArrayCollectionArray<TUniquePtr<TMultiChildProxyData<T, d>>> MMultiChildProxyData;
 
+	  // Multiple proxy pointers required for internal clusters
+	  TArrayCollectionArray<TSet<IPhysicsProxyBase*>> MPhysicsProxies;
+
 	  // Collision Impulses
 	  TArrayCollectionArray<T> MCollisionImpulses;
 
diff --git a/Engine/Source/Runtime/Experimental/Chaos/Public/Chaos/PBDRigidsEvolutionGBF.h b/Engine/Source/Runtime/Experimental/Chaos/Public/Chaos/PBDRigidsEvolutionGBF.h
--- a/Engine/Source/Runtime/Experimental/Chaos/Public/Chaos/PBDRigidsEvolutionGBF.h
+++ b/Engine/Source/Runtime/Experimental/Chaos/Public/Chaos/PBDRigidsEvolutionGBF.h
@@ -33,7 +33,7 @@ namespace Chaos
 
 	using FPBDRigidsEvolutionInternalHandleCallback = TFunction<void(
 		const FGeometryParticleHandle* OldParticle,
-		const FGeometryParticleHandle* NewParticle)>;
+		FGeometryParticleHandle* NewParticle)>;
 
 	class FPBDRigidsEvolutionGBF : public FPBDRigidsEvolutionBase
 	{
@@ -134,7 +134,7 @@ namespace Chaos
 			InternalParticleInitilization = Cb;
 		}
 
-		FORCEINLINE void DoInternalParticleInitilization(const FGeometryParticleHandle* OldParticle, const FGeometryParticleHandle* NewParticle) 
+		FORCEINLINE void DoInternalParticleInitilization(const FGeometryParticleHandle* OldParticle, FGeometryParticleHandle* NewParticle) 
 		{ 
 			if(InternalParticleInitilization)
 			{
diff --git a/Engine/Source/Runtime/Experimental/Chaos/Public/Chaos/ParticleHandle.h b/Engine/Source/Runtime/Experimental/Chaos/Public/Chaos/ParticleHandle.h
--- a/Engine/Source/Runtime/Experimental/Chaos/Public/Chaos/ParticleHandle.h
+++ b/Engine/Source/Runtime/Experimental/Chaos/Public/Chaos/ParticleHandle.h
@@ -459,6 +459,10 @@ public:
 	TGeometryParticle<T, d>* GTGeometryParticle() const { return GeometryParticles->GTGeometryParticle(ParticleIdx); }
 	TGeometryParticle<T, d>*& GTGeometryParticle() { return GeometryParticles->GTGeometryParticle(ParticleIdx); }
 
+	const IPhysicsProxyBase* PhysicsProxy() const { return GeometryParticles->PhysicsProxy(ParticleIdx); }
+	IPhysicsProxyBase* PhysicsProxy() { return GeometryParticles->PhysicsProxy(ParticleIdx); }
+	void SetPhysicsProxy(IPhysicsProxyBase* PhysicsProxy) { GeometryParticles->SetPhysicsProxy(ParticleIdx, PhysicsProxy); }
+
 	const TKinematicGeometryParticleHandleImp<T, d, bPersistent>* CastToKinematicParticle() const;
 	TKinematicGeometryParticleHandleImp<T, d, bPersistent>* CastToKinematicParticle();
 
@@ -938,6 +942,9 @@ public:
 		return Serializable;
 	}
 
+	const TSet<IPhysicsProxyBase*>& PhysicsProxies() const { return PBDRigidClusteredParticles->PhysicsProxies(ParticleIdx); }
+	void AddPhysicsProxy(IPhysicsProxyBase* PhysicsProxy) { PBDRigidClusteredParticles->PhysicsProxies(ParticleIdx).Add(PhysicsProxy); }
+	
 	void SetClusterId(const ClusterId& Id) { PBDRigidClusteredParticles->ClusterIds(ParticleIdx) = Id; }
 	const ClusterId& ClusterIds() const { return PBDRigidClusteredParticles->ClusterIds(ParticleIdx); }
 	ClusterId& ClusterIds() { return PBDRigidClusteredParticles->ClusterIds(ParticleIdx); }
diff --git a/Engine/Source/Runtime/Experimental/Chaos/Public/PBDRigidsSolver.h b/Engine/Source/Runtime/Experimental/Chaos/Public/PBDRigidsSolver.h
--- a/Engine/Source/Runtime/Experimental/Chaos/Public/PBDRigidsSolver.h
+++ b/Engine/Source/Runtime/Experimental/Chaos/Public/PBDRigidsSolver.h
@@ -222,27 +222,7 @@ namespace Chaos
 
 		FParticlesType& GetParticles() { return Particles; }
 		const FParticlesType& GetParticles() const { return Particles; }
-
-		void AddParticleToProxy(const Chaos::FGeometryParticleHandle* Particle, IPhysicsProxyBase* Proxy)
-		{
-			if (!MParticleToProxy.Find(Particle))
-			{
-				MParticleToProxy.Add(Particle, TSet<IPhysicsProxyBase*>());
-			}
-			MParticleToProxy[Particle].Add(Proxy); 
-		}
 		
-		void RemoveParticleToProxy(const Chaos::FGeometryParticleHandle* Particle)
-		{
-			MParticleToProxy.Remove(Particle);
-		}
-		
-		const TSet<IPhysicsProxyBase*> * GetProxies(const Chaos::FGeometryParticleHandle* Handle) const
-		{
-			const TSet<IPhysicsProxyBase*>* PhysicsProxyPtr = MParticleToProxy.Find(Handle);
-			return PhysicsProxyPtr ? PhysicsProxyPtr : nullptr;
-		}
-
 		/**/
 		FEventManager* GetEventManager() { return MEventManager.Get(); }
 
@@ -343,7 +323,6 @@ namespace Chaos
 		TUniquePtr<FEventManager> MEventManager;
 		TUniquePtr<FSolverEventFilters> MSolverEventFilters;
 		TUniquePtr<FDirtyParticlesBuffer> MDirtyParticlesBuffer;
-		TMap<const Chaos::FGeometryParticleHandle*, TSet<IPhysicsProxyBase*> > MParticleToProxy;
 
 		//
 		// Proxies
diff --git a/Engine/Source/Runtime/Experimental/Chaos/Public/PhysicsProxy/FieldSystemProxyHelper.h b/Engine/Source/Runtime/Experimental/Chaos/Public/PhysicsProxy/FieldSystemProxyHelper.h
--- a/Engine/Source/Runtime/Experimental/Chaos/Public/PhysicsProxy/FieldSystemProxyHelper.h
+++ b/Engine/Source/Runtime/Experimental/Chaos/Public/PhysicsProxy/FieldSystemProxyHelper.h
@@ -221,15 +221,8 @@ namespace Chaos
 
 			for (Chaos::FPBDRigidParticleHandle* Handle : Activated)
 			{
-				if (!RigidSolver->GetProxies(Handle))
-				{
-					const TSet<IPhysicsProxyBase*>* ParentProxies = RigidSolver->GetProxies(Iter.Key);
-					if (ensure(ParentProxies))
-					{
-						for (IPhysicsProxyBase* ParentProxy : *ParentProxies)
-							RigidSolver->AddParticleToProxy(Handle, ParentProxy);
-					}
-				}
+				IPhysicsProxyBase* ParentProxy = Iter.Key->PhysicsProxy();
+				Handle->SetPhysicsProxy(ParentProxy);
 			}
 		}
 	}
diff --git a/Engine/Source/Runtime/Experimental/ChaosSolverEngine/Private/Chaos/ChaosGameplayEventDispatcher.cpp b/Engine/Source/Runtime/Experimental/ChaosSolverEngine/Private/Chaos/ChaosGameplayEventDispatcher.cpp
--- a/Engine/Source/Runtime/Experimental/ChaosSolverEngine/Private/Chaos/ChaosGameplayEventDispatcher.cpp
+++ b/Engine/Source/Runtime/Experimental/ChaosSolverEngine/Private/Chaos/ChaosGameplayEventDispatcher.cpp
@@ -275,7 +275,7 @@ void UChaosGameplayEventDispatcher::HandleCollisionEvents(const Chaos::FCollisio
 								int32 CollisionIdx = Chaos::FEventManager::DecodeCollisionIndex(EncodedCollisionIdx, bSwapOrder);
 
 								Chaos::FCollidingData const& CollisionDataItem = CollisionData[CollisionIdx];
-								IPhysicsProxyBase* const PhysicsProxy1 = bSwapOrder ? CollisionDataItem.ParticleProxy : CollisionDataItem.LevelsetProxy;
+								IPhysicsProxyBase* const PhysicsProxy1 = bSwapOrder ? CollisionDataItem.Particle->PhysicsProxy() : CollisionDataItem.Levelset->PhysicsProxy();
 
 								{
 									bool bNewEntry = false;
@@ -385,18 +385,18 @@ void UChaosGameplayEventDispatcher::HandleBreakingEvents(const Chaos::FBreakingE
 		{
 			for (Chaos::FBreakingData const& BreakingDataItem : BreakingData)
 			{	
-				if (BreakingDataItem.Particle && (BreakingDataItem.ParticleProxy))
+				if (BreakingDataItem.Particle)
 				{
-					UPrimitiveComponent* const PrimComp = Cast<UPrimitiveComponent>(BreakingDataItem.ParticleProxy->GetOwner());
+					UPrimitiveComponent* const PrimComp = Cast<UPrimitiveComponent>(BreakingDataItem.Particle->PhysicsProxy()->GetOwner());
 					if (PrimComp && BreakEventRegistrations.Contains(PrimComp))
 					{
 						// queue them up so we can release the physics data before trigging BP events
 						FChaosBreakEvent& BreakEvent = PendingBreakEvents.AddZeroed_GetRef();
-							BreakEvent.Component = PrimComp;
-							BreakEvent.Location = BreakingDataItem.Location;
-							BreakEvent.Velocity = BreakingDataItem.Velocity;
-							BreakEvent.AngularVelocity = BreakingDataItem.AngularVelocity;
-							BreakEvent.Mass = BreakingDataItem.Mass;
+						BreakEvent.Component = PrimComp;
+						BreakEvent.Location = BreakingDataItem.Location;
+						BreakEvent.Velocity = BreakingDataItem.Velocity;
+						BreakEvent.AngularVelocity = BreakingDataItem.AngularVelocity;
+						BreakEvent.Mass = BreakingDataItem.Mass;
 					}
 				}
 			}
diff --git a/Engine/Source/Runtime/Experimental/GeometryCollectionEngine/Public/ChaosCollisionEventFilter.h b/Engine/Source/Runtime/Experimental/GeometryCollectionEngine/Public/ChaosCollisionEventFilter.h
--- a/Engine/Source/Runtime/Experimental/GeometryCollectionEngine/Public/ChaosCollisionEventFilter.h
+++ b/Engine/Source/Runtime/Experimental/GeometryCollectionEngine/Public/ChaosCollisionEventFilter.h
@@ -43,11 +43,11 @@ struct FChaosCollisionEventData
 
 	// The particle of the collision event
 	//UPROPERTY(EditAnywhere, BlueprintReadOnly, Category = "Data")
-	Chaos::FGeometryParticle* Particle;
+	Chaos::FGeometryParticleHandle* Particle;
 
 	// The level set of the collision event
 	//UPROPERTY(EditAnywhere, BlueprintReadOnly, Category = "Data")
-	Chaos::FGeometryParticle* Levelset;
+	Chaos::FGeometryParticleHandle* Levelset;
 
 	FChaosCollisionEventData()
 		: Location(FVector::ZeroVector)
