This fixes a UniqueID inconsistency in geometry collections that caused a crash in the past
- Full fix

#[fyi] Michael.Lentine, Benn.Gallagher, Brice.Criswell

#[robomerge] -ue5-main

diff --git a/Engine/Source/Runtime/Experimental/Chaos/Private/Chaos/PBDRigidClustering.cpp b/Engine/Source/Runtime/Experimental/Chaos/Private/Chaos/PBDRigidClustering.cpp
--- a/Engine/Source/Runtime/Experimental/Chaos/Private/Chaos/PBDRigidClustering.cpp
+++ b/Engine/Source/Runtime/Experimental/Chaos/Private/Chaos/PBDRigidClustering.cpp
@@ -416,14 +416,14 @@ namespace Chaos
 
 	DECLARE_CYCLE_STAT(TEXT("TPBDRigidClustering<>::CreateClusterParticle"), STAT_CreateClusterParticle, STATGROUP_Chaos);
 	template<class T_FPBDRigidsEvolution, class T_FPBDCollisionConstraint>
-	Chaos::FPBDRigidClusteredParticleHandle* TPBDRigidClustering<T_FPBDRigidsEvolution, T_FPBDCollisionConstraint>::CreateClusterParticle(const int32 ClusterGroupIndex, TArray<Chaos::FPBDRigidParticleHandle*>&& Children, const FClusterCreationParameters& Parameters, TSharedPtr<Chaos::FImplicitObject, ESPMode::ThreadSafe> ProxyGeometry, const FRigidTransform3* ForceMassOrientation)
+	Chaos::FPBDRigidClusteredParticleHandle* TPBDRigidClustering<T_FPBDRigidsEvolution, T_FPBDCollisionConstraint>::CreateClusterParticle(const int32 ClusterGroupIndex, TArray<Chaos::FPBDRigidParticleHandle*>&& Children, const FClusterCreationParameters& Parameters, TSharedPtr<Chaos::FImplicitObject, ESPMode::ThreadSafe> ProxyGeometry, const FRigidTransform3* ForceMassOrientation, const FUniqueIdx* ExistingIndex)
 	{
 		SCOPE_CYCLE_COUNTER(STAT_CreateClusterParticle);
 
 		Chaos::FPBDRigidClusteredParticleHandle* NewParticle = Parameters.ClusterParticleHandle;
 		if (!NewParticle)
 		{
-			NewParticle = MEvolution.CreateClusteredParticles(1)[0]; // calls Evolution.DirtyParticle()
+			NewParticle = MEvolution.CreateClusteredParticles(1, ExistingIndex)[0]; // calls Evolution.DirtyParticle()
 		}
 
 		// Must do this so that the constraint graph knows about this particle 
diff --git a/Engine/Source/Runtime/Experimental/Chaos/Private/PhysicsProxy/GeometryCollectionPhysicsProxy.cpp b/Engine/Source/Runtime/Experimental/Chaos/Private/PhysicsProxy/GeometryCollectionPhysicsProxy.cpp
--- a/Engine/Source/Runtime/Experimental/Chaos/Private/PhysicsProxy/GeometryCollectionPhysicsProxy.cpp
+++ b/Engine/Source/Runtime/Experimental/Chaos/Private/PhysicsProxy/GeometryCollectionPhysicsProxy.cpp
@@ -856,18 +856,19 @@ void FGeometryCollectionPhysicsProxy::InitializeBodiesPT(Chaos::FPBDRigidsSolver
 					Chaos::FClusterCreationParameters CreationParameters;
 					CreationParameters.ClusterParticleHandle = ClusterHandles.Num() ? ClusterHandles[ClusterHandlesIndex++] : nullptr;
 
-					Chaos::TPBDRigidClusteredParticleHandle<float, 3>* Handle = BuildClusters(TransformGroupIndex, RigidChildren, RigidChildrenTransformGroupIndex, CreationParameters);
+					// Hook the handle up with the GT particle
+					Chaos::FGeometryParticle* GTParticle = GTParticles[TransformGroupIndex].Get();
+
+					Chaos::FUniqueIdx ExistingIndex = GTParticle->UniqueIdx();
+					Chaos::TPBDRigidClusteredParticleHandle<float, 3>* Handle = BuildClusters(TransformGroupIndex, RigidChildren, RigidChildrenTransformGroupIndex, CreationParameters, &ExistingIndex);
+					Handle->GTGeometryParticle() = GTParticle;
 
 					int32 RigidChildrenIdx = 0;
 					for(const int32 ChildTransformIndex : RigidChildrenTransformGroupIndex)
 					{
 						SolverClusterID[ChildTransformIndex] = RigidChildren[RigidChildrenIdx++]->CastToClustered()->ClusterIds().Id;;
 					}
-					SolverClusterID[TransformGroupIndex] = Handle->ClusterIds().Id;
-
-					// Hook the handle up with the GT particle
-					Chaos::FGeometryParticle* GTParticle = GTParticles[TransformGroupIndex].Get();
-					Handle->GTGeometryParticle() = GTParticle;
+					SolverClusterID[TransformGroupIndex] = Handle->ClusterIds().Id;					
 					
 					// Cluster transform has been recalculated based on children - copy to the GT particle (not threadsafe - just testing)
 					GTParticle->SetX(Handle->X());
@@ -976,7 +977,8 @@ FGeometryCollectionPhysicsProxy::BuildClusters(
 	const uint32 CollectionClusterIndex, // TransformGroupIndex
 	TArray<Chaos::TPBDRigidParticleHandle<float,3>*>& ChildHandles,
 	const TArray<int32>& ChildTransformGroupIndices,
-	const Chaos::FClusterCreationParameters & ClusterParameters)
+	const Chaos::FClusterCreationParameters & ClusterParameters,
+	const Chaos::FUniqueIdx* ExistingIndex)
 {
 	SCOPE_CYCLE_COUNTER(STAT_BuildClusters);
 
@@ -1022,7 +1024,9 @@ FGeometryCollectionPhysicsProxy::BuildClusters(
 			MoveTemp(ChildHandlesCopy),
 			ClusterCreationParameters,
 			Implicits[CollectionClusterIndex], // union from children if null
-			&ParticleTM);
+			&ParticleTM,
+			ExistingIndex
+			);
 
 	if (ReportNoLevelsetCluster && 
 		Parent->DynamicGeometry())
diff --git a/Engine/Source/Runtime/Experimental/Chaos/Public/Chaos/PBDRigidClustering.h b/Engine/Source/Runtime/Experimental/Chaos/Public/Chaos/PBDRigidClustering.h
--- a/Engine/Source/Runtime/Experimental/Chaos/Public/Chaos/PBDRigidClustering.h
+++ b/Engine/Source/Runtime/Experimental/Chaos/Public/Chaos/PBDRigidClustering.h
@@ -8,6 +8,7 @@
 #include "Chaos/ExternalCollisionData.h"
 #include "Chaos/TriangleMesh.h"
 #include "Chaos/ClusterCreationParameters.h"
+#include "Chaos/GeometryParticlesfwd.h"
 #include "Framework/BufferedData.h"
 
 #define TODO_CONVERT_GEOMETRY_COLLECTION_PARTICLE_INDICES_TO_PARTICLE_POINTERS 0
@@ -81,7 +82,8 @@ public:
 		TArray<Chaos::FPBDRigidParticleHandle*>&& Children, 
 		const FClusterCreationParameters& Parameters = FClusterCreationParameters(),
 		TSharedPtr<Chaos::FImplicitObject, ESPMode::ThreadSafe> ProxyGeometry = nullptr,
-		const FRigidTransform3* ForceMassOrientation = nullptr);
+		const FRigidTransform3* ForceMassOrientation = nullptr,
+		const FUniqueIdx* ExistingIndex = nullptr);
 
 	/**
 	 *  CreateClusterParticleFromClusterChildren
diff --git a/Engine/Source/Runtime/Experimental/Chaos/Public/PhysicsProxy/GeometryCollectionPhysicsProxy.h b/Engine/Source/Runtime/Experimental/Chaos/Public/PhysicsProxy/GeometryCollectionPhysicsProxy.h
--- a/Engine/Source/Runtime/Experimental/Chaos/Public/PhysicsProxy/GeometryCollectionPhysicsProxy.h
+++ b/Engine/Source/Runtime/Experimental/Chaos/Public/PhysicsProxy/GeometryCollectionPhysicsProxy.h
@@ -16,6 +16,7 @@
 #include "Containers/Array.h"
 #include "PBDRigidsSolver.h"
 #include "Chaos/Defines.h"
+#include "Chaos/GeometryParticlesfwd.h"
 
 namespace Chaos
 {
@@ -252,7 +253,8 @@ protected:
 		const uint32 CollectionClusterIndex, 
 		TArray<Chaos::TPBDRigidParticleHandle<float, 3>*>& ChildHandles,
 		const TArray<int32>& ChildTransformGroupIndices,
-		const Chaos::FClusterCreationParameters & Parameters);
+		const Chaos::FClusterCreationParameters & Parameters,
+		const Chaos::FUniqueIdx* ExistingIndex);
 
 	/** 
 	 * Traverses the parents of \p TransformIndex in \p GeometryCollection, counting
