Make sure TPBDRigidClustering::ComputeStrainFromCollision only do computation for constraints that involve at least one clustered particle

diff --git a/Engine/Source/Runtime/Experimental/Chaos/Private/Chaos/PBDRigidClustering.cpp b/Engine/Source/Runtime/Experimental/Chaos/Private/Chaos/PBDRigidClustering.cpp
--- a/Engine/Source/Runtime/Experimental/Chaos/Private/Chaos/PBDRigidClustering.cpp
+++ b/Engine/Source/Runtime/Experimental/Chaos/Private/Chaos/PBDRigidClustering.cpp
@@ -1865,10 +1865,19 @@ namespace Chaos
 		for (const Chaos::FPBDCollisionConstraintHandle* ContactHandle : CollisionRule.GetConstConstraintHandles())
 		{
 			TVector<const FGeometryParticleHandle*, 2> ConstrainedParticles = ContactHandle->GetConstrainedParticles();
+			
+			// make sure we only compute things if one of the two particle is clustered
+			const FPBDRigidClusteredParticleHandle* ClusteredConstrainedParticles0 = ConstrainedParticles[0]->CastToClustered();
+			const FPBDRigidClusteredParticleHandle* ClusteredConstrainedParticles1 = ConstrainedParticles[1]->CastToClustered();
+			if (!ClusteredConstrainedParticles0 && !ClusteredConstrainedParticles1)
+			{
+				continue;
+			}
+
 			const FPBDRigidParticleHandle* Rigid0 = ConstrainedParticles[0]->CastToRigidParticle();
 			const FPBDRigidParticleHandle* Rigid1 = ConstrainedParticles[1]->CastToRigidParticle();
 
-			if(bUseContactSpeedForStrainThreshold)
+			if (bUseContactSpeedForStrainThreshold)
 			{
 				// Get dV between the two particles and project onto the normal to get the approach speed (take PreV as V is the new velocity post-solve)
 				const FVec3 V0 = Rigid0 ? Rigid0->PreV() : FVec3(0);
@@ -1877,18 +1886,18 @@ namespace Chaos
 				const FReal SpeedAlongNormal = FVec3::DotProduct(DeltaV, ContactHandle->GetContact().GetNormal());
 
 				// If we're not approaching at more than the min speed, reject the contact
-				if(SpeedAlongNormal > -MinContactSpeedForStrainEval && ContactHandle->GetAccumulatedImpulse().SizeSquared() > FReal(0))
+				if (SpeedAlongNormal > -MinContactSpeedForStrainEval && ContactHandle->GetAccumulatedImpulse().SizeSquared() > FReal(0))
 				{
 					continue;
 				}
 			}
-			else if(ContactHandle->GetAccumulatedImpulse().Size() < MinImpulseForStrainEval)
+			else if (ContactHandle->GetAccumulatedImpulse().Size() < MinImpulseForStrainEval)
 			{
 				continue;
 			}
 
 			auto ComputeStrainLambda = [&](
-				const FPBDRigidClusteredParticleHandle* Cluster, 
+				const FPBDRigidClusteredParticleHandle* Cluster,
 				const TArray<FPBDRigidParticleHandle*>& ParentToChildren)
 			{
 				const FRigidTransform3 WorldToClusterTM = FRigidTransform3(Cluster->P(), Cluster->Q());
@@ -1900,7 +1909,7 @@ namespace Chaos
 					const TArray<FPBDRigidParticleHandle*> Intersections = Cluster->ChildrenSpatial()->FindAllIntersectingChildren(ContactBox);
 					for (FPBDRigidParticleHandle* Child : Intersections)
 					{
-						if (TPBDRigidClusteredParticleHandle<FReal, 3>* ClusteredChild = Child->CastToClustered())
+						if (TPBDRigidClusteredParticleHandle<FReal, 3>*ClusteredChild = Child->CastToClustered())
 						{
 							const TUniquePtr<TMultiChildProxyData<FReal, 3>>& ProxyData = ClusteredChild->MultiChildProxyData();
 							const FPBDRigidParticleHandle* KeyChild = ProxyData ? ProxyData->KeyChild : nullptr;
@@ -1914,7 +1923,7 @@ namespace Chaos
 								ContactBoxProxy.Thicken(ClusterDistanceThreshold);
 								if (ClusteredChild->ChildrenSpatial())
 								{
-									const TArray<FPBDRigidParticleHandle*> SubIntersections = 
+									const TArray<FPBDRigidParticleHandle*> SubIntersections =
 										ClusteredChild->ChildrenSpatial()->FindAllIntersectingChildren(ContactBoxProxy);
 									for (FPBDRigidParticleHandle* SubChild : SubIntersections)
 									{
@@ -1934,14 +1943,14 @@ namespace Chaos
 				}
 			};
 
-			if (const TArray<FPBDRigidParticleHandle*>* ChildrenPtr = MParentToChildren.Find(ConstrainedParticles[0]->CastToClustered()))
+			if (const TArray<FPBDRigidParticleHandle*>* ChildrenPtr = MParentToChildren.Find(ClusteredConstrainedParticles0))
 			{
-				ComputeStrainLambda(ConstrainedParticles[0]->CastToClustered(), *ChildrenPtr);
+				ComputeStrainLambda(ClusteredConstrainedParticles0, *ChildrenPtr);
 			}
 
-			if (const TArray<FPBDRigidParticleHandle*>* ChildrenPtr = MParentToChildren.Find(ConstrainedParticles[1]->CastToClustered()))
+			if (const TArray<FPBDRigidParticleHandle*>* ChildrenPtr = MParentToChildren.Find(ClusteredConstrainedParticles1))
 			{
-				ComputeStrainLambda(ConstrainedParticles[1]->CastToClustered(), *ChildrenPtr);
+				ComputeStrainLambda(ClusteredConstrainedParticles1, *ChildrenPtr);
 			}
 
 			MCollisionImpulseArrayDirty = true;
