Chaos
-- Put a 2% pad on level set bounds to compensate for potential colume loss due to vertex collapsing in the underlying trimesh.

diff --git a/Engine/Source/Runtime/Experimental/Chaos/Private/Chaos/Levelset.cpp b/Engine/Source/Runtime/Experimental/Chaos/Private/Chaos/Levelset.cpp
--- a/Engine/Source/Runtime/Experimental/Chaos/Private/Chaos/Levelset.cpp
+++ b/Engine/Source/Runtime/Experimental/Chaos/Private/Chaos/Levelset.cpp
@@ -706,6 +706,10 @@ bool FLevelSet::ComputeDistancesNearZeroIsocontour(FErrorReporter& ErrorReporter
 		}
 	}
 
+	// Pad resulting bounds to compensate for potentially lost volume from vertex collapsing.
+	static const TVector<T, d> LevelSetBoundsPadding({ 1.02f, 1.02f, 1.02f });
+	MOriginalLocalBoundingBox.Scale(LevelSetBoundsPadding);
+
 	return true;
 }
 
