Chaos
-- Change to initialization of instance of TVector to fix ambiguity on non-windows platforms.

diff --git a/Engine/Source/Runtime/Experimental/Chaos/Private/Chaos/Levelset.cpp b/Engine/Source/Runtime/Experimental/Chaos/Private/Chaos/Levelset.cpp
--- a/Engine/Source/Runtime/Experimental/Chaos/Private/Chaos/Levelset.cpp
+++ b/Engine/Source/Runtime/Experimental/Chaos/Private/Chaos/Levelset.cpp
@@ -707,7 +707,7 @@ bool FLevelSet::ComputeDistancesNearZeroIsocontour(FErrorReporter& ErrorReporter
 	}
 
 	// Pad resulting bounds to compensate for potentially lost volume from vertex collapsing.
-	static const TVector<T, d> LevelSetBoundsPadding({ 1.02f, 1.02f, 1.02f });
+	static const TVector<T, d> LevelSetBoundsPadding(1.02f, 1.02f, 1.02f);
 	MOriginalLocalBoundingBox.Scale(LevelSetBoundsPadding);
 
 	return true;
