Chaos
-- removed templating in level set bounding box padding.

diff --git a/Engine/Source/Runtime/Experimental/Chaos/Private/Chaos/Levelset.cpp b/Engine/Source/Runtime/Experimental/Chaos/Private/Chaos/Levelset.cpp
--- a/Engine/Source/Runtime/Experimental/Chaos/Private/Chaos/Levelset.cpp
+++ b/Engine/Source/Runtime/Experimental/Chaos/Private/Chaos/Levelset.cpp
@@ -707,7 +707,7 @@ bool FLevelSet::ComputeDistancesNearZeroIsocontour(FErrorReporter& ErrorReporter
 	}
 
 	// Pad resulting bounds to compensate for potentially lost volume from vertex collapsing.
-	static const TVector<T, d> LevelSetBoundsPadding(1.02f, 1.02f, 1.02f);
+	static const FVector LevelSetBoundsPadding(1.02f);
 	MOriginalLocalBoundingBox.Scale(LevelSetBoundsPadding);
 
 	return true;
