Chaos - store CullDistance with Collision Constraint for use in Update
- fix early-out in collision Apply step
- add varying stiffness to Collision Apply
- fixed restitution for low stiffness

diff --git a/Engine/Source/Programs/HeadlessChaos/Private/GeometryCollection/GeometryCollectionTestCollisionResolution.cpp b/Engine/Source/Programs/HeadlessChaos/Private/GeometryCollection/GeometryCollectionTestCollisionResolution.cpp
--- a/Engine/Source/Programs/HeadlessChaos/Private/GeometryCollection/GeometryCollectionTestCollisionResolution.cpp
+++ b/Engine/Source/Programs/HeadlessChaos/Private/GeometryCollection/GeometryCollectionTestCollisionResolution.cpp
@@ -51,8 +51,9 @@ namespace GeometryCollectionTest
 			EXPECT_EQ(Collection->DynamicCollection->Simplicials[3], nullptr);
 			EXPECT_EQ(UnitTest.Solver->GetParticles().GetGeometryCollectionParticles().CollisionParticles(0), nullptr);
 
+			const FReal MaxRestingSeparation = -UnitTest.Solver->GetEvolution()->GetGravityForces().GetAcceleration().Z * UnitTest.Dt * UnitTest.Dt;	// PBD resting separation will be up to this
 			EXPECT_LT(FMath::Abs(Collection->RestCollection->Transform[0].GetTranslation().Z)-10.f, KINDA_SMALL_NUMBER);
-			EXPECT_LT(FMath::Abs(Collection->DynamicCollection->Transform[0].GetTranslation().Z - 1.0), 0.1);
+			EXPECT_LT(FMath::Abs(Collection->DynamicCollection->Transform[0].GetTranslation().Z - 1.0), 0.1 + MaxRestingSeparation);
 		}
 	}
 
diff --git a/Engine/Source/Programs/HeadlessChaos/Private/HeadlessChaosTestCollisions.cpp b/Engine/Source/Programs/HeadlessChaos/Private/HeadlessChaosTestCollisions.cpp
--- a/Engine/Source/Programs/HeadlessChaos/Private/HeadlessChaosTestCollisions.cpp
+++ b/Engine/Source/Programs/HeadlessChaos/Private/HeadlessChaosTestCollisions.cpp
@@ -56,7 +56,7 @@ namespace ChaosTest {
 		Box2->Q() = Box2->R();
 		Box2->AuxilaryValue(PhysicsMaterials) = MakeSerializable(PhysicsMaterial);
 
-		FPBDCollisionConstraintAccessor Collisions(Particles, Collided, PhysicsMaterials, PerParticlePhysicsMaterials, 1, 1, FReal(0));
+		FPBDCollisionConstraintAccessor Collisions(Particles, Collided, PhysicsMaterials, PerParticlePhysicsMaterials, 1, 1);
 		Collisions.ComputeConstraints(0.f);
 		EXPECT_EQ(Collisions.NumConstraints(), 1);
 
@@ -102,7 +102,7 @@ namespace ChaosTest {
 		Box2->Q() = Box2->R();
 		Box2->AuxilaryValue(PhysicsMaterials) = MakeSerializable(PhysicsMaterial);
 
-		FPBDCollisionConstraintAccessor Collisions(Particles, Collided, PhysicsMaterials, PerParticlePhysicsMaterials, 1, 1, FReal(0));
+		FPBDCollisionConstraintAccessor Collisions(Particles, Collided, PhysicsMaterials, PerParticlePhysicsMaterials, 1, 1);
 		Collisions.ComputeConstraints(0.f);
 		EXPECT_EQ(Collisions.NumConstraints(), 1);
 
@@ -142,7 +142,7 @@ namespace ChaosTest {
 
 		const FReal Dt = 1 / 24.;
 
-		FPBDCollisionConstraintAccessor Collisions(Particles, Collided, PhysicsMaterials, PerParticlePhysicsMaterials, 2, 5, FReal(0));
+		FPBDCollisionConstraintAccessor Collisions(Particles, Collided, PhysicsMaterials, PerParticlePhysicsMaterials, 2, 5);
 
 		Collisions.ComputeConstraints(Dt);
 		EXPECT_EQ(Collisions.NumConstraints(), 1);
@@ -219,7 +219,7 @@ namespace ChaosTest {
 
 		const FReal Dt = 1 / 24.;
 
-		FPBDCollisionConstraintAccessor Collisions(Particles, Collided, PhysicsMaterials, PerParticlePhysicsMaterials, 2, 5, FReal(0));
+		FPBDCollisionConstraintAccessor Collisions(Particles, Collided, PhysicsMaterials, PerParticlePhysicsMaterials, 2, 5);
 
 		Collisions.ComputeConstraints(Dt);
 		EXPECT_EQ(Collisions.NumConstraints(), 1);
@@ -292,7 +292,7 @@ namespace ChaosTest {
 
 		const FReal Dt = 1 / 24.;
 
-		FPBDCollisionConstraintAccessor Collisions(Particles, Collided, PhysicsMaterials, PerParticlePhysicsMaterials, 2, 5, FReal(0));
+		FPBDCollisionConstraintAccessor Collisions(Particles, Collided, PhysicsMaterials, PerParticlePhysicsMaterials, 2, 5);
 
 		Collisions.ComputeConstraints(Dt);
 		EXPECT_EQ(Collisions.NumConstraints(), 1);
@@ -363,7 +363,7 @@ namespace ChaosTest {
 
 		const FReal Dt = 1 / 24.;
 
-		FPBDCollisionConstraintAccessor Collisions(Particles, Collided, PhysicsMaterials, PerParticlePhysicsMaterials, 2, 5, FReal(0));
+		FPBDCollisionConstraintAccessor Collisions(Particles, Collided, PhysicsMaterials, PerParticlePhysicsMaterials, 2, 5);
 
 		Collisions.ComputeConstraints(Dt);
 		EXPECT_EQ(Collisions.NumConstraints(), 1);
@@ -440,7 +440,7 @@ namespace ChaosTest {
 
 		const FReal Dt = 1 / 24.;
 
-		FPBDCollisionConstraintAccessor Collisions(Particles, Collided, PhysicsMaterials, PerParticlePhysicsMaterials, 2, 5, FReal(0));
+		FPBDCollisionConstraintAccessor Collisions(Particles, Collided, PhysicsMaterials, PerParticlePhysicsMaterials, 2, 5);
 
 		Collisions.ComputeConstraints(Dt);
 		EXPECT_EQ(Collisions.NumConstraints(), 1);
@@ -503,12 +503,12 @@ namespace ChaosTest {
 
 		FReal Dt = 1 / 24.;
 
-		FPBDCollisionConstraintAccessor Collisions(Particles, Collided, PhysicsMaterials, PerParticlePhysicsMaterials, 1, 1, FReal(0));
+		FPBDCollisionConstraintAccessor Collisions(Particles, Collided, PhysicsMaterials, PerParticlePhysicsMaterials, 1, 1);
 		Collisions.ComputeConstraints(Dt);
 		EXPECT_EQ(Collisions.NumConstraints(), 1);
 
 		FCollisionConstraintBase & Constraint = Collisions.GetConstraint(0);
-		Collisions.Update(Constraint, 100.0f);
+		Collisions.Update(Constraint);
 
 		//Collisions.UpdateLevelsetConstraintGJK(Particles, Constraint);
 		//EXPECT_EQ(Constraint.ParticleIndex, 1);
diff --git a/Engine/Source/Programs/HeadlessChaos/Private/HeadlessChaosTestDetectCollision.cpp b/Engine/Source/Programs/HeadlessChaos/Private/HeadlessChaosTestDetectCollision.cpp
new file mode 100644
--- /dev/null
+++ b/Engine/Source/Programs/HeadlessChaos/Private/HeadlessChaosTestDetectCollision.cpp
@@ -0,0 +1,351 @@
+// Copyright Epic Games, Inc. All Rights Reserved.
+
+#include "HeadlessChaos.h"
+#include "HeadlessChaosCollisionConstraints.h"
+#include "HeadlessChaosTestUtility.h"
+#include "Chaos/Collision/CapsuleConvexContactPoint.h"
+#include "Chaos/Collision/ContactPoint.h"
+#include "Chaos/Collision/SphereConvexContactPoint.h"
+#include "Chaos/GJK.h"
+#include "Chaos/Sphere.h"
+#include "Chaos/Utilities.h"
+#include "Modules/ModuleManager.h"
+
+namespace ChaosTest 
+{
+	using namespace Chaos;
+
+	void CheckContactPoint(const FContactPoint& ContactPoint, const FVec3& ShapeContactPos0, const FVec3& ShapeContactPos1, const FVec3& ShapeContactNormal, const int32 NormalOwnerIndex, const FVec3 Normal, const FReal Phi)
+	{
+		const FReal DistanceTolerance = 1.e-3f;
+		const FReal NormalTolerance = 1.e-3f;
+
+		EXPECT_NEAR(ContactPoint.ShapeContactPoints[0].X, ShapeContactPos0.X, DistanceTolerance);
+		EXPECT_NEAR(ContactPoint.ShapeContactPoints[0].Y, ShapeContactPos0.Y, DistanceTolerance);
+		EXPECT_NEAR(ContactPoint.ShapeContactPoints[0].Z, ShapeContactPos0.Z, DistanceTolerance);
+		EXPECT_NEAR(ContactPoint.ShapeContactPoints[1].X, ShapeContactPos1.X, DistanceTolerance);
+		EXPECT_NEAR(ContactPoint.ShapeContactPoints[1].Y, ShapeContactPos1.Y, DistanceTolerance);
+		EXPECT_NEAR(ContactPoint.ShapeContactPoints[1].Z, ShapeContactPos1.Z, DistanceTolerance);
+		EXPECT_NEAR(ContactPoint.ShapeContactNormal.X, ShapeContactNormal.X, NormalTolerance);
+		EXPECT_NEAR(ContactPoint.ShapeContactNormal.Y, ShapeContactNormal.Y, NormalTolerance);
+		EXPECT_NEAR(ContactPoint.ShapeContactNormal.Z, ShapeContactNormal.Z, NormalTolerance);
+		EXPECT_EQ(ContactPoint.ContactNormalOwnerIndex, NormalOwnerIndex);
+		EXPECT_NEAR(ContactPoint.Normal.X, Normal.X, NormalTolerance);
+		EXPECT_NEAR(ContactPoint.Normal.Y, Normal.Y, NormalTolerance);
+		EXPECT_NEAR(ContactPoint.Normal.Z, Normal.Z, NormalTolerance);
+		EXPECT_NEAR(ContactPoint.Phi, Phi, DistanceTolerance);
+	}
+
+	//
+	//
+	// SPHERE - CONVEX TESTS
+	//
+	//
+
+	// Sphere-Convex(with margin) far separated when a face is the nearest feature
+	GTEST_TEST(DetectCollisionTests, TestSphereConvex_Separated_Face)
+	{
+		const FReal SphereRadius = 50.0f;
+		const FVec3 BoxSize = FVec3(1000, 1000, 50);
+		const FReal BoxMargin = 10.0f;
+		const FVec3 SpherePos = FVec3(0, 0, 70);
+		const FVec3 ConvexPos = FVec3(0, 0, -25);
+
+		const FImplicitSphere3 Sphere(FVec3(0,0,0), SphereRadius);
+		const FImplicitConvex3 Convex = CreateConvexBox(BoxSize, BoxMargin);
+
+		const FRigidTransform3 SphereTransform = FRigidTransform3(SpherePos, FRotation3::FromIdentity());
+		const FRigidTransform3 ConvexTransform = FRigidTransform3(ConvexPos, FRotation3::FromIdentity());
+
+		{
+			FContactPoint ContactPoint = SphereConvexContactPoint(Sphere, SphereTransform, Convex, ConvexTransform);
+			EXPECT_TRUE(ContactPoint.IsSet());
+			if (ContactPoint.IsSet())
+			{
+				CheckContactPoint(ContactPoint, FVec3(0, 0, -SphereRadius), FVec3(0, 0, 0.5f * BoxSize.Z), FVec3(0,0,1), 1, FVec3(0,0,1), SpherePos.Z - SphereRadius);
+			}
+		}
+		{
+			FContactPoint ContactPoint = ConvexSphereContactPoint(Convex, ConvexTransform, Sphere, SphereTransform);
+			EXPECT_TRUE(ContactPoint.IsSet());
+			if (ContactPoint.IsSet())
+			{
+				CheckContactPoint(ContactPoint, FVec3(0, 0, 0.5f * BoxSize.Z), FVec3(0, 0, -SphereRadius), FVec3(0, 0, -1), 0, FVec3(0, 0, -1), SpherePos.Z - SphereRadius);
+			}
+		}
+	}
+
+	// Sphere-Convex(with margin) far separated when an edge is the nearest feature.
+	GTEST_TEST(DetectCollisionTests, TestSphereConvex_Separated_Edge)
+	{
+		const FReal SphereRadius = 50.0f;
+		const FVec3 BoxSize = FVec3(1000, 1000, 50);
+		const FReal BoxMargin = 10.0f;
+		const FVec3 SpherePos = FVec3(600, 0, 70);
+		const FVec3 ConvexPos = FVec3(0, 0, -25);
+
+		const FImplicitSphere3 Sphere(FVec3(0, 0, 0), SphereRadius);
+		const FImplicitConvex3 Convex = CreateConvexBox(BoxSize, BoxMargin);
+
+		const FRigidTransform3 SphereTransform = FRigidTransform3(SpherePos, FRotation3::FromIdentity());
+		const FRigidTransform3 ConvexTransform = FRigidTransform3(ConvexPos, FRotation3::FromIdentity());
+
+		FContactPoint ContactPoint = SphereConvexContactPoint(Sphere, SphereTransform, Convex, ConvexTransform);
+		EXPECT_TRUE(ContactPoint.IsSet());
+		if (ContactPoint.IsSet())
+		{
+			const FVec3 SphereOffset = FVec3(SpherePos.X - 0.5f * BoxSize.X, 0, SpherePos.Z);
+			const FVec3 ExpectedNormal = SphereOffset.GetSafeNormal();
+			const FReal ExpectedDistance = SphereOffset.Size() - Sphere.GetRadius();
+
+			CheckContactPoint(ContactPoint, -SphereRadius * ExpectedNormal, 0.5f * FVec3(BoxSize.X, 0, BoxSize.Z), ExpectedNormal, 1, ExpectedNormal, ExpectedDistance);
+		}
+	}
+
+	// Sphere-Convex(with margin) far separated when a vertex is the nearest feature.
+	GTEST_TEST(DetectCollisionTests, TestSphereConvex_Separated_Vertex)
+	{
+		const FReal SphereRadius = 50.0f;
+		const FVec3 BoxSize = FVec3(1000, 1000, 50);
+		const FReal BoxMargin = 10.0f;
+		const FVec3 SpherePos = FVec3(600, 540, 70);
+		const FVec3 ConvexPos = FVec3(0, 0, -25);
+
+		const FImplicitSphere3 Sphere(FVec3(0, 0, 0), SphereRadius);
+		const FImplicitConvex3 Convex = CreateConvexBox(BoxSize, BoxMargin);
+
+		const FRigidTransform3 SphereTransform = FRigidTransform3(SpherePos, FRotation3::FromIdentity());
+		const FRigidTransform3 ConvexTransform = FRigidTransform3(ConvexPos, FRotation3::FromIdentity());
+
+		FContactPoint ContactPoint = SphereConvexContactPoint(Sphere, SphereTransform, Convex, ConvexTransform);
+		EXPECT_TRUE(ContactPoint.IsSet());
+		if (ContactPoint.IsSet())
+		{
+			const FVec3 SphereOffset = FVec3(SpherePos.X - 0.5f * BoxSize.X, SpherePos.Y - 0.5f * BoxSize.Y, SpherePos.Z);
+			const FVec3 ExpectedNormal = SphereOffset.GetSafeNormal();
+			const FReal ExpectedDistance = SphereOffset.Size() - Sphere.GetRadius();
+
+			CheckContactPoint(ContactPoint, -SphereRadius * ExpectedNormal, 0.5f * FVec3(BoxSize.X, BoxSize.Y, BoxSize.Z), ExpectedNormal, 1, ExpectedNormal, ExpectedDistance);
+		}
+	}
+
+	// Sphere-Convex(with margin) far separated when a vertex is the nearest feature.
+	GTEST_TEST(DetectCollisionTests, TestSphereConvex_Scaled_Separated_Vertex)
+	{
+		const FReal SphereRadius = 50.0f;
+		const FVec3 BoxSize = FVec3(1000, 1000, 50);
+		const FReal BoxMargin = 10.0f;
+		const FVec3 SpherePos = FVec3(600, 540, 70);
+		const FVec3 ConvexPos = FVec3(0, 0, -25);
+
+		const FImplicitSphere3 Sphere(FVec3(0, 0, 0), SphereRadius);
+		const TImplicitObjectScaled<FImplicitConvex3> Convex = CreateScaledConvexBox(FVec3(1,1,1), BoxSize, BoxMargin);
+
+		const FRigidTransform3 SphereTransform = FRigidTransform3(SpherePos, FRotation3::FromIdentity());
+		const FRigidTransform3 ConvexTransform = FRigidTransform3(ConvexPos, FRotation3::FromIdentity());
+
+		FContactPoint ContactPoint = SphereConvexContactPoint(Sphere, SphereTransform, Convex, ConvexTransform);
+		EXPECT_TRUE(ContactPoint.IsSet());
+		if (ContactPoint.IsSet())
+		{
+			const FVec3 SphereOffset = FVec3(SpherePos.X - 0.5f * BoxSize.X, SpherePos.Y - 0.5f * BoxSize.Y, SpherePos.Z);
+			const FVec3 ExpectedNormal = SphereOffset.GetSafeNormal();
+			const FReal ExpectedDistance = SphereOffset.Size() - Sphere.GetRadius();
+
+			CheckContactPoint(ContactPoint, -SphereRadius * ExpectedNormal, 0.5f * FVec3(BoxSize.X, BoxSize.Y, BoxSize.Z), ExpectedNormal, 1, ExpectedNormal, ExpectedDistance);
+		}
+	}
+
+	// Sphere-Convex(with margin) deep contact
+	GTEST_TEST(DetectCollisionTests, TestSphereConvex_DeepContact)
+	{
+		const FReal SphereRadius = 50.0f;
+		const FVec3 BoxSize = FVec3(1000, 1000, 50);
+		const FReal BoxMargin = 10.0f;
+		const FVec3 SpherePos = FVec3(400, 20, 10);
+		const FVec3 ConvexPos = FVec3(0, 0, -25);
+
+		const FImplicitSphere3 Sphere(FVec3(0, 0, 0), SphereRadius);
+		const FImplicitConvex3 Convex = CreateConvexBox(BoxSize, BoxMargin);
+
+		const FRigidTransform3 SphereTransform = FRigidTransform3(SpherePos, FRotation3::FromIdentity());
+		const FRigidTransform3 ConvexTransform = FRigidTransform3(ConvexPos, FRotation3::FromIdentity());
+
+		FContactPoint ContactPoint = SphereConvexContactPoint(Sphere, SphereTransform, Convex, ConvexTransform);
+		EXPECT_TRUE(ContactPoint.IsSet());
+		if (ContactPoint.IsSet())
+		{
+			CheckContactPoint(ContactPoint, FVec3(0, 0, -SphereRadius), FVec3(SpherePos.X, SpherePos.Y, 0.5f * BoxSize.Z), FVec3(0, 0, 1), 1, FVec3(0, 0, 1), SpherePos.Z - SphereRadius);
+		}
+	}
+
+	// Sphere-Convex(with margin) and touching core shapes
+	GTEST_TEST(DetectCollisionTests, TestSphereConvex_Touching_Face)
+	{
+		const FReal SphereRadius = 50.0f;
+		const FVec3 BoxSize = FVec3(1000, 1000, 50);
+		const FReal BoxMargin = 10.0f;
+		const FVec3 SpherePos = FVec3(0, 0, -50);
+		const FVec3 ConvexPos = FVec3(0, 0, -25);
+
+		const FImplicitSphere3 Sphere(FVec3(0, 0, 0), SphereRadius);
+		const FImplicitConvex3 Convex = CreateConvexBox(BoxSize, BoxMargin);
+
+		const FRigidTransform3 SphereTransform = FRigidTransform3(SpherePos, FRotation3::FromIdentity());
+		const FRigidTransform3 ConvexTransform = FRigidTransform3(ConvexPos, FRotation3::FromIdentity());
+
+		FContactPoint ContactPoint = SphereConvexContactPoint(Sphere, SphereTransform, Convex, ConvexTransform);
+		EXPECT_TRUE(ContactPoint.IsSet());
+		if (ContactPoint.IsSet())
+		{
+			CheckContactPoint(ContactPoint, FVec3(0, 0, SphereRadius), FVec3(SpherePos.X, SpherePos.Y, -0.5f * BoxSize.Z), FVec3(0, 0, -1), 1, FVec3(0, 0, -1), SpherePos.Z + SphereRadius - BoxSize.Z);
+		}
+	}
+
+	// Sphere-Convex(with margin) and core shapes separated by a small epsilon (less than GJK epsilon)
+	GTEST_TEST(DetectCollisionTests, TestSphereConvex_LessEpsilon_Face)
+	{
+		const FReal SphereRadius = 50.0f;
+		const FVec3 BoxSize = FVec3(1000, 1000, 50);
+		const FReal BoxMargin = 10.0f;
+		const FVec3 SpherePos = FVec3(0, 0, -50.0005f);		// Note: implicit knowledge of default epsilon of GJKDistance (1.e-3f)
+		const FVec3 ConvexPos = FVec3(0, 0, -25);
+
+		const FImplicitSphere3 Sphere(FVec3(0, 0, 0), SphereRadius);
+		const FImplicitConvex3 Convex = CreateConvexBox(BoxSize, BoxMargin);
+
+		const FRigidTransform3 SphereTransform = FRigidTransform3(SpherePos, FRotation3::FromIdentity());
+		const FRigidTransform3 ConvexTransform = FRigidTransform3(ConvexPos, FRotation3::FromIdentity());
+
+		FContactPoint ContactPoint = SphereConvexContactPoint(Sphere, SphereTransform, Convex, ConvexTransform);
+		EXPECT_TRUE(ContactPoint.IsSet());
+		if (ContactPoint.IsSet())
+		{
+			CheckContactPoint(ContactPoint, FVec3(0, 0, SphereRadius), FVec3(SpherePos.X, SpherePos.Y, -0.5f * BoxSize.Z), FVec3(0, 0, -1), 1, FVec3(0, 0, -1), SpherePos.Z + SphereRadius - BoxSize.Z);
+		}
+	}
+
+	// Sphere-Convex(with margin) and core shapes separated by a small epsilon (just greater than GJK epsilon)
+	GTEST_TEST(DetectCollisionTests, TestSphereConvex_GreaterEpsilon_Face)
+	{
+		const FReal SphereRadius = 50.0f;
+		const FVec3 BoxSize = FVec3(1000, 1000, 50);
+		const FReal BoxMargin = 10.0f;
+		const FVec3 SpherePos = FVec3(0, 0, -50.0015f);	// Note: implicit knowledge of default epsilon of GJKDistance (1.e-3f)
+		const FVec3 ConvexPos = FVec3(0, 0, -25);
+
+		const FImplicitSphere3 Sphere(FVec3(0, 0, 0), SphereRadius);
+		const FImplicitConvex3 Convex = CreateConvexBox(BoxSize, BoxMargin);
+
+		const FRigidTransform3 SphereTransform = FRigidTransform3(SpherePos, FRotation3::FromIdentity());
+		const FRigidTransform3 ConvexTransform = FRigidTransform3(ConvexPos, FRotation3::FromIdentity());
+
+		FContactPoint ContactPoint = SphereConvexContactPoint(Sphere, SphereTransform, Convex, ConvexTransform);
+		EXPECT_TRUE(ContactPoint.IsSet());
+		if (ContactPoint.IsSet())
+		{
+			CheckContactPoint(ContactPoint, FVec3(0, 0, SphereRadius), FVec3(SpherePos.X, SpherePos.Y, -0.5f * BoxSize.Z), FVec3(0, 0, -1), 1, FVec3(0, 0, -1), -BoxSize.Z - (SpherePos.Z + SphereRadius));
+		}
+	}
+
+
+	//
+	//
+	// CAPSULE - CONVEX TESTS
+	//
+	//
+
+	// Capsule-Convex far separated when a face is the nearest feature
+	GTEST_TEST(DetectCollisionTests, TestCapsuleConvex_Separated_Parallel_Face)
+	{
+		const FReal CapsuleRadius = 50.0f;
+		const FReal CapsuleSegmentLength = 100.0f;
+		const FVec3 BoxSize = FVec3(1000, 1000, 50);
+		const FReal BoxMargin = 10.0f;
+		const FVec3 CapsulePos = FVec3(0, 0, 70);
+		const FVec3 ConvexPos = FVec3(0, 0, -25);
+
+		const FImplicitCapsule3 Capsule(FVec3(-0.5f * CapsuleSegmentLength, 0.0f, 0.0f), FVec3(0.5f * CapsuleSegmentLength, 0.0f, 0.0f), CapsuleRadius);
+		const TImplicitObjectScaled<FImplicitConvex3> Convex = CreateScaledConvexBox(FVec3(1, 1, 1), BoxSize, BoxMargin);
+
+		const FRigidTransform3 CapsuleTransform = FRigidTransform3(CapsulePos, FRotation3::FromIdentity());
+		const FRigidTransform3 ConvexTransform = FRigidTransform3(ConvexPos, FRotation3::FromIdentity());
+
+		{
+			FContactPoint ContactPoint = CapsuleConvexContactPoint(Capsule, CapsuleTransform, Convex, ConvexTransform);
+			EXPECT_TRUE(ContactPoint.IsSet());
+			if (ContactPoint.IsSet())
+			{
+				CheckContactPoint(ContactPoint, FVec3(0, 0, -CapsuleRadius), FVec3(0, 0, 0.5f * BoxSize.Z), FVec3(0, 0, 1), 1, FVec3(0, 0, 1), CapsulePos.Z - CapsuleRadius);
+			}
+		}
+		{
+			FContactPoint ContactPoint = ConvexCapsuleContactPoint(Convex, ConvexTransform, Capsule, CapsuleTransform);
+			EXPECT_TRUE(ContactPoint.IsSet());
+			if (ContactPoint.IsSet())
+			{
+				CheckContactPoint(ContactPoint, FVec3(0, 0, 0.5f * BoxSize.Z), FVec3(0, 0, -CapsuleRadius), FVec3(0, 0, -1), 0, FVec3(0, 0, -1), CapsulePos.Z - CapsuleRadius);
+			}
+		}
+	}
+
+	// Capsule-Convex far separated when the nearest features are an edge on the convex and the middle of the capsule segment (capsule is rotated)
+	GTEST_TEST(DetectCollisionTests, TestCapsuleConvex_Separated_EdgeEdge)
+	{
+		const FReal CapsuleRadius = 50.0f;
+		const FReal CapsuleSegmentLength = 100.0f;
+		const FVec3 BoxSize = FVec3(1000, 1000, 50);
+		const FReal BoxMargin = 10.0f;
+		const FVec3 CapsulePos = FVec3(600, 0, 100);
+		const FVec3 ConvexPos = FVec3(0, 0, -25);
+
+		const FImplicitCapsule3 Capsule(FVec3(-0.5f * CapsuleSegmentLength, 0.0f, 0.0f), FVec3(0.5f * CapsuleSegmentLength, 0.0f, 0.0f), CapsuleRadius);
+		const TImplicitObjectScaled<FImplicitConvex3> Convex = CreateScaledConvexBox(FVec3(1, 1, 1), BoxSize, BoxMargin);
+
+		const FRigidTransform3 CapsuleTransform = FRigidTransform3(CapsulePos, FRotation3::FromAxisAngle(FVec3(0,1,0), FMath::DegreesToRadians(45)));
+		const FRigidTransform3 ConvexTransform = FRigidTransform3(ConvexPos, FRotation3::FromIdentity());
+
+		{
+			FContactPoint ContactPoint = CapsuleConvexContactPoint(Capsule, CapsuleTransform, Convex, ConvexTransform);
+			EXPECT_TRUE(ContactPoint.IsSet());
+			if (ContactPoint.IsSet())
+			{
+				const FVec3 CapsuleOffset = CapsulePos - FVec3(0.5f * BoxSize.X, 0.0f, 0.0f);
+				const FVec3 ExpectedNormal = CapsuleOffset.GetSafeNormal();
+				const FReal ExpectedPhi = CapsuleOffset.Size() - CapsuleRadius;
+				CheckContactPoint(ContactPoint, FVec3(0, 0, -CapsuleRadius), FVec3(0.5f * BoxSize.X, 0, 0.5f * BoxSize.Z), ExpectedNormal, 1, ExpectedNormal, ExpectedPhi);
+			}
+		}
+	}
+
+
+	// Capsule-Convex deep contact when the nearest features are an edge on the convex and the middle of the capsule segment (capsule is rotated)
+	GTEST_TEST(DetectCollisionTests, TestCapsuleConvex_DeepContact_EdgeEdge)
+	{
+		const FReal CapsuleRadius = 50.0f;
+		const FReal CapsuleSegmentLength = 100.0f;
+		const FVec3 BoxSize = FVec3(1000, 1000, 50);
+		const FReal BoxMargin = 10.0f;
+		const FVec3 CapsulePos = FVec3(480, 0, -20);
+		const FVec3 ConvexPos = FVec3(0, 0, -25);
+
+		const FImplicitCapsule3 Capsule(FVec3(-0.5f * CapsuleSegmentLength, 0.0f, 0.0f), FVec3(0.5f * CapsuleSegmentLength, 0.0f, 0.0f), CapsuleRadius);
+		const TImplicitObjectScaled<FImplicitConvex3> Convex = CreateScaledConvexBox(FVec3(1, 1, 1), BoxSize, BoxMargin);
+
+		const FRigidTransform3 CapsuleTransform = FRigidTransform3(CapsulePos, FRotation3::FromAxisAngle(FVec3(0, 1, 0), FMath::DegreesToRadians(45)));
+		const FRigidTransform3 ConvexTransform = FRigidTransform3(ConvexPos, FRotation3::FromIdentity());
+
+		{
+			FContactPoint ContactPoint = CapsuleConvexContactPoint(Capsule, CapsuleTransform, Convex, ConvexTransform);
+			EXPECT_TRUE(ContactPoint.IsSet());
+			if (ContactPoint.IsSet())
+			{
+				const FVec3 CapsuleOffset = CapsulePos - FVec3(0.5f * BoxSize.X, 0.0f, 0.0f);
+				const FVec3 ExpectedNormal = -CapsuleOffset.GetSafeNormal();
+				const FReal ExpectedPhi = -(CapsuleOffset.Size() + CapsuleRadius);
+				CheckContactPoint(ContactPoint, FVec3(0, 0, -CapsuleRadius), FVec3(0.5f * BoxSize.X, 0, 0.5f * BoxSize.Z), ExpectedNormal, 1, ExpectedNormal, ExpectedPhi);
+			}
+		}
+	}
+
+}
\ No newline at end of file
diff --git a/Engine/Source/Programs/HeadlessChaos/Private/HeadlessChaosTestOneShotManifolds.cpp b/Engine/Source/Programs/HeadlessChaos/Private/HeadlessChaosTestOneShotManifolds.cpp
--- a/Engine/Source/Programs/HeadlessChaos/Private/HeadlessChaosTestOneShotManifolds.cpp
+++ b/Engine/Source/Programs/HeadlessChaos/Private/HeadlessChaosTestOneShotManifolds.cpp
@@ -19,7 +19,6 @@ namespace Chaos
 			const Chaos::FRigidTransform3& Box1Transform, //world
 			const Chaos::FImplicitBox3& Box2,
 			const Chaos::FRigidTransform3& Box2Transform, //world
-			const Chaos::FReal CullDistance,
 			const Chaos::FReal Dt,
 			Chaos::FRigidBodyPointContactConstraint& Constraint);
 
@@ -29,7 +28,6 @@ namespace Chaos
 			const FRigidTransform3& Convex1Transform, //world
 			const ConvexImplicitType2& Implicit2,
 			const FRigidTransform3& Convex2Transform, //world
-			const FReal CullDistance,
 			const FReal Dt,
 			FRigidBodyPointContactConstraint& Constraint);
 	}
@@ -53,9 +51,7 @@ namespace ChaosTest
 			FRigidTransform3 Box2Transform = FRigidTransform3(FVec3(0.0f, 0.0f, 0.0f), FRotation3::FromElements(0.0f, 0.0f, 0.0f, 1.0f));
 
 			FRigidBodyPointContactConstraint Constraint;
-			FReal CullingDistance = 100.0f;
-
-			Collisions::ConstructBoxBoxOneShotManifold(Box1, Box1Transform, Box2, Box2Transform, CullingDistance, Dt, Constraint);
+			Collisions::ConstructBoxBoxOneShotManifold(Box1, Box1Transform, Box2, Box2Transform, Dt, Constraint);
 
 			// Result should give a negative phi on all contacts
 			// Phi direction may be in a random face direction
@@ -77,9 +73,7 @@ namespace ChaosTest
 			FRigidTransform3 Box2Transform = FRigidTransform3(FVec3(0.0f, 0.0f, 0.0f), FRotation3::FromElements(0.0f, 0.0f, 0.0f, 1.0f));
 
 			FRigidBodyPointContactConstraint Constraint;
-			FReal CullingDistance = 100.0f;
-
-			Collisions::ConstructBoxBoxOneShotManifold(Box1, Box1Transform, Box2, Box2Transform, CullingDistance, Dt, Constraint);
+			Collisions::ConstructBoxBoxOneShotManifold(Box1, Box1Transform, Box2, Box2Transform, Dt, Constraint);
 			int ContactCount = Constraint.GetManifoldPoints().Num();
 			EXPECT_EQ(ContactCount, 4);
 			for (int ConstraintIndex = 0; ConstraintIndex < ContactCount; ConstraintIndex++)
@@ -99,9 +93,7 @@ namespace ChaosTest
 			FRigidTransform3 Box2Transform = FRigidTransform3(FVec3(0.0f, 0.0f, 0.0f), FRotation3::FromAxisAngle(FVec3(0.0f, 1.0f, 0.0f), 0.1f));
 
 			FRigidBodyPointContactConstraint Constraint;
-			FReal CullingDistance = 100.0f;
-
-			Collisions::ConstructBoxBoxOneShotManifold(Box1, Box1Transform, Box2, Box2Transform, CullingDistance, Dt, Constraint);
+			Collisions::ConstructBoxBoxOneShotManifold(Box1, Box1Transform, Box2, Box2Transform, Dt, Constraint);
 			int ContactCount = Constraint.GetManifoldPoints().Num();
 			EXPECT_EQ(ContactCount, 4);
 			for (int ConstraintIndex = 0; ConstraintIndex < ContactCount; ConstraintIndex++)
@@ -123,9 +115,7 @@ namespace ChaosTest
 			FRigidTransform3 Box2Transform = FRigidTransform3(FVec3(0.0f, 0.0f, 0.0f) - OffsetBox2, FRotation3::FromElements(0.0f, 0.0f, 0.0f, 1.0f));
 
 			FRigidBodyPointContactConstraint Constraint;
-			FReal CullingDistance = 100.0f;
-
-			Collisions::ConstructBoxBoxOneShotManifold(Box1, Box1Transform, Box2, Box2Transform, CullingDistance, Dt, Constraint);
+			Collisions::ConstructBoxBoxOneShotManifold(Box1, Box1Transform, Box2, Box2Transform, Dt, Constraint);
 			int ContactCount = Constraint.GetManifoldPoints().Num();
 			EXPECT_EQ(ContactCount, 4);
 			for (int ConstraintIndex = 0; ConstraintIndex < ContactCount; ConstraintIndex++)
@@ -160,9 +150,7 @@ namespace ChaosTest
 			FRigidTransform3 Box2Transform = TranslationTransform2 * RotationTransform;
 
 			FRigidBodyPointContactConstraint Constraint;
-			FReal CullingDistance = 100.0f;
-
-			Collisions::ConstructBoxBoxOneShotManifold(Box1, Box1Transform, Box2, Box2Transform, CullingDistance, Dt, Constraint);
+			Collisions::ConstructBoxBoxOneShotManifold(Box1, Box1Transform, Box2, Box2Transform, Dt, Constraint);
 			int ContactCount = Constraint.GetManifoldPoints().Num();
 			EXPECT_EQ(ContactCount, 4);
 			for (int ConstraintIndex = 0; ConstraintIndex < ContactCount; ConstraintIndex++)
@@ -195,9 +183,7 @@ namespace ChaosTest
 			TBox<FReal, 3> Box2(FVec3(-100.0f, -100, -100.0f) + OffsetBox2, FVec3(100.0f, 100.0f, 100.0f) + OffsetBox2);
 
 			FRigidBodyPointContactConstraint Constraint;
-			FReal CullingDistance = 100.0f;
-
-			Collisions::ConstructBoxBoxOneShotManifold(Box1, Box1Transform, Box2, Box2Transform, CullingDistance, Dt, Constraint);
+			Collisions::ConstructBoxBoxOneShotManifold(Box1, Box1Transform, Box2, Box2Transform, Dt, Constraint);
 			int ContactCount = Constraint.GetManifoldPoints().Num();
 			for (int ConstraintIndex = 0; ConstraintIndex < ContactCount; ConstraintIndex++)
 			{
@@ -215,9 +201,7 @@ namespace ChaosTest
 			FRigidTransform3 Box2Transform = FRigidTransform3(FVec3(0.0f, 0.0f, 0.0f), FRotation3::FromElements(0.0f, 0.0f, 0.0f, 1.0f));
 
 			FRigidBodyPointContactConstraint Constraint;
-			FReal CullingDistance = 100.0f;
-
-			Collisions::ConstructBoxBoxOneShotManifold(Box1, Box1Transform, Box2, Box2Transform, CullingDistance, Dt, Constraint);
+			Collisions::ConstructBoxBoxOneShotManifold(Box1, Box1Transform, Box2, Box2Transform, Dt, Constraint);
 			int ContactCount = Constraint.GetManifoldPoints().Num();
 			EXPECT_EQ(ContactCount, 4);  
 			for (int ConstraintIndex = 0; ConstraintIndex < ContactCount; ConstraintIndex++)
@@ -234,7 +218,6 @@ namespace ChaosTest
 	GTEST_TEST(OneShotManifoldTests, TestConvexMarginEdgeEdge)
 	{
 		FReal Dt = 1 / 30.0f;
-		FReal CullingDistance = 100.0f;
 		FReal HalfSize = 100.0f;
 		FReal Margin = 0.2f * HalfSize;
 		float ExpectedPhi = 0.0f;
@@ -260,7 +243,7 @@ namespace ChaosTest
 
 		FRigidBodyPointContactConstraint Constraint;
 
-		Collisions::ConstructConvexConvexOneShotManifold(Convex, Convex1Transform, Convex, Convex2Transform, CullingDistance, Dt, Constraint);
+		Collisions::ConstructConvexConvexOneShotManifold(Convex, Convex1Transform, Convex, Convex2Transform, Dt, Constraint);
 		int ContactCount = Constraint.GetManifoldPoints().Num();
 		EXPECT_EQ(ContactCount, 1);
 		if (ContactCount > 0)
diff --git a/Engine/Source/Programs/HeadlessChaos/Private/HeadlessChaosTestRewind.cpp b/Engine/Source/Programs/HeadlessChaos/Private/HeadlessChaosTestRewind.cpp
--- a/Engine/Source/Programs/HeadlessChaos/Private/HeadlessChaosTestRewind.cpp
+++ b/Engine/Source/Programs/HeadlessChaos/Private/HeadlessChaosTestRewind.cpp
@@ -2225,7 +2225,7 @@ namespace ChaosTest {
 				}
 
 				TickSolverHelper(Solver);
-				EXPECT_LE(Dynamic.X()[2], 10);
+				EXPECT_LE(Dynamic.X()[2], 10 + KINDA_SMALL_NUMBER);
 
 #if REWIND_DESYNC
 				//kinematic desync will be known at end of frame because the simulation doesn't write results (so we know right away it's a desync)
diff --git a/Engine/Source/Programs/HeadlessChaos/Private/HeadlessChaosTestShapePadding.cpp b/Engine/Source/Programs/HeadlessChaos/Private/HeadlessChaosTestShapePadding.cpp
--- a/Engine/Source/Programs/HeadlessChaos/Private/HeadlessChaosTestShapePadding.cpp
+++ b/Engine/Source/Programs/HeadlessChaos/Private/HeadlessChaosTestShapePadding.cpp
@@ -90,10 +90,11 @@ namespace ChaosTest {
 			Box1->Geometry().Get(),
 			nullptr,
 			FRigidTransform3(),
+			FLT_MAX,
 			EContactShapesType::BoxBox, true);
 
 		// Detect collisions
-		Collisions::Update(Constraint, Delta.Size(), 1/30.0f);
+		Collisions::Update(Constraint, 1/30.0f);
 
 		EXPECT_NEAR(Constraint.Manifold.Phi, ExpectedPhi, Tolerance);
 		EXPECT_NEAR(Constraint.Manifold.Normal.X, ExpectedNormal.X, Tolerance);
@@ -194,10 +195,11 @@ namespace ChaosTest {
 			Box1->Geometry().Get(),
 			nullptr,
 			FRigidTransform3(),
+			FLT_MAX,
 			EContactShapesType::GenericConvexConvex, true);
 
 		// Detect collisions
-		Collisions::Update(Constraint, Delta.Size(), 1 / 30.0f);
+		Collisions::Update(Constraint, 1 / 30.0f);
 
 		EXPECT_NEAR(Constraint.Manifold.Phi, ExpectedPhi, Tolerance);
 		EXPECT_NEAR(Constraint.Manifold.Normal.X, ExpectedNormal.X, Tolerance);
diff --git a/Engine/Source/Programs/HeadlessChaos/Public/HeadlessChaosCollisionConstraints.h b/Engine/Source/Programs/HeadlessChaos/Public/HeadlessChaosCollisionConstraints.h
--- a/Engine/Source/Programs/HeadlessChaos/Public/HeadlessChaosCollisionConstraints.h
+++ b/Engine/Source/Programs/HeadlessChaos/Public/HeadlessChaosCollisionConstraints.h
@@ -41,10 +41,10 @@ public:
 	{}
 
 	FPBDCollisionConstraintAccessor(const FPBDRigidsSOAs& InParticles, TArrayCollectionArray<bool>& Collided, const TArrayCollectionArray<TSerializablePtr<FChaosPhysicsMaterial>>& PerParticleMaterials, const TArrayCollectionArray<TUniquePtr<FChaosPhysicsMaterial>>& PerParticleUniqueMaterials,
-	const int32 PushOutIterations, const int32 PushOutPairIterations, const FReal Thickness) 
+	const int32 PushOutIterations, const int32 PushOutPairIterations) 
 		: SpatialAcceleration(InParticles.GetNonDisabledView())
 		, BroadPhase(InParticles, (FReal)1, (FReal)0, (FReal)0)
-		, CollisionConstraints(InParticles, Collided, PerParticleMaterials, PerParticleUniqueMaterials, 1, 1, Thickness)
+		, CollisionConstraints(InParticles, Collided, PerParticleMaterials, PerParticleUniqueMaterials, 1, 1)
 		, CollisionDetector(BroadPhase, NarrowPhase, CollisionConstraints)
 	{}
 
@@ -59,11 +59,11 @@ public:
 		CollisionDetector.DetectCollisions(Dt);
 	}
 
-	void Update(FCollisionConstraintBase& Constraint, FReal CullDistance = FReal(0) )
+	void Update(FCollisionConstraintBase& Constraint)
 	{
 		if (Constraint.GetType() == FPointContactConstraint::StaticType())
 		{
-			Collisions::Update(*Constraint.As<FPointContactConstraint>(), CullDistance, 1/30.0f);
+			Collisions::Update(*Constraint.As<FPointContactConstraint>(), 1/30.0f);
 		}
 	}
 
@@ -73,7 +73,7 @@ public:
 		FRigidTransform3 WorldTransform0 = Constraint.ImplicitTransform[0] * Collisions::GetTransform(Constraint.Particle[0]);
 		FRigidTransform3 WorldTransform1 = Constraint.ImplicitTransform[1] * Collisions::GetTransform(Constraint.Particle[1]);
 
-		Collisions::UpdateLevelsetLevelsetConstraint<ECollisionUpdateType::Deepest>(WorldTransform0, WorldTransform1, FReal(0), FReal(1 / 30.0f), Constraint);
+		Collisions::UpdateLevelsetLevelsetConstraint<ECollisionUpdateType::Deepest>(WorldTransform0, WorldTransform1, FReal(1 / 30.0f), Constraint);
 	}
 
 	int32 NumConstraints() const
diff --git a/Engine/Source/Programs/HeadlessChaos/Public/HeadlessChaosTestUtility.h b/Engine/Source/Programs/HeadlessChaos/Public/HeadlessChaosTestUtility.h
--- a/Engine/Source/Programs/HeadlessChaos/Public/HeadlessChaosTestUtility.h
+++ b/Engine/Source/Programs/HeadlessChaos/Public/HeadlessChaosTestUtility.h
@@ -147,9 +147,10 @@ namespace ChaosTest {
 	{
 		// Settings used for unit tests
 		const float CullDistance = 0.0f;
-		Evolution.GetCollisionConstraints().SetCullDistance(CullDistance);
 		Evolution.GetBroadPhase().SetCullDistance(CullDistance);
 		Evolution.GetBroadPhase().SetBoundsThickness(CullDistance);
+		Evolution.GetCollisionDetector().GetNarrowPhase().GetContext().bDeferUpdate = false;
+		Evolution.GetCollisionDetector().GetNarrowPhase().GetContext().bAllowManifolds = true;
 	}
 
 	template<typename T_SOLVER>
diff --git a/Engine/Source/Runtime/Engine/Private/PhysicsEngine/ImmediatePhysics/ImmediatePhysicsChaos/ImmediatePhysicsSimulation_Chaos.cpp b/Engine/Source/Runtime/Engine/Private/PhysicsEngine/ImmediatePhysics/ImmediatePhysicsChaos/ImmediatePhysicsSimulation_Chaos.cpp
--- a/Engine/Source/Runtime/Engine/Private/PhysicsEngine/ImmediatePhysics/ImmediatePhysicsChaos/ImmediatePhysicsSimulation_Chaos.cpp
+++ b/Engine/Source/Runtime/Engine/Private/PhysicsEngine/ImmediatePhysics/ImmediatePhysicsChaos/ImmediatePhysicsSimulation_Chaos.cpp
@@ -789,7 +789,6 @@ namespace ImmediatePhysics_Chaos
 			Implementation->Collisions.SetCollisionsEnabled(ChaosImmediate_Collision_Enabled != 0);
 			Implementation->CollisionsRule.SetPriority(ChaosImmediate_Collision_Priority);
 
-			Implementation->Collisions.SetCullDistance(ChaosImmediate_Collision_CullDistance);
 			Implementation->BroadPhase.SetCullDustance(ChaosImmediate_Collision_CullDistance);
 			Implementation->Evolution.SetBoundsExtension(ChaosImmediate_Evolution_BoundsExtension);
 
diff --git a/Engine/Source/Runtime/Experimental/Chaos/Private/Chaos/Collision/CapsuleConvexContactPoint.cpp b/Engine/Source/Runtime/Experimental/Chaos/Private/Chaos/Collision/CapsuleConvexContactPoint.cpp
--- a/Engine/Source/Runtime/Experimental/Chaos/Private/Chaos/Collision/CapsuleConvexContactPoint.cpp
+++ b/Engine/Source/Runtime/Experimental/Chaos/Private/Chaos/Collision/CapsuleConvexContactPoint.cpp
@@ -12,7 +12,7 @@ namespace Chaos
 	// Use GJK (point to convex) to calculate separation.
 	// Fall back to plane testing if penetrating by more than Radius.
 	template<typename T_CONVEX>
-	FContactPoint ConvexCapsuleContactPointImpl(const T_CONVEX& Convex, const FRigidTransform3& ConvexTransform, const FImplicitCapsule3& Capsule, const FRigidTransform3& CapsuleTransform, FReal CullDistance)
+	FContactPoint ConvexCapsuleContactPointImpl(const T_CONVEX& Convex, const FRigidTransform3& ConvexTransform, const FImplicitCapsule3& Capsule, const FRigidTransform3& CapsuleTransform)
 	{
 		FContactPoint ContactPoint;
 		if (Convex.NumPlanes() > 0)
@@ -50,60 +50,60 @@ namespace Chaos
 	}
 
 	template<typename T_CONVEX>
-	FContactPoint CapsuleConvexContactPointImpl(const FImplicitCapsule3& Capsule, const FRigidTransform3& CapsuleTransform, const T_CONVEX& Convex, const FRigidTransform3& ConvexTransform, FReal CullDistance)
+	FContactPoint CapsuleConvexContactPointImpl(const FImplicitCapsule3& Capsule, const FRigidTransform3& CapsuleTransform, const T_CONVEX& Convex, const FRigidTransform3& ConvexTransform)
 	{
-		return ConvexCapsuleContactPointImpl(Convex, ConvexTransform, Capsule, CapsuleTransform, CullDistance).SwapShapes();
+		return ConvexCapsuleContactPointImpl(Convex, ConvexTransform, Capsule, CapsuleTransform).SwapShapes();
 	}
 
 
 
-	FContactPoint CapsuleConvexContactPoint(const FImplicitCapsule3& Capsule, const FRigidTransform3& CapsuleTransform, const FImplicitConvex3& Convex, const FRigidTransform3& ConvexTransform, FReal CullDistance)
+	FContactPoint CapsuleConvexContactPoint(const FImplicitCapsule3& Capsule, const FRigidTransform3& CapsuleTransform, const FImplicitConvex3& Convex, const FRigidTransform3& ConvexTransform)
 	{
-		return CapsuleConvexContactPointImpl(Capsule, CapsuleTransform, Convex, ConvexTransform, CullDistance);
+		return CapsuleConvexContactPointImpl(Capsule, CapsuleTransform, Convex, ConvexTransform);
 	}
 
-	FContactPoint CapsuleConvexContactPoint(const FImplicitCapsule3& Capsule, const FRigidTransform3& CapsuleTransform, const TImplicitObjectInstanced<FImplicitConvex3>& Convex, const FRigidTransform3& ConvexTransform, FReal CullDistance)
+	FContactPoint CapsuleConvexContactPoint(const FImplicitCapsule3& Capsule, const FRigidTransform3& CapsuleTransform, const TImplicitObjectInstanced<FImplicitConvex3>& Convex, const FRigidTransform3& ConvexTransform)
 	{
-		return CapsuleConvexContactPointImpl(Capsule, CapsuleTransform, Convex, ConvexTransform, CullDistance);
+		return CapsuleConvexContactPointImpl(Capsule, CapsuleTransform, Convex, ConvexTransform);
 	}
 
-	FContactPoint CapsuleConvexContactPoint(const FImplicitCapsule3& Capsule, const FRigidTransform3& CapsuleTransform, const TImplicitObjectScaled<FImplicitConvex3>& Convex, const FRigidTransform3& ConvexTransform, FReal CullDistance)
+	FContactPoint CapsuleConvexContactPoint(const FImplicitCapsule3& Capsule, const FRigidTransform3& CapsuleTransform, const TImplicitObjectScaled<FImplicitConvex3>& Convex, const FRigidTransform3& ConvexTransform)
 	{
-		return CapsuleConvexContactPointImpl(Capsule, CapsuleTransform, Convex, ConvexTransform, CullDistance);
+		return CapsuleConvexContactPointImpl(Capsule, CapsuleTransform, Convex, ConvexTransform);
 	}
 
 
 
-	FContactPoint ConvexCapsuleContactPoint(const FImplicitConvex3& Convex, const FRigidTransform3& ConvexTransform, const FImplicitCapsule3& Capsule, const FRigidTransform3& CapsuleTransform, FReal CullDistance)
+	FContactPoint ConvexCapsuleContactPoint(const FImplicitConvex3& Convex, const FRigidTransform3& ConvexTransform, const FImplicitCapsule3& Capsule, const FRigidTransform3& CapsuleTransform)
 	{
-		return ConvexCapsuleContactPointImpl(Convex, ConvexTransform, Capsule, CapsuleTransform, CullDistance);
+		return ConvexCapsuleContactPointImpl(Convex, ConvexTransform, Capsule, CapsuleTransform);
 	}
 
-	FContactPoint ConvexCapsuleContactPoint(const TImplicitObjectInstanced<FImplicitConvex3>& Convex, const FRigidTransform3& ConvexTransform, const FImplicitCapsule3& Capsule, const FRigidTransform3& CapsuleTransform, FReal CullDistance)
+	FContactPoint ConvexCapsuleContactPoint(const TImplicitObjectInstanced<FImplicitConvex3>& Convex, const FRigidTransform3& ConvexTransform, const FImplicitCapsule3& Capsule, const FRigidTransform3& CapsuleTransform)
 	{
-		return ConvexCapsuleContactPointImpl(Convex, ConvexTransform, Capsule, CapsuleTransform, CullDistance);
+		return ConvexCapsuleContactPointImpl(Convex, ConvexTransform, Capsule, CapsuleTransform);
 	}
 
-	FContactPoint ConvexCapsuleContactPoint(const TImplicitObjectScaled<FImplicitConvex3>& Convex, const FRigidTransform3& ConvexTransform, const FImplicitCapsule3& Capsule, const FRigidTransform3& CapsuleTransform, FReal CullDistance)
+	FContactPoint ConvexCapsuleContactPoint(const TImplicitObjectScaled<FImplicitConvex3>& Convex, const FRigidTransform3& ConvexTransform, const FImplicitCapsule3& Capsule, const FRigidTransform3& CapsuleTransform)
 	{
-		return ConvexCapsuleContactPointImpl(Convex, ConvexTransform, Capsule, CapsuleTransform, CullDistance);
+		return ConvexCapsuleContactPointImpl(Convex, ConvexTransform, Capsule, CapsuleTransform);
 	}
 
 
 
-	FContactPoint CapsuleConvexContactPoint(const FImplicitCapsule3& Capsule, const FRigidTransform3& CapsuleTransform, const FImplicitObject& Object, const FRigidTransform3& ConvexTransform, FReal CullDistance)
+	FContactPoint CapsuleConvexContactPoint(const FImplicitCapsule3& Capsule, const FRigidTransform3& CapsuleTransform, const FImplicitObject& Object, const FRigidTransform3& ConvexTransform)
 	{
 		if (const TImplicitObjectInstanced<FImplicitConvex3>* InstancedConvex = Object.template GetObject<const TImplicitObjectInstanced<FImplicitConvex3>>())
 		{
-			return CapsuleConvexContactPoint(Capsule, CapsuleTransform, *InstancedConvex, ConvexTransform, CullDistance);
+			return CapsuleConvexContactPoint(Capsule, CapsuleTransform, *InstancedConvex, ConvexTransform);
 		}
 		else if (const TImplicitObjectScaled<FImplicitConvex3>* ScaledConvex = Object.template GetObject<const TImplicitObjectScaled<FImplicitConvex3>>())
 		{
-			return CapsuleConvexContactPoint(Capsule, CapsuleTransform, *ScaledConvex, ConvexTransform, CullDistance);
+			return CapsuleConvexContactPoint(Capsule, CapsuleTransform, *ScaledConvex, ConvexTransform);
 		}
 		else if (const FImplicitConvex3* Convex = Object.template GetObject<const FImplicitConvex3>())
 		{
-			return CapsuleConvexContactPoint(Capsule, CapsuleTransform, *Convex, ConvexTransform, CullDistance);
+			return CapsuleConvexContactPoint(Capsule, CapsuleTransform, *Convex, ConvexTransform);
 		}
 		return FContactPoint();
 	}
diff --git a/Engine/Source/Runtime/Experimental/Chaos/Private/Chaos/Collision/CollisionSolver.cpp b/Engine/Source/Runtime/Experimental/Chaos/Private/Chaos/Collision/CollisionSolver.cpp
--- a/Engine/Source/Runtime/Experimental/Chaos/Private/Chaos/Collision/CollisionSolver.cpp
+++ b/Engine/Source/Runtime/Experimental/Chaos/Private/Chaos/Collision/CollisionSolver.cpp
@@ -24,27 +24,51 @@ namespace Chaos
 		extern int32 Chaos_Collision_UseShockPropagation;
 		extern FRealSingle Chaos_Collision_CollisionClipTolerance;
 
+		bool bChaos_Manifold_Apply_AllowNegativeIncrementalImpulse = true;
+		FAutoConsoleVariableRef CVarChaos_Manifold_NegativeApply(TEXT("p.Chaos.Collision.Manifold.Apply.NegativeIncrementalImpulse"), bChaos_Manifold_Apply_AllowNegativeIncrementalImpulse, TEXT(""));
+
 		bool Chaos_Manifold_PushOut_NegativePushOut = false;
 		bool Chaos_Manifold_PushOut_StaticFriction = true;
 		bool Chaos_Manifold_PushOut_Restitution = false;
 		bool Chaos_Manifold_PushOut_PositionCorrection = true;
 		int32 Chaos_Manifold_PushOut_VelocityCorrection = 2;
+		float Chaos_Manifold_PushOut_DistanceFactor = 0.0f;
 		FAutoConsoleVariableRef CVarChaos_Manifold_PushOut_NegativePushOut(TEXT("p.Chaos.Collision.Manifold.PushOut.NegativePushOut"), Chaos_Manifold_PushOut_NegativePushOut, TEXT(""));
 		FAutoConsoleVariableRef CVarChaos_Manifold_PushOut_StaticFriction(TEXT("p.Chaos.Collision.Manifold.PushOut.StaticFriction"), Chaos_Manifold_PushOut_StaticFriction, TEXT(""));
 		FAutoConsoleVariableRef CVarChaos_Manifold_PushOut_Restitution(TEXT("p.Chaos.Collision.Manifold.PushOut.Restitution"), Chaos_Manifold_PushOut_Restitution, TEXT(""));
 		FAutoConsoleVariableRef CVarChaos_Manifold_PushOut_PositionCorrection(TEXT("p.Chaos.Collision.Manifold.PushOut.PositionCorrection"), Chaos_Manifold_PushOut_PositionCorrection, TEXT(""));
 		FAutoConsoleVariableRef CVarChaos_Manifold_PushOut_VelocityCorrection(TEXT("p.Chaos.Collision.Manifold.PushOut.VelocityCorrectionMode"), Chaos_Manifold_PushOut_VelocityCorrection, TEXT("0 = No Velocity Correction; 1 = Normal Velocity Correction; 2 = Normal + Tangential Velocity Correction"));
+		FAutoConsoleVariableRef CVarChaos_Manifold_PushOut_DistanceFactor(TEXT("p.Chaos.Collision.Manifold.PushOut.DistanceFactor"), Chaos_Manifold_PushOut_DistanceFactor, TEXT(""));
+
+		// Note: Velocity solve requires full stiffness for restitution to work correctly so the max stiffness should be 1
+		FRealSingle Chaos_Manifold_Apply_MinStiffness = 0.25f;
+		FRealSingle Chaos_Manifold_Apply_MaxStiffness = 1.0f;
+		FAutoConsoleVariableRef CVarChaos_Manifold_Apply_MinStiffness(TEXT("p.Chaos.Collision.Manifold.Apply.MinStiffness"), Chaos_Manifold_Apply_MinStiffness, TEXT(""));
+		FAutoConsoleVariableRef CVarChaos_Manifold_Apply_MaxStiffness(TEXT("p.Chaos.Collision.Manifold.Apply.MaxStiffness"), Chaos_Manifold_Apply_MaxStiffness, TEXT(""));
+
+		FRealSingle Chaos_Manifold_PushOut_MinStiffness = 0.25f;
+		FRealSingle Chaos_Manifold_PushOut_MaxStiffness = 0.5f;
+		FAutoConsoleVariableRef CVarChaos_Manifold_PushOut_MinStiffness(TEXT("p.Chaos.Collision.Manifold.PushOut.MinStiffness"), Chaos_Manifold_PushOut_MinStiffness, TEXT(""));
+		FAutoConsoleVariableRef CVarChaos_Manifold_PushOut_MaxStiffness(TEXT("p.Chaos.Collision.Manifold.PushOut.MaxStiffness"), Chaos_Manifold_PushOut_MaxStiffness, TEXT(""));
+
+		FRealSingle Chaos_Manifold_Apply_ImpulseTolerance = 1.e-4f;
+		FAutoConsoleVariableRef CVarChaos_Manifold_Apply_ImpulseTolerance(TEXT("p.Chaos.Collision.Manifold.Apply.ImpulseTolerance"), Chaos_Manifold_Apply_ImpulseTolerance, TEXT(""));
+
+		float Chaos_Manifold_PushOut_PositionTolerance = 1.e-4f;
+		FAutoConsoleVariableRef CVarChaos_Manifold_PushOut_PositionTolerance(TEXT("p.Chaos.Collision.Manifold.PushOut.PositionTolerance"), Chaos_Manifold_PushOut_PositionTolerance, TEXT(""));
 
-		FRealSingle Chaos_Manifold_MinPushOutStiffness = 0.25f;
-		FRealSingle Chaos_Manifold_MaxPushOutStiffness = 0.5f;
-		FAutoConsoleVariableRef CVarChaos_Manifold_PushOut_MinStiffness(TEXT("p.Chaos.Collision.Manifold.MinPushOutStiffness"), Chaos_Manifold_MinPushOutStiffness, TEXT(""));
-		FAutoConsoleVariableRef CVarChaos_Manifold_PushOut_MaxStiffness(TEXT("p.Chaos.Collision.Manifold.MaxPushOutStiffness"), Chaos_Manifold_MaxPushOutStiffness, TEXT(""));
 
-		FRealSingle Chaos_Manifold_ImpulseTolerance = 1.e-4f;
-		FAutoConsoleVariableRef CVarChaos_Manifold_ImpulseTolerance(TEXT("p.Chaos.Collision.Manifold.ImpulseTolerance"), Chaos_Manifold_ImpulseTolerance, TEXT(""));
+		inline FReal GetApplyStiffness(int32 It, int32 NumIts)
+		{
+			const FReal Interpolant = (FReal)(It + 1) / (FReal)NumIts;
+			return FMath::Lerp(Chaos_Manifold_Apply_MinStiffness, Chaos_Manifold_Apply_MaxStiffness, Interpolant);
+		}
 
-		FRealSingle Chaos_Manifold_PositionTolerance = 1.e-4f;
-		FAutoConsoleVariableRef CVarChaos_Manifold_PositionTolerance(TEXT("p.Chaos.Collision.Manifold.PositionTolerance"), Chaos_Manifold_PositionTolerance, TEXT(""));
+		inline FReal GetPushOutStiffness(int32 It, int32 NumIts)
+		{
+			const FReal Interpolant = (FReal)(It + 1) / (FReal)NumIts;
+			return FMath::Lerp(Chaos_Manifold_PushOut_MinStiffness, Chaos_Manifold_PushOut_MaxStiffness, Interpolant);
+		}
 
 
 		// Calculate the impulse to drive normal contact velocity to zero (or positive for restitution), ignoring friction
@@ -78,7 +102,9 @@ namespace Chaos
 			const FReal ContactPhi = ManifoldPoint.ContactPoint.Phi;
 
 			// Reject non-contact points unless the point has previously been processed - we may want to undo some of the previous work
-			if ((ContactPhi > Chaos_Collision_CollisionClipTolerance) && !ManifoldPoint.bActive)
+			// The tolerance is so that we catch collisions that were resolved last frame, but resulted in slightly positive separations. However,
+			// it also means we are effectively padding objects by this amount.
+			if ((ContactPhi > Chaos_Collision_CollisionClipTolerance) && (ManifoldPoint.InitialPhi > Chaos_Collision_CollisionClipTolerance) && !ManifoldPoint.bActive)
 			{
 				return;
 			}
@@ -89,7 +115,8 @@ namespace Chaos
 			const FReal ContactVelocityNormalLen = FVec3::DotProduct(ContactVelocity, ContactNormal);
 
 			// Reject contacts moving apart unless the point has previously been processed, in which case we may want to undo some of the previous work
-			if ((ContactVelocityNormalLen > 0.0f) && !ManifoldPoint.bActive)
+			const bool bAllowNegativeImpulse = ManifoldPoint.bActive && bChaos_Manifold_Apply_AllowNegativeIncrementalImpulse;
+			if ((ContactVelocityNormalLen > 0.0f) && !bAllowNegativeImpulse)
 			{
 				return;
 			}
@@ -149,7 +176,7 @@ namespace Chaos
 			ManifoldPoint.NetImpulse = NetImpulse;
 
 			// If we applied any additional impulse, we need to go again next iteration
-			const FReal ImpulseTolerance = Chaos_Manifold_ImpulseTolerance * Chaos_Manifold_ImpulseTolerance;
+			const FReal ImpulseTolerance = Chaos_Manifold_Apply_ImpulseTolerance * Chaos_Manifold_Apply_ImpulseTolerance;
 			*IterationParameters.NeedsAnotherIteration |= (Impulse.SizeSquared() > ImpulseTolerance);
 		}
 
@@ -252,7 +279,7 @@ namespace Chaos
 			ManifoldPoint.NetImpulse = NetImpulse;
 
 			// If we applied any additional impulse, we need to go again next iteration
-			const FReal ImpulseTolerance = Chaos_Manifold_ImpulseTolerance * Chaos_Manifold_ImpulseTolerance;
+			const FReal ImpulseTolerance = Chaos_Manifold_Apply_ImpulseTolerance * Chaos_Manifold_Apply_ImpulseTolerance;
 			*IterationParameters.NeedsAnotherIteration |= (Impulse.SizeSquared() > ImpulseTolerance);
 		}
 
@@ -289,7 +316,7 @@ namespace Chaos
 			const FRotation3& PlaneQ = (ManifoldPoint.ContactPoint.ContactNormalOwnerIndex == 0) ? Q0 : Q1;
 			const FVec3 ContactNormal = PlaneQ * ManifoldPoint.CoMContactNormal;
 
-			const bool bApplyStaticFriction = (ManifoldPoint.bInsideStaticFrictionCone && ManifoldPoint.bPotentialRestingContact && Chaos_Manifold_PushOut_StaticFriction);
+			const bool bApplyStaticFriction = (ManifoldPoint.bActive && ManifoldPoint.bInsideStaticFrictionCone && ManifoldPoint.bPotentialRestingContact && Chaos_Manifold_PushOut_StaticFriction);
 			FVec3 LocalContactPoint0 = ManifoldPoint.CoMContactPoints[0];
 			FVec3 LocalContactPoint1 = ManifoldPoint.CoMContactPoints[1];
 
@@ -310,17 +337,33 @@ namespace Chaos
 
 			// Net error we need to correct, including lateral movement to correct for friction
 			FVec3 ContactError = (P1 + RelativeContactPoint1) - (P0 + RelativeContactPoint0);
+
+			// The max rest distance for PBD object is the distance it can move in one frame from rest along the contact normal (i.e., position delta from gravity)
+			if ((Chaos_Manifold_PushOut_DistanceFactor > 0.0f) && (ContactNormal.Z > 0.0f))
+			{
+				// @todo(chaos): pass gravity into the solver if we use this feature. Otherwise remove it.
+				const FReal GravityNegZ = 980.0f;
+				const FReal MaxRestDistance = GravityNegZ * IterationParameters.Dt * IterationParameters.Dt;
+				ContactError += (Chaos_Manifold_PushOut_DistanceFactor * MaxRestDistance * ContactNormal.Z) * ContactNormal;
+			}
+
 			FReal ContactErrorNormal = FVec3::DotProduct(ContactError, ContactNormal);
 
 			// Don't allow objects to be pulled together, but we may still have to correct static friction drift
 			// @todo(chaos): allow undoing of previously added normal pushout
 			if ((ContactErrorNormal < 0.0f) && !Chaos_Manifold_PushOut_NegativePushOut)
 			{
+				// If we have no static friction and no normal separation to enforce, we have nothing to do
+				if (!bApplyStaticFriction)
+				{
+					return;
+				}
+
 				ContactError = ContactError - ContactErrorNormal * ContactNormal;
 				ContactErrorNormal = 0.0f;
 			}
 
-			if (ContactError.SizeSquared() < Chaos_Manifold_PositionTolerance * Chaos_Manifold_PositionTolerance)
+			if (ContactError.SizeSquared() < Chaos_Manifold_PushOut_PositionTolerance * Chaos_Manifold_PushOut_PositionTolerance)
 			{
 				return;
 			}
@@ -541,8 +584,7 @@ namespace Chaos
 
 			Constraint.AccumulatedImpulse = FVec3(0);
 
-			// Velocity solve requires full stiffness for restitution to work correctly
-			const FReal Stiffness = 1.0f;
+			const FReal Stiffness = GetApplyStiffness(IterationParameters.Iteration, IterationParameters.NumIterations);
 
 			// Iterate over the manifold and accumulate velocity corrections - we will apply them after the loop
 			TArrayView<FManifoldPoint> ManifoldPoints = Constraint.GetManifoldPoints();
@@ -671,8 +713,7 @@ namespace Chaos
 			}
 
 			// Gradually increase position correction through iterations (optional based on cvars)
-			const FReal Interpolant = (FReal)(IterationParameters.Iteration + 1) / (FReal)IterationParameters.NumIterations;
-			const FReal Stiffness = FMath::Lerp(Chaos_Manifold_MinPushOutStiffness, Chaos_Manifold_MaxPushOutStiffness, Interpolant);
+			const FReal Stiffness = GetPushOutStiffness(IterationParameters.Iteration, IterationParameters.NumIterations);
 
 			FVec3 P0 = FParticleUtilities::GetCoMWorldPosition(Particle0);
 			FRotation3 Q0 = FParticleUtilities::GetCoMWorldRotation(Particle0);
@@ -690,7 +731,17 @@ namespace Chaos
 				{
 					FManifoldPoint& ManifoldPoint = Constraint.SetActiveManifoldPoint(PointIndex, P0, Q0, P1, Q1);
 
-					if (ManifoldPoint.bActive || (ManifoldPoint.ContactPoint.Phi < 0.0f))
+					// The distance we should pushout to
+					FReal PushOutSeparation = 0.0f;
+					if ((Chaos_Manifold_PushOut_DistanceFactor > 0.0f) && (ManifoldPoint.ContactPoint.Normal.Z > 0.0f))
+					{
+						// @todo(chaos): pass gravity into the solver if we use this feature. Otherwise remove it.
+						const FReal GravityNegZ = 980.0f;
+						const FReal MaxRestDistance = GravityNegZ * IterationParameters.Dt * IterationParameters.Dt;
+						PushOutSeparation = (Chaos_Manifold_PushOut_DistanceFactor * MaxRestDistance * ManifoldPoint.ContactPoint.Normal.Z);
+					}
+
+					if (ManifoldPoint.bActive || (ManifoldPoint.ContactPoint.Phi < PushOutSeparation))
 					{
 						ApplyManifoldPushOutCorrection(
 							Stiffness,
diff --git a/Engine/Source/Runtime/Experimental/Chaos/Private/Chaos/Collision/PBDCollisionConstraint.cpp b/Engine/Source/Runtime/Experimental/Chaos/Private/Chaos/Collision/PBDCollisionConstraint.cpp
--- a/Engine/Source/Runtime/Experimental/Chaos/Private/Chaos/Collision/PBDCollisionConstraint.cpp
+++ b/Engine/Source/Runtime/Experimental/Chaos/Private/Chaos/Collision/PBDCollisionConstraint.cpp
@@ -158,27 +158,35 @@
 
 	void FRigidBodyPointContactConstraint::AddIncrementalManifoldContact(const FContactPoint& ContactPoint, const FReal Dt)
 	{
-		int32 ManifoldPointIndex = INDEX_NONE;
 		if (bUseManifold)
 		{
 			// See if the manifold point already exists
-			ManifoldPointIndex = FindManifoldPoint(ContactPoint);
+			int32 ManifoldPointIndex = FindManifoldPoint(ContactPoint);
+			if (ManifoldPointIndex >= 0)
+			{
+				// This contact point is already in the manifold - update the state
+				UpdateManifoldPoint(ManifoldPointIndex, ContactPoint, Dt);
+			}
+			else
+			{
+				// This is a new manifold point - capture the state and generate initial properties
+				AddManifoldPoint(ContactPoint, Dt);
+			}
+
 		}
-		else if (ManifoldPoints.Num() > 0)
+		else 
 		{
 			// We are not using manifolds - reuse the first and only point
-			ManifoldPointIndex = 0;
+			if (ManifoldPoints.Num() == 0)
+			{
+				ManifoldPoints.Add(ContactPoint);
+			}
+			else
+			{
+				ManifoldPoints[0] = ContactPoint;
+			}
+
+			InitManifoldPoint(ManifoldPoints[0], Dt);
 		}
 
-		if (ManifoldPointIndex >= 0)
-		{
-			// This contact point is already in the manifold - update the state
-			UpdateManifoldPoint(ManifoldPointIndex, ContactPoint, Dt);
-		}
-		else
-		{
-			// This is a new manifold point - capture the state and generate initial properties
-			ManifoldPointIndex = AddManifoldPoint(ContactPoint, Dt);
-		}
-
 		// Copy currently active point
@@ -208,40 +216,29 @@
 		// @todo(chaos): support static friction position correction for non-manifold contacts (spheres, point clouds, etc)
 		ManifoldPoint.bPotentialRestingContact = bUseManifold;
 
+		// Update the derived contact state (CoM relative data)
+		UpdateManifoldPointFromContact(ManifoldPoint);
+
+		// World-space contact state used below
+		// @todo(chaos): some of this should already be known - optimize
+		const FRigidTransform3 WorldCoMTransform0 = FParticleUtilities::GetCoMWorldTransform(Particle0);
+		const FRigidTransform3 WorldCoMTransform1 = FParticleUtilities::GetCoMWorldTransform(Particle1);
+		FVec3 WorldContactLocation, WorldContactNormal;
+		FReal WorldContactPhi;
+		GetWorldSpaceManifoldPoint(ManifoldPoint, WorldCoMTransform0.GetLocation(), WorldCoMTransform0.GetRotation(), WorldCoMTransform1.GetLocation(), WorldCoMTransform1.GetRotation(), WorldContactLocation, WorldContactNormal, WorldContactPhi);
+
 		// Calculate and store the data required for static friction and restitution in PushOut
-		//	- contact point on second shape at previous location (to enforce static friction)
-		//	- initial and previous separations (for restitution)
-		// @todo(chaos): we shouldn't have to recalculate local space positions - they were available in collision update
-		check(ManifoldPoint.ContactPoint.ContactNormalOwnerIndex >= 0);
-		check(ManifoldPoint.ContactPoint.ContactNormalOwnerIndex < 2);
-		const FRigidTransform3& PlaneTransform = ImplicitTransform[ManifoldPoint.ContactPoint.ContactNormalOwnerIndex];
-		const FVec3 LocalContactPoint0 = ImplicitTransform[0].TransformPositionNoScale(ManifoldPoint.ContactPoint.ShapeContactPoints[0]);	// Particle Space on body 0
-		const FVec3 LocalContactPoint1 = ImplicitTransform[1].TransformPositionNoScale(ManifoldPoint.ContactPoint.ShapeContactPoints[1]);	// Particle Space on body 1
-		const FVec3 LocalContactNormal = PlaneTransform.TransformNormalNoScale(ManifoldPoint.ContactPoint.ShapeContactNormal);				// Particle Space on Plane owner
-
-		const FRotation3& PlaneRotationOfMass = (ManifoldPoint.ContactPoint.ContactNormalOwnerIndex == 0) ? Particle0->RotationOfMass() : Particle1->RotationOfMass();
-		const FVec3 CoMContactPoint0 = Particle0->RotationOfMass().Inverse() * (LocalContactPoint0 - Particle0->CenterOfMass());	// CoM Space on Body 0
-		const FVec3 CoMContactPoint1 = Particle1->RotationOfMass().Inverse() * (LocalContactPoint1 - Particle1->CenterOfMass());	// CoM Space on Body 1
-		const FVec3 CoMContactNormal = PlaneRotationOfMass.Inverse() * LocalContactNormal;											// CoM Space on Planer owner
-		ManifoldPoint.CoMContactPoints[0] = CoMContactPoint0;
-		ManifoldPoint.CoMContactPoints[1] = CoMContactPoint1;
-		ManifoldPoint.CoMContactNormal = CoMContactNormal;
-
-		// Calculate the initial contact velocity for use in restitution
 		// We use PreV and PreW to support incremental manifold generation. In this case, manifold points
 		// can be added after we have already run some solver iterations, which gives us an incorrect initial
 		// velocity if we just use V and W (one-shots will work the same either way since V=PreV on first pass)
-		const FRigidTransform3 WorldCoMTransform0 = FParticleUtilities::GetCoMWorldTransform(Particle0);
-		const FRigidTransform3 WorldCoMTransform1 = FParticleUtilities::GetCoMWorldTransform(Particle1);
-		const FVec3 WorldContactVel0 = Particle0->PreV() + FVec3::CrossProduct(Particle0->PreW(), WorldCoMTransform0.GetRotation() * CoMContactPoint0);
-		const FVec3 WorldContactVel1 = Particle1->PreV() + FVec3::CrossProduct(Particle1->PreW(), WorldCoMTransform1.GetRotation() * CoMContactPoint1);
-		const FVec3 WorldContactNormal = (ManifoldPoint.ContactPoint.ContactNormalOwnerIndex == 0) ? WorldCoMTransform0.GetRotation() * CoMContactNormal : WorldCoMTransform1.GetRotation() * CoMContactNormal;
+		const FVec3 WorldContactVel0 = Particle0->PreV() + FVec3::CrossProduct(Particle0->PreW(), WorldCoMTransform0.GetRotation() * ManifoldPoint.CoMContactPoints[0]);
+		const FVec3 WorldContactVel1 = Particle1->PreV() + FVec3::CrossProduct(Particle1->PreW(), WorldCoMTransform1.GetRotation() * ManifoldPoint.CoMContactPoints[1]);
 		const FReal WorldContactVelNorm = FVec3::DotProduct(WorldContactVel0 - WorldContactVel1, WorldContactNormal);
 		ManifoldPoint.InitialContactVelocity = WorldContactVelNorm;
 
 		// Store the initial penetration depth for use with restitution with PBD
 		// NOTE: This is incorrect if we are updating the point each iteration (which we are for incremental manifolds but not one-shots - see UpdateManifoldPoint)
-		ManifoldPoint.InitialPhi = ManifoldPoint.ContactPoint.Phi;
+		ManifoldPoint.InitialPhi = WorldContactPhi;
 	}
 
 	void FRigidBodyPointContactConstraint::CalculatePrevCoMContactPoints(
@@ -263,7 +260,7 @@
 				PrevCoMTransform.AddToTranslation(-Dt * ParticleHandle->V());
 				PrevCoMTransform.SetRotation(FRotation3::IntegrateRotationWithAngularVelocity(PrevCoMTransform.GetRotation(), -ParticleHandle->W(), Dt));
 			}
-			return PrevCoMTransform;
+			return PrevCoMTransform; 
 		};
 
 		const FRigidTransform3 PrevCoMTransform0 = CalculatePrevCoMTransform(Particle0);
@@ -305,7 +302,30 @@
 		// otherwise the PrevLocalContactPoint1 we calculated is not longer for the correct point.
 		FManifoldPoint& ManifoldPoint = ManifoldPoints[ManifoldPointIndex];
 		ManifoldPoint.ContactPoint = ContactPoint;
-		InitManifoldPoint(ManifoldPoint, Dt);
+		UpdateManifoldPointFromContact(ManifoldPoint);
 	}
 
+	// Update the derived contact state (CoM relative data)
+	// @todo(chaos): we shouldn't have to recalculate local space positions - they were available in collision update
+	void FRigidBodyPointContactConstraint::UpdateManifoldPointFromContact(FManifoldPoint& ManifoldPoint)
+	{
+		TConstGenericParticleHandle<FReal, 3> Particle0 = Particle[0];
+		TConstGenericParticleHandle<FReal, 3> Particle1 = Particle[1];
+
+		check(ManifoldPoint.ContactPoint.ContactNormalOwnerIndex >= 0);
+		check(ManifoldPoint.ContactPoint.ContactNormalOwnerIndex < 2);
+		const FRigidTransform3& PlaneTransform = ImplicitTransform[ManifoldPoint.ContactPoint.ContactNormalOwnerIndex];
+		const FVec3 LocalContactPoint0 = ImplicitTransform[0].TransformPositionNoScale(ManifoldPoint.ContactPoint.ShapeContactPoints[0]);	// Particle Space on body 0
+		const FVec3 LocalContactPoint1 = ImplicitTransform[1].TransformPositionNoScale(ManifoldPoint.ContactPoint.ShapeContactPoints[1]);	// Particle Space on body 1
+		const FVec3 LocalContactNormal = PlaneTransform.TransformNormalNoScale(ManifoldPoint.ContactPoint.ShapeContactNormal);				// Particle Space on Plane owner
+
+		const FRotation3& PlaneRotationOfMass = (ManifoldPoint.ContactPoint.ContactNormalOwnerIndex == 0) ? Particle0->RotationOfMass() : Particle1->RotationOfMass();
+		const FVec3 CoMContactPoint0 = Particle0->RotationOfMass().Inverse() * (LocalContactPoint0 - Particle0->CenterOfMass());	// CoM Space on Body 0
+		const FVec3 CoMContactPoint1 = Particle1->RotationOfMass().Inverse() * (LocalContactPoint1 - Particle1->CenterOfMass());	// CoM Space on Body 1
+		const FVec3 CoMContactNormal = PlaneRotationOfMass.Inverse() * LocalContactNormal;											// CoM Space on Planer owner
+		ManifoldPoint.CoMContactPoints[0] = CoMContactPoint0;
+		ManifoldPoint.CoMContactPoints[1] = CoMContactPoint1;
+		ManifoldPoint.CoMContactNormal = CoMContactNormal;
+	}
+
 	void FRigidBodyPointContactConstraint::SetActiveContactPoint(const FContactPoint& ContactPoint)
@@ -316,6 +336,30 @@
 		Manifold.Phi = ContactPoint.Phi;
 	}
 
+	void FRigidBodyPointContactConstraint::GetWorldSpaceManifoldPoint(
+		const FManifoldPoint& ManifoldPoint,
+		const FVec3& P0,			// World-Space CoM
+		const FRotation3& Q0,		// World-Space CoM
+		const FVec3& P1,			// World-Space CoM
+		const FRotation3& Q1,		// World-Space CoM
+		FVec3& OutContactLocation,
+		FVec3& OutContactNormal,
+		FReal& OutContactPhi)
+	{
+		const FReal Margin0 = ManifoldPoint.ContactPoint.ShapeMargins[0];
+		const FReal Margin1 = ManifoldPoint.ContactPoint.ShapeMargins[1];
+		const FRotation3& PlaneQ = (ManifoldPoint.ContactPoint.ContactNormalOwnerIndex == 0) ? Q0 : Q1;
+
+		const FVec3 ContactNormal = PlaneQ.RotateVector(ManifoldPoint.CoMContactNormal);
+		const FVec3 ContactPos0 = P0 + Q0.RotateVector(ManifoldPoint.CoMContactPoints[0]) - Margin0 * ContactNormal;
+		const FVec3 ContactPos1 = P1 + Q1.RotateVector(ManifoldPoint.CoMContactPoints[1]) + Margin1 * ContactNormal;
+
+		OutContactLocation = 0.5f * (ContactPos0 + ContactPos1);
+		OutContactNormal = ContactNormal;
+		OutContactPhi = FVec3::DotProduct(ContactPos0 - ContactPos1, ContactNormal);
+	}
+
+
 	FManifoldPoint& FRigidBodyPointContactConstraint::SetActiveManifoldPoint(
 		int32 ManifoldPointIndex,
 		const FVec3& P0,
@@ -325,18 +369,10 @@
 	{
 		FManifoldPoint& ManifoldPoint = ManifoldPoints[ManifoldPointIndex];
 
-		const FReal Margin0 = ManifoldPoint.ContactPoint.ShapeMargins[0];
-		const FReal Margin1 = ManifoldPoint.ContactPoint.ShapeMargins[1];
-		const FRotation3& PlaneQ = (ManifoldPoint.ContactPoint.ContactNormalOwnerIndex == 0) ? Q0 : Q1;
-
-		// Update the world-space point state based on current particle transforms
-		const FVec3 ContactNormal = PlaneQ.RotateVector(ManifoldPoint.CoMContactNormal);
-		const FVec3 ContactPos0 = P0 + Q0.RotateVector(ManifoldPoint.CoMContactPoints[0]) - Margin0 * ContactNormal;
-		const FVec3 ContactPos1 = P1 + Q1.RotateVector(ManifoldPoint.CoMContactPoints[1]) + Margin1 * ContactNormal;
-		ManifoldPoint.ContactPoint.Location = 0.5f * (ContactPos0 + ContactPos1);
-		ManifoldPoint.ContactPoint.Normal = ContactNormal;
-		ManifoldPoint.ContactPoint.Phi = FVec3::DotProduct(ContactPos0 - ContactPos1, ContactNormal);
+		// Update the world-space state in the manifold point
+		GetWorldSpaceManifoldPoint(ManifoldPoint, P0, Q0, P1, Q1, ManifoldPoint.ContactPoint.Location, ManifoldPoint.ContactPoint.Normal, ManifoldPoint.ContactPoint.Phi);
 	
+		// Copy the world-space state into the active contact
 		SetActiveContactPoint(ManifoldPoint.ContactPoint);
 
 		return ManifoldPoint;
diff --git a/Engine/Source/Runtime/Experimental/Chaos/Private/Chaos/Collision/SphereConvexContactPoint.cpp b/Engine/Source/Runtime/Experimental/Chaos/Private/Chaos/Collision/SphereConvexContactPoint.cpp
--- a/Engine/Source/Runtime/Experimental/Chaos/Private/Chaos/Collision/SphereConvexContactPoint.cpp
+++ b/Engine/Source/Runtime/Experimental/Chaos/Private/Chaos/Collision/SphereConvexContactPoint.cpp
@@ -45,7 +45,7 @@ namespace Chaos
 	// Use GJK (point to convex) to calculate separation.
 	// Fall back to plane testing if penetrating by more than Radius.
 	template<typename T_CONVEX>
-	FContactPoint ConvexSphereContactPointImpl(const T_CONVEX& Convex, const FRigidTransform3& ConvexTransform, const FImplicitSphere3& Sphere, const FRigidTransform3& SphereTransform, FReal CullDistance)
+	FContactPoint ConvexSphereContactPointImpl(const T_CONVEX& Convex, const FRigidTransform3& ConvexTransform, const FImplicitSphere3& Sphere, const FRigidTransform3& SphereTransform)
 	{
 		FContactPoint ContactPoint;
 		if (Convex.NumPlanes() > 0)
@@ -94,58 +94,58 @@ namespace Chaos
 	}
 
 	template<typename T_CONVEX>
-	FContactPoint SphereConvexContactPointImpl(const FImplicitSphere3& Sphere, const FRigidTransform3& SphereTransform, const T_CONVEX& Convex, const FRigidTransform3& ConvexTransform, FReal CullDistance)
+	FContactPoint SphereConvexContactPointImpl(const FImplicitSphere3& Sphere, const FRigidTransform3& SphereTransform, const T_CONVEX& Convex, const FRigidTransform3& ConvexTransform)
 	{
-		return ConvexSphereContactPointImpl(Convex, ConvexTransform, Sphere, SphereTransform, CullDistance).SwapShapes();
+		return ConvexSphereContactPointImpl(Convex, ConvexTransform, Sphere, SphereTransform).SwapShapes();
 	}
 
-	FContactPoint SphereConvexContactPoint(const FImplicitSphere3& Sphere, const FRigidTransform3& SphereTransform, const FImplicitConvex3& Convex, const FRigidTransform3& ConvexTransform, FReal CullDistance)
+	FContactPoint SphereConvexContactPoint(const FImplicitSphere3& Sphere, const FRigidTransform3& SphereTransform, const FImplicitConvex3& Convex, const FRigidTransform3& ConvexTransform)
 	{
-		return SphereConvexContactPointImpl(Sphere, SphereTransform, Convex, ConvexTransform, CullDistance);
+		return SphereConvexContactPointImpl(Sphere, SphereTransform, Convex, ConvexTransform);
 	}
 
-	FContactPoint SphereConvexContactPoint(const FImplicitSphere3& Sphere, const FRigidTransform3& SphereTransform, const TImplicitObjectInstanced<FImplicitConvex3>& Convex, const FRigidTransform3& ConvexTransform, FReal CullDistance)
+	FContactPoint SphereConvexContactPoint(const FImplicitSphere3& Sphere, const FRigidTransform3& SphereTransform, const TImplicitObjectInstanced<FImplicitConvex3>& Convex, const FRigidTransform3& ConvexTransform)
 	{
-		return SphereConvexContactPointImpl(Sphere, SphereTransform, Convex, ConvexTransform, CullDistance);
+		return SphereConvexContactPointImpl(Sphere, SphereTransform, Convex, ConvexTransform);
 	}
 
-	FContactPoint SphereConvexContactPoint(const FImplicitSphere3& Sphere, const FRigidTransform3& SphereTransform, const TImplicitObjectScaled<FImplicitConvex3>& Convex, const FRigidTransform3& ConvexTransform, FReal CullDistance)
+	FContactPoint SphereConvexContactPoint(const FImplicitSphere3& Sphere, const FRigidTransform3& SphereTransform, const TImplicitObjectScaled<FImplicitConvex3>& Convex, const FRigidTransform3& ConvexTransform)
 	{
-		return SphereConvexContactPointImpl(Sphere, SphereTransform, Convex, ConvexTransform, CullDistance);
+		return SphereConvexContactPointImpl(Sphere, SphereTransform, Convex, ConvexTransform);
 	}
 
 
 
-	FContactPoint ConvexSphereContactPoint(const FImplicitConvex3& Convex, const FRigidTransform3& ConvexTransform, const FImplicitSphere3& Sphere, const FRigidTransform3& SphereTransform, FReal CullDistance)
+	FContactPoint ConvexSphereContactPoint(const FImplicitConvex3& Convex, const FRigidTransform3& ConvexTransform, const FImplicitSphere3& Sphere, const FRigidTransform3& SphereTransform)
 	{
-		return ConvexSphereContactPointImpl(Convex, ConvexTransform, Sphere, SphereTransform, CullDistance);
+		return ConvexSphereContactPointImpl(Convex, ConvexTransform, Sphere, SphereTransform);
 	}
 
-	FContactPoint ConvexSphereContactPoint(const TImplicitObjectInstanced<FImplicitConvex3>& Convex, const FRigidTransform3& ConvexTransform, const FImplicitSphere3& Sphere, const FRigidTransform3& SphereTransform, FReal CullDistance)
+	FContactPoint ConvexSphereContactPoint(const TImplicitObjectInstanced<FImplicitConvex3>& Convex, const FRigidTransform3& ConvexTransform, const FImplicitSphere3& Sphere, const FRigidTransform3& SphereTransform)
 	{
-		return ConvexSphereContactPointImpl(Convex, ConvexTransform, Sphere, SphereTransform, CullDistance);
+		return ConvexSphereContactPointImpl(Convex, ConvexTransform, Sphere, SphereTransform);
 	}
 
-	FContactPoint ConvexSphereContactPoint(const TImplicitObjectScaled<FImplicitConvex3>& Convex, const FRigidTransform3& ConvexTransform, const FImplicitSphere3& Sphere, const FRigidTransform3& SphereTransform, FReal CullDistance)
+	FContactPoint ConvexSphereContactPoint(const TImplicitObjectScaled<FImplicitConvex3>& Convex, const FRigidTransform3& ConvexTransform, const FImplicitSphere3& Sphere, const FRigidTransform3& SphereTransform)
 	{
-		return ConvexSphereContactPointImpl(Convex, ConvexTransform, Sphere, SphereTransform, CullDistance);
+		return ConvexSphereContactPointImpl(Convex, ConvexTransform, Sphere, SphereTransform);
 	}
 
 
 
-	FContactPoint SphereConvexContactPoint(const FImplicitSphere3& Sphere, const FRigidTransform3& SphereTransform, const FImplicitObject& Object, const FRigidTransform3& ConvexTransform, FReal CullDistance)
+	FContactPoint SphereConvexContactPoint(const FImplicitSphere3& Sphere, const FRigidTransform3& SphereTransform, const FImplicitObject& Object, const FRigidTransform3& ConvexTransform)
 	{
 		if (const TImplicitObjectInstanced<FImplicitConvex3>* InstancedConvex = Object.template GetObject<const TImplicitObjectInstanced<FImplicitConvex3>>())
 		{
-			return SphereConvexContactPoint(Sphere, SphereTransform, *InstancedConvex, ConvexTransform, CullDistance);
+			return SphereConvexContactPoint(Sphere, SphereTransform, *InstancedConvex, ConvexTransform);
 		}
 		else if (const TImplicitObjectScaled<FImplicitConvex3>* ScaledConvex = Object.template GetObject<const TImplicitObjectScaled<FImplicitConvex3>>())
 		{
-			return SphereConvexContactPoint(Sphere, SphereTransform, *ScaledConvex, ConvexTransform, CullDistance);
+			return SphereConvexContactPoint(Sphere, SphereTransform, *ScaledConvex, ConvexTransform);
 		}
 		else if (const FImplicitConvex3* Convex = Object.template GetObject<const FImplicitConvex3>())
 		{
-			return SphereConvexContactPoint(Sphere, SphereTransform, *Convex, ConvexTransform, CullDistance);
+			return SphereConvexContactPoint(Sphere, SphereTransform, *Convex, ConvexTransform);
 		}
 		return FContactPoint();
 	}
diff --git a/Engine/Source/Runtime/Experimental/Chaos/Private/Chaos/CollisionOneShotManifolds.cpp b/Engine/Source/Runtime/Experimental/Chaos/Private/Chaos/CollisionOneShotManifolds.cpp
--- a/Engine/Source/Runtime/Experimental/Chaos/Private/Chaos/CollisionOneShotManifolds.cpp
+++ b/Engine/Source/Runtime/Experimental/Chaos/Private/Chaos/CollisionOneShotManifolds.cpp
@@ -21,14 +21,15 @@ namespace Chaos
 
 	// @todo(chaos): TEMP - use convex-convex collisio detection for box-box until TBox::GetClosestEdgePosition is implemented for that path (without plane hint)
 	bool bChaos_Collision_Manifold_BoxAsConvex = true;
-	FAutoConsoleVariableRef CVarChaosCollisioConvexManifodlBoxAsConvex(TEXT("p.Chaos.Collision.Manifold.BoxAsConvex"), bChaos_Collision_Manifold_BoxAsConvex, TEXT(""));
+	FAutoConsoleVariableRef CVarChaosCollisioConvexManifoldBoxAsConvex(TEXT("p.Chaos.Collision.Manifold.BoxAsConvex"), bChaos_Collision_Manifold_BoxAsConvex, TEXT(""));
+
 
 	namespace Collisions
 	{
 		// Forward delarations we need from CollisionRestitution.cpp
 
-		FContactPoint BoxBoxContactPoint(const FImplicitBox3& Box1, const FImplicitBox3& Box2, const FRigidTransform3& Box1TM, const FRigidTransform3& Box2TM, const FReal CullDistance, const FReal ShapePadding);
-		FContactPoint GenericConvexConvexContactPoint(const FImplicitObject& A, const FRigidTransform3& ATM, const FImplicitObject& B, const FRigidTransform3& BTM, const FReal CullDistance, const FReal ShapePadding);
+		FContactPoint BoxBoxContactPoint(const FImplicitBox3& Box1, const FImplicitBox3& Box2, const FRigidTransform3& Box1TM, const FRigidTransform3& Box2TM, const FReal ShapePadding);
+		FContactPoint GenericConvexConvexContactPoint(const FImplicitObject& A, const FRigidTransform3& ATM, const FImplicitObject& B, const FRigidTransform3& BTM, const FReal ShapePadding);
 
 		//////////////////////////
 		// Box Box
@@ -109,13 +110,12 @@ namespace Chaos
 			const FRigidTransform3& Box1Transform, //world
 			const FImplicitBox3& Box2,
 			const FRigidTransform3& Box2Transform, //world
-			const FReal CullDistance,
 			const FReal Dt,
 			FRigidBodyPointContactConstraint& Constraint)
 		{
 			if (bChaos_Collision_Manifold_BoxAsConvex)
 			{
-				ConstructConvexConvexOneShotManifold(Box1, Box1Transform, Box2, Box2Transform, CullDistance, Dt, Constraint);
+				ConstructConvexConvexOneShotManifold(Box1, Box1Transform, Box2, Box2Transform, Dt, Constraint);
 				return;
 			}
 
@@ -131,7 +131,7 @@ namespace Chaos
 			uint32 ContactPointCount = 0;
 
 			// Use GJK only once
-			FContactPoint GJKContactPoint = BoxBoxContactPoint(Box1, Box2, Box1Transform, Box2Transform, CullDistance, Constraint.Manifold.RestitutionPadding);
+			FContactPoint GJKContactPoint = BoxBoxContactPoint(Box1, Box2, Box1Transform, Box2Transform, Constraint.Manifold.RestitutionPadding);
 
 			FRigidTransform3 Box1TransformCenter = Box1Transform;
 			Box1TransformCenter.SetTranslation(Box1Transform.TransformPositionNoScale(Box1.GetCenter()));
@@ -527,7 +527,6 @@ namespace Chaos
 			const FRigidTransform3& Convex1Transform, //world
 			const ConvexImplicitType2& Convex2,
 			const FRigidTransform3& Convex2Transform, //world
-			const FReal CullDistance,
 			const FReal Dt,
 			FRigidBodyPointContactConstraint& Constraint)
 		{
@@ -540,8 +539,10 @@ namespace Chaos
 			ensure(Convex2Transform.GetScale3D() == FVec3(1.0f, 1.0f, 1.0f));
 
 			// Find the deepest penetration. This is used to determine the planes and points to use for the manifold
-			FContactPoint GJKContactPoint = GenericConvexConvexContactPoint(Convex1, Convex1Transform, Convex2, Convex2Transform, CullDistance, Constraint.Manifold.RestitutionPadding);
-			if (GJKContactPoint.Phi > CullDistance)
+			FContactPoint GJKContactPoint = GenericConvexConvexContactPoint(Convex1, Convex1Transform, Convex2, Convex2Transform, Constraint.Manifold.RestitutionPadding);
+
+			// @todo(chaos): this is not a good culling test - make it better. The separation from GJK can be larger than the real separation because of the margins (rounded corners)
+			if (GJKContactPoint.Phi > Constraint.GetCullDistance())
 			{
 				return;
 			}
@@ -678,7 +679,7 @@ namespace Chaos
 					MaxContactPointCount);
 			}
 			
-			// Reduce number of contacts to a maximum of 4
+			// Reduce number of contacts to the maximum allowed
 			if (ContactPointCount > 4)
 			{
 				FRotation3 RotateSeperationToZ = FRotation3::FromRotatedVector(RefPlaneNormal, FVec3(0.0f, 0.0f, 1.0f));
@@ -731,7 +732,6 @@ namespace Chaos
 			const FRigidTransform3& Convex1Transform, //world
 			const FImplicitBox3& Implicit2,
 			const FRigidTransform3& Convex2Transform, //world
-			const FReal CullDistance,
 			const FReal Dt,
 			FRigidBodyPointContactConstraint& Constraint);
 
@@ -741,7 +741,6 @@ namespace Chaos
 			const FRigidTransform3& Convex1Transform, //world
 			const FImplicitConvex3& Implicit2,
 			const FRigidTransform3& Convex2Transform, //world
-			const FReal CullDistance,
 			const FReal Dt,
 			FRigidBodyPointContactConstraint& Constraint);
 
@@ -751,7 +750,6 @@ namespace Chaos
 			const FRigidTransform3& Convex1Transform, //world
 			const FImplicitBox3& Implicit2,
 			const FRigidTransform3& Convex2Transform, //world
-			const FReal CullDistance,
 			const FReal Dt,
 			FRigidBodyPointContactConstraint& Constraint);
 
@@ -761,7 +759,6 @@ namespace Chaos
 			const FRigidTransform3& Convex1Transform, //world
 			const TImplicitObjectInstanced<FImplicitConvex3>& Implicit2,
 			const FRigidTransform3& Convex2Transform, //world
-			const FReal CullDistance,
 			const FReal Dt,
 			FRigidBodyPointContactConstraint& Constraint);
 
@@ -771,7 +768,6 @@ namespace Chaos
 			const FRigidTransform3& Convex1Transform, //world
 			const FImplicitBox3& Implicit2,
 			const FRigidTransform3& Convex2Transform, //world
-			const FReal CullDistance,
 			const FReal Dt,
 			FRigidBodyPointContactConstraint& Constraint);
 
@@ -781,7 +777,6 @@ namespace Chaos
 			const FRigidTransform3& Convex1Transform, //world
 			const TImplicitObjectScaled<FImplicitConvex3>& Implicit2,
 			const FRigidTransform3& Convex2Transform, //world
-			const FReal CullDistance,
 			const FReal Dt,
 			FRigidBodyPointContactConstraint& Constraint);
 
@@ -791,7 +786,6 @@ namespace Chaos
 			const FRigidTransform3& Convex1Transform, //world
 			const FImplicitBox3& Implicit2,
 			const FRigidTransform3& Convex2Transform, //world
-			const FReal CullDistance,
 			const FReal Dt,
 			FRigidBodyPointContactConstraint& Constraint);
 
@@ -801,7 +795,6 @@ namespace Chaos
 			const FRigidTransform3& Convex1Transform, //world
 			const FImplicitConvex3& Implicit2,
 			const FRigidTransform3& Convex2Transform, //world
-			const FReal CullDistance,
 			const FReal Dt,
 			FRigidBodyPointContactConstraint& Constraint);
 
@@ -811,7 +804,6 @@ namespace Chaos
 			const FRigidTransform3& Convex1Transform, //world
 			const FImplicitConvex3& Implicit2,
 			const FRigidTransform3& Convex2Transform, //world
-			const FReal CullDistance,
 			const FReal Dt,
 			FRigidBodyPointContactConstraint& Constraint);
 
@@ -821,7 +813,6 @@ namespace Chaos
 			const FRigidTransform3& Convex1Transform, //world
 			const FImplicitConvex3& Implicit2,
 			const FRigidTransform3& Convex2Transform, //world
-			const FReal CullDistance,
 			const FReal Dt,
 			FRigidBodyPointContactConstraint& Constraint);
 
@@ -831,7 +822,6 @@ namespace Chaos
 			const FRigidTransform3& Convex1Transform, //world
 			const TImplicitObjectInstanced<FImplicitConvex3>& Implicit2,
 			const FRigidTransform3& Convex2Transform, //world
-			const FReal CullDistance,
 			const FReal Dt,
 			FRigidBodyPointContactConstraint& Constraint);
 
@@ -841,7 +831,6 @@ namespace Chaos
 			const FRigidTransform3& Convex1Transform, //world
 			const TImplicitObjectScaled<FImplicitConvex3>& Implicit2,
 			const FRigidTransform3& Convex2Transform, //world
-			const FReal CullDistance,
 			const FReal Dt,
 			FRigidBodyPointContactConstraint& Constraint);
 
@@ -851,7 +840,6 @@ namespace Chaos
 			const FRigidTransform3& Convex1Transform, //world
 			const TImplicitObjectInstanced<FImplicitConvex3>& Implicit2,
 			const FRigidTransform3& Convex2Transform, //world
-			const FReal CullDistance,
 			const FReal Dt,
 			FRigidBodyPointContactConstraint& Constraint);
 
@@ -861,7 +849,6 @@ namespace Chaos
 			const FRigidTransform3& Convex1Transform, //world
 			const TImplicitObjectInstanced<FImplicitConvex3>& Implicit2,
 			const FRigidTransform3& Convex2Transform, //world
-			const FReal CullDistance,
 			const FReal Dt,
 			FRigidBodyPointContactConstraint& Constraint);
 
@@ -871,7 +858,6 @@ namespace Chaos
 			const FRigidTransform3& Convex1Transform, //world
 			const TImplicitObjectScaled<FImplicitConvex3>& Implicit2,
 			const FRigidTransform3& Convex2Transform, //world
-			const FReal CullDistance,
 			const FReal Dt,
 			FRigidBodyPointContactConstraint& Constraint);
 
@@ -881,7 +867,6 @@ namespace Chaos
 			const FRigidTransform3& Convex1Transform, //world
 			const TImplicitObjectScaled<FImplicitConvex3>& Implicit2,
 			const FRigidTransform3& Convex2Transform, //world
-			const FReal CullDistance,
 			const FReal Dt,
 			FRigidBodyPointContactConstraint& Constraint);
 	}
diff --git a/Engine/Source/Runtime/Experimental/Chaos/Private/Chaos/CollisionOneShotManifolds.h b/Engine/Source/Runtime/Experimental/Chaos/Private/Chaos/CollisionOneShotManifolds.h
--- a/Engine/Source/Runtime/Experimental/Chaos/Private/Chaos/CollisionOneShotManifolds.h
+++ b/Engine/Source/Runtime/Experimental/Chaos/Private/Chaos/CollisionOneShotManifolds.h
@@ -19,7 +19,6 @@ namespace Chaos
 			const FRigidTransform3& Box1Transform, //world
 			const FImplicitBox3& Box2,
 			const FRigidTransform3& Box2Transform, //world
-			const FReal CullDistance,
 			const FReal Dt,
 			FRigidBodyPointContactConstraint& Constraint);
 
@@ -29,7 +28,6 @@ namespace Chaos
 			const FRigidTransform3& Convex1Transform, //world
 			const ConvexImplicitType2& Implicit2,
 			const FRigidTransform3& Convex2Transform, //world
-			const FReal CullDistance,
 			const FReal Dt,
 			FRigidBodyPointContactConstraint& Constraint);
 	}
diff --git a/Engine/Source/Runtime/Experimental/Chaos/Private/Chaos/CollisionResolution.cpp b/Engine/Source/Runtime/Experimental/Chaos/Private/Chaos/CollisionResolution.cpp
--- a/Engine/Source/Runtime/Experimental/Chaos/Private/Chaos/CollisionResolution.cpp
+++ b/Engine/Source/Runtime/Experimental/Chaos/Private/Chaos/CollisionResolution.cpp
@@ -237,18 +237,19 @@ namespace Chaos
 			}
 		}
 
-		void UpdateContactPoint(FRigidBodyPointContactConstraint& Constraint, const FContactPoint& ContactPoint, const FReal CullDistance, const FReal Dt, bool bAllowContactDisable = true)
+		// Add the contact to the manifold (or update the existing point). Disable the contact if the contact distance is greater than the Cull Distance
+		void UpdateContactPoint(FRigidBodyPointContactConstraint& Constraint, const FContactPoint& ContactPoint, const FReal Dt)
 		{
 			// Permanently disable contacts beyond the CullDistance
 #if CHAOS_COLLISION_CREATE_BOUNDSCHECK
-			if (ContactPoint.Phi > CullDistance && bAllowContactDisable)
+			if (ContactPoint.Phi > Constraint.GetCullDistance())
 			{
 				Constraint.SetDisabled(true);
 				return;
 			}
 #endif
 
-			// Ignore points that have not been initialized - if there is no detectable contact 
+			// Ignore points that have not been initialized - i.e., if there is no detectable contact 
 			// point within reasonable range despite passing the AABB tests
 			if (ContactPoint.IsSet())
 			{
@@ -256,6 +257,15 @@ namespace Chaos
 			}
 		}
 
+		// Same as UpdateContact Point but without checking CullDistance. Used by CCD because sweeps do not set the separation unless the sweep actually hits
+		void UpdateContactPointNoCull(FRigidBodyPointContactConstraint& Constraint, const FContactPoint& ContactPoint, const FReal Dt)
+		{
+			if (ContactPoint.IsSet())
+			{
+				Constraint.AddIncrementalManifoldContact(ContactPoint, Dt);
+			}
+		}
+
 
 		template <typename GeometryA, typename GeometryB>
 		FContactPoint GJKContactPoint2(const GeometryA& A, const GeometryB& B, const FRigidTransform3& ATM, const FRigidTransform3& BToATM, const FVec3& InitialDir, const FReal ShapePadding)
@@ -297,7 +307,7 @@ namespace Chaos
 		}
 
 		template <typename GeometryA, typename GeometryB>
-		FContactPoint GJKContactPointSwept(const GeometryA& A, const FRigidTransform3& ATM, const GeometryB& B, const FRigidTransform3& BTM, const FVec3& Dir, const FReal Length, const FReal CullDistance, FReal& TOI)
+		FContactPoint GJKContactPointSwept(const GeometryA& A, const FRigidTransform3& ATM, const GeometryB& B, const FRigidTransform3& BTM, const FVec3& Dir, const FReal Length, FReal& TOI)
 		{
 			FContactPoint Contact;
 			const FRigidTransform3 AToBTM = ATM.GetRelativeTransform(BTM);
@@ -383,7 +393,7 @@ namespace Chaos
 		}
 
 		template <typename GeometryA, typename GeometryB>
-		FContactPoint GJKImplicitSweptContactPoint(const FImplicitObject& A, const FRigidTransform3& AStartTransform, const GeometryB& B, const FRigidTransform3& BTransform, const FVec3& Dir, const FReal Length, const FReal CullDistance, FReal& TOI)
+		FContactPoint GJKImplicitSweptContactPoint(const FImplicitObject& A, const FRigidTransform3& AStartTransform, const GeometryB& B, const FRigidTransform3& BTransform, const FVec3& Dir, const FReal Length, FReal& TOI)
 		{
 			FContactPoint Contact;
 			const FRigidTransform3 AToBTM = AStartTransform.GetRelativeTransform(BTransform);
@@ -415,7 +425,7 @@ namespace Chaos
 
 		// A is the implicit here, we want to return a contact point on B (trimesh)
 		template <typename GeometryA>
-		FContactPoint GJKImplicitScaledTriMeshSweptContactPoint(const FImplicitObject& A, const FRigidTransform3& AStartTransform, const TImplicitObjectScaled<FTriangleMeshImplicitObject>& B, const FRigidTransform3& BTransform, const FVec3& Dir, const FReal Length, const FReal CullDistance, FReal& TOI)
+		FContactPoint GJKImplicitScaledTriMeshSweptContactPoint(const FImplicitObject& A, const FRigidTransform3& AStartTransform, const TImplicitObjectScaled<FTriangleMeshImplicitObject>& B, const FRigidTransform3& BTransform, const FVec3& Dir, const FReal Length, FReal& TOI)
 		{
 			FContactPoint Contact;
 			const FRigidTransform3 AToBTM = AStartTransform.GetRelativeTransform(BTransform);
@@ -522,9 +532,9 @@ namespace Chaos
 		struct TConvexConvexUpdater
 		{
 			template<typename T_ImplicitA, typename T_ImplicitB>
-			static void UpdateConvexConvexConstraint(const T_ImplicitA& A, const FRigidTransform3& ATM, const T_ImplicitB& B, const FRigidTransform3& BTM, const FReal CullDistance, const FReal Dt, const FReal ShapePadding, FRigidBodyPointContactConstraint& Constraint)
+			static void UpdateConvexConvexConstraint(const T_ImplicitA& A, const FRigidTransform3& ATM, const T_ImplicitB& B, const FRigidTransform3& BTM, const FReal Dt, const FReal ShapePadding, FRigidBodyPointContactConstraint& Constraint)
 			{
-				UpdateContactPoint(Constraint, GJKContactPoint(A, ATM, B, BTM, FVec3(1, 0, 0), ShapePadding), CullDistance, Dt);
+				UpdateContactPoint(Constraint, GJKContactPoint(A, ATM, B, BTM, FVec3(1, 0, 0), ShapePadding), Dt);
 			}
 		};
 
@@ -533,14 +543,14 @@ namespace Chaos
 		struct TConvexConvexUpdater<true>
 		{
 			template<typename T_ImplicitA, typename T_ImplicitB>
-			static void UpdateConvexConvexConstraint(const T_ImplicitA& A, const FRigidTransform3& ATM, const T_ImplicitB& B, const FRigidTransform3& BTM, const FReal CullDistance, const FReal Dt, const FReal ShapePadding, FRigidBodyPointContactConstraint& Constraint)
+			static void UpdateConvexConvexConstraint(const T_ImplicitA& A, const FRigidTransform3& ATM, const T_ImplicitB& B, const FRigidTransform3& BTM, const FReal Dt, const FReal ShapePadding, FRigidBodyPointContactConstraint& Constraint)
 			{
 				if (Constraint.GetUseManifold())
 				{
 					// We only build one shot manifolds once, and from then on we just find the deepest point
 					if (Constraint.GetManifoldPoints().Num() == 0)
 					{
-						ConstructConvexConvexOneShotManifold(A, ATM, B, BTM, CullDistance, Dt, Constraint);
+						ConstructConvexConvexOneShotManifold(A, ATM, B, BTM, Dt, Constraint);
 					}
 					else
 					{
@@ -549,7 +559,7 @@ namespace Chaos
 				}
 				else
 				{
-					UpdateContactPoint(Constraint, GJKContactPoint(A, ATM, B, BTM, FVec3(1, 0, 0), ShapePadding), CullDistance, Dt);
+					UpdateContactPoint(Constraint, GJKContactPoint(A, ATM, B, BTM, FVec3(1, 0, 0), ShapePadding), Dt);
 				}
 			}
 		};
@@ -559,28 +569,28 @@ namespace Chaos
 		struct FConvexConvexUpdaterCaller
 		{
 			template<typename T_ImplicitA, typename T_ImplicitB>
-			static void Update(const T_ImplicitA& A, const FRigidTransform3& ATM, const T_ImplicitB& B, const FRigidTransform3& BTM, const FReal CullDistance, const FReal Dt, const FReal ShapePadding, FRigidBodyPointContactConstraint& Constraint)
+			static void Update(const T_ImplicitA& A, const FRigidTransform3& ATM, const T_ImplicitB& B, const FRigidTransform3& BTM, const FReal Dt, const FReal ShapePadding, FRigidBodyPointContactConstraint& Constraint)
 			{
 				using FConvexConvexUpdater = TConvexConvexUpdater<TConvexImplicitPairTraits<T_ImplicitA, T_ImplicitB>::bSupportsOneShotManifold>;
-				FConvexConvexUpdater::UpdateConvexConvexConstraint(A, ATM, B, BTM, CullDistance, Dt, ShapePadding, Constraint);
+				FConvexConvexUpdater::UpdateConvexConvexConstraint(A, ATM, B, BTM, Dt, ShapePadding, Constraint);
 			}
 		};
 
 
 		// Unwrap the many convex types, including scaled, and call the appropriate update which depends on the concrete types
-		void UpdateGenericConvexConvexConstraintHelper(const FImplicitObject& A, const FRigidTransform3& ATM, const FImplicitObject& B, const FRigidTransform3& BTM, const FReal CullDistance, const FReal Dt, const FReal ShapePadding, FRigidBodyPointContactConstraint& Constraint)
+		void UpdateGenericConvexConvexConstraintHelper(const FImplicitObject& A, const FRigidTransform3& ATM, const FImplicitObject& B, const FRigidTransform3& BTM, const FReal Dt, const FReal ShapePadding, FRigidBodyPointContactConstraint& Constraint)
 		{
 			// This expands to a switch of switches that calls the inner function with the appropriate concrete implicit types
 			Utilities::CastHelperNoUnwrap(A, ATM, [&](const auto& ADowncast, const FRigidTransform3& AFullTM)
 			{
 				Utilities::CastHelperNoUnwrap(B, BTM, [&](const auto& BDowncast, const FRigidTransform3& BFullTM)
 				{
-					FConvexConvexUpdaterCaller::Update(ADowncast, AFullTM, BDowncast, BFullTM, CullDistance, Dt, ShapePadding, Constraint);
+					FConvexConvexUpdaterCaller::Update(ADowncast, AFullTM, BDowncast, BFullTM, Dt, ShapePadding, Constraint);
 				});
 			});
 		}
 
-		FContactPoint GenericConvexConvexContactPoint(const FImplicitObject& A, const FRigidTransform3& ATM, const FImplicitObject& B, const FRigidTransform3& BTM, const FReal CullDistance, const FReal ShapePadding)
+		FContactPoint GenericConvexConvexContactPoint(const FImplicitObject& A, const FRigidTransform3& ATM, const FImplicitObject& B, const FRigidTransform3& BTM, const FReal ShapePadding)
 		{
 			// This expands to a switch of switches that calls the inner function with the appropriate concrete implicit types
 			return Utilities::CastHelperNoUnwrap(A, ATM, [&](const auto& ADowncast, const FRigidTransform3& AFullTM)
@@ -592,14 +602,14 @@ namespace Chaos
 			});
 		}
 
-		FContactPoint GenericConvexConvexContactPointSwept(const FImplicitObject& A, const FRigidTransform3& ATM, const FImplicitObject& B, const FRigidTransform3& BTM, const FVec3& Dir, const FReal Length, const FReal CullDistance, FReal& TOI)
+		FContactPoint GenericConvexConvexContactPointSwept(const FImplicitObject& A, const FRigidTransform3& ATM, const FImplicitObject& B, const FRigidTransform3& BTM, const FVec3& Dir, const FReal Length, FReal& TOI)
 		{
 			// This expands to a switch of switches that calls the inner function with the appropriate concrete implicit types
 			return Utilities::CastHelperNoUnwrap(A, ATM, [&](const auto& ADowncast, const FRigidTransform3& AFullTM)
 			{
 				return Utilities::CastHelperNoUnwrap(B, BTM, [&](const auto& BDowncast, const FRigidTransform3& BFullTM)
 				{
-					return GJKContactPointSwept(ADowncast, AFullTM, BDowncast, BFullTM, Dir, Length, CullDistance, TOI);
+					return GJKContactPointSwept(ADowncast, AFullTM, BDowncast, BFullTM, Dir, Length, TOI);
 				});
 			});
 		}
@@ -622,19 +632,19 @@ namespace Chaos
 		// Box - Box
 		//
 
-		FContactPoint BoxBoxContactPoint(const FImplicitBox3& Box1, const FImplicitBox3& Box2, const FRigidTransform3& Box1TM, const FRigidTransform3& Box2TM, const FReal CullDistance, const FReal ShapePadding)
+		FContactPoint BoxBoxContactPoint(const FImplicitBox3& Box1, const FImplicitBox3& Box2, const FRigidTransform3& Box1TM, const FRigidTransform3& Box2TM, const FReal ShapePadding)
 		{
 			return GJKContactPoint(Box1, Box1TM, Box2, Box2TM, FVec3(1, 0, 0), ShapePadding);
 		}
 
-		void UpdateBoxBoxConstraint(const FImplicitBox3& Box1, const FRigidTransform3& Box1Transform, const FImplicitBox3& Box2, const FRigidTransform3& Box2Transform, const FReal CullDistance, const FReal Dt, FRigidBodyPointContactConstraint& Constraint)
+		void UpdateBoxBoxConstraint(const FImplicitBox3& Box1, const FRigidTransform3& Box1Transform, const FImplicitBox3& Box2, const FRigidTransform3& Box2Transform, const FReal Dt, FRigidBodyPointContactConstraint& Constraint)
 		{
 			if (Constraint.GetUseManifold())
 			{
 				// We only build one shot manifolds once, and from then on we just find the deepest point
 				if (Constraint.GetManifoldPoints().Num() == 0)
 				{
-					ConstructBoxBoxOneShotManifold(Box1, Box1Transform, Box2, Box2Transform, CullDistance, Dt, Constraint);
+					ConstructBoxBoxOneShotManifold(Box1, Box1Transform, Box2, Box2Transform, Dt, Constraint);
 				}
 				else
 				{
@@ -643,7 +653,7 @@ namespace Chaos
 			}
 			else
 			{
-				UpdateContactPoint(Constraint, BoxBoxContactPoint(Box1, Box2, Box1Transform, Box2Transform, CullDistance, Constraint.Manifold.RestitutionPadding), CullDistance, Dt);
+				UpdateContactPoint(Constraint, BoxBoxContactPoint(Box1, Box2, Box1Transform, Box2Transform, Constraint.Manifold.RestitutionPadding), Dt);
 			}
 		}
 
@@ -654,12 +664,12 @@ namespace Chaos
 			const TBox<FReal, 3>* Object1 = Implicit1->template GetObject<const TBox<FReal, 3> >();
 			if (ensure(Object0 && Object1))
 			{
-				FRigidBodyPointContactConstraint Constraint = FRigidBodyPointContactConstraint(Particle0, Implicit0, nullptr, LocalTransform0, Particle1, Implicit1, nullptr, LocalTransform1, EContactShapesType::BoxBox, Context.bAllowManifolds);
+				FRigidBodyPointContactConstraint Constraint = FRigidBodyPointContactConstraint(Particle0, Implicit0, nullptr, LocalTransform0, Particle1, Implicit1, nullptr, LocalTransform1, CullDistance, EContactShapesType::BoxBox, Context.bAllowManifolds);
 				if (T_TRAITS::bImmediateUpdate)
 				{
 					FRigidTransform3 WorldTransform0 = LocalTransform0 * Collisions::GetTransform(Particle0);
 					FRigidTransform3 WorldTransform1 = LocalTransform1 * Collisions::GetTransform(Particle1);
-					UpdateBoxBoxConstraint(*Object0, WorldTransform0, *Object1, WorldTransform1, CullDistance, Dt, Constraint);
+					UpdateBoxBoxConstraint(*Object0, WorldTransform0, *Object1, WorldTransform1, Dt, Constraint);
 				}
 				NewConstraints.Add(Constraint);
 			}
@@ -676,10 +686,10 @@ namespace Chaos
 		}
 
 
-		void UpdateBoxHeightFieldConstraint(const FImplicitBox3& A, const FRigidTransform3& ATransform, const FHeightField& B, const FRigidTransform3& BTransform, const FReal CullDistance, const FReal Dt, FRigidBodyPointContactConstraint& Constraint)
+		void UpdateBoxHeightFieldConstraint(const FImplicitBox3& A, const FRigidTransform3& ATransform, const FHeightField& B, const FRigidTransform3& BTransform, const FReal Dt, FRigidBodyPointContactConstraint& Constraint)
 		{
 			// @todo(chaos): restitutionpadding
-			UpdateContactPoint(Constraint, BoxHeightFieldContactPoint(A, ATransform, B, BTransform, CullDistance, 0.0f), CullDistance, Dt);
+			UpdateContactPoint(Constraint, BoxHeightFieldContactPoint(A, ATransform, B, BTransform, Constraint.GetCullDistance(), 0.0f), Dt);
 		}
 
 		template<typename T_TRAITS>
@@ -691,12 +701,12 @@ namespace Chaos
 			if (ensure(Object0 && Object1))
 			{
 				// @todo(chaos): one-shot manifold
-				FRigidBodyPointContactConstraint Constraint = FRigidBodyPointContactConstraint(Particle0, Implicit0, nullptr, LocalTransform0, Particle1, Implicit1, nullptr, LocalTransform1, EContactShapesType::BoxHeightField, Context.bAllowManifolds);
+				FRigidBodyPointContactConstraint Constraint = FRigidBodyPointContactConstraint(Particle0, Implicit0, nullptr, LocalTransform0, Particle1, Implicit1, nullptr, LocalTransform1, CullDistance, EContactShapesType::BoxHeightField, Context.bAllowManifolds);
 				if (T_TRAITS::bImmediateUpdate)
 				{
 					FRigidTransform3 WorldTransform0 = LocalTransform0 * Collisions::GetTransform(Particle0);
 					FRigidTransform3 WorldTransform1 = LocalTransform1 * Collisions::GetTransform(Particle1);
-					UpdateBoxHeightFieldConstraint(*Object0, WorldTransform0, *Object1, WorldTransform1, CullDistance, Dt, Constraint);
+					UpdateBoxHeightFieldConstraint(*Object0, WorldTransform0, *Object1, WorldTransform1, Dt, Constraint);
 				}
 				NewConstraints.Add(Constraint);
 			}
@@ -708,7 +718,7 @@ namespace Chaos
 		// Box-Plane
 		//
 
-		FContactPoint BoxPlaneContactPoint(const FImplicitBox3& Box, const FImplicitPlane3& Plane, const FRigidTransform3& BoxTransform, const FRigidTransform3& PlaneTransform, const FReal CullDistance, const FReal ShapePadding)
+		FContactPoint BoxPlaneContactPoint(const FImplicitBox3& Box, const FImplicitPlane3& Plane, const FRigidTransform3& BoxTransform, const FRigidTransform3& PlaneTransform, const FReal ShapePadding)
 		{
 			FContactPoint Contact;
 
@@ -773,11 +783,11 @@ namespace Chaos
 			return Contact;
 		}
 
-		void UpdateBoxPlaneConstraint(const FImplicitBox3& A, const FRigidTransform3& ATransform, const TPlane<FReal, 3>& B, const FRigidTransform3& BTransform, const FReal CullDistance, const FReal Dt, FRigidBodyPointContactConstraint& Constraint)
+		void UpdateBoxPlaneConstraint(const FImplicitBox3& A, const FRigidTransform3& ATransform, const TPlane<FReal, 3>& B, const FRigidTransform3& BTransform, const FReal Dt, FRigidBodyPointContactConstraint& Constraint)
 		{
 			// @todo(chaos): one-shot manifold
 			// @todo(chaos): restitutionpadding
-			UpdateContactPoint(Constraint, BoxPlaneContactPoint(A, B, ATransform, BTransform, CullDistance, 0.0f), CullDistance, Dt);
+			UpdateContactPoint(Constraint, BoxPlaneContactPoint(A, B, ATransform, BTransform, 0.0f), Dt);
 		}
 
 
@@ -788,12 +798,12 @@ namespace Chaos
 			const TPlane<FReal, 3>* Object1 = Implicit1->template GetObject<const TPlane<FReal, 3> >();
 			if (ensure(Object0 && Object1))
 			{
-				FRigidBodyPointContactConstraint Constraint = FRigidBodyPointContactConstraint(Particle0, Implicit0, nullptr, LocalTransform0, Particle1, Implicit1, nullptr, LocalTransform1, EContactShapesType::BoxPlane, Context.bAllowManifolds);
+				FRigidBodyPointContactConstraint Constraint = FRigidBodyPointContactConstraint(Particle0, Implicit0, nullptr, LocalTransform0, Particle1, Implicit1, nullptr, LocalTransform1, CullDistance, EContactShapesType::BoxPlane, Context.bAllowManifolds);
 				if (T_TRAITS::bImmediateUpdate)
 				{
 					FRigidTransform3 WorldTransform0 = LocalTransform0 * Collisions::GetTransform(Particle0);
 					FRigidTransform3 WorldTransform1 = LocalTransform1 * Collisions::GetTransform(Particle1);
-					UpdateBoxPlaneConstraint(*Object0, WorldTransform0, *Object1, WorldTransform1, CullDistance, Dt, Constraint);
+					UpdateBoxPlaneConstraint(*Object0, WorldTransform0, *Object1, WorldTransform1, Dt, Constraint);
 				}
 				NewConstraints.Add(Constraint);
 			}
@@ -810,10 +820,10 @@ namespace Chaos
 		}
 
 		template <typename TriMeshType>
-		void UpdateBoxTriangleMeshConstraint(const FImplicitBox3& Box0, const FRigidTransform3& WorldTransform0, const TriMeshType& TriangleMesh1, const FRigidTransform3& WorldTransform1, const FReal CullDistance, const FReal Dt, FRigidBodyPointContactConstraint& Constraint)
+		void UpdateBoxTriangleMeshConstraint(const FImplicitBox3& Box0, const FRigidTransform3& WorldTransform0, const TriMeshType& TriangleMesh1, const FRigidTransform3& WorldTransform1, const FReal Dt, FRigidBodyPointContactConstraint& Constraint)
 		{
 			// @toto(chaos): restitutionpadding
-			UpdateContactPoint(Constraint, BoxTriangleMeshContactPoint(Box0, WorldTransform0, TriangleMesh1, WorldTransform1, CullDistance, 0.0f), CullDistance, Dt);
+			UpdateContactPoint(Constraint, BoxTriangleMeshContactPoint(Box0, WorldTransform0, TriangleMesh1, WorldTransform1, Constraint.GetCullDistance(), 0.0f), Dt);
 		}
 
 		template<typename T_TRAITS>
@@ -825,24 +835,24 @@ namespace Chaos
 				if (const TImplicitObjectScaled<FTriangleMeshImplicitObject>* ScaledTriangleMesh = Implicit1->template GetObject<const TImplicitObjectScaled<FTriangleMeshImplicitObject>>())
 				{
 					// @todo(chaos): one-shot manifold
-					FRigidBodyPointContactConstraint Constraint = FRigidBodyPointContactConstraint(Particle0, Implicit0, nullptr, LocalTransform0, Particle1, Implicit1, nullptr, LocalTransform1, EContactShapesType::BoxTriMesh, Context.bAllowManifolds);
+					FRigidBodyPointContactConstraint Constraint = FRigidBodyPointContactConstraint(Particle0, Implicit0, nullptr, LocalTransform0, Particle1, Implicit1, nullptr, LocalTransform1, CullDistance, EContactShapesType::BoxTriMesh, Context.bAllowManifolds);
 					if (T_TRAITS::bImmediateUpdate)
 					{
 						FRigidTransform3 WorldTransform0 = LocalTransform0 * Collisions::GetTransform(Particle0);
 						FRigidTransform3 WorldTransform1 = LocalTransform1 * Collisions::GetTransform(Particle1);
-						UpdateBoxTriangleMeshConstraint(*Object0, WorldTransform0, *ScaledTriangleMesh, WorldTransform1, CullDistance, Dt, Constraint);
+						UpdateBoxTriangleMeshConstraint(*Object0, WorldTransform0, *ScaledTriangleMesh, WorldTransform1, Dt, Constraint);
 					}
 					NewConstraints.Add(Constraint);
 				}
 				else if (const FTriangleMeshImplicitObject* TriangleMesh = Implicit1->template GetObject<const FTriangleMeshImplicitObject>())
 				{
 					// @todo(chaos): one-shot manifold
-					FRigidBodyPointContactConstraint Constraint = FRigidBodyPointContactConstraint(Particle0, Implicit0, nullptr, LocalTransform0, Particle1, Implicit1, nullptr, LocalTransform1, EContactShapesType::BoxTriMesh, Context.bAllowManifolds);
+					FRigidBodyPointContactConstraint Constraint = FRigidBodyPointContactConstraint(Particle0, Implicit0, nullptr, LocalTransform0, Particle1, Implicit1, nullptr, LocalTransform1, CullDistance, EContactShapesType::BoxTriMesh, Context.bAllowManifolds);
 					if (T_TRAITS::bImmediateUpdate)
 					{
 						FRigidTransform3 WorldTransform0 = LocalTransform0 * Collisions::GetTransform(Particle0);
 						FRigidTransform3 WorldTransform1 = LocalTransform1 * Collisions::GetTransform(Particle1);
-						UpdateBoxTriangleMeshConstraint(*Object0, WorldTransform0, *TriangleMesh, WorldTransform1, CullDistance, Dt, Constraint);
+						UpdateBoxTriangleMeshConstraint(*Object0, WorldTransform0, *TriangleMesh, WorldTransform1, Dt, Constraint);
 					}
 					NewConstraints.Add(Constraint);
 				}
@@ -858,7 +868,7 @@ namespace Chaos
 		//
 
 
-		FContactPoint SphereSphereContactPoint(const TSphere<FReal, 3>& Sphere1, const FRigidTransform3& Sphere1Transform, const TSphere<FReal, 3>& Sphere2, const FRigidTransform3& Sphere2Transform, const FReal CullDistance, const FReal ShapePadding)
+		FContactPoint SphereSphereContactPoint(const TSphere<FReal, 3>& Sphere1, const FRigidTransform3& Sphere1Transform, const TSphere<FReal, 3>& Sphere2, const FRigidTransform3& Sphere2Transform, const FReal ShapePadding)
 		{
 			FContactPoint Result;
 
@@ -887,9 +897,9 @@ namespace Chaos
 			return Result;
 		}
 
-		void UpdateSphereSphereConstraint(const TSphere<FReal, 3>& Sphere1, const FRigidTransform3& Sphere1Transform, const TSphere<FReal, 3>& Sphere2, const FRigidTransform3& Sphere2Transform, const FReal CullDistance, const FReal Dt, FRigidBodyPointContactConstraint& Constraint)
+		void UpdateSphereSphereConstraint(const TSphere<FReal, 3>& Sphere1, const FRigidTransform3& Sphere1Transform, const TSphere<FReal, 3>& Sphere2, const FRigidTransform3& Sphere2Transform, const FReal Dt, FRigidBodyPointContactConstraint& Constraint)
 		{
-			UpdateContactPoint(Constraint, SphereSphereContactPoint(Sphere1, Sphere1Transform, Sphere2, Sphere2Transform, CullDistance, Constraint.Manifold.RestitutionPadding), CullDistance, Dt);
+			UpdateContactPoint(Constraint, SphereSphereContactPoint(Sphere1, Sphere1Transform, Sphere2, Sphere2Transform, Constraint.Manifold.RestitutionPadding), Dt);
 		}
 
 		template<typename T_TRAITS>
@@ -899,12 +909,12 @@ namespace Chaos
 			const TSphere<FReal, 3>* Object1 = Implicit1->template GetObject<const TSphere<FReal, 3> >();
 			if (ensure(Object0 && Object1))
 			{
-				FRigidBodyPointContactConstraint Constraint = FRigidBodyPointContactConstraint(Particle0, Implicit0, nullptr, LocalTransform0, Particle1, Implicit1, nullptr, LocalTransform1, EContactShapesType::SphereSphere, Context.bAllowManifolds);
+				FRigidBodyPointContactConstraint Constraint = FRigidBodyPointContactConstraint(Particle0, Implicit0, nullptr, LocalTransform0, Particle1, Implicit1, nullptr, LocalTransform1, CullDistance, EContactShapesType::SphereSphere, Context.bAllowManifolds);
 				if (T_TRAITS::bImmediateUpdate)
 				{
 					FRigidTransform3 WorldTransform0 = LocalTransform0 * Collisions::GetTransform(Particle0);
 					FRigidTransform3 WorldTransform1 = LocalTransform1 * Collisions::GetTransform(Particle1);
-					UpdateSphereSphereConstraint(*Object0, WorldTransform0, *Object1, WorldTransform1, CullDistance, Dt, Constraint);
+					UpdateSphereSphereConstraint(*Object0, WorldTransform0, *Object1, WorldTransform1, Dt, Constraint);
 				}
 				NewConstraints.Add(Constraint);
 			}
@@ -920,16 +930,16 @@ namespace Chaos
 		}
 
 
-		void UpdateSphereHeightFieldConstraint(const TSphere<FReal, 3>& A, const FRigidTransform3& ATransform, const FHeightField& B, const FRigidTransform3& BTransform, const FReal CullDistance, const FReal Dt, FRigidBodyPointContactConstraint& Constraint)
+		void UpdateSphereHeightFieldConstraint(const TSphere<FReal, 3>& A, const FRigidTransform3& ATransform, const FHeightField& B, const FRigidTransform3& BTransform, const FReal Dt, FRigidBodyPointContactConstraint& Constraint)
 		{
 			// @todo(chaos): restitutionpadding
-			UpdateContactPoint(Constraint, SphereHeightFieldContactPoint(A, ATransform, B, BTransform, CullDistance, 0.0f), CullDistance, Dt);
+			UpdateContactPoint(Constraint, SphereHeightFieldContactPoint(A, ATransform, B, BTransform, Constraint.GetCullDistance(), 0.0f), Dt);
 		}
 
-		void UpdateSphereHeightFieldConstraintSwept(TGeometryParticleHandle<FReal, 3>* Particle0, const TSphere<FReal, 3>& A, const FRigidTransform3& ATransform, const FHeightField& B, const FRigidTransform3& BTransform, const FVec3& Dir, const FReal Length, const FReal CullDistance, const FReal Dt, FRigidBodySweptPointContactConstraint& Constraint)
+		void UpdateSphereHeightFieldConstraintSwept(TGeometryParticleHandle<FReal, 3>* Particle0, const TSphere<FReal, 3>& A, const FRigidTransform3& ATransform, const FHeightField& B, const FRigidTransform3& BTransform, const FVec3& Dir, const FReal Length, const FReal Dt, FRigidBodySweptPointContactConstraint& Constraint)
 		{
 			FReal TOI = 1.0f;
-			UpdateContactPoint(Constraint, GJKImplicitSweptContactPoint<TSphere<FReal, 3>>(A, ATransform, B, BTransform, Dir, Length, CullDistance, TOI), CullDistance, Dt);
+			UpdateContactPoint(Constraint, GJKImplicitSweptContactPoint<TSphere<FReal, 3>>(A, ATransform, B, BTransform, Dir, Length, TOI), Dt);
 			SetSweptConstraintTOI(Particle0, TOI, Length, Dir, Constraint);
 		}
 
@@ -941,12 +951,12 @@ namespace Chaos
 			const FHeightField* Object1 = Implicit1->template GetObject<const FHeightField >();
 			if (ensure(Object0 && Object1))
 			{
-				FRigidBodyPointContactConstraint Constraint = FRigidBodyPointContactConstraint(Particle0, Implicit0, nullptr, LocalTransform0, Particle1, Implicit1, nullptr, LocalTransform1, EContactShapesType::SphereHeightField, Context.bAllowManifolds);
+				FRigidBodyPointContactConstraint Constraint = FRigidBodyPointContactConstraint(Particle0, Implicit0, nullptr, LocalTransform0, Particle1, Implicit1, nullptr, LocalTransform1, CullDistance, EContactShapesType::SphereHeightField, Context.bAllowManifolds);
 				if (T_TRAITS::bImmediateUpdate)
 				{
 					FRigidTransform3 WorldTransform0 = LocalTransform0 * Collisions::GetTransform(Particle0);
 					FRigidTransform3 WorldTransform1 = LocalTransform1 * Collisions::GetTransform(Particle1);
-					UpdateSphereHeightFieldConstraint(*Object0, WorldTransform0, *Object1, WorldTransform1, CullDistance, Dt, Constraint);
+					UpdateSphereHeightFieldConstraint(*Object0, WorldTransform0, *Object1, WorldTransform1, Dt, Constraint);
 				}
 				NewConstraints.Add(Constraint);
 			}
@@ -958,10 +968,10 @@ namespace Chaos
 			const FHeightField* Object1 = Implicit1->template GetObject<const FHeightField>();
 			if (ensure(Object0 && Object1))
 			{
-				FRigidBodySweptPointContactConstraint Constraint = FRigidBodySweptPointContactConstraint(Particle0, Implicit0, nullptr, LocalTransform0, Particle1, Implicit1, nullptr, LocalTransform1, EContactShapesType::SphereHeightField);
+				FRigidBodySweptPointContactConstraint Constraint = FRigidBodySweptPointContactConstraint(Particle0, Implicit0, nullptr, LocalTransform0, Particle1, Implicit1, nullptr, LocalTransform1, CullDistance, EContactShapesType::SphereHeightField);
 				FRigidTransform3 WorldTransformX0 = LocalTransform0 * FRigidTransform3(Particle0->X(), Particle0->R());
 				FRigidTransform3 WorldTransform1 = LocalTransform1 * Collisions::GetTransform(Particle1);
-				UpdateSphereHeightFieldConstraintSwept(Particle0, *Object0, WorldTransformX0, *Object1, WorldTransform1, Dir, Length, CullDistance, Dt, Constraint);
+				UpdateSphereHeightFieldConstraintSwept(Particle0, *Object0, WorldTransformX0, *Object1, WorldTransform1, Dir, Length, Dt, Constraint);
 				NewConstraints.Add(Constraint);
 			}
 		}
@@ -970,7 +980,7 @@ namespace Chaos
 		//  Sphere-Plane
 		//
 
-		FContactPoint SpherePlaneContactPoint(const TSphere<FReal, 3>& Sphere, const FRigidTransform3& SphereTransform, const TPlane<FReal, 3>& Plane, const FRigidTransform3& PlaneTransform, const FReal CullDistance, const FReal ShapePadding)
+		FContactPoint SpherePlaneContactPoint(const TSphere<FReal, 3>& Sphere, const FRigidTransform3& SphereTransform, const TPlane<FReal, 3>& Plane, const FRigidTransform3& PlaneTransform, const FReal ShapePadding)
 		{
 			FContactPoint Result;
 
@@ -997,9 +1007,9 @@ namespace Chaos
 			return Result;
 		}
 
-		void UpdateSpherePlaneConstraint(const TSphere<FReal, 3>& Sphere, const FRigidTransform3& SphereTransform, const TPlane<FReal, 3>& Plane, const FRigidTransform3& PlaneTransform, const FReal CullDistance, const FReal Dt, FRigidBodyPointContactConstraint& Constraint)
+		void UpdateSpherePlaneConstraint(const TSphere<FReal, 3>& Sphere, const FRigidTransform3& SphereTransform, const TPlane<FReal, 3>& Plane, const FRigidTransform3& PlaneTransform, const FReal Dt, FRigidBodyPointContactConstraint& Constraint)
 		{
-			UpdateContactPoint(Constraint, SpherePlaneContactPoint(Sphere, SphereTransform, Plane, PlaneTransform, CullDistance, Constraint.Manifold.RestitutionPadding), CullDistance, Dt);
+			UpdateContactPoint(Constraint, SpherePlaneContactPoint(Sphere, SphereTransform, Plane, PlaneTransform, Constraint.Manifold.RestitutionPadding), Dt);
 		}
 
 		template<typename T_TRAITS>
@@ -1010,12 +1020,12 @@ namespace Chaos
 			const TPlane<FReal, 3>* Object1 = Implicit1->template GetObject<const TPlane<FReal, 3> >();
 			if (ensure(Object0 && Object1))
 			{
-				FRigidBodyPointContactConstraint Constraint = FRigidBodyPointContactConstraint(Particle0, Implicit0, nullptr, LocalTransform0, Particle1, Implicit1, nullptr, LocalTransform1, EContactShapesType::SpherePlane, Context.bAllowManifolds);
+				FRigidBodyPointContactConstraint Constraint = FRigidBodyPointContactConstraint(Particle0, Implicit0, nullptr, LocalTransform0, Particle1, Implicit1, nullptr, LocalTransform1, CullDistance, EContactShapesType::SpherePlane, Context.bAllowManifolds);
 				if (T_TRAITS::bImmediateUpdate)
 				{
 					FRigidTransform3 WorldTransform0 = LocalTransform0 * Collisions::GetTransform(Particle0);
 					FRigidTransform3 WorldTransform1 = LocalTransform1 * Collisions::GetTransform(Particle1);
-					UpdateSpherePlaneConstraint(*Object0, WorldTransform0, *Object1, WorldTransform1, CullDistance, Dt, Constraint);
+					UpdateSpherePlaneConstraint(*Object0, WorldTransform0, *Object1, WorldTransform1, Dt, Constraint);
 				}
 				NewConstraints.Add(Constraint);
 			}
@@ -1026,7 +1036,7 @@ namespace Chaos
 		//
 
 
-		FContactPoint SphereBoxContactPoint(const TSphere<FReal, 3>& Sphere, const FRigidTransform3& SphereTransform, const FImplicitBox3& Box, const FRigidTransform3& BoxTransform, const FReal CullDistance, const FReal ShapePadding)
+		FContactPoint SphereBoxContactPoint(const TSphere<FReal, 3>& Sphere, const FRigidTransform3& SphereTransform, const FImplicitBox3& Box, const FRigidTransform3& BoxTransform, const FReal ShapePadding)
 		{
 			FContactPoint Result;
 
@@ -1053,9 +1063,9 @@ namespace Chaos
 		}
 
 
-		void UpdateSphereBoxConstraint(const TSphere<FReal, 3>& Sphere, const FRigidTransform3& SphereTransform, const FImplicitBox3& Box, const FRigidTransform3& BoxTransform, const FReal CullDistance, const FReal Dt, FRigidBodyPointContactConstraint& Constraint)
+		void UpdateSphereBoxConstraint(const TSphere<FReal, 3>& Sphere, const FRigidTransform3& SphereTransform, const FImplicitBox3& Box, const FRigidTransform3& BoxTransform, const FReal Dt, FRigidBodyPointContactConstraint& Constraint)
 		{
-			UpdateContactPoint(Constraint, SphereBoxContactPoint(Sphere, SphereTransform, Box, BoxTransform, CullDistance, Constraint.Manifold.RestitutionPadding), CullDistance, Dt);
+			UpdateContactPoint(Constraint, SphereBoxContactPoint(Sphere, SphereTransform, Box, BoxTransform, Constraint.Manifold.RestitutionPadding), Dt);
 		}
 
 		template<typename T_TRAITS>
@@ -1066,12 +1076,12 @@ namespace Chaos
 			const TBox<FReal, 3>* Object1 = Implicit1->template GetObject<const TBox<FReal, 3> >();
 			if (ensure(Object0 && Object1))
 			{
-				FRigidBodyPointContactConstraint Constraint = FRigidBodyPointContactConstraint(Particle0, Implicit0, nullptr, LocalTransform0, Particle1, Implicit1, nullptr, LocalTransform1, EContactShapesType::SphereBox, Context.bAllowManifolds);
+				FRigidBodyPointContactConstraint Constraint = FRigidBodyPointContactConstraint(Particle0, Implicit0, nullptr, LocalTransform0, Particle1, Implicit1, nullptr, LocalTransform1, CullDistance, EContactShapesType::SphereBox, Context.bAllowManifolds);
 				if (T_TRAITS::bImmediateUpdate)
 				{
 					FRigidTransform3 WorldTransform0 = LocalTransform0 * Collisions::GetTransform(Particle0);
 					FRigidTransform3 WorldTransform1 = LocalTransform1 * Collisions::GetTransform(Particle1);
-					UpdateSphereBoxConstraint(*Object0, WorldTransform0, *Object1, WorldTransform1, CullDistance, Dt, Constraint);
+					UpdateSphereBoxConstraint(*Object0, WorldTransform0, *Object1, WorldTransform1, Dt, Constraint);
 				}
 				NewConstraints.Add(Constraint);
 			}
@@ -1082,7 +1092,7 @@ namespace Chaos
 		// Sphere - Capsule
 		//
 
-		FContactPoint SphereCapsuleContactPoint(const TSphere<FReal, 3>& A, const FRigidTransform3& ATransform, const FCapsule& B, const FRigidTransform3& BTransform, const FReal CullDistance, const FReal ShapePadding)
+		FContactPoint SphereCapsuleContactPoint(const TSphere<FReal, 3>& A, const FRigidTransform3& ATransform, const FCapsule& B, const FRigidTransform3& BTransform, const FReal ShapePadding)
 		{
 			FContactPoint Result;
 
@@ -1114,9 +1124,9 @@ namespace Chaos
 		}
 
 
-		void UpdateSphereCapsuleConstraint(const TSphere<FReal, 3>& A, const FRigidTransform3& ATransform, const FCapsule& B, const FRigidTransform3& BTransform, const FReal CullDistance, const FReal Dt, FRigidBodyPointContactConstraint& Constraint)
+		void UpdateSphereCapsuleConstraint(const TSphere<FReal, 3>& A, const FRigidTransform3& ATransform, const FCapsule& B, const FRigidTransform3& BTransform, const FReal Dt, FRigidBodyPointContactConstraint& Constraint)
 		{
-			UpdateContactPoint(Constraint, SphereCapsuleContactPoint(A, ATransform, B, BTransform, CullDistance, Constraint.Manifold.RestitutionPadding), CullDistance, Dt);
+			UpdateContactPoint(Constraint, SphereCapsuleContactPoint(A, ATransform, B, BTransform, Constraint.Manifold.RestitutionPadding), Dt);
 		}
 
 		template<typename T_TRAITS>
@@ -1126,12 +1136,12 @@ namespace Chaos
 			const FCapsule* Object1 = Implicit1->template GetObject<const FCapsule >();
 			if (ensure(Object0 && Object1))
 			{
-				FRigidBodyPointContactConstraint Constraint = FRigidBodyPointContactConstraint(Particle0, Implicit0, nullptr, LocalTransform0, Particle1, Implicit1, nullptr, LocalTransform1, EContactShapesType::SphereCapsule, Context.bAllowManifolds);
+				FRigidBodyPointContactConstraint Constraint = FRigidBodyPointContactConstraint(Particle0, Implicit0, nullptr, LocalTransform0, Particle1, Implicit1, nullptr, LocalTransform1, CullDistance, EContactShapesType::SphereCapsule, Context.bAllowManifolds);
 				if (T_TRAITS::bImmediateUpdate)
 				{
 					FRigidTransform3 WorldTransform0 = LocalTransform0 * Collisions::GetTransform(Particle0);
 					FRigidTransform3 WorldTransform1 = LocalTransform1 * Collisions::GetTransform(Particle1);
-					UpdateSphereCapsuleConstraint(*Object0, WorldTransform0, *Object1, WorldTransform1, CullDistance, Dt, Constraint);
+					UpdateSphereCapsuleConstraint(*Object0, WorldTransform0, *Object1, WorldTransform1, Dt, Constraint);
 				}
 				NewConstraints.Add(Constraint);
 			}
@@ -1141,9 +1151,9 @@ namespace Chaos
 		// Sphere-Convex
 		//
 
-		void UpdateSphereConvexConstraint(const TSphere<FReal, 3>& A, const FRigidTransform3& ATransform, const FImplicitObject3& B, const FRigidTransform3& BTransform, const FReal CullDistance, const FReal Dt, FRigidBodyPointContactConstraint& Constraint)
+		void UpdateSphereConvexConstraint(const TSphere<FReal, 3>& A, const FRigidTransform3& ATransform, const FImplicitObject3& B, const FRigidTransform3& BTransform, const FReal Dt, FRigidBodyPointContactConstraint& Constraint)
 		{
-			UpdateContactPoint(Constraint, SphereConvexContactPoint(A, ATransform, B, BTransform, CullDistance), CullDistance, Dt);
+			UpdateContactPoint(Constraint, SphereConvexContactPoint(A, ATransform, B, BTransform), Dt);
 		}
 
 		template<typename T_TRAITS>
@@ -1153,12 +1163,12 @@ namespace Chaos
 			const FImplicitObject* Object1 = Implicit1;
 			if (ensure(Object0 && Object1))
 			{
-				FRigidBodyPointContactConstraint Constraint = FRigidBodyPointContactConstraint(Particle0, Implicit0, nullptr, LocalTransform0, Particle1, Implicit1, nullptr, LocalTransform1, EContactShapesType::SphereConvex, Context.bAllowManifolds);
+				FRigidBodyPointContactConstraint Constraint = FRigidBodyPointContactConstraint(Particle0, Implicit0, nullptr, LocalTransform0, Particle1, Implicit1, nullptr, LocalTransform1, CullDistance, EContactShapesType::SphereConvex, Context.bAllowManifolds);
 				if (T_TRAITS::bImmediateUpdate)
 				{
 					FRigidTransform3 WorldTransform0 = LocalTransform0 * Collisions::GetTransform(Particle0);
 					FRigidTransform3 WorldTransform1 = LocalTransform1 * Collisions::GetTransform(Particle1);
-					UpdateSphereConvexConstraint(*Object0, WorldTransform0, *Object1, WorldTransform1, CullDistance, Dt, Constraint);
+					UpdateSphereConvexConstraint(*Object0, WorldTransform0, *Object1, WorldTransform1, Dt, Constraint);
 				}
 				NewConstraints.Add(Constraint);
 			}
@@ -1175,15 +1185,15 @@ namespace Chaos
 		}
 
 		template<typename TriMeshType>
-		FContactPoint SphereTriangleMeshSweptContactPoint(const TSphere<FReal, 3>& A, const FRigidTransform3& ATransform, const TriMeshType& B, const FRigidTransform3& BStartTransform, const FVec3& Dir, const FReal Length, const FReal CullDistance, FReal& TOI)
+		FContactPoint SphereTriangleMeshSweptContactPoint(const TSphere<FReal, 3>& A, const FRigidTransform3& ATransform, const TriMeshType& B, const FRigidTransform3& BStartTransform, const FVec3& Dir, const FReal Length, FReal& TOI)
 		{
 			if (const TImplicitObjectScaled<FTriangleMeshImplicitObject>* ScaledTriangleMesh = B.template GetObject<const TImplicitObjectScaled<FTriangleMeshImplicitObject>>())
 			{
-				return GJKImplicitScaledTriMeshSweptContactPoint<TSphere<FReal, 3>>(A, ATransform, *ScaledTriangleMesh, BStartTransform, Dir, Length, CullDistance, TOI);
+				return GJKImplicitScaledTriMeshSweptContactPoint<TSphere<FReal, 3>>(A, ATransform, *ScaledTriangleMesh, BStartTransform, Dir, Length, TOI);
 			}
 			else if (const FTriangleMeshImplicitObject* TriangleMesh = B.template GetObject<const FTriangleMeshImplicitObject>())
 			{
-				return GJKImplicitSweptContactPoint<TSphere<FReal, 3>>(A, ATransform, *TriangleMesh, BStartTransform, Dir, Length, CullDistance, TOI);
+				return GJKImplicitSweptContactPoint<TSphere<FReal, 3>>(A, ATransform, *TriangleMesh, BStartTransform, Dir, Length, TOI);
 			}
 
 			ensure(false);
@@ -1191,17 +1201,17 @@ namespace Chaos
 		}
 
 		template <typename TriMeshType>
-		void UpdateSphereTriangleMeshConstraint(const TSphere<FReal, 3>& Sphere0, const FRigidTransform3& WorldTransform0, const TriMeshType& TriangleMesh1, const FRigidTransform3& WorldTransform1, const FReal CullDistance, const FReal Dt, FRigidBodyPointContactConstraint& Constraint)
+		void UpdateSphereTriangleMeshConstraint(const TSphere<FReal, 3>& Sphere0, const FRigidTransform3& WorldTransform0, const TriMeshType& TriangleMesh1, const FRigidTransform3& WorldTransform1, const FReal Dt, FRigidBodyPointContactConstraint& Constraint)
 		{
 			// @todo(chaos): restitutionpadding
-			UpdateContactPoint(Constraint, SphereTriangleMeshContactPoint(Sphere0, WorldTransform0, TriangleMesh1, WorldTransform1, CullDistance, 0.0f), CullDistance, Dt);
+			UpdateContactPoint(Constraint, SphereTriangleMeshContactPoint(Sphere0, WorldTransform0, TriangleMesh1, WorldTransform1, Constraint.GetCullDistance(), 0.0f), Dt);
 		}
 
 		template<typename TriMeshType>
-		void UpdateSphereTriangleMeshConstraintSwept(TGeometryParticleHandle<FReal, 3>* Particle0, const TSphere<FReal, 3>& Sphere0, const FRigidTransform3& WorldTransform0, const TriMeshType& TriangleMesh1, const FRigidTransform3& WorldTransform1, const FVec3& Dir, const FReal Length, const FReal CullDistance, const FReal Dt, FRigidBodySweptPointContactConstraint& Constraint)
+		void UpdateSphereTriangleMeshConstraintSwept(TGeometryParticleHandle<FReal, 3>* Particle0, const TSphere<FReal, 3>& Sphere0, const FRigidTransform3& WorldTransform0, const TriMeshType& TriangleMesh1, const FRigidTransform3& WorldTransform1, const FVec3& Dir, const FReal Length, const FReal Dt, FRigidBodySweptPointContactConstraint& Constraint)
 		{
 			FReal TOI = 1.0f;
-			UpdateContactPoint(Constraint, SphereTriangleMeshSweptContactPoint(Sphere0, WorldTransform0, TriangleMesh1, WorldTransform1, Dir, Length, CullDistance, TOI), CullDistance, Dt);
+			UpdateContactPoint(Constraint, SphereTriangleMeshSweptContactPoint(Sphere0, WorldTransform0, TriangleMesh1, WorldTransform1, Dir, Length, TOI), Dt);
 			SetSweptConstraintTOI(Particle0, TOI, Length, Dir, Constraint);
 		}
 
@@ -1213,23 +1223,23 @@ namespace Chaos
 			{
 				if (const TImplicitObjectScaled<FTriangleMeshImplicitObject>* ScaledTriangleMesh = Implicit1->template GetObject<const TImplicitObjectScaled<FTriangleMeshImplicitObject>>())
 				{
-					FRigidBodyPointContactConstraint Constraint = FRigidBodyPointContactConstraint(Particle0, Implicit0, nullptr, LocalTransform0, Particle1, Implicit1, nullptr, LocalTransform1, EContactShapesType::SphereTriMesh, Context.bAllowManifolds);
+					FRigidBodyPointContactConstraint Constraint = FRigidBodyPointContactConstraint(Particle0, Implicit0, nullptr, LocalTransform0, Particle1, Implicit1, nullptr, LocalTransform1, CullDistance, EContactShapesType::SphereTriMesh, Context.bAllowManifolds);
 					if (T_TRAITS::bImmediateUpdate)
 					{
 						FRigidTransform3 WorldTransform0 = LocalTransform0 * Collisions::GetTransform(Particle0);
 						FRigidTransform3 WorldTransform1 = LocalTransform1 * Collisions::GetTransform(Particle1);
-						UpdateSphereTriangleMeshConstraint(*Object0, WorldTransform0, *ScaledTriangleMesh, WorldTransform1, CullDistance, Dt, Constraint);
+						UpdateSphereTriangleMeshConstraint(*Object0, WorldTransform0, *ScaledTriangleMesh, WorldTransform1, Dt, Constraint);
 					}
 					NewConstraints.Add(Constraint);
 				}
 				else if (const FTriangleMeshImplicitObject* TriangleMesh = Implicit1->template GetObject<const FTriangleMeshImplicitObject>())
 				{
-					FRigidBodyPointContactConstraint Constraint = FRigidBodyPointContactConstraint(Particle0, Implicit0, nullptr, LocalTransform0, Particle1, Implicit1, nullptr, LocalTransform1, EContactShapesType::SphereTriMesh, Context.bAllowManifolds);
+					FRigidBodyPointContactConstraint Constraint = FRigidBodyPointContactConstraint(Particle0, Implicit0, nullptr, LocalTransform0, Particle1, Implicit1, nullptr, LocalTransform1, CullDistance, EContactShapesType::SphereTriMesh, Context.bAllowManifolds);
 					if (T_TRAITS::bImmediateUpdate)
 					{
 						FRigidTransform3 WorldTransform0 = LocalTransform0 * Collisions::GetTransform(Particle0);
 						FRigidTransform3 WorldTransform1 = LocalTransform1 * Collisions::GetTransform(Particle1);
-						UpdateSphereTriangleMeshConstraint(*Object0, WorldTransform0, *TriangleMesh, WorldTransform1, CullDistance, Dt, Constraint);
+						UpdateSphereTriangleMeshConstraint(*Object0, WorldTransform0, *TriangleMesh, WorldTransform1, Dt, Constraint);
 					}
 					NewConstraints.Add(Constraint);
 				}
@@ -1248,18 +1258,18 @@ namespace Chaos
 			{
 				if (const TImplicitObjectScaled<FTriangleMeshImplicitObject>* ScaledTriangleMesh = Implicit1->template GetObject<const TImplicitObjectScaled<FTriangleMeshImplicitObject>>())
 				{
-					FRigidBodySweptPointContactConstraint Constraint = FRigidBodySweptPointContactConstraint(Particle0, Implicit0, nullptr, LocalTransform0, Particle1, Implicit1, nullptr, LocalTransform1, EContactShapesType::SphereTriMesh);
+					FRigidBodySweptPointContactConstraint Constraint = FRigidBodySweptPointContactConstraint(Particle0, Implicit0, nullptr, LocalTransform0, Particle1, Implicit1, nullptr, LocalTransform1, CullDistance, EContactShapesType::SphereTriMesh);
 					FRigidTransform3 WorldTransformX0 = LocalTransform0 * FRigidTransform3(Particle0->X(), Particle0->R());
 					FRigidTransform3 WorldTransform1 = LocalTransform1 * Collisions::GetTransform(Particle1);
-					UpdateSphereTriangleMeshConstraintSwept(Particle0, *Object0, WorldTransformX0, *ScaledTriangleMesh, WorldTransform1, Dir, Length, CullDistance, Dt, Constraint);
+					UpdateSphereTriangleMeshConstraintSwept(Particle0, *Object0, WorldTransformX0, *ScaledTriangleMesh, WorldTransform1, Dir, Length, Dt, Constraint);
 					NewConstraints.Add(Constraint);
 				}
 				else if (const FTriangleMeshImplicitObject* TriangleMesh = Implicit1->template GetObject<const FTriangleMeshImplicitObject>())
 				{
-					FRigidBodySweptPointContactConstraint Constraint = FRigidBodySweptPointContactConstraint(Particle0, Implicit0, nullptr, LocalTransform0, Particle1, Implicit1, nullptr, LocalTransform1, EContactShapesType::SphereTriMesh);
+					FRigidBodySweptPointContactConstraint Constraint = FRigidBodySweptPointContactConstraint(Particle0, Implicit0, nullptr, LocalTransform0, Particle1, Implicit1, nullptr, LocalTransform1, CullDistance, EContactShapesType::SphereTriMesh);
 					FRigidTransform3 WorldTransformX0 = LocalTransform0 * FRigidTransform3(Particle0->X(), Particle0->R());
 					FRigidTransform3 WorldTransform1 = LocalTransform1 * Collisions::GetTransform(Particle1);
-					UpdateSphereTriangleMeshConstraintSwept(Particle0, *Object0, WorldTransformX0, *TriangleMesh, WorldTransform1, Dir, Length, CullDistance, Dt, Constraint);
+					UpdateSphereTriangleMeshConstraintSwept(Particle0, *Object0, WorldTransformX0, *TriangleMesh, WorldTransform1, Dir, Length, Dt, Constraint);
 					NewConstraints.Add(Constraint);
 				}
 				else
@@ -1276,7 +1286,7 @@ namespace Chaos
 		//
 
 
-		FContactPoint CapsuleCapsuleContactPoint(const FCapsule& A, const FRigidTransform3& ATransform, const FCapsule& B, const FRigidTransform3& BTransform, const FReal CullDistance, const FReal ShapePadding)
+		FContactPoint CapsuleCapsuleContactPoint(const FCapsule& A, const FRigidTransform3& ATransform, const FCapsule& B, const FRigidTransform3& BTransform, const FReal ShapePadding)
 		{
 			FContactPoint Result;
 
@@ -1310,9 +1320,9 @@ namespace Chaos
 		}
 
 
-		void UpdateCapsuleCapsuleConstraint(const FCapsule& A, const FRigidTransform3& ATransform, const FCapsule& B, const FRigidTransform3& BTransform, const FReal CullDistance, const FReal Dt, FRigidBodyPointContactConstraint& Constraint)
+		void UpdateCapsuleCapsuleConstraint(const FCapsule& A, const FRigidTransform3& ATransform, const FCapsule& B, const FRigidTransform3& BTransform, const FReal Dt, FRigidBodyPointContactConstraint& Constraint)
 		{
-			UpdateContactPoint(Constraint, CapsuleCapsuleContactPoint(A, ATransform, B, BTransform, CullDistance, Constraint.Manifold.RestitutionPadding), CullDistance, Dt);
+			UpdateContactPoint(Constraint, CapsuleCapsuleContactPoint(A, ATransform, B, BTransform, Constraint.Manifold.RestitutionPadding), Dt);
 		}
 
 		template<typename T_TRAITS>
@@ -1322,12 +1332,12 @@ namespace Chaos
 			const FCapsule* Object1 = Implicit1->template GetObject<const FCapsule >();
 			if (ensure(Object0 && Object1))
 			{
-				FRigidBodyPointContactConstraint Constraint = FRigidBodyPointContactConstraint(Particle0, Implicit0, nullptr, LocalTransform0, Particle1, Implicit1, nullptr, LocalTransform1, EContactShapesType::CapsuleCapsule, Context.bAllowManifolds);
+				FRigidBodyPointContactConstraint Constraint = FRigidBodyPointContactConstraint(Particle0, Implicit0, nullptr, LocalTransform0, Particle1, Implicit1, nullptr, LocalTransform1, CullDistance, EContactShapesType::CapsuleCapsule, Context.bAllowManifolds);
 				if (T_TRAITS::bImmediateUpdate)
 				{
 					FRigidTransform3 WorldTransform0 = LocalTransform0 * Collisions::GetTransform(Particle0);
 					FRigidTransform3 WorldTransform1 = LocalTransform1 * Collisions::GetTransform(Particle1);
-					UpdateCapsuleCapsuleConstraint(*Object0, WorldTransform0, *Object1, WorldTransform1, CullDistance, Dt, Constraint);
+					UpdateCapsuleCapsuleConstraint(*Object0, WorldTransform0, *Object1, WorldTransform1, Dt, Constraint);
 				}
 				NewConstraints.Add(Constraint);
 			}
@@ -1338,25 +1348,16 @@ namespace Chaos
 		//
 
 
-		FContactPoint CapsuleBoxContactPoint(const FCapsule& A, const FRigidTransform3& ATransform, const FImplicitBox3& B, const FRigidTransform3& BTransform, const FVec3& InitialDir, const FReal CullDistance, const FReal ShapePadding)
+		FContactPoint CapsuleBoxContactPoint(const FCapsule& A, const FRigidTransform3& ATransform, const FImplicitBox3& B, const FRigidTransform3& BTransform, const FVec3& InitialDir, const FReal ShapePadding)
 		{
 			return GJKContactPoint(A, ATransform, B, BTransform, InitialDir, ShapePadding);
 		}
 
 
-		void UpdateCapsuleBoxConstraint(const FCapsule& A, const FRigidTransform3& ATransform, const FImplicitBox3& B, const FRigidTransform3& BTransform, const FReal CullDistance, const FReal Dt, FRigidBodyPointContactConstraint& Constraint)
+		void UpdateCapsuleBoxConstraint(const FCapsule& A, const FRigidTransform3& ATransform, const FImplicitBox3& B, const FRigidTransform3& BTransform, const FReal Dt, FRigidBodyPointContactConstraint& Constraint)
 		{
-			// @todo(chaos): the inverse of CapsuleToBoxTM is calculated in GJKContactPoint - try to eliminate this one
-			const FRigidTransform3 CapsuleToBoxTM = ATransform.GetRelativeTransform(BTransform);
-			const FVec3 P1 = CapsuleToBoxTM.TransformPosition(A.GetX1());
-			const FVec3 P2 = CapsuleToBoxTM.TransformPosition(A.GetX2());
-			FAABB3 CapsuleAABB(P1.ComponentMin(P2), P1.ComponentMax(P2));
-			CapsuleAABB.Thicken(A.GetRadius() + CullDistance);
-			if (CapsuleAABB.Intersects(B.BoundingBox()))
-			{
-				const FVec3 InitialDir = ATransform.GetRotation().Inverse() * -Constraint.GetNormal();
-				UpdateContactPoint(Constraint, CapsuleBoxContactPoint(A, ATransform, B, BTransform, InitialDir, CullDistance, Constraint.Manifold.RestitutionPadding), CullDistance, Dt);
-			}
+			const FVec3 InitialDir = ATransform.GetRotation().Inverse() * -Constraint.GetNormal();
+			UpdateContactPoint(Constraint, CapsuleBoxContactPoint(A, ATransform, B, BTransform, InitialDir, Constraint.Manifold.RestitutionPadding), Dt);
 		}
 
 		// @todo(chaos): make this work with the new one-shot manifold
@@ -1572,12 +1573,12 @@ namespace Chaos
 			const TBox<FReal, 3>* Object1 = Implicit1->template GetObject<const TBox<FReal, 3>>();
 			if (ensure(Object0 && Object1))
 			{
-				FRigidBodyPointContactConstraint Constraint = FRigidBodyPointContactConstraint(Particle0, Implicit0, nullptr, LocalTransform0, Particle1, Implicit1, nullptr, LocalTransform1, EContactShapesType::CapsuleBox, Context.bAllowManifolds);
+				FRigidBodyPointContactConstraint Constraint = FRigidBodyPointContactConstraint(Particle0, Implicit0, nullptr, LocalTransform0, Particle1, Implicit1, nullptr, LocalTransform1, CullDistance, EContactShapesType::CapsuleBox, Context.bAllowManifolds);
 				if (T_TRAITS::bImmediateUpdate)
 				{
 					FRigidTransform3 WorldTransform0 = LocalTransform0 * Collisions::GetTransform(Particle0);
 					FRigidTransform3 WorldTransform1 = LocalTransform1 * Collisions::GetTransform(Particle1);
-					UpdateCapsuleBoxConstraint(*Object0, WorldTransform0, *Object1, WorldTransform1, CullDistance, Dt, Constraint);
+					UpdateCapsuleBoxConstraint(*Object0, WorldTransform0, *Object1, WorldTransform1, Dt, Constraint);
 				}
 				NewConstraints.Add(Constraint);
 			}
@@ -1588,9 +1589,9 @@ namespace Chaos
 		// Capsule-Convex
 		//
 
-		void UpdateCapsuleConvexConstraint(const FCapsule& A, const FRigidTransform3& ATransform, const FImplicitObject& B, const FRigidTransform3& BTransform, const FReal CullDistance, const FReal Dt, FRigidBodyPointContactConstraint& Constraint)
+		void UpdateCapsuleConvexConstraint(const FCapsule& A, const FRigidTransform3& ATransform, const FImplicitObject& B, const FRigidTransform3& BTransform, const FReal Dt, FRigidBodyPointContactConstraint& Constraint)
 		{
-			UpdateContactPoint(Constraint, CapsuleConvexContactPoint(A, ATransform, B, BTransform, CullDistance), CullDistance, Dt);
+			UpdateContactPoint(Constraint, CapsuleConvexContactPoint(A, ATransform, B, BTransform), Dt);
 		}
 
 		template<typename T_TRAITS>
@@ -1600,12 +1601,12 @@ namespace Chaos
 			const FImplicitObject* Object1 = Implicit1;
 			if (ensure(Object0 && Object1))
 			{
-				FRigidBodyPointContactConstraint Constraint = FRigidBodyPointContactConstraint(Particle0, Implicit0, nullptr, LocalTransform0, Particle1, Implicit1, nullptr, LocalTransform1, EContactShapesType::CapsuleConvex, Context.bAllowManifolds);
+				FRigidBodyPointContactConstraint Constraint = FRigidBodyPointContactConstraint(Particle0, Implicit0, nullptr, LocalTransform0, Particle1, Implicit1, nullptr, LocalTransform1, CullDistance, EContactShapesType::CapsuleConvex, Context.bAllowManifolds);
 				if (T_TRAITS::bImmediateUpdate)
 				{
 					FRigidTransform3 WorldTransform0 = LocalTransform0 * Collisions::GetTransform(Particle0);
 					FRigidTransform3 WorldTransform1 = LocalTransform1 * Collisions::GetTransform(Particle1);
-					UpdateCapsuleConvexConstraint(*Object0, WorldTransform0, *Object1, WorldTransform1, CullDistance, Dt, Constraint);
+					UpdateCapsuleConvexConstraint(*Object0, WorldTransform0, *Object1, WorldTransform1, Dt, Constraint);
 				}
 				NewConstraints.Add(Constraint);
 			}
@@ -1624,19 +1625,19 @@ namespace Chaos
 		}
 
 
-		void UpdateCapsuleHeightFieldConstraint(const FCapsule& A, const FRigidTransform3& ATransform, const FHeightField& B, const FRigidTransform3& BTransform, const FReal CullDistance, const FReal Dt, FRigidBodyPointContactConstraint& Constraint)
+		void UpdateCapsuleHeightFieldConstraint(const FCapsule& A, const FRigidTransform3& ATransform, const FHeightField& B, const FRigidTransform3& BTransform, const FReal Dt, FRigidBodyPointContactConstraint& Constraint)
 		{
 			// @todo(chaos): restitutionpadding
 			CONDITIONAL_SCOPE_CYCLE_COUNTER(STAT_Collisions_UpdateCapsuleHeightFieldConstraint, ConstraintsDetailedStats);
-			UpdateContactPoint(Constraint, CapsuleHeightFieldContactPoint(A, ATransform, B, BTransform, CullDistance, 0.0f), CullDistance, Dt);
+			UpdateContactPoint(Constraint, CapsuleHeightFieldContactPoint(A, ATransform, B, BTransform, Constraint.GetCullDistance(), 0.0f), Dt);
 		}
 
 
-		void UpdateCapsuleHeightFieldConstraintSwept(TGeometryParticleHandle<FReal, 3>* Particle0, const FCapsule& A, const FRigidTransform3& ATransform, const FHeightField& B, const FRigidTransform3& BTransform, const FVec3& Dir, const FReal Length, const FReal CullDistance, const FReal Dt, FRigidBodySweptPointContactConstraint& Constraint)
+		void UpdateCapsuleHeightFieldConstraintSwept(TGeometryParticleHandle<FReal, 3>* Particle0, const FCapsule& A, const FRigidTransform3& ATransform, const FHeightField& B, const FRigidTransform3& BTransform, const FVec3& Dir, const FReal Length, const FReal Dt, FRigidBodySweptPointContactConstraint& Constraint)
 		{
 			CONDITIONAL_SCOPE_CYCLE_COUNTER(STAT_Collisions_UpdateCapsuleHeightFieldConstraintSwept, ConstraintsDetailedStats);
 			FReal TOI = 1.0f;
-			UpdateContactPoint(Constraint, GJKImplicitSweptContactPoint<FCapsule >(A, ATransform, B, BTransform, Dir, Length, CullDistance, TOI), CullDistance, Dt);
+			UpdateContactPoint(Constraint, GJKImplicitSweptContactPoint<FCapsule >(A, ATransform, B, BTransform, Dir, Length, TOI), Dt);
 			SetSweptConstraintTOI(Particle0, TOI, Length, Dir, Constraint);
 		}
 
@@ -1649,12 +1650,12 @@ namespace Chaos
 			const FHeightField* Object1 = Implicit1->template GetObject<const FHeightField >();
 			if (ensure(Object0 && Object1))
 			{
-				FRigidBodyPointContactConstraint Constraint = FRigidBodyPointContactConstraint(Particle0, Implicit0, nullptr, LocalTransform0, Particle1, Implicit1, nullptr, LocalTransform1, EContactShapesType::CapsuleHeightField, Context.bAllowManifolds);
+				FRigidBodyPointContactConstraint Constraint = FRigidBodyPointContactConstraint(Particle0, Implicit0, nullptr, LocalTransform0, Particle1, Implicit1, nullptr, LocalTransform1, CullDistance, EContactShapesType::CapsuleHeightField, Context.bAllowManifolds);
 				if (T_TRAITS::bImmediateUpdate)
 				{
 					FRigidTransform3 WorldTransform0 = LocalTransform0 * Collisions::GetTransform(Particle0);
 					FRigidTransform3 WorldTransform1 = LocalTransform1 * Collisions::GetTransform(Particle1);
-					UpdateCapsuleHeightFieldConstraint(*Object0, WorldTransform0, *Object1, WorldTransform1, CullDistance, Dt, Constraint);
+					UpdateCapsuleHeightFieldConstraint(*Object0, WorldTransform0, *Object1, WorldTransform1, Dt, Constraint);
 				}
 				NewConstraints.Add(Constraint);
 			}
@@ -1667,10 +1668,10 @@ namespace Chaos
 			const FHeightField* Object1 = Implicit1->template GetObject<const FHeightField >();
 			if (ensure(Object0 && Object1))
 			{
-				FRigidBodySweptPointContactConstraint Constraint = FRigidBodySweptPointContactConstraint(Particle0, Implicit0, nullptr, LocalTransform0, Particle1, Implicit1, nullptr, LocalTransform1, EContactShapesType::CapsuleHeightField);
+				FRigidBodySweptPointContactConstraint Constraint = FRigidBodySweptPointContactConstraint(Particle0, Implicit0, nullptr, LocalTransform0, Particle1, Implicit1, nullptr, LocalTransform1, CullDistance, EContactShapesType::CapsuleHeightField);
 				FRigidTransform3 WorldTransformX0 = LocalTransform0 * FRigidTransform3(Particle0->X(), Particle0->R());
 				FRigidTransform3 WorldTransform1 = LocalTransform1 * Collisions::GetTransform(Particle1);
-				UpdateCapsuleHeightFieldConstraintSwept(Particle0, *Object0, WorldTransformX0, *Object1, WorldTransform1, Dir, Length, CullDistance, Dt, Constraint);
+				UpdateCapsuleHeightFieldConstraintSwept(Particle0, *Object0, WorldTransformX0, *Object1, WorldTransform1, Dir, Length, Dt, Constraint);
 				NewConstraints.Add(Constraint);
 			}
 		}
@@ -1688,16 +1689,16 @@ namespace Chaos
 		}
 
 		template <typename TriMeshType>
-		FContactPoint CapsuleTriangleMeshSweptContactPoint(const FCapsule& A, const FRigidTransform3& ATransform, const TriMeshType& B, const FRigidTransform3& BStartTransform, const FVec3& Dir, const FReal Length, const FReal CullDistance, FReal& TOI)
+		FContactPoint CapsuleTriangleMeshSweptContactPoint(const FCapsule& A, const FRigidTransform3& ATransform, const TriMeshType& B, const FRigidTransform3& BStartTransform, const FVec3& Dir, const FReal Length, FReal& TOI)
 		{
 			CONDITIONAL_SCOPE_CYCLE_COUNTER(STAT_Collisions_CapsuleTriangleMeshSweptContactPoint, ConstraintsDetailedStats);
 			if (const TImplicitObjectScaled<FTriangleMeshImplicitObject>* ScaledTriangleMesh = B.template GetObject<const TImplicitObjectScaled<FTriangleMeshImplicitObject>>())
 			{
-				return GJKImplicitScaledTriMeshSweptContactPoint<FCapsule>(A, ATransform, *ScaledTriangleMesh, BStartTransform, Dir, Length, CullDistance, TOI);
+				return GJKImplicitScaledTriMeshSweptContactPoint<FCapsule>(A, ATransform, *ScaledTriangleMesh, BStartTransform, Dir, Length, TOI);
 			}
 			else if (const FTriangleMeshImplicitObject* TriangleMesh = B.template GetObject<const FTriangleMeshImplicitObject>())
 			{
-				return GJKImplicitSweptContactPoint<FCapsule>(A, ATransform, *TriangleMesh, BStartTransform, Dir, Length, CullDistance, TOI);
+				return GJKImplicitSweptContactPoint<FCapsule>(A, ATransform, *TriangleMesh, BStartTransform, Dir, Length, TOI);
 			}
 
 			ensure(false);
@@ -1706,19 +1707,19 @@ namespace Chaos
 
 
 		template <typename TriMeshType>
-		void UpdateCapsuleTriangleMeshConstraint(const FCapsule& Capsule0, const FRigidTransform3& WorldTransform0, const TriMeshType& TriangleMesh1, const FRigidTransform3& WorldTransform1, const FReal CullDistance, const FReal Dt, FRigidBodyPointContactConstraint& Constraint)
+		void UpdateCapsuleTriangleMeshConstraint(const FCapsule& Capsule0, const FRigidTransform3& WorldTransform0, const TriMeshType& TriangleMesh1, const FRigidTransform3& WorldTransform1, const FReal Dt, FRigidBodyPointContactConstraint& Constraint)
 		{
 			// @todo(chaos): restitutionpadding
 			CONDITIONAL_SCOPE_CYCLE_COUNTER(STAT_Collisions_UpdateCapsuleTriangleMeshConstraint, ConstraintsDetailedStats);
-			UpdateContactPoint(Constraint, CapsuleTriangleMeshContactPoint(Capsule0, WorldTransform0, TriangleMesh1, WorldTransform1, CullDistance, 0.0f), CullDistance, Dt);
+			UpdateContactPoint(Constraint, CapsuleTriangleMeshContactPoint(Capsule0, WorldTransform0, TriangleMesh1, WorldTransform1, Constraint.GetCullDistance(), 0.0f), Dt);
 		}
 
 		template <typename TriMeshType>
-		void UpdateCapsuleTriangleMeshConstraintSwept(TGeometryParticleHandle<FReal, 3>* Particle0, const FCapsule& Capsule0, const FRigidTransform3& WorldTransform0, const TriMeshType& TriangleMesh1, const FRigidTransform3& WorldTransform1, const FVec3& Dir, const FReal Length, const FReal CullDistance, const FReal Dt, FRigidBodySweptPointContactConstraint& Constraint)
+		void UpdateCapsuleTriangleMeshConstraintSwept(TGeometryParticleHandle<FReal, 3>* Particle0, const FCapsule& Capsule0, const FRigidTransform3& WorldTransform0, const TriMeshType& TriangleMesh1, const FRigidTransform3& WorldTransform1, const FVec3& Dir, const FReal Length, const FReal Dt, FRigidBodySweptPointContactConstraint& Constraint)
 		{
 			CONDITIONAL_SCOPE_CYCLE_COUNTER(STAT_Collisions_UpdateCapsuleTriangleMeshConstraint, ConstraintsDetailedStats);
 			FReal TOI = 1.0f;
-			UpdateContactPoint(Constraint, CapsuleTriangleMeshSweptContactPoint(Capsule0, WorldTransform0, TriangleMesh1, WorldTransform1, Dir, Length, CullDistance, TOI), CullDistance, Dt);
+			UpdateContactPoint(Constraint, CapsuleTriangleMeshSweptContactPoint(Capsule0, WorldTransform0, TriangleMesh1, WorldTransform1, Dir, Length, TOI), Dt);
 			SetSweptConstraintTOI(Particle0, TOI, Length, Dir, Constraint);
 		}
 
@@ -1732,23 +1733,23 @@ namespace Chaos
 			{
 				if (const TImplicitObjectScaled<FTriangleMeshImplicitObject>* ScaledTriangleMesh = Implicit1->template GetObject<const TImplicitObjectScaled<FTriangleMeshImplicitObject>>())
 				{
-					FRigidBodyPointContactConstraint Constraint = FRigidBodyPointContactConstraint(Particle0, Implicit0, nullptr, LocalTransform0, Particle1, Implicit1, nullptr, LocalTransform1, EContactShapesType::CapsuleTriMesh, Context.bAllowManifolds);
+					FRigidBodyPointContactConstraint Constraint = FRigidBodyPointContactConstraint(Particle0, Implicit0, nullptr, LocalTransform0, Particle1, Implicit1, nullptr, LocalTransform1, CullDistance, EContactShapesType::CapsuleTriMesh, Context.bAllowManifolds);
 					if (T_TRAITS::bImmediateUpdate)
 					{
 						FRigidTransform3 WorldTransform0 = LocalTransform0 * Collisions::GetTransform(Particle0);
 						FRigidTransform3 WorldTransform1 = LocalTransform1 * Collisions::GetTransform(Particle1);
-						UpdateCapsuleTriangleMeshConstraint(*Object0, WorldTransform0, *ScaledTriangleMesh, WorldTransform1, CullDistance, Dt, Constraint);
+						UpdateCapsuleTriangleMeshConstraint(*Object0, WorldTransform0, *ScaledTriangleMesh, WorldTransform1, Dt, Constraint);
 					}
 					NewConstraints.Add(Constraint);
 				}
 				else if (const FTriangleMeshImplicitObject* TriangleMesh = Implicit1->template GetObject<const FTriangleMeshImplicitObject>())
 				{
-					FRigidBodyPointContactConstraint Constraint = FRigidBodyPointContactConstraint(Particle0, Implicit0, nullptr, LocalTransform0, Particle1, Implicit1, nullptr, LocalTransform1, EContactShapesType::CapsuleTriMesh, Context.bAllowManifolds);
+					FRigidBodyPointContactConstraint Constraint = FRigidBodyPointContactConstraint(Particle0, Implicit0, nullptr, LocalTransform0, Particle1, Implicit1, nullptr, LocalTransform1, CullDistance, EContactShapesType::CapsuleTriMesh, Context.bAllowManifolds);
 					if (T_TRAITS::bImmediateUpdate)
 					{
 						FRigidTransform3 WorldTransform0 = LocalTransform0 * Collisions::GetTransform(Particle0);
 						FRigidTransform3 WorldTransform1 = LocalTransform1 * Collisions::GetTransform(Particle1);
-						UpdateCapsuleTriangleMeshConstraint(*Object0, WorldTransform0, *TriangleMesh, WorldTransform1, CullDistance, Dt, Constraint);
+						UpdateCapsuleTriangleMeshConstraint(*Object0, WorldTransform0, *TriangleMesh, WorldTransform1, Dt, Constraint);
 					}
 					NewConstraints.Add(Constraint);
 				}
@@ -1768,18 +1769,18 @@ namespace Chaos
 			{
 				if (const TImplicitObjectScaled<FTriangleMeshImplicitObject>* ScaledTriangleMesh = Implicit1->template GetObject<const TImplicitObjectScaled<FTriangleMeshImplicitObject>>())
 				{
-					FRigidBodySweptPointContactConstraint Constraint = FRigidBodySweptPointContactConstraint(Particle0, Implicit0, nullptr, LocalTransform0, Particle1, Implicit1, nullptr, LocalTransform1, EContactShapesType::CapsuleTriMesh);
+					FRigidBodySweptPointContactConstraint Constraint = FRigidBodySweptPointContactConstraint(Particle0, Implicit0, nullptr, LocalTransform0, Particle1, Implicit1, nullptr, LocalTransform1, CullDistance, EContactShapesType::CapsuleTriMesh);
 					FRigidTransform3 WorldTransformX0 = LocalTransform0 * FRigidTransform3(Particle0->X(), Particle0->R());
 					FRigidTransform3 WorldTransform1 = LocalTransform1 * Collisions::GetTransform(Particle1);
-					UpdateCapsuleTriangleMeshConstraintSwept(Particle0, *Object0, WorldTransformX0, *ScaledTriangleMesh, WorldTransform1, Dir, Length, CullDistance, Dt, Constraint);
+					UpdateCapsuleTriangleMeshConstraintSwept(Particle0, *Object0, WorldTransformX0, *ScaledTriangleMesh, WorldTransform1, Dir, Length, Dt, Constraint);
 					NewConstraints.Add(Constraint);
 				}
 				else if (const FTriangleMeshImplicitObject* TriangleMesh = Implicit1->template GetObject<const FTriangleMeshImplicitObject>())
 				{
-					FRigidBodySweptPointContactConstraint Constraint = FRigidBodySweptPointContactConstraint(Particle0, Implicit0, nullptr, LocalTransform0, Particle1, Implicit1, nullptr, LocalTransform1, EContactShapesType::CapsuleTriMesh);
+					FRigidBodySweptPointContactConstraint Constraint = FRigidBodySweptPointContactConstraint(Particle0, Implicit0, nullptr, LocalTransform0, Particle1, Implicit1, nullptr, LocalTransform1, CullDistance, EContactShapesType::CapsuleTriMesh);
 					FRigidTransform3 WorldTransformX0 = LocalTransform0 * FRigidTransform3(Particle0->X(), Particle0->R());
 					FRigidTransform3 WorldTransform1 = LocalTransform1 * Collisions::GetTransform(Particle1);
-					UpdateCapsuleTriangleMeshConstraintSwept(Particle0, *Object0, WorldTransformX0, *TriangleMesh, WorldTransform1, Dir, Length, CullDistance, Dt, Constraint);
+					UpdateCapsuleTriangleMeshConstraintSwept(Particle0, *Object0, WorldTransformX0, *TriangleMesh, WorldTransform1, Dir, Length, Dt, Constraint);
 					NewConstraints.Add(Constraint);
 				}
 				else
@@ -1793,18 +1794,18 @@ namespace Chaos
 		// Generic Convex - Convex (actualy concrete type could be anything)
 		//
 
-		void UpdateGenericConvexConvexConstraint(const FImplicitObject& Implicit0, const FRigidTransform3& WorldTransform0, const FImplicitObject& Implicit1, const FRigidTransform3& WorldTransform1, const FReal CullDistance, const FReal Dt, FRigidBodyPointContactConstraint& Constraint)
+		void UpdateGenericConvexConvexConstraint(const FImplicitObject& Implicit0, const FRigidTransform3& WorldTransform0, const FImplicitObject& Implicit1, const FRigidTransform3& WorldTransform1, const FReal Dt, FRigidBodyPointContactConstraint& Constraint)
 		{
 			CONDITIONAL_SCOPE_CYCLE_COUNTER(STAT_Collisions_UpdateGenericConvexConvexConstraint, ConstraintsDetailedStats);
 
-			UpdateGenericConvexConvexConstraintHelper(Implicit0, WorldTransform0, Implicit1, WorldTransform1, CullDistance, Dt, Constraint.Manifold.RestitutionPadding, Constraint);
+			UpdateGenericConvexConvexConstraintHelper(Implicit0, WorldTransform0, Implicit1, WorldTransform1, Dt, Constraint.Manifold.RestitutionPadding, Constraint);
 		}
 
-		void UpdateGenericConvexConvexConstraintSwept(TGeometryParticleHandle<FReal, 3>* Particle0, const FImplicitObject& Implicit0, const FRigidTransform3& WorldTransform0, const FImplicitObject& Implicit1, const FRigidTransform3& WorldTransform1, const FVec3& Dir, const FReal Length, const FReal CullDistance, const FReal Dt, FRigidBodySweptPointContactConstraint& Constraint)
+		void UpdateGenericConvexConvexConstraintSwept(TGeometryParticleHandle<FReal, 3>* Particle0, const FImplicitObject& Implicit0, const FRigidTransform3& WorldTransform0, const FImplicitObject& Implicit1, const FRigidTransform3& WorldTransform1, const FVec3& Dir, const FReal Length, const FReal Dt, FRigidBodySweptPointContactConstraint& Constraint)
 		{
 			CONDITIONAL_SCOPE_CYCLE_COUNTER(STAT_Collisions_UpdateGenericConvexConvexConstraintSwept, ConstraintsDetailedStats);
 			FReal TOI = 1.0f;
-			UpdateContactPoint(Constraint, GenericConvexConvexContactPointSwept(Implicit0, WorldTransform0, Implicit1, WorldTransform1, Dir, Length, CullDistance, TOI), CullDistance, Dt, false);
+			UpdateContactPointNoCull(Constraint, GenericConvexConvexContactPointSwept(Implicit0, WorldTransform0, Implicit1, WorldTransform1, Dir, Length, TOI), Dt);
 			SetSweptConstraintTOI(Particle0, TOI, Length, Dir, Constraint);
 		}
 
@@ -1816,12 +1817,12 @@ namespace Chaos
 			EImplicitObjectType Implicit0Type = Particle0->Geometry()->GetType();
 			EImplicitObjectType Implicit1Type = Particle1->Geometry()->GetType();
 
-			FRigidBodyPointContactConstraint Constraint = FRigidBodyPointContactConstraint(Particle0, Implicit0, nullptr, LocalTransform0, Particle1, Implicit1, nullptr, LocalTransform1, EContactShapesType::GenericConvexConvex, Context.bAllowManifolds);
+			FRigidBodyPointContactConstraint Constraint = FRigidBodyPointContactConstraint(Particle0, Implicit0, nullptr, LocalTransform0, Particle1, Implicit1, nullptr, LocalTransform1, CullDistance, EContactShapesType::GenericConvexConvex, Context.bAllowManifolds);
 			if (T_TRAITS::bImmediateUpdate)
 			{
 				FRigidTransform3 WorldTransform0 = LocalTransform0 * Collisions::GetTransform(Particle0);
 				FRigidTransform3 WorldTransform1 = LocalTransform1 * Collisions::GetTransform(Particle1);
-				UpdateGenericConvexConvexConstraint(*Implicit0, WorldTransform0, *Implicit1, WorldTransform1, CullDistance, Dt, Constraint);
+				UpdateGenericConvexConvexConstraint(*Implicit0, WorldTransform0, *Implicit1, WorldTransform1, Dt, Constraint);
 			}
 			NewConstraints.Add(Constraint);
 		}
@@ -1830,12 +1831,12 @@ namespace Chaos
 		{
 			CONDITIONAL_SCOPE_CYCLE_COUNTER(STAT_Collisions_ConstructGenericConvexConvexConstraintsSwept, ConstraintsDetailedStats);
 			FGenericParticleHandle RigidParticle = FGenericParticleHandle(Particle0);
-			FRigidBodySweptPointContactConstraint Constraint = FRigidBodySweptPointContactConstraint(Particle0, Implicit0, nullptr, LocalTransform0, Particle1, Implicit1, nullptr, LocalTransform1, EContactShapesType::GenericConvexConvex);
+			FRigidBodySweptPointContactConstraint Constraint = FRigidBodySweptPointContactConstraint(Particle0, Implicit0, nullptr, LocalTransform0, Particle1, Implicit1, nullptr, LocalTransform1, CullDistance, EContactShapesType::GenericConvexConvex);
 			// Note: This is unusual but we are using a mix of the previous and current transform
 			// This is due to how CCD only rewinds the position
 			FRigidTransform3 WorldTransformXQ0 = LocalTransform0 * FRigidTransform3(RigidParticle->X(), RigidParticle->Q());
 			FRigidTransform3 WorldTransform1 = LocalTransform1 * Collisions::GetTransform(Particle1);
-			UpdateGenericConvexConvexConstraintSwept(Particle0, *Implicit0, WorldTransformXQ0, *Implicit1, WorldTransform1, Dir, Length, CullDistance, Dt, Constraint);
+			UpdateGenericConvexConvexConstraintSwept(Particle0, *Implicit0, WorldTransformXQ0, *Implicit1, WorldTransform1, Dir, Length, Dt, Constraint);
 			NewConstraints.Add(Constraint);
 		}
 
@@ -1851,19 +1852,19 @@ namespace Chaos
 		}
 
 
-		void UpdateConvexHeightFieldConstraint(const FImplicitObject& A, const FRigidTransform3& ATransform, const FHeightField& B, const FRigidTransform3& BTransform, const FReal CullDistance, const FReal Dt, FRigidBodyPointContactConstraint& Constraint)
+		void UpdateConvexHeightFieldConstraint(const FImplicitObject& A, const FRigidTransform3& ATransform, const FHeightField& B, const FRigidTransform3& BTransform, const FReal Dt, FRigidBodyPointContactConstraint& Constraint)
 		{
 			// @todo(chaos): restitutionpadding
 			CONDITIONAL_SCOPE_CYCLE_COUNTER(STAT_Collisions_UpdateConvexHeightFieldConstraint, ConstraintsDetailedStats);
-			UpdateContactPoint(Constraint, ConvexHeightFieldContactPoint(A, ATransform, B, BTransform, CullDistance, 0.0f), CullDistance, Dt);
+			UpdateContactPoint(Constraint, ConvexHeightFieldContactPoint(A, ATransform, B, BTransform, Constraint.GetCullDistance(), 0.0f), Dt);
 		}
 
 
-		void UpdateConvexHeightFieldConstraintSwept(TGeometryParticleHandle<FReal, 3>* Particle0, const FImplicitObject& A, const FRigidTransform3& ATransform, const FHeightField& B, const FRigidTransform3& BTransform, const FVec3& Dir, const FReal Length, const FReal CullDistance, const FReal Dt, FRigidBodySweptPointContactConstraint& Constraint)
+		void UpdateConvexHeightFieldConstraintSwept(TGeometryParticleHandle<FReal, 3>* Particle0, const FImplicitObject& A, const FRigidTransform3& ATransform, const FHeightField& B, const FRigidTransform3& BTransform, const FVec3& Dir, const FReal Length, const FReal Dt, FRigidBodySweptPointContactConstraint& Constraint)
 		{
 			CONDITIONAL_SCOPE_CYCLE_COUNTER(STAT_Collisions_UpdateConvexHeightFieldConstraintSwept, ConstraintsDetailedStats);
 			FReal TOI = 1.0f;
-			UpdateContactPoint(Constraint, GJKImplicitSweptContactPoint< FConvex >(A, ATransform, B, BTransform, Dir, Length, CullDistance, TOI), CullDistance, Dt);
+			UpdateContactPoint(Constraint, GJKImplicitSweptContactPoint< FConvex >(A, ATransform, B, BTransform, Dir, Length, TOI), Dt);
 			SetSweptConstraintTOI(Particle0, TOI, Length, Dir, Constraint);
 		}
 
@@ -1874,12 +1875,12 @@ namespace Chaos
 			const FHeightField* Object1 = Implicit1->template GetObject<const FHeightField >();
 			if (ensure(Implicit0->IsConvex() && Object1))
 			{
-				FRigidBodyPointContactConstraint Constraint = FRigidBodyPointContactConstraint(Particle0, Implicit0, nullptr, LocalTransform0, Particle1, Implicit1, nullptr, LocalTransform1, EContactShapesType::ConvexHeightField, Context.bAllowManifolds);
+				FRigidBodyPointContactConstraint Constraint = FRigidBodyPointContactConstraint(Particle0, Implicit0, nullptr, LocalTransform0, Particle1, Implicit1, nullptr, LocalTransform1, CullDistance, EContactShapesType::ConvexHeightField, Context.bAllowManifolds);
 				if (T_TRAITS::bImmediateUpdate)
 				{
 					FRigidTransform3 WorldTransform0 = LocalTransform0 * Collisions::GetTransform(Particle0);
 					FRigidTransform3 WorldTransform1 = LocalTransform1 * Collisions::GetTransform(Particle1);
-					UpdateConvexHeightFieldConstraint(*Implicit0, WorldTransform0, *Object1, WorldTransform1, CullDistance, Dt, Constraint);
+					UpdateConvexHeightFieldConstraint(*Implicit0, WorldTransform0, *Object1, WorldTransform1, Dt, Constraint);
 				}
 				NewConstraints.Add(Constraint);
 			}
@@ -1891,10 +1892,10 @@ namespace Chaos
 			const FHeightField* Object1 = Implicit1->template GetObject<const FHeightField >();
 			if (ensure(Implicit0->IsConvex() && Object1))
 			{
-				FRigidBodySweptPointContactConstraint Constraint = FRigidBodySweptPointContactConstraint(Particle0, Implicit0, nullptr, LocalTransform0, Particle1, Implicit1, nullptr, LocalTransform1, EContactShapesType::ConvexHeightField);
+				FRigidBodySweptPointContactConstraint Constraint = FRigidBodySweptPointContactConstraint(Particle0, Implicit0, nullptr, LocalTransform0, Particle1, Implicit1, nullptr, LocalTransform1, CullDistance, EContactShapesType::ConvexHeightField);
 				FRigidTransform3 WorldTransformX0 = LocalTransform0 * FRigidTransform3(Particle0->X(), Particle0->R());
 				FRigidTransform3 WorldTransform1 = LocalTransform1 * Collisions::GetTransform(Particle1);
-				UpdateConvexHeightFieldConstraintSwept(Particle0, *Implicit0, WorldTransformX0, *Object1, WorldTransform1, Dir, Length, CullDistance, Dt, Constraint);
+				UpdateConvexHeightFieldConstraintSwept(Particle0, *Implicit0, WorldTransformX0, *Object1, WorldTransform1, Dir, Length, Dt, Constraint);
 				NewConstraints.Add(Constraint);
 			}
 		}
@@ -1911,16 +1912,16 @@ namespace Chaos
 		}
 
 		template <typename TriMeshType>
-		FContactPoint ConvexTriangleMeshSweptContactPoint(const FImplicitObject& A, const FRigidTransform3& ATransform, const TriMeshType& B, const FRigidTransform3& BStartTransform, const FVec3& Dir, const FReal Length, const FReal CullDistance, FReal& TOI)
+		FContactPoint ConvexTriangleMeshSweptContactPoint(const FImplicitObject& A, const FRigidTransform3& ATransform, const TriMeshType& B, const FRigidTransform3& BStartTransform, const FVec3& Dir, const FReal Length, FReal& TOI)
 		{
 			CONDITIONAL_SCOPE_CYCLE_COUNTER(STAT_Collisions_ConvexTriangleMeshSweptContactPoint, ConstraintsDetailedStats);
 			if (const TImplicitObjectScaled<FTriangleMeshImplicitObject>* ScaledTriangleMesh = B.template GetObject<const TImplicitObjectScaled<FTriangleMeshImplicitObject>>())
 			{
-				return GJKImplicitScaledTriMeshSweptContactPoint<FConvex>(A, ATransform, *ScaledTriangleMesh, BStartTransform, Dir, Length, CullDistance, TOI);
+				return GJKImplicitScaledTriMeshSweptContactPoint<FConvex>(A, ATransform, *ScaledTriangleMesh, BStartTransform, Dir, Length, TOI);
 			}
 			else if (const FTriangleMeshImplicitObject* TriangleMesh = B.template GetObject<const FTriangleMeshImplicitObject>())
 			{
-				return GJKImplicitSweptContactPoint<FConvex>(A, ATransform, *TriangleMesh, BStartTransform, Dir, Length, CullDistance, TOI);
+				return GJKImplicitSweptContactPoint<FConvex>(A, ATransform, *TriangleMesh, BStartTransform, Dir, Length, TOI);
 			}
 
 			ensure(false);
@@ -1928,20 +1929,20 @@ namespace Chaos
 		}
 
 		template <typename TriMeshType>
-		void UpdateConvexTriangleMeshConstraint(const FImplicitObject& Convex0, const FRigidTransform3& WorldTransform0, const TriMeshType& TriangleMesh1, const FRigidTransform3& WorldTransform1, const FReal CullDistance, const FReal Dt, FRigidBodyPointContactConstraint& Constraint)
+		void UpdateConvexTriangleMeshConstraint(const FImplicitObject& Convex0, const FRigidTransform3& WorldTransform0, const TriMeshType& TriangleMesh1, const FRigidTransform3& WorldTransform1, const FReal Dt, FRigidBodyPointContactConstraint& Constraint)
 		{
 			// @todo(chaos): restitutionpadding
 			CONDITIONAL_SCOPE_CYCLE_COUNTER(STAT_Collisions_UpdateConvexTriangleMeshConstraint, ConstraintsDetailedStats);
-			UpdateContactPoint(Constraint, ConvexTriangleMeshContactPoint(Convex0, WorldTransform0, TriangleMesh1, WorldTransform1, CullDistance, 0.0f), CullDistance, Dt);
+			UpdateContactPoint(Constraint, ConvexTriangleMeshContactPoint(Convex0, WorldTransform0, TriangleMesh1, WorldTransform1, Constraint.GetCullDistance(), 0.0f), Dt);
 		}
 
 		// Sweeps convex against trimesh
 		template <typename TriMeshType>
-		void UpdateConvexTriangleMeshConstraintSwept(TGeometryParticleHandle<FReal, 3>* Particle0, const FImplicitObject& Convex0, const FRigidTransform3& WorldTransform0, const TriMeshType& TriangleMesh1, const FRigidTransform3& WorldTransform1, const FVec3& Dir, const FReal Length, const FReal CullDistance, const FReal Dt, FRigidBodySweptPointContactConstraint& Constraint)
+		void UpdateConvexTriangleMeshConstraintSwept(TGeometryParticleHandle<FReal, 3>* Particle0, const FImplicitObject& Convex0, const FRigidTransform3& WorldTransform0, const TriMeshType& TriangleMesh1, const FRigidTransform3& WorldTransform1, const FVec3& Dir, const FReal Length, const FReal Dt, FRigidBodySweptPointContactConstraint& Constraint)
 		{
 			CONDITIONAL_SCOPE_CYCLE_COUNTER(STAT_Collisions_UpdateConvexTriangleMeshConstraintSwept, ConstraintsDetailedStats);
 			FReal TOI = 1.0f;
-			UpdateContactPoint(Constraint, ConvexTriangleMeshSweptContactPoint(Convex0, WorldTransform0, TriangleMesh1, WorldTransform1, Dir, Length, CullDistance, TOI), CullDistance, Dt);
+			UpdateContactPoint(Constraint, ConvexTriangleMeshSweptContactPoint(Convex0, WorldTransform0, TriangleMesh1, WorldTransform1, Dir, Length, TOI), Dt);
 			SetSweptConstraintTOI(Particle0, TOI, Length, Dir, Constraint);
 		}
 
@@ -1953,23 +1954,23 @@ namespace Chaos
 			{
 				if (const TImplicitObjectScaled<FTriangleMeshImplicitObject>* ScaledTriangleMesh = Implicit1->template GetObject<const TImplicitObjectScaled<FTriangleMeshImplicitObject>>())
 				{
-					FRigidBodyPointContactConstraint Constraint = FRigidBodyPointContactConstraint(Particle0, Implicit0, nullptr, LocalTransform0, Particle1, Implicit1, nullptr, LocalTransform0, EContactShapesType::ConvexTriMesh, Context.bAllowManifolds);
+					FRigidBodyPointContactConstraint Constraint = FRigidBodyPointContactConstraint(Particle0, Implicit0, nullptr, LocalTransform0, Particle1, Implicit1, nullptr, LocalTransform0, CullDistance, EContactShapesType::ConvexTriMesh, Context.bAllowManifolds);
 					if (T_TRAITS::bImmediateUpdate)
 					{
 						FRigidTransform3 WorldTransform0 = LocalTransform0 * Collisions::GetTransform(Particle0);
 						FRigidTransform3 WorldTransform1 = LocalTransform1 * Collisions::GetTransform(Particle1);
-						UpdateConvexTriangleMeshConstraint(*Implicit0, WorldTransform0, *ScaledTriangleMesh, WorldTransform1, CullDistance, Dt, Constraint);
+						UpdateConvexTriangleMeshConstraint(*Implicit0, WorldTransform0, *ScaledTriangleMesh, WorldTransform1, Dt, Constraint);
 					}
 					NewConstraints.Add(Constraint);
 				}
 				else if (const FTriangleMeshImplicitObject* TriangleMesh = Implicit1->template GetObject<const FTriangleMeshImplicitObject>())
 				{
-					FRigidBodyPointContactConstraint Constraint = FRigidBodyPointContactConstraint(Particle0, Implicit0, nullptr, LocalTransform0, Particle1, Implicit1, nullptr, LocalTransform1, EContactShapesType::ConvexTriMesh, Context.bAllowManifolds);
+					FRigidBodyPointContactConstraint Constraint = FRigidBodyPointContactConstraint(Particle0, Implicit0, nullptr, LocalTransform0, Particle1, Implicit1, nullptr, LocalTransform1, CullDistance, EContactShapesType::ConvexTriMesh, Context.bAllowManifolds);
 					if (T_TRAITS::bImmediateUpdate)
 					{
 						FRigidTransform3 WorldTransform0 = LocalTransform0 * Collisions::GetTransform(Particle0);
 						FRigidTransform3 WorldTransform1 = LocalTransform1 * Collisions::GetTransform(Particle1);
-						UpdateConvexTriangleMeshConstraint(*Implicit0, WorldTransform0, *TriangleMesh, WorldTransform1, CullDistance, Dt, Constraint);
+						UpdateConvexTriangleMeshConstraint(*Implicit0, WorldTransform0, *TriangleMesh, WorldTransform1, Dt, Constraint);
 					}
 					NewConstraints.Add(Constraint);
 				}
@@ -1987,18 +1988,18 @@ namespace Chaos
 			{
 				if (const TImplicitObjectScaled<FTriangleMeshImplicitObject>* ScaledTriangleMesh = Implicit1->template GetObject<const TImplicitObjectScaled<FTriangleMeshImplicitObject>>())
 				{
-					FRigidBodySweptPointContactConstraint Constraint = FRigidBodySweptPointContactConstraint(Particle0, Implicit0, nullptr, LocalTransform0, Particle1, Implicit1, nullptr, LocalTransform1, EContactShapesType::ConvexTriMesh);
+					FRigidBodySweptPointContactConstraint Constraint = FRigidBodySweptPointContactConstraint(Particle0, Implicit0, nullptr, LocalTransform0, Particle1, Implicit1, nullptr, LocalTransform1, CullDistance, EContactShapesType::ConvexTriMesh);
 					FRigidTransform3 WorldTransformX0 = LocalTransform0 * FRigidTransform3(Particle0->X(), Particle0->R());
 					FRigidTransform3 WorldTransform1 = LocalTransform1 * Collisions::GetTransform(Particle1);
-					UpdateConvexTriangleMeshConstraintSwept(Particle0, *Implicit0, WorldTransformX0, *ScaledTriangleMesh, WorldTransform1, Dir, Length, CullDistance, Dt, Constraint);
+					UpdateConvexTriangleMeshConstraintSwept(Particle0, *Implicit0, WorldTransformX0, *ScaledTriangleMesh, WorldTransform1, Dir, Length, Dt, Constraint);
 					NewConstraints.Add(Constraint);
 				}
 				else if (const FTriangleMeshImplicitObject* TriangleMesh = Implicit1->template GetObject<const FTriangleMeshImplicitObject>())
 				{
-					FRigidBodySweptPointContactConstraint Constraint = FRigidBodySweptPointContactConstraint(Particle0, Implicit0, nullptr, LocalTransform0, Particle1, Implicit1, nullptr, LocalTransform1, EContactShapesType::ConvexTriMesh);
+					FRigidBodySweptPointContactConstraint Constraint = FRigidBodySweptPointContactConstraint(Particle0, Implicit0, nullptr, LocalTransform0, Particle1, Implicit1, nullptr, LocalTransform1, CullDistance, EContactShapesType::ConvexTriMesh);
 					FRigidTransform3 WorldTransformX0 = LocalTransform0 * FRigidTransform3(Particle0->X(), Particle0->R());
 					FRigidTransform3 WorldTransform1 = LocalTransform1 * Collisions::GetTransform(Particle1);
-					UpdateConvexTriangleMeshConstraintSwept(Particle0, *Implicit0, WorldTransformX0, *TriangleMesh, WorldTransform1, Dir, Length, CullDistance, Dt, Constraint);
+					UpdateConvexTriangleMeshConstraintSwept(Particle0, *Implicit0, WorldTransformX0, *TriangleMesh, WorldTransform1, Dir, Length, Dt, Constraint);
 					NewConstraints.Add(Constraint);
 				}
 				else
@@ -2014,7 +2015,7 @@ namespace Chaos
 		//
 
 		template<ECollisionUpdateType UpdateType>
-		void UpdateLevelsetLevelsetConstraint(const FRigidTransform3& WorldTransform0, const FRigidTransform3& WorldTransform1, const FReal CullDistance, const FReal Dt, FRigidBodyPointContactConstraint& Constraint)
+		void UpdateLevelsetLevelsetConstraint(const FRigidTransform3& WorldTransform0, const FRigidTransform3& WorldTransform1, const FReal Dt, FRigidBodyPointContactConstraint& Constraint)
 		{
 			SCOPE_CYCLE_COUNTER(STAT_UpdateLevelsetLevelsetConstraint);
 
@@ -2041,7 +2042,7 @@ namespace Chaos
 			if (SampleParticles)
 			{
 				const FImplicitObject* Obj1 = Constraint.Manifold.Implicit[1];
-				FContactPoint ContactPoint = SampleObject<UpdateType>(*Obj1, LevelsetTM, *SampleParticles, ParticlesTM, CullDistance);
+				FContactPoint ContactPoint = SampleObject<UpdateType>(*Obj1, LevelsetTM, *SampleParticles, ParticlesTM, Constraint.GetCullDistance());
 
 				if (ContactPoint.IsSet())
 				{
@@ -2053,11 +2054,11 @@ namespace Chaos
 					ContactPoint.ShapeContactNormal = WorldTransform1.InverseTransformVector(ContactPoint.Normal);
 				}
 
-				UpdateContactPoint(Constraint, ContactPoint, CullDistance, Dt);
+				UpdateContactPoint(Constraint, ContactPoint, Dt);
 			}
 		}
 
-		void UpdateLevelsetLevelsetManifold(FCollisionConstraintBase& Constraint, const FRigidTransform3& ATM, const FRigidTransform3& BTM, const FReal CullDistance)
+		void UpdateLevelsetLevelsetManifold(FCollisionConstraintBase& Constraint, const FRigidTransform3& ATM, const FRigidTransform3& BTM)
 		{
 			// @todo(chaos) : Stub Update Manifold
 			//   Stub function for updating the manifold prior to the Apply and ApplyPushOut
@@ -2067,7 +2068,7 @@ namespace Chaos
 		template<typename T_TRAITS>
 		void ConstructLevelsetLevelsetConstraints(TGeometryParticleHandle<FReal, 3>* Particle0, TGeometryParticleHandle<FReal, 3>* Particle1, const FImplicitObject* Implicit0, const FBVHParticles* Simplicial0, const FImplicitObject* Implicit1, const FBVHParticles* Simplicial1, const FRigidTransform3& LocalTransform0, const FRigidTransform3& LocalTransform1, const FReal CullDistance, const FReal Dt, const FCollisionContext& Context, FCollisionConstraintsArray& NewConstraints)
 		{
-			FRigidBodyPointContactConstraint Constraint = FRigidBodyPointContactConstraint(Particle0, Implicit0, Simplicial0, LocalTransform0, Particle1, Implicit1, Simplicial1, LocalTransform1, EContactShapesType::LevelSetLevelSet, Context.bAllowManifolds);
+			FRigidBodyPointContactConstraint Constraint = FRigidBodyPointContactConstraint(Particle0, Implicit0, Simplicial0, LocalTransform0, Particle1, Implicit1, Simplicial1, LocalTransform1, CullDistance, EContactShapesType::LevelSetLevelSet, Context.bAllowManifolds);
 
 			bool bIsParticleDynamic0 = Particle0->CastToRigidParticle() && Particle0->ObjectState() == EObjectStateType::Dynamic;
 			int32 P0NumCollisionParticles = Simplicial0 ? Simplicial0->Size() : Particle0->CastToRigidParticle()->CollisionParticlesSize();
@@ -2092,7 +2093,7 @@ namespace Chaos
 			{
 				FRigidTransform3 WorldTransform0 = Constraint.ImplicitTransform[0] * Collisions::GetTransform(Constraint.Particle[0]);
 				FRigidTransform3 WorldTransform1 = Constraint.ImplicitTransform[1] * Collisions::GetTransform(Constraint.Particle[1]);
-				UpdateLevelsetLevelsetConstraint<ECollisionUpdateType::Deepest>(WorldTransform0, WorldTransform1, CullDistance, Dt, Constraint);
+				UpdateLevelsetLevelsetConstraint<ECollisionUpdateType::Deepest>(WorldTransform0, WorldTransform1, Dt, Constraint);
 			}
 			NewConstraints.Add(Constraint);
 		}
@@ -2109,7 +2110,7 @@ namespace Chaos
 		// @todo(chaos): add the missing cases below
 		// @todo(chaos): see use GetInnerObject below - we should try to use the leaf types for all (currently only TriMesh needs this)
 		template<ECollisionUpdateType UpdateType>
-		inline void UpdateConstraintFromGeometryImpl(FRigidBodyPointContactConstraint& Constraint, const FRigidTransform3& WorldTransform0, const FRigidTransform3& WorldTransform1, const FReal CullDistance, const FReal Dt)
+		inline void UpdateConstraintFromGeometryImpl(FRigidBodyPointContactConstraint& Constraint, const FRigidTransform3& WorldTransform0, const FRigidTransform3& WorldTransform1, const FReal Dt)
 		{
 			CONDITIONAL_SCOPE_CYCLE_COUNTER(STAT_Collisions_UpdateConstraintFromGeometryInternal, ConstraintsDetailedStats);
 
@@ -2122,6 +2123,7 @@ namespace Chaos
 
 			if (Implicit0.HasBoundingBox() && Implicit1.HasBoundingBox())
 			{
+				const FReal CullDistance = Constraint.GetCullDistance();
 				if (Chaos_Collision_NarrowPhase_SphereBoundsCheck)
 				{
 					const FReal R1 = Implicit0.BoundingBox().OriginRadius();
@@ -2150,107 +2152,107 @@ namespace Chaos
 			switch (Constraint.Manifold.ShapesType)
 			{
 			case EContactShapesType::SphereSphere:
-				UpdateSphereSphereConstraint(*Implicit0.template GetObject<TSphere<FReal, 3>>(), WorldTransform0, *Implicit1.template GetObject<TSphere<FReal, 3>>(), WorldTransform1, CullDistance, Dt, Constraint);
+				UpdateSphereSphereConstraint(*Implicit0.template GetObject<TSphere<FReal, 3>>(), WorldTransform0, *Implicit1.template GetObject<TSphere<FReal, 3>>(), WorldTransform1, Dt, Constraint);
 				break;
 			case EContactShapesType::SphereCapsule:
-				UpdateSphereCapsuleConstraint(*Implicit0.template GetObject<TSphere<FReal, 3>>(), WorldTransform0, *Implicit1.template GetObject<FCapsule>(), WorldTransform1, CullDistance, Dt, Constraint);
+				UpdateSphereCapsuleConstraint(*Implicit0.template GetObject<TSphere<FReal, 3>>(), WorldTransform0, *Implicit1.template GetObject<FCapsule>(), WorldTransform1, Dt, Constraint);
 				break;
 			case EContactShapesType::SphereBox:
-				UpdateSphereBoxConstraint(*Implicit0.template GetObject<TSphere<FReal, 3>>(), WorldTransform0, *Implicit1.template GetObject<TBox<FReal, 3>>(), WorldTransform1, CullDistance, Dt, Constraint);
+				UpdateSphereBoxConstraint(*Implicit0.template GetObject<TSphere<FReal, 3>>(), WorldTransform0, *Implicit1.template GetObject<TBox<FReal, 3>>(), WorldTransform1, Dt, Constraint);
 				break;
 			case EContactShapesType::SphereConvex:
-				UpdateSphereConvexConstraint(*Implicit0.template GetObject<TSphere<FReal, 3>>(), WorldTransform0, Implicit1, WorldTransform1, CullDistance, Dt, Constraint);
+				UpdateSphereConvexConstraint(*Implicit0.template GetObject<TSphere<FReal, 3>>(), WorldTransform0, Implicit1, WorldTransform1, Dt, Constraint);
 				break;
 			case EContactShapesType::SphereTriMesh:
 				if (const TImplicitObjectScaled<FTriangleMeshImplicitObject>* ScaledTriMesh = Implicit1.template GetObject<const TImplicitObjectScaled<FTriangleMeshImplicitObject> >())
 				{
-					UpdateSphereTriangleMeshConstraint(*Implicit0.template GetObject<TSphere<FReal, 3>>(), WorldTransform0, *ScaledTriMesh, WorldTransform1, CullDistance, Dt, Constraint);
+					UpdateSphereTriangleMeshConstraint(*Implicit0.template GetObject<TSphere<FReal, 3>>(), WorldTransform0, *ScaledTriMesh, WorldTransform1, Dt, Constraint);
 					break;
 				}
 				else if (const FTriangleMeshImplicitObject* TriangleMeshImplicit = Implicit1.template GetObject<const FTriangleMeshImplicitObject>())
 				{
-					UpdateSphereTriangleMeshConstraint(*Implicit0.template GetObject<TSphere<FReal, 3>>(), WorldTransform0, *TriangleMeshImplicit, WorldTransform1, CullDistance, Dt, Constraint);
+					UpdateSphereTriangleMeshConstraint(*Implicit0.template GetObject<TSphere<FReal, 3>>(), WorldTransform0, *TriangleMeshImplicit, WorldTransform1, Dt, Constraint);
 					break;
 				}
 				ensure(false);
 				break;
 			case EContactShapesType::SphereHeightField:
-				UpdateSphereHeightFieldConstraint(*Implicit0.template GetObject<TSphere<FReal, 3>>(), WorldTransform0, *Implicit1.template GetObject< FHeightField >(), WorldTransform1, CullDistance, Dt, Constraint);
+				UpdateSphereHeightFieldConstraint(*Implicit0.template GetObject<TSphere<FReal, 3>>(), WorldTransform0, *Implicit1.template GetObject< FHeightField >(), WorldTransform1, Dt, Constraint);
 				break;
 			case EContactShapesType::SpherePlane:
-				UpdateSpherePlaneConstraint(*Implicit0.template GetObject<TSphere<FReal, 3>>(), WorldTransform0, *Implicit1.template GetObject<TPlane<FReal, 3>>(), WorldTransform1, CullDistance, Dt, Constraint);
+				UpdateSpherePlaneConstraint(*Implicit0.template GetObject<TSphere<FReal, 3>>(), WorldTransform0, *Implicit1.template GetObject<TPlane<FReal, 3>>(), WorldTransform1, Dt, Constraint);
 				break;
 			case EContactShapesType::CapsuleCapsule:
-				UpdateCapsuleCapsuleConstraint(*Implicit0.template GetObject<FCapsule>(), WorldTransform0, *Implicit1.template GetObject<FCapsule>(), WorldTransform1, CullDistance, Dt, Constraint);
+				UpdateCapsuleCapsuleConstraint(*Implicit0.template GetObject<FCapsule>(), WorldTransform0, *Implicit1.template GetObject<FCapsule>(), WorldTransform1, Dt, Constraint);
 				break;
 			case EContactShapesType::CapsuleBox:
-				UpdateCapsuleBoxConstraint(*Implicit0.template GetObject<FCapsule>(), WorldTransform0, *Implicit1.template GetObject<TBox<FReal, 3>>(), WorldTransform1, CullDistance, Dt, Constraint);
+				UpdateCapsuleBoxConstraint(*Implicit0.template GetObject<FCapsule>(), WorldTransform0, *Implicit1.template GetObject<TBox<FReal, 3>>(), WorldTransform1, Dt, Constraint);
 				break;
 			case EContactShapesType::CapsuleConvex:
-				UpdateCapsuleConvexConstraint(*Implicit0.template GetObject<FCapsule>(), WorldTransform0, Implicit1, WorldTransform1, CullDistance, Dt, Constraint);
+				UpdateCapsuleConvexConstraint(*Implicit0.template GetObject<FCapsule>(), WorldTransform0, Implicit1, WorldTransform1, Dt, Constraint);
 				break;
 			case EContactShapesType::CapsuleTriMesh:
 				if (const TImplicitObjectScaled<FTriangleMeshImplicitObject>* ScaledTriMesh = Implicit1.template GetObject<const TImplicitObjectScaled<FTriangleMeshImplicitObject> >())
 				{
-					UpdateCapsuleTriangleMeshConstraint(*Implicit0.template GetObject<FCapsule>(), WorldTransform0, *ScaledTriMesh, WorldTransform1, CullDistance, Dt, Constraint);
+					UpdateCapsuleTriangleMeshConstraint(*Implicit0.template GetObject<FCapsule>(), WorldTransform0, *ScaledTriMesh, WorldTransform1, Dt, Constraint);
 					break;
 				}
 				else if (const FTriangleMeshImplicitObject* TriangleMeshImplicit = Implicit1.template GetObject<const FTriangleMeshImplicitObject>())
 				{
-					UpdateCapsuleTriangleMeshConstraint(*Implicit0.template GetObject<FCapsule>(), WorldTransform0, *TriangleMeshImplicit, WorldTransform1, CullDistance, Dt, Constraint);
+					UpdateCapsuleTriangleMeshConstraint(*Implicit0.template GetObject<FCapsule>(), WorldTransform0, *TriangleMeshImplicit, WorldTransform1, Dt, Constraint);
 					break;
 				}
 				ensure(false);
 				break;
 			case EContactShapesType::CapsuleHeightField:
-				UpdateCapsuleHeightFieldConstraint(*Implicit0.template GetObject<FCapsule>(), WorldTransform0, *Implicit1.template GetObject< FHeightField >(), WorldTransform1, CullDistance, Dt, Constraint);
+				UpdateCapsuleHeightFieldConstraint(*Implicit0.template GetObject<FCapsule>(), WorldTransform0, *Implicit1.template GetObject< FHeightField >(), WorldTransform1, Dt, Constraint);
 				break;
 			case EContactShapesType::BoxBox:
-				UpdateBoxBoxConstraint(*Implicit0.template GetObject<TBox<FReal, 3>>(), WorldTransform0, *Implicit1.template GetObject<TBox<FReal, 3>>(), WorldTransform1, CullDistance, Dt, Constraint);
+				UpdateBoxBoxConstraint(*Implicit0.template GetObject<TBox<FReal, 3>>(), WorldTransform0, *Implicit1.template GetObject<TBox<FReal, 3>>(), WorldTransform1, Dt, Constraint);
 				break;
 			case EContactShapesType::BoxConvex:
-				UpdateGenericConvexConvexConstraint(Implicit0, WorldTransform0, Implicit1, WorldTransform1, CullDistance, Dt, Constraint);
+				UpdateGenericConvexConvexConstraint(Implicit0, WorldTransform0, Implicit1, WorldTransform1, Dt, Constraint);
 				break;
 			case EContactShapesType::BoxTriMesh:
 				if (const TImplicitObjectScaled<FTriangleMeshImplicitObject>* ScaledTriMesh = Implicit1.template GetObject<const TImplicitObjectScaled<FTriangleMeshImplicitObject> >())
 				{
-					UpdateBoxTriangleMeshConstraint(*Implicit0.template GetObject<TBox<FReal, 3>>(), WorldTransform0, *ScaledTriMesh, WorldTransform1, CullDistance, Dt, Constraint);
+					UpdateBoxTriangleMeshConstraint(*Implicit0.template GetObject<TBox<FReal, 3>>(), WorldTransform0, *ScaledTriMesh, WorldTransform1, Dt, Constraint);
 					break;
 				}
 				else if (const FTriangleMeshImplicitObject* TriangleMeshImplicit = Implicit1.template GetObject<const FTriangleMeshImplicitObject>())
 				{
-					UpdateBoxTriangleMeshConstraint(*Implicit0.template GetObject<TBox<FReal, 3>>(), WorldTransform0, *TriangleMeshImplicit, WorldTransform1, CullDistance, Dt, Constraint);
+					UpdateBoxTriangleMeshConstraint(*Implicit0.template GetObject<TBox<FReal, 3>>(), WorldTransform0, *TriangleMeshImplicit, WorldTransform1, Dt, Constraint);
 					break;
 				}
 				ensure(false);
 				break;
 			case EContactShapesType::BoxHeightField:
-				UpdateBoxHeightFieldConstraint(*Implicit0.template GetObject<TBox<FReal, 3>>(), WorldTransform0, *Implicit1.template GetObject< FHeightField >(), WorldTransform1, CullDistance, Dt, Constraint);
+				UpdateBoxHeightFieldConstraint(*Implicit0.template GetObject<TBox<FReal, 3>>(), WorldTransform0, *Implicit1.template GetObject< FHeightField >(), WorldTransform1, Dt, Constraint);
 				break;
 			case EContactShapesType::BoxPlane:
-				UpdateBoxPlaneConstraint(*Implicit0.template GetObject<TBox<FReal, 3>>(), WorldTransform0, *Implicit1.template GetObject<TPlane<FReal, 3>>(), WorldTransform1, CullDistance, Dt, Constraint);
+				UpdateBoxPlaneConstraint(*Implicit0.template GetObject<TBox<FReal, 3>>(), WorldTransform0, *Implicit1.template GetObject<TPlane<FReal, 3>>(), WorldTransform1, Dt, Constraint);
 				break;
 			case EContactShapesType::GenericConvexConvex:
-				UpdateGenericConvexConvexConstraint(Implicit0, WorldTransform0, Implicit1, WorldTransform1, CullDistance, Dt, Constraint);
+				UpdateGenericConvexConvexConstraint(Implicit0, WorldTransform0, Implicit1, WorldTransform1, Dt, Constraint);
 				break;
 			case EContactShapesType::ConvexTriMesh:
 				if (const TImplicitObjectScaled<FTriangleMeshImplicitObject>* ScaledTriMesh = Implicit1.template GetObject<const TImplicitObjectScaled<FTriangleMeshImplicitObject> >())
 				{
-					UpdateConvexTriangleMeshConstraint(Implicit0, WorldTransform0, *ScaledTriMesh, WorldTransform1, CullDistance, Dt, Constraint);
+					UpdateConvexTriangleMeshConstraint(Implicit0, WorldTransform0, *ScaledTriMesh, WorldTransform1, Dt, Constraint);
 					break;
 				}
 				else if (const FTriangleMeshImplicitObject* TriangleMeshImplicit = Implicit1.template GetObject<const FTriangleMeshImplicitObject>())
 				{
-					UpdateConvexTriangleMeshConstraint(Implicit0, WorldTransform0, *TriangleMeshImplicit, WorldTransform1, CullDistance, Dt, Constraint);
+					UpdateConvexTriangleMeshConstraint(Implicit0, WorldTransform0, *TriangleMeshImplicit, WorldTransform1, Dt, Constraint);
 					break;
 				}
 				ensure(false);
 				break;
 			case EContactShapesType::ConvexHeightField:
-				UpdateConvexHeightFieldConstraint(Implicit0, WorldTransform0, *Implicit1.template GetObject< FHeightField >(), WorldTransform1, CullDistance, Dt, Constraint);
+				UpdateConvexHeightFieldConstraint(Implicit0, WorldTransform0, *Implicit1.template GetObject< FHeightField >(), WorldTransform1, Dt, Constraint);
 				break;
 			case EContactShapesType::LevelSetLevelSet:
-				UpdateLevelsetLevelsetConstraint<UpdateType>(WorldTransform0, WorldTransform1, CullDistance, Dt, Constraint);
+				UpdateLevelsetLevelsetConstraint<UpdateType>(WorldTransform0, WorldTransform1, Dt, Constraint);
 				break;
 			default:
 				// Switch needs updating....
@@ -2267,7 +2269,7 @@ namespace Chaos
 		// NOTE: Transforms are world space shape transforms
 		
 		template<ECollisionUpdateType UpdateType>
-		inline void UpdateConstraintFromGeometryImpl(FRigidBodySweptPointContactConstraint& Constraint, const FRigidTransform3& WorldTransform0, const FRigidTransform3& WorldTransform1, const FReal CullDistance, const FReal Dt)
+		inline void UpdateConstraintFromGeometryImpl(FRigidBodySweptPointContactConstraint& Constraint, const FRigidTransform3& WorldTransform0, const FRigidTransform3& WorldTransform1, const FReal Dt)
 		{
 			FReal LengthCCD = 0.0f;
 			FVec3 DirCCD(0.0f);
@@ -2283,20 +2285,20 @@ namespace Chaos
 				switch (Constraint.Manifold.ShapesType)
 				{
 				case EContactShapesType::GenericConvexConvex:
-					UpdateGenericConvexConvexConstraintSwept(Particle0, Implicit0, WorldTransform0, Implicit1, WorldTransform1, DirCCD, LengthCCD, CullDistance, Dt, Constraint);
+					UpdateGenericConvexConvexConstraintSwept(Particle0, Implicit0, WorldTransform0, Implicit1, WorldTransform1, DirCCD, LengthCCD, Dt, Constraint);
 					return;
 				case EContactShapesType::SphereHeightField:
 				{
 					const TSphere<FReal, 3>* Object0 = Implicit0.template GetObject<const TSphere<FReal, 3>>();
 					const FHeightField* Object1 = Implicit1.template GetObject<const FHeightField>();
-					UpdateSphereHeightFieldConstraintSwept(Particle0, *Object0, WorldTransform0, *Object1, WorldTransform1, DirCCD, LengthCCD, CullDistance, Dt, Constraint);
+					UpdateSphereHeightFieldConstraintSwept(Particle0, *Object0, WorldTransform0, *Object1, WorldTransform1, DirCCD, LengthCCD, Dt, Constraint);
 					return;
 				}
 				case EContactShapesType::CapsuleHeightField:
 				{
 					const TCapsule<FReal>* Object0 = Implicit0.template GetObject<const TCapsule<FReal> >();
 					const FHeightField* Object1 = Implicit1.template GetObject<const FHeightField >();
-					UpdateCapsuleHeightFieldConstraintSwept(Particle0, *Object0, WorldTransform0, *Object1, WorldTransform1, DirCCD, LengthCCD, CullDistance, Dt, Constraint);
+					UpdateCapsuleHeightFieldConstraintSwept(Particle0, *Object0, WorldTransform0, *Object1, WorldTransform1, DirCCD, LengthCCD, Dt, Constraint);
 					return;
 				}
 				case EContactShapesType::SphereTriMesh:
@@ -2304,11 +2306,11 @@ namespace Chaos
 					const TSphere<FReal, 3>* Object0 = Implicit0.template GetObject<const TSphere<FReal, 3>>();
 					if (const TImplicitObjectScaled<FTriangleMeshImplicitObject>* ScaledTriangleMesh = Implicit1.template GetObject<const TImplicitObjectScaled<FTriangleMeshImplicitObject>>())
 					{
-						UpdateSphereTriangleMeshConstraintSwept(Particle0, *Object0, WorldTransform0, *ScaledTriangleMesh, WorldTransform1, DirCCD, LengthCCD, CullDistance, Dt, Constraint);
+						UpdateSphereTriangleMeshConstraintSwept(Particle0, *Object0, WorldTransform0, *ScaledTriangleMesh, WorldTransform1, DirCCD, LengthCCD, Dt, Constraint);
 					}
 					else if (const FTriangleMeshImplicitObject* TriangleMesh = Implicit1.template GetObject<const FTriangleMeshImplicitObject>())
 					{
-						UpdateSphereTriangleMeshConstraintSwept(Particle0, *Object0, WorldTransform0, *TriangleMesh, WorldTransform1, DirCCD, LengthCCD, CullDistance, Dt, Constraint);
+						UpdateSphereTriangleMeshConstraintSwept(Particle0, *Object0, WorldTransform0, *TriangleMesh, WorldTransform1, DirCCD, LengthCCD, Dt, Constraint);
 					}
 					else
 					{
@@ -2321,11 +2323,11 @@ namespace Chaos
 					const TCapsule<FReal>* Object0 = Implicit0.template GetObject<const TCapsule<FReal> >();
 					if (const TImplicitObjectScaled<FTriangleMeshImplicitObject>* ScaledTriangleMesh = Implicit1.template GetObject<const TImplicitObjectScaled<FTriangleMeshImplicitObject>>())
 					{
-						UpdateCapsuleTriangleMeshConstraintSwept(Particle0, *Object0, WorldTransform0, *ScaledTriangleMesh, WorldTransform1, DirCCD, LengthCCD, CullDistance, Dt, Constraint);
+						UpdateCapsuleTriangleMeshConstraintSwept(Particle0, *Object0, WorldTransform0, *ScaledTriangleMesh, WorldTransform1, DirCCD, LengthCCD, Dt, Constraint);
 					}
 					else if (const FTriangleMeshImplicitObject* TriangleMesh = Implicit1.template GetObject<const FTriangleMeshImplicitObject>())
 					{
-						UpdateCapsuleTriangleMeshConstraintSwept(Particle0, *Object0, WorldTransform0, *TriangleMesh, WorldTransform1, DirCCD, LengthCCD, CullDistance, Dt, Constraint);
+						UpdateCapsuleTriangleMeshConstraintSwept(Particle0, *Object0, WorldTransform0, *TriangleMesh, WorldTransform1, DirCCD, LengthCCD, Dt, Constraint);
 					}
 					else
 					{
@@ -2336,17 +2338,17 @@ namespace Chaos
 				case EContactShapesType::ConvexHeightField:
 				{
 					const FHeightField* Object1 = Implicit1.template GetObject<const FHeightField >();
-					UpdateConvexHeightFieldConstraintSwept(Particle0, Implicit0, WorldTransform0, *Object1, WorldTransform1, DirCCD, LengthCCD, CullDistance, Dt, Constraint);
+					UpdateConvexHeightFieldConstraintSwept(Particle0, Implicit0, WorldTransform0, *Object1, WorldTransform1, DirCCD, LengthCCD, Dt, Constraint);
 					return;
 				}
 				case EContactShapesType::ConvexTriMesh:
 					if (const TImplicitObjectScaled<FTriangleMeshImplicitObject>* ScaledTriangleMesh = Implicit1.template GetObject<const TImplicitObjectScaled<FTriangleMeshImplicitObject>>())
 					{
-						UpdateConvexTriangleMeshConstraintSwept(Particle0, Implicit0, WorldTransform0, *ScaledTriangleMesh, WorldTransform1, DirCCD, LengthCCD, CullDistance, Dt, Constraint);
+						UpdateConvexTriangleMeshConstraintSwept(Particle0, Implicit0, WorldTransform0, *ScaledTriangleMesh, WorldTransform1, DirCCD, LengthCCD, Dt, Constraint);
 					}
 					else if (const FTriangleMeshImplicitObject* TriangleMesh = Implicit1.template GetObject<const FTriangleMeshImplicitObject>())
 					{
-						UpdateConvexTriangleMeshConstraintSwept(Particle0, Implicit0, WorldTransform0, *TriangleMesh, WorldTransform1, DirCCD, LengthCCD, CullDistance, Dt, Constraint);
+						UpdateConvexTriangleMeshConstraintSwept(Particle0, Implicit0, WorldTransform0, *TriangleMesh, WorldTransform1, DirCCD, LengthCCD, Dt, Constraint);
 					}
 					else
 					{
@@ -2358,7 +2360,7 @@ namespace Chaos
 
 			Constraint.TimeOfImpact = 1.0f; // CCD will not be used
 			// Do a normal non-swept update if we reach this point - Not required for now, since this will be done in solver
-			// UpdateConstraintFromGeometryImpl<UpdateType>(*(Constraint.As<FRigidBodyPointContactConstraint>()), WorldTransform0, WorldTransform1, CullDistance, Dt);
+			// UpdateConstraintFromGeometryImpl<UpdateType>(*(Constraint.As<FRigidBodyPointContactConstraint>()), WorldTransform0, WorldTransform1, Dt);
 		}
 
 		template<typename T_TRAITS>
@@ -2631,11 +2633,11 @@ namespace Chaos
 		// Run collision detection for the specified constraint to update the nearest contact point.
 		// NOTE: Transforms are world space particle transforms
 		template<ECollisionUpdateType UpdateType, typename RigidBodyContactConstraint>
-		void UpdateConstraintFromGeometry(RigidBodyContactConstraint& Constraint, const FRigidTransform3& ParticleTransform0, const FRigidTransform3& ParticleTransform1, const FReal CullDistance, const FReal Dt)
+		void UpdateConstraintFromGeometry(RigidBodyContactConstraint& Constraint, const FRigidTransform3& ParticleTransform0, const FRigidTransform3& ParticleTransform1, const FReal Dt)
 		{
 			const FRigidTransform3 WorldTransform0 = Constraint.ImplicitTransform[0] * ParticleTransform0;
 			const FRigidTransform3 WorldTransform1 = Constraint.ImplicitTransform[1] * ParticleTransform1;
-			UpdateConstraintFromGeometryImpl<UpdateType>(Constraint, WorldTransform0, WorldTransform1, CullDistance, Dt);
+			UpdateConstraintFromGeometryImpl<UpdateType>(Constraint, WorldTransform0, WorldTransform1, Dt);
 		}
 
 
@@ -2943,14 +2945,14 @@ namespace Chaos
 		}
 
 
-		template void UpdateLevelsetLevelsetConstraint<ECollisionUpdateType::Any>(const FRigidTransform3& WorldTransform0, const FRigidTransform3& WorldTransform1, const FReal CullDistance, const FReal Dt, FRigidBodyPointContactConstraint& Constraint);
-		template void UpdateLevelsetLevelsetConstraint<ECollisionUpdateType::Deepest>(const FRigidTransform3& WorldTransform0, const FRigidTransform3& WorldTransform1, const FReal CullDistance, const FReal Dt, FRigidBodyPointContactConstraint& Constraint);
+		template void UpdateLevelsetLevelsetConstraint<ECollisionUpdateType::Any>(const FRigidTransform3& WorldTransform0, const FRigidTransform3& WorldTransform1, const FReal Dt, FRigidBodyPointContactConstraint& Constraint);
+		template void UpdateLevelsetLevelsetConstraint<ECollisionUpdateType::Deepest>(const FRigidTransform3& WorldTransform0, const FRigidTransform3& WorldTransform1, const FReal Dt, FRigidBodyPointContactConstraint& Constraint);
 
-		template void UpdateConstraintFromGeometry<ECollisionUpdateType::Any, FRigidBodyPointContactConstraint>(FRigidBodyPointContactConstraint& ConstraintBase, const FRigidTransform3& Transform0, const FRigidTransform3& Transform1, const FReal CullDistance, const FReal Dt);
-		template void UpdateConstraintFromGeometry<ECollisionUpdateType::Deepest, FRigidBodyPointContactConstraint>(FRigidBodyPointContactConstraint& ConstraintBase, const FRigidTransform3& Transform0, const FRigidTransform3& Transform1, const FReal CullDistance, const FReal Dt);
+		template void UpdateConstraintFromGeometry<ECollisionUpdateType::Any, FRigidBodyPointContactConstraint>(FRigidBodyPointContactConstraint& ConstraintBase, const FRigidTransform3& Transform0, const FRigidTransform3& Transform1, const FReal Dt);
+		template void UpdateConstraintFromGeometry<ECollisionUpdateType::Deepest, FRigidBodyPointContactConstraint>(FRigidBodyPointContactConstraint& ConstraintBase, const FRigidTransform3& Transform0, const FRigidTransform3& Transform1, const FReal Dt);
 
-		template void UpdateConstraintFromGeometry<ECollisionUpdateType::Any, FRigidBodySweptPointContactConstraint>(FRigidBodySweptPointContactConstraint& ConstraintBase, const FRigidTransform3& Transform0, const FRigidTransform3& Transform1, const FReal CullDistance, const FReal Dt);
-		template void UpdateConstraintFromGeometry<ECollisionUpdateType::Deepest, FRigidBodySweptPointContactConstraint>(FRigidBodySweptPointContactConstraint& ConstraintBase, const FRigidTransform3& Transform0, const FRigidTransform3& Transform1, const FReal CullDistance, const FReal Dt);
+		template void UpdateConstraintFromGeometry<ECollisionUpdateType::Any, FRigidBodySweptPointContactConstraint>(FRigidBodySweptPointContactConstraint& ConstraintBase, const FRigidTransform3& Transform0, const FRigidTransform3& Transform1, const FReal Dt);
+		template void UpdateConstraintFromGeometry<ECollisionUpdateType::Deepest, FRigidBodySweptPointContactConstraint>(FRigidBodySweptPointContactConstraint& ConstraintBase, const FRigidTransform3& Transform0, const FRigidTransform3& Transform1, const FReal Dt);
 
 	} // Collisions
 
diff --git a/Engine/Source/Runtime/Experimental/Chaos/Private/Chaos/PBDCollisionConstraints.cpp b/Engine/Source/Runtime/Experimental/Chaos/Private/Chaos/PBDCollisionConstraints.cpp
--- a/Engine/Source/Runtime/Experimental/Chaos/Private/Chaos/PBDCollisionConstraints.cpp
+++ b/Engine/Source/Runtime/Experimental/Chaos/Private/Chaos/PBDCollisionConstraints.cpp
@@ -72,9 +72,6 @@ namespace Chaos
 	FRealSingle CollisionRestitutionThresholdOverride = -1.0f;
 	FAutoConsoleVariableRef CVarDefaultCollisionRestitutionThreshold(TEXT("p.CollisionRestitutionThreshold"), CollisionRestitutionThresholdOverride, TEXT("Collision restitution threshold override if >= 0 (units of acceleration)"));
 
-	FRealSingle CollisionCullDistanceOverride = -1.0f;
-	FAutoConsoleVariableRef CVarDefaultCollisionCullDistance(TEXT("p.CollisionCullDistance"), CollisionCullDistanceOverride, TEXT("Collision culling distance override if >= 0"));
-
 	int32 CollisionCanAlwaysDisableContacts = 0;
 	FAutoConsoleVariableRef CVarCollisionCanAlwaysDisableContacts(TEXT("p.CollisionCanAlwaysDisableContacts"), CollisionCanAlwaysDisableContacts, TEXT("Collision culling will always be able to permanently disable contacts"));
 
@@ -103,7 +100,6 @@ namespace Chaos
 		const TArrayCollectionArray<TUniquePtr<FChaosPhysicsMaterial>>& InPerParticlePhysicsMaterials,
 		const int32 InApplyPairIterations /*= 1*/,
 		const int32 InApplyPushOutPairIterations /*= 1*/,
-		const FReal InCullDistance /*= (FReal)0*/,
 		const FReal InRestitutionThreshold /*= (FReal)0*/)
 		: bInAppendOperation(false)
 		, Particles(InParticles)
@@ -114,7 +110,6 @@ namespace Chaos
 		, MPerParticlePhysicsMaterials(InPerParticlePhysicsMaterials)
 		, MApplyPairIterations(InApplyPairIterations)
 		, MApplyPushOutPairIterations(InApplyPushOutPairIterations)
-		, MCullDistance(InCullDistance)
 		, RestitutionThreshold(InRestitutionThreshold)	// @todo(chaos): expose as property
 		, bUseCCD(false)
 		, bEnableCollisions(true)
@@ -483,14 +478,6 @@ namespace Chaos
 	{
 		SCOPE_CYCLE_COUNTER(STAT_Collisions_UpdatePointConstraints);
 
-		// Make sure the cull distance is enough if we switched to Accumulated Impulse clipping		
-		// @todo(chaos): remove this - it should be handled in physics settings
-		const int MinCullDistanceForImpulseClipping = 5;
-		if (Collisions::Chaos_Collision_UseAccumulatedImpulseClipSolve && MCullDistance < MinCullDistanceForImpulseClipping)
-		{
-			MCullDistance = MinCullDistanceForImpulseClipping;
-		}
-
 		// @todo(chaos): parallelism needs to be optional
 
 		//PhysicsParallelFor(Handles.Num(), [&](int32 ConstraintHandleIndex)
@@ -507,8 +494,8 @@ namespace Chaos
 
 		for (FRigidBodyPointContactConstraint& Contact : Constraints.SinglePointConstraints)
 		{
-			Collisions::Update(Contact, MCullDistance, Dt);
-			if (Contact.GetPhi() < MCullDistance)
+			Collisions::Update(Contact, Dt);
+			if (Contact.GetPhi() < Contact.GetCullDistance())
 			{
 				Contact.Timestamp = LifespanCounter;
 			}
@@ -518,7 +505,6 @@ namespace Chaos
 	Collisions::FContactParticleParameters FPBDCollisionConstraints::GetContactParticleParameters(const FReal Dt)
 	{
 		return { 
-			(CollisionCullDistanceOverride >= 0.0f) ? CollisionCullDistanceOverride : MCullDistance,
 			(CollisionRestitutionThresholdOverride >= 0.0f) ? CollisionRestitutionThresholdOverride * Dt : RestitutionThreshold * Dt,
 			CollisionCanAlwaysDisableContacts ? true : (CollisionCanNeverDisableContacts ? false : bCanDisableContacts),
 			&MCollided,
diff --git a/Engine/Source/Runtime/Experimental/Chaos/Private/Chaos/PBDCollisionConstraintsContact.cpp b/Engine/Source/Runtime/Experimental/Chaos/Private/Chaos/PBDCollisionConstraintsContact.cpp
--- a/Engine/Source/Runtime/Experimental/Chaos/Private/Chaos/PBDCollisionConstraintsContact.cpp
+++ b/Engine/Source/Runtime/Experimental/Chaos/Private/Chaos/PBDCollisionConstraintsContact.cpp
@@ -44,32 +44,32 @@ namespace Chaos
 		int32 Chaos_Collision_UseShockPropagation = 1;
 		FAutoConsoleVariableRef CVarChaosCollisionUseShockPropagation(TEXT("p.Chaos.Collision.UseShockPropagation"), Chaos_Collision_UseShockPropagation, TEXT(""));
 
-		FReal Chaos_Collision_CollisionClipTolerance = 0.1f;
+		FReal Chaos_Collision_CollisionClipTolerance = 0.01f;
 		FAutoConsoleVariableRef CVarChaosCollisionClipTolerance(TEXT("p.Chaos.Collision.ClipTolerance"), Chaos_Collision_CollisionClipTolerance, TEXT(""));
 
 		bool Chaos_Collision_CheckManifoldComplete = false;
 		FAutoConsoleVariableRef CVarChaosCollisionCheckManifoldComplete(TEXT("p.Chaos.Collision.CheckManifoldComplete"), Chaos_Collision_CheckManifoldComplete, TEXT(""));
 
-		void Update(FRigidBodyPointContactConstraint& Constraint, const FReal CullDistance, const FReal Dt)
+		void Update(FRigidBodyPointContactConstraint& Constraint, const FReal Dt)
 		{
 			const FRigidTransform3 Transform0 = GetTransform(Constraint.Particle[0]);
 			const FRigidTransform3 Transform1 = GetTransform(Constraint.Particle[1]);
 
-			Constraint.ResetPhi(CullDistance);
-			UpdateConstraintFromGeometry<ECollisionUpdateType::Deepest, FRigidBodyPointContactConstraint>(Constraint, Transform0, Transform1, CullDistance, Dt);
+			Constraint.ResetPhi(Constraint.GetCullDistance());
+			UpdateConstraintFromGeometry<ECollisionUpdateType::Deepest, FRigidBodyPointContactConstraint>(Constraint, Transform0, Transform1, Dt);
 		}
 
-		void UpdateSwept(FRigidBodySweptPointContactConstraint& Constraint, const FReal CullDistance, const FReal Dt)
+		void UpdateSwept(FRigidBodySweptPointContactConstraint& Constraint, const FReal Dt)
 		{
 			FGenericParticleHandle Particle0 = FGenericParticleHandle(Constraint.Particle[0]);
 			// Note: This is unusual but we are using a mix of the previous and current transform
 			// This is due to how CCD only rewinds the position
 			const FRigidTransform3 TransformXQ0(Particle0->X(), Particle0->Q());
 			const FRigidTransform3 Transform1 = GetTransform(Constraint.Particle[1]);
-			Constraint.ResetPhi(CullDistance);
+			Constraint.ResetPhi(Constraint.GetCullDistance());
 			
 			// Update as a point constraint (base class).
-			UpdateConstraintFromGeometry<ECollisionUpdateType::Deepest, FRigidBodySweptPointContactConstraint>(Constraint, TransformXQ0, Transform1, CullDistance, Dt);
+			UpdateConstraintFromGeometry<ECollisionUpdateType::Deepest, FRigidBodySweptPointContactConstraint>(Constraint, TransformXQ0, Transform1, Dt);
 		}
 
 
@@ -215,11 +215,11 @@ namespace Chaos
 				bool bRequiresCollisionUpdate = true;
 				if (bRequiresCollisionUpdate)
 				{
-					Collisions::Update(Constraint, ParticleParameters.CullDistance, IterationParameters.Dt);
+					Collisions::Update(Constraint, IterationParameters.Dt);
 				}
 
 				// Permanently disable a constraint that is beyond the cull distance
-				if (Constraint.GetPhi() >= ParticleParameters.CullDistance)
+				if (Constraint.GetPhi() >= Constraint.GetCullDistance())
 				{
 					if (ParticleParameters.bCanDisableContacts)
 					{
@@ -266,10 +266,9 @@ namespace Chaos
 			FGenericParticleHandle Particle0 = FGenericParticleHandle(Constraint.Particle[0]);
 			FGenericParticleHandle Particle1 = FGenericParticleHandle(Constraint.Particle[1]);
 
-			const FReal ExpandedCullDistance = ParticleParameters.CullDistance + IterationParameters.Dt * Particle0->V().Size() * 2.0f; // Particle will move at most by this much during PBD solve
-			Collisions::UpdateSwept(Constraint, ExpandedCullDistance, IterationParameters.Dt);
+			Collisions::UpdateSwept(Constraint, IterationParameters.Dt);
 			
-			const FContactParticleParameters CCDParticleParamaters{ ExpandedCullDistance, ParticleParameters.RestitutionVelocityThreshold, true, ParticleParameters.Collided };
+			const FContactParticleParameters CCDParticleParamaters{ ParticleParameters.RestitutionVelocityThreshold, true, ParticleParameters.Collided };
 			if (Constraint.TimeOfImpact == 1)
 			{
 				// If TOI = 1 (normal constraint) we don't want to split timestep at TOI.
@@ -336,10 +335,10 @@ namespace Chaos
 
 			for (int32 PairIt = 0; PairIt < IterationParameters.NumPairIterations; ++PairIt)
 			{
-				Update(Constraint, ParticleParameters.CullDistance, IterationParameters.Dt);
+				Update(Constraint, IterationParameters.Dt);
 
 				// Ignore contacts where the closest point is greater than cull distance
-				if (Constraint.GetPhi() >= ParticleParameters.CullDistance)
+				if (Constraint.GetPhi() >= Constraint.GetCullDistance())
 				{
 					// Optionally permanently disable the contact for the remaining iterations
 					if (ParticleParameters.bCanDisableContacts)
diff --git a/Engine/Source/Runtime/Experimental/Chaos/Public/Chaos/BoundingVolumeUtilities.h b/Engine/Source/Runtime/Experimental/Chaos/Public/Chaos/BoundingVolumeUtilities.h
--- a/Engine/Source/Runtime/Experimental/Chaos/Public/Chaos/BoundingVolumeUtilities.h
+++ b/Engine/Source/Runtime/Experimental/Chaos/Public/Chaos/BoundingVolumeUtilities.h
@@ -39,13 +39,6 @@ inline FVec3 ComputeBoundsThickness(const TPBDRigidParticles<FReal, 3>& InPartic
 template <typename THandle>
 FVec3 ComputeBoundsThickness(const THandle& PBDRigid, FReal Dt, FReal BoundsThickness, FReal BoundsVelocityInflation)
 {
-	FReal NewBoundsVelocityInflation = BoundsVelocityInflation;
-	auto* RigidBody = PBDRigid.CastToRigidParticle();
-	if (RigidBody && !RigidBody->CCDEnabled())
-	{
-		NewBoundsVelocityInflation = 0.0f;
-	}
-
 	FVec3 Vel{ 0 };
 	auto* KinematicBody = PBDRigid.CastToKinematicParticle();
 	if (KinematicBody)
@@ -53,7 +46,7 @@ FVec3 ComputeBoundsThickness(const THandle& PBDRigid, FReal Dt, FReal BoundsThic
 		Vel = KinematicBody->V();
 	}
 
-	return ComputeBoundsThickness(Vel, Dt, BoundsThickness, NewBoundsVelocityInflation);
+	return ComputeBoundsThickness(Vel, Dt, BoundsThickness, BoundsVelocityInflation);
 }
 
 template<class OBJECT_ARRAY>
diff --git a/Engine/Source/Runtime/Experimental/Chaos/Public/Chaos/Collision/CapsuleConvexContactPoint.h b/Engine/Source/Runtime/Experimental/Chaos/Public/Chaos/Collision/CapsuleConvexContactPoint.h
--- a/Engine/Source/Runtime/Experimental/Chaos/Public/Chaos/Collision/CapsuleConvexContactPoint.h
+++ b/Engine/Source/Runtime/Experimental/Chaos/Public/Chaos/Collision/CapsuleConvexContactPoint.h
@@ -10,13 +10,13 @@ namespace Chaos
 {
 	class FContactPoint;
 
-	extern CHAOS_API FContactPoint CapsuleConvexContactPoint(const FImplicitCapsule3& Capsule, const FRigidTransform3& CapsuleTransform, const FImplicitObject& Convex, const FRigidTransform3& ConvexTransform, FReal CullDistance);
+	extern CHAOS_API FContactPoint CapsuleConvexContactPoint(const FImplicitCapsule3& Capsule, const FRigidTransform3& CapsuleTransform, const FImplicitObject& Convex, const FRigidTransform3& ConvexTransform);
 
-	extern CHAOS_API FContactPoint CapsuleConvexContactPoint(const FImplicitCapsule3& Capsule, const FRigidTransform3& CapsuleTransform, const FImplicitConvex3& Convex, const FRigidTransform3& ConvexTransform, FReal CullDistance);
-	extern CHAOS_API FContactPoint CapsuleConvexContactPoint(const FImplicitCapsule3& Capsule, const FRigidTransform3& CapsuleTransform, const TImplicitObjectInstanced<FImplicitConvex3>& Convex, const FRigidTransform3& ConvexTransform, FReal CullDistance);
-	extern CHAOS_API FContactPoint CapsuleConvexContactPoint(const FImplicitCapsule3& Capsule, const FRigidTransform3& CapsuleTransform, const TImplicitObjectScaled<FImplicitConvex3>& Convex, const FRigidTransform3& ConvexTransform, FReal CullDistance);
+	extern CHAOS_API FContactPoint CapsuleConvexContactPoint(const FImplicitCapsule3& Capsule, const FRigidTransform3& CapsuleTransform, const FImplicitConvex3& Convex, const FRigidTransform3& ConvexTransform);
+	extern CHAOS_API FContactPoint CapsuleConvexContactPoint(const FImplicitCapsule3& Capsule, const FRigidTransform3& CapsuleTransform, const TImplicitObjectInstanced<FImplicitConvex3>& Convex, const FRigidTransform3& ConvexTransform);
+	extern CHAOS_API FContactPoint CapsuleConvexContactPoint(const FImplicitCapsule3& Capsule, const FRigidTransform3& CapsuleTransform, const TImplicitObjectScaled<FImplicitConvex3>& Convex, const FRigidTransform3& ConvexTransform);
 
-	extern CHAOS_API FContactPoint ConvexCapsuleContactPoint(const FImplicitConvex3& Convex, const FRigidTransform3& ConvexTransform, const FImplicitCapsule3& Capsule, const FRigidTransform3& CapsuleTransform, FReal CullDistance);
-	extern CHAOS_API FContactPoint ConvexCapsuleContactPoint(const TImplicitObjectInstanced<FImplicitConvex3>& Convex, const FRigidTransform3& ConvexTransform, const FImplicitCapsule3& Capsule, const FRigidTransform3& CapsuleTransform, FReal CullDistance);
-	extern CHAOS_API FContactPoint ConvexCapsuleContactPoint(const TImplicitObjectScaled<FImplicitConvex3>& Convex, const FRigidTransform3& ConvexTransform, const FImplicitCapsule3& Capsule, const FRigidTransform3& CapsuleTransform, FReal CullDistance);
+	extern CHAOS_API FContactPoint ConvexCapsuleContactPoint(const FImplicitConvex3& Convex, const FRigidTransform3& ConvexTransform, const FImplicitCapsule3& Capsule, const FRigidTransform3& CapsuleTransform);
+	extern CHAOS_API FContactPoint ConvexCapsuleContactPoint(const TImplicitObjectInstanced<FImplicitConvex3>& Convex, const FRigidTransform3& ConvexTransform, const FImplicitCapsule3& Capsule, const FRigidTransform3& CapsuleTransform);
+	extern CHAOS_API FContactPoint ConvexCapsuleContactPoint(const TImplicitObjectScaled<FImplicitConvex3>& Convex, const FRigidTransform3& ConvexTransform, const FImplicitCapsule3& Capsule, const FRigidTransform3& CapsuleTransform);
 }
diff --git a/Engine/Source/Runtime/Experimental/Chaos/Public/Chaos/Collision/PBDCollisionConstraint.h b/Engine/Source/Runtime/Experimental/Chaos/Public/Chaos/Collision/PBDCollisionConstraint.h
--- a/Engine/Source/Runtime/Experimental/Chaos/Public/Chaos/Collision/PBDCollisionConstraint.h
+++ b/Engine/Source/Runtime/Experimental/Chaos/Public/Chaos/Collision/PBDCollisionConstraint.h
@@ -251,6 +251,7 @@ namespace Chaos
 
 		FRigidBodyPointContactConstraint() 
 			: Base(Base::FType::SinglePoint)
+			, CullDistance(FLT_MAX)
 			, bUseManifold(false)
 		{}
 
@@ -263,14 +264,19 @@ namespace Chaos
 			const FImplicitObject* Implicit1, 
 			const FBVHParticles* Simplicial1, 
 			const FRigidTransform3& Transform1,
+			const FReal InCullDistance,
 			const EContactShapesType ShapesType,
 			const bool bInUseManifold)
 			: Base(Particle0, Implicit0, Simplicial0, Transform0, Particle1, Implicit1, Simplicial0, Transform1, Base::FType::SinglePoint, ShapesType)
+			, CullDistance(InCullDistance)
 			, bUseManifold(bInUseManifold)
 		{}
 
 		static typename Base::FType StaticType() { return Base::FType::SinglePoint; };
 
+		FReal GetCullDistance() const { return CullDistance; }
+		void SetCullDistance(FReal InCullDistance) { CullDistance = InCullDistance; }
+
 		bool GetUseManifold() const { return bUseManifold; }
 
 		TArrayView<FManifoldPoint> GetManifoldPoints() { return MakeArrayView(ManifoldPoints); }
@@ -307,9 +313,11 @@ namespace Chaos
 			const FImplicitObject* Implicit1, 
 			const FBVHParticles* Simplicial1,
 			const FRigidTransform3& Transform1,
-			typename Base::FType InType, 
+			const FReal InCullDistance,
+			typename Base::FType InType,
 			const EContactShapesType ShapesType)
 			: Base(Particle0, Implicit0, Simplicial0, Transform0, Particle1, Implicit1, Simplicial1, Transform1, InType, ShapesType)
+			, CullDistance(InCullDistance)
 			, bUseManifold(false)
 		{}
 
@@ -318,11 +326,14 @@ namespace Chaos
 		int32 AddManifoldPoint(const FContactPoint& ContactPoint, const FReal Dt);
 		void InitManifoldPoint(FManifoldPoint& ManifoldPoint, FReal Dt);
 		void UpdateManifoldPoint(int32 ManifoldPointIndex, const FContactPoint& ContactPoint, const FReal Dt);
+		void UpdateManifoldPointFromContact(FManifoldPoint& ManifoldPoint);
 		void SetActiveContactPoint(const FContactPoint& ContactPoint);
+		void GetWorldSpaceManifoldPoint(const FManifoldPoint& ManifoldPoint, const FVec3& P0, const FRotation3& Q0, const FVec3& P1, const FRotation3& Q1, FVec3& OutContactLocation, FVec3& OutContactNormal, FReal& OutContactPhi);
 
 		// @todo(chaos): Switching this to inline allocator does not help, but maybe a physics scratchpad would
 		//TArray<FManifoldPoint, TInlineAllocator<4>> ManifoldPoints;
 		TArray<FManifoldPoint> ManifoldPoints;
+		FReal CullDistance;
 		bool bUseManifold;
 	};
 
@@ -346,8 +357,9 @@ namespace Chaos
 			const FImplicitObject* Implicit1, 
 			const FBVHParticles* Simplicial1, 
 			const FRigidTransform3& Transform1,
+			const FReal InCullDistance,
 			EContactShapesType ShapesType)
-			: Base(Particle0, Implicit0, Simplicial0, Transform0, Particle1, Implicit1, Simplicial1, Transform1, Base::FType::SinglePointSwept, ShapesType)
+			: Base(Particle0, Implicit0, Simplicial0, Transform0, Particle1, Implicit1, Simplicial1, Transform1, InCullDistance, Base::FType::SinglePointSwept, ShapesType)
 			, TimeOfImpact(0)
 		{}
 
diff --git a/Engine/Source/Runtime/Experimental/Chaos/Public/Chaos/Collision/SphereConvexContactPoint.h b/Engine/Source/Runtime/Experimental/Chaos/Public/Chaos/Collision/SphereConvexContactPoint.h
--- a/Engine/Source/Runtime/Experimental/Chaos/Public/Chaos/Collision/SphereConvexContactPoint.h
+++ b/Engine/Source/Runtime/Experimental/Chaos/Public/Chaos/Collision/SphereConvexContactPoint.h
@@ -10,13 +10,13 @@ namespace Chaos
 {
 	class FContactPoint;
 
-	extern CHAOS_API FContactPoint SphereConvexContactPoint(const FImplicitSphere3& Sphere, const FRigidTransform3& SphereTransform, const FImplicitObject& Convex, const FRigidTransform3& ConvexTransform, FReal CullDistance);
+	extern CHAOS_API FContactPoint SphereConvexContactPoint(const FImplicitSphere3& Sphere, const FRigidTransform3& SphereTransform, const FImplicitObject& Convex, const FRigidTransform3& ConvexTransform);
 
-	extern CHAOS_API FContactPoint SphereConvexContactPoint(const FImplicitSphere3& Sphere, const FRigidTransform3& SphereTransform, const FImplicitConvex3& Convex, const FRigidTransform3& ConvexTransform, FReal CullDistance);
-	extern CHAOS_API FContactPoint SphereConvexContactPoint(const FImplicitSphere3& Sphere, const FRigidTransform3& SphereTransform, const TImplicitObjectInstanced<FImplicitConvex3>& Convex, const FRigidTransform3& ConvexTransform, FReal CullDistance);
-	extern CHAOS_API FContactPoint SphereConvexContactPoint(const FImplicitSphere3& Sphere, const FRigidTransform3& SphereTransform, const TImplicitObjectScaled<FImplicitConvex3>& Convex, const FRigidTransform3& ConvexTransform, FReal CullDistance);
+	extern CHAOS_API FContactPoint SphereConvexContactPoint(const FImplicitSphere3& Sphere, const FRigidTransform3& SphereTransform, const FImplicitConvex3& Convex, const FRigidTransform3& ConvexTransform);
+	extern CHAOS_API FContactPoint SphereConvexContactPoint(const FImplicitSphere3& Sphere, const FRigidTransform3& SphereTransform, const TImplicitObjectInstanced<FImplicitConvex3>& Convex, const FRigidTransform3& ConvexTransform);
+	extern CHAOS_API FContactPoint SphereConvexContactPoint(const FImplicitSphere3& Sphere, const FRigidTransform3& SphereTransform, const TImplicitObjectScaled<FImplicitConvex3>& Convex, const FRigidTransform3& ConvexTransform);
 
-	extern CHAOS_API FContactPoint ConvexSphereContactPoint(const FImplicitConvex3& Convex, const FRigidTransform3& ConvexTransform, const FImplicitSphere3& Sphere, const FRigidTransform3& SphereTransform, FReal CullDistance);
-	extern CHAOS_API FContactPoint ConvexSphereContactPoint(const TImplicitObjectInstanced<FImplicitConvex3>& Convex, const FRigidTransform3& ConvexTransform, const FImplicitSphere3& Sphere, const FRigidTransform3& SphereTransform, FReal CullDistance);
-	extern CHAOS_API FContactPoint ConvexSphereContactPoint(const TImplicitObjectScaled<FImplicitConvex3>& Convex, const FRigidTransform3& ConvexTransform, const FImplicitSphere3& Sphere, const FRigidTransform3& SphereTransform, FReal CullDistance);
+	extern CHAOS_API FContactPoint ConvexSphereContactPoint(const FImplicitConvex3& Convex, const FRigidTransform3& ConvexTransform, const FImplicitSphere3& Sphere, const FRigidTransform3& SphereTransform);
+	extern CHAOS_API FContactPoint ConvexSphereContactPoint(const TImplicitObjectInstanced<FImplicitConvex3>& Convex, const FRigidTransform3& ConvexTransform, const FImplicitSphere3& Sphere, const FRigidTransform3& SphereTransform);
+	extern CHAOS_API FContactPoint ConvexSphereContactPoint(const TImplicitObjectScaled<FImplicitConvex3>& Convex, const FRigidTransform3& ConvexTransform, const FImplicitSphere3& Sphere, const FRigidTransform3& SphereTransform);
 }
diff --git a/Engine/Source/Runtime/Experimental/Chaos/Public/Chaos/CollisionResolution.h b/Engine/Source/Runtime/Experimental/Chaos/Public/Chaos/CollisionResolution.h
--- a/Engine/Source/Runtime/Experimental/Chaos/Public/Chaos/CollisionResolution.h
+++ b/Engine/Source/Runtime/Experimental/Chaos/Public/Chaos/CollisionResolution.h
@@ -37,7 +37,7 @@ namespace Chaos
 		// Update the constraint by re-running collision detection on the shape pair.
 
 		template<ECollisionUpdateType UpdateType, typename RigidBodyContactConstraint>
-		void CHAOS_API UpdateConstraintFromGeometry(RigidBodyContactConstraint& Constraint, const FRigidTransform3& ParticleTransform0, const FRigidTransform3& ParticleTransform1, const FReal CullDistance, const FReal Dt);
+		void CHAOS_API UpdateConstraintFromGeometry(RigidBodyContactConstraint& Constraint, const FRigidTransform3& ParticleTransform0, const FRigidTransform3& ParticleTransform1, const FReal Dt);
 
 		// Create constraints for the particle pair. This could create multiple constraints: one for each potentially colliding shape pair in multi-shape particles.
 		void CHAOS_API ConstructConstraints(TGeometryParticleHandle<FReal, 3>* Particle0, TGeometryParticleHandle<FReal, 3>* Particle1, const FImplicitObject* Implicit0, const FBVHParticles* Simplicial0, const FImplicitObject* Implicit1, const FBVHParticles* Simplicial1, const FRigidTransform3& Transform0, const FRigidTransform3& Transform1, const FReal CullDistance, const FReal dT,const FCollisionContext& Context, FCollisionConstraintsArray& NewConstraints);
@@ -45,7 +45,7 @@ namespace Chaos
 
 		// @todo(chaos): this is only called in tests - should it really be exposed?
 		template<ECollisionUpdateType UpdateType>
-		void UpdateLevelsetLevelsetConstraint(const FRigidTransform3& WorldTransform0, const FRigidTransform3& WorldTransform1, const FReal CullDistance, const FReal Dt, FRigidBodyPointContactConstraint& Constraint);
+		void UpdateLevelsetLevelsetConstraint(const FRigidTransform3& WorldTransform0, const FRigidTransform3& WorldTransform1, const FReal Dt, FRigidBodyPointContactConstraint& Constraint);
 
 	}
 }
diff --git a/Engine/Source/Runtime/Experimental/Chaos/Public/Chaos/PBDCollisionConstraints.h b/Engine/Source/Runtime/Experimental/Chaos/Public/Chaos/PBDCollisionConstraints.h
--- a/Engine/Source/Runtime/Experimental/Chaos/Public/Chaos/PBDCollisionConstraints.h
+++ b/Engine/Source/Runtime/Experimental/Chaos/Public/Chaos/PBDCollisionConstraints.h
@@ -58,7 +58,6 @@ public:
 		const TArrayCollectionArray<TUniquePtr<FChaosPhysicsMaterial>>& PerParticlePhysicsMaterials, 
 		const int32 ApplyPairIterations = 1, 
 		const int32 ApplyPushOutPairIterations = 1, 
-		const FReal CullDistance = (FReal)0, 
 		const FReal RestitutionThreshold = 2000.0f);
 
 	virtual ~FPBDCollisionConstraints() {}
@@ -254,16 +253,6 @@ public:
 #endif
 	}
 
-	void SetCullDistance(FReal InCullDistance)
-	{
-		MCullDistance = InCullDistance;
-	}
-
-	FReal GetCullDistance() const
-	{
-		return MCullDistance;
-	}
-
 	void SetCanDisableContacts(bool bInCanDisableContacts)
 	{
 		bCanDisableContacts = bInCanDisableContacts;
@@ -370,7 +359,6 @@ private:
 	const TArrayCollectionArray<TUniquePtr<FChaosPhysicsMaterial>>& MPerParticlePhysicsMaterials;
 	int32 MApplyPairIterations;
 	int32 MApplyPushOutPairIterations;
-	FReal MCullDistance;
 	FReal RestitutionThreshold;
 	bool bUseCCD;
 	bool bEnableCollisions;
diff --git a/Engine/Source/Runtime/Experimental/Chaos/Public/Chaos/PBDCollisionConstraintsContact.h b/Engine/Source/Runtime/Experimental/Chaos/Public/Chaos/PBDCollisionConstraintsContact.h
--- a/Engine/Source/Runtime/Experimental/Chaos/Public/Chaos/PBDCollisionConstraintsContact.h
+++ b/Engine/Source/Runtime/Experimental/Chaos/Public/Chaos/PBDCollisionConstraintsContact.h
@@ -12,14 +12,55 @@ namespace Chaos
 
 	namespace Collisions
 	{
-		struct FContactParticleParameters {
-			FReal CullDistance;
+		struct FContactParticleParameters 
+		{
+			FContactParticleParameters()
+				: RestitutionVelocityThreshold(0)
+				, bCanDisableContacts(false)
+				, Collided(nullptr)
+			{}
+
+			FContactParticleParameters(
+				FReal InRestitutionVelocityThreshold,
+				bool bInCanDisableContacts,
+				TArrayCollectionArray<bool>* InCollided)
+				: RestitutionVelocityThreshold(InRestitutionVelocityThreshold)
+				, bCanDisableContacts(bInCanDisableContacts)
+				, Collided(InCollided)
+			{}
+
+
 			FReal RestitutionVelocityThreshold;
 			bool bCanDisableContacts;
 			TArrayCollectionArray<bool>* Collided;
 		};
 
-		struct FContactIterationParameters {
+		struct FContactIterationParameters 
+		{
+			FContactIterationParameters()
+				: Dt(0)
+				, Iteration(0)
+				, NumIterations(0)
+				, NumPairIterations(0)
+				, ApplyType(ECollisionApplyType::None)
+				, NeedsAnotherIteration(nullptr)
+			{}
+
+			FContactIterationParameters(
+				const FReal InDt,
+				const int32 InIteration,
+				const int32 InNumIterations,
+				const int32 InNumPairIterations,
+				const ECollisionApplyType InApplyType,
+				bool* InNeedsAnotherIteration)
+				: Dt(InDt)
+				, Iteration(InIteration)
+				, NumIterations(InNumIterations)
+				, NumPairIterations(InNumPairIterations)
+				, ApplyType(InApplyType)
+				, NeedsAnotherIteration(InNeedsAnotherIteration)
+			{}
+
 			const FReal Dt;
 			const int32 Iteration;
 			const int32 NumIterations;
@@ -29,7 +70,7 @@ namespace Chaos
 		};
 
 		// Update the constraint (re-runs collision detection for this contact)
-		extern void Update(FRigidBodyPointContactConstraint& Constraint, const FReal CullDistance, const FReal Dt);
+		extern void Update(FRigidBodyPointContactConstraint& Constraint, const FReal Dt);
 
 		extern void Apply(FCollisionConstraintBase& Constraint, const FContactIterationParameters& IterationParameters, const FContactParticleParameters& ParticleParameters);
 
diff --git a/Engine/Source/Runtime/Experimental/Chaos/Public/PBDRigidsSolver.h b/Engine/Source/Runtime/Experimental/Chaos/Public/PBDRigidsSolver.h
--- a/Engine/Source/Runtime/Experimental/Chaos/Public/PBDRigidsSolver.h
+++ b/Engine/Source/Runtime/Experimental/Chaos/Public/PBDRigidsSolver.h
@@ -193,7 +193,7 @@ namespace Chaos
 		void SetCollisionPushOutPairIterations(const int32 InNumIterations) { GetEvolution()->GetCollisionConstraints().SetPushOutPairIterations(InNumIterations); }
 		void SetJointPairIterations(const int32 InNumIterations) { GetJointConstraints().SetNumPairIterations(InNumIterations); }
 		void SetJointPushOutPairIterations(const int32 InNumIterations) {GetJointConstraints().SetNumPushOutPairIterations(InNumIterations); }
-		void SetCollisionCullDistance(const FReal InCullDistance) { GetEvolution()->GetCollisionConstraints().SetCullDistance(InCullDistance); GetEvolution()->GetBroadPhase().SetCullDistance(InCullDistance); }
+		void SetCollisionCullDistance(const FReal InCullDistance) { GetEvolution()->GetBroadPhase().SetCullDistance(InCullDistance); }
 		void SetUseContactGraph(const bool bInUseContactGraph) { GetEvolution()->GetCollisionConstraintsRule().SetUseContactGraph(bInUseContactGraph); }
 
 		/**/
