CIS fix - fix post merge error in PBDCollisionConstraint.cpp

diff --git a/Engine/Source/Runtime/Experimental/Chaos/Private/Chaos/Collision/PBDCollisionConstraint.cpp b/Engine/Source/Runtime/Experimental/Chaos/Private/Chaos/Collision/PBDCollisionConstraint.cpp
--- a/Engine/Source/Runtime/Experimental/Chaos/Private/Chaos/Collision/PBDCollisionConstraint.cpp
+++ b/Engine/Source/Runtime/Experimental/Chaos/Private/Chaos/Collision/PBDCollisionConstraint.cpp
@@ -309,8 +309,8 @@ namespace Chaos
 	// @todo(chaos): we shouldn't have to recalculate local space positions - they were available in collision update
 	void FRigidBodyPointContactConstraint::UpdateManifoldPointFromContact(FManifoldPoint& ManifoldPoint)
 	{
-		TConstGenericParticleHandle<FReal, 3> Particle0 = Particle[0];
-		TConstGenericParticleHandle<FReal, 3> Particle1 = Particle[1];
+		FConstGenericParticleHandle Particle0 = Particle[0];
+		FConstGenericParticleHandle Particle1 = Particle[1];
 
 		check(ManifoldPoint.ContactPoint.ContactNormalOwnerIndex >= 0);
 		check(ManifoldPoint.ContactPoint.ContactNormalOwnerIndex < 2);
