GeometryCollectionComponent
-- changed the way ISMCs are spawned for embedded geometry to correct issues with cache spawns.

diff --git a/Engine/Source/Runtime/Experimental/GeometryCollectionEngine/Private/GeometryCollection/GeometryCollectionComponent.cpp b/Engine/Source/Runtime/Experimental/GeometryCollectionEngine/Private/GeometryCollection/GeometryCollectionComponent.cpp
--- a/Engine/Source/Runtime/Experimental/GeometryCollectionEngine/Private/GeometryCollection/GeometryCollectionComponent.cpp
+++ b/Engine/Source/Runtime/Experimental/GeometryCollectionEngine/Private/GeometryCollection/GeometryCollectionComponent.cpp
@@ -2942,13 +2942,16 @@ void UGeometryCollectionComponent::ClearEmbeddedGeometry()
 	AActor* OwningActor = GetOwner();
 	TArray<UActorComponent*> TargetComponents;
 	OwningActor->GetComponents(TargetComponents, false);
-	
-	for (UActorComponent* EmbeddedGeometryComponent : TargetComponents)
+
+	for (UActorComponent* TargetComponent : TargetComponents)
 	{
-		if (UInstancedStaticMeshComponent* ISMComponent = Cast<UInstancedStaticMeshComponent>(EmbeddedGeometryComponent))
+		if (TargetComponent->GetOuter() == this)
 		{
-			ISMComponent->ClearInstances();
-			ISMComponent->DestroyComponent();
+			if (UInstancedStaticMeshComponent* ISMComponent = Cast<UInstancedStaticMeshComponent>(TargetComponent))
+			{
+				ISMComponent->ClearInstances();
+				ISMComponent->DestroyComponent();
+			}
 		}
 	}
 
@@ -2969,7 +2972,7 @@ void UGeometryCollectionComponent::InitializeEmbeddedGeometry()
 		{
 			if (UStaticMesh* ExemplarStaticMesh = Cast<UStaticMesh>(Exemplar.StaticMeshExemplar.TryLoad()))
 			{
-				if (UInstancedStaticMeshComponent* ISMC = NewObject<UInstancedStaticMeshComponent>(ActorOwner))
+				if (UInstancedStaticMeshComponent* ISMC = NewObject<UInstancedStaticMeshComponent>(this))
 				{
 					ISMC->SetStaticMesh(ExemplarStaticMesh);
 					ISMC->SetCullDistances(Exemplar.StartCullDistance, Exemplar.EndCullDistance);
@@ -2979,7 +2982,6 @@ void UGeometryCollectionComponent::InitializeEmbeddedGeometry()
 					ISMC->SetMobility(EComponentMobility::Stationary);
 					ISMC->SetupAttachment(this);
 					ISMC->RegisterComponent();
-					ActorOwner->AddInstanceComponent(ISMC);
 
 					EmbeddedGeometryComponents.Add(ISMC);
 				}
