GeometryCollection
-- Removed vertex visibility test from CalcBounds as it was having the opposite performance effect than intended.

diff --git a/Engine/Source/Runtime/Experimental/GeometryCollectionEngine/Private/GeometryCollection/GeometryCollectionComponent.cpp b/Engine/Source/Runtime/Experimental/GeometryCollectionEngine/Private/GeometryCollection/GeometryCollectionComponent.cpp
--- a/Engine/Source/Runtime/Experimental/GeometryCollectionEngine/Private/GeometryCollection/GeometryCollectionComponent.cpp
+++ b/Engine/Source/Runtime/Experimental/GeometryCollectionEngine/Private/GeometryCollection/GeometryCollectionComponent.cpp
@@ -390,7 +390,7 @@ FBoxSphereBounds UGeometryCollectionComponent::CalcBounds(const FTransform& Loca
 	{
 		return WorldBounds;
 	} 
-	else if (RestCollection && RestCollection->HasVisibleGeometry())
+	else if (RestCollection)
 	{			
 		const FMatrix LocalToWorldWithScale = LocalToWorldIn.ToMatrixWithScale();
 
@@ -403,6 +403,7 @@ FBoxSphereBounds UGeometryCollectionComponent::CalcBounds(const FTransform& Loca
 		const TManagedArray<int32>& TransformIndices = GetTransformIndexArray();
 		const TManagedArray<int32>& ParentIndices = GetParentArray();
 		const TManagedArray<int32>& TransformToGeometryIndex = GetTransformToGeometryIndexArray();
+		const TManagedArray<FTransform>& Transforms = GetTransformArray();
 
 		const int32 NumBoxes = BoundingBoxes.Num();
 
@@ -414,8 +415,8 @@ FBoxSphereBounds UGeometryCollectionComponent::CalcBounds(const FTransform& Loca
 		if (GlobalMatrices.Num() != RestCollection->NumElements(FGeometryCollection::TransformGroup))
 		{
 			TArray<FMatrix> TmpGlobalMatrices;
-			GeometryCollectionAlgo::GlobalMatrices(GetTransformArray(), ParentIndices, TmpGlobalMatrices);
-
+			
+			GeometryCollectionAlgo::GlobalMatrices(Transforms, ParentIndices, TmpGlobalMatrices);
 			if (TmpGlobalMatrices.Num() == 0)
 			{
 				return FBoxSphereBounds(ForceInitToZero);
