GeometryCollection
-- fixed a bug in calc bounds when creating a new (empty) rest collection

diff --git a/Engine/Source/Runtime/Experimental/GeometryCollectionEngine/Private/GeometryCollection/GeometryCollectionComponent.cpp b/Engine/Source/Runtime/Experimental/GeometryCollectionEngine/Private/GeometryCollection/GeometryCollectionComponent.cpp
--- a/Engine/Source/Runtime/Experimental/GeometryCollectionEngine/Private/GeometryCollection/GeometryCollectionComponent.cpp
+++ b/Engine/Source/Runtime/Experimental/GeometryCollectionEngine/Private/GeometryCollection/GeometryCollectionComponent.cpp
@@ -412,7 +412,8 @@ FBoxSphereBounds UGeometryCollectionComponent::CalcBounds(const FTransform& Loca
 		// has set the global matrices cache while in the editor.  This is a somewhat weak guard against this
 		// to default to just calculating tmp global matrices.  This should be removed or modified somehow
 		// such that we always cache the global matrices and this method always does the correct behavior
-		if (GlobalMatrices.Num() != RestCollection->NumElements(FGeometryCollection::TransformGroup))
+		int32 NumElements = RestCollection->NumElements(FGeometryCollection::TransformGroup);
+		if (NumElements == 0 || GlobalMatrices.Num() != RestCollection->NumElements(FGeometryCollection::TransformGroup))
 		{
 			TArray<FMatrix> TmpGlobalMatrices;
 			
