Fix vehicle template and remove old wheel code path

diff --git a/Engine/Plugins/Experimental/ChaosVehiclesPlugin/Source/ChaosVehicles/Private/ChaosVehicleWheel.cpp b/Engine/Plugins/Experimental/ChaosVehiclesPlugin/Source/ChaosVehicles/Private/ChaosVehicleWheel.cpp
--- a/Engine/Plugins/Experimental/ChaosVehiclesPlugin/Source/ChaosVehicles/Private/ChaosVehicleWheel.cpp
+++ b/Engine/Plugins/Experimental/ChaosVehiclesPlugin/Source/ChaosVehicles/Private/ChaosVehicleWheel.cpp
@@ -25,11 +25,8 @@ UChaosVehicleWheel::UChaosVehicleWheel(const FObjectInitializer& ObjectInitializ
 	WheelWidth = 20.0f;
 	//bAutoAdjustCollisionSize = true;
 	//WheelMass = 20.0f;
-	LongitudinalFrictionForceMultiplier = 1.0f;
-	LateralFrictionForceMultiplier = 2.0f;
 	FrictionForceMultiplier = 2.0f;
 	CorneringStiffness = 1000.0f;
-	bNewWheelSimulation = false;
 	SideSlipModifier = 1.0f;
 	SlipThreshold = 20.0f;
 	SkidThreshold = 20.0f;
diff --git a/Engine/Plugins/Experimental/ChaosVehiclesPlugin/Source/ChaosVehicles/Private/ChaosWheeledVehicleMovementComponent.cpp b/Engine/Plugins/Experimental/ChaosVehiclesPlugin/Source/ChaosVehicles/Private/ChaosWheeledVehicleMovementComponent.cpp
--- a/Engine/Plugins/Experimental/ChaosVehiclesPlugin/Source/ChaosVehicles/Private/ChaosWheeledVehicleMovementComponent.cpp
+++ b/Engine/Plugins/Experimental/ChaosVehiclesPlugin/Source/ChaosVehicles/Private/ChaosWheeledVehicleMovementComponent.cpp
@@ -57,6 +57,7 @@ FAutoConsoleVariableRef CVarChaosVehiclesShowBatchQueryExtents(TEXT("p.Vehicle.S
 FAutoConsoleVariableRef CVarChaosVehiclesDisableSuspensionForces(TEXT("p.Vehicle.DisableSuspensionForces"), GWheeledVehicleDebugParams.DisableSuspensionForces, TEXT("Enable/Disable Suspension Forces."));
 FAutoConsoleVariableRef CVarChaosVehiclesDisableFrictionForces(TEXT("p.Vehicle.DisableFrictionForces"), GWheeledVehicleDebugParams.DisableFrictionForces, TEXT("Enable/Disable Wheel Friction Forces."));
 FAutoConsoleVariableRef CVarChaosVehiclesDisableRollbarForces(TEXT("p.Vehicle.DisableRollbarForces"), GWheeledVehicleDebugParams.DisableRollbarForces, TEXT("Enable/Disable Rollbar Forces."));
+FAutoConsoleVariableRef CVarChaosVehiclesDisableConstraintSuspension(TEXT("p.Vehicle.DisableConstraintSuspension"), GWheeledVehicleDebugParams.DisableConstraintSuspension, TEXT("Enable/Disable Constraint based suspension, swaps to basic force based suspension without hardstops instead."));
 
 FAutoConsoleVariableRef CVarChaosVehiclesThrottleOverride(TEXT("p.Vehicle.ThrottleOverride"), GWheeledVehicleDebugParams.ThrottleOverride, TEXT("Hard code throttle input on."));
 FAutoConsoleVariableRef CVarChaosVehiclesSteeringOverride(TEXT("p.Vehicle.SteeringOverride"), GWheeledVehicleDebugParams.SteeringOverride, TEXT("Hard code steering input on."));
@@ -555,7 +556,7 @@ void UChaosWheeledVehicleSimulation::ApplySuspensionForces(float DeltaTime)
 		auto& PSuspension = PVehicle->Suspension[WheelIdx];
 		float SuspensionMovePosition = -PSuspension.Setup().MaxLength;
 
-		if (PWheel.Setup().NewSimulationPath)
+		if (!GWheeledVehicleDebugParams.DisableConstraintSuspension)
 		{
 #if WITH_CHAOS
 			if (WheelIdx < ConstraintHandles.Num())
@@ -597,7 +598,7 @@ void UChaosWheeledVehicleSimulation::ApplySuspensionForces(float DeltaTime)
 			FVector SusApplicationPoint = WheelState.WheelWorldLocation[WheelIdx] + PVehicle->Suspension[WheelIdx].Setup().SuspensionForceOffset;
 
 			check(PWheel.InContact());
-			if (!PWheel.Setup().NewSimulationPath)
+			if (GWheeledVehicleDebugParams.DisableConstraintSuspension)
 			{
 				AddForceAtPosition(SuspensionForceVector, SusApplicationPoint);
 			}
@@ -1057,7 +1058,7 @@ void UChaosWheeledVehicleMovementComponent::FixupSkeletalMesh()
 
 					}
 
-					if (Wheels[WheelIdx]->bNewWheelSimulation)
+					if (!GWheeledVehicleDebugParams.DisableConstraintSuspension)
 					{
 						FBodyInstance* TargetInstance = UpdatedPrimitive->GetBodyInstance();
 						if (TargetInstance)
diff --git a/Engine/Plugins/Experimental/ChaosVehiclesPlugin/Source/ChaosVehicles/Public/ChaosVehicleWheel.h b/Engine/Plugins/Experimental/ChaosVehiclesPlugin/Source/ChaosVehicles/Public/ChaosVehicleWheel.h
--- a/Engine/Plugins/Experimental/ChaosVehiclesPlugin/Source/ChaosVehicles/Public/ChaosVehicleWheel.h
+++ b/Engine/Plugins/Experimental/ChaosVehiclesPlugin/Source/ChaosVehicles/Public/ChaosVehicleWheel.h
@@ -98,17 +98,6 @@ class UChaosWheeledVehicleMovementComponent;
 		//UPROPERTY(EditAnywhere, Category = Wheel, meta = (ClampMin = "0.01", UIMin = "0.01"))
 		//float WheelMass;
 
-		UPROPERTY(EditAnywhere, Category = Wheel)
-		bool bNewWheelSimulation;
-
-		/** Longitudinal Friction Force Multiplier */
-		UPROPERTY(EditAnywhere, Category = Wheel, meta = (EditCondition = "!bNewWheelSimulation"))
-		float LongitudinalFrictionForceMultiplier;
-
-		/** Lateral Friction Force Multiplier */
-		UPROPERTY(EditAnywhere, Category = Wheel, meta = (EditCondition = "!bNewWheelSimulation"))
-		float LateralFrictionForceMultiplier;
-
 		/** Tyre Cornering Ability */
 		UPROPERTY(EditAnywhere, Category = Wheel, meta = (EditCondition = "bNewWheelSimulation"))
 		float CorneringStiffness;
@@ -375,11 +364,8 @@ class UChaosWheeledVehicleMovementComponent;
 			PWheelConfig.ABSEnabled = this->bABSEnabled;
 			PWheelConfig.TractionControlEnabled = this->bTractionControlEnabled;
 			PWheelConfig.AxleType = static_cast<Chaos::FSimpleWheelConfig::EAxleType>(this->AxleType);
-			PWheelConfig.LongitudinalFrictionMultiplier = this->LongitudinalFrictionForceMultiplier;
-			PWheelConfig.LateralFrictionMultiplier = this->LateralFrictionForceMultiplier;
 			PWheelConfig.FrictionMultiplier = this->FrictionForceMultiplier;
 			PWheelConfig.CorneringStiffness = this->CorneringStiffness * 10000.0f;
-			PWheelConfig.NewSimulationPath = this->bNewWheelSimulation;
 			PWheelConfig.SideSlipModifier = this->SideSlipModifier;
 			PWheelConfig.SlipThreshold = this->SlipThreshold;
 			PWheelConfig.SkidThreshold = this->SkidThreshold;
diff --git a/Engine/Plugins/Experimental/ChaosVehiclesPlugin/Source/ChaosVehicles/Public/ChaosWheeledVehicleMovementComponent.h b/Engine/Plugins/Experimental/ChaosVehiclesPlugin/Source/ChaosVehicles/Public/ChaosWheeledVehicleMovementComponent.h
--- a/Engine/Plugins/Experimental/ChaosVehiclesPlugin/Source/ChaosVehicles/Public/ChaosWheeledVehicleMovementComponent.h
+++ b/Engine/Plugins/Experimental/ChaosVehiclesPlugin/Source/ChaosVehicles/Public/ChaosWheeledVehicleMovementComponent.h
@@ -28,6 +28,7 @@ struct FWheeledVehicleDebugParams
 	bool DisableSuspensionForces = false;
 	bool DisableFrictionForces = false;
 	bool DisableRollbarForces = false;
+	bool DisableConstraintSuspension = false;
 
 	float ThrottleOverride = 0.f;
 	float SteeringOverride = 0.f;
diff --git a/Engine/Source/Runtime/Experimental/ChaosVehicles/ChaosVehiclesCore/Private/WheelSystem.cpp b/Engine/Source/Runtime/Experimental/ChaosVehicles/ChaosVehiclesCore/Private/WheelSystem.cpp
--- a/Engine/Source/Runtime/Experimental/ChaosVehicles/ChaosVehiclesCore/Private/WheelSystem.cpp
+++ b/Engine/Source/Runtime/Experimental/ChaosVehicles/ChaosVehiclesCore/Private/WheelSystem.cpp
@@ -46,165 +46,6 @@ namespace Chaos
 	}
 
 	void FSimpleWheelSim::Simulate(float DeltaTime)
-	{
-		if (Setup().NewSimulationPath)
-		{
-			SimulateNew(DeltaTime);
-
-			return;
-		}
-
-		SlipAngle = FMath::Atan2(GroundVelocityVector.Y, GroundVelocityVector.X);
-
-		float FinalLongitudinalForce = 0.f;
-
-		// The physics system is mostly unit-less i.e. can work in meters or cm, however there are 
-		// a couple of places where the results are wrong if Cm is used. This is one of them, the simulated radius
-		// for torque must be real size to obtain the correct output values.
-		AppliedLinearDriveForce = DriveTorque / Re;
-		AppliedLinearBrakeForce = BrakeTorque / Re;
-
-		// currently just letting the brake override the throttle
-		bool Braking = BrakeTorque > FMath::Abs(DriveTorque);
-		float BrakeFactor = 1.0f;
-		float K = 0.4f;
-
-		// are we actually touching the ground
-		if (ForceIntoSurface > SMALL_NUMBER)
-		{
-			LongitudinalAdhesiveLimit = ForceIntoSurface * SurfaceFriction * Setup().LongitudinalFrictionMultiplier;
-			LateralAdhesiveLimit = ForceIntoSurface * SurfaceFriction * Setup().LateralFrictionMultiplier;
-
-			if (Braking)
-			{
-
-				// whether the velocity is +ve or -ve when we brake we are slowing the vehicle down
-				// so force is opposing current direction of travel.
-				float ForceRequiredToBringToStop = MassPerWheel * K * (GroundVelocityVector.X) / DeltaTime;
-				FinalLongitudinalForce = AppliedLinearBrakeForce;
-
-				// check we are not applying more force than required so we end up overshooting 
-				// and accelerating in the opposite direction
-				FinalLongitudinalForce = FMath::Clamp(FinalLongitudinalForce, -FMath::Abs(ForceRequiredToBringToStop), FMath::Abs(ForceRequiredToBringToStop));
-
-				// ensure the brake opposes current direction of travel
-				if (GroundVelocityVector.X > 0.0f)
-				{
-					FinalLongitudinalForce = -FinalLongitudinalForce;
-				}
-
-			}
-			else
-			{
-				FinalLongitudinalForce = AppliedLinearDriveForce;
-			}
-
-			// lateral grip
-			float FinalLateralForce = -(MassPerWheel * K * GroundVelocityVector.Y) / DeltaTime;
-
-			ForceFromFriction.X = FinalLongitudinalForce;
-
-			float DynamicFrictionLongitudialScaling = 0.75f;
-			float TractionControlAndABSScaling = 0.98f;	// how close to perfection is the system working
-
-			SideSlipModifier = 1.0f;
-			bool Locked = false;
-			bool Spinning = false;
-
-			// we can only obtain as much accel/decel force as the friction will allow
-			if (FMath::Abs(FinalLongitudinalForce) > LongitudinalAdhesiveLimit)
-			{
-				if (Braking)
-				{
-					BrakeFactor = FMath::Clamp(LongitudinalAdhesiveLimit / FMath::Abs(FinalLongitudinalForce), 0.6f, 1.0f);
-				}
-
-				if ((Braking && ABSEnabled) || (!Braking && TractionControlEnabled))
-				{
-					Spin = 0.0f;
-
-					float Sign = (FinalLongitudinalForce > 0.0f) ? 1.0f : -1.0f;
-					ForceFromFriction.X = LongitudinalAdhesiveLimit * TractionControlAndABSScaling * Sign;
-				}
-				else
-				{
-					if (!Braking)
-					{
-						Spinning = true;
-						Spin += 0.5f * DeltaTime;
-						Spin = FMath::Clamp(Spin, -2.f, 2.f);
-					}
-					else
-					{
-						Locked = true;
-					}
-
-					float Sign = (ForceFromFriction.X >= 0.0f) ? 1.0 : -1.0f;
-					ForceFromFriction.X = LongitudinalAdhesiveLimit * DynamicFrictionLongitudialScaling * Sign;
-				}
-			}
-			else
-			{
-				Spin = 0.0f;
-			}
-
-			static float DynamicFrictionLateralScaling = 0.75f;
-			if (Locked || Spinning)
-			{
-				SideSlipModifier *= Setup().SideSlipModifier;
-			}
-
-			// Lateral needs more grip to feel right!
-			LateralAdhesiveLimit *= 1.0f * SideSlipModifier;
-			ForceFromFriction.Y = FinalLateralForce;
-			if (FMath::Abs(FinalLateralForce) > LateralAdhesiveLimit)
-			{
-				ForceFromFriction.Y = LateralAdhesiveLimit * DynamicFrictionLateralScaling;
-			}
-
-			if (FinalLateralForce < -LateralAdhesiveLimit)
-			{
-				ForceFromFriction.Y = -ForceFromFriction.Y;
-			}
-
-
-			// wheel rolling - just match the ground speed exactly
-			if (BrakeFactor < 1.0f)
-			{
-				Omega *= BrakeFactor;
-			}
-			else if (Spin > 0.1f)
-			{
-				Omega += Spin;
-			}
-			else
-			{
-				float GroundOmega = GroundVelocityVector.X / Re;
-				Omega += (GroundOmega - Omega);
-			}
-
-		}
-		// Wheel angular position
-		AngularPosition += Omega * DeltaTime;
-
-		while (AngularPosition >= PI * 2.f)
-		{
-			AngularPosition -= PI * 2.f;
-		}
-		while (AngularPosition <= -PI * 2.f)
-		{
-			AngularPosition += PI * 2.f;
-		}
-
-		if (!bInContact)
-		{
-			ForceFromFriction = FVector::ZeroVector;
-		}
-
-		return;
-	}
-
-	void FSimpleWheelSim::SimulateNew(float DeltaTime)
 	{
 		float K = 0.4f;
 		float TractionControlAndABSScaling = 0.98f;	// how close to perfection is the system working
diff --git a/Engine/Source/Runtime/Experimental/ChaosVehicles/ChaosVehiclesCore/Public/WheelSystem.h b/Engine/Source/Runtime/Experimental/ChaosVehicles/ChaosVehiclesCore/Public/WheelSystem.h
--- a/Engine/Source/Runtime/Experimental/ChaosVehicles/ChaosVehiclesCore/Public/WheelSystem.h
+++ b/Engine/Source/Runtime/Experimental/ChaosVehicles/ChaosVehiclesCore/Public/WheelSystem.h
@@ -85,15 +85,12 @@ struct CHAOSVEHICLESCORE_API FSimpleWheelConfig
 		, TractionControlEnabled(false)
 		, AxleType(EAxleType::Undefined)
 		, FrictionCombineMethod(EFrictionCombineMethod::Multiply)
-		, LateralFrictionMultiplier(1.0f)
-		, LongitudinalFrictionMultiplier(1.0f)
 		, FrictionMultiplier(2.0f)
 		, LateralSlipGraphMultiplier(1.0f)
 		, CorneringStiffness(1000.0f)
 		, SideSlipModifier(1.0f)
 		, SlipThreshold(20.0f)
 		, SkidThreshold(20.0f)
-		, NewSimulationPath(false)
 		, MaxSpinRotation(30.0f)
 	{
 
@@ -126,8 +123,6 @@ struct CHAOSVEHICLESCORE_API FSimpleWheelConfig
 
 	EFrictionCombineMethod FrictionCombineMethod; //#todo: use this variable
 
-	float LateralFrictionMultiplier;
-	float LongitudinalFrictionMultiplier;
 	float FrictionMultiplier;
 	float LateralSlipGraphMultiplier;
 	float CorneringStiffness;
@@ -136,8 +131,6 @@ struct CHAOSVEHICLESCORE_API FSimpleWheelConfig
 	float SlipThreshold;
 	float SkidThreshold;
 
-	bool NewSimulationPath;
-
 	FGraph LateralSlipGraph;
 	float MaxSpinRotation;
 
@@ -380,13 +373,8 @@ public:
 	 *	Wheel load force from body weight and the surface friction together determine the grip available at the wheel
 	 *	DriveTorque accelerates the wheel
 	 *	BrakeTorque decelerates the wheel
-	 *
-	 *	#todo: this function is a mess and needs tidied/consolidated
-	 *	#todo: lateral friction is currently not affected by longitudinal wheel slip
-	 *	#todo: wheel slip angle isn't being used
 	 */
 	void Simulate(float DeltaTime);
-	void SimulateNew(float DeltaTime);
 
 
 	void SetMassPerWheel(float VehicleMassPerWheel)
diff --git a/Templates/TP_VehicleAdv/Source/TP_VehicleAdv/TP_VehicleAdvWheelFront.cpp b/Templates/TP_VehicleAdv/Source/TP_VehicleAdv/TP_VehicleAdvWheelFront.cpp
--- a/Templates/TP_VehicleAdv/Source/TP_VehicleAdv/TP_VehicleAdvWheelFront.cpp
+++ b/Templates/TP_VehicleAdv/Source/TP_VehicleAdv/TP_VehicleAdvWheelFront.cpp
@@ -9,13 +9,16 @@ UTP_VehicleAdvWheelFront::UTP_VehicleAdvWheelFront()
 {
 	WheelRadius = 18.f;
 	WheelWidth = 20.0f;
-	LongitudinalFrictionForceMultiplier = 2.5f;
-	LateralFrictionForceMultiplier = 2.0f;
+	FrictionForceMultiplier = 2.0f;
+	CorneringStiffness = 500.0f;
+	bAffectedByEngine = false;
 	bAffectedByHandbrake = false;
 	bAffectedBySteering = true;
 	AxleType = EAxleType::Front;
-	SpringRate = 400.0f;
+	SpringRate = 200.0f;
 	SpringPreload = 100.f;
+	SuspensionDampingRatio = 0.5f;
+	WheelLoadRatio = 0.0f;
 	RollbarScaling = 0.5f;
 	SuspensionMaxRaise = 8;
 	SuspensionMaxDrop = 12;
diff --git a/Templates/TP_VehicleAdv/Source/TP_VehicleAdv/TP_VehicleAdvWheelRear.cpp b/Templates/TP_VehicleAdv/Source/TP_VehicleAdv/TP_VehicleAdvWheelRear.cpp
--- a/Templates/TP_VehicleAdv/Source/TP_VehicleAdv/TP_VehicleAdvWheelRear.cpp
+++ b/Templates/TP_VehicleAdv/Source/TP_VehicleAdv/TP_VehicleAdvWheelRear.cpp
@@ -9,13 +9,16 @@ UTP_VehicleAdvWheelRear::UTP_VehicleAdvWheelRear()
 {
 	WheelRadius = 18.f;
 	WheelWidth = 20.0f;
-	LongitudinalFrictionForceMultiplier = 2.5f;
-	LateralFrictionForceMultiplier = 2.0f;
+	FrictionForceMultiplier = 2.0f;
+	CorneringStiffness = 500.0f;
+	bAffectedByEngine = true;
 	bAffectedByHandbrake = true;
 	bAffectedBySteering = false;
 	AxleType = EAxleType::Rear;
-	SpringRate = 400.0f;
+	SpringRate = 200.0f;
 	SpringPreload = 100.f;
+	SuspensionDampingRatio = 0.5f;
+	WheelLoadRatio = 0.0f;
 	RollbarScaling = 0.5f;
 	SuspensionMaxRaise = 8;
 	SuspensionMaxDrop = 12.0f;
