Chaos Physics - Add support for being able to update Clustered particles by setting their KinematicTarget properties
- Changed GeometryCollection proxy to use KinematicTarget property instead of setting X/R directly
- Make sure to update P/Q in UpdateKinematicTargets when the particle is a rigid particle
- Consolidated PBDRigidsSOA Arrays and Map pairs into one single class
- Keep track of Dynamic/Kinematic/Static clustered particles
- Made sure PDBRigidClustering is calling SetParticleObjectState on the evolution instead of setting it directly on the particle

diff --git a/Engine/Source/Runtime/Experimental/Chaos/Private/Chaos/PBDRigidClustering.cpp b/Engine/Source/Runtime/Experimental/Chaos/Private/Chaos/PBDRigidClustering.cpp
--- a/Engine/Source/Runtime/Experimental/Chaos/Private/Chaos/PBDRigidClustering.cpp
+++ b/Engine/Source/Runtime/Experimental/Chaos/Private/Chaos/PBDRigidClustering.cpp
@@ -1344,10 +1344,16 @@ namespace Chaos
 	{
 		SCOPE_CYCLE_COUNTER(STAT_BreakingModel);
 
+		//make copy because release cluster modifies active indices. We want to iterate over original active indices
+		TArray<FPBDRigidClusteredParticleHandle*> ClusteredParticlesToProcess;
+		for (auto& Particle : MEvolution.GetNonDisabledClusteredView())
+		{
+			ClusteredParticlesToProcess.Add(Particle.Handle()->CastToClustered());
+		}
+
 		TMap<FPBDRigidClusteredParticleHandle*, TSet<FPBDRigidParticleHandle*>> AllActivatedChildren;
 
-		auto NonDisabledClusteredParticles = MEvolution.GetNonDisabledClusteredArray(); //make copy because release cluster modifies active indices. We want to iterate over original active indices
-		for (Chaos::TPBDRigidClusteredParticleHandleImp<FReal, 3, true>* ClusteredParticle : NonDisabledClusteredParticles)
+		for (FPBDRigidClusteredParticleHandle* ClusteredParticle : ClusteredParticlesToProcess)
 		{
 			if (ClusteredParticle->ClusterIds().NumChildren)
 			{
@@ -1454,7 +1460,7 @@ namespace Chaos
 					}
 				}
 
-				ClusteredCurrentNode->SetObjectState(ObjectState);
+				MEvolution.SetParticleObjectState(ClusteredCurrentNode, ObjectState);
 			}
 		}
 	}
diff --git a/Engine/Source/Runtime/Experimental/Chaos/Private/PhysicsProxy/GeometryCollectionPhysicsProxy.cpp b/Engine/Source/Runtime/Experimental/Chaos/Private/PhysicsProxy/GeometryCollectionPhysicsProxy.cpp
--- a/Engine/Source/Runtime/Experimental/Chaos/Private/PhysicsProxy/GeometryCollectionPhysicsProxy.cpp
+++ b/Engine/Source/Runtime/Experimental/Chaos/Private/PhysicsProxy/GeometryCollectionPhysicsProxy.cpp
@@ -1460,7 +1460,7 @@ void FGeometryCollectionPhysicsProxy::PushStateOnGameThread(Chaos::FPBDRigidsSol
 void FGeometryCollectionPhysicsProxy::PushToPhysicsState()
 {
 	// CONTEXT: PHYSICSTHREAD
-	// because the aattached actor can be dynamic, we need to update the kinematic particles properly
+	// because the attached actor can be dynamic, we need to update the kinematic particles properly
 	if (bIsPhysicsThreadWorldTransformDirty)
 	{
 		const FTransform& ActorToWorld = Parameters.WorldTransform;
@@ -1469,23 +1469,19 @@ void FGeometryCollectionPhysicsProxy::PushToPhysicsState()
 		for (int32 TransformGroupIndex = 0; TransformGroupIndex < NumTransformGroupElements; ++TransformGroupIndex)
 		{
 			Chaos::FPBDRigidClusteredParticleHandle* Handle = SolverParticleHandles[TransformGroupIndex];
-			if (!Handle)
+			if (!Handle || Handle->Disabled())
 			{
 				continue;
 			}
 			const Chaos::EObjectStateType ObjectState = Handle->ObjectState();
-			switch (ObjectState)
-			{
-			case Chaos::EObjectStateType::Kinematic:
+			if (ObjectState == Chaos::EObjectStateType::Kinematic)
 			{
-				FTransform LocalTransform = PhysicsThreadCollection.Transform[TransformGroupIndex];
-				FTransform WorldTransform = PhysicsThreadCollection.MassToLocal[TransformGroupIndex] * PhysicsThreadCollection.Transform[TransformGroupIndex] * ActorToWorld;;
+				FTransform WorldTransform = PhysicsThreadCollection.MassToLocal[TransformGroupIndex] * PhysicsThreadCollection.Transform[TransformGroupIndex] * ActorToWorld;
 
-				Handle->SetX(WorldTransform.GetTranslation());
-				Handle->SetR(WorldTransform.GetRotation().GetNormalized());
-				Handle->SetP(Handle->X());
-				Handle->SetQ(Handle->R());
-			}
+				Chaos::TKinematicTarget<Chaos::FReal, 3> newKinematicTarget;
+				Chaos::FRigidTransform3 PreviousWorldTransform(Handle->X(), Handle->R());
+				newKinematicTarget.SetTargetMode(WorldTransform, PreviousWorldTransform);
+				Handle->SetKinematicTarget(newKinematicTarget);
 			}
 		}
 	}
diff --git a/Engine/Source/Runtime/Experimental/Chaos/Public/Chaos/PBDRigidsEvolution.h b/Engine/Source/Runtime/Experimental/Chaos/Public/Chaos/PBDRigidsEvolution.h
--- a/Engine/Source/Runtime/Experimental/Chaos/Public/Chaos/PBDRigidsEvolution.h
+++ b/Engine/Source/Runtime/Experimental/Chaos/Public/Chaos/PBDRigidsEvolution.h
@@ -424,7 +424,15 @@ public:
 		EObjectStateType InitialState = Particle->ObjectState();
 
 		Particle->SetObjectStateLowLevel(ObjectState);
-		Particles.SetDynamicParticleSOA(Particle);
+
+		if (auto ClusteredParticle = Particle->CastToClustered())
+		{
+			Particles.SetClusteredParticleSOA(ClusteredParticle);
+		}
+		else
+		{
+			Particles.SetDynamicParticleSOA(Particle);
+		}
 
 		if (InitialState == EObjectStateType::Sleeping && InitialState != ObjectState && !Particle->Disabled())
 		{
@@ -524,7 +532,7 @@ public:
 		}
 	}
 
-	const auto& GetNonDisabledClusteredArray() const { return Particles.GetNonDisabledClusteredArray(); }
+	const TParticleView<FPBDRigidClusteredParticles>& GetNonDisabledClusteredView() const { return Particles.GetNonDisabledClusteredView(); }
 
 	CHAOS_API TSerializablePtr<FChaosPhysicsMaterial> GetPhysicsMaterial(const FGeometryParticleHandle* Particle) const { return Particle->AuxilaryValue(PhysicsMaterials); }
 	
@@ -694,11 +702,13 @@ public:
 			}
 			}
 			
-			// Set previous velocities if we can
+			// Set previous and positions and velocities if we can
 			// Note: At present kininematics are in fact rigid bodies
 			auto* Rigid = Particle.CastToRigidParticle();
 			if (Rigid)
 			{
+				Rigid->P() = Rigid->X();
+				Rigid->Q() = Rigid->R();
 				Rigid->PreV() = Rigid->V();
 				Rigid->PreW() = Rigid->W();
 			}
diff --git a/Engine/Source/Runtime/Experimental/Chaos/Public/Chaos/PBDRigidsEvolutionGBF.h b/Engine/Source/Runtime/Experimental/Chaos/Public/Chaos/PBDRigidsEvolutionGBF.h
--- a/Engine/Source/Runtime/Experimental/Chaos/Public/Chaos/PBDRigidsEvolutionGBF.h
+++ b/Engine/Source/Runtime/Experimental/Chaos/Public/Chaos/PBDRigidsEvolutionGBF.h
@@ -61,7 +61,7 @@ namespace Chaos
 		using Base::EnableParticle;
 		using Base::DisableParticles;
 		using Base::NumIslands;
-		using Base::GetNonDisabledClusteredArray;
+		using Base::GetNonDisabledClusteredView;
 		using Base::DisableParticle;
 		using Base::PrepareIteration;
 		using Base::GetConstraintGraph;
diff --git a/Engine/Source/Runtime/Experimental/Chaos/Public/Chaos/PBDRigidsSOAs.h b/Engine/Source/Runtime/Experimental/Chaos/Public/Chaos/PBDRigidsSOAs.h
--- a/Engine/Source/Runtime/Experimental/Chaos/Public/Chaos/PBDRigidsSOAs.h
+++ b/Engine/Source/Runtime/Experimental/Chaos/Public/Chaos/PBDRigidsSOAs.h
@@ -90,6 +90,88 @@ private:
 	volatile int8 Block;
 };
 
+template <typename TParticleType>
+class TParticleMapArray
+{
+public:
+	void Reset()
+	{
+		ParticleToIndex.Reset();
+		ParticleArray.Reset();
+	}
+
+	template <typename TParticle1>
+	void Insert(const TArray<TParticle1*>& ParticlesToInsert)
+	{
+		TArray<bool> Contains;
+		Contains.AddZeroed(ParticlesToInsert.Num());
+
+		// TODO: Compile time check ensuring TParticle2 is derived from TParticle1?
+		int32 NextIdx = ParticleArray.Num();
+		for (int32 Idx = 0; Idx < ParticlesToInsert.Num(); ++Idx)
+		{
+			auto Particle = ParticlesToInsert[Idx];
+			Contains[Idx] = ParticleToIndex.Contains(Particle);
+			if (!Contains[Idx])
+			{
+				ParticleToIndex.Add(Particle, NextIdx++);
+			}
+		}
+		ParticleArray.Reserve(ParticleArray.Num() + NextIdx - ParticleArray.Num());
+		for (int32 Idx = 0; Idx < ParticlesToInsert.Num(); ++Idx)
+		{
+			if (!Contains[Idx])
+			{
+				auto Particle = ParticlesToInsert[Idx];
+				ParticleArray.Add(Particle);
+			}
+		}
+	}
+
+	void Insert(TParticleType* Particle)
+	{
+		if (ParticleToIndex.Contains(Particle) == false)
+		{
+			ParticleToIndex.Add(Particle, ParticleArray.Num());
+			ParticleArray.Add(Particle);
+		}
+	}
+
+	void Remove(TParticleType* Particle)
+	{
+		if (int32* IdxPtr = ParticleToIndex.Find(Particle))
+		{
+			int32 Idx = *IdxPtr;
+			ParticleArray.RemoveAtSwap(Idx);
+			if (Idx < ParticleArray.Num())
+			{
+				//update swapped element with new index
+				ParticleToIndex[ParticleArray[Idx]] = Idx;
+			}
+			ParticleToIndex.Remove(Particle);
+		}
+	}
+
+	void Serialize(FChaosArchive& Ar)
+	{
+		TArray<TSerializablePtr<TParticleType>>& SerializableArray = AsAlwaysSerializableArray(ParticleArray);
+		Ar << SerializableArray;
+
+		int32 Idx = 0;
+		for (auto Particle : ParticleArray)
+		{
+			ParticleToIndex.Add(Particle, Idx++);
+		}
+	}
+
+	const TArray<TParticleType*>& GetArray() const { return ParticleArray;	}
+	TArray<TParticleType*>& GetArray() { return ParticleArray; }
+
+private:
+	TMap<TParticleType*, int32> ParticleToIndex;
+	TArray<TParticleType*> ParticleArray;
+};
+
 class FPBDRigidsSOAs
 {
 public:
@@ -152,7 +234,7 @@ public:
 
 		if (!Params.bStartSleeping)
 		{
-			InsertToMapAndArray(Results, ActiveParticlesToIndex, ActiveParticlesArray);
+			ActiveParticlesMapArray.Insert(Results);
 		}
 		UpdateViews();
 		return Results;
@@ -185,12 +267,12 @@ public:
 		
 		if (!Params.bDisabled)
 		{
-			InsertToMapAndArray(NewClustered, NonDisabledClusteredToIndex, NonDisabledClusteredArray);
+			InsertClusteredParticles(NewClustered);
 		}
 
 		if (!Params.bStartSleeping)
 		{
-			InsertToMapAndArray(reinterpret_cast<TArray<FPBDRigidParticleHandle*>&>(NewClustered), ActiveParticlesToIndex, ActiveParticlesArray);
+			ActiveParticlesMapArray.Insert(reinterpret_cast<TArray<FPBDRigidParticleHandle*>&>(NewClustered));
 		}
 
 		UpdateViews();
@@ -200,8 +282,7 @@ public:
 	
 	void ClearTransientDirty()
 	{
-		TransientDirtyToIndex.Reset();
-		TransientDirtyArray.Reset();
+		TransientDirtyMapArray.Reset();
 	}
 
 	void MarkTransientDirtyParticle(FGeometryParticleHandle* Particle)
@@ -209,7 +290,7 @@ public:
 		FPBDRigidParticleHandle* Rigid =  Particle->CastToRigidParticle();
 		if(Rigid)
 		{
-			InsertToMapAndArray(Rigid, TransientDirtyToIndex, TransientDirtyArray);
+			TransientDirtyMapArray.Insert(Rigid);
 		}
 
 		UpdateViews();
@@ -220,8 +301,8 @@ public:
 		auto PBDRigid = Particle->CastToRigidParticle();
 		if(PBDRigid)
 		{
-			RemoveFromMapAndArray(PBDRigid, ActiveParticlesToIndex, ActiveParticlesArray);
-			RemoveFromMapAndArray(PBDRigid, TransientDirtyToIndex, TransientDirtyArray);
+			ActiveParticlesMapArray.Remove(PBDRigid);
+			TransientDirtyMapArray.Remove(PBDRigid);
 
 			if (auto PBDRigidClustered = Particle->CastToClustered())
 			{
@@ -231,8 +312,7 @@ public:
 				}
 				else // clustered
 				{
-					RemoveFromMapAndArray(PBDRigidClustered,
-						NonDisabledClusteredToIndex, NonDisabledClusteredArray);
+					DynamicClusteredMapArray.Remove(PBDRigidClustered);
 				}
 			}
 			else
@@ -281,8 +361,7 @@ public:
 				}
 				else // clustered
 				{
-					RemoveFromMapAndArray(PBDRigidClustered, 
-						NonDisabledClusteredToIndex, NonDisabledClusteredArray);
+					RemoveClusteredParticle(PBDRigidClustered);
 				}
 			}
 			else
@@ -292,7 +371,7 @@ public:
 
 			// All active particles RIGID particles
 			{
-				RemoveFromMapAndArray(PBDRigid, ActiveParticlesToIndex, ActiveParticlesArray);
+				ActiveParticlesMapArray.Remove(PBDRigid);
 			}
 		}
 		else if (Particle->CastToKinematicParticle())
@@ -326,8 +405,7 @@ public:
 				}
 				else // clustered
 				{
-					InsertToMapAndArray(PBDRigidClustered, 
-						NonDisabledClusteredToIndex, NonDisabledClusteredArray);
+					InsertClusteredParticle(PBDRigidClustered);
 				}
 			}
 			else
@@ -337,8 +415,7 @@ public:
 
 			if (!PBDRigid->Sleeping() && Particle->ObjectState() == EObjectStateType::Dynamic)
 			{
-				InsertToMapAndArray(PBDRigid, 
-					ActiveParticlesToIndex, ActiveParticlesArray);
+				ActiveParticlesMapArray.Insert(PBDRigid);
 			}
 		}
 		else if (Particle->CastToKinematicParticle())
@@ -381,8 +458,7 @@ public:
 					else
 					{
 						// Non clustered rigid particles:
-						InsertToMapAndArray(PBDRigid, 
-							ActiveParticlesToIndex, ActiveParticlesArray);
+						ActiveParticlesMapArray.Insert(PBDRigid);
 					}
 
 					UpdateViews();
@@ -433,7 +509,7 @@ public:
 					PBDRigid->SetSleeping(true);
 					PBDRigid->SetObjectState(EObjectStateType::Sleeping);
 
-					InsertToMapAndArray(PBDRigid,TransientDirtyToIndex,TransientDirtyArray);
+					TransientDirtyMapArray.Insert(PBDRigid);
 
 					if (auto PBDRigidClustered = Particle->CastToClustered())
 					{
@@ -444,8 +520,7 @@ public:
 					}
 					else
 					{
-						RemoveFromMapAndArray(PBDRigid, 
-							ActiveParticlesToIndex, ActiveParticlesArray);
+						ActiveParticlesMapArray.Remove(PBDRigid);
 					}
 
 					if (!DeferUpdateViews)
@@ -477,17 +552,17 @@ public:
 		if (Particle->Disabled())
 		{
 			Particle->MoveToSOA(*DynamicDisabledParticles);
-			RemoveFromMapAndArray(Particle->CastToRigidParticle(), ActiveParticlesToIndex, ActiveParticlesArray);
+			ActiveParticlesMapArray.Remove(Particle->CastToRigidParticle());
 		}
 		else
 		{
 			if (Particle->ObjectState() != EObjectStateType::Dynamic)
 			{
-				RemoveFromMapAndArray(Particle->CastToRigidParticle(), ActiveParticlesToIndex, ActiveParticlesArray);
+				ActiveParticlesMapArray.Remove(Particle->CastToRigidParticle());
 			}
 			else
 			{
-				InsertToMapAndArray(Particle->CastToRigidParticle(), ActiveParticlesToIndex, ActiveParticlesArray);
+				ActiveParticlesMapArray.Insert(Particle->CastToRigidParticle());
 			}
 
 			// Move to appropriate dynamic SOA
@@ -510,6 +585,13 @@ public:
 		UpdateViews();
 	}
 
+	void SetClusteredParticleSOA(FPBDRigidClusteredParticleHandle* ClusteredParticle)
+	{
+		RemoveClusteredParticle(ClusteredParticle);
+		InsertClusteredParticle(ClusteredParticle);
+		UpdateViews();
+	}
+
 	void Serialize(FChaosArchive& Ar)
 	{
 		static const FName SOAsName = TEXT("PBDRigidsSOAs");
@@ -559,8 +641,8 @@ public:
 		//Ar << ClusteredParticles;
 		Ar << GeometryCollectionParticles;
 
-		SerializeMapAndArray(Ar, ActiveParticlesToIndex, ActiveParticlesArray);
-		//SerializeMapAndArray(Ar, NonDisabledClusteredToIndex, NonDisabledClusteredArray);
+		ActiveParticlesMapArray.Serialize(Ar);
+		//DynamicClusteredMapArray.Serialize(Ar);
 
 		//todo: update deterministic ID
 
@@ -575,10 +657,12 @@ public:
 
 	const TParticleView<FPBDRigidParticles>& GetNonDisabledDynamicView() const { UpdateIfNeeded(); return NonDisabledDynamicView; }
 
+	const TParticleView<FPBDRigidClusteredParticles>& GetNonDisabledClusteredView() const { return NonDisabledClusteredView; }
+
 	const TParticleView<FPBDRigidParticles>& GetActiveParticlesView() const { UpdateIfNeeded(); return ActiveParticlesView; }
 	TParticleView<FPBDRigidParticles>& GetActiveParticlesView() { UpdateIfNeeded(); return ActiveParticlesView; }
 
-	const TArray<FPBDRigidParticleHandle*>& GetActiveParticlesArray() const { return ActiveParticlesArray; }
+	const TArray<FPBDRigidParticleHandle*>& GetActiveParticlesArray() const { return ActiveParticlesMapArray.GetArray(); }
 	
 	const TParticleView<FPBDRigidParticles>& GetDirtyParticlesView() const { UpdateIfNeeded(); return DirtyParticlesView; }
 	TParticleView<FPBDRigidParticles>& GetDirtyParticlesView() { UpdateIfNeeded(); return DirtyParticlesView; }
@@ -726,8 +810,6 @@ public:
 		bGeometryCollectionDirty = false;
 	}
 
-	const auto& GetNonDisabledClusteredArray() const { return NonDisabledClusteredArray; }
-
 	const auto& GetClusteredParticles() const { return *ClusteredParticles; }
 	auto& GetClusteredParticles() { return *ClusteredParticles; }
 
@@ -767,73 +849,6 @@ private:
 
 		return ReturnHandles;
 	}
-
-	template <typename TParticle1, typename TParticle2>
-	void InsertToMapAndArray(const TArray<TParticle1*>& ParticlesToInsert, TMap<TParticle2*, int32>& ParticleToIndex, TArray<TParticle2*>& ParticleArray)
-	{
-		TArray<bool> Contains;
-		Contains.AddZeroed(ParticlesToInsert.Num());
-
-		// TODO: Compile time check ensuring TParticle2 is derived from TParticle1?
-		int32 NextIdx = ParticleArray.Num();
-		for(int32 Idx = 0; Idx < ParticlesToInsert.Num(); ++Idx)
-		{
-			auto Particle = ParticlesToInsert[Idx];
-			Contains[Idx] = ParticleToIndex.Contains(Particle);
-			if (!Contains[Idx])
-			{
-				ParticleToIndex.Add(Particle, NextIdx++);
-			}
-		}
-		ParticleArray.Reserve(ParticleArray.Num() + NextIdx - ParticleArray.Num());
-		for (int32 Idx = 0; Idx < ParticlesToInsert.Num(); ++Idx)
-		{
-			if (!Contains[Idx])
-			{
-				auto Particle = ParticlesToInsert[Idx];
-				ParticleArray.Add(Particle);
-			}
-		}
-	}
-
-	template <typename TParticle>
-	static void InsertToMapAndArray(TParticle* Particle, TMap<TParticle*, int32>& ParticleToIndex, TArray<TParticle*>& ParticleArray)
-	{
-		if (ParticleToIndex.Contains(Particle) == false)
-		{
-			ParticleToIndex.Add(Particle, ParticleArray.Num());
-			ParticleArray.Add(Particle);
-		}
-	}
-
-	template <typename TParticle>
-	static void RemoveFromMapAndArray(TParticle* Particle, TMap<TParticle*, int32>& ParticleToIndex, TArray<TParticle*>& ParticleArray)
-	{
-		if (int32* IdxPtr = ParticleToIndex.Find(Particle))
-		{
-			int32 Idx = *IdxPtr;
-			ParticleArray.RemoveAtSwap(Idx);
-			if (Idx < ParticleArray.Num())
-			{
-				//update swapped element with new index
-				ParticleToIndex[ParticleArray[Idx]] = Idx;
-			}
-			ParticleToIndex.Remove(Particle);
-		}
-	}
-
-	template <typename TParticle>
-	void SerializeMapAndArray(FChaosArchive& Ar, TMap<TParticle*, int32>& ParticleToIndex, TArray<TParticle*>& ParticleArray)
-	{
-		TArray<TSerializablePtr<TParticle>>& SerializableArray = AsAlwaysSerializableArray(ParticleArray);
-		Ar << SerializableArray;
-
-		int32 Idx = 0;
-		for (auto Particle : ParticleArray)
-		{
-			ParticleToIndex.Add(Particle, Idx++);
-		}
-	}
 	
 	//should be called whenever particles are added / removed / reordered
 	void UpdateViews()
@@ -846,7 +861,9 @@ private:
 				KinematicParticles.Get(), 
 				DynamicParticles.Get(),
 				DynamicKinematicParticles.Get(),
-				{&NonDisabledClusteredArray},
+				{&StaticClusteredMapArray.GetArray()},
+				{&KinematicClusteredMapArray.GetArray()},
+				{&DynamicClusteredMapArray.GetArray()},
 				{&StaticGeometryCollectionArray},
 				{&KinematicGeometryCollectionArray},
 				{&SleepingGeometryCollectionArray},
@@ -858,7 +875,7 @@ private:
 			TArray<TSOAView<FPBDRigidParticles>> TmpArray = 
 			{ 
 				DynamicParticles.Get(), 
-				{&NonDisabledClusteredArray}, 
+				{&DynamicClusteredMapArray.GetArray()},
 				{&SleepingGeometryCollectionArray},
 				{&DynamicGeometryCollectionArray}
 			};
@@ -867,8 +884,8 @@ private:
 		{
 			TArray<TSOAView<FPBDRigidParticles>> TmpArray = 
 			{ 
-				{&ActiveParticlesArray},
-			//	{&NonDisabledClusteredArray},  Cluster particles appear in the ActiveParticlesArray
+				{&ActiveParticlesMapArray.GetArray()},
+			//	{&DynamicClusteredMapArray.GetArray()},  Cluster particles appear in the ActiveParticlesArray
 				{&StaticGeometryCollectionArray},
 				{&KinematicGeometryCollectionArray},
 				{&DynamicGeometryCollectionArray}
@@ -878,12 +895,12 @@ private:
 		{
 			TArray<TSOAView<FPBDRigidParticles>> TmpArray =
 			{
-				{&ActiveParticlesArray},
-				//	{&NonDisabledClusteredArray},  Cluster particles appear in the ActiveParticlesArray
+				{&ActiveParticlesMapArray.GetArray()},
+				//	{&DynamicClusteredMapArray.GetArray()},  Cluster particles appear in the ActiveParticlesArray
 				{&StaticGeometryCollectionArray},
 				{&KinematicGeometryCollectionArray},
 				{&DynamicGeometryCollectionArray},
-				{&TransientDirtyArray}
+				{&TransientDirtyMapArray.GetArray()}
 			};
 			DirtyParticlesView = MakeParticleView(MoveTemp(TmpArray));
 		}
@@ -907,7 +924,8 @@ private:
 			{ 
 				KinematicParticles.Get(),
 				DynamicKinematicParticles.Get(),
-				{&KinematicGeometryCollectionArray}
+				{&KinematicGeometryCollectionArray},
+				{&KinematicClusteredMapArray.GetArray()}
 			};
 			ActiveKinematicParticlesView = MakeParticleView(MoveTemp(TmpArray));
 		}
@@ -915,6 +933,7 @@ private:
 			TArray<TSOAView<FGeometryParticles>> TmpArray =
 			{
 				StaticParticles.Get(),
+				{&StaticClusteredMapArray.GetArray()}
 			};
 			ActiveStaticParticlesView = MakeParticleView(MoveTemp(TmpArray));
 		}
@@ -927,6 +946,55 @@ private:
 			};
 			ActiveGeometryCollectionParticlesView = MakeParticleView(MoveTemp(TmpArray));
 		}
+
+		{
+			TArray<TSOAView<FPBDRigidClusteredParticles>> TmpArray =
+			{
+				{&StaticClusteredMapArray.GetArray()},
+				{&KinematicClusteredMapArray.GetArray()},
+				{&DynamicClusteredMapArray.GetArray()}
+			};
+			NonDisabledClusteredView = MakeParticleView(MoveTemp(TmpArray));
+		}
+		
+	}
+
+	void InsertClusteredParticle(FPBDRigidClusteredParticleHandle* ClusteredParticle)
+	{
+		if (!ClusteredParticle->Disabled())
+		{
+			switch (ClusteredParticle->ObjectState())
+			{
+			case EObjectStateType::Uninitialized:
+				ensure(false); // we should probably not be here 
+				break;
+			case EObjectStateType::Static:
+				StaticClusteredMapArray.Insert(ClusteredParticle);
+				break;
+			case EObjectStateType::Kinematic:
+				KinematicClusteredMapArray.Insert(ClusteredParticle);
+				break;
+			case EObjectStateType::Dynamic:
+			case EObjectStateType::Sleeping:
+				DynamicClusteredMapArray.Insert(ClusteredParticle);
+				break;
+			}
+		}
+	}
+
+	void InsertClusteredParticles(const TArray<FPBDRigidClusteredParticleHandle*>& ClusteredParticleArray)
+	{
+		for (FPBDRigidClusteredParticleHandle* ClusteredParticle : ClusteredParticleArray)
+		{
+			InsertClusteredParticle(ClusteredParticle);
+		}
+	}
+
+	void RemoveClusteredParticle(FPBDRigidClusteredParticleHandle* ClusteredParticle)
+	{
+		StaticClusteredMapArray.Remove(ClusteredParticle);
+		KinematicClusteredMapArray.Remove(ClusteredParticle);
+		DynamicClusteredMapArray.Remove(ClusteredParticle);
 	}
 
 	//Organized by SOA type
@@ -953,18 +1021,18 @@ private:
 	bool bGeometryCollectionDirty;
 
 	//Utility structures for maintaining an Active particles view
-	TMap<FPBDRigidParticleHandle*, int32> ActiveParticlesToIndex;
-	TArray<FPBDRigidParticleHandle*> ActiveParticlesArray;
-	TArray<FPBDRigidParticleHandle*> TransientDirtyArray;
-	TMap<FPBDRigidParticleHandle*,int32> TransientDirtyToIndex;
+	TParticleMapArray<FPBDRigidParticleHandle> ActiveParticlesMapArray;
+	TParticleMapArray<FPBDRigidParticleHandle> TransientDirtyMapArray;
 
-	//Utility structures for maintaining a NonDisabled particle view
-	TMap<FPBDRigidClusteredParticleHandle*, int32> NonDisabledClusteredToIndex;
-	TArray<FPBDRigidClusteredParticleHandle*> NonDisabledClusteredArray;
+	// NonDisabled clustered particle arrays
+	TParticleMapArray<FPBDRigidClusteredParticleHandle> StaticClusteredMapArray;
+	TParticleMapArray<FPBDRigidClusteredParticleHandle> KinematicClusteredMapArray;
+	TParticleMapArray<FPBDRigidClusteredParticleHandle> DynamicClusteredMapArray;
 
 	//Particle Views
 	TParticleView<FGeometryParticles> NonDisabledView;							//all particles that are not disabled
 	TParticleView<FPBDRigidParticles> NonDisabledDynamicView;						//all dynamic particles that are not disabled
+	TParticleView<FPBDRigidClusteredParticles> NonDisabledClusteredView;			//all clustered particles that are not disabled
 	TParticleView<FPBDRigidParticles> ActiveParticlesView;						//all particles that are active
 	TParticleView<FPBDRigidParticles> DirtyParticlesView;							//all particles that are active + any that were put to sleep this frame
 	TParticleView<FGeometryParticles> AllParticlesView;							//all particles
