Added ray-cone and ray-cylinder interesection tests based on implementations from the Geometric Tools library.

diff --git a/Engine/Source/Runtime/Experimental/InteractiveToolsFramework/Private/BaseGizmos/Geometry-Tools-Engine-License.txt b/Engine/Source/Runtime/Experimental/InteractiveToolsFramework/Private/BaseGizmos/Geometry-Tools-Engine-License.txt
new file mode 100644
--- /dev/null
+++ b/Engine/Source/Runtime/Experimental/InteractiveToolsFramework/Private/BaseGizmos/Geometry-Tools-Engine-License.txt
@@ -0,0 +1,23 @@
+Boost Software License - Version 1.0 - August 17th, 2003
+
+Permission is hereby granted, free of charge, to any person or organization
+obtaining a copy of the software and accompanying documentation covered by
+this license (the "Software") to use, reproduce, display, distribute,
+execute, and transmit the Software, and to prepare derivative works of the
+Software, and to permit third-parties to whom the Software is furnished to
+do so, all subject to the following:
+
+The copyright notices in the Software and this entire statement, including
+the above license grant, this restriction and the following disclaimer,
+must be included in all copies of the Software, in whole or in part, and
+all derivative works of the Software, unless such copies or derivative
+works are solely in the form of machine-executable object code generated by
+a source language processor.
+
+THE SOFTWARE IS PROVIDED "AS IS", WITHOUT WARRANTY OF ANY KIND, EXPRESS OR
+IMPLIED, INCLUDING BUT NOT LIMITED TO THE WARRANTIES OF MERCHANTABILITY,
+FITNESS FOR A PARTICULAR PURPOSE, TITLE AND NON-INFRINGEMENT. IN NO EVENT
+SHALL THE COPYRIGHT HOLDERS OR ANYONE DISTRIBUTING THE SOFTWARE BE LIABLE
+FOR ANY DAMAGES OR OTHER LIABILITY, WHETHER IN CONTRACT, TORT OR OTHERWISE,
+ARISING FROM, OUT OF OR IN CONNECTION WITH THE SOFTWARE OR THE USE OR OTHER
+DEALINGS IN THE SOFTWARE.
diff --git a/Engine/Source/Runtime/Experimental/InteractiveToolsFramework/Private/BaseGizmos/GizmoMath.cpp b/Engine/Source/Runtime/Experimental/InteractiveToolsFramework/Private/BaseGizmos/GizmoMath.cpp
--- a/Engine/Source/Runtime/Experimental/InteractiveToolsFramework/Private/BaseGizmos/GizmoMath.cpp
+++ b/Engine/Source/Runtime/Experimental/InteractiveToolsFramework/Private/BaseGizmos/GizmoMath.cpp
@@ -124,6 +124,481 @@ void GizmoMath::RaySphereIntersection(
 }
 
 
+template<typename RealType>
+void GizmoMath::RayCylinderIntersection(
+	const FVector& CylinderCenter, const FVector& CylinderAxis, RealType CylinderRadius, RealType CylinderHeight,
+	const FVector& RayOrigin, const FVector& RayDirection,
+	bool& bIntersectsOut, RealType& OutRayParam)
+{
+	// adapted from GeometricTools GTEngine
+	// https://www.geometrictools.com/GTE/Mathematics/IntrRay3Cylinder3.h and
+	// https://www.geometrictools.com/GTE/Mathematics/IntrLine3Cylinder3.h 
+	// (Engine/Plugins/Experimental/GeometryProcessing/Source/GeometryAlgorithms/Private/ThirdParty/GTEngine/Mathematics/GteIntrRay3Cylinder3.h and
+	//  Engine/Plugins/Experimental/GeometryProcessing/Source/GeometryAlgorithms/Private/ThirdParty/GTEngine/Mathematics/GteIntrLine3Cylinder3.h)	
+
+	// The cylinder axis is a line.  The origin of the cylinder is chosen to be
+	// the line origin.  The cylinder wall is at a distance R units from the axis.
+	// An infinite cylinder has infinite height.  A finite cylinder has center C
+	// at the line origin and has a finite height H.  The segment for the finite
+	// cylinder has endpoints C-(H/2)*D and C+(H/2)*D where D is a unit-length
+	// direction of the line.
+
+	// Initialize the result as if there is no intersection.  If we discover
+	// an intersection, these values will be modified accordingly.
+	bIntersectsOut  = false;
+	int NumIntersections = 0;
+	RealType RayParam[2];
+
+	// Create a coordinate system for the cylinder.  In this system, the
+	// cylinder segment center C is the origin and the cylinder axis direction
+	// W is the z-axis.  U and V are the other coordinate axis directions.
+	// If P = x*U+y*V+z*W, the cylinder is x^2 + y^2 = r^2, where r is the
+	// cylinder radius.  The end caps are |z| = h/2, where h is the cylinder
+	// height.
+	FVector Basis[3];  // {W, U, V}
+	Basis[0] = CylinderAxis.GetSafeNormal();
+	// @todo: replace with MakePerpVectors() once it has been moved out of plugins
+	Basis[1] = GetOrthogonalVector(Basis[0]).GetSafeNormal();
+	Basis[2] = (Basis[0] ^ Basis[1]).GetSafeNormal();
+
+	RealType HalfHeight = 0.5 * CylinderHeight;
+	RealType RadiusSquared = CylinderRadius * CylinderRadius;
+
+	// Convert incoming line origin to capsule coordinates.
+	FVector Diff = RayOrigin - CylinderCenter;
+	FVector P(Basis[1] | Diff, Basis[2] | Diff, Basis[0] | Diff);
+
+	// Get the z-value, in cylinder coordinates, of the incoming line's
+	// unit-length direction.
+	RealType Dz = Basis[0] | RayDirection;
+	if (FMath::Abs(Dz) == 1.0)
+	{
+		// The line is parallel to the cylinder axis.  Determine whether the
+		// line intersects the cylinder end disks.
+		RealType RadialSquaredDist = RadiusSquared - P[0] * P[0] - P[1] * P[1];
+		if (RadialSquaredDist >= 0.0)
+		{
+			// The line intersects the cylinder end disks.
+			NumIntersections = 2;
+			if (Dz > 0.0)
+			{
+				RayParam[0] = -P[2] - HalfHeight;
+				RayParam[1] = -P[2] + HalfHeight;
+			}
+			else
+			{
+				RayParam[0] = P[2] - HalfHeight;
+				RayParam[1] = P[2] + HalfHeight;
+			}
+		}
+		// else:  The line is outside the cylinder, no intersection.
+	}
+	else
+	{
+		// Convert the incoming line unit-length direction to cylinder
+		// coordinates.
+		FVector D((Basis[1] | RayDirection), (Basis[2] | RayDirection), Dz);
+
+		RealType A0, A1, A2, Discr, Root, Inv, TValue;
+
+		if (D[2] == 0.0)
+		{
+			// The line is perpendicular to the cylinder axis.
+			if (FMath::Abs(P[2]) <= HalfHeight)
+			{
+				// Test intersection of line P+t*D with infinite cylinder
+				// x^2+y^2 = r^2.  This reduces to computing the Roots of a
+				// quadratic equation.  If P = (px,py,pz) and D = (dx,dy,dz),
+				// then the quadratic equation is
+				//   (dx^2+dy^2)*t^2 + 2*(px*dx+py*dy)*t + (px^2+py^2-r^2) = 0
+				A0 = P[0] * P[0] + P[1] * P[1] - RadiusSquared;
+				A1 = P[0] * D[0] + P[1] * D[1];
+				A2 = D[0] * D[0] + D[1] * D[1];
+				Discr = A1 * A1 - A0 * A2;
+				if (Discr > 0.0)
+				{
+					// The line intersects the cylinder in two places.
+					NumIntersections = 2;
+					Root = FMath::Sqrt(Discr);
+					Inv = 1.0f / A2;
+					RayParam[0] = (-A1 - Root) * Inv;
+					RayParam[1] = (-A1 + Root) * Inv;
+				}
+				else if (Discr == 0.0)
+				{
+					// The line is tangent to the cylinder.
+					NumIntersections = 1;
+					RayParam[0] = -A1 / A2;
+					RayParam[1] = RayParam[0];
+				}
+				// else: The line does not intersect the cylinder.
+			}
+			// else: The line is outside the planes of the cylinder end disks.
+		}
+		else
+		{
+			// Test for intersections with the planes of the end disks.
+			Inv = 1.0f / D[2];
+
+			RealType T0 = (-HalfHeight - P[2]) * Inv;
+			RealType TmpX = P[0] + T0 * D[0];
+			RealType TmpY = P[1] + T0 * D[1];
+			if (TmpX * TmpX + TmpY * TmpY <= RadiusSquared)
+			{
+				// Plane intersection inside the top cylinder end disk.
+				RayParam[NumIntersections++] = T0;
+			}
+
+			RealType T1 = (+HalfHeight - P[2]) * Inv;
+			TmpX = P[0] + T1 * D[0];
+			TmpY = P[1] + T1 * D[1];
+			if (TmpX * TmpX + TmpY * TmpY <= RadiusSquared)
+			{
+				// Plane intersection inside the bottom cylinder end disk.
+				RayParam[NumIntersections++] = T1;
+			}
+
+			if (NumIntersections < 2)
+			{
+				// Test for intersection with the cylinder wall.
+				A0 = P[0] * P[0] + P[1] * P[1] - RadiusSquared;
+				A1 = P[0] * D[0] + P[1] * D[1];
+				A2 = D[0] * D[0] + D[1] * D[1];
+				Discr = A1 * A1 - A0 * A2;
+				if (Discr > 0.0f)
+				{
+					Root = FMath::Sqrt(Discr);
+					Inv = (1.0f) / A2;
+					TValue = (-A1 - Root) * Inv;
+					if (T0 <= T1)
+					{
+						if (T0 <= TValue && TValue <= T1)
+						{
+							RayParam[NumIntersections++] = TValue;
+						}
+					}
+					else
+					{
+						if (T1 <= TValue && TValue <= T0)
+						{
+							RayParam[NumIntersections++] = TValue;
+						}
+					}
+
+					if (NumIntersections < 2)
+					{
+						TValue = (-A1 + Root) * Inv;
+						if (T0 <= T1)
+						{
+							if (T0 <= TValue && TValue <= T1)
+							{
+								RayParam[NumIntersections++] = TValue;
+							}
+						}
+						else
+						{
+							if (T1 <= TValue && TValue <= T0)
+							{
+								RayParam[NumIntersections++] = TValue;
+							}
+						}
+					}
+					// else: Line intersects end disk and cylinder wall.
+				}
+				else if (Discr == 0.0f)
+				{
+					TValue = -A1 / A2;
+					if (T0 <= T1)
+					{
+						if (T0 <= TValue && TValue <= T1)
+						{
+							RayParam[NumIntersections++] = TValue;
+						}
+					}
+					else
+					{
+						if (T1 <= TValue && TValue <= T0)
+						{
+							RayParam[NumIntersections++] = TValue;
+						}
+					}
+				}
+				// else: Line does not intersect cylinder wall.
+			}
+			// else: Line intersects both top and bottom cylinder end disks.
+
+			if (NumIntersections == 2)
+			{
+				if (RayParam[0] > RayParam[1])
+				{
+					RealType TmpT = RayParam[0];
+					RayParam[0] = RayParam[1];
+					RayParam[1] = TmpT;
+				}
+			}
+			else if (NumIntersections == 1)
+			{
+				RayParam[1] = RayParam[0];
+			}
+		}
+	}
+
+	// Get rid of hits before ray origin
+	if (NumIntersections > 0 && RayParam[0] >= RealType(0))
+	{
+		bIntersectsOut = true;
+		OutRayParam = RayParam[0];
+	}
+	else if (NumIntersections == 2 && RayParam[1] >= RealType(0))
+	{
+		bIntersectsOut = true;
+		OutRayParam = RayParam[1];
+	}
+	else
+	{
+		bIntersectsOut = false;
+	}
+}
+
+template<typename RealType>
+void GizmoMath::RayConeIntersection(
+	const FVector& ConeCenter, const FVector& ConeDirection, RealType ConeCosAngle, RealType ConeHeight,
+	const FVector& RayOrigin, const FVector& RayDirection,
+	bool& bIntersectsOut, RealType& OutRayParam)
+{
+	// adapted from GeometricTools GTEngine
+	// https://www.geometrictools.com/GTE/Mathematics/IntrRay3Cone3.h and
+	// https://www.geometrictools.com/GTE/Mathematics/IntrLine3Cone3.h 
+	// (Engine/Plugins/Experimental/GeometryProcessing/Source/GeometryAlgorithms/Private/ThirdParty/GTEngine/Mathematics/GteIntrRay3Cone3.h and
+	//  Engine/Plugins/Experimental/GeometryProcessing/Source/GeometryAlgorithms/Private/ThirdParty/GTEngine/Mathematics/GteIntrLine3Cone3.h)	
+
+	// The cone has vertex V, unit-length axis direction D, angle theta in
+	// (0,pi/2), and height h in (0,+infinity).  The line is P + t*U, where U
+	// is a unit-length direction vector.  Define g = cos(theta).  The cone
+	// is represented by
+	//   (X-V)^T * (D*D^T - g^2*I) * (X-V) = 0,  0 <= Dot(D,X-V) <= h
+	// The first equation defines a double-sided cone.  The first inequality
+	// in the second equation limits this to a single-sided cone containing
+	// the ray V + s*D with s >= 0.  We will call this the 'positive cone'.
+	// The single-sided cone containing ray V + s * t with s <= 0 is called
+	// the 'negative cone'.  The double-sided cone is the union of the
+	// positive cone and negative cone.  The second inequality in the second
+	// equation limits the single-sided cone to the region bounded by the
+	// height.  Setting X(t) = P + t*U, the equations are
+	//   C2*t^2 + 2*C1*t + C0 = 0,  0 <= Dot(D,U)*t + Dot(D,P-V) <= h
+	// where
+	//   C2 = Dot(D,U)^2 - g^2
+	//   C1 = Dot(D,U)*Dot(D,P-V) - g^2*Dot(U,P-V)
+	//   C0 = Dot(D,P-V)^2 - g^2*Dot(P-V,P-V)
+	// The following code computes the t-interval that satisfies the quadratic
+	// equation subject to the linear inequality constraints.
+
+	FVector PmV = RayOrigin - ConeCenter;
+	RealType DdU = FVector::DotProduct(ConeDirection, RayDirection);
+	RealType DdPmV = FVector::DotProduct(ConeDirection, PmV);
+	RealType UdPmV = FVector::DotProduct(RayDirection, PmV);
+	RealType PmVdPmV = FVector::DotProduct(PmV, PmV);
+	RealType CosAngleSqr = ConeCosAngle * ConeCosAngle;
+	RealType C2 = DdU * DdU - CosAngleSqr;
+	RealType C1 = DdU * DdPmV - CosAngleSqr * UdPmV;
+	RealType C0 = DdPmV * DdPmV - CosAngleSqr * PmVdPmV;
+	RealType T;
+	RealType RayParam[2];
+
+	if (C2 != 0.0)
+	{
+		RealType Discr = C1 * C1 - C0 * C2;
+		if (Discr < 0.0)
+		{
+			// The quadratic has no real-valued Roots.  The line does not
+			// intersect the double-sided cone.
+			bIntersectsOut = false;
+			return;
+		}
+		else if (Discr > 0.0)
+		{
+			// The quadratic has two distinct real-valued Roots.  However, one
+			// or both of them might intersect the negative cone.  We are
+			// interested only in those intersections with the positive cone.
+			RealType Root = FMath::Sqrt(Discr);
+			RealType InvC2 = 1.0 / C2;
+			int NumIntersections = 0;
+
+			T = (-C1 - Root) * InvC2;
+			if (DdU * T + DdPmV >= 0.0)
+			{
+				RayParam[NumIntersections++] = T;
+			}
+
+			T = (-C1 + Root) * InvC2;
+			if (DdU * T + DdPmV >= 0.0)
+			{
+				RayParam[NumIntersections++] = T;
+			}
+
+			if (NumIntersections == 2)
+			{
+				// The line intersects the positive cone in two distinct
+				// points.
+				if (RayParam[0] > RayParam[1])
+				{
+					RealType TmpT = RayParam[0];
+					RayParam[0] = RayParam[1];
+					RayParam[1] = TmpT;
+				}
+			}
+			else if (NumIntersections == 1)
+			{
+				// The line intersects the positive cone in a single point and
+				// the negative cone in a single point.  We report only the
+				// intersection with the positive cone.
+				if (DdU > 0.0)
+				{
+					RayParam[1] = FLT_MAX;
+				}
+				else
+				{
+					bIntersectsOut = false;
+					return;
+				}
+			}
+			else
+			{
+				bIntersectsOut = false;
+				return;
+			}
+		}
+		else  // Discr == 0
+		{
+			// One repeated real Root; the line is tangent to the double-sided
+			// cone at a single point.  Report only the point if it is on the
+			// positive cone.
+			T = -C1 / C2;
+			if (DdU * T + DdPmV >= 0.0)
+			{
+				RayParam[0] = T;
+				RayParam[1] = T;
+			}
+			else
+			{
+				bIntersectsOut = false;
+				return;
+			}
+		}
+	}
+	else
+	{
+		bIntersectsOut = false;
+		return;
+	}
+
+	if (DdU != 0.0)
+	{
+		// Clamp the intersection to the height of the cone.
+		RealType InvDdU = (1.0f) / DdU;
+		RealType hInterval[2];
+		if (DdU > 0.0)
+		{
+			hInterval[0] = -DdPmV * InvDdU;
+			hInterval[1] = (ConeHeight - DdPmV) * InvDdU;
+		}
+		else // (DdU < 0.0f)
+		{
+			hInterval[0] = (ConeHeight - DdPmV) * InvDdU;
+			hInterval[1] = -DdPmV * InvDdU;
+		}
+
+		RealType Result0, Result1;
+		int NumIntersections;
+		GizmoMath::IntervalIntervalIntersection<RealType>(RayParam, hInterval, NumIntersections, Result0, Result1);
+
+		if (NumIntersections > 0)
+		{
+			RayParam[0] = Result0;
+			RayParam[1] = NumIntersections == 2 ? Result1 : Result0;
+		}
+		else
+		{
+			bIntersectsOut = false;
+			return;
+		}
+	}
+	else if (bIntersectsOut)
+	{
+		if (DdPmV > ConeHeight)
+		{
+			bIntersectsOut = false;
+			return;
+		}
+	}
+
+	// Get rid of hits before ray origin
+	if (RayParam[0] >= RealType(0))
+	{
+		bIntersectsOut = true;
+		OutRayParam = RayParam[0];
+	}
+	else if (RayParam[1] >= RealType(0))
+	{
+		bIntersectsOut = true;
+		OutRayParam = RayParam[1];
+	}
+	else
+	{
+		bIntersectsOut = false;
+	}
+	
+}
+
+template <typename RealType>
+void GizmoMath::IntervalIntervalIntersection(
+	const RealType Interval0[2], const RealType Interval1[2], 
+	int& OutNumIntersections, RealType& OutResult0, RealType& OutResult1)
+{
+	// adapted from GeometricTools GTEngine
+	// https://www.geometrictools.com/GTE/Mathematics/IntrIntervals.h 
+	// (Engine/Plugins/Experimental/GeometryProcessing/Source/GeometryAlgorithms/Private/ThirdParty/GTEngine/Mathematics/GteIntrIntervals.h)
+	// Determines the intersection between two floating point intervals in which
+	// the input intervals' values are sorted in increasing order.
+	// Used by RayConeIntersection.
+
+	if (Interval0[1] < Interval1[0] || Interval0[0] > Interval1[1])
+	{
+		OutNumIntersections = 0;
+		OutResult0 = FLT_MAX;
+		OutResult1 = -FLT_MAX;
+	}
+	else if (Interval0[1] > Interval1[0])
+	{
+		if (Interval0[0] < Interval1[1])
+		{
+			OutNumIntersections = 2;
+			OutResult0 =
+				(Interval0[0] < Interval1[0] ? Interval1[0] : Interval0[0]);
+			OutResult1 =
+				(Interval0[1] > Interval1[1] ? Interval1[1] : Interval0[1]);
+			if (OutResult0 == OutResult1)
+			{
+				OutNumIntersections = 1;
+			}
+		}
+		else  // Interval0[0] == Interval1[1]
+		{
+			OutNumIntersections = 1;
+			OutResult0 = Interval0[0];
+			OutResult1 = OutResult0;
+		}
+	}
+	else  // Interval0[1] == Interval1[0]
+	{
+		OutNumIntersections = 1;
+		OutResult0 = Interval0[1];
+		OutResult1 = OutResult0;
+	}
+}
+
+
+
 
 void GizmoMath::ClosetPointOnCircle(
 	const FVector& QueryPoint,
@@ -234,4 +709,51 @@ float GizmoMath::SnapToIncrement(float Value, float Increment)
 		++IntIncrement;
 	}
 	return Sign * (float)IntIncrement * Increment;
-}
\ No newline at end of file
+}
+
+// @todo: Remove this and replace calls to it with MakePerpVectors() once
+// Engine\Plugins\Experimental\GeometryProcessing\Source\GeometricObjects\Public\VectorUtil.h has been moved out of Plugins dir.
+FVector GizmoMath::GetOrthogonalVector(const FVector& V)
+{
+	FVector AbsVector(FMath::Abs(V.X), FMath::Abs(V.Y), FMath::Abs(V.Z));
+	if ((AbsVector.X <= AbsVector.Y) && (AbsVector.X <= AbsVector.Z))
+	{
+		// X is the smallest component
+		return FVector(0, V.Z, -V.Y);
+	}
+	if ((AbsVector.Z <= AbsVector.X) && (AbsVector.Z <= AbsVector.Y))
+	{
+		// Z is the smallest component
+		return FVector(V.Y, -V.X, 0);
+	}
+	// Y is the smallest component
+	return FVector(-V.Z, 0, V.X);
+}
+
+namespace GizmoMath
+{
+	template
+	void INTERACTIVETOOLSFRAMEWORK_API GizmoMath::RayCylinderIntersection<float>(
+		const FVector& CylinderCenter, const FVector& CylinderAxis, float CylinderRadius, float CylinderHeight,
+		const FVector& RayOrigin, const FVector& RayDirection,
+		bool& bIntersectsOut, float& OutRayParam);
+
+	template
+	void INTERACTIVETOOLSFRAMEWORK_API GizmoMath::RayCylinderIntersection<double>(
+		const FVector& CylinderCenter, const FVector& CylinderAxis, double CylinderRadius, double CylinderHeight,
+		const FVector& RayOrigin, const FVector& RayDirection,
+		bool& bIntersectsOut, double& OutRayParam);
+
+	template
+	void INTERACTIVETOOLSFRAMEWORK_API GizmoMath::RayConeIntersection<float>(
+		const FVector& ConeCenter, const FVector& ConeDirection, float ConeAngle, float ConeHeight,
+		const FVector& RayOrigin, const FVector& RayDirection,
+		bool& bIntersectsOut, float& OutHitDepth);
+
+	template
+	void INTERACTIVETOOLSFRAMEWORK_API GizmoMath::RayConeIntersection<double>(
+		const FVector& ConeCenter, const FVector& ConeDirection, double ConeAngle, double ConeHeight,
+		const FVector& RayOrigin, const FVector& RayDirection,
+		bool& bIntersectsOut, double& OutHitDepth);
+
+} // end namespace GizmoMath
\ No newline at end of file
diff --git a/Engine/Source/Runtime/Experimental/InteractiveToolsFramework/Public/BaseGizmos/GizmoMath.h b/Engine/Source/Runtime/Experimental/InteractiveToolsFramework/Public/BaseGizmos/GizmoMath.h
--- a/Engine/Source/Runtime/Experimental/InteractiveToolsFramework/Public/BaseGizmos/GizmoMath.h
+++ b/Engine/Source/Runtime/Experimental/InteractiveToolsFramework/Public/BaseGizmos/GizmoMath.h
@@ -61,6 +61,46 @@ namespace GizmoMath
 		const FVector& RayOrigin, const FVector& RayDirection,
 		bool& bIntersectsOut, FVector& SphereIntersectionPointOut);
 
+	/**
+	 * @todo change FVector to a templated vector type specialized by RealType (double or float)
+	 * Find the intersection of the ray defined by RayOrigin and RayDirection
+	 * with the cylinder defined by CylinderCenter, CylinderAxis, CylinderRadius, and CylinderHeight.
+	 * Returns intersection success/failure in bOutIntersects and the hit depth parameter 
+	 * along the ray for the first intersection in OutHitDepth.
+	 */
+	template<typename RealType>
+	void INTERACTIVETOOLSFRAMEWORK_API RayCylinderIntersection(
+		const FVector& CylinderCenter, const FVector& CylinderAxis, RealType CylinderRadius, RealType CylinderHeight,
+		const FVector& RayOrigin, const FVector& RayDirection,
+		bool& bIntersectsOut, RealType& OutHitDepth);
+
+
+	/**
+	 * @todo change FVector to a templated vector type specialized by RealType (double or float)
+	 * Find the intersection of the ray defined by RayOrigin and RayDirection
+	 * with the cone defined by ConeCenter, ConeDirection, ConeCosAngle, and ConeHeight.
+	 * Where ConeCosAngle is the cosine of the cone opening angle, defined as the angle
+	 * at the apex between the cone's axis and slant edge.
+	 * Returns intersection success/failure in bOutIntersects and the hit depth parameter
+	 * along the ray for the first intersection in OutHitDepth.
+	 */
+	template <typename RealType>
+	void INTERACTIVETOOLSFRAMEWORK_API RayConeIntersection(
+		const FVector& ConeCenter, const FVector& ConeDirection, RealType ConeCosAngle, RealType ConeHeight,
+		const FVector& RayOrigin, const FVector& RayDirection,
+		bool& bIntersectsOut, RealType& OutHitDepth);
+
+	/**
+	 * Find the intersection of two floating point intervals in which
+	 * the interval inputs contain pairs of floats sorted in increasing order.
+	 * Returns the number of intersections intersection and intersection 
+	 * results in increasing order. 
+	 */
+	template <typename RealType>
+	void INTERACTIVETOOLSFRAMEWORK_API IntervalIntervalIntersection(
+		const RealType Interval0[2], const RealType Interval1[2],
+		int& OutNumIntersections, RealType& OutResult0, RealType& OutResult1);
+
 	/**
 	 * Find the closest point to QueryPoint that is on the circle
 	 * defined by CircleOrigin, CircleNormal, and CircleRadius.
@@ -115,4 +155,12 @@ namespace GizmoMath
 	 * @return snapped/rounded value
 	 */
 	INTERACTIVETOOLSFRAMEWORK_API float SnapToIncrement(float Value, float Increment);
-}
\ No newline at end of file
+
+	/**
+	 * @todo change FVector to a templated vector type specialized by RealType (double or float)
+	 * Compute vector orthogonal to V.
+	 * @return Orthogonal vector
+	 */
+	INTERACTIVETOOLSFRAMEWORK_API FVector GetOrthogonalVector(const FVector& V);
+}
+
