Added legacy view dependent scaling used by UGizmo*Object classes. Eventually this should be probably be replaced by CalculateLocalPixelToWorldScale, for now needed for legacy feature parity in Level Editor gizmo.

diff --git a/Engine/Source/Runtime/Experimental/InteractiveToolsFramework/Private/BaseGizmos/GizmoRenderingUtil.cpp b/Engine/Source/Runtime/Experimental/InteractiveToolsFramework/Private/BaseGizmos/GizmoRenderingUtil.cpp
--- a/Engine/Source/Runtime/Experimental/InteractiveToolsFramework/Private/BaseGizmos/GizmoRenderingUtil.cpp
+++ b/Engine/Source/Runtime/Experimental/InteractiveToolsFramework/Private/BaseGizmos/GizmoRenderingUtil.cpp
@@ -2,6 +2,8 @@
 
 #include "BaseGizmos/GizmoRenderingUtil.h"
 #include "RHI.h"
+#include "Math/UnrealMathUtility.h"
+#include "Math/NumericLimits.h"
 
 
 // yuck global value set by Editor
@@ -141,4 +143,49 @@ float GizmoRenderingUtil::CalculateLocalPixelToWorldScale(
 	double PixelDeltaSqr = VectorDifferenceSqr(PixelA, PixelB);
 	double WorldDeltaSqr = VectorDifferenceSqr(Location, OffsetPointWorld);
 	return (float)(sqrt(WorldDeltaSqr / PixelDeltaSqr));
-}
\ No newline at end of file
+}
+
+float GizmoRenderingUtil::CalculateViewDependentScaleAndFlatten(
+	const FSceneView* View,
+	const FVector& Location,
+	const float InScale,
+	FVector& OutFlattenScale)
+{
+	const FMatrix& ViewMatrix = View->ViewMatrices.GetViewMatrix();
+	bool bIsPerspective = View->ViewMatrices.GetProjectionMatrix().M[3][3] < 1.0f;
+	bool bIsOrthoXY = !bIsPerspective && FMath::Abs(ViewMatrix.M[2][2]) > 0.0f;
+	bool bIsOrthoXZ = !bIsPerspective && FMath::Abs(ViewMatrix.M[1][2]) > 0.0f;
+	bool bIsOrthoYZ = !bIsPerspective && FMath::Abs(ViewMatrix.M[0][2]) > 0.0f;
+	float UniformScale = InScale * View->WorldToScreen(Location).W * (4.0f / View->UnscaledViewRect.Width() / View->ViewMatrices.GetProjectionMatrix().M[0][0]);
+
+	// Clamp to tolerance to prevent division by zero.
+	// @todo change to use MathUtil<RealType>::ZeroTolerance and TMathUtil<RealType>::SignNonZero(Value)
+	float MinimumScale = TNumericLimits<float>::Lowest();
+	if (FMath::Abs(UniformScale) < MinimumScale)
+	{
+		UniformScale = MinimumScale * (UniformScale < 0.0f ? -1.0f : 1.0f);
+	}
+
+	if (bIsPerspective)
+	{
+		OutFlattenScale = FVector(1.0f, 1.0f, 1.0f);
+	}
+	else
+	{
+		// Flatten scale prevents scaling in the direction of the camera and thus intersecting the near plane.
+		// Based on legacy FWidget render code but is flatten actually necessary?? that axis wasn't scaled anyways!
+		if (bIsOrthoXY)
+		{
+			OutFlattenScale = FVector(1.0f, 1.0f, 1.0f / UniformScale);
+		}
+		else if (bIsOrthoXZ)
+		{
+			OutFlattenScale = FVector(1.0f, 1.0f / UniformScale, 1.0f);
+		}
+		else if (bIsOrthoYZ)
+		{
+			OutFlattenScale = FVector(1.0f / UniformScale, 1.0f, 1.0f);
+		}
+	}
+	return UniformScale;
+}
diff --git a/Engine/Source/Runtime/Experimental/InteractiveToolsFramework/Public/BaseGizmos/GizmoRenderingUtil.h b/Engine/Source/Runtime/Experimental/InteractiveToolsFramework/Public/BaseGizmos/GizmoRenderingUtil.h
--- a/Engine/Source/Runtime/Experimental/InteractiveToolsFramework/Public/BaseGizmos/GizmoRenderingUtil.h
+++ b/Engine/Source/Runtime/Experimental/InteractiveToolsFramework/Public/BaseGizmos/GizmoRenderingUtil.h
@@ -4,7 +4,7 @@
 
 #include "CoreMinimal.h"
 #include "Components/PrimitiveComponent.h"
-
+#include "SceneView.h"
 
 /**
  * Utility functions for standard GizmoComponent rendering
@@ -49,4 +49,14 @@ namespace GizmoRenderingUtil
 		const FSceneView* View,
 		const FVector& Location);
 
+	/**
+	 * @return Legacy view dependent conversion factor.
+	 * @return OutWorldFlattenScale vector to be applied in world space, can be used to flatten excluded 
+	 *         dimension in orthographic views as it reverses the scale in that dimension.
+	 */
+	INTERACTIVETOOLSFRAMEWORK_API float CalculateViewDependentScaleAndFlatten(
+		const FSceneView* View,
+		const FVector& Location,
+		const float Scale,
+		FVector& OutWorldFlattenScale);
 }
