Added UGizmo*Object classes which can be used to define, render and hit test various geometric primitives and/or composites of primitives such as Arrow. These are similar to UGizmo*Components.

diff --git a/Engine/Source/Runtime/Experimental/InteractiveToolsFramework/Private/BaseGizmos/GizmoArrowObject.cpp b/Engine/Source/Runtime/Experimental/InteractiveToolsFramework/Private/BaseGizmos/GizmoArrowObject.cpp
new file mode 100644
--- /dev/null
+++ b/Engine/Source/Runtime/Experimental/InteractiveToolsFramework/Private/BaseGizmos/GizmoArrowObject.cpp
@@ -0,0 +1,204 @@
+// Copyright Epic Games, Inc. All Rights Reserved.
+ 
+#include "BaseGizmos/GizmoArrowObject.h"
+#include "BaseGizmos/GizmoRenderingUtil.h"
+
+UGizmoArrowObject::UGizmoArrowObject()
+{
+    CylinderObject = NewObject<UGizmoCylinderObject>();
+    ConeObject = NewObject<UGizmoConeObject>();
+    BoxObject = NewObject<UGizmoBoxObject>();
+}
+
+void UGizmoArrowObject::Render(IToolsContextRenderAPI* RenderAPI)
+{
+	if (CylinderObject)
+	{
+		CylinderObject->Render(RenderAPI);
+	}
+
+	if (bHasConeHead)
+	{
+		if (ConeObject)
+		{
+			ConeObject->Render(RenderAPI);
+		}
+	}
+	else
+	{
+		if (BoxObject)
+		{
+			BoxObject->Render(RenderAPI);
+		}
+	}
+}
+
+FInputRayHit UGizmoArrowObject::LineTraceObject(const FVector RayOrigin, const FVector RayDirection)
+{
+	FInputRayHit RayHit;
+	if (CylinderObject)
+	{
+		RayHit = CylinderObject->LineTraceObject(RayOrigin, RayDirection);
+		if (!RayHit.bHit)
+		{
+			if (bHasConeHead)
+			{
+				if (ConeObject)
+				{
+					RayHit = ConeObject->LineTraceObject(RayOrigin, RayDirection);
+				}
+			}
+			else
+			{
+				if (BoxObject)
+				{
+					RayHit = BoxObject->LineTraceObject(RayOrigin, RayDirection);
+				}
+			}
+		}
+	}
+	return RayHit;
+}
+
+void UGizmoArrowObject::SetHoverState(bool bHoveringIn)
+{
+	bHovering = bHoveringIn;
+
+	if (CylinderObject)
+	{
+		CylinderObject->SetHoverState(bHoveringIn);
+	}
+	if (ConeObject)
+	{
+		ConeObject->SetHoverState(bHoveringIn);
+	}
+	if (BoxObject)
+	{
+		BoxObject->SetHoverState(bHoveringIn);
+	}
+}
+
+void UGizmoArrowObject::SetInteractingState(bool bInteractingIn)
+{
+	bInteracting = bInteractingIn;
+
+	if (CylinderObject)
+	{
+		CylinderObject->SetInteractingState(bInteractingIn);
+	}
+	if (ConeObject)
+	{
+		ConeObject->SetInteractingState(bInteractingIn);
+	}
+	if (BoxObject)
+	{
+		BoxObject->SetInteractingState(bInteractingIn);
+	}
+}
+
+void UGizmoArrowObject::SetWorldLocalState(bool bWorldIn)
+{
+	bWorld = bWorldIn;
+
+	if (CylinderObject)
+	{
+		CylinderObject->SetWorldLocalState(bWorldIn);
+	}
+	if (ConeObject)
+	{
+		ConeObject->SetWorldLocalState(bWorldIn);
+	}
+	if (BoxObject)
+	{
+		BoxObject->SetWorldLocalState(bWorldIn);
+	}
+}
+
+void UGizmoArrowObject::SetVisibility(bool bVisibleIn)
+{
+	bVisible = bVisibleIn;
+
+	if (CylinderObject)
+	{
+		CylinderObject->SetVisibility(bVisibleIn);
+	}
+	if (ConeObject)
+	{
+		ConeObject->SetVisibility(bVisibleIn);
+	}
+	if (BoxObject)
+	{
+		BoxObject->SetVisibility(bVisibleIn);
+	}
+}
+
+void UGizmoArrowObject::SetLocalToWorldTransform(FTransform LocalToWorldTransformIn)
+{
+	LocalToWorldTransform = LocalToWorldTransformIn;
+
+	if (CylinderObject)
+	{
+		CylinderObject->SetLocalToWorldTransform(LocalToWorldTransformIn);
+	}
+	if (ConeObject)
+	{
+		ConeObject->SetLocalToWorldTransform(LocalToWorldTransformIn);
+	}
+	if (BoxObject)
+	{
+		BoxObject->SetLocalToWorldTransform(LocalToWorldTransformIn);
+	}
+}
+
+void UGizmoArrowObject::SetGizmoScale(float InGizmoScale)
+{
+	GizmoScale = InGizmoScale;
+
+	if (CylinderObject)
+	{
+		CylinderObject->SetGizmoScale(InGizmoScale);
+	}
+	if (ConeObject)
+	{
+		ConeObject->SetGizmoScale(InGizmoScale);
+	}
+	if (BoxObject)
+	{
+		BoxObject->SetGizmoScale(InGizmoScale);
+	}
+}
+
+void UGizmoArrowObject::SetMaterial(UMaterialInterface* MaterialIn)
+{
+	Material = MaterialIn;
+	if (CylinderObject)
+	{
+		CylinderObject->SetMaterial(MaterialIn);
+	}
+	if (ConeObject)
+	{
+		ConeObject->SetMaterial(MaterialIn);
+	}
+	if (BoxObject)
+	{
+		BoxObject->SetMaterial(MaterialIn);
+	}
+}
+
+void UGizmoArrowObject::SetCurrentMaterial(UMaterialInterface* CurrentMaterialIn)
+{
+	CurrentMaterial = CurrentMaterialIn;
+	if (CylinderObject)
+	{
+		CylinderObject->SetCurrentMaterial(CurrentMaterialIn);
+	}
+	if (ConeObject)
+	{
+		ConeObject->SetCurrentMaterial(CurrentMaterialIn);
+	}
+	if (BoxObject)
+	{
+		BoxObject->SetCurrentMaterial(CurrentMaterialIn);
+	}
+}
+
diff --git a/Engine/Source/Runtime/Experimental/InteractiveToolsFramework/Private/BaseGizmos/GizmoBoxObject.cpp b/Engine/Source/Runtime/Experimental/InteractiveToolsFramework/Private/BaseGizmos/GizmoBoxObject.cpp
new file mode 100644
--- /dev/null
+++ b/Engine/Source/Runtime/Experimental/InteractiveToolsFramework/Private/BaseGizmos/GizmoBoxObject.cpp
@@ -0,0 +1,92 @@
+// Copyright Epic Games, Inc. All Rights Reserved.
+ 
+#include "BaseGizmos/GizmoBoxObject.h"
+#include "BaseGizmos/GizmoMath.h"
+#include "BaseGizmos/GizmoRenderingUtil.h"
+#include "Materials/MaterialInterface.h"
+#include "InputState.h"
+
+void UGizmoBoxObject::Render(IToolsContextRenderAPI* RenderAPI)
+{
+	if (bVisible)
+	{
+		check(RenderAPI);
+
+		FPrimitiveDrawInterface* PDI = RenderAPI->GetPrimitiveDrawInterface();
+		const FSceneView* View = RenderAPI->GetSceneView();
+
+
+		const FMatrix LocalToWorldMatrix = LocalToWorldTransform.ToMatrixNoScale();
+		FVector UseOrigin = LocalToWorldMatrix.TransformPosition(FVector::ZeroVector);
+
+		bool bIsOrtho = !View->IsPerspectiveProjection();
+
+		// direction to origin of gizmo
+		FVector ViewDirection =
+			(bIsOrtho) ? (View->GetViewDirection()) : (UseOrigin - View->ViewLocation);
+		ViewDirection.Normalize();
+
+		const FVector UseDirection = (bWorld) ? UpDirection : FVector{ LocalToWorldMatrix.TransformVector(UpDirection) };
+		bool bFlipToCompare = (FVector::DotProduct(ViewDirection, UseDirection) > 0);
+		const FVector CompareDirection = (bFlipToCompare) ? -UseDirection : UseDirection;
+
+		// @todo - make this an object property and static constant on the gizmo
+		static const float ViewMaxCosAngle = 0.995f;  // ~5 degrees, cos(0.087 radians)
+		bVisibleViewDependent = FMath::Abs(FVector::DotProduct(CompareDirection, ViewDirection)) < ViewMaxCosAngle;
+
+		if (bVisibleViewDependent)
+		{
+			FMatrix UseRotMatrix = (bWorld) ? FMatrix::Identity : FRotationMatrix(LocalToWorldTransform.GetRotation().Rotator());
+
+			// @todo: replace with a call to CalculateLocalPixelToWorldScale instead if possible
+			FVector FlattenScale = FVector::OneVector;
+			DynamicPixelToWorldScale = GizmoRenderingUtil::CalculateViewDependentScaleAndFlatten(View, UseOrigin, GizmoScale, FlattenScale);
+
+			FVector UniformDimensions = DynamicPixelToWorldScale * Dimensions;
+			FVector FlattenDimensions(FlattenScale[0] == 1.0f ? 1.0f : 1.0f / UniformDimensions[0],
+									  FlattenScale[1] == 1.0f ? 1.0f : 1.0f / UniformDimensions[1],
+									  FlattenScale[2] == 1.0f ? 1.0f : 1.0f / UniformDimensions[2]);
+
+			UMaterialInterface* UseMaterial = Material;
+			if (bHovering || bInteracting)
+			{
+				UseMaterial = CurrentMaterial;
+			}
+			check(UseMaterial);
+
+			FVector BoxCenter(0, 0, Offset * DynamicPixelToWorldScale);
+
+			FVector ForwardDirection = FVector::CrossProduct(SideDirection, UpDirection);
+			FMatrix AxisRotation = FRotationMatrix::MakeFromYZ(ForwardDirection, UpDirection);
+
+			FMatrix BoxToWorld = FScaleMatrix(UniformDimensions) * FTranslationMatrix(BoxCenter) * AxisRotation * UseRotMatrix * FTranslationMatrix(UseOrigin) * FScaleMatrix(FlattenScale);
+
+			DrawBox(PDI, BoxToWorld, FVector(1, 1, 1), UseMaterial->GetRenderProxy(), SDPG_Foreground);
+		}
+	}
+}
+
+FInputRayHit UGizmoBoxObject::LineTraceObject(const FVector RayStart, const FVector RayDirection)
+{
+	if (bVisible && bVisibleViewDependent)
+	{
+		/*
+		// @todo
+		bool bIntersects = false;
+		float RayParam = 0.0f;
+
+
+		const FVector BoxOrigin = Origin + Offset;
+		GizmoMath::RayBoxIntersection(
+			BoxOrigin, CubeObject->Axis, CubeObject->Radius, CubeObject->Height,
+			ClickPos.WorldRay.Origin, ClickPos.WorldRay.Direction,
+			bIntersects, RayParam);
+
+		if (bIntersects)
+		{
+			return FInputRayHit(RayParam);
+		}
+		*/
+	}
+	return FInputRayHit();
+}
diff --git a/Engine/Source/Runtime/Experimental/InteractiveToolsFramework/Private/BaseGizmos/GizmoConeObject.cpp b/Engine/Source/Runtime/Experimental/InteractiveToolsFramework/Private/BaseGizmos/GizmoConeObject.cpp
new file mode 100644
--- /dev/null
+++ b/Engine/Source/Runtime/Experimental/InteractiveToolsFramework/Private/BaseGizmos/GizmoConeObject.cpp
@@ -0,0 +1,87 @@
+// Copyright Epic Games, Inc. All Rights Reserved.
+ 
+#include "BaseGizmos/GizmoConeObject.h"
+#include "BaseGizmos/GizmoRenderingUtil.h"
+#include "BaseGizmos/GizmoMath.h"
+#include "InputState.h"
+#include "Materials/MaterialInterface.h"
+
+void UGizmoConeObject::Render(IToolsContextRenderAPI* RenderAPI)
+{
+	if (bVisible)
+	{
+		check(RenderAPI);
+
+		FPrimitiveDrawInterface* PDI = RenderAPI->GetPrimitiveDrawInterface();
+		const FSceneView* View = RenderAPI->GetSceneView();
+
+		const FMatrix LocalToWorldMatrix = LocalToWorldTransform.ToMatrixNoScale();
+
+		FVector UseOrigin = LocalToWorldMatrix.TransformPosition(FVector::ZeroVector);
+
+		bool bIsOrtho = !View->IsPerspectiveProjection();
+
+		// direction to origin of gizmo
+		FVector ViewDirection =
+			(bIsOrtho) ? (View->GetViewDirection()) : (UseOrigin - View->ViewLocation);
+		ViewDirection.Normalize();
+
+		FVector UseDirection = (bWorld) ? Direction : FVector{ LocalToWorldMatrix.TransformVector(Direction) };
+		bool bFlipToCompare = (FVector::DotProduct(ViewDirection, UseDirection) > 0);
+		FVector CompareDirection = (bFlipToCompare) ? -UseDirection : UseDirection;
+
+		// @todo - make this an object property and static constant on the gizmo
+		static const float ViewMaxCosAngle = 0.995f;  // ~5 degrees, cos(0.087 radians)
+		bVisibleViewDependent = FMath::Abs(FVector::DotProduct(CompareDirection, ViewDirection)) < ViewMaxCosAngle;  
+
+		if (bVisibleViewDependent)
+		{
+			// @todo: replace with a call to CalculateLocalPixelToWorldScale instead if possible
+			FVector FlattenScale = FVector::OneVector;
+			DynamicPixelToWorldScale = GizmoRenderingUtil::CalculateViewDependentScaleAndFlatten(View, UseOrigin, GizmoScale, FlattenScale);
+
+			UMaterialInterface* UseMaterial = Material;
+			if (bHovering || bInteracting)
+			{
+				UseMaterial = CurrentMaterial;
+			}
+			check(UseMaterial);
+
+			const FVector Base(Offset * DynamicPixelToWorldScale, 0, 0);
+			FMatrix AxisRotation = FRotationMatrix::MakeFromX(UseDirection);
+			FMatrix ConeToWorld = FScaleMatrix(Height * DynamicPixelToWorldScale) * FTranslationMatrix(Base) * AxisRotation * FTranslationMatrix(UseOrigin) * FScaleMatrix(FlattenScale);
+			DrawCone(PDI, ConeToWorld, Angle, Angle, NumSides, false, FColor::White, UseMaterial->GetRenderProxy(), SDPG_Foreground);
+		}
+	}
+}
+
+FInputRayHit UGizmoConeObject::LineTraceObject(const FVector RayOrigin, const FVector RayDirection)
+{
+	if (bVisible && bVisibleViewDependent)
+	{
+		bool bIntersects = false;
+		double RayParam = 0.0;
+
+		const FMatrix LocalToWorldMatrix = LocalToWorldTransform.ToMatrixNoScale();
+
+		const FVector Origin = LocalToWorldMatrix.TransformPosition(FVector::ZeroVector);
+		const FVector ConeDirection = (bWorld) ? Direction : FVector{ LocalToWorldMatrix.TransformVector(Direction) };
+		const double ConeHeight = Height * DynamicPixelToWorldScale;
+		const double ConeOffset = Offset * DynamicPixelToWorldScale;
+		const FVector ConeOrigin = Origin + ConeOffset * ConeDirection;
+
+		GizmoMath::RayConeIntersection(
+			ConeOrigin, 
+			ConeDirection, 
+			FMath::Cos(Angle), 
+			ConeHeight,
+			RayOrigin, RayDirection,
+			bIntersects, RayParam);
+
+		if (bIntersects)
+		{
+			return FInputRayHit(RayParam);
+		}
+	}
+	return FInputRayHit();
+}
\ No newline at end of file
diff --git a/Engine/Source/Runtime/Experimental/InteractiveToolsFramework/Private/BaseGizmos/GizmoCylinderObject.cpp b/Engine/Source/Runtime/Experimental/InteractiveToolsFramework/Private/BaseGizmos/GizmoCylinderObject.cpp
new file mode 100644
--- /dev/null
+++ b/Engine/Source/Runtime/Experimental/InteractiveToolsFramework/Private/BaseGizmos/GizmoCylinderObject.cpp
@@ -0,0 +1,91 @@
+// Copyright Epic Games, Inc. All Rights Reserved.
+ 
+#include "BaseGizmos/GizmoCylinderObject.h"
+#include "BaseGizmos/GizmoRenderingUtil.h"
+#include "BaseGizmos/GizmoMath.h"
+#include "InputState.h"
+#include "Materials/MaterialInterface.h"
+
+void UGizmoCylinderObject::Render(IToolsContextRenderAPI* RenderAPI)
+{
+	if (!bVisible)
+	{
+		return;
+	}
+
+	check(RenderAPI);
+
+	FPrimitiveDrawInterface* PDI = RenderAPI->GetPrimitiveDrawInterface();
+	const FSceneView* View = RenderAPI->GetSceneView();
+
+	const FMatrix LocalToWorldMatrix = LocalToWorldTransform.ToMatrixNoScale();
+
+	FVector UseOrigin = LocalToWorldMatrix.TransformPosition(FVector::ZeroVector);
+
+	bool bIsOrtho = !View->IsPerspectiveProjection();
+
+	// direction to origin of gizmo
+	FVector ViewDirection =
+		(bIsOrtho) ? (View->GetViewDirection()) : (UseOrigin - View->ViewLocation);
+	ViewDirection.Normalize();
+
+	FVector UseDirection = (bWorld) ? Direction : FVector{ LocalToWorldMatrix.TransformVector(Direction) };
+	bool bFlipToCompare = (FVector::DotProduct(ViewDirection, UseDirection) > 0);
+	FVector CompareDirection = (bFlipToCompare) ? -UseDirection : UseDirection;
+
+	// @todo - make this an object property and static constant on the gizmo
+	static const float ViewMaxCosAngle = 0.995f;  // ~5 degrees, cos(0.087 radians)
+	bVisibleViewDependent = FMath::Abs(FVector::DotProduct(CompareDirection, ViewDirection)) < ViewMaxCosAngle; 
+
+	if (bVisibleViewDependent)
+	{
+		// @todo: replace with a call to CalculateLocalPixelToWorldScale instead if possible
+		FVector FlattenScale = FVector::OneVector;
+		DynamicPixelToWorldScale = GizmoRenderingUtil::CalculateViewDependentScaleAndFlatten(View, UseOrigin, GizmoScale, FlattenScale);
+
+		UMaterialInterface* UseMaterial = Material;
+		if (bHovering || bInteracting)
+		{
+			UseMaterial = CurrentMaterial;
+		}
+		check(UseMaterial);
+
+		const float HalfLength = Length * 0.5f;
+		const FVector CylinderCenter(0, 0, Offset + HalfLength);
+		FMatrix AxisRotation = FRotationMatrix::MakeFromZ(UseDirection);
+		FMatrix CylinderToWorld = FScaleMatrix(DynamicPixelToWorldScale) * AxisRotation * FTranslationMatrix(UseOrigin) * FScaleMatrix(FlattenScale);
+
+		DrawCylinder(PDI, CylinderToWorld, CylinderCenter, FVector(1, 0, 0), FVector(0, 1, 0), FVector(0, 0, 1), Radius, HalfLength, NumSides, UseMaterial->GetRenderProxy(), SDPG_Foreground);
+	}
+}
+
+FInputRayHit UGizmoCylinderObject::LineTraceObject(const FVector RayOrigin, const FVector RayDirection)
+{
+	if (bVisible && bVisibleViewDependent)
+	{
+		bool bIntersects = false;
+		double RayParam;
+
+		const FMatrix LocalToWorldMatrix = LocalToWorldTransform.ToMatrixNoScale();
+		const FVector ScaleVector = LocalToWorldTransform.GetScale3D();
+
+		FVector UseOrigin = LocalToWorldMatrix.TransformPosition(FVector::ZeroVector);
+		FVector CylinderDirection = (bWorld) ? Direction : FVector{ LocalToWorldMatrix.TransformVector(Direction) };
+
+		const double CylinderOffsetLength = DynamicPixelToWorldScale * Offset;
+		const double CylinderLength = DynamicPixelToWorldScale * Length;
+		const double CylinderRadius = DynamicPixelToWorldScale * Radius;
+		const FVector CylinderCenter = UseOrigin + (CylinderOffsetLength + CylinderLength * 0.5) * CylinderDirection;
+
+		GizmoMath::RayCylinderIntersection<double>(
+			CylinderCenter, CylinderDirection, CylinderRadius, CylinderLength,
+			RayOrigin, RayDirection,
+			bIntersects, RayParam);
+
+		if (bIntersects)
+		{
+			return FInputRayHit(RayParam);
+		}
+	}
+	return FInputRayHit();
+}
\ No newline at end of file
diff --git a/Engine/Source/Runtime/Experimental/InteractiveToolsFramework/Private/BaseGizmos/GizmoGroupObject.cpp b/Engine/Source/Runtime/Experimental/InteractiveToolsFramework/Private/BaseGizmos/GizmoGroupObject.cpp
new file mode 100644
--- /dev/null
+++ b/Engine/Source/Runtime/Experimental/InteractiveToolsFramework/Private/BaseGizmos/GizmoGroupObject.cpp
@@ -0,0 +1,141 @@
+// Copyright Epic Games, Inc. All Rights Reserved.
+ 
+#include "BaseGizmos/GizmoGroupObject.h"
+#include "BaseGizmos/GizmoRenderingUtil.h"
+
+UGizmoGroupObject::UGizmoGroupObject()
+{
+}
+
+void UGizmoGroupObject::Render(IToolsContextRenderAPI* RenderAPI)
+{
+	for (UGizmoBaseObject* Object : Objects)
+	{
+		if (Object)
+		{
+			Object->Render(RenderAPI);
+		}
+	}
+}
+
+FInputRayHit UGizmoGroupObject::LineTraceObject(const FVector RayOrigin, const FVector RayDirection)
+{
+	FInputRayHit RayHit;
+	for (UGizmoBaseObject* Object : Objects)
+	{
+		if (Object)
+		{
+			RayHit = Object->LineTraceObject(RayOrigin, RayDirection);
+			if (RayHit.bHit)
+			{
+				return RayHit;
+			}
+		}
+	}
+	return RayHit;
+}
+
+void UGizmoGroupObject::SetHoverState(bool bHoverIn)
+{
+	bHovering = bHoverIn;
+
+	for (UGizmoBaseObject* Object : Objects)
+	{
+		if (Object)
+		{
+			Object->SetHoverState(bHoverIn);
+		}
+	}
+}
+
+void UGizmoGroupObject::SetInteractingState(bool bHoverIn)
+{
+	bHovering = bHoverIn;
+
+	for (UGizmoBaseObject* Object : Objects)
+	{
+		if (Object)
+		{
+			Object->SetInteractingState(bHoverIn);
+		}
+	}
+}
+
+void UGizmoGroupObject::SetWorldLocalState(bool bWorldIn)
+{
+	bWorld = bWorldIn;
+
+	for (UGizmoBaseObject* Object : Objects)
+	{
+		if (Object)
+		{
+			Object->SetWorldLocalState(bWorldIn);
+		}
+	}
+}
+
+void UGizmoGroupObject::SetVisibility(bool bVisibleIn)
+{
+	bVisible = bVisibleIn;
+
+	for (UGizmoBaseObject* Object : Objects)
+	{
+		if (Object)
+		{
+			Object->SetVisibility(bVisibleIn);
+		}
+	}
+}
+
+void UGizmoGroupObject::SetLocalToWorldTransform(FTransform LocalToWorldTransformIn)
+{
+	LocalToWorldTransform = LocalToWorldTransformIn;
+
+	for (UGizmoBaseObject* Object : Objects)
+	{
+		if (Object)
+		{
+			Object->SetLocalToWorldTransform(LocalToWorldTransformIn);
+		}
+	}
+}
+
+void UGizmoGroupObject::SetMaterial(UMaterialInterface* InMaterial)
+{
+	Material = InMaterial;
+
+	for (UGizmoBaseObject* Object : Objects)
+	{
+		if (Object)
+		{
+			Object->SetMaterial(InMaterial);
+		}
+	}
+}
+
+void UGizmoGroupObject::SetCurrentMaterial(UMaterialInterface* InCurrentMaterial)
+{
+	CurrentMaterial = InCurrentMaterial;
+
+	for (UGizmoBaseObject* Object : Objects)
+	{
+		if (Object)
+		{
+			Object->SetCurrentMaterial(InCurrentMaterial);
+		}
+	}
+}
+
+void UGizmoGroupObject::SetGizmoScale(float InGizmoScale)
+{
+	GizmoScale = InGizmoScale;
+
+	for (UGizmoBaseObject* Object : Objects)
+	{
+		if (Object)
+		{
+			Object->SetGizmoScale(InGizmoScale);
+		}
+	}
+}
+
diff --git a/Engine/Source/Runtime/Experimental/InteractiveToolsFramework/Public/BaseGizmos/GizmoArrowObject.h b/Engine/Source/Runtime/Experimental/InteractiveToolsFramework/Public/BaseGizmos/GizmoArrowObject.h
new file mode 100644
--- /dev/null
+++ b/Engine/Source/Runtime/Experimental/InteractiveToolsFramework/Public/BaseGizmos/GizmoArrowObject.h
@@ -0,0 +1,55 @@
+// Copyright Epic Games, Inc. All Rights Reserved.
+
+#pragma once
+
+#include "CoreMinimal.h"
+#include "BaseGizmos/GizmoBaseObject.h"
+#include "BaseGizmos/GizmoBoxObject.h"
+#include "BaseGizmos/GizmoConeObject.h"
+#include "BaseGizmos/GizmoCylinderObject.h"
+#include "InputState.h"
+#include "UObject/ObjectMacros.h"
+#include "GizmoArrowObject.generated.h"
+
+/**
+ * Simple arrow object intended to be used as part of 3D Gizmos.
+ * Draws a solid 3D arrow based with a cylinder body and
+ * cone or box head.
+ */
+UCLASS(Transient)
+class INTERACTIVETOOLSFRAMEWORK_API UGizmoArrowObject : public UGizmoBaseObject
+{
+	GENERATED_BODY()
+
+public:
+	UGizmoArrowObject();
+
+	//~ Begin UGizmoBaseObject Interface.
+	virtual void Render(IToolsContextRenderAPI* RenderAPI) override;
+	virtual FInputRayHit LineTraceObject(const FVector RayOrigin, const FVector RayDirection) override;
+	virtual void SetHoverState(bool bHoveringIn) override;
+	virtual void SetInteractingState(bool bInteractingIn) override;
+	virtual void SetWorldLocalState(bool bWorldIn) override;
+	virtual void SetVisibility(bool bVisibleIn) override;
+	virtual void SetLocalToWorldTransform(FTransform LocalToWorldTransformIn) override;
+	virtual void SetGizmoScale(float InGizmoScale) override;
+	virtual void SetMaterial(UMaterialInterface* InMaterial) override;
+	virtual void SetCurrentMaterial(UMaterialInterface* InCurrentMaterial) override;
+	//~ End UGizmoBaseObject Interface.
+
+public:
+
+	UPROPERTY(EditAnywhere, Category = Options)
+	int32 bHasConeHead = true;
+		
+	UPROPERTY(EditAnywhere, Category = Options)
+	FVector Direction = FVector(1.0f, 0.0f, 0.0f);
+
+	// Arrow origin is located at (Direction * Offset) 
+	UPROPERTY(EditAnywhere, Category = Options)
+	float Offset = 0.0f;
+
+	TObjectPtr<UGizmoCylinderObject> CylinderObject;
+	TObjectPtr<UGizmoConeObject> ConeObject;
+	TObjectPtr<UGizmoBoxObject> BoxObject;
+};
\ No newline at end of file
diff --git a/Engine/Source/Runtime/Experimental/InteractiveToolsFramework/Public/BaseGizmos/GizmoBaseObject.h b/Engine/Source/Runtime/Experimental/InteractiveToolsFramework/Public/BaseGizmos/GizmoBaseObject.h
new file mode 100644
--- /dev/null
+++ b/Engine/Source/Runtime/Experimental/InteractiveToolsFramework/Public/BaseGizmos/GizmoBaseObject.h
@@ -0,0 +1,170 @@
+// Copyright Epic Games, Inc. All Rights Reserved.
+
+#pragma once
+
+#include "CoreMinimal.h"
+#include "InputState.h"
+#include "Materials/MaterialInterface.h"
+#include "ToolContextInterfaces.h"
+#include "GizmoBaseObject.generated.h"
+
+/**
+ * Base class for simple objects intended to be used as part of 3D Gizmos.
+ * Contains common properties and utility functions.
+ * This class does nothing by itself, use subclasses like UGizmoCylinderObject
+ */
+UCLASS(Transient)
+class INTERACTIVETOOLSFRAMEWORK_API UGizmoBaseObject : public UObject
+{
+	GENERATED_BODY()
+
+public:
+
+	virtual void Render(IToolsContextRenderAPI* RenderAPI) {}
+	virtual FInputRayHit LineTraceObject(const FVector Start, const FVector Direction) { return FInputRayHit(); }
+
+public:
+
+	UPROPERTY(EditAnywhere, Category = Options)
+	FTransform LocalToWorldTransform = FTransform::Identity;
+
+	UPROPERTY(EditAnywhere, Category = Options)
+	TObjectPtr<UMaterialInterface> Material;
+
+	UPROPERTY(EditAnywhere, Category = Options)
+	TObjectPtr<UMaterialInterface> CurrentMaterial;
+
+
+	UPROPERTY(EditAnywhere, Category = Options)
+	float GizmoScale;
+
+	// @todo need this?
+	UPROPERTY(EditAnywhere, Category = Options)
+	float PixelHitDistanceThreshold = 7.0f;
+
+	// @todo add view dependent controls
+	// UPROPERTY(EditAnywhere, Category = Options)
+	// uint8 bEnableViewDependentVisibility
+
+	// angle in radians
+	// UPROPERTY(EditAnywhere, Category = Options)
+	// float ViewDependentVisibilityAngle
+
+	// UPROPERTY(EditAnywhere, Category = Options)
+	// uint8 bEnableViewDependentScaling
+
+
+public:
+
+	UFUNCTION()
+	virtual void SetHoverState(bool bHoveringIn)
+	{
+		bHovering = bHoveringIn;
+	}
+
+	UFUNCTION()
+	virtual bool GetHoverState()
+	{
+		return bHovering;
+	}
+
+	UFUNCTION()
+	virtual void SetInteractingState(bool bInteractingIn)
+	{
+		bInteracting = bInteractingIn;
+	}
+
+	UFUNCTION()
+		virtual bool GetInteractingState()
+	{
+		return bInteracting;
+	}
+
+	UFUNCTION()
+	virtual void SetWorldLocalState(bool bWorldIn)
+	{
+		bWorld = bWorldIn;
+	}
+
+	UFUNCTION()
+	virtual bool GetWorldLocalState()
+	{
+		return bWorld;
+	}
+
+	UFUNCTION()
+	virtual void SetVisibility(bool bVisibleIn)
+	{
+		bVisible = bVisibleIn;
+	}
+
+	UFUNCTION()
+	virtual bool GetVisibility()
+	{
+		return bVisible;
+	}
+
+	UFUNCTION()
+	virtual void SetLocalToWorldTransform(const FTransform LocalToWorldTransformIn)
+	{
+		LocalToWorldTransform = LocalToWorldTransformIn;
+	}
+
+	UFUNCTION()
+	virtual FTransform GetLocalToWorldTransform()
+	{
+		return LocalToWorldTransform;
+	}
+
+	UFUNCTION()
+	virtual void SetGizmoScale(float InGizmoScale)
+	{
+		GizmoScale = InGizmoScale;
+	}
+
+	UFUNCTION()
+	virtual float GetGizmoScale()
+	{
+		return GizmoScale;
+	}
+
+	UFUNCTION()
+	virtual void SetMaterial(UMaterialInterface* InMaterial)
+	{
+		Material = InMaterial;
+	}
+
+	UFUNCTION()
+	virtual UMaterialInterface* GetMaterial()
+	{
+		return Material;
+	}
+
+	UFUNCTION()
+	virtual void SetCurrentMaterial(UMaterialInterface* InCurrentMaterial)
+	{
+		CurrentMaterial = InCurrentMaterial;
+	}
+
+	UFUNCTION()
+	virtual UMaterialInterface* GetCurrentMaterial()
+	{
+		return CurrentMaterial;
+	}
+
+protected:
+	// scale factor between pixel distances and world distances at Gizmo origin
+	float DynamicPixelToWorldScale = 1.0f;
+
+	// hover state
+	bool bHovering = false;
+
+	// interacting state
+	bool bInteracting = false;
+
+	// world/local coordinates state
+	bool bWorld = false;
+
+	// visibility state
+	bool bVisible = true;
+};
diff --git a/Engine/Source/Runtime/Experimental/InteractiveToolsFramework/Public/BaseGizmos/GizmoBoxObject.h b/Engine/Source/Runtime/Experimental/InteractiveToolsFramework/Public/BaseGizmos/GizmoBoxObject.h
new file mode 100644
--- /dev/null
+++ b/Engine/Source/Runtime/Experimental/InteractiveToolsFramework/Public/BaseGizmos/GizmoBoxObject.h
@@ -0,0 +1,48 @@
+// Copyright Epic Games, Inc. All Rights Reserved.
+
+#pragma once
+
+#include "CoreMinimal.h"
+#include "BaseGizmos/GizmoBaseObject.h"
+#include "InputState.h"
+#include "UObject/ObjectMacros.h"
+#include "GizmoBoxObject.generated.h"
+
+/**
+ * Simple object intended to be used as part of 3D Gizmos.
+ * Draws a solid axis-aligned 3D box based on parameters.
+ */
+UCLASS(Transient)
+class INTERACTIVETOOLSFRAMEWORK_API UGizmoBoxObject : public UGizmoBaseObject
+{
+	GENERATED_BODY()
+
+public:
+	//~ Begin UGizmoBaseObject Interface.
+	virtual void Render(IToolsContextRenderAPI* RenderAPI) override;
+	virtual FInputRayHit LineTraceObject(const FVector RayStart, const FVector RayDirection) override;
+	//~ End UGizmoBaseObject Interface.
+
+public:
+
+	// Dimensions of box
+	UPROPERTY(EditAnywhere, Category = Options)
+	FVector Dimensions = FVector(20.0f, 20.0f, 20.0f);
+
+	// Direction corresponding to the box's Z-dimension, must be a unit vector
+	// This direction is used for view-dependent visibility.
+	UPROPERTY(EditAnywhere, Category = Options)
+	FVector UpDirection = FVector::ZeroVector;
+
+	// Direction corresponding to the box's Y-dimension, must be a unit vector
+	UPROPERTY(EditAnywhere, Category = Options)
+	FVector SideDirection = FVector::ZeroVector;
+
+	// Cube center is located at (UpDirection * Offset) 
+	UPROPERTY(EditAnywhere, Category = Options)
+	float Offset = 0.0f;
+
+protected:
+
+	bool bVisibleViewDependent = true;
+};
\ No newline at end of file
diff --git a/Engine/Source/Runtime/Experimental/InteractiveToolsFramework/Public/BaseGizmos/GizmoConeObject.h b/Engine/Source/Runtime/Experimental/InteractiveToolsFramework/Public/BaseGizmos/GizmoConeObject.h
new file mode 100644
--- /dev/null
+++ b/Engine/Source/Runtime/Experimental/InteractiveToolsFramework/Public/BaseGizmos/GizmoConeObject.h
@@ -0,0 +1,53 @@
+// Copyright Epic Games, Inc. All Rights Reserved.
+
+#pragma once
+
+#include "CoreMinimal.h"
+#include "BaseGizmos/GizmoBaseObject.h"
+#include "InputState.h"
+#include "UObject/ObjectMacros.h"
+#include "GizmoConeObject.generated.h"
+
+/**
+ * Simple object intended to be used as part of 3D Gizmos.
+ * Draws a solid 3D cone based on parameters.
+ */
+UCLASS(Transient)
+class INTERACTIVETOOLSFRAMEWORK_API UGizmoConeObject : public UGizmoBaseObject
+{
+	GENERATED_BODY()
+
+public:
+	//~ Begin UGizmoBaseObject Interface.
+	virtual void Render(IToolsContextRenderAPI* RenderAPI) override;
+	virtual FInputRayHit LineTraceObject(const FVector Start, const FVector Direction) override;
+	//~ End UGizmoBaseObject Interface.
+
+public:
+
+	// Direction of cone's axis, must be a unit vector
+	// This direction will also be used for view-dependent culling.
+	UPROPERTY(EditAnywhere, Category = Options)
+	FVector Direction = FVector(0.0f, 0.0f, 1.0f);
+
+	// Height of cone
+	UPROPERTY(EditAnywhere, Category = Options)
+	double Height = 1.0f;
+
+	// Cone's point is located at (Direction * Offset)
+	// or the origin if Offset is 0
+	UPROPERTY(EditAnywhere, Category = Options)
+	double Offset = 0.0f;
+
+	// Angle in radians between cone's axis and slant edge
+	UPROPERTY(EditAnywhere, Category = Options)
+	double Angle = 0.274f;
+
+	// Number of sides for tessellating cone
+	UPROPERTY(EditAnywhere, Category = Options)
+	int NumSides = 32;
+
+protected:
+
+	bool bVisibleViewDependent = true;
+};
\ No newline at end of file
diff --git a/Engine/Source/Runtime/Experimental/InteractiveToolsFramework/Public/BaseGizmos/GizmoCylinderObject.h b/Engine/Source/Runtime/Experimental/InteractiveToolsFramework/Public/BaseGizmos/GizmoCylinderObject.h
new file mode 100644
--- /dev/null
+++ b/Engine/Source/Runtime/Experimental/InteractiveToolsFramework/Public/BaseGizmos/GizmoCylinderObject.h
@@ -0,0 +1,53 @@
+// Copyright Epic Games, Inc. All Rights Reserved.
+
+#pragma once
+
+#include "CoreMinimal.h"
+#include "BaseGizmos/GizmoBaseObject.h"
+#include "InputState.h"
+#include "UObject/ObjectMacros.h"
+#include "GizmoCylinderObject.generated.h"
+
+/**
+ * Simple object intended to be used as part of 3D Gizmos.
+ * Draws a solid 3D cylinder based on parameters.
+ */
+UCLASS(Transient)
+class INTERACTIVETOOLSFRAMEWORK_API UGizmoCylinderObject : public UGizmoBaseObject
+{
+	GENERATED_BODY()
+
+public:
+
+	//~ Begin UGizmoBaseObject Interface.
+	virtual void Render(IToolsContextRenderAPI* RenderAPI) override;
+	virtual FInputRayHit LineTraceObject(const FVector RayOrigin, const FVector RayDirection) override;
+	//~ End UGizmoBaseObject Interface.
+
+public:
+
+	// Orientation of cylinder's axis, must be a unit-vector
+	// This direction will also be used for view-dependent culling.
+	UPROPERTY(EditAnywhere, Category = Options)
+	FVector Direction = FVector(1.0, 0.0, 0.0);
+
+	// Length
+	UPROPERTY(EditAnywhere, Category = Options)
+	double Length = 1.0;
+
+	// Base of cylinder is located at (Direction * Offset)
+	UPROPERTY(EditAnywhere, Category = Options)
+	double Offset = 0.0;
+
+	// Radius
+	UPROPERTY(EditAnywhere, Category = Options)
+	double Radius = 1.2;
+
+	// Number of sides for tessellating cylinder
+	UPROPERTY(EditAnywhere, Category = Options)
+	int NumSides = 16;
+
+protected:
+
+	bool bVisibleViewDependent = true;
+};
\ No newline at end of file
diff --git a/Engine/Source/Runtime/Experimental/InteractiveToolsFramework/Public/BaseGizmos/GizmoGroupObject.h b/Engine/Source/Runtime/Experimental/InteractiveToolsFramework/Public/BaseGizmos/GizmoGroupObject.h
new file mode 100644
--- /dev/null
+++ b/Engine/Source/Runtime/Experimental/InteractiveToolsFramework/Public/BaseGizmos/GizmoGroupObject.h
@@ -0,0 +1,48 @@
+// Copyright Epic Games, Inc. All Rights Reserved.
+
+#pragma once
+
+#include "CoreMinimal.h"
+#include "UObject/ObjectMacros.h"
+#include "BaseGizmos/GizmoBaseObject.h"
+#include "BaseGizmos/GizmoBoxObject.h"
+#include "BaseGizmos/GizmoConeObject.h"
+#include "BaseGizmos/GizmoCylinderObject.h"
+#include "InputState.h"
+#include "GizmoGroupObject.generated.h"
+
+/**
+ * Simple group object intended to be used as part of 3D Gizmos.
+ * Contains multiple gizmo objects.
+ */
+UCLASS(Transient)
+class INTERACTIVETOOLSFRAMEWORK_API UGizmoGroupObject : public UGizmoBaseObject
+{
+	GENERATED_BODY()
+
+public:
+	UGizmoGroupObject();
+
+	//~ Begin UGizmoBaseObject Interface.
+	virtual void Render(IToolsContextRenderAPI* RenderAPI) override;
+	virtual FInputRayHit LineTraceObject(const FVector RayOrigin, const FVector RayDirection) override;
+	virtual void SetHoverState(bool bHoveringIn) override;
+	virtual void SetInteractingState(bool bInteractingIn) override;
+	virtual void SetWorldLocalState(bool bWorldIn) override;
+	virtual void SetVisibility(bool bVisibleIn) override;
+	virtual void SetLocalToWorldTransform(FTransform LocalToWorldTransformIn) override;
+	virtual void SetGizmoScale(float InGizmoScale) override;
+	virtual void SetMaterial(UMaterialInterface* InMaterial) override;
+	virtual void SetCurrentMaterial(UMaterialInterface* InCurrentMaterial) override;
+	//~ End UGizmoBaseObject Interface.
+
+	virtual void Add(UGizmoBaseObject* Object)
+	{
+		Objects.Add(Object);
+	}
+
+public:
+
+	UPROPERTY(EditAnywhere, Category = Options)
+	TArray<TObjectPtr<UGizmoBaseObject>> Objects;
+};
\ No newline at end of file
