Added UpdateInteractingState to UGizmoClickTarget interface. Updated hit targets and added new hit target for UGizmo*Object classes.

diff --git a/Engine/Source/Runtime/Experimental/InteractiveToolsFramework/Private/BaseGizmos/HitTargets.cpp b/Engine/Source/Runtime/Experimental/InteractiveToolsFramework/Private/BaseGizmos/HitTargets.cpp
--- a/Engine/Source/Runtime/Experimental/InteractiveToolsFramework/Private/BaseGizmos/HitTargets.cpp
+++ b/Engine/Source/Runtime/Experimental/InteractiveToolsFramework/Private/BaseGizmos/HitTargets.cpp
@@ -4,7 +4,30 @@
 #include "Components/PrimitiveComponent.h"
 #include "Engine/EngineTypes.h"    // FHitResult
 
+FInputRayHit UGizmoLambdaHitTarget::IsHit(const FInputDeviceRay& ClickPos) const
+{
+	if (IsHitFunction)
+	{
+		return IsHit(ClickPos);
+	}
+	return FInputRayHit();
+}
 
+void UGizmoLambdaHitTarget::UpdateHoverState(bool bHovering)
+{
+	if (UpdateHoverFunction)
+	{
+		UpdateHoverFunction(bHovering);
+	}
+}
+
+void UGizmoLambdaHitTarget::UpdateInteractingState(bool bHovering)
+{
+	if (UpdateInteractingFunction)
+	{
+		UpdateInteractingFunction(bHovering);
+	}
+}
 
 FInputRayHit UGizmoComponentHitTarget::IsHit(const FInputDeviceRay& ClickPos) const
 {
@@ -30,3 +53,60 @@ FInputRayHit UGizmoComponentHitTarget::IsHit(const FInputDeviceRay& ClickPos) co
 	return Hit;
 }
 
+void UGizmoComponentHitTarget::UpdateHoverState(bool bHovering)
+{
+	if (UpdateHoverFunction)
+	{
+		UpdateHoverFunction(bHovering);
+	}
+}
+
+void UGizmoComponentHitTarget::UpdateInteractingState(bool bHovering)
+{
+	if (UpdateInteractingFunction)
+	{
+		UpdateInteractingFunction(bHovering);
+	}
+}
+
+UGizmoComponentHitTarget* UGizmoComponentHitTarget::Construct(UPrimitiveComponent* Component, UObject* Outer)
+{
+	UGizmoComponentHitTarget* NewHitTarget = NewObject<UGizmoComponentHitTarget>(Outer);
+	NewHitTarget->Component = Component;
+	return NewHitTarget;
+}
+
+FInputRayHit UGizmoObjectHitTarget::IsHit(const FInputDeviceRay& ClickPos) const
+{
+	if (GizmoObject && (!Condition || Condition(ClickPos)))
+	{
+		return GizmoObject->LineTraceObject(ClickPos.WorldRay.Origin, ClickPos.WorldRay.Direction);
+	}
+	return FInputRayHit();
+}
+
+void UGizmoObjectHitTarget::UpdateHoverState(bool bHovering)
+{
+	if (GizmoObject)
+	{
+		GizmoObject->SetHoverState(bHovering);
+	}
+}
+
+void UGizmoObjectHitTarget::UpdateInteractingState(bool bInteracting)
+{
+	if (GizmoObject)
+	{
+		GizmoObject->SetInteractingState(bInteracting);
+	}
+}
+
+UGizmoObjectHitTarget* UGizmoObjectHitTarget::Construct(UGizmoBaseObject* InGizmoObject, UObject* Outer)
+{
+	UGizmoObjectHitTarget* NewHitTarget = NewObject<UGizmoObjectHitTarget>(Outer);
+	NewHitTarget->GizmoObject = InGizmoObject;
+	return NewHitTarget;
+}
+
+
+
diff --git a/Engine/Source/Runtime/Experimental/InteractiveToolsFramework/Public/BaseGizmos/GizmoInterfaces.h b/Engine/Source/Runtime/Experimental/InteractiveToolsFramework/Public/BaseGizmos/GizmoInterfaces.h
--- a/Engine/Source/Runtime/Experimental/InteractiveToolsFramework/Public/BaseGizmos/GizmoInterfaces.h
+++ b/Engine/Source/Runtime/Experimental/InteractiveToolsFramework/Public/BaseGizmos/GizmoInterfaces.h
@@ -5,6 +5,7 @@
 #include "CoreMinimal.h"
 #include "UObject/Interface.h"
 #include "InputState.h"
+#include "ToolContextInterfaces.h"
 #include "GizmoInterfaces.generated.h"
 
 
@@ -103,8 +104,19 @@ public:
 	//UFUNCTION()    // FInputDeviceRay is not USTRUCT because FRay is not USTRUCT
 	virtual FInputRayHit IsHit(const FInputDeviceRay& ClickPos) const = 0;
 
+	/*
+	 * Updates the hover state indicating whether the input device is currently hovering over the Standard gizmo.
+	 * This should be be set to false once interaction with the gizmo commences.
+	 */
+	UFUNCTION()
+	virtual void UpdateHoverState(bool bHovering) = 0;
+
+	/*
+	 * Updates the interacting state indicating when interaction with the Standard gizmo is actively occurring, 
+	 * typically upon the input device clicking and dragging the Standard gizmo.
+	 */
 	UFUNCTION()
-	virtual void UpdateHoverState(bool bHovering) const = 0;
+	virtual void UpdateInteractingState(bool bInteracting) = 0;
 };
 
 
diff --git a/Engine/Source/Runtime/Experimental/InteractiveToolsFramework/Public/BaseGizmos/HitTargets.h b/Engine/Source/Runtime/Experimental/InteractiveToolsFramework/Public/BaseGizmos/HitTargets.h
--- a/Engine/Source/Runtime/Experimental/InteractiveToolsFramework/Public/BaseGizmos/HitTargets.h
+++ b/Engine/Source/Runtime/Experimental/InteractiveToolsFramework/Public/BaseGizmos/HitTargets.h
@@ -4,6 +4,7 @@
 
 #include "CoreMinimal.h"
 #include "BaseGizmos/GizmoInterfaces.h"
+#include "BaseGizmos/GizmoBaseObject.h"
 #include "HitTargets.generated.h"
 
 class UPrimitiveComponent;
@@ -24,24 +25,16 @@ public:
 	/** This function is called to update hover state of the target */
 	TFunction<void(bool)> UpdateHoverFunction;
 
+	/** This function is called to update interacting state of the target */
+	TFunction<void(bool)> UpdateInteractingFunction;
+
 
 public:
-	virtual FInputRayHit IsHit(const FInputDeviceRay& ClickPos) const
-	{
-		if (IsHitFunction)
-		{
-			return IsHit(ClickPos);
-		}
-		return FInputRayHit();
-	}
-
-	virtual void UpdateHoverState(bool bHovering) const
-	{
-		if (UpdateHoverFunction)
-		{
-			UpdateHoverFunction(bHovering);
-		}
-	}
+	virtual FInputRayHit IsHit(const FInputDeviceRay& ClickPos) const;
+
+	virtual void UpdateHoverState(bool bHovering);
+
+	virtual void UpdateInteractingState(bool bHovering);
 };
 
 
@@ -73,26 +66,56 @@ public:
 	/** This function is called to update hover state of the target */
 	TFunction<void(bool)> UpdateHoverFunction;
 
+	/** This function is called to update interacting state of the target */
+	TFunction<void(bool)> UpdateInteractingFunction;
 
 public:
 	virtual FInputRayHit IsHit(const FInputDeviceRay& ClickPos) const;
 
-	virtual void UpdateHoverState(bool bHovering) const
-	{
-		if (UpdateHoverFunction)
-		{
-			UpdateHoverFunction(bHovering);
-		}
-	}
+	virtual void UpdateHoverState(bool bHovering);
+
+	virtual void UpdateInteractingState(bool bHovering);
 
 public:
 	static UGizmoComponentHitTarget* Construct(
 		UPrimitiveComponent* Component,
-		UObject* Outer = (UObject*)GetTransientPackage())
-	{
-		UGizmoComponentHitTarget* NewHitTarget = NewObject<UGizmoComponentHitTarget>(Outer);
-		NewHitTarget->Component = Component;
-		return NewHitTarget;
-	}
+		UObject* Outer = (UObject*)GetTransientPackage());
+};
+
+/**
+ * UGizmoObjectHitTarget is an IGizmoClickTarget implementation that
+ * hit-tests any object derived from UGizmoBaseObject
+ */
+UCLASS()
+class INTERACTIVETOOLSFRAMEWORK_API UGizmoObjectHitTarget : public UObject, public IGizmoClickTarget
+{
+	GENERATED_BODY()
+public:
+
+	/**
+	 * Gizmo object.
+	 */
+	UPROPERTY()
+	TObjectPtr<UGizmoBaseObject> GizmoObject;
+
+	/**
+	 * If set, this condition is checked before performing the hit test. This gives a way
+	 * to disable the hit test without hiding the component. This is useful, for instance,
+	 * in a repositionable transform gizmo in world-coordinate mode, where the rotation
+	 * components need to be hittable for movement, but not for repositioning.
+	 */
+	TFunction<bool(const FInputDeviceRay&)> Condition = nullptr;
+
+public:
+	virtual FInputRayHit IsHit(const FInputDeviceRay& ClickPos) const;
+
+	virtual void UpdateHoverState(bool bHovering);
+
+	virtual void UpdateInteractingState(bool bInteracting);
+
+public:
+	static UGizmoObjectHitTarget* Construct(
+		UGizmoBaseObject* InGizmoObject,
+		UObject* Outer = (UObject*)GetTransientPackage());
 };
 
