Added support for rendering UGizmo*Objects. Also update hit target interacting state when interacting. Eventually all other sub-gizmos should be modified to support the same.

diff --git a/Engine/Source/Runtime/Experimental/InteractiveToolsFramework/Private/BaseGizmos/AxisPositionGizmo.cpp b/Engine/Source/Runtime/Experimental/InteractiveToolsFramework/Private/BaseGizmos/AxisPositionGizmo.cpp
--- a/Engine/Source/Runtime/Experimental/InteractiveToolsFramework/Private/BaseGizmos/AxisPositionGizmo.cpp
+++ b/Engine/Source/Runtime/Experimental/InteractiveToolsFramework/Private/BaseGizmos/AxisPositionGizmo.cpp
@@ -39,6 +39,13 @@ void UAxisPositionGizmo::Setup()
 	bInInteraction = false;
 }
 
+void UAxisPositionGizmo::Render(IToolsContextRenderAPI* RenderAPI)
+{
+	if (GizmoObject)
+	{
+		GizmoObject->Render(RenderAPI);
+	}
+}
 
 FInputRayHit UAxisPositionGizmo::CanBeginClickDragSequence(const FInputDeviceRay& PressPos)
 {
@@ -87,6 +94,11 @@ void UAxisPositionGizmo::OnClickPress(const FInputDeviceRay& PressPos)
 
 	bInInteraction = true;
 
+	if (HitTarget)
+	{
+		HitTarget->UpdateInteractingState(bInInteraction);
+	}
+
 	if (StateTarget)
 	{
 		StateTarget->BeginUpdate();
@@ -136,6 +148,11 @@ void UAxisPositionGizmo::OnClickRelease(const FInputDeviceRay& ReleasePos)
 		StateTarget->EndUpdate();
 	}
 	bInInteraction = false;
+
+	if (HitTarget)
+	{
+		HitTarget->UpdateInteractingState(bInInteraction);
+	}
 }
 
 
@@ -149,6 +166,11 @@ void UAxisPositionGizmo::OnTerminateDragSequence()
 		StateTarget->EndUpdate();
 	}
 	bInInteraction = false;
+
+	if (HitTarget)
+	{
+		HitTarget->UpdateInteractingState(bInInteraction);
+	}
 }
 
 
diff --git a/Engine/Source/Runtime/Experimental/InteractiveToolsFramework/Public/BaseGizmos/AxisPositionGizmo.h b/Engine/Source/Runtime/Experimental/InteractiveToolsFramework/Public/BaseGizmos/AxisPositionGizmo.h
--- a/Engine/Source/Runtime/Experimental/InteractiveToolsFramework/Public/BaseGizmos/AxisPositionGizmo.h
+++ b/Engine/Source/Runtime/Experimental/InteractiveToolsFramework/Public/BaseGizmos/AxisPositionGizmo.h
@@ -47,6 +47,7 @@ public:
 	// UInteractiveGizmo overrides
 
 	virtual void Setup() override;
+	virtual void Render(IToolsContextRenderAPI* RenderAPI) override;
 
 	// IClickDragBehaviorTarget implementation
 
@@ -84,6 +85,11 @@ public:
 	UPROPERTY()
 	TObjectPtr<UClickDragInputBehavior> MouseBehavior;
 
+	/** Optional object used to render the gizmo. */
+	UPROPERTY()
+	TObjectPtr<UGizmoBaseObject> GizmoObject;
+
+
 public:
 	/** If enabled, then the sign on the parameter delta is always "increasing" when moving away from the origin point, rather than just being a projection onto the axis */
 	UPROPERTY()
