Added StateTarget and TransformSource for use with UGizmoBaseObject derived classes.

diff --git a/Engine/Source/Runtime/Experimental/InteractiveToolsFramework/Private/BaseGizmos/TransformSources.cpp b/Engine/Source/Runtime/Experimental/InteractiveToolsFramework/Private/BaseGizmos/TransformSources.cpp
--- a/Engine/Source/Runtime/Experimental/InteractiveToolsFramework/Private/BaseGizmos/TransformSources.cpp
+++ b/Engine/Source/Runtime/Experimental/InteractiveToolsFramework/Private/BaseGizmos/TransformSources.cpp
@@ -20,8 +20,20 @@ void UGizmoComponentWorldTransformSource::SetTransform(const FTransform& NewTran
 	OnTransformChanged.Broadcast(this);
 }
 
+FTransform UGizmoObjectWorldTransformSource::GetTransform() const
+{
+	return Object->LocalToWorldTransform;
+}
 
-
+void UGizmoObjectWorldTransformSource::SetTransform(const FTransform& TransformIn)
+{
+	if (bModifyObjectOnTransform)
+	{
+		Object->Modify();
+	}
+	Object->SetLocalToWorldTransform(TransformIn);
+	OnTransformChanged.Broadcast(this);
+}
 
 
 FTransform UGizmoScaledTransformSource::GetTransform() const
diff --git a/Engine/Source/Runtime/Experimental/InteractiveToolsFramework/Public/BaseGizmos/StateTargets.h b/Engine/Source/Runtime/Experimental/InteractiveToolsFramework/Public/BaseGizmos/StateTargets.h
--- a/Engine/Source/Runtime/Experimental/InteractiveToolsFramework/Public/BaseGizmos/StateTargets.h
+++ b/Engine/Source/Runtime/Experimental/InteractiveToolsFramework/Public/BaseGizmos/StateTargets.h
@@ -7,6 +7,7 @@
 #include "InteractiveGizmoManager.h"
 #include "InteractiveToolManager.h"
 #include "BaseGizmos/GizmoInterfaces.h"
+#include "BaseGizmos/GizmoBaseObject.h"
 #include "Changes/TransformChange.h"
 #include "StateTargets.generated.h"
 
@@ -269,3 +270,138 @@ public:
 
 
 };
+
+/**
+ * UGizmoObjectTransformChangeStateTarget is an implementation of IGizmoStateTarget that
+ * emits an FComponentWorldTransformChange on a target USceneComponent. This StateTarget
+ * also opens/closes an undo transaction via GizmoManager.
+ *
+ * The DependentChangeSources and ExternalDependentChangeSources lists allow additional
+ * FChange objects to be inserted into the transaction, provided by IToolCommandChangeSource implementations.
+ */
+UCLASS()
+class INTERACTIVETOOLSFRAMEWORK_API UGizmoObjectTransformChangeStateTarget : public UObject, public IGizmoStateTarget
+{
+	GENERATED_BODY()
+public:
+	virtual void BeginUpdate()
+	{
+		if (TargetObject.IsValid())
+		{
+			if (TransactionManager)
+			{
+				TransactionManager->BeginUndoTransaction(ChangeDescription);
+			}
+
+			InitialTransform = TargetObject->GetLocalToWorldTransform();
+
+			for (TUniquePtr<IToolCommandChangeSource>& Source : DependentChangeSources)
+			{
+				Source->BeginChange();
+			}
+
+			for (IToolCommandChangeSource* Source : ExternalDependentChangeSources)
+			{
+				Source->BeginChange();
+			}
+		}
+	}
+
+	virtual void EndUpdate()
+	{
+		if (TargetObject.IsValid())
+		{
+			FinalTransform = TargetObject->GetLocalToWorldTransform();
+
+			if (TransactionManager)
+			{
+				TUniquePtr<FComponentWorldTransformChange> TransformChange
+					= MakeUnique<FComponentWorldTransformChange>(InitialTransform, FinalTransform);
+				TransactionManager->EmitObjectChange(TargetObject.Get(), MoveTemp(TransformChange), ChangeDescription);
+
+				for (TUniquePtr<IToolCommandChangeSource>& Source : DependentChangeSources)
+				{
+					TUniquePtr<FToolCommandChange> Change = Source->EndChange();
+					if (Change)
+					{
+						UObject* Target = Source->GetChangeTarget();
+						TransactionManager->EmitObjectChange(Target, MoveTemp(Change), Source->GetChangeDescription());
+					}
+				}
+
+				for (IToolCommandChangeSource* Source : ExternalDependentChangeSources)
+				{
+					TUniquePtr<FToolCommandChange> Change = Source->EndChange();
+					if (Change)
+					{
+						UObject* Target = Source->GetChangeTarget();
+						TransactionManager->EmitObjectChange(Target, MoveTemp(Change), Source->GetChangeDescription());
+					}
+				}
+
+				TransactionManager->EndUndoTransaction();
+			}
+		}
+	}
+
+
+	/**
+	 * The object that will be changed, ie have Modify() called on it on BeginUpdate()
+	 */
+	TWeakObjectPtr<UGizmoBaseObject> TargetObject;
+
+	/**
+	 * Localized text description of the transaction (will be visible in Editor on undo/redo)
+	 */
+	FText ChangeDescription;
+
+	/**
+	 * Pointer to the GizmoManager or ToolManager that is used to open/close the transaction
+	 */
+	UPROPERTY()
+	TScriptInterface<IToolContextTransactionProvider> TransactionManager;
+
+	/** Start Transform, saved on BeginUpdate() */
+	FTransform InitialTransform;
+	/** End Transform, saved on EndUpdate() */
+	FTransform FinalTransform;
+
+
+	/**
+	 * Dependent-change generators. This will be told about update start/end, and any generated changes will also be emitted.
+	 * This allows (eg) TransformProxy change events to be collected at the same time as changes to a gizmo target component.
+	 */
+	TArray<TUniquePtr<IToolCommandChangeSource>> DependentChangeSources;
+
+	/**
+	 * Dependent-change generators that are not owned by this class, otherwise handled identically to DependentChangeSources
+	 */
+	TArray<IToolCommandChangeSource*> ExternalDependentChangeSources;
+
+public:
+
+	/**
+	 * Create and initialize an standard instance of UGizmoObjectTransformChangeStateTarget
+	 * @param TargetComponentIn the USceneComponent this StateTarget will track transform changes on
+	 * @param DescriptionIn Localized text description of the transaction
+	 * @param TransactionManagerIn pointer to the GizmoManager or ToolManager used to manage transactions
+	 */
+	static UGizmoObjectTransformChangeStateTarget* Construct(
+		UGizmoBaseObject* TargetObjectIn,
+		const FText& DescriptionIn,
+		IToolContextTransactionProvider* TransactionManagerIn,
+		UObject* Outer = (UObject*)GetTransientPackage())
+	{
+		UGizmoObjectTransformChangeStateTarget* NewTarget = NewObject<UGizmoObjectTransformChangeStateTarget>(Outer);
+		NewTarget->TargetObject = TargetObjectIn;
+		NewTarget->ChangeDescription = DescriptionIn;
+
+		// have to explicitly configure this because we only have IToolContextTransactionProvider pointer
+		NewTarget->TransactionManager.SetInterface(TransactionManagerIn);
+		NewTarget->TransactionManager.SetObject(CastChecked<UObject>(TransactionManagerIn));
+
+		return NewTarget;
+	}
+
+
+};
diff --git a/Engine/Source/Runtime/Experimental/InteractiveToolsFramework/Public/BaseGizmos/TransformSources.h b/Engine/Source/Runtime/Experimental/InteractiveToolsFramework/Public/BaseGizmos/TransformSources.h
--- a/Engine/Source/Runtime/Experimental/InteractiveToolsFramework/Public/BaseGizmos/TransformSources.h
+++ b/Engine/Source/Runtime/Experimental/InteractiveToolsFramework/Public/BaseGizmos/TransformSources.h
@@ -3,6 +3,7 @@
 #pragma once
 
 #include "CoreMinimal.h"
+#include "BaseGizmos/GizmoBaseObject.h"
 #include "BaseGizmos/GizmoInterfaces.h"
 #include "BaseGizmos/TransformProxy.h"
 #include "TransformSources.generated.h"
@@ -72,6 +73,40 @@ public:
 	}
 };
 
+/**
+ * UGizmoObjectWorldTransformSource implements IGizmoTransformSource (via UGizmoBaseTransformSource)
+ * based on the internal transform of a UGizmoBaseObject;
+ */
+UCLASS()
+class INTERACTIVETOOLSFRAMEWORK_API UGizmoObjectWorldTransformSource : public UGizmoBaseTransformSource
+{
+	GENERATED_BODY()
+public:
+
+	virtual FTransform GetTransform() const override;
+
+	virtual void SetTransform(const FTransform& NewTransform) override;
+
+	UPROPERTY()
+	TObjectPtr<UGizmoBaseObject> Object;
+
+	/** If true, Object->Modify() is called on SetTransform */
+	UPROPERTY()
+	bool bModifyObjectOnTransform = true;
+
+public:
+	/**
+	 * Construct a default instance of UGizmoObjectWorldTransformSource with the given Component
+	 */
+	static UGizmoObjectWorldTransformSource* Construct(
+		UGizmoBaseObject* InObject,
+		UObject* Outer = (UObject*)GetTransientPackage())
+	{
+		UGizmoObjectWorldTransformSource* NewSource = NewObject<UGizmoObjectWorldTransformSource>(Outer);
+		NewSource->Object = InObject;
+		return NewSource;
+	}
+};
 
 
 
