Chaos - Remove unused stat

diff --git a/Engine/Source/Runtime/Experimental/Chaos/Private/Chaos/Collision/CollisionDetector.cpp b/Engine/Source/Runtime/Experimental/Chaos/Private/Chaos/Collision/CollisionDetector.cpp
--- a/Engine/Source/Runtime/Experimental/Chaos/Private/Chaos/Collision/CollisionDetector.cpp
+++ b/Engine/Source/Runtime/Experimental/Chaos/Private/Chaos/Collision/CollisionDetector.cpp
@@ -11,7 +11,6 @@ namespace Chaos
 	DEFINE_STAT(STAT_Collisions_BroadPhase);
 	DEFINE_STAT(STAT_Collisions_SpatialBroadPhase);
 	DEFINE_STAT(STAT_Collisions_Filtering);
-	DEFINE_STAT(STAT_Collisions_ComputeBoundsThickness);
 	DEFINE_STAT(STAT_Collisions_GenerateCollisions);
 	DEFINE_STAT(STAT_Collisions_ReceiveCollisions);
 #if CHAOS_ENABLE_STAT_NARROWPHASE
diff --git a/Engine/Source/Runtime/Experimental/Chaos/Public/Chaos/Collision/SpatialAccelerationBroadPhase.h b/Engine/Source/Runtime/Experimental/Chaos/Public/Chaos/Collision/SpatialAccelerationBroadPhase.h
--- a/Engine/Source/Runtime/Experimental/Chaos/Public/Chaos/Collision/SpatialAccelerationBroadPhase.h
+++ b/Engine/Source/Runtime/Experimental/Chaos/Public/Chaos/Collision/SpatialAccelerationBroadPhase.h
@@ -18,7 +18,6 @@ namespace Chaos
 {
 	DECLARE_CYCLE_STAT_EXTERN(TEXT("Collisions::BroadPhase"), STAT_Collisions_SpatialBroadPhase, STATGROUP_ChaosCollision, CHAOS_API);
 	DECLARE_CYCLE_STAT_EXTERN(TEXT("Collisions::Filtering"), STAT_Collisions_Filtering, STATGROUP_ChaosCollision, CHAOS_API);
-	DECLARE_CYCLE_STAT_EXTERN(TEXT("Collisions::ComputeBoundsThickness"), STAT_Collisions_ComputeBoundsThickness, STATGROUP_ChaosCollision, CHAOS_API);
 	DECLARE_CYCLE_STAT_EXTERN(TEXT("Collisions::GenerateCollisions"), STAT_Collisions_GenerateCollisions, STATGROUP_ChaosCollision, CHAOS_API);
 	DECLARE_CYCLE_STAT_EXTERN(TEXT("Collisions::ReceiveCollisions"), STAT_Collisions_ReceiveCollisions, STATGROUP_ChaosCollision, CHAOS_API);
 	
