Unshelved from pending changelist '15313623':

PR #7742: Chaos: Fix bugs that unwelding deletes valid FPerShapeData (Contributed by [REDACTED_PERSON])

- When more than one bodies are welded and a body(except last one) is unwelding, last body's FPerShapeData gets deleted by 'ShapesArray.SetNum' at UpdateShapesArrayFromGeometry()
- Which invalidates FBodyInstance::ShapeToBodiesMap

diff --git a/Engine/Source/Runtime/Experimental/Chaos/Private/Chaos/ParticleHandle.cpp b/Engine/Source/Runtime/Experimental/Chaos/Private/Chaos/ParticleHandle.cpp
--- a/Engine/Source/Runtime/Experimental/Chaos/Private/Chaos/ParticleHandle.cpp
+++ b/Engine/Source/Runtime/Experimental/Chaos/Private/Chaos/ParticleHandle.cpp
@@ -136,6 +136,7 @@ namespace Chaos
 		{
 			if (InShape == MShapesArray[Index].Get())
 			{
+				MShapesArray.RemoveAt(Index);
 				FoundIndex = Index;
 				break;
 			}
@@ -143,7 +144,7 @@ namespace Chaos
 
 		if (MNonFrequentData.Read().Geometry()->GetType() == FImplicitObjectUnion::StaticType())
 		{
-			// if we are currently a union then add the new geometry to this union
+			// if we are currently a union then remove geometry from this union
 			MNonFrequentData.Modify(true, MDirtyFlags, Proxy, [FoundIndex](auto& Data)
 				{
 					if (Data.AccessGeometry())
