Chaos - Add maximum linear and angular speeds per-particle. Default linear limit to float max, angular speed is configured by project settings

diff --git a/Engine/Source/Runtime/Core/Public/UObject/UE5ReleaseStreamObjectVersion.h b/Engine/Source/Runtime/Core/Public/UObject/UE5ReleaseStreamObjectVersion.h
--- a/Engine/Source/Runtime/Core/Public/UObject/UE5ReleaseStreamObjectVersion.h
+++ b/Engine/Source/Runtime/Core/Public/UObject/UE5ReleaseStreamObjectVersion.h
@@ -32,6 +32,9 @@ struct CORE_API FUE5ReleaseStreamObjectVersion
 
 		// Change hair strands vertex format
 		HairStrandsVertexFormatChange,
+		
+		// Added max linear and angular speed to Chaos bodies
+		AddChaosMaxLinearAngularSpeed,
 
 		// -----<new versions can be added above this line>-------------------------------------------------
 		VersionPlusOne,
diff --git a/Engine/Source/Runtime/Engine/Private/PhysicsEngine/BodyInstance.cpp b/Engine/Source/Runtime/Engine/Private/PhysicsEngine/BodyInstance.cpp
--- a/Engine/Source/Runtime/Engine/Private/PhysicsEngine/BodyInstance.cpp
+++ b/Engine/Source/Runtime/Engine/Private/PhysicsEngine/BodyInstance.cpp
@@ -1212,6 +1212,7 @@ void FInitBodiesHelperBase::CreateActor_AssumesLocked(FBodyInstance* Instance, c
 		FPhysicsInterface::CreateActor(ActorParams, Instance->ActorHandle);
 		FPhysicsInterface::SetCcdEnabled_AssumesLocked(Instance->ActorHandle, Instance->bUseCCD);
 		FPhysicsInterface::SetIsKinematic_AssumesLocked(Instance->ActorHandle, !Instance->ShouldInstanceSimulatingPhysics());
+		FPhysicsInterface::SetMaxLinearVelocity_AssumesLocked(Instance->ActorHandle, TNumericLimits<float>::Max());
 
 		// Set sleep even notification
 		FPhysicsInterface::SetSendsSleepNotifies_AssumesLocked(Instance->ActorHandle, Instance->bGenerateWakeEvents);
diff --git a/Engine/Source/Runtime/Experimental/Chaos/Public/Chaos/ParticleDirtyFlags.h b/Engine/Source/Runtime/Experimental/Chaos/Public/Chaos/ParticleDirtyFlags.h
--- a/Engine/Source/Runtime/Experimental/Chaos/Public/Chaos/ParticleDirtyFlags.h
+++ b/Engine/Source/Runtime/Experimental/Chaos/Public/Chaos/ParticleDirtyFlags.h
@@ -14,6 +14,7 @@
 #include "UObject/ExternalPhysicsMaterialCustomObjectVersion.h"
 #include "UObject/PhysicsObjectVersion.h"
 #include "UObject/UE5MainStreamObjectVersion.h"
+#include "UObject/UE5ReleaseStreamObjectVersion.h"
 
 class FName;
 
@@ -251,12 +252,15 @@ class FParticleDynamicMisc
 public:
 	void Serialize(FChaosArchive& Ar)
 	{
+		Ar.UsingCustomVersion(FUE5ReleaseStreamObjectVersion::GUID);
 		Ar.UsingCustomVersion(FUE5MainStreamObjectVersion::GUID);
 		Ar.UsingCustomVersion(FExternalPhysicsCustomObjectVersion::GUID);
+
 		Ar << MLinearEtherDrag;
 		Ar << MAngularEtherDrag;
 		Ar << MObjectState;
 		Ar << MGravityEnabled;
+
 		if (Ar.CustomVer(FExternalPhysicsCustomObjectVersion::GUID) >= FExternalPhysicsCustomObjectVersion::AddOneWayInteraction)
 		{
 			Ar << MOneWayInteraction;
@@ -265,6 +269,7 @@ public:
 		{
 			MOneWayInteraction = false;
 		}
+
 		if (Ar.CustomVer(FPhysicsObjectVersion::GUID) >= FPhysicsObjectVersion::AddCCDEnableFlag)
 		{
 			Ar << bCCDEnabled;
@@ -279,6 +284,12 @@ public:
 		{
 			Ar << bDisabled;
 		}
+		
+		if(Ar.CustomVer(FUE5ReleaseStreamObjectVersion::GUID) >= FUE5ReleaseStreamObjectVersion::AddChaosMaxLinearAngularSpeed)
+		{
+			Ar << MMaxLinearSpeedSq;
+			Ar << MMaxAngularSpeedSq;
+		}
 	}
 
 	template <typename TOther>
@@ -286,6 +297,8 @@ public:
 	{
 		SetLinearEtherDrag(Other.LinearEtherDrag());
 		SetAngularEtherDrag(Other.AngularEtherDrag());
+		SetMaxLinearSpeedSq(Other.MaxLinearSpeedSq());
+		SetMaxAngularSpeedSq(Other.MaxAngularSpeedSq());
 		SetObjectState(Other.ObjectState());
 		SetGravityEnabled(Other.GravityEnabled());
 		SetCollisionGroup(Other.CollisionGroup());
@@ -302,6 +315,8 @@ public:
 		return ObjectState() == Other.ObjectState()
 			&& LinearEtherDrag() == Other.LinearEtherDrag()
 			&& AngularEtherDrag() == Other.AngularEtherDrag()
+			&& MaxLinearSpeedSq() == Other.MaxLinearSpeedSq()
+			&& MaxAngularSpeedSq() == Other.MaxAngularSpeedSq()
 			&& GravityEnabled() == Other.GravityEnabled()
 			&& CollisionGroup() == Other.CollisionGroup()
 			&& ResimType() == Other.ResimType()
@@ -322,6 +337,12 @@ public:
 	FReal AngularEtherDrag() const { return MAngularEtherDrag; }
 	void SetAngularEtherDrag(FReal InAngularEtherDrag) { MAngularEtherDrag = InAngularEtherDrag; }
 
+	FReal MaxLinearSpeedSq() const { return MMaxLinearSpeedSq; }
+	void SetMaxLinearSpeedSq(FReal InMaxLinearSpeed) { MMaxLinearSpeedSq = InMaxLinearSpeed; }
+
+	FReal MaxAngularSpeedSq() const { return MMaxAngularSpeedSq; }
+	void SetMaxAngularSpeedSq(FReal InMaxAngularSpeed) { MMaxAngularSpeedSq = InMaxAngularSpeed; }
+
 	EObjectStateType ObjectState() const { return MObjectState; }
 	void SetObjectState(EObjectStateType InState){ MObjectState = InState; }
 
@@ -348,6 +369,8 @@ public:
 private:
 	FReal MLinearEtherDrag;
 	FReal MAngularEtherDrag;
+	FReal MMaxLinearSpeedSq;
+	FReal MMaxAngularSpeedSq;
 	int32 MCollisionGroup;
 
 	EObjectStateType MObjectState;
diff --git a/Engine/Source/Runtime/Experimental/Chaos/Public/Chaos/ParticleHandle.h b/Engine/Source/Runtime/Experimental/Chaos/Public/Chaos/ParticleHandle.h
--- a/Engine/Source/Runtime/Experimental/Chaos/Public/Chaos/ParticleHandle.h
+++ b/Engine/Source/Runtime/Experimental/Chaos/Public/Chaos/ParticleHandle.h
@@ -803,6 +803,8 @@ public:
 	{
 		SetLinearEtherDrag(DynamicMisc.LinearEtherDrag());
 		SetAngularEtherDrag(DynamicMisc.AngularEtherDrag());
+		SetMaxLinearSpeedSq(DynamicMisc.MaxLinearSpeedSq());
+		SetMaxAngularSpeedSq(DynamicMisc.MaxAngularSpeedSq());
 		SetCollisionGroup(DynamicMisc.CollisionGroup());
 		SetGravityEnabled(DynamicMisc.GravityEnabled());
 		SetCCDEnabled(DynamicMisc.CCDEnabled());
@@ -849,6 +851,14 @@ public:
 	T& AngularEtherDrag() { return PBDRigidParticles->AngularEtherDrag(ParticleIdx); }
 	void SetAngularEtherDrag(const T& InAngularEtherDrag) { PBDRigidParticles->AngularEtherDrag(ParticleIdx) = InAngularEtherDrag; }
 
+	T MaxLinearSpeedSq() const { return PBDRigidParticles->MaxLinearSpeedSq(ParticleIdx); }
+	T& MaxLinearSpeedSq() { return PBDRigidParticles->MaxLinearSpeedSq(ParticleIdx); }
+	void SetMaxLinearSpeedSq(const T& InMaxLinearSpeed) { PBDRigidParticles->MaxLinearSpeedSq(ParticleIdx) = InMaxLinearSpeed; }
+
+	T MaxAngularSpeedSq() const { return PBDRigidParticles->MaxAngularSpeedSq(ParticleIdx); }
+	T& MaxAngularSpeedSq() { return PBDRigidParticles->MaxAngularSpeedSq(ParticleIdx); }
+	void SetMaxAngularSpeedSq(const T& InMaxAngularSpeed) { PBDRigidParticles->MaxAngularSpeedSq(ParticleIdx) = InMaxAngularSpeed; }
+
 	int32 Island() const { return PBDRigidParticles->Island(ParticleIdx); }
 	int32& Island() { return PBDRigidParticles->Island(ParticleIdx); }
 	void SetIsland(const int32 InIsland) { PBDRigidParticles->Island(ParticleIdx) = InIsland; }
@@ -2274,6 +2284,18 @@ public:
 		MMiscData.Modify(true,MDirtyFlags,Proxy,[&InAngularEtherDrag](auto& Data){ Data.SetAngularEtherDrag(InAngularEtherDrag);});
 	}
 
+	T MaxLinearSpeedSq() const { return MMiscData.Read().MaxLinearSpeedSq(); }
+	void SetMaxLinearSpeedSq(const T& InLinearSpeed)
+	{
+		MMiscData.Modify(true,MDirtyFlags,Proxy,[&InLinearSpeed](auto& Data){ Data.SetMaxLinearSpeedSq(InLinearSpeed);});
+	}
+
+	T MaxAngularSpeedSq() const { return MMiscData.Read().MaxAngularSpeedSq(); }
+	void SetMaxAngularSpeedSq(const T& InAngularSpeed)
+	{
+		MMiscData.Modify(true,MDirtyFlags,Proxy,[&InAngularSpeed](auto& Data){ Data.SetMaxAngularSpeedSq(InAngularSpeed);});
+	}
+
 	int32 Island() const { return MIsland; }
 	// TODO(stett): Make the setter private. It is public right now to provide access to proxies.
 	void SetIsland(const int32 InIsland)
diff --git a/Engine/Source/Runtime/Experimental/Chaos/Public/Chaos/PerParticleEtherDrag.h b/Engine/Source/Runtime/Experimental/Chaos/Public/Chaos/PerParticleEtherDrag.h
--- a/Engine/Source/Runtime/Experimental/Chaos/Public/Chaos/PerParticleEtherDrag.h
+++ b/Engine/Source/Runtime/Experimental/Chaos/Public/Chaos/PerParticleEtherDrag.h
@@ -16,7 +16,7 @@ namespace Chaos
 		FPerParticleEtherDrag() {}
 		virtual ~FPerParticleEtherDrag() {}
 
-		inline void ApplyHelper(FVec3& V, FVec3& W, FReal LinearDamp, FReal AngularDamp, FReal Dt) const
+		inline void ApplyHelper(FVec3& V, FVec3& W, FReal LinearDamp, FReal AngularDamp, FReal LinearLimitSq, FReal AngularLimitSq, FReal Dt) const
 		{
 			const FReal LinearDrag = LinearEtherDragOverride >= 0 ? LinearEtherDragOverride : LinearDamp * Dt;
 			const FReal LinearMultiplier = FMath::Max(FReal(0), FReal(1) - LinearDrag);
@@ -25,6 +25,19 @@ namespace Chaos
 			const FReal AngularDrag = AngularEtherDragOverride >= 0 ? AngularEtherDragOverride : AngularDamp * Dt;
 			const FReal AngularMultiplier = FMath::Max(FReal(0), FReal(1) - AngularDrag);
 			W *= AngularMultiplier;
+
+			const FReal LinearSpeedSq = V.SizeSquared();
+			const FReal AngularSpeedSq = W.SizeSquared();
+
+			if(LinearSpeedSq > LinearLimitSq)
+			{
+				V *= FMath::Sqrt(LinearLimitSq / LinearSpeedSq);
+			}
+
+			if(AngularSpeedSq > AngularLimitSq)
+			{
+				W *= FMath::Sqrt(AngularLimitSq / AngularSpeedSq);
+			}
 		}
 
 		inline void Apply(FDynamicParticles& InParticles, const FReal Dt, const int32 Index) const override //-V762
@@ -34,12 +47,12 @@ namespace Chaos
 
 		inline void Apply(TRigidParticles<FReal, 3>& InParticles, const FReal Dt, const int32 Index) const override //-V762
 		{
-			ApplyHelper(InParticles.V(Index), InParticles.W(Index), InParticles.LinearEtherDrag(Index), InParticles.AngularEtherDrag(Index), Dt);
+			ApplyHelper(InParticles.V(Index), InParticles.W(Index), InParticles.LinearEtherDrag(Index), InParticles.AngularEtherDrag(Index), InParticles.MaxLinearSpeedSq(Index), InParticles.MaxAngularSpeedSq(Index), Dt);
 		}
 
 		inline void Apply(TTransientPBDRigidParticleHandle<FReal, 3>& Particle, const FReal Dt) const override //-V762
 		{
-			ApplyHelper(Particle.V(), Particle.W(), Particle.LinearEtherDrag(), Particle.AngularEtherDrag(), Dt);
+			ApplyHelper(Particle.V(), Particle.W(), Particle.LinearEtherDrag(), Particle.AngularEtherDrag(), Particle.MaxLinearSpeedSq(), Particle.MaxAngularSpeedSq(), Dt);
 		}
 	};
 
diff --git a/Engine/Source/Runtime/Experimental/Chaos/Public/Chaos/RigidParticles.h b/Engine/Source/Runtime/Experimental/Chaos/Public/Chaos/RigidParticles.h
--- a/Engine/Source/Runtime/Experimental/Chaos/Public/Chaos/RigidParticles.h
+++ b/Engine/Source/Runtime/Experimental/Chaos/Public/Chaos/RigidParticles.h
@@ -83,6 +83,8 @@ public:
 		TArrayCollection::AddArray(&MRotationOfMass);
 		TArrayCollection::AddArray(&MLinearEtherDrag);
 		TArrayCollection::AddArray(&MAngularEtherDrag);
+		TArrayCollection::AddArray(&MaxLinearSpeedsSq);
+		TArrayCollection::AddArray(&MaxAngularSpeedsSq);
 		TArrayCollection::AddArray(&MCollisionParticles);
 		TArrayCollection::AddArray(&MCollisionGroup);
 		TArrayCollection::AddArray(&MCollisionConstraintFlags);
@@ -109,6 +111,8 @@ public:
 		, MInvM(MoveTemp(Other.MInvM))
 		, MCenterOfMass(MoveTemp(Other.MCenterOfMass))
 		, MRotationOfMass(MoveTemp(Other.MRotationOfMass))
+		, MaxLinearSpeedsSq(MoveTemp(Other.MaxLinearSpeedsSq))
+		, MaxAngularSpeedsSq(MoveTemp(Other.MaxAngularSpeedsSq))
 		, MCollisionParticles(MoveTemp(Other.MCollisionParticles))
 		, MCollisionGroup(MoveTemp(Other.MCollisionGroup))
 		, MCollisionConstraintFlags(MoveTemp(Other.MCollisionConstraintFlags))
@@ -132,6 +136,8 @@ public:
 		TArrayCollection::AddArray(&MRotationOfMass);
 		TArrayCollection::AddArray(&MLinearEtherDrag);
 		TArrayCollection::AddArray(&MAngularEtherDrag);
+		TArrayCollection::AddArray(&MaxLinearSpeedsSq);
+		TArrayCollection::AddArray(&MaxAngularSpeedsSq);
 		TArrayCollection::AddArray(&MCollisionParticles);
 		TArrayCollection::AddArray(&MCollisionGroup);
 		TArrayCollection::AddArray(&MCollisionConstraintFlags);
@@ -190,6 +196,12 @@ public:
 	FORCEINLINE const T& AngularEtherDrag(const int32 index) const { return MAngularEtherDrag[index]; }
 	FORCEINLINE T& AngularEtherDrag(const int32 index) { return MAngularEtherDrag[index]; }
 
+	FORCEINLINE const T& MaxLinearSpeedSq(const int32 index) const { return MaxLinearSpeedsSq[index]; }
+	FORCEINLINE T& MaxLinearSpeedSq(const int32 index) { return MaxLinearSpeedsSq[index]; }
+
+	FORCEINLINE const T& MaxAngularSpeedSq(const int32 index) const { return MaxAngularSpeedsSq[index]; }
+	FORCEINLINE T& MaxAngularSpeedSq(const int32 index) { return MaxAngularSpeedsSq[index]; }
+
 	FORCEINLINE int32 CollisionParticlesSize(int32 Index) const { return MCollisionParticles[Index] == nullptr ? 0 : MCollisionParticles[Index]->Size(); }
 
 	void CHAOS_API CollisionParticlesInitIfNeeded(const int32 Index);
@@ -306,6 +318,8 @@ public:
 	FORCEINLINE TArray<TRotation<T, d>>& AllRotationOfMass() { return MRotationOfMass; }
 	FORCEINLINE TArray<FReal>& AllLinearEtherDrag() { return MLinearEtherDrag; }
 	FORCEINLINE TArray<FReal>& AllAngularEtherDrag() { return MAngularEtherDrag; }
+	FORCEINLINE TArray<FReal>& AllMaxLinearSpeeds() { return MaxLinearSpeedsSq; }
+	FORCEINLINE TArray<FReal>& AllMaxAngularSpeeds() { return MaxAngularSpeedsSq; }
 	FORCEINLINE TArray<bool>& AllDisabled() { return MDisabled; }
 	FORCEINLINE TArray<EObjectStateType>& AllObjectState() { return MObjectState; }
 	FORCEINLINE TArray<bool>& AllGravityEnabled() { return MGravityEnabled; }
@@ -326,6 +340,8 @@ private:
 	TArrayCollectionArray<TRotation<T,d>> MRotationOfMass;
 	TArrayCollectionArray<T> MLinearEtherDrag;
 	TArrayCollectionArray<T> MAngularEtherDrag;
+	TArrayCollectionArray<T> MaxLinearSpeedsSq;
+	TArrayCollectionArray<T> MaxAngularSpeedsSq;
 	TArrayCollectionArray<TUniquePtr<TBVHParticles<T, d>>> MCollisionParticles;
 	TArrayCollectionArray<int32> MCollisionGroup;
 	TArrayCollectionArray<uint32> MCollisionConstraintFlags;
diff --git a/Engine/Source/Runtime/Experimental/Chaos/Public/PhysicsProxy/SingleParticlePhysicsProxy.h b/Engine/Source/Runtime/Experimental/Chaos/Public/PhysicsProxy/SingleParticlePhysicsProxy.h
--- a/Engine/Source/Runtime/Experimental/Chaos/Public/PhysicsProxy/SingleParticlePhysicsProxy.h
+++ b/Engine/Source/Runtime/Experimental/Chaos/Public/PhysicsProxy/SingleParticlePhysicsProxy.h
@@ -14,6 +14,7 @@
 #include "Chaos/Core.h"
 #include "PhysicsProxy/SingleParticlePhysicsProxyFwd.h"
 #include "Framework/Threading.h"
+#include "Math/NumericLimits.h"
 
 namespace Chaos
 {
@@ -838,7 +839,7 @@ public:
 	int32 Island() const
 	{
 		VerifyContext();
-		if (auto Rigid = GetParticle_LowLevel()->CastToRigidParticle())
+		if (const TPBDRigidParticle<FReal, 3>* Rigid = GetParticle_LowLevel()->CastToRigidParticle())
 		{
 			return Rigid->Island();
 		}
@@ -849,7 +850,7 @@ public:
 	void SetIsland(const int32 InIsland)
 	{
 		VerifyContext();
-		if (auto Rigid = GetParticle_LowLevel()->CastToRigidParticle())
+		if (TPBDRigidParticle<FReal, 3>* Rigid = GetParticle_LowLevel()->CastToRigidParticle())
 		{
 			Rigid->SetIsland(InIsland);
 		}
@@ -858,7 +859,7 @@ public:
 	bool ToBeRemovedOnFracture() const
 	{
 		VerifyContext();
-		if (auto Rigid = GetParticle_LowLevel()->CastToRigidParticle())
+		if (const TPBDRigidParticle<FReal, 3>* Rigid = GetParticle_LowLevel()->CastToRigidParticle())
 		{
 			return Rigid->ToBeRemovedOnFracture();
 		}
@@ -869,7 +870,7 @@ public:
 	void SetToBeRemovedOnFracture(const bool InToBeRemovedOnFracture)
 	{
 		VerifyContext();
-		if (auto Rigid = GetParticle_LowLevel()->CastToRigidParticle())
+		if (TPBDRigidParticle<FReal, 3>* Rigid = GetParticle_LowLevel()->CastToRigidParticle())
 		{
 			Rigid->SetToBeRemovedOnFracture(InToBeRemovedOnFracture);
 		}
@@ -878,7 +879,7 @@ public:
 	void ClearEvents()
 	{
 		VerifyContext();
-		if (auto Rigid = GetParticle_LowLevel()->CastToRigidParticle())
+		if (TPBDRigidParticle<FReal, 3>* Rigid = GetParticle_LowLevel()->CastToRigidParticle())
 		{
 			Rigid->ClearEvents();
 		}
@@ -887,7 +888,7 @@ public:
 	EWakeEventEntry GetWakeEvent()
 	{
 		VerifyContext();
-		if (auto Rigid = GetParticle_LowLevel()->CastToRigidParticle())
+		if (TPBDRigidParticle<FReal, 3>* Rigid = GetParticle_LowLevel()->CastToRigidParticle())
 		{
 			return Rigid->GetWakeEvent();
 		}
@@ -898,7 +899,7 @@ public:
 	void ClearForces(bool bInvalidate = true)
 	{
 		VerifyContext();
-		if (auto Rigid = GetParticle_LowLevel()->CastToRigidParticle())
+		if (TPBDRigidParticle<FReal, 3>*Rigid = GetParticle_LowLevel()->CastToRigidParticle())
 		{
 			Rigid->ClearForces(bInvalidate);
 		}
@@ -907,7 +908,7 @@ public:
 	void ClearTorques(bool bInvalidate = true)
 	{
 		VerifyContext();
-		if (auto Rigid = GetParticle_LowLevel()->CastToRigidParticle())
+		if (TPBDRigidParticle<FReal, 3>*Rigid = GetParticle_LowLevel()->CastToRigidParticle())
 		{
 			Rigid->ClearTorques(bInvalidate);
 		}
@@ -943,7 +944,7 @@ public:
 	void SetCCDEnabled(bool bEnabled)
 	{
 		VerifyContext();
-		if (auto Rigid = GetParticle_LowLevel()->CastToRigidParticle())
+		if (TPBDRigidParticle<FReal, 3>* Rigid = GetParticle_LowLevel()->CastToRigidParticle())
 		{
 			Rigid->SetCCDEnabled(bEnabled);
 		}
@@ -952,7 +953,7 @@ public:
 	bool CCDEnabled() const
 	{
 		VerifyContext();
-		if (auto Rigid = GetParticle_LowLevel()->CastToRigidParticle())
+		if (const TPBDRigidParticle<FReal, 3>* Rigid = GetParticle_LowLevel()->CastToRigidParticle())
 		{
 			return Rigid->CCDEnabled();
 		}
@@ -960,6 +961,46 @@ public:
 		return false;
 	}
 
+	void SetMaxLinearSpeedSq(FReal InNewSpeed)
+	{
+		VerifyContext();
+		if(TPBDRigidParticle<FReal, 3>* Rigid = GetParticle_LowLevel()->CastToRigidParticle())
+		{
+			Rigid->SetMaxLinearSpeedSq(InNewSpeed);
+		}
+	}
+
+	FReal GetMaxLinearSpeedSq()
+	{
+		VerifyContext();
+		if(const TPBDRigidParticle<FReal, 3>* Rigid = GetParticle_LowLevel()->CastToRigidParticle())
+		{
+			return Rigid->MaxLinearSpeedSq();
+		}
+
+		return TNumericLimits<FReal>::Max();
+	}
+
+	void SetMaxAngularSpeedSq(FReal InNewSpeed)
+	{
+		VerifyContext();
+		if(TPBDRigidParticle<FReal, 3>* Rigid = GetParticle_LowLevel()->CastToRigidParticle())
+		{
+			Rigid->SetMaxAngularSpeedSq(InNewSpeed);
+		}
+	}
+
+	FReal GetMaxAngularSpeedSq()
+	{
+		VerifyContext();
+		if(const TPBDRigidParticle<FReal, 3>* Rigid = GetParticle_LowLevel()->CastToRigidParticle())
+		{
+			return Rigid->MaxAngularSpeedSq();
+		}
+
+		return TNumericLimits<FReal>::Max();
+	}
+
 	void SetDisabled(bool bDisable)
 	{
 		VerifyContext();
diff --git a/Engine/Source/Runtime/Experimental/Chaos/Public/RewindData.h b/Engine/Source/Runtime/Experimental/Chaos/Public/RewindData.h
--- a/Engine/Source/Runtime/Experimental/Chaos/Public/RewindData.h
+++ b/Engine/Source/Runtime/Experimental/Chaos/Public/RewindData.h
@@ -123,6 +123,18 @@ public:
 		return DynamicsMisc.IsSet() ? DynamicsMisc.Read().AngularEtherDrag() : Particle.CastToRigidParticle()->AngularEtherDrag();
 	}
 
+	template <typename TParticle>
+	FReal MaxLinearSpeedSq(const TParticle& Particle) const
+	{
+		return DynamicsMisc.IsSet() ? DynamicsMisc.Read().MaxLinearSpeedSq() : Particle.CastToRigidParticle()->MaxLinearSpeedSq();
+	}
+
+	template <typename TParticle>
+	FReal MaxAngularSpeedSq(const TParticle& Particle) const
+	{
+		return DynamicsMisc.IsSet() ? DynamicsMisc.Read().MaxAngularSpeedSq() : Particle.CastToRigidParticle()->MaxAngularSpeedSq();
+	}
+
 	template <typename TParticle>
 	EObjectStateType ObjectState(const TParticle& Particle) const
 	{
@@ -300,6 +312,16 @@ public:
 		return State.AngularEtherDrag(Particle);
 	}
 
+	FReal MaxLinearSpeedSq() const
+	{
+		return State.MaxLinearSpeedSq(Particle);
+	}
+
+	FReal MaxAngularSpeedSq() const
+	{
+		return State.MaxAngularSpeedSq(Particle);
+	}
+
 	EObjectStateType ObjectState() const
 	{
 		return State.ObjectState(Particle);
diff --git a/Engine/Source/Runtime/PhysicsCore/Private/ChaosEngineInterface.cpp b/Engine/Source/Runtime/PhysicsCore/Private/ChaosEngineInterface.cpp
--- a/Engine/Source/Runtime/PhysicsCore/Private/ChaosEngineInterface.cpp
+++ b/Engine/Source/Runtime/PhysicsCore/Private/ChaosEngineInterface.cpp
@@ -580,13 +580,44 @@ void FChaosEngineInterface::SetAngularVelocity_AssumesLocked(const FPhysicsActor
 
 float FChaosEngineInterface::GetMaxAngularVelocity_AssumesLocked(const FPhysicsActorHandle& InActorReference)
 {
-	CHAOS_ENSURE(false);
-	return FLT_MAX;
+	if(ensure(FChaosEngineInterface::IsValid(InActorReference)))
+	{
+		return FMath::Sqrt(InActorReference->GetGameThreadAPI().GetMaxAngularSpeedSq());
+	}
+
+	return TNumericLimits<float>::Max();
 }
 
-void FChaosEngineInterface::SetMaxAngularVelocity_AssumesLocked(const FPhysicsActorHandle& InActorReference,float InMaxAngularVelocity)
+float FChaosEngineInterface::GetMaxLinearVelocity_AssumesLocked(const FPhysicsActorHandle& InActorReference)
 {
-	CHAOS_ENSURE(false);
+	if(ensure(FChaosEngineInterface::IsValid(InActorReference)))
+	{
+		return FMath::Sqrt(InActorReference->GetGameThreadAPI().GetMaxLinearSpeedSq());
+	}
+
+	return TNumericLimits<float>::Max();
+}
+
+void FChaosEngineInterface::SetMaxAngularVelocity_AssumesLocked(const FPhysicsActorHandle& InActorReference,float InMaxAngularVelocityRadians)
+{
+	// We're about to square the input so we clamp to this maximum
+	static const float MaxInput = FMath::Sqrt(TNumericLimits<float>::Max());
+
+	if(ensure(FChaosEngineInterface::IsValid(InActorReference)))
+	{
+		InActorReference->GetGameThreadAPI().SetMaxAngularSpeedSq(InMaxAngularVelocityRadians > MaxInput ? TNumericLimits<float>::Max() : InMaxAngularVelocityRadians * InMaxAngularVelocityRadians);
+	}
+}
+
+void FChaosEngineInterface::SetMaxLinearVelocity_AssumesLocked(const FPhysicsActorHandle& InActorReference, float InMaxLinearVelocity)
+{
+	// We're about to square the input so we clamp to this maximum
+	static const float MaxInput = FMath::Sqrt(TNumericLimits<float>::Max());
+
+	if(ensure(FChaosEngineInterface::IsValid(InActorReference)))
+	{
+		InActorReference->GetGameThreadAPI().SetMaxLinearSpeedSq(InMaxLinearVelocity > MaxInput ? TNumericLimits<float>::Max() : InMaxLinearVelocity * InMaxLinearVelocity);
+	}
 }
 
 float FChaosEngineInterface::GetMaxDepenetrationVelocity_AssumesLocked(const FPhysicsActorHandle& InActorReference)
diff --git a/Engine/Source/Runtime/PhysicsCore/Public/Chaos/ChaosEngineInterface.h b/Engine/Source/Runtime/PhysicsCore/Public/Chaos/ChaosEngineInterface.h
--- a/Engine/Source/Runtime/PhysicsCore/Public/Chaos/ChaosEngineInterface.h
+++ b/Engine/Source/Runtime/PhysicsCore/Public/Chaos/ChaosEngineInterface.h
@@ -408,7 +408,9 @@ public:
 	static FVector GetAngularVelocity_AssumesLocked(const FPhysicsActorHandle& InActorReference);
 	static void SetAngularVelocity_AssumesLocked(const FPhysicsActorHandle& InActorReference,const FVector& InNewVelocity,bool bAutoWake = true);
 	static float GetMaxAngularVelocity_AssumesLocked(const FPhysicsActorHandle& InActorReference);
-	static void SetMaxAngularVelocity_AssumesLocked(const FPhysicsActorHandle& InActorReference,float InMaxAngularVelocity);
+	static float GetMaxLinearVelocity_AssumesLocked(const FPhysicsActorHandle& InActorReference);
+	static void SetMaxAngularVelocity_AssumesLocked(const FPhysicsActorHandle& InActorReference,float InMaxAngularVelocityRadians);
+	static void SetMaxLinearVelocity_AssumesLocked(const FPhysicsActorHandle& InActorReference, float InMaxLinearVelocity);
 
 	static float GetMaxDepenetrationVelocity_AssumesLocked(const FPhysicsActorHandle& InActorReference);
 	static void SetMaxDepenetrationVelocity_AssumesLocked(const FPhysicsActorHandle& InActorReference,float InMaxDepenetrationVelocity);
