Chaos - Added missed default initialization for particle speed limits (missed from previous checkin 15930001)

diff --git a/Engine/Source/Runtime/Experimental/Chaos/Public/Chaos/ParticleHandle.h b/Engine/Source/Runtime/Experimental/Chaos/Public/Chaos/ParticleHandle.h
--- a/Engine/Source/Runtime/Experimental/Chaos/Public/Chaos/ParticleHandle.h
+++ b/Engine/Source/Runtime/Experimental/Chaos/Public/Chaos/ParticleHandle.h
@@ -2,6 +2,7 @@
 #pragma once
 
 #include "Templates/ChooseClass.h"
+#include "Math/NumericLimits.h"
 #include "Chaos/PBDRigidClusteredParticles.h"
 #include "Chaos/PBDGeometryCollectionParticles.h"
 #include "Chaos/ParticleHandleFwd.h"
@@ -77,6 +78,8 @@ void PBDRigidParticleDefaultConstruct(FConcrete& Concrete, const FPBDRigidPartic
 	Concrete.SetCollisionGroup(0);
 	Concrete.SetLinearImpulse(TVector<T, d>(0));
 	Concrete.SetAngularImpulse(TVector<T, d>(0));
+	Concrete.SetMaxLinearSpeedSq(TNumericLimits<T>::Max());
+	Concrete.SetMaxAngularSpeedSq(TNumericLimits<T>::Max());
 	Concrete.SetM(1);
 	Concrete.SetInvM(1);
 	Concrete.SetCenterOfMass(TVector<T,d>(0));
