fix crash in fracture's flatten tool where it tries to rename based on the parent of the root node, but root does not have a parent

#rnx

diff --git a/Engine/Source/Runtime/Experimental/Chaos/Private/GeometryCollection/GeometryCollectionClusteringUtility.cpp b/Engine/Source/Runtime/Experimental/Chaos/Private/GeometryCollection/GeometryCollectionClusteringUtility.cpp
--- a/Engine/Source/Runtime/Experimental/Chaos/Private/GeometryCollection/GeometryCollectionClusteringUtility.cpp
+++ b/Engine/Source/Runtime/Experimental/Chaos/Private/GeometryCollection/GeometryCollectionClusteringUtility.cpp
@@ -282,7 +282,10 @@ void FGeometryCollectionClusteringUtility::ClusterBonesUnderExistingNode(FGeomet
 
 			for (int32 NodeIndex : ParentsToUpdateNames)
 			{
-				RecursivelyUpdateChildBoneNames(NodeIndex, Children, BoneNames);
+				if (NodeIndex != INDEX_NONE)
+				{
+					RecursivelyUpdateChildBoneNames(NodeIndex, Children, BoneNames);
+				}
 			}
 		}
 	}
@@ -537,7 +540,10 @@ int32 FGeometryCollectionClusteringUtility::GetParentOfBoneAtSpecifiedLevel(cons
 
 void FGeometryCollectionClusteringUtility::RecursivelyUpdateChildBoneNames(int32 BoneIndex, const TManagedArray<TSet<int32>>& Children, TManagedArray<FString>& BoneNames, bool OverrideBoneNames /*= false*/)
 {
-	check(BoneIndex < Children.Num());
+	if (!ensure(BoneIndex > -1 && BoneIndex < Children.Num()))
+	{
+		return;
+	}
 
 	if (Children[BoneIndex].Num() > 0)
 	{
