Chaos - debug draw optimization

diff --git a/Engine/Source/Runtime/Experimental/Chaos/Private/Chaos/ChaosDebugDraw.cpp b/Engine/Source/Runtime/Experimental/Chaos/Private/Chaos/ChaosDebugDraw.cpp
--- a/Engine/Source/Runtime/Experimental/Chaos/Private/Chaos/ChaosDebugDraw.cpp
+++ b/Engine/Source/Runtime/Experimental/Chaos/Private/Chaos/ChaosDebugDraw.cpp
@@ -45,41 +45,42 @@ namespace Chaos
 		FAutoConsoleVariableRef CVarChaosDebugDrawExactShapes(TEXT("p.Chaos.DebugDraw.ShowExactCoreShapes"), bChaosDebugDebugDrawExactCoreShapes, TEXT("Whether to show the exact core shape. NOTE: Extremely expensive and should only be used on a small scene with a couple convex shapes in it"));
 
 
-		// NOTE: These settings should never really be used - they are the fallback defaults
-		// if the user does not specify settings in the debug draw call.
-		// See PBDRigidsColver.cpp and ImmediatePhysicsSimulation_Chaos.cpp for example.
-		FChaosDebugDrawSettings ChaosDefaultDebugDebugDrawSettings(
-			/* ArrowSize =					*/ 1.5f,
-			/* BodyAxisLen =				*/ 4.0f,
-			/* ContactLen =					*/ 4.0f,
-			/* ContactWidth =				*/ 2.0f,
-			/* ContactPhiWidth =			*/ 0.0f,
-			/* ContactOwnerWidth =			*/ 0.0f,
-			/* ConstraintAxisLen =			*/ 5.0f,
-			/* JointComSize =				*/ 2.0f,
-			/* LineThickness =				*/ 0.15f,
-			/* DrawScale =					*/ 1.0f,
-			/* FontHeight =					*/ 10.0f,
-			/* FontScale =					*/ 1.5f,
-			/* ShapeThicknesScale =			*/ 1.0f,
-			/* PointSize =					*/ 2.0f,
-			/* VelScale =					*/ 0.0f,
-			/* AngVelScale =				*/ 0.0f,
-			/* ImpulseScale =				*/ 0.0f,
-			/* InertiaScale =				*/ 0.0f,
-			/* DrawPriority =				*/ 10.0f,
-			/* bShowSimple =				*/ true,
-			/* bShowComplex =				*/ false,
-			/* bInShowLevelSetCollision =	*/ false
-		);
-
 		const FChaosDebugDrawSettings& GetChaosDebugDrawSettings(const FChaosDebugDrawSettings* InSettings)
 		{
 			if (InSettings != nullptr)
 			{
 				return *InSettings;
 			}
-			return ChaosDefaultDebugDebugDrawSettings;
+
+			{
+				static FChaosDebugDrawSettings ChaosDefaultDebugDebugDrawSettings;
+				return ChaosDefaultDebugDebugDrawSettings;
+			}
+		}
+
+		// NOTE: These settings should never really be used - they are the fallback defaults
+		// if the user does not specify settings in the debug draw call.
+		// See PBDRigidsColver.cpp and ImmediatePhysicsSimulation_Chaos.cpp for example.
+		FChaosDebugDrawSettings::FChaosDebugDrawSettings()
+			: ArrowSize(1.5f)
+			, BodyAxisLen(4.0f)
+			, ContactLen(4.0f)
+			, ContactWidth(2.0f)
+			, ContactPhiWidth(0.0f)
+			, ContactOwnerWidth(0.0f)
+			, ConstraintAxisLen(5.0f)
+			, JointComSize(2.0f)
+			, LineThickness(0.15f)
+			, DrawScale(1.0f)
+			, FontHeight(10.0f)
+			, FontScale(1.5f)
+			, ShapeThicknesScale(1.0f)
+			, PointSize(2.0f)
+			, VelScale(0.0f)
+			, AngVelScale(0.0f)
+			, ImpulseScale(0.0f)
+			, DrawPriority(10.0f)
+		{
 		}
 
 		//
@@ -298,6 +299,13 @@ namespace Chaos
 			const EImplicitObjectType PackedType = Shape->GetType(); // Type includes scaling and instancing data
 			const EImplicitObjectType InnerType = GetInnerType(Shape->GetType());
 
+			// Are we within the region of interest?
+			const FReal ParticleSize = Particle->HasBounds() ? 0.5f * Particle->LocalBounds().Extents().Size() : TNumericLimits<FReal>::Max();
+			if (!FDebugDrawQueue::GetInstance().IsInRegionOfInterest(ShapeTransform.GetLocation(), ParticleSize))
+			{
+				return;
+			}
+
 			// Unwrap the wrapper/aggregating shapes
 			if (IsScaled(PackedType))
 			{
@@ -608,6 +616,18 @@ namespace Chaos
 					const FRigidTransform3 WorldCoMTransform0 = FParticleUtilities::GetCoMWorldTransform(Particle0);
 					const FRigidTransform3 WorldCoMTransform1 = FParticleUtilities::GetCoMWorldTransform(Particle1);
 
+					// Are we within the region of interest?
+					const FReal Particle0Size = Contact.Particle[0]->HasBounds() ? 0.5f * Contact.Particle[0]->LocalBounds().Extents().Size() : TNumericLimits<FReal>::Max();
+					const FReal Particle1Size = Contact.Particle[1]->HasBounds() ? 0.5f * Contact.Particle[1]->LocalBounds().Extents().Size() : TNumericLimits<FReal>::Max();
+					if (!FDebugDrawQueue::GetInstance().IsInRegionOfInterest(WorldCoMTransform0.GetLocation(), Particle0Size))
+					{
+						return;
+					}
+					if (!FDebugDrawQueue::GetInstance().IsInRegionOfInterest(WorldCoMTransform1.GetLocation(), Particle1Size))
+					{
+						return;
+					}
+
 					for (const FManifoldPoint& ManifoldPoint : PointConstraint->GetManifoldPoints())
 					{
 						const bool bIsActive = ManifoldPoint.bActive || !ManifoldPoint.NetPushOut.IsNearlyZero();
diff --git a/Engine/Source/Runtime/Experimental/Chaos/Public/Chaos/DebugDrawQueue.h b/Engine/Source/Runtime/Experimental/Chaos/Public/Chaos/DebugDrawQueue.h
--- a/Engine/Source/Runtime/Experimental/Chaos/Public/Chaos/DebugDrawQueue.h
+++ b/Engine/Source/Runtime/Experimental/Chaos/Public/Chaos/DebugDrawQueue.h
@@ -398,6 +398,16 @@ public:
 		return RadiusOfInterest;
 	}
 
+	bool IsInRegionOfInterest(FVector Pos) const
+	{
+		return (RadiusOfInterest <= 0.0f) || ((Pos - CenterOfInterest).SizeSquared() < RadiusOfInterest * RadiusOfInterest);
+	}
+
+	bool IsInRegionOfInterest(FVector Pos, FReal Radius) const
+	{
+		return (RadiusOfInterest <= 0.0f) || ((Pos - CenterOfInterest).SizeSquared() < (RadiusOfInterest + Radius) * (RadiusOfInterest + Radius));
+	}
+
 	void SetConsumerActive(void* Consumer, bool bConsumerActive);
 
 	static FDebugDrawQueue& GetInstance()
@@ -437,11 +447,6 @@ private:
 		return false;
 	}
 
-	bool IsInRegionOfInterest(FVector Pos) const
-	{
-		return (RadiusOfInterest <= 0.0f) || ((Pos - CenterOfInterest).SizeSquared() < RadiusOfInterest * RadiusOfInterest);
-	}
-
 	bool IsInBudget() const
 	{
 		return (MaxCommandCost <= 0) || (RequestedCommandCost <= MaxCommandCost);
