Chaos - add logging to potential degenerate triangle ensure

diff --git a/Engine/Source/Runtime/Experimental/Chaos/Private/Chaos/TriangleMeshImplicitObject.cpp b/Engine/Source/Runtime/Experimental/Chaos/Private/Chaos/TriangleMeshImplicitObject.cpp
--- a/Engine/Source/Runtime/Experimental/Chaos/Private/Chaos/TriangleMeshImplicitObject.cpp
+++ b/Engine/Source/Runtime/Experimental/Chaos/Private/Chaos/TriangleMeshImplicitObject.cpp
@@ -6,6 +6,7 @@
 #include "Chaos/Convex.h"
 #include "Chaos/ImplicitObjectScaled.h"
 #include "Chaos/GeometryQueries.h"
+#include "Chaos/Utilities.h"
 
 namespace Chaos
 {
@@ -985,9 +986,14 @@ FVec3 FTriangleMeshImplicitObject::GetFaceNormal(const int32 FaceIdx) const
 			const FVec3 AB = B - A;
 			const FVec3 AC = C - A;
 			FVec3 Normal = FVec3::CrossProduct(AB, AC);
-			const FReal Length = Normal.SafeNormalize();
-			ensure(Length);
-			return Normal;
+			if (Utilities::NormalizeSafe(Normal))
+			{
+				return Normal;
+			}
+	
+			UE_LOG(LogChaos, Warning, TEXT("Degenerate triangle %d: (%f %f %f) (%f %f %f) (%f %f %f)"), FaceIdx, A.X, A.Y, A.Z, B.X, B.Y, B.Z, C.X, C.Y, C.Z);
+			ensure(false);
+			return FVec3(0, 0, 1);
 		};
 		
 		if (MElements.RequiresLargeIndices())
