Chaos - add upper limit to velocity bounds expansion
This prevents perf issues when moving bodies large distances without teleporting.

diff --git a/Engine/Source/Runtime/Experimental/Chaos/Private/Chaos/BoundingVolumeHierarchy.cpp b/Engine/Source/Runtime/Experimental/Chaos/Private/Chaos/BoundingVolumeHierarchy.cpp
--- a/Engine/Source/Runtime/Experimental/Chaos/Private/Chaos/BoundingVolumeHierarchy.cpp
+++ b/Engine/Source/Runtime/Experimental/Chaos/Private/Chaos/BoundingVolumeHierarchy.cpp
@@ -9,11 +9,14 @@ using namespace Chaos;
 
 namespace Chaos
 {
-	FRealSingle BoundsThicknessMultiplier = 1.f;	//should not be based on bounds, should be based on distance
-	FAutoConsoleVariableRef CVarBoundsThicknessMultiplier(TEXT("p.BoundsThicknessMultiplier"), BoundsThicknessMultiplier, TEXT(""));
-
-	FRealSingle MinBoundsThickness = 0.1f;
-	FAutoConsoleVariableRef CVarMinBoundsThickness(TEXT("p.MinBoundsThickness"), MinBoundsThickness, TEXT(""));
+	// Bounds expansion (for speculative contact creation) is limited by this factor multiplied by the objects size.
+	// This prevents objects that are moving large distances in one frame from having bounds which overlap too many objects. 
+	// We only detect collisions at the final position anyway (when CCD is disabled) so having a bounds which is much larger 
+	// than the object doesn't help much. E.g., if an object moves more than its size per frame it will encounter tunneling 
+	// problems when colliding against zero-thickness surfaces such as tri meshes, regardless of the bounds size.
+	// NOTE: This is only used by objects that have CCD disabled.
+	FRealSingle Chaos_Bounds_MaxInflationScale = 2.f;
+	FAutoConsoleVariableRef CVarBoundsMaxInflatiuonScale(TEXT("p.Chaos.MaxInflationScale"), Chaos_Bounds_MaxInflationScale, TEXT("A limit on the bounds used to detect collisions when CCD is disabled. The bounds limit is this scale multiplied by the object's max dimension"));
 }
 
 template<class OBJECT_ARRAY, class LEAF_TYPE, class T, int d>
diff --git a/Engine/Source/Runtime/Experimental/Chaos/Public/Chaos/BoundingVolumeUtilities.h b/Engine/Source/Runtime/Experimental/Chaos/Public/Chaos/BoundingVolumeUtilities.h
--- a/Engine/Source/Runtime/Experimental/Chaos/Public/Chaos/BoundingVolumeUtilities.h
+++ b/Engine/Source/Runtime/Experimental/Chaos/Public/Chaos/BoundingVolumeUtilities.h
@@ -16,37 +16,56 @@
 
 namespace Chaos
 {
+extern FRealSingle Chaos_Bounds_MaxInflationScale;
 
-// @todo(ccaulfield): COLLISION - get rid of these?
-extern FRealSingle MinBoundsThickness;
-extern FRealSingle BoundsThicknessMultiplier;
-
-inline FVec3 ComputeBoundsThickness(FVec3 Vel, FReal Dt, FReal BoundsThickness, FReal BoundsVelocityInflation)
+inline FVec3 ComputeBoundsThickness(FVec3 Vel, FReal Dt, FReal MinBoundsThickness, FReal MaxBoundsThickness, FReal BoundsVelocityInflation)
 {
-	// @todo(ccaulfield): this expands the bounds in negative velocity direction as well...maybe we don't want that?
 	for (int i = 0; i < 3; ++i)
 	{
-		Vel[i] = FMath::Max(MinBoundsThickness, BoundsThicknessMultiplier * (BoundsThickness + FMath::Abs(Vel[i]) * Dt * BoundsVelocityInflation));
+		const FReal BoundsThickness = FMath::Abs(Vel[i]) * Dt * BoundsVelocityInflation;
+		Vel[i] = FMath::Clamp(BoundsThickness, MinBoundsThickness, MaxBoundsThickness);
 	}
 	return Vel;
 }
 
-inline FVec3 ComputeBoundsThickness(const TPBDRigidParticles<FReal, 3>& InParticles, FReal Dt, int32 BodyIndex, FReal BoundsThickness, FReal BoundsVelocityInflation)
+inline FVec3 ComputeBoundsThickness(const TPBDRigidParticles<FReal, 3>& InParticles, FReal Dt, int32 BodyIndex, FReal MinBoundsThickness, FReal BoundsVelocityInflation)
 {
-	return ComputeBoundsThickness(InParticles.V(BodyIndex), Dt, BoundsThickness, !InParticles.CCDEnabled(BodyIndex) ? FReal(0) : BoundsVelocityInflation);
+	// See comments in ComputeBoundsThickness<THandle> below
+	FReal MaxBoundsThickness = TNumericLimits<FReal>::Max();
+	const bool bIsBounded = InParticles.HasBounds(BodyIndex);
+	const bool bIsCCD = InParticles.CCDEnabled(BodyIndex);
+	if (bIsBounded && !bIsCCD)
+	{
+		MaxBoundsThickness = Chaos_Bounds_MaxInflationScale * InParticles.LocalBounds(BodyIndex).Extents().GetMax();
+	}
+
+	return ComputeBoundsThickness(InParticles.V(BodyIndex), Dt, MinBoundsThickness, MaxBoundsThickness, BoundsVelocityInflation);
 }
 
 template <typename THandle>
-FVec3 ComputeBoundsThickness(const THandle& PBDRigid, FReal Dt, FReal BoundsThickness, FReal BoundsVelocityInflation)
-{
-	FVec3 Vel{ 0 };
-	auto* KinematicBody = PBDRigid.CastToKinematicParticle();
-	if (KinematicBody)
+FVec3 ComputeBoundsThickness(const THandle& ParticleHandle, FReal Dt, FReal MinBoundsThickness, FReal BoundsVelocityInflation)
+{
+	const typename THandle::FDynamicParticleHandleType* RigidParticle = ParticleHandle.CastToRigidParticle();
+	const typename THandle::FKinematicParticleHandleType* KinematicParticle = ParticleHandle.CastToKinematicParticle();
+
+	// Limit the bounds expansion based on the size of the object. This prevents objects that are moved a large
+	// distance without resetting physics from having excessive bounds. Objects that move more than their size per
+	// tick without CCD enabled will have simulation issues anyway, so expanding bounds beyond this is unnecessary.
+	FReal MaxBoundsThickness = TNumericLimits<FReal>::Max();
+	const bool bIsBounded = ParticleHandle.HasBounds();
+	const bool bIsCCD = (RigidParticle != nullptr) && RigidParticle->CCDEnabled();
+	if (bIsBounded && !bIsCCD)
+	{
+		MaxBoundsThickness = Chaos_Bounds_MaxInflationScale * ParticleHandle.LocalBounds().Extents().GetMax();
+	}
+
+	FVec3 Vel(0);
+	if (KinematicParticle != nullptr)
 	{
-		Vel = KinematicBody->V();
+		Vel = KinematicParticle->V();
 	}
 
-	return ComputeBoundsThickness(Vel, Dt, BoundsThickness, BoundsVelocityInflation);
+	return ComputeBoundsThickness(Vel, Dt, MinBoundsThickness, MaxBoundsThickness, BoundsVelocityInflation);
 }
 
 template<class OBJECT_ARRAY>
@@ -158,7 +177,6 @@ TAABB<T, d> ComputeWorldSpaceBoundingBox(const TPBDRigidParticles<T, d>& Objects
 	
 	if (bUseVelocity)
 	{
-		// @todo(ccaulfield): COLLISION - thickness
 		WorldSpaceBox.ThickenSymmetrically(ComputeBoundsThickness(Objects, Dt, i, 0, 1));
 	}
 	return WorldSpaceBox;
@@ -350,7 +368,6 @@ typename TEnableIf<TModels<CParticleView, ParticleView>::Value>::Type ComputeAll
 			{
 				if (const auto PBDRigid = Particle.AsDynamic())
 				{
-					// @todo(ccaulfield): COLLISION - thickness
 					WorldSpaceBoxes.Last().ThickenSymmetrically(ComputeBoundsThickness(*PBDRigid, Dt, 0, 1));
 				}
 			}
diff --git a/Engine/Source/Runtime/Experimental/Chaos/Public/Chaos/ParticleHandle.h b/Engine/Source/Runtime/Experimental/Chaos/Public/Chaos/ParticleHandle.h
--- a/Engine/Source/Runtime/Experimental/Chaos/Public/Chaos/ParticleHandle.h
+++ b/Engine/Source/Runtime/Experimental/Chaos/Public/Chaos/ParticleHandle.h
@@ -307,6 +307,9 @@ template <typename T, int d, bool bPersistent>
 class TGeometryParticleHandleImp : public TParticleHandleBase<T,d>
 {
 public:
+	using FDynamicParticleHandleType = TPBDRigidParticleHandleImp<T, d, bPersistent>;
+	using FKinematicParticleHandleType = TKinematicGeometryParticleHandleImp<T, d, bPersistent>;
+
 	using TTransientHandle = TTransientGeometryParticleHandle<T, d>;
 	using THandleBase = TParticleHandleBase<T, d>;
 	using THandleBase::GeometryParticles;
@@ -1135,6 +1138,9 @@ TGeometryParticleHandleImp<T,d,bPersistent>* TGeometryParticleHandleImp<T,d, bPe
 class CHAOS_API FGenericParticleHandleHandleImp
 {
 public:
+	using FDynamicParticleHandleType = FPBDRigidParticleHandle;
+	using FKinematicParticleHandleType = FKinematicGeometryParticleHandle;
+
 	FGenericParticleHandleHandleImp(FGeometryParticleHandle* InHandle) { MHandle = InHandle; }
 
 	// Check for the exact type of particle (see also AsKinematic etc, which will work on derived types)
@@ -1344,13 +1350,23 @@ public:
 		return MHandle->UniqueIdx();
 	}
 
-	//Named this way for templated code
+	bool HasBounds() const
+	{
+		return MHandle->HasBounds();
+	}
+
+	const FAABB3& LocalBounds() const
+	{
+		return MHandle->LocalBounds();
+	}
+
+	//Named this way for templated code (GT/PT particles)
 	bool HasBoundingBox() const
 	{
 		return MHandle->HasBounds();
 	}
 
-	//Named this way for templated code
+	//Named this way for templated code (GT/PT particles)
 	const FAABB3& BoundingBox() const
 	{
 		return MHandle->WorldSpaceInflatedBounds();
