Chaos: Allow the constraint proxy to correctly set up ignored colliisions when one side of the constraint is static

diff --git a/Engine/Source/Runtime/Experimental/Chaos/Private/PhysicsProxy/JointConstraintProxy.cpp b/Engine/Source/Runtime/Experimental/Chaos/Private/PhysicsProxy/JointConstraintProxy.cpp
--- a/Engine/Source/Runtime/Experimental/Chaos/Private/PhysicsProxy/JointConstraintProxy.cpp
+++ b/Engine/Source/Runtime/Experimental/Chaos/Private/PhysicsProxy/JointConstraintProxy.cpp
@@ -254,21 +254,29 @@ void FJointConstraintPhysicsProxy::PushStateOnPhysicsThread(Chaos::FPBDRigidsSol
 					// .. IgnoreCollisionsManager
 					if (Handle0 && Handle1)
 					{
-						Chaos::TPBDRigidParticleHandle<FReal, 3>* ParticleHandle0 = Handle0->CastToRigidParticle();
-						Chaos::TPBDRigidParticleHandle<FReal, 3>* ParticleHandle1 = Handle1->CastToRigidParticle();
+						Chaos::TPBDRigidParticleHandle<FReal, 3>* RigidHandle0 = Handle0->CastToRigidParticle();
+						Chaos::TPBDRigidParticleHandle<FReal, 3>* RigidHandle1 = Handle1->CastToRigidParticle();
 
-						if (ParticleHandle0 && ParticleHandle1)
+						// As long as one particle is a rigid we can add the ignore entry, one particle can be a static
+						if (RigidHandle0 || RigidHandle1)
 						{
+							const Chaos::FUniqueIdx ID0 = Handle0->UniqueIdx();
+							const Chaos::FUniqueIdx ID1 = Handle1->UniqueIdx();
 							Chaos::FIgnoreCollisionManager& IgnoreCollisionManager = InSolver->GetEvolution()->GetBroadPhase().GetIgnoreCollisionManager();
-							Chaos::FUniqueIdx ID0 = ParticleHandle0->UniqueIdx();
-							Chaos::FUniqueIdx ID1 = ParticleHandle1->UniqueIdx();
 
+							// For rigid/dynamic particles, add the broadphase flag and the IDs to check for disabled collisions
+							if(RigidHandle0)
+							{
+								RigidHandle0->AddCollisionConstraintFlag(Chaos::ECollisionConstraintFlags::CCF_BroadPhaseIgnoreCollisions);
+								IgnoreCollisionManager.AddIgnoreCollisionsFor(ID0, ID1);
+							}
 
-							ParticleHandle0->AddCollisionConstraintFlag(Chaos::ECollisionConstraintFlags::CCF_BroadPhaseIgnoreCollisions);
-							IgnoreCollisionManager.AddIgnoreCollisionsFor(ID0, ID1);
+							if(RigidHandle1)
+							{
+								RigidHandle1->AddCollisionConstraintFlag(Chaos::ECollisionConstraintFlags::CCF_BroadPhaseIgnoreCollisions);
+								IgnoreCollisionManager.AddIgnoreCollisionsFor(ID1, ID0);
+							}
 
-							ParticleHandle1->AddCollisionConstraintFlag(Chaos::ECollisionConstraintFlags::CCF_BroadPhaseIgnoreCollisions);
-							IgnoreCollisionManager.AddIgnoreCollisionsFor(ID1, ID0);
 							ConstraintSettings.bCollisionEnabled = JointSettingsBuffer.bCollisionEnabled;
 						}
 					}
