Chaos
-- Optimized UpdateConnectivityGraphUsingPointImplicit function in PBDRigidClustering.

diff --git a/Engine/Source/Runtime/Experimental/Chaos/Private/Chaos/PBDRigidClustering.cpp b/Engine/Source/Runtime/Experimental/Chaos/Private/Chaos/PBDRigidClustering.cpp
--- a/Engine/Source/Runtime/Experimental/Chaos/Private/Chaos/PBDRigidClustering.cpp
+++ b/Engine/Source/Runtime/Experimental/Chaos/Private/Chaos/PBDRigidClustering.cpp
@@ -2095,63 +2095,62 @@ namespace Chaos
 
 		const FReal Delta = FMath::Min(FMath::Max(Parameters.CoillisionThicknessPercent, FReal(0)), FReal(1));
 		const TArray<FPBDRigidParticleHandle*>& Children = MChildren[Parent];
-		for (int32 i = 0; i < Children.Num(); ++i)
-		{
-			FPBDRigidParticleHandle* Child1 = Children[i];
-			if (!Child1->Geometry() || !Child1->Geometry()->HasBoundingBox())
-			{
-				continue;
-			}
-			const FVec3& Child1X = Child1->X();
-			if (!(ensure(!FMath::IsNaN(Child1X[0])) && ensure(!FMath::IsNaN(Child1X[1])) && ensure(!FMath::IsNaN(Child1X[2]))))
-			{
-				continue;
-			}
-			FRigidTransform3 TM1 = FRigidTransform3(Child1X, Child1->R());
 
-			const int32 Offset = i + 1;
-			const int32 NumRemainingChildren = Children.Num() - Offset;
-			typedef TPair<FPBDRigidParticleHandle*, FPBDRigidParticleHandle*> ParticlePair;
-			typedef TArray<ParticlePair> ParticlePairArray;
-			TArray<ParticlePairArray> Connections;
-			Connections.Init(ParticlePairArray(), NumRemainingChildren);
-			PhysicsParallelFor(NumRemainingChildren, [&](int32 Idx) 
+		typedef TPair<FPBDRigidParticleHandle*, FPBDRigidParticleHandle*> ParticlePair;
+		typedef TSet<ParticlePair> ParticlePairArray;
+
+		TArray<ParticlePairArray> Connections;
+		Connections.Init(ParticlePairArray(), Children.Num());
+
+		PhysicsParallelFor(Children.Num(), [&](int32 i)
 			{
-				const int32 ChildrenIdx = Offset + Idx;
-				FPBDRigidParticleHandle* Child2 = Children[ChildrenIdx];
-				if(!Child2->CollisionParticles())
-					return;
+				FPBDRigidParticleHandle* Child1 = Children[i];
+				if (Child1->Geometry() && Child1->Geometry()->HasBoundingBox())
+				{
+					ParticlePairArray& ConnectionList = Connections[i];
 
-				const FVec3& Child2X = Child2->X();
-				if (!(ensure(!FMath::IsNaN(Child2X[0])) && ensure(!FMath::IsNaN(Child2X[1])) && ensure(!FMath::IsNaN(Child2X[2]))))
-					return;
+					const FVec3& Child1X = Child1->X();
+					FRigidTransform3 TM1 = FRigidTransform3(Child1X, Child1->R());
 
-				const FRigidTransform3 TM = TM1.GetRelativeTransform(FRigidTransform3(Child2X, Child2->R()));
+					const int32 Offset = i + 1;
+					const int32 NumRemainingChildren = Children.Num() - Offset;
 
-				bool bCollided = false;
-				for (uint32 CollisionIdx = 0; !bCollided && CollisionIdx < Child2->CollisionParticles()->Size(); ++CollisionIdx)
-				{
-					const FVec3 LocalPoint = 
-						TM.TransformPositionNoScale(Child2->CollisionParticles()->X(CollisionIdx));
-					const FReal Phi = Child1->Geometry()->SignedDistance(LocalPoint - (LocalPoint * Delta));
-					if (Phi < 0.0)
-						bCollided = true;
-				}
-				if (bCollided)
-				{
-					Connections[Idx].Add(ParticlePair(Child1, Child2));
+					for (int32 Idx = 0; Idx < NumRemainingChildren; ++Idx)
+					{
+						const int32 ChildrenIdx = Offset + Idx;
+						FPBDRigidParticleHandle* Child2 = Children[ChildrenIdx];
+						if (Child2->CollisionParticles())
+						{
+
+							const FVec3& Child2X = Child2->X();
+							const FRigidTransform3 TM = TM1.GetRelativeTransform(FRigidTransform3(Child2X, Child2->R()));
+							const uint32 NumCollisionParticles = Child2->CollisionParticles()->Size();
+							for (uint32 CollisionIdx = 0; CollisionIdx < NumCollisionParticles; ++CollisionIdx)
+							{
+								const FVec3 LocalPoint =
+									TM.TransformPositionNoScale(Child2->CollisionParticles()->X(CollisionIdx));
+								const FReal Phi = Child1->Geometry()->SignedDistance(LocalPoint - (LocalPoint * Delta));
+								if (Phi < 0.0)
+								{
+									ConnectionList.Add(ParticlePair(Child1, Child2));
+									break;
+								}
+
+							}
+						}
+					}
 				}
 			});
 
-			// join results and make connections
-			for (const ParticlePairArray& ConnectionList : Connections)
+		// join results and make connections
+		for (const ParticlePairArray& ConnectionList : Connections)
+		{
+			for (const ParticlePair& Edge : ConnectionList)
 			{
-				for (const ParticlePair& Edge : ConnectionList)
-				{
-					ConnectNodes(Edge.Key, Edge.Value);
-				}
+				ConnectNodes(Edge.Key, Edge.Value);
 			}
 		}
+
 	}
 
 	DECLARE_CYCLE_STAT(TEXT("TPBDRigidClustering<>::FixConnectivityGraphUsingDelaunayTriangulation"), STAT_FixConnectivityGraphUsingDelaunayTriangulation, STATGROUP_Chaos);
