GeometryCollection
-- Changed physics proxy ensure to a logged warning if geometry is found to be too small to simulate

diff --git a/Engine/Source/Runtime/Experimental/Chaos/Private/PhysicsProxy/GeometryCollectionPhysicsProxy.cpp b/Engine/Source/Runtime/Experimental/Chaos/Private/PhysicsProxy/GeometryCollectionPhysicsProxy.cpp
--- a/Engine/Source/Runtime/Experimental/Chaos/Private/PhysicsProxy/GeometryCollectionPhysicsProxy.cpp
+++ b/Engine/Source/Runtime/Experimental/Chaos/Private/PhysicsProxy/GeometryCollectionPhysicsProxy.cpp
@@ -1888,6 +1888,10 @@ void FGeometryCollectionPhysicsProxy::InitializeSharedCollisionStructures(
 	TArray<bool> InertiaComputationNeeded;
 	InertiaComputationNeeded.Init(false, NumGeometries);
 
+	// We skip very small geometry and log as a warning. To avoid log spamming, we wait
+	// until we complete the loop before reporting the skips.
+	bool bSkippedSmallGeometry = false;
+	
 	float TotalVolume = 0.f;
 	// The geometry group has a set of transform indices that maps a geometry index
 	// to a transform index, but only in the case where there is a 1-to-1 mapping 
@@ -1917,8 +1921,9 @@ void FGeometryCollectionPhysicsProxy::InitializeSharedCollisionStructures(
 				MassProperties.InertiaTensor = FMatrix33(1,1,1);
 				MassProperties.Mass = 1.0f; // start with unit mass, scaled later by density
 
-				if (!ensureMsgf(BoundingBox[GeometryIndex].GetExtent().GetAbsMin() > MinVolume, TEXT("Geometry too small to simulate. Idx (%d)"), GeometryIndex))
+				if (BoundingBox[GeometryIndex].GetExtent().GetAbsMin() < MinVolume)
 				{
+					bSkippedSmallGeometry = true;
 					CollectionSimulatableParticles[TransformGroupIndex] = false;	//do not simulate tiny particles
 					MassProperties.Mass = 0.f;
 					MassProperties.InertiaTensor = FMatrix33(0,0,0);
@@ -2010,6 +2015,11 @@ void FGeometryCollectionPhysicsProxy::InitializeSharedCollisionStructures(
 		}
 	}
 
+	if (bSkippedSmallGeometry)
+	{
+		UE_LOG(LogChaos, Warning, TEXT("Some geometry is too small to be simulated and has been skipped."));
+	}
+
 	//User provides us with total mass or density.
 	//Density must be the same for individual parts and the total. Density_i = Density = Mass_i / Volume_i
 	//Total mass must equal sum of individual parts. Mass_i = TotalMass * Volume_i / TotalVolume => Density_i = TotalMass / TotalVolume
