CHaos Physics - Allow union of cluster to be attached to dynamic parents
- Added Parent() function to lcustred particle handle ( and adjusted code that was using ClusterIds().id instead
- In the geometry collection collection , make sure to check the internal clustered parent and update their kinematic target when necessary
- Fix typo issue ClustersToReuild => ClustersToRebuild

diff --git a/Engine/Source/Runtime/Experimental/Chaos/Private/Chaos/PBDRigidClustering.cpp b/Engine/Source/Runtime/Experimental/Chaos/Private/Chaos/PBDRigidClustering.cpp
--- a/Engine/Source/Runtime/Experimental/Chaos/Private/Chaos/PBDRigidClustering.cpp
+++ b/Engine/Source/Runtime/Experimental/Chaos/Private/Chaos/PBDRigidClustering.cpp
@@ -2045,7 +2045,7 @@ namespace Chaos
 		// disconnect from the parents
 		if (ClusteredParticle->ClusterIds().Id)
 		{
-			ParentParticle = ClusteredParticle->ClusterIds().Id->CastToClustered();
+			ParentParticle = ClusteredParticle->Parent();
 
 			ClusteredParticle->ClusterIds() = ClusterId();
 			ClusteredParticle->ClusterGroupIndex() = 0;
diff --git a/Engine/Source/Runtime/Experimental/Chaos/Private/PhysicsProxy/GeometryCollectionPhysicsProxy.cpp b/Engine/Source/Runtime/Experimental/Chaos/Private/PhysicsProxy/GeometryCollectionPhysicsProxy.cpp
--- a/Engine/Source/Runtime/Experimental/Chaos/Private/PhysicsProxy/GeometryCollectionPhysicsProxy.cpp
+++ b/Engine/Source/Runtime/Experimental/Chaos/Private/PhysicsProxy/GeometryCollectionPhysicsProxy.cpp
@@ -1341,7 +1341,7 @@ void FGeometryCollectionPhysicsProxy::OnRemoveFromSolver(Chaos::FPBDRigidsSolver
 {
 	Chaos::FPBDRigidsEvolutionGBF* Evolution = RBDSolver->GetEvolution();
 
-	TSet< FClusterHandle* > ClustersToReuild;
+	TSet< FClusterHandle* > ClustersToRebuild;
 	for (int i = 0; i < SolverParticleHandles.Num(); i++)
 	{
 		if (FClusterHandle* Handle = SolverParticleHandles[i])
@@ -1350,7 +1350,7 @@ void FGeometryCollectionPhysicsProxy::OnRemoveFromSolver(Chaos::FPBDRigidsSolver
 			{
 				if (ParentCluster->InternalCluster())
 				{
-					ClustersToReuild.Add(ParentCluster);
+					ClustersToRebuild.Add(ParentCluster);
 				}
 			}
 		}
@@ -1364,7 +1364,7 @@ void FGeometryCollectionPhysicsProxy::OnRemoveFromSolver(Chaos::FPBDRigidsSolver
 		}
 	}
 
-	for (FClusterHandle* Cluster : ClustersToReuild)
+	for (FClusterHandle* Cluster : ClustersToRebuild)
 	{
 		ensure(Cluster->InternalCluster());
 		if (ensure(Evolution->GetRigidClustering().GetChildrenMap().Contains(Cluster)))
@@ -1465,28 +1465,56 @@ void FGeometryCollectionPhysicsProxy::PushToPhysicsState()
 	{
 		const FTransform& ActorToWorld = Parameters.WorldTransform;
 
+		// used to avoid doing the work twice if we have a internalCluster parent 
+		bool InternalClusterParentUpdated = false;
+
 		int32 NumTransformGroupElements = PhysicsThreadCollection.NumElements(FGeometryCollection::TransformGroup);
 		for (int32 TransformGroupIndex = 0; TransformGroupIndex < NumTransformGroupElements; ++TransformGroupIndex)
 		{
 			Chaos::FPBDRigidClusteredParticleHandle* Handle = SolverParticleHandles[TransformGroupIndex];
-			if (!Handle || Handle->Disabled())
+			if (Handle)
 			{
-				continue;
-			}
-			const Chaos::EObjectStateType ObjectState = Handle->ObjectState();
-			if (ObjectState == Chaos::EObjectStateType::Kinematic)
-			{
-				FTransform WorldTransform = PhysicsThreadCollection.MassToLocal[TransformGroupIndex] * PhysicsThreadCollection.Transform[TransformGroupIndex] * ActorToWorld;
+				if (Handle->ObjectState() == Chaos::EObjectStateType::Kinematic)
+				{
+					// in the case of cluster union we need to find our Internal Cluster parent and update it
+					if (!InternalClusterParentUpdated)
+					{
+						FClusterHandle* ParentHandle = Handle->Parent();
+						if (ParentHandle && ParentHandle->InternalCluster() && !ParentHandle->Disabled() && ParentHandle->ObjectState() == Chaos::EObjectStateType::Kinematic)
+						{
+							FTransform NewChildWorldTransform = PhysicsThreadCollection.MassToLocal[TransformGroupIndex] * PhysicsThreadCollection.Transform[TransformGroupIndex] * ActorToWorld;
+							Chaos::FRigidTransform3 ParentToChildTransform = Handle->ChildToParent().Inverse();
+							FTransform NewParentWorldTRansform = ParentToChildTransform * NewChildWorldTransform;
+							SetClusteredParticleKinematicTarget(ParentHandle, NewParentWorldTRansform);
 
-				Chaos::TKinematicTarget<Chaos::FReal, 3> newKinematicTarget;
-				Chaos::FRigidTransform3 PreviousWorldTransform(Handle->X(), Handle->R());
-				newKinematicTarget.SetTargetMode(WorldTransform, PreviousWorldTransform);
-				Handle->SetKinematicTarget(newKinematicTarget);
+							InternalClusterParentUpdated = true;
+						}
+					}
+
+					if (!Handle->Disabled())
+					{
+						FTransform WorldTransform = PhysicsThreadCollection.MassToLocal[TransformGroupIndex] * PhysicsThreadCollection.Transform[TransformGroupIndex] * ActorToWorld;
+						SetClusteredParticleKinematicTarget(Handle, WorldTransform);
+					}
+				}
 			}
 		}
 	}
 }
 
+void FGeometryCollectionPhysicsProxy::SetClusteredParticleKinematicTarget(Chaos::FPBDRigidClusteredParticleHandle* Handle, const FTransform& NewWorldTransform)
+{
+	// CONTEXT: PHYSICSTHREAD
+	// this should be called only on teh physics thread
+	const Chaos::EObjectStateType ObjectState = Handle->ObjectState();
+	if (ensure(ObjectState == Chaos::EObjectStateType::Kinematic))
+	{
+		Chaos::TKinematicTarget<Chaos::FReal, 3> NewKinematicTarget;
+		Chaos::FRigidTransform3 CurrentWorldTransform(Handle->X(), Handle->R());
+		NewKinematicTarget.SetTargetMode(NewWorldTransform, CurrentWorldTransform);
+		Handle->SetKinematicTarget(NewKinematicTarget);
+	}
+}
 void FGeometryCollectionPhysicsProxy::BufferPhysicsResults(Chaos::FPBDRigidsSolver* CurrentSolver, Chaos::FDirtyGeometryCollectionData& BufferData)
 {
 	/**
diff --git a/Engine/Source/Runtime/Experimental/Chaos/Public/Chaos/ParticleHandle.h b/Engine/Source/Runtime/Experimental/Chaos/Public/Chaos/ParticleHandle.h
--- a/Engine/Source/Runtime/Experimental/Chaos/Public/Chaos/ParticleHandle.h
+++ b/Engine/Source/Runtime/Experimental/Chaos/Public/Chaos/ParticleHandle.h
@@ -965,6 +965,8 @@ public:
 	const ClusterId& ClusterIds() const { return PBDRigidClusteredParticles->ClusterIds(ParticleIdx); }
 	ClusterId& ClusterIds() { return PBDRigidClusteredParticles->ClusterIds(ParticleIdx); }
 
+	FPBDRigidClusteredParticleHandle* Parent() { return (ClusterIds().Id)? ClusterIds().Id->CastToClustered(): nullptr; }
+
 	const TRigidTransform<T,d>& ChildToParent() const { return PBDRigidClusteredParticles->ChildToParent(ParticleIdx); }
 	TRigidTransform<T,d>& ChildToParent() { return PBDRigidClusteredParticles->ChildToParent(ParticleIdx); }
 	void SetChildToParent(const TRigidTransform<T, d>& Xf) { PBDRigidClusteredParticles->ChildToParent(ParticleIdx) = Xf; }
diff --git a/Engine/Source/Runtime/Experimental/Chaos/Public/PhysicsProxy/GeometryCollectionPhysicsProxy.h b/Engine/Source/Runtime/Experimental/Chaos/Public/PhysicsProxy/GeometryCollectionPhysicsProxy.h
--- a/Engine/Source/Runtime/Experimental/Chaos/Public/PhysicsProxy/GeometryCollectionPhysicsProxy.h
+++ b/Engine/Source/Runtime/Experimental/Chaos/Public/PhysicsProxy/GeometryCollectionPhysicsProxy.h
@@ -264,6 +264,8 @@ protected:
 
 	void InitializeRemoveOnFracture(FParticlesType& Particles, const TManagedArray<int32>& DynamicState);
 
+	void SetClusteredParticleKinematicTarget(Chaos::FPBDRigidClusteredParticleHandle* Handle, const FTransform& WorldTransform);
+
 private:
 
 	FSimulationParameters Parameters;
diff --git a/Engine/Source/Runtime/Experimental/GeometryCollectionEngine/Private/GeometryCollection/GeometryCollectionComponent.cpp b/Engine/Source/Runtime/Experimental/GeometryCollectionEngine/Private/GeometryCollection/GeometryCollectionComponent.cpp
--- a/Engine/Source/Runtime/Experimental/GeometryCollectionEngine/Private/GeometryCollection/GeometryCollectionComponent.cpp
+++ b/Engine/Source/Runtime/Experimental/GeometryCollectionEngine/Private/GeometryCollection/GeometryCollectionComponent.cpp
@@ -949,7 +949,7 @@ void ActivateClusters(Chaos::FPBDRigidsEvolution::FRigidClustering& Clustering,
 
 	if(Cluster->ClusterIds().Id)
 	{
-		ActivateClusters(Clustering, Cluster->ClusterIds().Id->CastToClustered());
+		ActivateClusters(Clustering, Cluster->Parent());
 	}
 
 	Clustering.DeactivateClusterParticle(Cluster);
@@ -1002,7 +1002,7 @@ void UGeometryCollectionComponent::OnRep_RepData(const FGeometryCollectionRepDat
 				if(Particle->ClusterIds().Id)
 				{
 					// This particle is clustered but the remote authority has it activated. Fracture the parent cluster
-					ActivateClusters(Clustering, Particle->ClusterIds().Id->CastToClustered());
+					ActivateClusters(Clustering, Particle->Parent());
 				}
 				else if(Particle->Disabled())
 				{
