Chaos - fix sweep bounds for sweeps against scaled objects. Fixes per problem with sweeps against large tri meshes scaled down.

diff --git a/Engine/Source/Runtime/Experimental/Chaos/Public/Chaos/ImplicitObject.h b/Engine/Source/Runtime/Experimental/Chaos/Public/Chaos/ImplicitObject.h
--- a/Engine/Source/Runtime/Experimental/Chaos/Public/Chaos/ImplicitObject.h
+++ b/Engine/Source/Runtime/Experimental/Chaos/Public/Chaos/ImplicitObject.h
@@ -279,7 +279,7 @@ public:
 		// F_i = D . (N_i / S) / |(N_i / S)|
 		// Below we are effectively testing F_i = ((D / S) . N_i) and ignoring the |(N_i / S)| which is dependent on i and therefore cannot 
 		// be ignored in the minimum check.
-		const FReal UnscaledSearchDist = SearchDist * Scale.Max();	//this is not quite right since it's no longer a sphere, but the whole thing is fuzzy anyway
+		const FReal UnscaledSearchDist = SearchDist / Scale.Min();	//this is not quite right since it's no longer a sphere, but the whole thing is fuzzy anyway
 		const FVec3 UnscaledPosition = Position / Scale;
 		const FVec3 UnscaledNormal = GetInnerUnscaledNormal(UnitDir, Scale);
 		return FindMostOpposingFace(UnscaledPosition, UnscaledNormal, HintFaceIndex, UnscaledSearchDist);
