Fix fields metadata display name to match the documentation

diff --git a/Engine/Source/Runtime/Experimental/Chaos/Public/Field/FieldSystemTypes.h b/Engine/Source/Runtime/Experimental/Chaos/Public/Field/FieldSystemTypes.h
--- a/Engine/Source/Runtime/Experimental/Chaos/Public/Field/FieldSystemTypes.h
+++ b/Engine/Source/Runtime/Experimental/Chaos/Public/Field/FieldSystemTypes.h
@@ -83,7 +83,7 @@ enum EFieldFilterType
 	Field_Filter_Dynamic  UMETA(DisplayName = "Dynamic", ToolTip = "Apply the field to all the dynamic particles"),
 	Field_Filter_Kinematic  UMETA(DisplayName = "Kinematic", ToolTip = "Apply the field to all the kinematic particles"),
 	Field_Filter_Static  UMETA(DisplayName = "Static", ToolTip = "Apply the field to all the static particles"),
-	Field_Filter_All  UMETA(DisplayName = "Maximum", ToolTip = "Apply the field to all the solver particles"),
+	Field_Filter_All  UMETA(DisplayName = "All", ToolTip = "Apply the field to all the solver particles"),
 	//~~~
 	//256th entry
 	Field_Filter_Max      UMETA(Hidden)
@@ -120,7 +120,7 @@ enum EFieldPhysicsType
 	Field_AngularVelociy			UMETA(DisplayName = "Angular Velocity", ToolTip = "Add a vector field to the particles angular velocity."),
 	Field_AngularTorque				UMETA(DisplayName = "Angular Torque", ToolTip = "Add a vector field to the particles angular torque."),
 	Field_InternalClusterStrain		UMETA(DisplayName = "Internal Strain", ToolTip = "Add a strain field to the particles internal one."),
-	Field_DisableThreshold			UMETA(DisplayName = "Disable Threshold", ToolTip = "Disable the particles if their linear and angular velocity are less than the threshold."),
+	Field_DisableThreshold			UMETA(DisplayName = "Disabled Threshold", ToolTip = "Disable the particles if their linear and angular velocity are less than the threshold."),
 	Field_SleepingThreshold			UMETA(DisplayName = "Sleeping Threshold", ToolTip = "Set particles in sleeping mode if their linear and angular velocity are less than the threshold."),
 	Field_PositionStatic			UMETA(DisplayName = "Position Static", ToolTip = "Add a position constraint to the particles to remain static", Hidden),
 	Field_PositionAnimated			UMETA(DisplayName = "Position Animated", ToolTip = "Add a position constraint to the particles to follow its kinematic position", Hidden),
@@ -152,7 +152,7 @@ enum EFieldScalarType
 {
 	Scalar_ExternalClusterStrain		UMETA(DisplayName = "External Strain", ToolTip = "Apply an external strain over the particles. If this strain is over the internal one, the cluster will break."),
 	Scalar_Kill   						UMETA(DisplayName = "Kill Particle", ToolTip = "Disable the particles for which the field will be higher than 0."),
-	Scalar_DisableThreshold				UMETA(DisplayName = "Disable Threshold", ToolTip = "Disable the particles if their linear and angular velocity are less than the threshold."),
+	Scalar_DisableThreshold				UMETA(DisplayName = "Disabled Threshold", ToolTip = "Disable the particles if their linear and angular velocity are less than the threshold."),
 	Scalar_SleepingThreshold			UMETA(DisplayName = "Sleeping Threshold", ToolTip = "Set particles in sleeping mode if their linear and angular velocity are less than the threshold."),
 	Scalar_InternalClusterStrain		UMETA(DisplayName = "Internal Strain", ToolTip = "Add a strain field to the particles internal one."),
 	Scalar_DynamicConstraint			UMETA(DisplayName = "Dynamic Constraint", ToolTip = "Add the particles to a spring constraint holding them together", Hidden),
