Chaos
-- Update world bounds during ApplyKinematicTargets in PBD evolution, ensuring that spatial acceleration structures are accurate for traces and sweeps.

diff --git a/Engine/Source/Runtime/Experimental/Chaos/Public/Chaos/PBDRigidsEvolution.h b/Engine/Source/Runtime/Experimental/Chaos/Public/Chaos/PBDRigidsEvolution.h
--- a/Engine/Source/Runtime/Experimental/Chaos/Public/Chaos/PBDRigidsEvolution.h
+++ b/Engine/Source/Runtime/Experimental/Chaos/Public/Chaos/PBDRigidsEvolution.h
@@ -711,6 +711,19 @@ public:
 				Rigid->Q() = Rigid->R();
 				Rigid->PreV() = Rigid->V();
 				Rigid->PreW() = Rigid->W();
+
+				// Update the world bounds
+				if (Rigid->HasBounds())
+				{
+					const FAABB3& LocalBounds = Rigid->LocalBounds();
+					FAABB3 WorldSpaceBounds = LocalBounds.TransformedAABB(FRigidTransform3(Rigid->P(), Rigid->Q()));
+					if (Rigid->CCDEnabled())
+					{
+						WorldSpaceBounds.ThickenSymmetrically(Rigid->V() * Dt);
+					}
+					Rigid->SetWorldSpaceInflatedBounds(WorldSpaceBounds);
+				}
+				
 			}
 		}
 	}
