GeometryCollectionComponent
-- check for nullptr on FNaniteGeometryCollectionSceneProxy->GetPrimitiveSceneInfo

diff --git a/Engine/Source/Runtime/Experimental/GeometryCollectionEngine/Private/GeometryCollection/GeometryCollectionComponent.cpp b/Engine/Source/Runtime/Experimental/GeometryCollectionEngine/Private/GeometryCollection/GeometryCollectionComponent.cpp
--- a/Engine/Source/Runtime/Experimental/GeometryCollectionEngine/Private/GeometryCollection/GeometryCollectionComponent.cpp
+++ b/Engine/Source/Runtime/Experimental/GeometryCollectionEngine/Private/GeometryCollection/GeometryCollectionComponent.cpp
@@ -514,7 +514,10 @@ FPrimitiveSceneProxy* UGeometryCollectionComponent::CreateSceneProxy()
 						{
 							GeometryCollectionSceneProxy->SetConstantData_RenderThread(ConstantData);
 							GeometryCollectionSceneProxy->SetDynamicData_RenderThread(DynamicData);
-							GeometryCollectionSceneProxy->GetPrimitiveSceneInfo()->RequestGPUSceneUpdate();
+							if (GeometryCollectionSceneProxy->GetPrimitiveSceneInfo())
+							{
+								GeometryCollectionSceneProxy->GetPrimitiveSceneInfo()->RequestGPUSceneUpdate();
+							}
 						}
 					}
 				);
@@ -2995,3 +2998,4 @@ void UGeometryCollectionComponent::InitializeEmbeddedGeometry()
 
 	}
 }
+
