Chaos - Change main scene substepping to run external to the solver update loop, unifying the process with async mode for the dispatch of substeps so we end up with correctly interpolated input forces when substepping. We still support internal substepping for other use cases in the short term (RBAN) until such time they can also be unified.

Included unit test for edge case where we sleep an object somewhere in the update before the end of the batch of substeps to make sure we correctly collapse the result queue before updating the game thread

diff --git a/Engine/Source/Programs/HeadlessChaos/Private/HeadlessChaosTestDataMarshalling.cpp b/Engine/Source/Programs/HeadlessChaos/Private/HeadlessChaosTestDataMarshalling.cpp
--- a/Engine/Source/Programs/HeadlessChaos/Private/HeadlessChaosTestDataMarshalling.cpp
+++ b/Engine/Source/Programs/HeadlessChaos/Private/HeadlessChaosTestDataMarshalling.cpp
@@ -207,7 +207,7 @@ namespace ChaosTest
 				ExternalTime += ExternalDt;
 				MarshallingManager.FinalizePullData_Internal(Step, StartTime, ExternalDt);
 				//in sync mode the external time we pass in doesn't matter
-				FPullPhysicsData* Next = ResultsManager.PullSyncPhysicsResults_External(MarshallingManager);
+				FPullPhysicsData* Next = ResultsManager.PullSyncPhysicsResults_External(MarshallingManager, false);
 				EXPECT_EQ(Next->ExternalStartTime, StartTime);
 			}
 		}
diff --git a/Engine/Source/Programs/HeadlessChaos/Private/HeadlessChaosTestPhysicsScene.cpp b/Engine/Source/Programs/HeadlessChaos/Private/HeadlessChaosTestPhysicsScene.cpp
--- a/Engine/Source/Programs/HeadlessChaos/Private/HeadlessChaosTestPhysicsScene.cpp
+++ b/Engine/Source/Programs/HeadlessChaos/Private/HeadlessChaosTestPhysicsScene.cpp
@@ -576,7 +576,7 @@ namespace ChaosTest {
 				//tick solver
 				{
 					FVec3 Grav(0,0,-1);
-					Scene.SetUpForFrame(&Grav,1,99999,99999,10,false, -1);
+					Scene.SetUpForFrame(&Grav,1,99999,99999,1,false, -1);
 					Scene.StartFrame();
 					Scene.EndFrame();
 				}
@@ -598,7 +598,7 @@ namespace ChaosTest {
 			//tick solver one last time
 			{
 				FVec3 Grav(0,0,-1);
-				Scene.SetUpForFrame(&Grav,1,99999,99999,10,false, -1);
+				Scene.SetUpForFrame(&Grav,1,99999,99999,1,false, -1);
 				Scene.StartFrame();
 				Scene.EndFrame();
 			}
@@ -617,7 +617,7 @@ namespace ChaosTest {
 				//tick solver
 				{
 					FVec3 Grav(0,0,-1);
-					Scene.SetUpForFrame(&Grav,1,99999,99999,10,false, -1);
+					Scene.SetUpForFrame(&Grav,1,99999,99999,1,false, -1);
 					Scene.StartFrame();
 					Scene.EndFrame();
 				}
@@ -634,7 +634,7 @@ namespace ChaosTest {
 			//tick solver one last time
 			{
 				FVec3 Grav(0,0,-1);
-				Scene.SetUpForFrame(&Grav,1,99999,99999,10,false, -1);
+				Scene.SetUpForFrame(&Grav,1,99999,99999,1,false, -1);
 				Scene.StartFrame();
 				Scene.EndFrame();
 			}
@@ -664,7 +664,7 @@ namespace ChaosTest {
 			{
 				//tick solver
 				FVec3 Grav(0,0,-1);
-				Scene.SetUpForFrame(&Grav,1,99999,99999,10,false, -1);
+				Scene.SetUpForFrame(&Grav,1,99999,99999,1,false, -1);
 				Scene.StartFrame();
 				Scene.EndFrame();
 
diff --git a/Engine/Source/Programs/HeadlessChaos/Private/HeadlessChaosTestRewind.cpp b/Engine/Source/Programs/HeadlessChaos/Private/HeadlessChaosTestRewind.cpp
--- a/Engine/Source/Programs/HeadlessChaos/Private/HeadlessChaosTestRewind.cpp
+++ b/Engine/Source/Programs/HeadlessChaos/Private/HeadlessChaosTestRewind.cpp
@@ -63,7 +63,7 @@ namespace ChaosTest {
 					FChaosSolversModule* Module = FChaosSolversModule::GetModule();
 					FReal SimDt;
 					auto* Solver = CreateSolverHelper(DtMode, RewindHistorySize, Optimization, SimDt);
-					Solver->SetMaxDeltaTime(SimDt);	//make sure it can step even for huge steps
+					Solver->SetMaxDeltaTime_External(SimDt);	//make sure it can step even for huge steps
 
 					Lambda(Solver, SimDt, Optimization);
 
diff --git a/Engine/Source/Programs/HeadlessChaos/Private/HeadlessChaosTestSim.cpp b/Engine/Source/Programs/HeadlessChaos/Private/HeadlessChaosTestSim.cpp
--- a/Engine/Source/Programs/HeadlessChaos/Private/HeadlessChaosTestSim.cpp
+++ b/Engine/Source/Programs/HeadlessChaos/Private/HeadlessChaosTestSim.cpp
@@ -240,5 +240,62 @@ namespace ChaosTest {
 		EXPECT_TRUE(Dynamic1WentToSleep);
 		EXPECT_TRUE(Dynamic1HasWokeAgain);
 	}
+
+	TYPED_TEST(AllTraits, SimTests_MidSubstepSleep)
+	{
+		TSharedPtr<FImplicitObject, ESPMode::ThreadSafe> Sphere{new TSphere<FReal, 3>(FVec3(0), 10)};
+
+		FChaosSolversModule* Module = FChaosSolversModule::GetModule();
+		auto Solver = Module->CreateSolver<TypeParam>(nullptr);
+
+		FSingleParticlePhysicsProxy* Proxy = FSingleParticlePhysicsProxy::Create(Chaos::TPBDRigidParticle<FReal, 3>::CreateParticle());
+		Chaos::FRigidBodyHandle_External& Particle = Proxy->GetGameThreadAPI();
+		Particle.SetGeometry(Sphere);
+
+		Solver->RegisterObject(Proxy);
+
+		Solver->SetMaxSubSteps_External(4);
+		Solver->SetMaxDeltaTime_External(1.0f / 60.0f);
+		Solver->DisableAsyncMode();
+
+		const FVector InitialX = Proxy->GetGameThreadAPI().X();
+
+		struct FCallback : TSimCallbackObject<>
+		{
+			virtual void OnPreSimulate_Internal() override
+			{
+				check(Proxy);
+				PTX = Proxy->GetPhysicsThreadAPI()->X();
+				if(Hits == 1)
+				{
+					Proxy->GetPhysicsThreadAPI()->SetObjectState(EObjectStateType::Sleeping);
+				}
+				++Hits;
+			}
+
+			FSingleParticlePhysicsProxy* Proxy = nullptr;
+			FVector PTX = FVector::ZeroVector;
+			int32 Hits = 0;
+			asdasd accumulate dt and check it's 4 * 1/60
+		};
+
+		FCallback* Callback = Solver->CreateAndRegisterSimCallbackObject_External<FCallback>();
+		Callback->Proxy = Proxy;
+
+		// This should ensure 4 steps take place.
+		Solver->AdvanceAndDispatch_External(1.0f);
+		Solver->UpdateGameThreadStructures();
+
+		EXPECT_EQ(Callback->Hits, 4);
+
+		const FVector GTX = Proxy->GetGameThreadAPI().X();
+
+		EXPECT_NE(GTX, InitialX);
+		EXPECT_EQ(GTX, Callback->PTX);
+
+		Solver->UnregisterAndFreeSimCallbackObject_External(Callback);
+
+		Module->DestroySolver(Solver);
+	}
 }
 
diff --git a/Engine/Source/Runtime/Experimental/Chaos/Private/Chaos/ChaosMarshallingManager.cpp b/Engine/Source/Runtime/Experimental/Chaos/Private/Chaos/ChaosMarshallingManager.cpp
--- a/Engine/Source/Runtime/Experimental/Chaos/Private/Chaos/ChaosMarshallingManager.cpp
+++ b/Engine/Source/Runtime/Experimental/Chaos/Private/Chaos/ChaosMarshallingManager.cpp
@@ -57,7 +57,7 @@ void FChaosMarshallingManager::PrepareExternalQueue_External()
 	ProducerData->StartTime = ExternalTime_External;
 }
 
-void FChaosMarshallingManager::Step_External(FReal ExternalDT, const int32 NumSteps)
+void FChaosMarshallingManager::Step_External(FReal ExternalDT, const int32 NumSteps, bool bInSolverSubstepped)
 {
 	ensure(NumSteps > 0);
 
@@ -83,6 +83,7 @@ void FChaosMarshallingManager::Step_External(FReal ExternalDT, const int32 NumSt
 		ProducerData->InternalStep = InternalStep_External++;
 		ProducerData->IntervalStep = Step;
 		ProducerData->IntervalNumSteps = NumSteps;
+		ProducerData->bSolverSubstepped = bInSolverSubstepped;
 
 		ExternalQueue.Insert(ProducerData, 0);
 
diff --git a/Engine/Source/Runtime/Experimental/Chaos/Private/Chaos/Framework/ChaosResultsManager.cpp b/Engine/Source/Runtime/Experimental/Chaos/Private/Chaos/Framework/ChaosResultsManager.cpp
--- a/Engine/Source/Runtime/Experimental/Chaos/Private/Chaos/Framework/ChaosResultsManager.cpp
+++ b/Engine/Source/Runtime/Experimental/Chaos/Private/Chaos/Framework/ChaosResultsManager.cpp
@@ -6,48 +6,6 @@
 
 namespace Chaos
 {	
-	FPullPhysicsData* FChaosResultsManager::PullSyncPhysicsResults_External(FChaosMarshallingManager& MarshallingManager)
-	{
-		//sync mode doesn't use prev results, but if we were async previously we need to clean it up
-		if (Results.Prev)
-		{
-			MarshallingManager.FreePullData_External(Results.Prev);
-			Results.Prev = nullptr;
-		}
-
-		//if we switched from async to sync we may have multiple pending results, so discard them all except latest
-		if(bUsingSync == false)
-		{
-			FPullPhysicsData* PullDataLatest = nullptr;
-			do
-			{
-				//free any old results
-				if (Results.Next)
-				{
-					MarshallingManager.FreePullData_External(Results.Next);
-				}
-
-				Results.Next = PullDataLatest;
-				PullDataLatest = MarshallingManager.PopPullData_External();
-
-			} while (PullDataLatest && !bUsingSync);
-		}
-		else
-		{
-			//already in sync mode so don't take latest, just take next result
-			if (Results.Next)
-			{
-				MarshallingManager.FreePullData_External(Results.Next);
-			}
-
-			Results.Next = MarshallingManager.PopPullData_External();
-		}
-
-		bUsingSync = true;	//indicate that we used sync mode so don't rely on any cached results
-
-		return Results.Next;
-	}
-
 	enum class ESetPrevNextDataMode
 	{
 		Prev,
@@ -92,6 +50,84 @@ namespace Chaos
 		}
 	}
 
+	/** 
+	 Advance the results in the marshaller queue by one if it is available 
+
+	 @param MarshallingManager Manger to advance
+	 @param Results Results to update with advanced state
+	 @return whether an advance occurred
+	*/
+	bool AdvanceResult(FChaosMarshallingManager& MarshallingManager, FChaosInterpolationResults& Results)
+	{
+		if(FPullPhysicsData* PotentialNext = MarshallingManager.PopPullData_External())
+		{
+			//newer result exists so prev is overwritten
+			if(Results.Prev)
+			{
+				MarshallingManager.FreePullData_External(Results.Prev);
+			}
+
+			Results.Prev = Results.Next;
+			//mark prev with next's data.
+			//any particles that were dirty in the previous results and are now constant will still have the old values set
+			SetPrevNextDataHelper<ESetPrevNextDataMode::Prev>(*Results.Prev, Results.RigidInterpolations);
+
+			Results.Next = PotentialNext;
+			return true;
+		}
+		
+		return false;
+	}
+
+	/**
+	 Collapse the whole pending queue inside a marshalling manager to one results object written to Results.Next
+
+	 @param MarshallingManager Manger to advance
+	 @param Results Results to update with advanced state
+	*/
+	void CollapseResultsToLatest(FChaosMarshallingManager& MarshallingManager, FChaosInterpolationResults& Results)
+	{
+		if(Results.Next == nullptr)
+		{
+			//nothing in Next (first time), so get latest if possible
+			Results.Next = MarshallingManager.PopPullData_External();
+		}
+
+		while(AdvanceResult(MarshallingManager, Results))
+		{}
+	}
+
+	FPullPhysicsData* FChaosResultsManager::PullSyncPhysicsResults_External(FChaosMarshallingManager& MarshallingManager, bool bWasSubstepping)
+	{
+		//sync mode doesn't use prev results, but if we were async previously we need to clean it up
+		if (Results.Prev)
+		{
+			MarshallingManager.FreePullData_External(Results.Prev);
+			Results.Prev = nullptr;
+		}
+
+		// If we switched from async to sync we may have multiple pending results, so discard them all except latest.
+		// If we dispatched substeps there will be multiple results pending but the latest is the one we want.
+		if(bUsingSync == false || bWasSubstepping)
+		{
+			CollapseResultsToLatest(MarshallingManager, Results);
+		}
+		else
+		{
+			//already in sync mode so don't take latest, just take next result
+			if (Results.Next)
+			{
+				MarshallingManager.FreePullData_External(Results.Next);
+			}
+
+			Results.Next = MarshallingManager.PopPullData_External();
+		}
+
+		bUsingSync = true;	//indicate that we used sync mode so don't rely on any cached results
+
+		return Results.Next;
+	}
+
 	FReal ComputeAlphaHelper(const FPullPhysicsData& Next, const FReal ResultsTime)
 	{
 		const FReal Denom = Next.ExternalEndTime - Next.ExternalStartTime;
@@ -140,22 +176,7 @@ namespace Chaos
 			//go through every result and record the dirty proxies
 			while (Results.Next->ExternalEndTime < ResultsTime)
 			{
-				if (FPullPhysicsData* PotentialNext = MarshallingManager.PopPullData_External())
-				{
-					//newer result exists so prev is overwritten
-					if (Results.Prev)
-					{
-						MarshallingManager.FreePullData_External(Results.Prev);
-					}
-
-					Results.Prev = Results.Next;
-					//mark prev with next's data.
-					//any particles that were dirty in the previous results and are now constant will still have the old values set
-					SetPrevNextDataHelper<ESetPrevNextDataMode::Prev>(*Results.Prev, Results.RigidInterpolations);
-
-					Results.Next = PotentialNext;
-				}
-				else
+				if(!AdvanceResult(MarshallingManager, Results))
 				{
 					break;
 				}
diff --git a/Engine/Source/Runtime/Experimental/Chaos/Private/Chaos/Framework/PhysicsSolverBase.cpp b/Engine/Source/Runtime/Experimental/Chaos/Private/Chaos/Framework/PhysicsSolverBase.cpp
--- a/Engine/Source/Runtime/Experimental/Chaos/Private/Chaos/Framework/PhysicsSolverBase.cpp
+++ b/Engine/Source/Runtime/Experimental/Chaos/Private/Chaos/Framework/PhysicsSolverBase.cpp
@@ -102,7 +102,7 @@ namespace Chaos
 		// E.g., for 4 steps this will be: 1/4, 1/3, 1/2, 1
 		const FReal PseudoFraction = (FReal)1 / (FReal)(PushData->IntervalNumSteps - PushData->IntervalStep);
 		
-		Solver.AdvanceSolverBy(PushData->ExternalDt, FSubStepInfo{PseudoFraction, PushData->IntervalStep, PushData->IntervalNumSteps });
+		Solver.AdvanceSolverBy(PushData->ExternalDt, FSubStepInfo{PseudoFraction, PushData->IntervalStep, PushData->IntervalNumSteps, PushData->bSolverSubstepped});
 		Solver.GetMarshallingManager().FreeDataToHistory_Internal(PushData);	//cannot use push data after this point
 		PushData = nullptr;
 
@@ -136,8 +136,12 @@ namespace Chaos
 		, Owner(InOwner)
 		, ExternalDataLock_External(new FPhysicsSceneGuard())
 		, bIsShuttingDown(false)
+		, bSolverSubstep_External(false)
 		, AsyncDt(-1)
 		, AccumulatedTime(0)
+		, MMaxDeltaTime(0.0)
+		, MMinDeltaTime(SMALL_NUMBER)
+		, MMaxSubSteps(1)
 		, ExternalSteps(0)
 #if !UE_BUILD_SHIPPING
 		, bStealAdvanceTasksForTesting(false)
diff --git a/Engine/Source/Runtime/Experimental/Chaos/Private/PBDRigidsSolver.cpp b/Engine/Source/Runtime/Experimental/Chaos/Private/PBDRigidsSolver.cpp
--- a/Engine/Source/Runtime/Experimental/Chaos/Private/PBDRigidsSolver.cpp
+++ b/Engine/Source/Runtime/Experimental/Chaos/Private/PBDRigidsSolver.cpp
@@ -230,9 +230,9 @@ namespace Chaos
 				//
 				// * If we have used all of our substeps but still have time remaining, then some
 				//   energy will be lost.
-				const FReal MinDeltaTime = MSolver->GetMinDeltaTime();
-				const FReal MaxDeltaTime = MSolver->GetMaxDeltaTime();
-				int32 StepsRemaining = MSolver->GetMaxSubSteps();
+				const FReal MinDeltaTime = MSolver->GetMinDeltaTime_External();
+				const FReal MaxDeltaTime = MSolver->GetMaxDeltaTime_External();
+				int32 StepsRemaining = MSubStepInfo.bSolverSubstepped ? 1 : MSolver->GetMaxSubSteps_External();
 				FReal TimeRemaining = MDeltaTime;
 				bool bFirstStep = true;
 				while (StepsRemaining > 0 && TimeRemaining > MinDeltaTime)
@@ -341,9 +341,6 @@ namespace Chaos
 		, CurrentFrame(0)
 		, MTime(0.0)
 		, MLastDt(0.0)
-		, MMaxDeltaTime(0.0)
-		, MMinDeltaTime(SMALL_NUMBER)
-		, MMaxSubSteps(1)
 		, bHasFloor(true)
 		, bIsFloorAnalytic(false)
 		, FloorHeight(0.f)
@@ -616,9 +613,9 @@ namespace Chaos
 		MTime = 0;
 		MLastDt = 0.0f;
 		CurrentFrame = 0;
-		MMaxDeltaTime = 1.f;
-		MMinDeltaTime = SMALL_NUMBER;
-		MMaxSubSteps = 1;
+		SetMaxDeltaTime_External(1.0f);
+		SetMinDeltaTime_External(SMALL_NUMBER);
+		SetMaxSubSteps_External(1);
 		MEvolution = TUniquePtr<FPBDRigidsEvolution>(new FPBDRigidsEvolution(Particles, SimMaterials, &ContactModifiers, BufferMode == EMultiBufferMode::Single)); 
 
 		PerSolverField = MakeUnique<FPerSolverFieldSystem>();
@@ -814,7 +811,7 @@ namespace Chaos
 
 		GetEvolution()->GetBroadPhase().GetIgnoreCollisionManager().PushProducerStorageData_External(MarshallingManager.GetExternalTimestamp_External());
 
-		MarshallingManager.Step_External(DeltaTime, NumSteps);
+		MarshallingManager.Step_External(DeltaTime, NumSteps, GetSolverSubstep_External());
 	}
 
 	void FPBDRigidsSolver::ProcessSinglePushedData_Internal(FPushPhysicsData& PushData)
diff --git a/Engine/Source/Runtime/Experimental/Chaos/Public/Chaos/ChaosMarshallingManager.h b/Engine/Source/Runtime/Experimental/Chaos/Public/Chaos/ChaosMarshallingManager.h
--- a/Engine/Source/Runtime/Experimental/Chaos/Public/Chaos/ChaosMarshallingManager.h
+++ b/Engine/Source/Runtime/Experimental/Chaos/Public/Chaos/ChaosMarshallingManager.h
@@ -190,6 +190,7 @@ struct FPushPhysicsData
 	int32 InternalStep;		//The solver step this data will be associated with
 	int32 IntervalStep;		//The step we are currently at for this simulation interval. If not sub-stepping both step and num steps are 1: step is [0, IntervalNumSteps-1]
 	int32 IntervalNumSteps;	//The total number of steps associated with this simulation interval
+	bool bSolverSubstepped;
 
 	TArray<ISimCallbackObject*> SimCallbackObjectsToAdd;	//callback object registered at this specific time
 	TArray<ISimCallbackObject*> SimCallbackObjectsToRemove;	//callback object removed at this specific time
@@ -236,7 +237,7 @@ public:
 		GetProducerData_External()->SimCallbackInputs.Add(FSimCallbackInputAndObject{ SimCallbackObject, InputData });
 	}
 	/** Step forward using the external delta time. Should only be called by external thread */
-	void Step_External(FReal ExternalDT, const int32 NumSteps = 1);
+	void Step_External(FReal ExternalDT, const int32 NumSteps = 1, bool bSolverSubstepped = false);
 
 	/** Step the internal time forward if possible*/
 	FPushPhysicsData* StepInternalTime_External();
diff --git a/Engine/Source/Runtime/Experimental/Chaos/Public/Chaos/Framework/ChaosResultsManager.h b/Engine/Source/Runtime/Experimental/Chaos/Public/Chaos/Framework/ChaosResultsManager.h
--- a/Engine/Source/Runtime/Experimental/Chaos/Public/Chaos/Framework/ChaosResultsManager.h
+++ b/Engine/Source/Runtime/Experimental/Chaos/Public/Chaos/Framework/ChaosResultsManager.h
@@ -41,7 +41,7 @@ namespace Chaos
 		{
 		}
 
-		FPullPhysicsData* PullSyncPhysicsResults_External(FChaosMarshallingManager& MarshallingManager);
+		FPullPhysicsData* PullSyncPhysicsResults_External(FChaosMarshallingManager& MarshallingManager, bool bWasSubstepping);
 		const FChaosInterpolationResults& PullAsyncPhysicsResults_External(FChaosMarshallingManager& MarshallingManager, const FReal ResultsTime);
 
 	private:
diff --git a/Engine/Source/Runtime/Experimental/Chaos/Public/Chaos/Framework/PhysicsSolverBase.h b/Engine/Source/Runtime/Experimental/Chaos/Public/Chaos/Framework/PhysicsSolverBase.h
--- a/Engine/Source/Runtime/Experimental/Chaos/Public/Chaos/Framework/PhysicsSolverBase.h
+++ b/Engine/Source/Runtime/Experimental/Chaos/Public/Chaos/Framework/PhysicsSolverBase.h
@@ -36,6 +36,7 @@ namespace Chaos
 			: PseudoFraction(1.0)
 			, Step(1)
 			, NumSteps(1)
+			, bSolverSubstepped(false)
 		{
 		}
 
@@ -43,6 +44,16 @@ namespace Chaos
 			: PseudoFraction(InPseudoFraction)
 			, Step(InStep)
 			, NumSteps(InNumSteps)
+			, bSolverSubstepped(false)
+		{
+
+		}
+
+		FSubStepInfo(const FReal InPseudoFraction, const int32 InStep, const int32 InNumSteps, bool bInSolverSubstepped)
+			: PseudoFraction(InPseudoFraction)
+			, Step(InStep)
+			, NumSteps(InNumSteps)
+			, bSolverSubstepped(bInSolverSubstepped)
 		{
 
 		}
@@ -51,6 +62,7 @@ namespace Chaos
 		FReal PseudoFraction;
 		int32 Step;
 		int32 NumSteps;
+		bool bSolverSubstepped;
 	};
 
 	/**
@@ -333,6 +345,46 @@ namespace Chaos
 			return IsUsingAsyncResults() && UseAsyncInterpolation;
 		}
 
+		void SetMaxDeltaTime_External(float InMaxDeltaTime)
+		{
+			MMaxDeltaTime = InMaxDeltaTime;
+		}
+
+		void SetMinDeltaTime_External(float InMinDeltaTime)
+		{
+			MMinDeltaTime = InMinDeltaTime;
+		}
+
+		float GetMaxDeltaTime_External() const
+		{
+			return MMaxDeltaTime;
+		}
+
+		float GetMinDeltaTime_External() const
+		{
+			return MMinDeltaTime;
+		}
+
+		void SetMaxSubSteps_External(const int32 InMaxSubSteps)
+		{
+			MMaxSubSteps = InMaxSubSteps;
+		}
+
+		int32 GetMaxSubSteps_External() const
+		{
+			return MMaxSubSteps;
+		}
+
+		void SetSolverSubstep_External(bool bInSubstepExternal)
+		{
+			bSolverSubstep_External = bInSubstepExternal;
+		}
+
+		bool GetSolverSubstep_External() const
+		{
+			return bSolverSubstep_External;
+		}
+
 		/** Returns the time used by physics results. If fixed dt is used this will be the interpolated time */
 		FReal GetPhysicsResultsTime_External() const
 		{
@@ -432,8 +484,12 @@ namespace Chaos
 		friend struct TSolverSimMaterialScope;
 
 		bool bIsShuttingDown;
+		bool bSolverSubstep_External;
 		FReal AsyncDt;
 		FReal AccumulatedTime;
+		float MMaxDeltaTime;
+		float MMinDeltaTime;
+		int32 MMaxSubSteps;
 		int32 ExternalSteps;
 		TArray<FGeometryParticle*> UniqueIdxToGTParticles;
 
diff --git a/Engine/Source/Runtime/Experimental/Chaos/Public/Chaos/PhysicsSolverBaseImpl.h b/Engine/Source/Runtime/Experimental/Chaos/Public/Chaos/PhysicsSolverBaseImpl.h
--- a/Engine/Source/Runtime/Experimental/Chaos/Public/Chaos/PhysicsSolverBaseImpl.h
+++ b/Engine/Source/Runtime/Experimental/Chaos/Public/Chaos/PhysicsSolverBaseImpl.h
@@ -54,8 +54,11 @@ namespace Chaos
 		}
 		else
 		{
-			//no interpolation so just use latest
-			if (FPullPhysicsData* PullData = PullResultsManager->PullSyncPhysicsResults_External(MarshallingManager))
+			//no interpolation so just use latest, in non-substepping modes this will just be the next result
+			// available in the queue - however if we substepped externally we need to consume the whole
+			// queue by telling the sync pull that we expect multiple results.
+			const bool bSubstepping = bSolverSubstep_External && MMaxSubSteps > 1;
+			if (FPullPhysicsData* PullData = PullResultsManager->PullSyncPhysicsResults_External(MarshallingManager, bSubstepping))
 			{
 				LatestData = PullData;
 				const int32 SyncTimestamp = PullData->SolverTimestamp;
diff --git a/Engine/Source/Runtime/Experimental/Chaos/Public/PBDRigidsSolver.h b/Engine/Source/Runtime/Experimental/Chaos/Public/PBDRigidsSolver.h
--- a/Engine/Source/Runtime/Experimental/Chaos/Public/PBDRigidsSolver.h
+++ b/Engine/Source/Runtime/Experimental/Chaos/Public/PBDRigidsSolver.h
@@ -179,13 +179,8 @@ namespace Chaos
 		const FReal GetSolverTime() const { return MTime; }
 
 		/**/
-		void SetMaxDeltaTime(const FReal InMaxDeltaTime) { MMaxDeltaTime = InMaxDeltaTime; }
 		FReal GetLastDt() const { return MLastDt; }
-		FReal GetMaxDeltaTime() const { return MMaxDeltaTime; }
-		FReal GetMinDeltaTime() const { return MMinDeltaTime; }
-		void SetMaxSubSteps(const int32 InMaxSubSteps) { MMaxSubSteps = InMaxSubSteps; }
-		int32 GetMaxSubSteps() const { return MMaxSubSteps; }
-
+		
 		/**/
 		void SetIterations(const int32 InNumIterations) { GetEvolution()->SetNumIterations(InNumIterations); }
 		void SetPushOutIterations(const int32 InNumIterations) {  GetEvolution()->SetNumPushOutIterations(InNumIterations); }
@@ -311,9 +306,6 @@ namespace Chaos
 		int32 CurrentFrame;
 		FReal MTime;
 		FReal MLastDt;
-		FReal MMaxDeltaTime;
-		FReal MMinDeltaTime;
-		int32 MMaxSubSteps;
 		bool bHasFloor;
 		bool bIsFloorAnalytic;
 		FReal FloorHeight;
diff --git a/Engine/Source/Runtime/PhysicsCore/Private/ChaosScene.cpp b/Engine/Source/Runtime/PhysicsCore/Private/ChaosScene.cpp
--- a/Engine/Source/Runtime/PhysicsCore/Private/ChaosScene.cpp
+++ b/Engine/Source/Runtime/PhysicsCore/Private/ChaosScene.cpp
@@ -324,11 +324,20 @@ void FChaosScene::SetUpForFrame(const FVector* NewGrav,float InDeltaSeconds /*=
 #if WITH_CHAOS
 	using namespace Chaos;
 	SetGravity(*NewGrav);
-	MDeltaTime = InMaxPhysicsDeltaTime > 0.f ? FMath::Min(InDeltaSeconds,InMaxPhysicsDeltaTime) : InDeltaSeconds;
+
+	if(bSubstepping)
+	{
+		MDeltaTime = FMath::Min(InDeltaSeconds, InMaxSubsteps * InMaxSubstepDeltaTime);
+	}
+	else
+	{
+		MDeltaTime = InMaxPhysicsDeltaTime > 0.f ? FMath::Min(InDeltaSeconds, InMaxPhysicsDeltaTime) : InDeltaSeconds;
+	}
 
 	if(FPhysicsSolver* Solver = GetSolver())
 	{
-		if (InAsyncDt > 0)
+		const bool bIsAsync = InAsyncDt > 0;
+		if (bIsAsync)
 		{
 			Solver->EnableAsyncMode(InAsyncDt);
 		}
@@ -338,12 +347,13 @@ void FChaosScene::SetUpForFrame(const FVector* NewGrav,float InDeltaSeconds /*=
 		}
 		if(bSubstepping)
 		{
-			Solver->SetMaxDeltaTime(InMaxSubstepDeltaTime);
-			Solver->SetMaxSubSteps(InMaxSubsteps);
-		} else
+			Solver->SetMaxDeltaTime_External(InMaxSubstepDeltaTime);
+			Solver->SetMaxSubSteps_External(InMaxSubsteps);
+		} 
+		else
 		{
-			Solver->SetMaxDeltaTime(InMaxPhysicsDeltaTime);
-			Solver->SetMaxSubSteps(1);
+			Solver->SetMaxDeltaTime_External(InMaxPhysicsDeltaTime);
+			Solver->SetMaxSubSteps_External(1);
 		}
 	}
 #endif
