Required updates after merging substepping changes in CL#15984572 after solver dispatch was moved to another file

diff --git a/Engine/Source/Runtime/Experimental/Chaos/Private/Chaos/Framework/PhysicsSolverBase.cpp b/Engine/Source/Runtime/Experimental/Chaos/Private/Chaos/Framework/PhysicsSolverBase.cpp
--- a/Engine/Source/Runtime/Experimental/Chaos/Private/Chaos/Framework/PhysicsSolverBase.cpp
+++ b/Engine/Source/Runtime/Experimental/Chaos/Private/Chaos/Framework/PhysicsSolverBase.cpp
@@ -286,14 +286,16 @@ namespace Chaos
 
 	FGraphEventRef FPhysicsSolverBase::AdvanceAndDispatch_External(FReal InDt)
 	{
+		const bool bSubstepping = MMaxSubSteps > 1;
+		SetSolverSubstep_External(bSubstepping);
 		const FReal DtWithPause = bPaused_External ? 0.0f : InDt;
 		FReal InternalDt = DtWithPause;
 		int32 NumSteps = 1;
 
-		if (IsUsingFixedDt())
+		if(IsUsingFixedDt())
 		{
 			AccumulatedTime += DtWithPause;
-			if (InDt == 0)	//this is a special flush case
+			if(InDt == 0)	//this is a special flush case
 			{
 				//just use any remaining time and sync up to latest no matter what
 				InternalDt = AccumulatedTime;
@@ -302,18 +304,34 @@ namespace Chaos
 			}
 			else
 			{
+
 				InternalDt = AsyncDt;
 				NumSteps = FMath::FloorToInt(AccumulatedTime / InternalDt);
 				AccumulatedTime -= InternalDt * NumSteps;
 			}
 		}
+		else if(bSubstepping && InDt > 0)
+		{
+			NumSteps = FMath::CeilToInt(DtWithPause / MMaxDeltaTime);
+			if(NumSteps > MMaxSubSteps)
+			{
+				// Hitting this case means we're losing time, given the constraints of MaxSteps and MaxDt we can't
+				// fully handle the Dt requested, the simulation will appear to the viewer to run slower than realtime
+				NumSteps = MMaxSubSteps;
+				InternalDt = MMaxDeltaTime;
+			}
+			else
+			{
+				InternalDt = DtWithPause / NumSteps;
+			}
+		}
 
-		if (InDt > 0)
+		if(InDt > 0)
 		{
 			ExternalSteps++;	//we use this to average forces. It assumes external dt is about the same. 0 dt should be ignored as it typically has nothing to do with force
 		}
 
-		if (NumSteps > 0)
+		if(NumSteps > 0)
 		{
 			//make sure any GT state is pushed into necessary buffer
 			PushPhysicsState(InternalDt, NumSteps, FMath::Max(ExternalSteps, 1));
@@ -323,19 +341,19 @@ namespace Chaos
 		// Ensures we block on any tasks generated from previous frames
 		FGraphEventRef BlockingTasks = PendingTasks;
 
-		while (FPushPhysicsData* PushData = MarshallingManager.StepInternalTime_External())
+		while(FPushPhysicsData* PushData = MarshallingManager.StepInternalTime_External())
 		{
 			if(MRewindCallback && !bIsShuttingDown)
 			{
 				MRewindCallback->ProcessInputs_External(PushData->InternalStep, PushData->SimCallbackInputs);
 			}
 
-			if (ThreadingMode == EThreadingModeTemp::SingleThread)
+			if(ThreadingMode == EThreadingModeTemp::SingleThread)
 			{
 				ensure(!PendingTasks || PendingTasks->IsComplete());	//if mode changed we should have already blocked
 				FPhysicsSolverAdvanceTask ImmediateTask(*this, *PushData);
 #if !UE_BUILD_SHIPPING
-				if (bStealAdvanceTasksForTesting)
+				if(bStealAdvanceTasksForTesting)
 				{
 					StolenSolverAdvanceTasks.Emplace(MoveTemp(ImmediateTask));
 				}
@@ -350,27 +368,29 @@ namespace Chaos
 			else
 			{
 				// If enabled, block on all but most recent physics task, even tasks generated this frame.
-				if (AsyncPhysicsBlockMode == 1)
+				if(AsyncPhysicsBlockMode == 1)
 				{
 					BlockingTasks = PendingTasks;
 				}
 
 				FGraphEventArray Prereqs;
-				if (PendingTasks && !PendingTasks->IsComplete())
+				if(PendingTasks && !PendingTasks->IsComplete())
 				{
 					Prereqs.Add(PendingTasks);
 				}
 
 				PendingTasks = TGraphTask<FPhysicsSolverAdvanceTask>::CreateTask(&Prereqs).ConstructAndDispatchWhenReady(*this, *PushData);
-				if (IsUsingAsyncResults() == false)
+				if(IsUsingAsyncResults() == false)
 				{
 					BlockingTasks = PendingTasks;	//block right away
 				}
 			}
 
-			if (IsUsingAsyncResults() == false)
+			// This break is mainly here to satisfy unit testing. The call to StepInternalTime_External will decrement the
+			// delay in the marshaling manager and throw of tests that are explicitly testing for propagation delays
+			if(IsUsingAsyncResults() == false && !bSubstepping)
 			{
-				break;	//non async can only process one step at a time
+				break;
 			}
 		}
 
