Chaos : Geometry Collection
- Remove GUID support from the FManagedArrayCollection, its to slow and should be generated on-demand within editor operations.

diff --git a/Engine/Plugins/Experimental/GeometryCollectionPlugin/Source/GeometryCollectionEditor/Private/GeometryCollection/GeometryCollectionSelectRigidBodyEdMode.cpp b/Engine/Plugins/Experimental/GeometryCollectionPlugin/Source/GeometryCollectionEditor/Private/GeometryCollection/GeometryCollectionSelectRigidBodyEdMode.cpp
--- a/Engine/Plugins/Experimental/GeometryCollectionPlugin/Source/GeometryCollectionEditor/Private/GeometryCollection/GeometryCollectionSelectRigidBodyEdMode.cpp
+++ b/Engine/Plugins/Experimental/GeometryCollectionPlugin/Source/GeometryCollectionEditor/Private/GeometryCollection/GeometryCollectionSelectRigidBodyEdMode.cpp
@@ -155,7 +155,7 @@ bool FGeometryCollectionSelectRigidBodyEdMode::HandleClick(FEditorViewportClient
 				{
 					TransformIndex = ParentArray[TransformIndex];
 				}
-
+				/*
 				// Retrieve the rigid body id
 				//const TManagedArray<int32>& RigidBodyIdArray = GeometryCollectionComponent->GetRigidBodyIdArray();
 				const TManagedArray<FGuid>& RigidBodyIdArray = GeometryCollectionComponent->GetRigidBodyGuidArray();
@@ -177,6 +177,7 @@ bool FGeometryCollectionSelectRigidBodyEdMode::HandleClick(FEditorViewportClient
 						PropertyHandleSolverPin->SetValue(GeometryCollectionComponent->ChaosSolverActor);
 					}
 				}
+				*/
 			}
 		}
 	}
diff --git a/Engine/Source/Runtime/Experimental/Chaos/Private/GeometryCollection/ManagedArrayCollection.cpp b/Engine/Source/Runtime/Experimental/Chaos/Private/GeometryCollection/ManagedArrayCollection.cpp
--- a/Engine/Source/Runtime/Experimental/Chaos/Private/GeometryCollection/ManagedArrayCollection.cpp
+++ b/Engine/Source/Runtime/Experimental/Chaos/Private/GeometryCollection/ManagedArrayCollection.cpp
@@ -144,7 +144,6 @@ int32 FManagedArrayCollection::AddElements(int32 NumberElements, FName Group)
 			Entry.Value.Value->Resize(StartSize + NumberElements);
 		}
 	}
-	GenerateGuids(Group, StartSize);
 	GroupInfo[Group].Size += NumberElements;
 
 	SetDefaults(Group, StartSize, NumberElements);
@@ -204,7 +203,6 @@ void FManagedArrayCollection::Resize(int32 Size, FName Group)
 			Entry.Value.Value->Resize(Size);
 		}
 	}
-	GenerateGuids(Group, StartSize);
 	GroupInfo[Group].Size = Size;
 }
 
diff --git a/Engine/Source/Runtime/Experimental/GeometryCollectionEngine/Private/GeometryCollection/GeometryCollectionActor.cpp b/Engine/Source/Runtime/Experimental/GeometryCollectionEngine/Private/GeometryCollection/GeometryCollectionActor.cpp
--- a/Engine/Source/Runtime/Experimental/GeometryCollectionEngine/Private/GeometryCollection/GeometryCollectionActor.cpp
+++ b/Engine/Source/Runtime/Experimental/GeometryCollectionEngine/Private/GeometryCollection/GeometryCollectionActor.cpp
@@ -68,17 +68,14 @@ const Chaos::FPhysicsSolver* GetSolver(const AGeometryCollectionActor& GeomColle
 
 bool LowLevelRaycastImp(const Chaos::FVec3& Start, const Chaos::FVec3& Dir, float DeltaMag, const AGeometryCollectionActor& GeomCollectionActor, FHitResult& OutHit)
 {
+#if TODO_REIMPLEMENT_GET_RIGID_PARTICLES
 	using namespace Chaos;
-	//todo(ocohen): need to add thread safety / lock semantics
-	//const TManagedArray<int32>& RigidBodyIdArray = GeomCollectionActor.GetGeometryCollectionComponent()->GetRigidBodyIdArray();
-	const TManagedArray<FGuid>& RigidBodyIdArray = GeomCollectionActor.GetGeometryCollectionComponent()->GetRigidBodyGuidArray();
 	FPhysScene_Chaos* Scene = GeomCollectionActor.GetGeometryCollectionComponent()->GetInnerChaosScene();
 	ensure(Scene);
 
 	const Chaos::FPhysicsSolver* Solver = GetSolver(GeomCollectionActor);
 	if(ensure(Solver))
 	{
-#if TODO_REIMPLEMENT_GET_RIGID_PARTICLES
 		const TPBDRigidParticles<float, 3>& Particles = Solver->GetRigidParticles();	//todo(ocohen): should these just get passed in instead of hopping through scene?
 		FAABB3 RayBox(Start, Start);
 		RayBox.Thicken(Dir * DeltaMag);
@@ -112,8 +109,8 @@ bool LowLevelRaycastImp(const Chaos::FVec3& Start, const Chaos::FVec3& Dir, floa
 				return true;
 			}
 		}
-#endif
 	}
+#endif
 
 	return false;
 }
diff --git a/Engine/Source/Runtime/Experimental/GeometryCollectionEngine/Private/GeometryCollection/GeometryCollectionDebugDrawComponent.cpp b/Engine/Source/Runtime/Experimental/GeometryCollectionEngine/Private/GeometryCollection/GeometryCollectionDebugDrawComponent.cpp
--- a/Engine/Source/Runtime/Experimental/GeometryCollectionEngine/Private/GeometryCollection/GeometryCollectionDebugDrawComponent.cpp
+++ b/Engine/Source/Runtime/Experimental/GeometryCollectionEngine/Private/GeometryCollection/GeometryCollectionDebugDrawComponent.cpp
@@ -445,7 +445,7 @@ void UGeometryCollectionDebugDrawComponent::UpdateSelectedTransformIndex()
 		SelectedChaosSolver = GeometryCollectionDebugDrawActor->SelectedRigidBody.Solver;
 		return;
 	}
-
+	/*
 	// Check rigid body id sync
 	// Note that this test alone isn't enough to ensure that the rigid body ids are valid.
 	//const TManagedArray<int32>& RigidBodyIds = GeometryCollectionComponent->RigidBodyIds;
@@ -505,6 +505,7 @@ void UGeometryCollectionDebugDrawComponent::UpdateSelectedTransformIndex()
 	// Update selected rigid body index and solver
 	SelectedRigidBodyId = GeometryCollectionDebugDrawActor->SelectedRigidBody.Id;
 	SelectedChaosSolver = GeometryCollectionDebugDrawActor->SelectedRigidBody.Solver;
+	*/
 }
 
 int32 UGeometryCollectionDebugDrawComponent::CountFaces(int32 TransformIndex, bool bDebugDrawClustering) const
@@ -724,6 +725,7 @@ void UGeometryCollectionDebugDrawComponent::UpdateTickStatus()
 
 void UGeometryCollectionDebugDrawComponent::DebugDrawChaosTick()
 {
+	/*
 	check(GeometryCollectionComponent);
 	check(GeometryCollectionDebugDrawActor);
 	check(GeometryCollectionRenderLevelSetActor);
@@ -731,6 +733,7 @@ void UGeometryCollectionDebugDrawComponent::DebugDrawChaosTick()
 	AActor* const Actor = GetOwner();
 	check(Actor);
 
+	
 	// Retrieve synced particle and clustering data
 	//const TManagedArray<int32>& RigidBodyIds = GeometryCollectionComponent->RigidBodyIds;
 	const TManagedArray<FGuid>& RigidBodyIds = GeometryCollectionComponent->GetRigidBodyGuidArray();
@@ -879,6 +882,7 @@ void UGeometryCollectionDebugDrawComponent::DebugDrawChaosTick()
 			GeometryCollectionDebugDrawActor->DrawRigidBodiesId(GeometryCollectionComponent, ParticlesData, RigidBodyIds, GeometryCollectionDebugDrawActor->RigidBodyIdColor);
 		}
 	}
+	*/
 }
 
 void UGeometryCollectionDebugDrawComponent::UpdateLevelSetVisibility()
diff --git a/Engine/Source/Runtime/Experimental/GeometryCollectionEngine/Public/GeometryCollection/GeometryCollectionComponent.h b/Engine/Source/Runtime/Experimental/GeometryCollectionEngine/Public/GeometryCollection/GeometryCollectionComponent.h
--- a/Engine/Source/Runtime/Experimental/GeometryCollectionEngine/Public/GeometryCollection/GeometryCollectionComponent.h
+++ b/Engine/Source/Runtime/Experimental/GeometryCollectionEngine/Public/GeometryCollection/GeometryCollectionComponent.h
@@ -592,8 +592,6 @@ public:
 #endif  // #if !(UE_BUILD_SHIPPING || UE_BUILD_TEST)
 
 	/**/
-	//const TManagedArray<int32>& GetRigidBodyIdArray() const { return RigidBodyIds; }
-	const TManagedArray<FGuid>& GetRigidBodyGuidArray() const { return RestCollection->GetGeometryCollection()->GetAttribute<FGuid>(FName("GUID"), FGeometryCollection::TransformGroup); }
 	const TArray<bool>& GetDisabledFlags() const { return DisabledFlags; }
 
 	virtual void OnCreatePhysicsState() override;
