Fix jittering dynamic objects standing on a kinematic platform
- Fix shifted parameter passed to constructor
- Make sure kinematic particles do return previous velocities so that InitialContactVelocity is properly relative and can be compared to the threshold

diff --git a/Engine/Source/Runtime/Experimental/Chaos/Private/Chaos/PBDRigidsEvolutionGBF.cpp b/Engine/Source/Runtime/Experimental/Chaos/Private/Chaos/PBDRigidsEvolutionGBF.cpp
--- a/Engine/Source/Runtime/Experimental/Chaos/Private/Chaos/PBDRigidsEvolutionGBF.cpp
+++ b/Engine/Source/Runtime/Experimental/Chaos/Private/Chaos/PBDRigidsEvolutionGBF.cpp
@@ -406,7 +406,7 @@ void FPBDRigidsEvolutionGBF::AdvanceOneTimeStepImpl(const FReal Dt,const FSubSte
 FPBDRigidsEvolutionGBF::FPBDRigidsEvolutionGBF(FPBDRigidsSOAs& InParticles,THandleArray<FChaosPhysicsMaterial>& SolverPhysicsMaterials, const TArray<ISimCallbackObject*>* InCollisionModifiers, bool InIsSingleThreaded)
 	: Base(InParticles, SolverPhysicsMaterials, DefaultNumIterations, DefaultNumPushOutIterations, InIsSingleThreaded)
 	, Clustering(*this, Particles.GetClusteredParticles())
-	, CollisionConstraints(InParticles, Collided, PhysicsMaterials, PerParticlePhysicsMaterials, DefaultNumCollisionPairIterations, DefaultNumCollisionPushOutPairIterations, DefaultCollisionCullDistance)
+	, CollisionConstraints(InParticles, Collided, PhysicsMaterials, PerParticlePhysicsMaterials, DefaultNumCollisionPairIterations, DefaultNumCollisionPushOutPairIterations, DefaultRestitutionThreshold)
 	, CollisionRule(CollisionConstraints)
 	, BroadPhase(InParticles, DefaultCollisionCullDistance, BoundsThicknessVelocityMultiplier, DefaultCollisionCullDistance)
 	, NarrowPhase()
diff --git a/Engine/Source/Runtime/Experimental/Chaos/Public/Chaos/PBDRigidsEvolutionGBF.h b/Engine/Source/Runtime/Experimental/Chaos/Public/Chaos/PBDRigidsEvolutionGBF.h
--- a/Engine/Source/Runtime/Experimental/Chaos/Public/Chaos/PBDRigidsEvolutionGBF.h
+++ b/Engine/Source/Runtime/Experimental/Chaos/Public/Chaos/PBDRigidsEvolutionGBF.h
@@ -97,6 +97,7 @@ namespace Chaos
 		static constexpr FRealSingle DefaultCollisionCullDistance = 5.0f;
 		static constexpr int32 DefaultNumJointPairIterations = 3;
 		static constexpr int32 DefaultNumJointPushOutPairIterations = 0;
+		static constexpr int32 DefaultRestitutionThreshold = 1000;
 
 		// @todo(chaos): Required by clustering - clean up
 		using Base::ApplyPushOut;
diff --git a/Engine/Source/Runtime/Experimental/Chaos/Public/Chaos/ParticleHandle.h b/Engine/Source/Runtime/Experimental/Chaos/Public/Chaos/ParticleHandle.h
--- a/Engine/Source/Runtime/Experimental/Chaos/Public/Chaos/ParticleHandle.h
+++ b/Engine/Source/Runtime/Experimental/Chaos/Public/Chaos/ParticleHandle.h
@@ -1237,10 +1237,9 @@ public:
 		return false;
 	}
 
-	// @todo(ccaulfield): should be available on kinematics?
 	const FVec3& PreV() const
 	{
-		if (MHandle->CastToRigidParticle() && MHandle->ObjectState() == EObjectStateType::Dynamic)
+		if (MHandle->CastToRigidParticle())
 		{
 			return MHandle->CastToRigidParticle()->PreV();
 		}
@@ -1248,10 +1247,9 @@ public:
 		return ZeroVector;
 	}
 
-	// @todo(ccaulfield): should be available on kinematics?
 	const FVec3& PreW() const
 	{
-		if (MHandle->CastToRigidParticle() && MHandle->ObjectState() == EObjectStateType::Dynamic)
+		if (MHandle->CastToRigidParticle())
 		{
 			return MHandle->CastToRigidParticle()->PreW();
 		}
