Chaos : Geometry Collection

-- On demand Guids are generated when selection changes, when new geometry collections are created or reset, and when new nodes are introduced due to fracture/cluster operations.
-- Extended Validate tool tostrip unnecessary attributes from a geometry collection.
-- Pad out levelset grid by one voxel all around and report bounding box using the grids bounds, not the original mesh bounds. This avoids the issue of mass particles falling slightly outside the implicit domain.

diff --git a/Engine/Plugins/Experimental/ChaosEditor/Source/FractureEditor/Private/FractureEditorMode.cpp b/Engine/Plugins/Experimental/ChaosEditor/Source/FractureEditor/Private/FractureEditorMode.cpp
--- a/Engine/Plugins/Experimental/ChaosEditor/Source/FractureEditor/Private/FractureEditorMode.cpp
+++ b/Engine/Plugins/Experimental/ChaosEditor/Source/FractureEditor/Private/FractureEditorMode.cpp
@@ -321,6 +321,20 @@ void FFractureEditorMode::OnActorSelectionChanged(const TArray<UObject*>& NewSel
 		
 		for(UGeometryCollectionComponent* GeometryCollectionComponent : GeometryCollectionComponents)
 		{
+			// If collection does not already have guids, make them
+			FGeometryCollectionEdit RestCollectionEdit = GeometryCollectionComponent->EditRestCollection(GeometryCollection::EEditUpdate::None);
+			UGeometryCollection* RestCollection = RestCollectionEdit.GetRestCollection();
+			TSharedPtr<FGeometryCollection, ESPMode::ThreadSafe> GeometryCollection = RestCollection->GetGeometryCollection();
+			if (!GeometryCollection->HasAttribute("GUID", FGeometryCollection::TransformGroup))
+			{
+				FManagedArrayCollection::FConstructionParameters Params(FName(""), false);
+				TManagedArray<FGuid>& Guids = GeometryCollection->AddAttribute<FGuid>("GUID", FGeometryCollection::TransformGroup, Params);
+				for (int32 Idx = 0; Idx < Guids.Num(); ++Idx)
+				{
+					Guids[Idx] = FGuid::NewGuid();
+				}
+			}
+			
 			FScopedColorEdit ShowBoneColorsEdit(GeometryCollectionComponent);
 			ShowBoneColorsEdit.SetEnableBoneSelection(true);
 			// ShowBoneColorsEdit.SetLevelViewMode(ViewLevel);
diff --git a/Engine/Plugins/Experimental/ChaosEditor/Source/FractureEditor/Private/FractureToolAutoCluster.cpp b/Engine/Plugins/Experimental/ChaosEditor/Source/FractureEditor/Private/FractureToolAutoCluster.cpp
--- a/Engine/Plugins/Experimental/ChaosEditor/Source/FractureEditor/Private/FractureToolAutoCluster.cpp
+++ b/Engine/Plugins/Experimental/ChaosEditor/Source/FractureEditor/Private/FractureToolAutoCluster.cpp
@@ -73,6 +73,8 @@ void UFractureToolAutoCluster::Execute(TWeakPtr<FFractureEditorModeToolkit> InTo
 			Context.ConvertSelectionToClusterNodes();
 			FGeometryCollection* GeometryCollection = Context.GetGeometryCollection().Get();
 
+			int32 StartTransformCount = GeometryCollection->Transform.Num();
+
 			if (AutoClusterSettings->AutoClusterMode < EFractureAutoClusterMode::Voronoi)
 			{
 				UAutoClusterFractureCommand::ClusterChildBonesOfASingleMesh(Context.GetGeometryCollectionComponent(), AutoClusterSettings->AutoClusterMode, AutoClusterSettings->SiteCount);
@@ -112,6 +114,9 @@ void UFractureToolAutoCluster::Execute(TWeakPtr<FFractureEditorModeToolkit> InTo
 					FGeometryCollectionClusteringUtility::ValidateResults(GeometryCollection);
 				}
 			}
+
+			Context.GenerateGuids(StartTransformCount);
+
 			Refresh(Context, Toolkit);
 		}
 		SetOutlinerComponents(Contexts, Toolkit);
diff --git a/Engine/Plugins/Experimental/ChaosEditor/Source/FractureEditor/Private/FractureToolClusterMagnet.cpp b/Engine/Plugins/Experimental/ChaosEditor/Source/FractureEditor/Private/FractureToolClusterMagnet.cpp
--- a/Engine/Plugins/Experimental/ChaosEditor/Source/FractureEditor/Private/FractureToolClusterMagnet.cpp
+++ b/Engine/Plugins/Experimental/ChaosEditor/Source/FractureEditor/Private/FractureToolClusterMagnet.cpp
@@ -86,7 +86,8 @@ void UFractureToolClusterMagnet::Execute(TWeakPtr<FFractureEditorModeToolkit> In
 
 			const TManagedArray<TSet<int32>>& Children = Context.GetGeometryCollection()->Children;
 			const TManagedArray<int32>& Levels = Context.GetGeometryCollection()->GetAttribute<int32>("Level", FGeometryCollection::TransformGroup);
-
+			int32 StartTransformCount = Children.Num();
+			
 			Context.Sanitize();
 			TMap<int32, TArray<int32>> ClusteredSelection = Context.GetClusteredSelections();
 
@@ -135,6 +136,8 @@ void UFractureToolClusterMagnet::Execute(TWeakPtr<FFractureEditorModeToolkit> In
 				}
 			}
 
+			Context.GenerateGuids(StartTransformCount);
+
 			Refresh(Context, Toolkit);
 		}
 
diff --git a/Engine/Plugins/Experimental/ChaosEditor/Source/FractureEditor/Private/FractureToolClustering.cpp b/Engine/Plugins/Experimental/ChaosEditor/Source/FractureEditor/Private/FractureToolClustering.cpp
--- a/Engine/Plugins/Experimental/ChaosEditor/Source/FractureEditor/Private/FractureToolClustering.cpp
+++ b/Engine/Plugins/Experimental/ChaosEditor/Source/FractureEditor/Private/FractureToolClustering.cpp
@@ -100,6 +100,8 @@ void UFractureToolCluster::Execute(TWeakPtr<FFractureEditorModeToolkit> InToolki
 
 		for (FFractureToolContext& Context : Contexts)
 		{
+			int32 StartTransformCount = Context.GetGeometryCollection()->Transform.Num();
+			
 			Context.RemoveRootNodes();
 			Context.Sanitize();
 
@@ -119,6 +121,8 @@ void UFractureToolCluster::Execute(TWeakPtr<FFractureEditorModeToolkit> InToolki
 					}
 				}
 
+				Context.GenerateGuids(StartTransformCount);
+
 				Refresh(Context, Toolkit);
 			}
 		}
diff --git a/Engine/Plugins/Experimental/ChaosEditor/Source/FractureEditor/Private/FractureToolContext.cpp b/Engine/Plugins/Experimental/ChaosEditor/Source/FractureEditor/Private/FractureToolContext.cpp
--- a/Engine/Plugins/Experimental/ChaosEditor/Source/FractureEditor/Private/FractureToolContext.cpp
+++ b/Engine/Plugins/Experimental/ChaosEditor/Source/FractureEditor/Private/FractureToolContext.cpp
@@ -51,6 +51,21 @@ void FFractureToolContext::RemoveRootNodes()
 		});
 }
 
+void FFractureToolContext::GenerateGuids(int32 StartIdx)
+{
+	if (!GeometryCollection->HasAttribute("GUID", FGeometryCollection::TransformGroup))
+	{
+		FManagedArrayCollection::FConstructionParameters Params(FName(""), false);
+		GeometryCollection->AddAttribute<FGuid>("GUID", FGeometryCollection::TransformGroup, Params);
+	}
+
+	TManagedArray<FGuid>& Guids = GeometryCollection->GetAttribute<FGuid>("GUID", FGeometryCollection::TransformGroup);
+	for (int32 Idx = StartIdx; Idx < Guids.Num(); ++Idx)
+	{
+		Guids[Idx] = FGuid::NewGuid();
+	}
+}
+
 TMap<int32, TArray<int32>> FFractureToolContext::GetClusteredSelections()
 {
 	TMap<int32, TArray<int32>> SiblingGroups;
@@ -213,3 +228,5 @@ bool FFractureToolContext::HasSelectedAncestor(int32 Index) const
 	return false;
 }
 
+
+
diff --git a/Engine/Plugins/Experimental/ChaosEditor/Source/FractureEditor/Private/FractureToolContext.h b/Engine/Plugins/Experimental/ChaosEditor/Source/FractureEditor/Private/FractureToolContext.h
--- a/Engine/Plugins/Experimental/ChaosEditor/Source/FractureEditor/Private/FractureToolContext.h
+++ b/Engine/Plugins/Experimental/ChaosEditor/Source/FractureEditor/Private/FractureToolContext.h
@@ -44,6 +44,8 @@ public:
 
 	bool IsValid() const { return GeometryCollection.IsValid() && (SelectedBones.Num()>0); }
 
+	void GenerateGuids(int32 StartIdx);
+
 private:
 	bool HasSelectedAncestor(int32 Index) const;
 	bool IsValidBone(int32 Index) const
@@ -64,4 +66,5 @@ private:
 	FTransform Transform;
 	FBox Bounds;
 	int32 RandomSeed;
+
 };
\ No newline at end of file
diff --git a/Engine/Plugins/Experimental/ChaosEditor/Source/FractureEditor/Private/FractureToolEditing.cpp b/Engine/Plugins/Experimental/ChaosEditor/Source/FractureEditor/Private/FractureToolEditing.cpp
--- a/Engine/Plugins/Experimental/ChaosEditor/Source/FractureEditor/Private/FractureToolEditing.cpp
+++ b/Engine/Plugins/Experimental/ChaosEditor/Source/FractureEditor/Private/FractureToolEditing.cpp
@@ -176,6 +176,8 @@ void UFractureToolValidate::Execute(TWeakPtr<FFractureEditorModeToolkit> InToolk
 					// #todo leaving this out for the moment because we don't want to invalidate existing caches.
 					// FGeometryCollectionClusteringUtility::RemoveDanglingClusters(GeometryCollection);
 
+					bDirty |= StripUnnecessaryAttributes(GeometryCollection);
+
 					if (bDirty)
 					{
 						FGeometryCollectionClusteringUtility::UpdateHierarchyLevelOfChildren(GeometryCollection, -1);
@@ -199,6 +201,97 @@ void UFractureToolValidate::Execute(TWeakPtr<FFractureEditorModeToolkit> InToolk
 	}
 }
 
+bool UFractureToolValidate::StripUnnecessaryAttributes(FGeometryCollection* GeometryCollection)
+{	
+	static TMap<FName, TArray<FName>> Necessary = {
+		{ "Transform", 
+			{
+			"GUID",
+			"Transform",
+			"BoneColor",
+			"Parent",
+			"Children",
+			"TransformToGeometryIndex",
+			"SimulationType",
+			"StatusFlags",
+			"InitialDynamicState",
+			"SimulatableParticlesAttribute",
+			"InertiaTensor",
+			"Mass",
+			"ExemplarIndex",
+			"MassToLocal",
+			"DefaultMaterialIndex"
+			} 
+		},
+		
+		{ "Vertices",
+			{
+			"Vertex",
+			"UV",
+			"Color",
+			"TangentU",
+			"TangentV",
+			"Normal",
+			"BoneMap"
+			}
+		},
+		
+		{ "Faces",
+			{
+			"Indices",
+			"Visible",
+			"MaterialIndex",
+			"MaterialID"
+			}
+		},
+
+		{ "Geometry",
+			{
+			"TransformIndex",
+			"BoundingBox",
+			"InnerRadius",
+			"OuterRadius",
+			"VertexStart",
+			"VertexCount",
+			"FaceStart",
+			"FaceCount"
+			}
+		},
+
+		{ "Material",
+			{
+			"Sections"
+			}
+		}
+	};
+
+	bool bChangesMade = false;
+
+	const TArray<FName> GroupNames = GeometryCollection->GroupNames();
+	for (const FName Group : GroupNames)
+	{
+		if (Necessary.Contains(Group))
+		{
+			const TArray<FName> AttributeNames = GeometryCollection->AttributeNames(Group);
+			const TArray<FName>& NecessaryAttributes = Necessary[Group];
+			for (const FName AttributeName : AttributeNames)
+			{
+				if (!NecessaryAttributes.Contains(AttributeName))
+				{
+					GeometryCollection->RemoveAttribute(AttributeName, Group);
+					bChangesMade = true;
+				}
+			}
+		}
+		else
+		{
+			GeometryCollection->RemoveGroup(Group);
+			bChangesMade = true;
+		}
+	}
+
+	return bChangesMade;
+}
 
 #undef LOCTEXT_NAMESPACE
 
diff --git a/Engine/Plugins/Experimental/ChaosEditor/Source/FractureEditor/Private/FractureToolEditing.h b/Engine/Plugins/Experimental/ChaosEditor/Source/FractureEditor/Private/FractureToolEditing.h
--- a/Engine/Plugins/Experimental/ChaosEditor/Source/FractureEditor/Private/FractureToolEditing.h
+++ b/Engine/Plugins/Experimental/ChaosEditor/Source/FractureEditor/Private/FractureToolEditing.h
@@ -36,4 +36,7 @@ public:
 	virtual FSlateIcon GetToolIcon() const override;
 	virtual void RegisterUICommand(FFractureEditorCommands* BindingContext) override;
 	virtual void Execute(TWeakPtr<FFractureEditorModeToolkit> InToolkit) override;
+
+private:
+	bool StripUnnecessaryAttributes(FGeometryCollection* GeometryCollection);
 };
diff --git a/Engine/Plugins/Experimental/ChaosEditor/Source/FractureEditor/Private/FractureToolEmbed.cpp b/Engine/Plugins/Experimental/ChaosEditor/Source/FractureEditor/Private/FractureToolEmbed.cpp
--- a/Engine/Plugins/Experimental/ChaosEditor/Source/FractureEditor/Private/FractureToolEmbed.cpp
+++ b/Engine/Plugins/Experimental/ChaosEditor/Source/FractureEditor/Private/FractureToolEmbed.cpp
@@ -66,6 +66,8 @@ void UFractureToolAddEmbeddedGeometry::Execute(TWeakPtr<FFractureEditorModeToolk
 			FFractureToolContext::FGeometryCollectionPtr GeometryCollection = Context.GetGeometryCollection();
 			const TManagedArray<FTransform>& Transform = GeometryCollection->Transform;
 			const TManagedArray<int32>& Parent = GeometryCollection->Parent;
+
+			int32 StartTransformCount = Transform.Num();
 			
 			const FTransform TargetActorTransform(Context.GetGeometryCollectionComponent()->GetOwner()->GetTransform());
 
@@ -95,6 +97,8 @@ void UFractureToolAddEmbeddedGeometry::Execute(TWeakPtr<FFractureEditorModeToolk
 				}
 			}
 
+			Context.GenerateGuids(StartTransformCount);
+
 			Context.GetGeometryCollectionComponent()->InitializeEmbeddedGeometry();
 			Refresh(Context, Toolkit);
 
@@ -249,6 +253,10 @@ void UFractureToolAutoEmbedGeometry::Execute(TWeakPtr<FFractureEditorModeToolkit
 						FracturedGeometryCollection->EmbeddedGeometryExemplar[ExemplarIndex].InstanceCount++;
 					}
 
+					// Get a guid generated for the new instance
+					TManagedArray<FGuid>& Guids = ClosestGeometryCollection->GetAttribute<FGuid>("GUID", FGeometryCollection::TransformGroup);
+					Guids[Guids.Num()-1] = FGuid::NewGuid();
+
 					// #todo there might be a lot of these -- collect and put outside the loop.
 					ClosestComponent->InitializeEmbeddedGeometry();
 					Refresh(*ClosestContext, Toolkit);
diff --git a/Engine/Plugins/Experimental/ChaosEditor/Source/FractureEditor/Private/FractureToolGenerators.cpp b/Engine/Plugins/Experimental/ChaosEditor/Source/FractureEditor/Private/FractureToolGenerators.cpp
--- a/Engine/Plugins/Experimental/ChaosEditor/Source/FractureEditor/Private/FractureToolGenerators.cpp
+++ b/Engine/Plugins/Experimental/ChaosEditor/Source/FractureEditor/Private/FractureToolGenerators.cpp
@@ -256,6 +256,20 @@ AGeometryCollectionActor* UFractureToolGenerateAsset::ConvertActorsToGeometryCol
 		FracturedGeometryCollection->EnsureDataIsCooked(false /* init resources */);
 	}
 
+	// Add and initialize guids
+	TSharedPtr<FGeometryCollection, ESPMode::ThreadSafe> GeometryCollection = FracturedGeometryCollection->GetGeometryCollection();
+	if (!GeometryCollection->HasAttribute("GUID", FGeometryCollection::TransformGroup))
+	{
+		FManagedArrayCollection::FConstructionParameters Params(FName(""), false);
+		GeometryCollection->AddAttribute<FGuid>("GUID", FGeometryCollection::TransformGroup, Params);
+	}
+
+	TManagedArray<FGuid>& Guids = GeometryCollection->GetAttribute<FGuid>("GUID", FGeometryCollection::TransformGroup);
+	for (int32 Idx = 0; Idx < Guids.Num(); ++Idx)
+	{
+		Guids[Idx] = FGuid::NewGuid();
+	}
+
 	return NewActor;
 }
 
diff --git a/Engine/Plugins/Experimental/ChaosEditor/Source/FractureEditor/Public/FractureTool.cpp b/Engine/Plugins/Experimental/ChaosEditor/Source/FractureEditor/Public/FractureTool.cpp
--- a/Engine/Plugins/Experimental/ChaosEditor/Source/FractureEditor/Public/FractureTool.cpp
+++ b/Engine/Plugins/Experimental/ChaosEditor/Source/FractureEditor/Public/FractureTool.cpp
@@ -219,6 +219,8 @@ void UFractureModalTool::Execute(TWeakPtr<FFractureEditorModeToolkit> InToolkit)
 			
 			if (FirstNewGeometryIndex > INDEX_NONE)
 			{
+				FractureContext.GenerateGuids(FirstNewGeometryIndex);
+
 				// Based on the first new geometry index, generate a list of new transforms generated by the fracture.
 				const TManagedArray<int32>& TransformIndex = FractureContext.GetGeometryCollection()->GetAttribute<int32>("TransformIndex", FGeometryCollection::GeometryGroup);
 				int32 LastGeometryIndex = TransformIndex.Num();
diff --git a/Engine/Plugins/Experimental/ChaosNiagara/Source/ChaosNiagara/Private/NiagaraDataInterfaceChaosDestruction.cpp b/Engine/Plugins/Experimental/ChaosNiagara/Source/ChaosNiagara/Private/NiagaraDataInterfaceChaosDestruction.cpp
--- a/Engine/Plugins/Experimental/ChaosNiagara/Source/ChaosNiagara/Private/NiagaraDataInterfaceChaosDestruction.cpp
+++ b/Engine/Plugins/Experimental/ChaosNiagara/Source/ChaosNiagara/Private/NiagaraDataInterfaceChaosDestruction.cpp
@@ -1253,7 +1253,9 @@ void UNiagaraDataInterfaceChaosDestruction::HandleBreakingEvents(const Chaos::FB
 			CopyData.BoundingboxVolume = CopyData.BoundingBox.GetVolume();
 
 			UPhysicalMaterial* PhysicalMaterial = nullptr;
-			int32 MaterialID = DataIn.Particle->Geometry()->GetMaterialIndex(0);
+			int32 MaterialID = 0;
+			if(DataIn.Particle && DataIn.Particle->Geometry() )
+				MaterialID = DataIn.Particle->Geometry()->GetMaterialIndex(0);
 
 			UGeometryCollectionComponent* GeometryCollectionComponent = nullptr;
 			UMaterialInterface* Material = nullptr;
diff --git a/Engine/Source/Runtime/Experimental/Chaos/Private/GeometryCollection/ManagedArrayCollection.cpp b/Engine/Source/Runtime/Experimental/Chaos/Private/GeometryCollection/ManagedArrayCollection.cpp
--- a/Engine/Source/Runtime/Experimental/Chaos/Private/GeometryCollection/ManagedArrayCollection.cpp
+++ b/Engine/Source/Runtime/Experimental/Chaos/Private/GeometryCollection/ManagedArrayCollection.cpp
@@ -15,8 +15,6 @@ FManagedArrayCollection::FManagedArrayCollection()
 	Version = 5;
 }
 
-static const FName GuidName("GUID");
-
 void FManagedArrayCollection::AddGroup(FName Group)
 {
 	ensure(!GroupInfo.Contains(Group));
@@ -24,9 +22,6 @@ void FManagedArrayCollection::AddGroup(FName Group)
 		0
 	};
 	GroupInfo.Add(Group, info);
-
-	//Every group has to have a GUID attribute
-	AddAttribute<FGuid>(GuidName, Group);
 }
 
 void FManagedArrayCollection::RemoveElements(const FName& Group, const TArray<int32>& SortedDeletionList, FProcessingParameters Params)
@@ -112,22 +107,6 @@ int32 FManagedArrayCollection::NumElements(FName Group) const
 	return Num;
 }
 
-/** Should be called whenever new elements are added. Generates guids for new entries */
-void FManagedArrayCollection::GenerateGuids(FName Group, int32 StartIdx)
-{
-	TManagedArray<FGuid>& Guids = GetAttribute<FGuid>(GuidName, Group);
-
-	// we don't actually rely on this at the moment and generating the guids is very expensive.
-	// we don't need these at runtime in any case, so if we need in the editor later, make sure this is in-editor only
-	if (GIsEditor)
-	{
-		for (int32 Idx = StartIdx; Idx < Guids.Num(); ++Idx)
-		{
-			Guids[Idx] = FGuid::NewGuid();
-		}
-	}
-}
-
 int32 FManagedArrayCollection::AddElements(int32 NumberElements, FName Group)
 {
 	int32 StartSize = 0;
@@ -405,6 +384,7 @@ FString FManagedArrayCollection::ToString() const
 	return Buffer;
 }
 
+static const FName GuidName("GUID");
 
 void FManagedArrayCollection::Serialize(Chaos::FChaosArchive& Ar)
 {
@@ -455,15 +435,41 @@ void FManagedArrayCollection::Serialize(Chaos::FChaosArchive& Ar)
 				Pair.Value.Value->Resize(GroupSize);
 			}
 		}
-		if (Version < 4)
+
+		// Material indices are tied to the implicit collision geometry
+		// baking cached assets will strip out the collision geometry
+		// so we copy the indices to an attribute for back access on 
+		// cached only assets. 
+		if (HasGroup("Transform"))
 		{
-			//old content has no guids
-			for (TTuple<FName, FGroupInfo>& Pair : GroupInfo)
+			if (HasAttribute("Implicits", "Transform"))
 			{
-				GenerateGuids(Pair.Key, 0);
+				if (!HasAttribute("DefaultMaterialIndex", "Transform"))
+				{
+					TManagedArray<int32>& DefaultMaterialID = AddAttribute<int32>("DefaultMaterialIndex", "Transform");
+
+					typedef TSharedPtr<Chaos::FImplicitObject, ESPMode::ThreadSafe> FSharedImplicit;
+					TManagedArray<FSharedImplicit>& Implicits = GetAttribute<FSharedImplicit>("Implicits", "Transform");
+					for (int i = 0; i < NumElements("Transform"); i++)
+					{
+						if (Implicits[i])
+							DefaultMaterialID[i] = Implicits[i]->GetMaterialIndex(0);
+					}
+				}
 			}
 		}
+
 #endif
+
+		// strip out GUID Attributes
+		for (auto& GroupName : GroupNames())
+		{
+			if (HasAttribute(GuidName, GroupName))
+			{
+				RemoveAttribute(GuidName, GroupName);
+			}
+		}
+
 	}
 	else
 	{
diff --git a/Engine/Source/Runtime/Experimental/Chaos/Private/PhysicsProxy/GeometryCollectionPhysicsProxy.cpp b/Engine/Source/Runtime/Experimental/Chaos/Private/PhysicsProxy/GeometryCollectionPhysicsProxy.cpp
--- a/Engine/Source/Runtime/Experimental/Chaos/Private/PhysicsProxy/GeometryCollectionPhysicsProxy.cpp
+++ b/Engine/Source/Runtime/Experimental/Chaos/Private/PhysicsProxy/GeometryCollectionPhysicsProxy.cpp
@@ -58,6 +58,13 @@ FAutoConsoleVariableRef CVarGeometryCollectionCollideAll(
 	GeometryCollectionCollideAll,
 	TEXT("Bypass the collision matrix and make geometry collections collide against everything"));
 
+
+bool GeometryCollectionBypassPhysicsAttributes = false;
+FAutoConsoleVariableRef CVarGeometryCollectionBypassPhysicsAttributes(
+	TEXT("p.GeometryCollectionBypassPhysicsAttributes"),
+	GeometryCollectionBypassPhysicsAttributes,
+	TEXT("Bypass the construction of simulation properties when all bodies are simply cached. for playback."));
+
 DEFINE_LOG_CATEGORY_STATIC(UGCC_LOG, Error, All);
 
 //==============================================================================
@@ -1817,6 +1824,8 @@ void FGeometryCollectionPhysicsProxy::InitializeSharedCollisionStructures(
 	FGeometryCollection& RestCollection,
 	const FSharedSimulationParameters& SharedParams)
 {
+	if (GeometryCollectionBypassPhysicsAttributes)  return;
+
 	FString BaseErrorPrefix = ErrorReporter.GetPrefix();
 
 	// fracture tools can create an empty GC before appending new geometry
diff --git a/Engine/Source/Runtime/Experimental/Chaos/Public/Chaos/Levelset.h b/Engine/Source/Runtime/Experimental/Chaos/Public/Chaos/Levelset.h
--- a/Engine/Source/Runtime/Experimental/Chaos/Public/Chaos/Levelset.h
+++ b/Engine/Source/Runtime/Experimental/Chaos/Public/Chaos/Levelset.h
@@ -42,7 +42,7 @@ class CHAOS_API FLevelSet final : public FImplicitObject
 	virtual FReal PhiWithNormal(const FVec3& x, FVec3& Normal) const override;
 	FReal SignedDistance(const FVec3& x) const;
 
-	virtual const FAABB3 BoundingBox() const override { return MOriginalLocalBoundingBox; }
+	virtual const FAABB3 BoundingBox() const override { return MLocalBoundingBox; }
 
 	// Returns a const ref to the underlying phi grid
 	const TArrayND<FReal, 3>& GetPhiArray() const { return MPhi; }
diff --git a/Engine/Source/Runtime/Experimental/Chaos/Public/GeometryCollection/ManagedArrayCollection.h b/Engine/Source/Runtime/Experimental/Chaos/Public/GeometryCollection/ManagedArrayCollection.h
--- a/Engine/Source/Runtime/Experimental/Chaos/Public/GeometryCollection/ManagedArrayCollection.h
+++ b/Engine/Source/Runtime/Experimental/Chaos/Public/GeometryCollection/ManagedArrayCollection.h
@@ -446,8 +446,6 @@ private:
 
 	virtual void SetDefaults(FName Group, uint32 StartSize, uint32 NumElements) {};
 
-	void GenerateGuids(FName Group, int32 StartIdx);
-
 	TMap< FKeyType, FValueType> Map;	//data is owned by the map explicitly
 	TMap< FName, FGroupInfo> GroupInfo;
 	bool bDirty;
