ChaosVehicles: fix wheel spin anim, engine braking was locking wheel too aggressively

diff --git a/Engine/Plugins/Experimental/ChaosVehiclesPlugin/Source/ChaosVehicles/Private/ChaosWheeledVehicleMovementComponent.cpp b/Engine/Plugins/Experimental/ChaosVehiclesPlugin/Source/ChaosVehicles/Private/ChaosWheeledVehicleMovementComponent.cpp
--- a/Engine/Plugins/Experimental/ChaosVehiclesPlugin/Source/ChaosVehicles/Private/ChaosWheeledVehicleMovementComponent.cpp
+++ b/Engine/Plugins/Experimental/ChaosVehiclesPlugin/Source/ChaosVehicles/Private/ChaosWheeledVehicleMovementComponent.cpp
@@ -752,11 +752,11 @@ void UChaosWheeledVehicleSimulation::ApplyInput(const FControlInputs& ControlInp
 		if (PWheel.BrakeEnabled)
 		{
 			float BrakeForce = PWheel.MaxBrakeTorque * ModifiedInputs.BrakeInput;
-			PWheel.SetBrakeTorque(TorqueMToCm(BrakeForce + EngineBrakingForce));
+			PWheel.SetBrakeTorque(TorqueMToCm(BrakeForce + EngineBrakingForce), FMath::Abs(EngineBrakingForce) > FMath::Abs(BrakeForce));
 		}
 		else
 		{
-			PWheel.SetBrakeTorque(TorqueMToCm(EngineBraking));
+			PWheel.SetBrakeTorque(TorqueMToCm(EngineBraking), true);
 		}
 
 		if ((ModifiedInputs.HandbrakeInput && PWheel.HandbrakeEnabled) || ModifiedInputs.ParkingEnabled)
diff --git a/Engine/Source/Runtime/Experimental/ChaosVehicles/ChaosVehiclesCore/Private/WheelSystem.cpp b/Engine/Source/Runtime/Experimental/ChaosVehicles/ChaosVehiclesCore/Private/WheelSystem.cpp
--- a/Engine/Source/Runtime/Experimental/ChaosVehicles/ChaosVehiclesCore/Private/WheelSystem.cpp
+++ b/Engine/Source/Runtime/Experimental/ChaosVehicles/ChaosVehiclesCore/Private/WheelSystem.cpp
@@ -37,6 +37,7 @@ namespace Chaos
 		, SlipAngle(0.f)
 		, bInContact(false)
 		, WheelIndex(0)
+		, bEngineBraking(false)
 		, Spin(0.f)
 		, AvailableGrip(0.f)
 		, InputForces(FVector::ZeroVector)
@@ -73,7 +74,7 @@ namespace Chaos
 		// are we actually touching the ground
 		if (ForceIntoSurface > SMALL_NUMBER)
 		{
-			// ABS limiting brake force to match force from the grip avialable
+			// ABS limiting brake force to match force from the grip available
 			if (ABSEnabled && Braking && FMath::Abs(AppliedLinearBrakeForce) > AvailableGrip)
 			{
 				if ((Braking && ABSEnabled) || (!Braking && TractionControlEnabled))
@@ -156,7 +157,7 @@ namespace Chaos
 				float Clip = (AvailableGrip) / Length;
 				if (Clip < 1.0f)
 				{
-					if (Braking)
+					if (Braking && !bEngineBraking)
 					{
 						WheelLocked = true;
 					}
@@ -178,7 +179,7 @@ namespace Chaos
 
 		if (WheelLocked)
 		{
-			Omega = 0.0f;
+			Omega *= 0.9f; // velocity reduced quickly
 		}
 		else
 		{ 
diff --git a/Engine/Source/Runtime/Experimental/ChaosVehicles/ChaosVehiclesCore/Public/WheelSystem.h b/Engine/Source/Runtime/Experimental/ChaosVehicles/ChaosVehiclesCore/Public/WheelSystem.h
--- a/Engine/Source/Runtime/Experimental/ChaosVehicles/ChaosVehiclesCore/Public/WheelSystem.h
+++ b/Engine/Source/Runtime/Experimental/ChaosVehicles/ChaosVehiclesCore/Public/WheelSystem.h
@@ -176,9 +176,10 @@ public:
 	}
 
 	/** Set the braking torque - decelerating rotational force */
-	void SetBrakeTorque(float BrakeTorqueIn)
+	void SetBrakeTorque(float BrakeTorqueIn, bool bEngineBrakingIn = false)
 	{
 		BrakeTorque = BrakeTorqueIn;
+		bEngineBraking = bEngineBrakingIn;
 	}
 
 	/** Set the drive torque - accelerating rotational force */
@@ -418,6 +419,7 @@ public:
 	float SlipAngle;			// Angle between wheel forwards and velocity vector
 	bool bInContact;			// Is tire in contact with the ground or free in the air
 	uint32 WheelIndex;			// purely for debugging purposes
+	bool bEngineBraking;		// Is the braking force coming from the engine
 
 	public:
 	// debug for now
