ChaosVehicles: Edge case crash fix

diff --git a/Engine/Source/Programs/HeadlessChaos/Private/HeadlessChaosTestVehicles.cpp b/Engine/Source/Programs/HeadlessChaos/Private/HeadlessChaosTestVehicles.cpp
--- a/Engine/Source/Programs/HeadlessChaos/Private/HeadlessChaosTestVehicles.cpp
+++ b/Engine/Source/Programs/HeadlessChaos/Private/HeadlessChaosTestVehicles.cpp
@@ -51,6 +51,8 @@ namespace ChaosTest
 		EXPECT_NEAR(Graph.EvaluateY(50.f), 90.f, Tolerance);
 		EXPECT_NEAR(Graph.EvaluateY(55.f), 90.f, Tolerance);
 		EXPECT_NEAR(Graph.EvaluateY(60.f), 90.f, Tolerance);
+		EXPECT_NEAR(Graph.EvaluateY(1.0E21f), 90.f, Tolerance);
+		EXPECT_NEAR(Graph.EvaluateY(-1.0E21f), 0.f, Tolerance);
 	}
 
 	GTEST_TEST(AllTraits, VehicleTest_CalculateSlipAngle)
diff --git a/Engine/Source/Runtime/Experimental/ChaosVehicles/ChaosVehiclesCore/Private/VehicleUtility.cpp b/Engine/Source/Runtime/Experimental/ChaosVehicles/ChaosVehiclesCore/Private/VehicleUtility.cpp
--- a/Engine/Source/Runtime/Experimental/ChaosVehicles/ChaosVehiclesCore/Private/VehicleUtility.cpp
+++ b/Engine/Source/Runtime/Experimental/ChaosVehicles/ChaosVehiclesCore/Private/VehicleUtility.cpp
@@ -198,10 +198,14 @@ namespace Chaos
 			float Step = (RangeX) / (Graph.Num() - 1);
 			float Start = (InX - BoundsX.X) / Step;
 
-			if (Start < 0.f)
+			if (InX <= BoundsX.X || Start < 0.f)
 			{
 				return Graph[0].Y;
 			}
+			else if (InX >= BoundsX.Y)
+			{
+				return Graph[Graph.Num() - 1].Y;
+			}
 			
 			int StartIndex = (int)Start;
 
