Chaos - partial fix for ragdoll explosions
- Edge Collision collision improvement in CollisionOneShotManifolds.cpp. This change rolls back to a previous version of the edge-edge contact creation and adds a couple additional fixes to that.
  - better face selection from GJK results
  - use the minimum margin of the two shapes in contact (fixes issues with very large shapes)
  - add a margin clamp (not currently used, but canb be enabled with a cvar)
- fix culling radius check in CollisionResolution.cpp
- fix margin-reduced vertex calculations for malformed convexes that have vertices with less than 3 associated faces

diff --git a/Engine/Source/Runtime/Experimental/Chaos/Private/Chaos/CollisionOneShotManifolds.cpp b/Engine/Source/Runtime/Experimental/Chaos/Private/Chaos/CollisionOneShotManifolds.cpp
--- a/Engine/Source/Runtime/Experimental/Chaos/Private/Chaos/CollisionOneShotManifolds.cpp
+++ b/Engine/Source/Runtime/Experimental/Chaos/Private/Chaos/CollisionOneShotManifolds.cpp
@@ -6,6 +6,7 @@
 #include "Chaos/Collision/PBDCollisionConstraint.h"
 #include "Chaos/Convex.h"
 #include "Chaos/Defines.h"
+#include "Chaos/GJK.h"
 #include "Chaos/ImplicitObjectScaled.h"
 #include "Chaos/Transform.h"
 #include "Chaos/Utilities.h"
@@ -23,6 +24,14 @@ namespace Chaos
 	bool bChaos_Collision_Manifold_BoxAsConvex = true;
 	FAutoConsoleVariableRef CVarChaosCollisioConvexManifoldBoxAsConvex(TEXT("p.Chaos.Collision.Manifold.BoxAsConvex"), bChaos_Collision_Manifold_BoxAsConvex, TEXT(""));
 
+	FReal Chaos_Collision_Manifold_CullDistanceMarginMultiplier = 1.0f;
+	FAutoConsoleVariableRef CVarChaosCollisioConvexManifoldCullDistanceMarginMultiplier(TEXT("p.Chaos.Collision.Manifold.CullDistanceMarginMultiplier"), Chaos_Collision_Manifold_CullDistanceMarginMultiplier, TEXT(""));
+
+	bool bChaos_Collision_Manifold_UseMinMargin = true;
+	FAutoConsoleVariableRef CVarChaosCollisioConvexManifoldUseMinMargin(TEXT("p.Chaos.Collision.Manifold.UseMinMargin"), bChaos_Collision_Manifold_UseMinMargin, TEXT(""));
+
+	FReal Chaos_Collision_Manifold_MaxMargin = -1.0f;
+	FAutoConsoleVariableRef CVarChaosCollisioConvexManifoldMaxMargin(TEXT("p.Chaos.Collision.Manifold.MaxMargin"), Chaos_Collision_Manifold_MaxMargin, TEXT(""));
 
 	namespace Collisions
 	{
@@ -521,6 +530,93 @@ namespace Chaos
 			return VertexBuffer1;
 		}
 
+		// Use GJK to find the closest points (or shallowest penetrating points) on two convex shapes usingthe specified margin
+		// @todo(chaos): dedupe from GJKContactPoint in CollisionResolution.cpp
+		template <typename GeometryA, typename GeometryB>
+		FContactPoint GJKContactPointMargin(const GeometryA& A, const GeometryB& B, const FRigidTransform3& ATM, const FRigidTransform3& BToATM, FReal MarginA, FReal MarginB)
+		{
+			FContactPoint Contact;
+
+			FReal Penetration;
+			FVec3 ClosestA, ClosestBInA, Normal;
+			int32 ClosestVertexIndexA, ClosestVertexIndexB;
+
+			// Slightly increased epsilon to reduce error in normal for almost touching objects.
+			const FReal Epsilon = 3.e-3f;
+
+			const TGJKCoreShape<GeometryA> AWithMargin(A, MarginA);
+			const TGJKCoreShape<GeometryB> BWithMargin(B, MarginB);
+
+			if (GJKPenetration<true>(AWithMargin, BWithMargin, BToATM, Penetration, ClosestA, ClosestBInA, Normal, ClosestVertexIndexA, ClosestVertexIndexB, 0.0f, 0.0f, FVec3(1,0,0), Epsilon))
+			{
+				// GJK output is all in the local space of A. We need to transform the B-relative position and the normal in to B-space
+				Contact.ShapeMargins[0] = 0.0f;
+				Contact.ShapeMargins[1] = 0.0f;
+				Contact.ShapeContactPoints[0] = ClosestA;
+				Contact.ShapeContactPoints[1] = BToATM.InverseTransformPosition(ClosestBInA);
+				Contact.ShapeContactNormal = -BToATM.InverseTransformVector(Normal);
+				Contact.Location = ATM.TransformPosition(ClosestA);
+				Contact.Normal = -ATM.TransformVectorNoScale(Normal);
+				Contact.Phi = -Penetration;
+			}
+
+			return Contact;
+		}
+
+		// Select one of the planes on the convex to use as the contact plane, given an estimated contact posiiton and opposing 
+		// normal from GJK with margins (which gives the shapes rounded corners/edges).
+		template <typename ConvexImplicitType>
+		int32 SelectContactPlane(
+			const ConvexImplicitType& Convex,
+			const FVec3 X,
+			const FVec3 N,
+			const FReal MaxDistance)
+		{
+			int32 BestPlaneIndex = INDEX_NONE;
+			FReal BestPlaneDot = 1.0f;
+
+			// Finds the most opposing plane (to N) that is within the specified distance (MaxDistance) of the position (X)
+			const int32 NumPlanes = Convex.NumPlanes();
+			for (int32 PlaneIndex = 0; PlaneIndex < NumPlanes; ++PlaneIndex)
+			{
+				const TPlaneConcrete<FReal, 3> Plane = Convex.GetPlane(PlaneIndex);
+				const FReal PlaneNormalDotN = FVec3::DotProduct(N, Plane.Normal());
+
+				// Ignore planes that do not oppose N
+				if (PlaneNormalDotN > -SMALL_NUMBER)
+				{
+					continue;
+				}
+
+				// Reject planes farther than MaxDistance
+				const FReal PlaneDistance = Plane.SignedDistance(X);
+				if (FMath::Abs(PlaneDistance) > MaxDistance)
+				{
+					continue;
+				}
+
+				// Keep the most opposing plane
+				if (PlaneNormalDotN < BestPlaneDot)
+				{
+					BestPlaneDot = PlaneNormalDotN;
+					BestPlaneIndex = PlaneIndex;
+				}
+			}
+
+			// Malformed convexes or half-spaces could have all planes rejected above.
+			// If that happens, select the most opposing plane including those that
+			// may point the same direction as N. 
+			if (BestPlaneIndex == INDEX_NONE)
+			{
+				// This always returns a valid plane.
+				BestPlaneIndex = Convex.GetMostOpposingPlane(N);
+			}
+
+			check(BestPlaneIndex != INDEX_NONE);
+			return BestPlaneIndex;
+		}
+
+
 		template <typename ConvexImplicitType1, typename ConvexImplicitType2>
 		void ConstructConvexConvexOneShotManifold(
 			const ConvexImplicitType1& Convex1,
@@ -533,29 +629,52 @@ namespace Chaos
 			const uint32 SpaceDimension = 3;
 
 			// We only build one shot manifolds once
-			// All convexes are pre-scaled
+			// All convexes are pre-scaled, or wrapped in TImplicitObjectScaled
 			ensure(Constraint.GetManifoldPoints().Num() == 0);
 			ensure(Convex1Transform.GetScale3D() == FVec3(1.0f, 1.0f, 1.0f));
 			ensure(Convex2Transform.GetScale3D() == FVec3(1.0f, 1.0f, 1.0f));
 
-			// Find the deepest penetration. This is used to determine the planes and points to use for the manifold
-			FContactPoint GJKContactPoint = GenericConvexConvexContactPoint(Convex1, Convex1Transform, Convex2, Convex2Transform, Constraint.Manifold.RestitutionPadding);
+			const FRigidTransform3 Convex2ToConvex1Transform = Convex2Transform.GetRelativeTransformNoScale(Convex1Transform);
 
-			// @todo(chaos): this is not a good culling test - make it better. The separation from GJK can be larger than the real separation because of the margins (rounded corners)
-			if (GJKContactPoint.Phi > Constraint.GetCullDistance())
+			// Get the adjusted margins for each convex
+			FReal Margin1 = Convex1.GetMargin();
+			FReal Margin2 = Convex2.GetMargin();
+			if (bChaos_Collision_Manifold_UseMinMargin)
+			{
+				// Use the smaller of the two margins for both shapes
+				Margin1 = FMath::Min(Margin1, Margin2);
+				Margin2 = Margin1;
+			}
+			if (Chaos_Collision_Manifold_MaxMargin >= 0.0f)
+			{
+				// Clamp the margin
+				Margin1 = FMath::Min(Margin1, Chaos_Collision_Manifold_MaxMargin);
+				Margin2 = FMath::Min(Margin2, Chaos_Collision_Manifold_MaxMargin);
+			}
+
+			// Find the deepest penetration. This is used to determine the planes and points to use for the manifold
+			FContactPoint GJKContactPoint = GJKContactPointMargin(Convex1, Convex2, Convex1Transform, Convex2ToConvex1Transform, Margin1, Margin2);
+
+			// GJK is using margins and rounded corner, so if we have a corner-to-corner contact it will under-report the actual distance by an amount that depends on how
+			// "pointy" the edge/corner is - it can be arbitrarily large. For now just expand the cull distance by some multiple of the margin (it would need to be
+			// (sqrt(3) - 1)*margin for a box, but can be much larger for acute angles.
+			// @todo(chaos): return the correct margin error with the GJK result, or do something so we don't need this at all. SAT?
+			const FReal GJKCullDistanceMarginMultiplier = Chaos_Collision_Manifold_CullDistanceMarginMultiplier;
+			const FReal GJKCullDistance = Constraint.GetCullDistance() + GJKCullDistanceMarginMultiplier * (Margin1 + Margin2);
+			if (GJKContactPoint.Phi > GJKCullDistance)
 			{
 				return;
 			}
 
 			// @todo(chaos): get the vertex index from GJK and use to to get the plane
 			const FVec3 SeparationDirectionLocalConvex1 = Convex1Transform.InverseTransformVectorNoScale(GJKContactPoint.Normal);
-			const int32 MostOpposingPlaneIndexConvex1 = Convex1.GetMostOpposingPlane(SeparationDirectionLocalConvex1);
+			const int32 MostOpposingPlaneIndexConvex1 = SelectContactPlane(Convex1, GJKContactPoint.ShapeContactPoints[0], SeparationDirectionLocalConvex1, Margin1);
 			const TPlaneConcrete<FReal, 3> BestPlaneConvex1 = Convex1.GetPlane(MostOpposingPlaneIndexConvex1);
 			const FReal BestPlaneDotNormalConvex1 = FVec3::DotProduct(-SeparationDirectionLocalConvex1, BestPlaneConvex1.Normal());
 
 			// Now for Convex2
 			const FVec3 SeparationDirectionLocalConvex2 = Convex2Transform.InverseTransformVectorNoScale(GJKContactPoint.Normal);
-			const int32 MostOpposingPlaneIndexConvex2 = Convex2.GetMostOpposingPlane(-SeparationDirectionLocalConvex2);
+			const int32 MostOpposingPlaneIndexConvex2 = SelectContactPlane(Convex2, GJKContactPoint.ShapeContactPoints[1], -SeparationDirectionLocalConvex2, Margin2);
 			const TPlaneConcrete<FReal, 3> BestPlaneConvex2 = Convex2.GetPlane(MostOpposingPlaneIndexConvex2);
 			const FReal BestPlaneDotNormalConvex2 = FVec3::DotProduct(SeparationDirectionLocalConvex2, BestPlaneConvex2.Normal());
 
@@ -570,56 +689,17 @@ namespace Chaos
 			const FReal PlaneContactNormalEpsilon = Chaos_Collision_Manifold_PlaneContactNormalEpsilon;
 			const bool bIsPlaneContact = FMath::IsNearlyEqual(BestPlaneDotNormalConvex1, 1.0f, PlaneContactNormalEpsilon) || FMath::IsNearlyEqual(BestPlaneDotNormalConvex2, 1.0f, PlaneContactNormalEpsilon);
 
-			// For edge-edge contacts, we find the edges involved and generate contact from the closest points on the edges
+			// For edge-edge contacts, we find the edges involved and project the contact onto the edges
 			if (!bIsPlaneContact)
 			{
-				// Get the edge vertices on each convex closest to the GJK contact
-				int32 EdgeVertexIndex1A, EdgeVertexIndex1B;
-				int32 EdgeVertexIndex2A, EdgeVertexIndex2B;
-				bool bFoundEdge1 = Convex1.GetClosestEdgeVertices(MostOpposingPlaneIndexConvex1, GJKContactPoint.ShapeContactPoints[0], EdgeVertexIndex1A, EdgeVertexIndex1B);
-				bool bFoundEdge2 = Convex2.GetClosestEdgeVertices(MostOpposingPlaneIndexConvex2, GJKContactPoint.ShapeContactPoints[1], EdgeVertexIndex2A, EdgeVertexIndex2B);
-				
-				// Sanity check - should only happen with null convexes
-				if (!bFoundEdge1 || !bFoundEdge2)
-				{
-					return;
-				}
-
-				// Get the world-space edge vertices and find tge closest point
-				const FVec3 EdgeVertex1A = Convex1Transform.TransformPosition(Convex1.GetVertex(EdgeVertexIndex1A));
-				const FVec3 EdgeVertex1B = Convex1Transform.TransformPosition(Convex1.GetVertex(EdgeVertexIndex1B));
-				const FVec3 EdgeVertex2A = Convex2Transform.TransformPosition(Convex2.GetVertex(EdgeVertexIndex2A));
-				const FVec3 EdgeVertex2B = Convex2Transform.TransformPosition(Convex2.GetVertex(EdgeVertexIndex2B));
-				FReal S, T;
-				FVec3 EdgePos1, EdgePos2;
-				Utilities::NearestPointsOnLineSegments(EdgeVertex1A, EdgeVertex1B, EdgeVertex2A, EdgeVertex2B, S, T, EdgePos1, EdgePos2);
-
-				// Generate the normal from the edges (make sure we don't change the sign)
-				FVec3 EdgeNormal = FVec3::CrossProduct(EdgeVertex1B - EdgeVertex1A, EdgeVertex2B - EdgeVertex2A);
-				const FReal EdgeNormalSq = EdgeNormal.SizeSquared();
-				if (EdgeNormalSq > SMALL_NUMBER)
-				{
-					EdgeNormal = EdgeNormal * FMath::InvSqrt(EdgeNormalSq);
-					if (FVec3::DotProduct(EdgeNormal, GJKContactPoint.Normal) < 0.0f)
-					{
-						EdgeNormal = -EdgeNormal;
-					}
-				}
-				else
-				{
-					EdgeNormal = GJKContactPoint.Normal;
-				}
-
-				// Shape-space contact data
-				const FReal EdgePhi = FVec3::DotProduct(EdgePos1 - EdgePos2, EdgeNormal);
-				const FVec3 ShapeEdgePos1 = Convex1Transform.InverseTransformPositionNoScale(EdgePos1);
-				const FVec3 ShapeEdgePos2 = Convex2Transform.InverseTransformPositionNoScale(EdgePos2);
-				const FVec3 ShapeEdgeNormal2 = Convex2Transform.InverseTransformVectorNoScale(EdgeNormal);
+				FVec3 ShapeEdgePos1 = Convex1.GetClosestEdgePosition(MostOpposingPlaneIndexConvex1, GJKContactPoint.ShapeContactPoints[0]);
+				FVec3 ShapeEdgePos2 = Convex2.GetClosestEdgePosition(MostOpposingPlaneIndexConvex2, GJKContactPoint.ShapeContactPoints[1]);
+				FVec3 EdgePos1 = Convex1Transform.TransformPosition(ShapeEdgePos1);
+				FVec3 EdgePos2 = Convex2Transform.TransformPosition(ShapeEdgePos2);
+				FReal EdgePhi = FVec3::DotProduct(EdgePos1 - EdgePos2, GJKContactPoint.Normal);
 
-				// Update the contact with the edge-edge contact
 				GJKContactPoint.ShapeContactPoints[0] = ShapeEdgePos1;
 				GJKContactPoint.ShapeContactPoints[1] = ShapeEdgePos2;
-				GJKContactPoint.ShapeContactNormal = ShapeEdgeNormal2;
 				GJKContactPoint.Phi = EdgePhi;
 				GJKContactPoint.Location = 0.5f * (EdgePos1 + EdgePos2);
 
@@ -708,7 +788,7 @@ namespace Chaos
 				ContactPoint.ShapeMargins[1] = 0.0f;
 				ContactPoint.ShapeContactPoints[0] = ReferenceFaceConvex1 ? PointProjectedOntoReferenceFace : ClippedPointInOtherCoordinates;
 				ContactPoint.ShapeContactPoints[1] = ReferenceFaceConvex1 ? ClippedPointInOtherCoordinates : PointProjectedOntoReferenceFace;
-				ContactPoint.ShapeContactNormal = ReferenceFaceConvex1 ? -RefPlaneNormal : RefPlaneNormal;
+				ContactPoint.ShapeContactNormal = RefSeparationDirection;
 				ContactPoint.ContactNormalOwnerIndex = ReferenceFaceConvex1 ? 0 : 1;
 
 				ContactPoint.Location = RefConvexTM->TransformPositionNoScale(PointProjectedOntoReferenceFace);
diff --git a/Engine/Source/Runtime/Experimental/Chaos/Private/Chaos/CollisionResolution.cpp b/Engine/Source/Runtime/Experimental/Chaos/Private/Chaos/CollisionResolution.cpp
--- a/Engine/Source/Runtime/Experimental/Chaos/Private/Chaos/CollisionResolution.cpp
+++ b/Engine/Source/Runtime/Experimental/Chaos/Private/Chaos/CollisionResolution.cpp
@@ -2132,8 +2132,7 @@ namespace Chaos
 					const FReal R1 = Implicit0.BoundingBox().OriginRadius();
 					const FReal R2 = Implicit1.BoundingBox().OriginRadius();
 					const FReal SeparationSq = (WorldTransform1.GetTranslation() - WorldTransform0.GetTranslation()).SizeSquared();
-					const FReal CullDistanceSq = CullDistance * CullDistance;
-					if ((SeparationSq - FMath::Square(R1 + R2)) > CullDistanceSq)
+					if (SeparationSq > FMath::Square(R1 + R2 + CullDistance))
 					{
 						return;
 					}
@@ -2395,8 +2394,7 @@ namespace Chaos
 						const FReal R1 = Implicit0->BoundingBox().OriginRadius();
 						const FReal R2 = Implicit1->BoundingBox().OriginRadius();
 						const FReal SeparationSq = (WorldTransform1.GetTranslation() - WorldTransform0.GetTranslation()).SizeSquared();
-						const FReal CullDistanceSq = CullDistance * CullDistance;
-						if ((SeparationSq - FMath::Square(R1 + R2)) > CullDistanceSq)
+						if (SeparationSq > FMath::Square(R1 + R2 + CullDistance))
 						{
 							return;
 						}
diff --git a/Engine/Source/Runtime/Experimental/Chaos/Public/Chaos/Convex.h b/Engine/Source/Runtime/Experimental/Chaos/Public/Chaos/Convex.h
--- a/Engine/Source/Runtime/Experimental/Chaos/Public/Chaos/Convex.h
+++ b/Engine/Source/Runtime/Experimental/Chaos/Public/Chaos/Convex.h
@@ -457,7 +457,7 @@ namespace Chaos
 
 			if (InMargin == 0.0f)
 			{
-				return GetVertex(InMargin);
+				return GetVertex(VertexIndex);
 			}
 
 			// Get any 3 planes that contribute to this vertex
@@ -482,14 +482,37 @@ namespace Chaos
 				}
 			}
 
-			return FVec3(0);
+			// If we get here, the convex hull is malformed. Try to handle it anyway 
+			// @todo(chaos): track down the invalid hull issue
+
+			if (NumVertexPlanes(VertexIndex) == 2)
+			{
+				const int32 PlaneIndex0 = GetVertexPlane(VertexIndex, 0);
+				const int32 PlaneIndex1 = GetVertexPlane(VertexIndex, 1);
+				const FVec3 NewPlaneX = GetVertex(VertexIndex);
+				const FVec3 NewPlaneN0 = Planes[PlaneIndex0].Normal();
+				const FVec3 NewPlaneN1 = Planes[PlaneIndex1].Normal();
+				const FVec3 NewPlaneN = (NewPlaneN0 + NewPlaneN1).GetSafeNormal();
+				return NewPlaneX - (InMargin * NewPlaneN);
+			}
+
+			if (NumVertexPlanes(VertexIndex) == 1)
+			{
+				const int32 PlaneIndex0 = GetVertexPlane(VertexIndex, 0);
+				const FVec3 NewPlaneX = GetVertex(VertexIndex);
+				const FVec3 NewPlaneN = Planes[PlaneIndex0].Normal();
+				return NewPlaneX - (InMargin * NewPlaneN);
+			}
+
+			// Ok now we really are done...just return the outer vertex and duck
+			return GetVertex(VertexIndex);
 		}
 
 		FVec3 GetMarginAdjustedVertexScaled(int32 VertexIndex, FReal InMargin, const FVec3& Scale) const
 		{
 			if (InMargin == 0.0f)
 			{
-				return GetVertex(InMargin) * Scale;
+				return GetVertex(VertexIndex) * Scale;
 			}
 
 			// Get any 3 planes that contribute to this vertex
@@ -528,7 +551,30 @@ namespace Chaos
 				}
 			}
 
-			return FVec3(0);
+			// If we get here, the convex hull is malformed. Try to handle it anyway 
+			// @todo(chaos): track down the invalid hull issue
+
+			if (NumVertexPlanes(VertexIndex) == 2)
+			{
+				const int32 PlaneIndex0 = GetVertexPlane(VertexIndex, 0);
+				const int32 PlaneIndex1 = GetVertexPlane(VertexIndex, 1);
+				const FVec3 NewPlaneX = Scale * GetVertex(VertexIndex);
+				const FVec3 NewPlaneN0 = (Planes[PlaneIndex0].Normal() / Scale).GetUnsafeNormal();
+				const FVec3 NewPlaneN1 = (Planes[PlaneIndex1].Normal() / Scale).GetUnsafeNormal();
+				const FVec3 NewPlaneN = (NewPlaneN0 + NewPlaneN1).GetSafeNormal();
+				return NewPlaneX - (InMargin * NewPlaneN);
+			}
+
+			if (NumVertexPlanes(VertexIndex) == 1)
+			{
+				const int32 PlaneIndex0 = GetVertexPlane(VertexIndex, 0);
+				const FVec3 NewPlaneX = Scale * GetVertex(VertexIndex);
+				const FVec3 NewPlaneN = (Planes[PlaneIndex0].Normal() / Scale).GetUnsafeNormal();
+				return NewPlaneX - (InMargin * NewPlaneN);
+			}
+
+			// Ok now we really are done...just return the outer vertex and duck
+			return GetVertex(VertexIndex) * Scale;
 		}
 
 	public:
diff --git a/Engine/Source/Runtime/Experimental/Chaos/Public/Chaos/GJKShape.h b/Engine/Source/Runtime/Experimental/Chaos/Public/Chaos/GJKShape.h
--- a/Engine/Source/Runtime/Experimental/Chaos/Public/Chaos/GJKShape.h
+++ b/Engine/Source/Runtime/Experimental/Chaos/Public/Chaos/GJKShape.h
@@ -59,11 +59,19 @@ namespace Chaos
 	{
 		using FImplicitObjectType = T_SHAPE;
 
-		TGJKCoreShape(const FImplicitObjectType& InShape) : Shape(InShape) {}
+		TGJKCoreShape(const FImplicitObjectType& InShape) 
+			: Shape(InShape)
+			, Margin(InShape.GetMargin())
+		{}
+
+		TGJKCoreShape(const FImplicitObjectType& InShape, FReal InMargin)
+			: Shape(InShape)
+			, Margin(InMargin)
+		{}
 
 		FReal GetMargin() const
 		{
-			return Shape.GetMargin();
+			return Margin;
 		}
 
 		FVec3 SupportCore(const FVec3 Dir, FReal InMargin) const
@@ -77,6 +85,7 @@ namespace Chaos
 		}
 
 		const FImplicitObjectType& Shape;
+		const FReal Margin;
 	};
 
 
