Chaos
- fix collision/joint fighting caused by manifolds by disabling manifolds for bodies with joints

diff --git a/Engine/Source/Runtime/Experimental/Chaos/Private/Chaos/Collision/PBDCollisionConstraint.cpp b/Engine/Source/Runtime/Experimental/Chaos/Private/Chaos/Collision/PBDCollisionConstraint.cpp
--- a/Engine/Source/Runtime/Experimental/Chaos/Private/Chaos/Collision/PBDCollisionConstraint.cpp
+++ b/Engine/Source/Runtime/Experimental/Chaos/Private/Chaos/Collision/PBDCollisionConstraint.cpp
@@ -377,4 +377,11 @@ namespace Chaos
 
 		return ManifoldPoint;
 	}
+
+	bool FRigidBodyPointContactConstraint::CanUseManifold(FGeometryParticleHandle* Particle0, FGeometryParticleHandle* Particle1) const
+	{
+		// Do not use manifolds when a body is connected by a joint to another. Manifolds do not work when the bodies may be moved
+		// and rotated by significant amounts and joints can do that.
+		return (Particle0->ParticleConstraints().Num() == 0) && (Particle1->ParticleConstraints().Num() == 0);
+	}
 }
\ No newline at end of file
diff --git a/Engine/Source/Runtime/Experimental/Chaos/Public/Chaos/Collision/PBDCollisionConstraint.h b/Engine/Source/Runtime/Experimental/Chaos/Public/Chaos/Collision/PBDCollisionConstraint.h
--- a/Engine/Source/Runtime/Experimental/Chaos/Public/Chaos/Collision/PBDCollisionConstraint.h
+++ b/Engine/Source/Runtime/Experimental/Chaos/Public/Chaos/Collision/PBDCollisionConstraint.h
@@ -269,8 +269,10 @@ namespace Chaos
 			const bool bInUseManifold)
 			: Base(Particle0, Implicit0, Simplicial0, Transform0, Particle1, Implicit1, Simplicial0, Transform1, Base::FType::SinglePoint, ShapesType)
 			, CullDistance(InCullDistance)
-			, bUseManifold(bInUseManifold)
-		{}
+			, bUseManifold(false)
+		{
+			bUseManifold = bInUseManifold && CanUseManifold(Particle0, Particle1);
+		}
 
 		static typename Base::FType StaticType() { return Base::FType::SinglePoint; };
 
@@ -321,6 +323,10 @@ namespace Chaos
 			, bUseManifold(false)
 		{}
 
+		// Whether we can use manifolds for the given partices. Manifolds do not work well with Joints and PBD
+		// because the bodies may be moved (and especially rotated) a lot in the solver and this can make the manifold extremely inaccurate
+		bool CanUseManifold(FGeometryParticleHandle* Particle0, FGeometryParticleHandle* Particle1) const;
+
 		bool AreMatchingContactPoints(const FContactPoint& A, const FContactPoint& B, FReal& OutScore) const;
 		int32 FindManifoldPoint(const FContactPoint& ContactPoint) const;
 		int32 AddManifoldPoint(const FContactPoint& ContactPoint, const FReal Dt);
