Geometry Collection : Cache Attribute
- Strip geometry attributes from the collection when using the collection as a simulation viewer.
- Use the transferred bounds on the minimal asset for rendering bounds.
- Minor clean to expose user controllable cache only state to the GeometryCollectionObject

diff --git a/Engine/Source/Runtime/Experimental/Chaos/Private/Chaos/MassProperties.cpp b/Engine/Source/Runtime/Experimental/Chaos/Private/Chaos/MassProperties.cpp
--- a/Engine/Source/Runtime/Experimental/Chaos/Private/Chaos/MassProperties.cpp
+++ b/Engine/Source/Runtime/Experimental/Chaos/Private/Chaos/MassProperties.cpp
@@ -15,9 +15,9 @@ namespace Chaos
 
 		// Extract Eigenvalues
 		FReal OffDiagSize = FMath::Square(Inertia.M[1][0]) + FMath::Square(Inertia.M[2][0]) + FMath::Square(Inertia.M[2][1]);
-		FReal Trace = (Inertia.M[0][0] + Inertia.M[1][1] + Inertia.M[2][2]) / 3;
+		double Trace = (Inertia.M[0][0] + Inertia.M[1][1] + Inertia.M[2][2]) / 3;
 
-		if (!ensure(Trace > SMALL_NUMBER))
+		if (Trace <= SMALL_NUMBER)
 		{
 			// Tiny inertia - numerical instability would follow. We should not get this unless we have bad input.
 			return FRotation3::FromElements(FVec3(0), 1);
diff --git a/Engine/Source/Runtime/Experimental/Chaos/Private/GeometryCollection/ManagedArrayCollection.cpp b/Engine/Source/Runtime/Experimental/Chaos/Private/GeometryCollection/ManagedArrayCollection.cpp
--- a/Engine/Source/Runtime/Experimental/Chaos/Private/GeometryCollection/ManagedArrayCollection.cpp
+++ b/Engine/Source/Runtime/Experimental/Chaos/Private/GeometryCollection/ManagedArrayCollection.cpp
@@ -436,29 +436,6 @@ void FManagedArrayCollection::Serialize(Chaos::FChaosArchive& Ar)
 			}
 		}
 
-		// Material indices are tied to the implicit collision geometry
-		// baking cached assets will strip out the collision geometry
-		// so we copy the indices to an attribute for back access on 
-		// cached only assets. 
-		if (HasGroup("Transform"))
-		{
-			if (HasAttribute("Implicits", "Transform"))
-			{
-				if (!HasAttribute("DefaultMaterialIndex", "Transform"))
-				{
-					TManagedArray<int32>& DefaultMaterialID = AddAttribute<int32>("DefaultMaterialIndex", "Transform");
-
-					typedef TSharedPtr<Chaos::FImplicitObject, ESPMode::ThreadSafe> FSharedImplicit;
-					TManagedArray<FSharedImplicit>& Implicits = GetAttribute<FSharedImplicit>("Implicits", "Transform");
-					for (int i = 0; i < NumElements("Transform"); i++)
-					{
-						if (Implicits[i])
-							DefaultMaterialID[i] = Implicits[i]->GetMaterialIndex(0);
-					}
-				}
-			}
-		}
-
 #endif
 
 		// strip out GUID Attributes
diff --git a/Engine/Source/Runtime/Experimental/Chaos/Private/PhysicsProxy/GeometryCollectionPhysicsProxy.cpp b/Engine/Source/Runtime/Experimental/Chaos/Private/PhysicsProxy/GeometryCollectionPhysicsProxy.cpp
--- a/Engine/Source/Runtime/Experimental/Chaos/Private/PhysicsProxy/GeometryCollectionPhysicsProxy.cpp
+++ b/Engine/Source/Runtime/Experimental/Chaos/Private/PhysicsProxy/GeometryCollectionPhysicsProxy.cpp
@@ -58,13 +58,6 @@ FAutoConsoleVariableRef CVarGeometryCollectionCollideAll(
 	GeometryCollectionCollideAll,
 	TEXT("Bypass the collision matrix and make geometry collections collide against everything"));
 
-
-bool GeometryCollectionBypassPhysicsAttributes = false;
-FAutoConsoleVariableRef CVarGeometryCollectionBypassPhysicsAttributes(
-	TEXT("p.GeometryCollectionBypassPhysicsAttributes"),
-	GeometryCollectionBypassPhysicsAttributes,
-	TEXT("Bypass the construction of simulation properties when all bodies are simply cached. for playback."));
-
 DEFINE_LOG_CATEGORY_STATIC(UGCC_LOG, Error, All);
 
 //==============================================================================
@@ -663,6 +656,7 @@ void FGeometryCollectionPhysicsProxy::InitializeBodiesPT(Chaos::FPBDRigidsSolver
 		int32 NextIdx = 0;
 		for (int32 Idx = 0; Idx < SimulatableParticles.Num(); ++Idx)
 		{
+			SolverParticleHandles[Idx] = nullptr;
 			if (SimulatableParticles[Idx] && !RestCollection->IsClustered(Idx))
 			{
 				// todo: Unblocked read access of game thread data on the physics thread.
@@ -690,12 +684,8 @@ void FGeometryCollectionPhysicsProxy::InitializeBodiesPT(Chaos::FPBDRigidsSolver
 		// for the case when there's a 1-to-1 mapping between transforms and geometries.
 		// At the point that we start supporting instancing, this assumption will no longer
 		// hold, and those reverse mappints will be INDEX_NONE.
-
-		const int32 NumGeometries = DynamicCollection.NumElements(FGeometryCollection::GeometryGroup);
-		ParallelFor(NumGeometries, [&](int32 GeometryIndex)
-		//for (int32 GeometryIndex = 0; GeometryIndex < NumGeometries; ++GeometryIndex)
+		ParallelFor(NumTransforms, [&](int32 TransformGroupIndex)
 		{
-			const int32 TransformGroupIndex = TransformIndex[GeometryIndex];
 			if (FClusterHandle* Handle = SolverParticleHandles[TransformGroupIndex])
 			{
 				// Mass space -> Composed parent space -> world
@@ -1824,8 +1814,6 @@ void FGeometryCollectionPhysicsProxy::InitializeSharedCollisionStructures(
 	FGeometryCollection& RestCollection,
 	const FSharedSimulationParameters& SharedParams)
 {
-	if (GeometryCollectionBypassPhysicsAttributes)  return;
-
 	FString BaseErrorPrefix = ErrorReporter.GetPrefix();
 
 	// fracture tools can create an empty GC before appending new geometry
diff --git a/Engine/Source/Runtime/Experimental/GeometryCollectionEngine/Private/GeometryCollection/GeometryCollectionComponent.cpp b/Engine/Source/Runtime/Experimental/GeometryCollectionEngine/Private/GeometryCollection/GeometryCollectionComponent.cpp
--- a/Engine/Source/Runtime/Experimental/GeometryCollectionEngine/Private/GeometryCollection/GeometryCollectionComponent.cpp
+++ b/Engine/Source/Runtime/Experimental/GeometryCollectionEngine/Private/GeometryCollection/GeometryCollectionComponent.cpp
@@ -407,14 +407,26 @@ FBoxSphereBounds UGeometryCollectionComponent::CalcBounds(const FTransform& Loca
 
 		const int32 NumBoxes = BoundingBoxes.Num();
 
-		// #todo(dmp): we could do the bbox transform in parallel with a bit of reformulating		
-		// #todo(dmp):  there are some cases where the calcbounds function is called before the component
-		// has set the global matrices cache while in the editor.  This is a somewhat weak guard against this
-		// to default to just calculating tmp global matrices.  This should be removed or modified somehow
-		// such that we always cache the global matrices and this method always does the correct behavior
-		int32 NumElements = RestCollection->NumElements(FGeometryCollection::TransformGroup);
-		if (NumElements == 0 || GlobalMatrices.Num() != RestCollection->NumElements(FGeometryCollection::TransformGroup))
+		int32 NumElements = HackGeometryCollectionPtr->NumElements(FGeometryCollection::TransformGroup);
+		if (RestCollection->EnableNanite && HackGeometryCollectionPtr->HasAttribute("BoundingBox", FGeometryCollection::TransformGroup) && NumElements)
 		{
+			TArray<FMatrix> TmpGlobalMatrices;
+			GeometryCollectionAlgo::GlobalMatrices(Transforms, ParentIndices, TmpGlobalMatrices);
+
+			TManagedArray<FBox>& TransformBounds = HackGeometryCollectionPtr->GetAttribute<FBox>("BoundingBox", "Transform");
+			for (int TransformIndex = 0; TransformIndex < HackGeometryCollectionPtr->NumElements(FGeometryCollection::TransformGroup); TransformIndex++)
+			{
+				BoundingBox += TransformBounds[TransformIndex].TransformBy(TmpGlobalMatrices[TransformIndex] * LocalToWorldWithScale);
+			}
+		}
+		else if (NumElements == 0 || GlobalMatrices.Num() != RestCollection->NumElements(FGeometryCollection::TransformGroup))
+		{
+			// #todo(dmp): we could do the bbox transform in parallel with a bit of reformulating		
+			// #todo(dmp):  there are some cases where the calcbounds function is called before the component
+			// has set the global matrices cache while in the editor.  This is a somewhat weak guard against this
+			// to default to just calculating tmp global matrices.  This should be removed or modified somehow
+			// such that we always cache the global matrices and this method always does the correct behavior
+
 			TArray<FMatrix> TmpGlobalMatrices;
 			
 			GeometryCollectionAlgo::GlobalMatrices(Transforms, ParentIndices, TmpGlobalMatrices);
@@ -1164,169 +1176,116 @@ void UGeometryCollectionComponent::InitConstantData(FGeometryCollectionConstantD
 	const FGeometryCollection* Collection = RestCollection->GetGeometryCollection().Get();
 	check(Collection);
 
-	const int32 NumPoints = Collection->NumElements(FGeometryCollection::VerticesGroup);
-	const TManagedArray<FVector>& Vertex = Collection->Vertex;
-	const TManagedArray<int32>& BoneMap = Collection->BoneMap;
-	const TManagedArray<FVector>& TangentU = Collection->TangentU;
-	const TManagedArray<FVector>& TangentV = Collection->TangentV;
-	const TManagedArray<FVector>& Normal = Collection->Normal;
-	const TManagedArray<FVector2D>& UV = Collection->UV;
-	const TManagedArray<FLinearColor>& Color = Collection->Color;
-	const TManagedArray<FLinearColor>& BoneColors = Collection->BoneColor;
-
-	ConstantData->Vertices = TArray<FVector>(Vertex.GetData(), Vertex.Num());
-	ConstantData->BoneMap = TArray<int32>(BoneMap.GetData(), BoneMap.Num());
-	ConstantData->TangentU = TArray<FVector>(TangentU.GetData(), TangentU.Num());
-	ConstantData->TangentV = TArray<FVector>(TangentV.GetData(), TangentV.Num());
-	ConstantData->Normals = TArray<FVector>(Normal.GetData(), Normal.Num());
-	ConstantData->UVs = TArray<FVector2D>(UV.GetData(), UV.Num());
-	ConstantData->Colors = TArray<FLinearColor>(Color.GetData(), Color.Num());
-
-	ConstantData->BoneColors.AddUninitialized(NumPoints);
-	
-	ParallelFor(NumPoints, [&](const int32 InPointIndex)
-	{
-		const int32 BoneIndex = ConstantData->BoneMap[InPointIndex];
-		ConstantData->BoneColors[InPointIndex] = BoneColors[BoneIndex];
-	});
+	if (!RestCollection->EnableNanite)
+	{
+		const int32 NumPoints = Collection->NumElements(FGeometryCollection::VerticesGroup);
+		const TManagedArray<FVector>& Vertex = Collection->Vertex;
+		const TManagedArray<int32>& BoneMap = Collection->BoneMap;
+		const TManagedArray<FVector>& TangentU = Collection->TangentU;
+		const TManagedArray<FVector>& TangentV = Collection->TangentV;
+		const TManagedArray<FVector>& Normal = Collection->Normal;
+		const TManagedArray<FVector2D>& UV = Collection->UV;
+		const TManagedArray<FLinearColor>& Color = Collection->Color;
+		const TManagedArray<FLinearColor>& BoneColors = Collection->BoneColor;
+
+		ConstantData->Vertices = TArray<FVector>(Vertex.GetData(), Vertex.Num());
+		ConstantData->BoneMap = TArray<int32>(BoneMap.GetData(), BoneMap.Num());
+		ConstantData->TangentU = TArray<FVector>(TangentU.GetData(), TangentU.Num());
+		ConstantData->TangentV = TArray<FVector>(TangentV.GetData(), TangentV.Num());
+		ConstantData->Normals = TArray<FVector>(Normal.GetData(), Normal.Num());
+		ConstantData->UVs = TArray<FVector2D>(UV.GetData(), UV.Num());
+		ConstantData->Colors = TArray<FLinearColor>(Color.GetData(), Color.Num());
+
+		ConstantData->BoneColors.AddUninitialized(NumPoints);
+
+		ParallelFor(NumPoints, [&](const int32 InPointIndex)
+			{
+				const int32 BoneIndex = ConstantData->BoneMap[InPointIndex];
+				ConstantData->BoneColors[InPointIndex] = BoneColors[BoneIndex];
+			});
 
-	int32 NumIndices = 0;
-	const TManagedArray<FIntVector>& Indices = Collection->Indices;
-	const TManagedArray<int32>& MaterialID = Collection->MaterialID;
-	
-	const TManagedArray<bool>& Visible = GetVisibleArray();  // Use copy on write attribute. The rest collection visible array can be overriden for the convenience of debug drawing the collision volumes
-	const TManagedArray<int32>& MaterialIndex = Collection->MaterialIndex;
+		int32 NumIndices = 0;
+		const TManagedArray<FIntVector>& Indices = Collection->Indices;
+		const TManagedArray<int32>& MaterialID = Collection->MaterialID;
 
-	const int32 NumFaceGroupEntries = Collection->NumElements(FGeometryCollection::FacesGroup);
+		const TManagedArray<bool>& Visible = GetVisibleArray();  // Use copy on write attribute. The rest collection visible array can be overriden for the convenience of debug drawing the collision volumes
+		const TManagedArray<int32>& MaterialIndex = Collection->MaterialIndex;
 
-	for (int FaceIndex = 0; FaceIndex < NumFaceGroupEntries; ++FaceIndex)
-	{
-		NumIndices += static_cast<int>(Visible[FaceIndex]);
-	}
+		const int32 NumFaceGroupEntries = Collection->NumElements(FGeometryCollection::FacesGroup);
 
-	ConstantData->Indices.AddUninitialized(NumIndices);
-	for (int IndexIdx = 0, cdx = 0; IndexIdx < NumFaceGroupEntries; ++IndexIdx)
-	{
-		if (Visible[ MaterialIndex[IndexIdx] ])
+		for (int FaceIndex = 0; FaceIndex < NumFaceGroupEntries; ++FaceIndex)
 		{
-			ConstantData->Indices[cdx++] = Indices[ MaterialIndex[IndexIdx] ];
+			NumIndices += static_cast<int>(Visible[FaceIndex]);
 		}
-	}
 
-	// We need to correct the section index start point & number of triangles since only the visible ones have been copied across in the code above
-	const int32 NumMaterialSections = Collection->NumElements(FGeometryCollection::MaterialGroup);
-	ConstantData->Sections.AddUninitialized(NumMaterialSections);
-	const TManagedArray<FGeometryCollectionSection>& Sections = Collection->Sections;
-	for (int SectionIndex = 0; SectionIndex < NumMaterialSections; ++SectionIndex)
-	{
-		FGeometryCollectionSection Section = Sections[SectionIndex]; // deliberate copy
-
-		for (int32 TriangleIndex = 0; TriangleIndex < Sections[SectionIndex].FirstIndex / 3; TriangleIndex++)
+		ConstantData->Indices.AddUninitialized(NumIndices);
+		for (int IndexIdx = 0, cdx = 0; IndexIdx < NumFaceGroupEntries; ++IndexIdx)
 		{
-			if(!Visible[MaterialIndex[TriangleIndex]])
+			if (Visible[MaterialIndex[IndexIdx]])
 			{
-				Section.FirstIndex -= 3;
+				ConstantData->Indices[cdx++] = Indices[MaterialIndex[IndexIdx]];
 			}
 		}
 
-		for (int32 TriangleIndex = 0; TriangleIndex < Sections[SectionIndex].NumTriangles; TriangleIndex++)
+		// We need to correct the section index start point & number of triangles since only the visible ones have been copied across in the code above
+		const int32 NumMaterialSections = Collection->NumElements(FGeometryCollection::MaterialGroup);
+		ConstantData->Sections.AddUninitialized(NumMaterialSections);
+		const TManagedArray<FGeometryCollectionSection>& Sections = Collection->Sections;
+		for (int SectionIndex = 0; SectionIndex < NumMaterialSections; ++SectionIndex)
 		{
-			if(!Visible[MaterialIndex[Sections[SectionIndex].FirstIndex / 3 + TriangleIndex]])
+			FGeometryCollectionSection Section = Sections[SectionIndex]; // deliberate copy
+
+			for (int32 TriangleIndex = 0; TriangleIndex < Sections[SectionIndex].FirstIndex / 3; TriangleIndex++)
 			{
-				Section.NumTriangles--;
+				if (!Visible[MaterialIndex[TriangleIndex]])
+				{
+					Section.FirstIndex -= 3;
+				}
 			}
-		}
 
-		ConstantData->Sections[SectionIndex] = MoveTemp(Section);
-	}
-	ConstantData->NumTransforms = Collection->NumElements(FGeometryCollection::TransformGroup);
-	ConstantData->LocalBounds = LocalBounds;
-
-	// store the index buffer and render sections for the base unfractured mesh
-	const TManagedArray<int32>& TransformToGeometryIndex = Collection->TransformToGeometryIndex;
-	const TManagedArray<int32>&	FaceStart = Collection->FaceStart;
-	const TManagedArray<int32>&	FaceCount = Collection->FaceCount;
-	
-	const int32 NumFaces = Collection->NumElements(FGeometryCollection::FacesGroup);
-	TArray<FIntVector> BaseMeshIndices;
-	TArray<int32> BaseMeshOriginalFaceIndices;	
+			for (int32 TriangleIndex = 0; TriangleIndex < Sections[SectionIndex].NumTriangles; TriangleIndex++)
+			{
+				if (!Visible[MaterialIndex[Sections[SectionIndex].FirstIndex / 3 + TriangleIndex]])
+				{
+					Section.NumTriangles--;
+				}
+			}
 
-	BaseMeshIndices.Reserve(NumFaces);
-	BaseMeshOriginalFaceIndices.Reserve(NumFaces);
+			ConstantData->Sections[SectionIndex] = MoveTemp(Section);
+		}
+		ConstantData->NumTransforms = Collection->NumElements(FGeometryCollection::TransformGroup);
+		ConstantData->LocalBounds = LocalBounds;
 
-	// add all visible external faces to the original geometry index array
-	// #note:  This is a stopgap because the original geometry array is broken
-	for (int FaceIndex = 0; FaceIndex < NumFaces; ++FaceIndex)
-	{
-		// only add visible external faces.  MaterialID that is even is an external material
-		if (Visible[FaceIndex] && MaterialID[FaceIndex] % 2 == 0)
-		{
-			BaseMeshIndices.Add(Indices[FaceIndex]);
-			BaseMeshOriginalFaceIndices.Add(FaceIndex);
-		}				
-	}
+		// store the index buffer and render sections for the base unfractured mesh
+		const TManagedArray<int32>& TransformToGeometryIndex = Collection->TransformToGeometryIndex;
+		const TManagedArray<int32>& FaceStart = Collection->FaceStart;
+		const TManagedArray<int32>& FaceCount = Collection->FaceCount;
 
-	// We should always have external faces of a geometry collection
-	ensure(BaseMeshIndices.Num() > 0);
+		const int32 NumFaces = Collection->NumElements(FGeometryCollection::FacesGroup);
+		TArray<FIntVector> BaseMeshIndices;
+		TArray<int32> BaseMeshOriginalFaceIndices;
 
-	// #todo(dmp): we should eventually get this working where we use geometry nodes
-	// that signify original unfractured geometry.  For now, this system is broken.
-	/*
-	for (int i = 0; i < Collection->NumElements(FGeometryCollection::TransformGroup); ++i)
-	{
-		const FGeometryCollectionBoneNode &CurrBone = BoneHierarchy[i];
+		BaseMeshIndices.Reserve(NumFaces);
+		BaseMeshOriginalFaceIndices.Reserve(NumFaces);
 
-		// root node could be parent geo
-		if (CurrBone.Parent == INDEX_NONE)
+		// add all visible external faces to the original geometry index array
+		// #note:  This is a stopgap because the original geometry array is broken
+		for (int FaceIndex = 0; FaceIndex < NumFaces; ++FaceIndex)
 		{
-			int32 GeometryIndex = TransformToGeometryIndex[i];
-
-			// found geometry associated with base mesh root node
-			if (GeometryIndex != INDEX_NONE)
-			{				
-				int32 CurrFaceStart = FaceStart[GeometryIndex];
-				int32 CurrFaceCount = FaceCount[GeometryIndex];
-			
-				// add all the faces to the original geometry face array
-				for (int face = CurrFaceStart; face < CurrFaceStart + CurrFaceCount; ++face)
-				{
-					BaseMeshIndices.Add(Indices[face]);
-					BaseMeshOriginalFaceIndices.Add(face);
-				}
-
-				// build an array of mesh sections
-				ConstantData->HasOriginalMesh = true;				
-			}
-			else
+			// only add visible external faces.  MaterialID that is even is an external material
+			if (Visible[FaceIndex] && MaterialID[FaceIndex] % 2 == 0)
 			{
-				// all the direct decedents of the root node with no geometry are original geometry
-				for (int32 CurrChild : CurrBone.Children)
-				{					
-					int32 GeometryIndex = TransformToGeometryIndex[CurrChild];
-					if (GeometryIndex != INDEX_NONE)
-					{
-						// original geo static mesh					
-						int32 CurrFaceStart = FaceStart[GeometryIndex];
-						int32 CurrFaceCount = FaceCount[GeometryIndex];
-
-						// add all the faces to the original geometry face array
-						for (int face = CurrFaceStart; face < CurrFaceStart + CurrFaceCount; ++face)
-						{
-							BaseMeshIndices.Add(Indices[face]);
-							BaseMeshOriginalFaceIndices.Add(face);
-						}
-
-						ConstantData->HasOriginalMesh = true;
-					}					
-				}
+				BaseMeshIndices.Add(Indices[FaceIndex]);
+				BaseMeshOriginalFaceIndices.Add(FaceIndex);
 			}
 		}
-	}
-	*/
-
 
-	ConstantData->OriginalMeshSections = Collection->BuildMeshSections(BaseMeshIndices, BaseMeshOriginalFaceIndices, ConstantData->OriginalMeshIndices);
+		// We should always have external faces of a geometry collection
+		ensure(BaseMeshIndices.Num() > 0);
 
+		ConstantData->OriginalMeshSections = Collection->BuildMeshSections(BaseMeshIndices, BaseMeshOriginalFaceIndices, ConstantData->OriginalMeshIndices);
+	}
+	
 	TArray<FMatrix> RestMatrices;
 	GeometryCollectionAlgo::GlobalMatrices(RestCollection->GetGeometryCollection()->Transform, RestCollection->GetGeometryCollection()->Parent, RestMatrices);
 
@@ -1510,121 +1469,7 @@ void UGeometryCollectionComponent::InitDynamicData(FGeometryCollectionDynamicDat
 						}
 					}
 				}
-
 				DynamicData->Transforms = GlobalMatrices;
-
-				/**********************************************************************************************************************************************************************************/
-				// Capture all events for the given time
-			
-				if(false)
-				{
-					// clear events on the solver
-					Chaos::FPhysicsSolver *Solver = GetSolver(*this);
-					
-					if (Solver)
-					{
-#if TODO_REPLACE_SOLVER_LOCK
-						Chaos::FSolverWriteLock ScopedWriteLock(Solver);
-#endif
-
-#if TODO_REIMPLEMENT_EVENTS_DATA_ARRAYS
-						//////////////////////////////////////////////////////////////////////////
-						// Temporary workaround for writing on multiple threads.
-						// The following is called wide from the render thread to populate
-						// Niagara data from a cache without invoking a solver directly
-						//
-						// The above write lock guarantees we can safely write to these buffers
-						//
-						// TO BE REFACTORED
-						// #TODO BG
-						//////////////////////////////////////////////////////////////////////////
-						Chaos::FPhysicsSolver::FAllCollisionData *CollisionDataToWriteTo = const_cast<Chaos::FPhysicsSolver::FAllCollisionData*>(Solver->GetAllCollisions_FromSequencerCache_NEEDSLOCK());
-						Chaos::FPhysicsSolver::FAllBreakingData *BreakingDataToWriteTo = const_cast<Chaos::FPhysicsSolver::FAllBreakingData*>(Solver->GetAllBreakings_FromSequencerCache_NEEDSLOCK());
-						Chaos::FPhysicsSolver::FAllTrailingData *TrailingDataToWriteTo = const_cast<Chaos::FPhysicsSolver::FAllTrailingData*>(Solver->GetAllTrailings_FromSequencerCache_NEEDSLOCK());
-
-						if (!FMath::IsNearlyEqual(CollisionDataToWriteTo->TimeCreated, DesiredCacheTime))
-						{
-							CollisionDataToWriteTo->AllCollisionsArray.Empty();
-							CollisionDataToWriteTo->TimeCreated = DesiredCacheTime;
-						}
-
-						int32 Index = CacheParameters.TargetCache->GetData()->FindLastKeyBefore(CurrentCacheTime);
-						const FRecordedFrame *RecordedFrame = &CacheParameters.TargetCache->GetData()->Records[Index];				
-
-						if (RecordedFrame && PhysicsProxy && !EventsPlayed[Index])
-						{
-							EventsPlayed[Index] = true;
-
-							// Collisions
-							if (RecordedFrame->Collisions.Num() > 0)
-							{							
-								for (int32 Idx = 0; Idx < RecordedFrame->Collisions.Num(); ++Idx)
-								{
-									// Check if the particle is still kinematic
-									int32 NewIdx = CollisionDataToWriteTo->AllCollisionsArray.Add(Chaos::FCollidingData());
-									Chaos::FCollidingData& AllCollisionsDataArrayItem = CollisionDataToWriteTo->AllCollisionsArray[NewIdx];
-
-									AllCollisionsDataArrayItem.Location = RecordedFrame->Collisions[Idx].Location;
-									AllCollisionsDataArrayItem.AccumulatedImpulse = RecordedFrame->Collisions[Idx].AccumulatedImpulse;
-									AllCollisionsDataArrayItem.Normal = RecordedFrame->Collisions[Idx].Normal;
-									AllCollisionsDataArrayItem.Velocity1 = RecordedFrame->Collisions[Idx].Velocity1;
-									AllCollisionsDataArrayItem.Velocity2 = RecordedFrame->Collisions[Idx].Velocity2;
-									AllCollisionsDataArrayItem.AngularVelocity1 = RecordedFrame->Collisions[Idx].AngularVelocity1;
-									AllCollisionsDataArrayItem.AngularVelocity2 = RecordedFrame->Collisions[Idx].AngularVelocity2;
-									AllCollisionsDataArrayItem.Mass1 = RecordedFrame->Collisions[Idx].Mass1;
-									AllCollisionsDataArrayItem.Mass2 = RecordedFrame->Collisions[Idx].Mass2;
-#if TODO_CONVERT_GEOMETRY_COLLECTION_PARTICLE_INDICES_TO_PARTICLE_POINTERS
-									AllCollisionsDataArrayItem.ParticleIndex = RecordedFrame->Collisions[Idx].ParticleIndex;
-#endif
-									AllCollisionsDataArrayItem.LevelsetIndex = RecordedFrame->Collisions[Idx].LevelsetIndex;
-									AllCollisionsDataArrayItem.ParticleIndexMesh = RecordedFrame->Collisions[Idx].ParticleIndexMesh;
-									AllCollisionsDataArrayItem.LevelsetIndexMesh = RecordedFrame->Collisions[Idx].LevelsetIndexMesh;
-								}
-							}
-
-							// Breaking
-							if (RecordedFrame->Breakings.Num() > 0)
-							{
-								for (int32 Idx = 0; Idx < RecordedFrame->Breakings.Num(); ++Idx)
-								{
-									// Check if the particle is still kinematic							
-									int32 NewIdx = BreakingDataToWriteTo->AllBreakingsArray.Add(Chaos::FBreakingData());
-									Chaos::FBreakingData& AllBreakingsDataArrayItem = BreakingDataToWriteTo->AllBreakingsArray[NewIdx];
-
-									AllBreakingsDataArrayItem.Location = RecordedFrame->Breakings[Idx].Location;
-									AllBreakingsDataArrayItem.Velocity = RecordedFrame->Breakings[Idx].Velocity;
-									AllBreakingsDataArrayItem.AngularVelocity = RecordedFrame->Breakings[Idx].AngularVelocity;
-									AllBreakingsDataArrayItem.Mass = RecordedFrame->Breakings[Idx].Mass;
-#if TODO_CONVERT_GEOMETRY_COLLECTION_PARTICLE_INDICES_TO_PARTICLE_POINTERS
-									AllBreakingsDataArrayItem.ParticleIndex = RecordedFrame->Breakings[Idx].ParticleIndex;
-#endif
-									AllBreakingsDataArrayItem.ParticleIndexMesh = RecordedFrame->Breakings[Idx].ParticleIndexMesh;
-								}
-							}
-
-							// Trailing
-							if (RecordedFrame->Trailings.Num() > 0)
-							{
-								for (FSolverTrailingData Trailing : RecordedFrame->Trailings)
-								{
-									// Check if the particle is still kinematic
-									int32 NewIdx = TrailingDataToWriteTo->AllTrailingsArray.Add(Chaos::FTrailingData());
-									Chaos::FTrailingData& AllTrailingsDataArrayItem = TrailingDataToWriteTo->AllTrailingsArray[NewIdx];
-
-									AllTrailingsDataArrayItem.Location = Trailing.Location;
-									AllTrailingsDataArrayItem.Velocity = Trailing.Velocity;
-									AllTrailingsDataArrayItem.AngularVelocity = Trailing.AngularVelocity;
-									AllTrailingsDataArrayItem.Mass = Trailing.Mass;
-#if TODO_CONVERT_GEOMETRY_COLLECTION_PARTICLE_INDICES_TO_PARTICLE_POINTERS
-									AllTrailingsDataArrayItem.ParticleIndex = Trailing.ParticleIndex;
-#endif
-									AllTrailingsDataArrayItem.ParticleIndexMesh = Trailing.ParticleIndexMesh;
-								}
-							}
-						}
-#endif
-					}
-				}								
 			}
 
 			// check if transforms at start of this tick are the same as what is calculated from the cache
diff --git a/Engine/Source/Runtime/Experimental/GeometryCollectionEngine/Private/GeometryCollection/GeometryCollectionObject.cpp b/Engine/Source/Runtime/Experimental/GeometryCollectionEngine/Private/GeometryCollection/GeometryCollectionObject.cpp
--- a/Engine/Source/Runtime/Experimental/GeometryCollectionEngine/Private/GeometryCollection/GeometryCollectionObject.cpp
+++ b/Engine/Source/Runtime/Experimental/GeometryCollectionEngine/Private/GeometryCollection/GeometryCollectionObject.cpp
@@ -34,6 +34,13 @@
 
 DEFINE_LOG_CATEGORY_STATIC(LogGeometryCollectionInternal, Log, All);
 
+bool GeometryCollectionAssetForceStripOnCook = false;
+FAutoConsoleVariableRef CVarGeometryCollectionBypassPhysicsAttributes(
+	TEXT("p.GeometryCollectionAssetForceStripOnCook"),
+	GeometryCollectionAssetForceStripOnCook,
+	TEXT("Bypass the construction of simulation properties when all bodies are simply cached. for playback."));
+
+
 #if ENABLE_COOK_STATS
 namespace GeometryCollectionCookStats
 {
@@ -55,6 +62,7 @@ UGeometryCollection::UGeometryCollection(const FObjectInitializer& ObjectInitial
 	, MaxClusterLevel(100)
 	, DamageThreshold({ 250.0 })
 	, ClusterConnectionType(EClusterConnectionTypeEnum::Chaos_PointImplicit)
+	, bStripOnCook(false)
 	, EnableNanite(false)
 	, CollisionType(ECollisionTypeEnum::Chaos_Volumetric)
 	, ImplicitType(EImplicitTypeEnum::Chaos_Implicit_Box)
@@ -75,6 +83,7 @@ UGeometryCollection::UGeometryCollection(const FObjectInitializer& ObjectInitial
 	InvalidateCollection();
 #if WITH_EDITOR
 	SimulationDataGuid = StateGuid;
+	bStripOnCook = GeometryCollectionAssetForceStripOnCook;
 #endif
 }
 
@@ -314,7 +323,7 @@ bool UGeometryCollection::HasVisibleGeometry() const
 {
 	if(ensureMsgf(GeometryCollection.IsValid(), TEXT("Geometry Collection %s has an invalid internal collection")))
 	{
-		return GeometryCollection->HasVisibleGeometry();
+		return ( (EnableNanite && NaniteData) || GeometryCollection->HasVisibleGeometry());
 	}
 
 	return false;
@@ -407,11 +416,33 @@ void UGeometryCollection::Serialize(FArchive& Ar)
 	Ar.UsingCustomVersion(FUE5MainStreamObjectVersion::GUID);
 	Chaos::FChaosArchive ChaosAr(Ar);
 
+	// The Geometry Collection we will be archiving. This may be replaced with a transient, stripped back Geometry Collection if we are cooking.
+	TSharedPtr<FGeometryCollection, ESPMode::ThreadSafe> ArchiveGeometryCollection = GeometryCollection;
+
+	bool bIsCookedOrCooking = Ar.IsCooking();
+	if (bIsCookedOrCooking && Ar.IsSaving())
+	{
+#if WITH_EDITOR
+		if (bStripOnCook && EnableNanite && NaniteData)
+		{
+			// If this is a cooked archive, we strip unnecessary data from the Geometry Collection to keep the memory footprint as small as possible.
+			ArchiveGeometryCollection = GenerateMinimalGeometryCollection();
+		}
+#endif
+	}
+
 #if WITH_EDITOR
 	//Early versions did not have tagged properties serialize first
 	if (Ar.CustomVer(FDestructionObjectVersion::GUID) < FDestructionObjectVersion::GeometryCollectionInDDC)
 	{
-		GeometryCollection->Serialize(ChaosAr);
+		if (Ar.IsLoading())
+		{
+			GeometryCollection->Serialize(ChaosAr);
+		}
+		else
+		{
+			ArchiveGeometryCollection->Serialize(ChaosAr);
+		}
 	}
 
 	if (Ar.CustomVer(FDestructionObjectVersion::GUID) < FDestructionObjectVersion::AddedTimestampedGeometryComponentCache)
@@ -447,7 +478,6 @@ void UGeometryCollection::Serialize(FArchive& Ar)
 		}
 	}
 
-	bool bIsCookedOrCooking = Ar.IsCooking();
 	if (Ar.CustomVer(FDestructionObjectVersion::GUID) >= FDestructionObjectVersion::GeometryCollectionInDDC)
 	{
 		Ar << bIsCookedOrCooking;
@@ -462,31 +492,38 @@ void UGeometryCollection::Serialize(FArchive& Ar)
 			CreateSimulationDataImp(/*bCopyFromDDC=*/false);	//make sure content is built before saving
 		}
 #endif
-		GeometryCollection->Serialize(ChaosAr);
+		if (Ar.IsLoading())
+		{
+			GeometryCollection->Serialize(ChaosAr);
+		}
+		else
+		{
+			ArchiveGeometryCollection->Serialize(ChaosAr);
+		}
 
 		// Fix up the type change for implicits here, previously they were unique ptrs, now they're shared
-		TManagedArray<TUniquePtr<Chaos::FImplicitObject>>* OldAttr = GeometryCollection->FindAttributeTyped<TUniquePtr<Chaos::FImplicitObject>>(FGeometryDynamicCollection::ImplicitsAttribute, FTransformCollection::TransformGroup);
-		TManagedArray<TSharedPtr<Chaos::FImplicitObject, ESPMode::ThreadSafe>>* NewAttr = GeometryCollection->FindAttributeTyped<TSharedPtr<Chaos::FImplicitObject, ESPMode::ThreadSafe>>(FGeometryDynamicCollection::SharedImplicitsAttribute, FTransformCollection::TransformGroup);
-		if(OldAttr)
+		TManagedArray<TUniquePtr<Chaos::FImplicitObject>>* OldAttr = ArchiveGeometryCollection->FindAttributeTyped<TUniquePtr<Chaos::FImplicitObject>>(FGeometryDynamicCollection::ImplicitsAttribute, FTransformCollection::TransformGroup);
+		TManagedArray<TSharedPtr<Chaos::FImplicitObject, ESPMode::ThreadSafe>>* NewAttr = ArchiveGeometryCollection->FindAttributeTyped<TSharedPtr<Chaos::FImplicitObject, ESPMode::ThreadSafe>>(FGeometryDynamicCollection::SharedImplicitsAttribute, FTransformCollection::TransformGroup);
+		if (OldAttr)
 		{
-			if(!NewAttr)
+			if (!NewAttr)
 			{
-				NewAttr = &GeometryCollection->AddAttribute<TSharedPtr<Chaos::FImplicitObject, ESPMode::ThreadSafe>>(FGeometryDynamicCollection::SharedImplicitsAttribute, FTransformCollection::TransformGroup);
+				NewAttr = &ArchiveGeometryCollection->AddAttribute<TSharedPtr<Chaos::FImplicitObject, ESPMode::ThreadSafe>>(FGeometryDynamicCollection::SharedImplicitsAttribute, FTransformCollection::TransformGroup);
 
 				const int32 NumElems = GeometryCollection->NumElements(FTransformCollection::TransformGroup);
-				for(int32 Index = 0; Index < NumElems; ++Index)
+				for (int32 Index = 0; Index < NumElems; ++Index)
 				{
 					(*NewAttr)[Index] = TSharedPtr<Chaos::FImplicitObject, ESPMode::ThreadSafe>((*OldAttr)[Index].Release());
 				}
 			}
 
-			GeometryCollection->RemoveAttribute(FGeometryDynamicCollection::ImplicitsAttribute, FTransformCollection::TransformGroup);
+			ArchiveGeometryCollection->RemoveAttribute(FGeometryDynamicCollection::ImplicitsAttribute, FTransformCollection::TransformGroup);
 		}
 	}
 
 	if (Ar.CustomVer(FDestructionObjectVersion::GUID) < FDestructionObjectVersion::GroupAndAttributeNameRemapping)
 	{
-		GeometryCollection->UpdateOldAttributeNames();
+		ArchiveGeometryCollection->UpdateOldAttributeNames();
 		InvalidateCollection();
 #if WITH_EDITOR
 		CreateSimulationData();
@@ -517,6 +554,7 @@ void UGeometryCollection::Serialize(FArchive& Ar)
 		}
 	}
 
+
 #if WITH_EDITOR
 	//for all versions loaded, make sure sim data is up to date
 	if (Ar.IsLoading())
@@ -702,6 +740,81 @@ TUniquePtr<FGeometryCollectionNaniteData> UGeometryCollection::CreateNaniteData(
 	return NaniteData;
 }
 
+TSharedPtr<FGeometryCollection, ESPMode::ThreadSafe> UGeometryCollection::GenerateMinimalGeometryCollection() const
+{
+	TMap<FName, TSet<FName>> SkipList;
+	static TSet<FName> GeometryGroups{ FGeometryCollection::GeometryGroup, FGeometryCollection::VerticesGroup, FGeometryCollection::FacesGroup };
+	if (bStripOnCook)
+	{
+		// Remove all geometry
+		//static TSet<FName> GeometryGroups{ FGeometryCollection::GeometryGroup, FGeometryCollection::VerticesGroup, FGeometryCollection::FacesGroup, FGeometryCollection::MaterialGroup };
+		for (const FName& GeometryGroup : GeometryGroups)
+		{
+			TSet<FName>& SkipAttributes = SkipList.Add(GeometryGroup);
+			SkipAttributes.Append(GeometryCollection->AttributeNames(GeometryGroup));
+		}
+	}
+
+	TSharedPtr<FGeometryCollection, ESPMode::ThreadSafe> DuplicateGeometryCollection(new FGeometryCollection());
+	DuplicateGeometryCollection->AddAttribute<bool>(FGeometryCollection::SimulatableParticlesAttribute, FTransformCollection::TransformGroup);
+	DuplicateGeometryCollection->AddAttribute<FVector>("InertiaTensor", FGeometryCollection::TransformGroup);
+	DuplicateGeometryCollection->AddAttribute<float>("Mass", FGeometryCollection::TransformGroup);
+	DuplicateGeometryCollection->AddAttribute<FTransform>("MassToLocal", FGeometryCollection::TransformGroup);
+	DuplicateGeometryCollection->AddAttribute<FGeometryDynamicCollection::FSharedImplicit>(
+		FGeometryDynamicCollection::ImplicitsAttribute, FTransformCollection::TransformGroup);
+	DuplicateGeometryCollection->CopyMatchingAttributesFrom(*GeometryCollection, &SkipList);
+	// If we've removed all geometry, we need to make sure any references to that geometry are removed.
+	// We also need to resize geometry groups to ensure that they are empty.
+	if (bStripOnCook)
+	{
+		TManagedArray<int32>& TransformToGeometryIndex = DuplicateGeometryCollection->GetAttribute<int32>("TransformToGeometryIndex", FTransformCollection::TransformGroup);
+
+		//
+		// Copy the bounds to the TransformGroup.
+		//  @todo(nanite.bounds) : Rely on Nanite bounds in the component instead and dont copy here
+		//
+		if (!DuplicateGeometryCollection->HasAttribute("BoundingBox", "Transform"))
+		{
+			DuplicateGeometryCollection->AddAttribute<FBox>("BoundingBox", "Transform");
+		}
+
+		if (!DuplicateGeometryCollection->HasAttribute("NaniteIndex", "Transform"))
+		{
+			DuplicateGeometryCollection->AddAttribute<FBox>("NaniteIndex", "Transform");
+		}
+
+		int32 NumTransforms = GeometryCollection->NumElements(FGeometryCollection::TransformGroup);
+		TManagedArray<int32>& NaniteIndex = DuplicateGeometryCollection->GetAttribute<int32>("NaniteIndex", "Transform");
+		TManagedArray<FBox>& TransformBounds = DuplicateGeometryCollection->GetAttribute<FBox>("BoundingBox", "Transform");
+		TManagedArray<FBox>& GeometryBounds = GeometryCollection->GetAttribute<FBox>("BoundingBox", "Geometry");
+
+		NaniteIndex.Fill(INDEX_NONE);
+		for (int TransformIndex = 0; TransformIndex < NumTransforms; TransformIndex++)
+		{
+			NaniteIndex[TransformIndex] = TransformToGeometryIndex[TransformIndex];
+			int32 GeometryIndex = TransformToGeometryIndex[TransformIndex];
+			if (GeometryIndex != INDEX_NONE)
+			{
+				TransformBounds[TransformIndex] = GeometryBounds[GeometryIndex];
+			}
+			else
+			{
+				TransformBounds[TransformIndex].Init();
+			}
+		}
+
+		//
+		//  Clear the geometry and the transforms connection to it. 
+		//
+		//TransformToGeometryIndex.Fill(INDEX_NONE);
+		for (const FName& GeometryGroup : GeometryGroups)
+		{
+			DuplicateGeometryCollection->EmptyGroup(GeometryGroup);
+		}
+	}
+	return DuplicateGeometryCollection;
+}
+
 #endif
 
 void UGeometryCollection::InitResources()
@@ -810,12 +923,15 @@ void UGeometryCollection::EnsureDataIsCooked(bool bInitResources)
 	if (StateGuid != LastBuiltGuid)
 	{
 		CreateSimulationDataImp(/*bCopyFromDDC=*/ true);
-		
+
 		if (FApp::CanEverRender() && bInitResources)
 		{
-			NaniteData->InitResources(this);
+			// If there is no geometry in the collection, we leave Nanite data alone.
+			if (GeometryCollection->NumElements(FGeometryCollection::GeometryGroup) > 0)
+			{
+				NaniteData->InitResources(this);
+			}
 		}
-
 		LastBuiltGuid = StateGuid;
 	}
 }
diff --git a/Engine/Source/Runtime/Experimental/GeometryCollectionEngine/Private/GeometryCollection/GeometryCollectionSceneProxy.cpp b/Engine/Source/Runtime/Experimental/GeometryCollectionEngine/Private/GeometryCollection/GeometryCollectionSceneProxy.cpp
--- a/Engine/Source/Runtime/Experimental/GeometryCollectionEngine/Private/GeometryCollection/GeometryCollectionSceneProxy.cpp
+++ b/Engine/Source/Runtime/Experimental/GeometryCollectionEngine/Private/GeometryCollection/GeometryCollectionSceneProxy.cpp
@@ -1123,19 +1123,41 @@ FNaniteGeometryCollectionSceneProxy::FNaniteGeometryCollectionSceneProxy(UGeomet
 		MaterialSections[SectionIndex].MaterialIndex = MeshSection.MaterialID;
 	}
 
-	const TManagedArray<FBox>& BoundingBoxes = Collection->BoundingBox;
-
-	const int32 NumGeometry = Collection->NumElements(FGeometryCollection::GeometryGroup);
+	Nanite::FResources& Resource = GeometryCollection->NaniteData->NaniteResource;
+	const int32 NumGeometry = Resource.HierarchyRootOffsets.Num();
 	GeometryNaniteData.SetNumUninitialized(NumGeometry);
 
-	for (int32 GeometryIndex = 0; GeometryIndex < NumGeometry; ++GeometryIndex)
+	bool bHasGeometryBoundingBoxes = Collection->HasAttribute("BoundingBox", FGeometryCollection::GeometryGroup) && Collection->NumElements(FGeometryCollection::GeometryGroup)>0;
+	if (bHasGeometryBoundingBoxes)
 	{
-		FGeometryNaniteData& Instance = GeometryNaniteData[GeometryIndex];
-		Instance.PrimitiveId = ~uint32(0);
-		Instance.RenderBounds = BoundingBoxes[GeometryIndex];
-		Instance.NaniteInfo = GeometryCollection->GetNaniteInfo(GeometryIndex);
+		const TManagedArray<FBox>& BoundingBoxes = Collection->GetAttribute<FBox>("BoundingBox", FGeometryCollection::GeometryGroup);
+		for (int32 GeometryIndex = 0; GeometryIndex < NumGeometry; ++GeometryIndex)
+		{
+			FGeometryNaniteData& Instance = GeometryNaniteData[GeometryIndex];
+			Instance.PrimitiveId = ~uint32(0);
+			Instance.NaniteInfo = GeometryCollection->GetNaniteInfo(GeometryIndex);
+			Instance.RenderBounds = BoundingBoxes[GeometryIndex];
+		}
+	}
+	else
+	{
+		const TManagedArray<FBox>& BoundingBoxes = Collection->GetAttribute<FBox>("BoundingBox", FGeometryCollection::TransformGroup);
+		const TManagedArray<int32>& TransformToGeometry = Collection->GetAttribute<int32>("TransformToGeometryIndex", FGeometryCollection::TransformGroup);
+		const int32 NumTransforms = TransformToGeometry.Num();
+		for (int32 TransformIndex = 0; TransformIndex < NumTransforms; ++TransformIndex)
+		{
+			const int32 GeometryIndex = TransformToGeometry[TransformIndex];
+			if (GeometryIndex > INDEX_NONE)
+			{
+				FGeometryNaniteData& Instance = GeometryNaniteData[GeometryIndex];
+				Instance.PrimitiveId = ~uint32(0);
+				Instance.NaniteInfo = GeometryCollection->GetNaniteInfo(GeometryIndex);
+				Instance.RenderBounds = BoundingBoxes[GeometryIndex];
+			}
+		}
 	}
 
+
 	// Need to specify initial instance list, even with just identity transforms, so that the
 	// GPUScene instance data allocator reserves space for the instances early on. The instance
 	// transforms will be corrected during the first frame before any rendering occurs.
@@ -1257,6 +1279,7 @@ void FNaniteGeometryCollectionSceneProxy::SetDynamicData_RenderThread(FGeometryC
 		const TSharedPtr<FGeometryCollection, ESPMode::ThreadSafe> Collection = GeometryCollection->GetGeometryCollection();
 		const TManagedArray<int32>& TransformToGeometryIndices = Collection->TransformToGeometryIndex;
 		const TManagedArray<TSet<int32>>& TransformChildren = Collection->Children;
+		const TManagedArray<int32>& SimulationType = Collection->SimulationType;
 
 		check(NewDynamicData->Transforms.Num() == TransformToGeometryIndices.Num());
 		check(NewDynamicData->Transforms.Num() == TransformChildren.Num());
@@ -1265,7 +1288,8 @@ void FNaniteGeometryCollectionSceneProxy::SetDynamicData_RenderThread(FGeometryC
 		for (int32 TransformIndex = 0; TransformIndex < NewDynamicData->Transforms.Num(); ++TransformIndex)
 		{
 			const int32 TransformToGeometryIndex = TransformToGeometryIndices[TransformIndex];
-			if (!Collection->IsGeometry(TransformIndex))
+			//if (!Collection->IsGeometry(TransformIndex))
+			if( SimulationType[TransformIndex] != FGeometryCollection::ESimulationTypes::FST_Rigid)
 			{
 				continue;
 			}
diff --git a/Engine/Source/Runtime/Experimental/GeometryCollectionEngine/Public/GeometryCollection/GeometryCollectionObject.h b/Engine/Source/Runtime/Experimental/GeometryCollectionEngine/Public/GeometryCollection/GeometryCollectionObject.h
--- a/Engine/Source/Runtime/Experimental/GeometryCollectionEngine/Public/GeometryCollection/GeometryCollectionObject.h
+++ b/Engine/Source/Runtime/Experimental/GeometryCollectionEngine/Public/GeometryCollection/GeometryCollectionObject.h
@@ -238,6 +238,9 @@ public:
 	/** Remove embedded geometry exemplars with indices matching the sorted removal list. */
 	void RemoveExemplars(const TArray<int32>& SortedRemovalIndices);
 
+	/** Produce a deep copy of GeometryCollection member, stripped of data unecessary for gameplay. */
+	TSharedPtr<FGeometryCollection, ESPMode::ThreadSafe> GenerateMinimalGeometryCollection() const;
+
 #if WITH_EDITOR
 	/** If this flag is set, we only regenerate simulation data when requested via CreateSimulationData() */
 	bool bManualDataCreate;
@@ -295,6 +298,9 @@ public:
 	/**
 	* Enable support for Nanite.
 	*/
+	UPROPERTY(EditAnywhere, Category = "Nanite")
+	bool bStripOnCook;
+
 	UPROPERTY(EditAnywhere, BlueprintReadOnly, Category = "Nanite")
 	bool EnableNanite;
 
@@ -399,6 +405,7 @@ public:
 private:
 #if WITH_EDITOR
 	void CreateSimulationDataImp(bool bCopyFromDDC);
+
 #endif
 
 private:
