Geometry Collection : Render Bounds
- Fix to an indexing error when the bounds are setup on the FNaniteGeometryCollectionSceneProxy, and a transform error when updating the instance bounds.

diff --git a/Engine/Source/Runtime/Experimental/GeometryCollectionEngine/Private/GeometryCollection/GeometryCollectionSceneProxy.cpp b/Engine/Source/Runtime/Experimental/GeometryCollectionEngine/Private/GeometryCollection/GeometryCollectionSceneProxy.cpp
--- a/Engine/Source/Runtime/Experimental/GeometryCollectionEngine/Private/GeometryCollection/GeometryCollectionSceneProxy.cpp
+++ b/Engine/Source/Runtime/Experimental/GeometryCollectionEngine/Private/GeometryCollection/GeometryCollectionSceneProxy.cpp
@@ -1127,7 +1127,11 @@ FNaniteGeometryCollectionSceneProxy::FNaniteGeometryCollectionSceneProxy(UGeomet
 	const int32 NumGeometry = Resource.HierarchyRootOffsets.Num();
 	GeometryNaniteData.SetNumUninitialized(NumGeometry);
 
-	bool bHasGeometryBoundingBoxes = Collection->HasAttribute("BoundingBox", FGeometryCollection::GeometryGroup) && Collection->NumElements(FGeometryCollection::GeometryGroup)>0;
+	bool bHasGeometryBoundingBoxes = Collection->HasAttribute("BoundingBox", FGeometryCollection::GeometryGroup) && Collection->NumElements(FGeometryCollection::GeometryGroup);
+	bool bHasTransformBoundingBoxes = Collection->NumElements(FGeometryCollection::TransformGroup) && 
+		Collection->HasAttribute("BoundingBox", FGeometryCollection::TransformGroup) &&
+		Collection->HasAttribute("TransformToGeometryIndex", FGeometryCollection::TransformGroup);
+
 	if (bHasGeometryBoundingBoxes)
 	{
 		const TManagedArray<FBox>& BoundingBoxes = Collection->GetAttribute<FBox>("BoundingBox", FGeometryCollection::GeometryGroup);
@@ -1139,7 +1143,7 @@ FNaniteGeometryCollectionSceneProxy::FNaniteGeometryCollectionSceneProxy(UGeomet
 			Instance.RenderBounds = BoundingBoxes[GeometryIndex];
 		}
 	}
-	else
+	else if(bHasTransformBoundingBoxes)
 	{
 		const TManagedArray<FBox>& BoundingBoxes = Collection->GetAttribute<FBox>("BoundingBox", FGeometryCollection::TransformGroup);
 		const TManagedArray<int32>& TransformToGeometry = Collection->GetAttribute<int32>("TransformToGeometryIndex", FGeometryCollection::TransformGroup);
@@ -1152,7 +1156,7 @@ FNaniteGeometryCollectionSceneProxy::FNaniteGeometryCollectionSceneProxy(UGeomet
 				FGeometryNaniteData& Instance = GeometryNaniteData[GeometryIndex];
 				Instance.PrimitiveId = ~uint32(0);
 				Instance.NaniteInfo = GeometryCollection->GetNaniteInfo(GeometryIndex);
-				Instance.RenderBounds = BoundingBoxes[GeometryIndex];
+				Instance.RenderBounds = BoundingBoxes[TransformIndex];
 			}
 		}
 	}
@@ -1261,8 +1265,8 @@ void FNaniteGeometryCollectionSceneProxy::SetConstantData_RenderThread(FGeometry
 		Instance.InstanceToLocal	= NewConstantData->RestTransforms[TransformIndex];
 		Instance.LocalToWorld		= Instance.InstanceToLocal;
 		Instance.PrimitiveId		= NaniteData.PrimitiveId;
-		Instance.RenderBounds		= NaniteData.RenderBounds;
-		Instance.LocalBounds		= Instance.RenderBounds.TransformBy(Instance.InstanceToLocal);
+		Instance.RenderBounds       = NaniteData.RenderBounds;
+		Instance.LocalBounds        = Instance.RenderBounds.TransformBy(Instance.InstanceToLocal);
 		Instance.NaniteInfo			= NaniteData.NaniteInfo;
 	}
 
@@ -1286,7 +1290,6 @@ void FNaniteGeometryCollectionSceneProxy::SetDynamicData_RenderThread(FGeometryC
 		for (int32 TransformIndex = 0; TransformIndex < NewDynamicData->Transforms.Num(); ++TransformIndex)
 		{
 			const int32 TransformToGeometryIndex = TransformToGeometryIndices[TransformIndex];
-			//if (!Collection->IsGeometry(TransformIndex))
 			if( SimulationType[TransformIndex] != FGeometryCollection::ESimulationTypes::FST_Rigid)
 			{
 				continue;
@@ -1299,8 +1302,8 @@ void FNaniteGeometryCollectionSceneProxy::SetDynamicData_RenderThread(FGeometryC
 			Instance.InstanceToLocal	= NewDynamicData->Transforms[TransformIndex];
 			Instance.LocalToWorld		= Instance.InstanceToLocal;
 			Instance.PrimitiveId		= NaniteData.PrimitiveId;
-			Instance.RenderBounds		= NaniteData.RenderBounds;
-			Instance.LocalBounds		= Instance.RenderBounds.TransformBy(Instance.InstanceToLocal);
+			Instance.RenderBounds       = NaniteData.RenderBounds;
+			Instance.LocalBounds        = Instance.RenderBounds.TransformBy(Instance.InstanceToLocal);
 			Instance.NaniteInfo			= NaniteData.NaniteInfo;
 		}
 	}
