Rewrite of UGeometryCollection::InitDynamicData()
 * Implemented thread safe object pool / free list allocator (FGeometryCollectionDynamicDataPool)
 * Added Chaos stats for total transforms and changed transforms
 * Removed TransformsAreEqual mechanism that caused full transform updates to occur unnecessarily for 2 subsequent frames, as the state wasn't stable until the previous 3 frames were stable. This was incredibly innefficient for GPUScene, etc..
 * InitDynamicData now returns the data allocation (which can be null if not dynamic or if no transforms have changed), users of this now check if it's available, and either skip or fall back to the static/rest path otherwise
 * Cleaned up some code paths by combining them with similar passes
 * Removed tons of allocations, including heavy allocations/memcopies of transform data, even on stable frames (which were immediately deallocated right after allocation and setup of the dynamic data)
 * Saves ~0.5 to 1.0ms of CPU time in heavy scenes
 * Saves ~0.5ms of GPU time doing far less GPUScene updates
 * Main optimization around this refactor is forthcoming in a followup CL (avoid calling MarkRenderTransformDirty in TickComponent).

diff --git a/Engine/Source/Runtime/Experimental/Chaos/Private/ChaosStats.cpp b/Engine/Source/Runtime/Experimental/Chaos/Private/ChaosStats.cpp
--- a/Engine/Source/Runtime/Experimental/Chaos/Private/ChaosStats.cpp
+++ b/Engine/Source/Runtime/Experimental/Chaos/Private/ChaosStats.cpp
@@ -24,6 +24,8 @@ DEFINE_STAT(STAT_GCSweep);
 DEFINE_STAT(STAT_GCCUpdateBounds);
 DEFINE_STAT(STAT_GCCUGlobalMatrices);
 DEFINE_STAT(STAT_GCInitDynamicData);
+DEFINE_STAT(STAT_GCTotalTransforms);
+DEFINE_STAT(STAT_GCChangedTransforms);
 DEFINE_STAT(STAT_LockWaits);
 DEFINE_STAT(STAT_GeomBeginFrame);
 DEFINE_STAT(STAT_SkelMeshUpdateAnim);
diff --git a/Engine/Source/Runtime/Experimental/Chaos/Public/ChaosStats.h b/Engine/Source/Runtime/Experimental/Chaos/Public/ChaosStats.h
--- a/Engine/Source/Runtime/Experimental/Chaos/Public/ChaosStats.h
+++ b/Engine/Source/Runtime/Experimental/Chaos/Public/ChaosStats.h
@@ -38,6 +38,8 @@ DECLARE_CYCLE_STAT_EXTERN(TEXT("Geometry Collection Sweep"), STAT_GCSweep, STATG
 DECLARE_CYCLE_STAT_EXTERN(TEXT("Geometry Collection Component UpdateBounds"), STAT_GCCUpdateBounds, STATGROUP_Chaos, CHAOS_API);
 DECLARE_CYCLE_STAT_EXTERN(TEXT("Geometry Collection Component CalculateGlobalMatrices"), STAT_GCCUGlobalMatrices, STATGROUP_Chaos, CHAOS_API);
 DECLARE_CYCLE_STAT_EXTERN(TEXT("Geometry Collection InitDynamicData"), STAT_GCInitDynamicData, STATGROUP_Chaos, CHAOS_API);
+DECLARE_DWORD_COUNTER_STAT_EXTERN(TEXT("Geometry Collection Total Transforms"), STAT_GCTotalTransforms, STATGROUP_Chaos, CHAOS_API);
+DECLARE_DWORD_COUNTER_STAT_EXTERN(TEXT("Geometry Collection Changed Transforms"), STAT_GCChangedTransforms, STATGROUP_Chaos, CHAOS_API);
 DECLARE_CYCLE_STAT_EXTERN(TEXT("Physics Lock Waits"), STAT_LockWaits, STATGROUP_Chaos, CHAOS_API);
 DECLARE_CYCLE_STAT_EXTERN(TEXT("Geometry Collection Begin Frame"), STAT_GeomBeginFrame, STATGROUP_Chaos, CHAOS_API);
 DECLARE_CYCLE_STAT_EXTERN(TEXT("Skeletal Mesh Update Anim"), STAT_SkelMeshUpdateAnim, STATGROUP_Chaos, CHAOS_API);
diff --git a/Engine/Source/Runtime/Experimental/GeometryCollectionEngine/Private/GeometryCollection/GeometryCollectionComponent.cpp b/Engine/Source/Runtime/Experimental/GeometryCollectionEngine/Private/GeometryCollection/GeometryCollectionComponent.cpp
--- a/Engine/Source/Runtime/Experimental/GeometryCollectionEngine/Private/GeometryCollection/GeometryCollectionComponent.cpp
+++ b/Engine/Source/Runtime/Experimental/GeometryCollectionEngine/Private/GeometryCollection/GeometryCollectionComponent.cpp
@@ -63,6 +63,8 @@ MSVC_PRAGMA(warning(pop))
 
 DEFINE_LOG_CATEGORY_STATIC(UGCC_LOG, Error, All);
 
+extern FGeometryCollectionDynamicDataPool GDynamicDataPool;
+
 FString NetModeToString(ENetMode InMode)
 {
 	switch(InMode)
@@ -262,10 +264,6 @@ UGeometryCollectionComponent::UGeometryCollectionComponent(const FObjectInitiali
 	// default current cache time
 	CurrentCacheTime = MAX_flt;
 
-	// Buffer for rolling cache of past N=3 transforms being equal.
-	TransformsAreEqual.AddDefaulted(3);
-	TransformsAreEqualIndex = 0;
-
 	SetGenerateOverlapEvents(false);
 
 	// By default use the destructible object channel unless the user specifies otherwise
@@ -503,8 +501,7 @@ FPrimitiveSceneProxy* UGeometryCollectionComponent::CreateSceneProxy()
 			FGeometryCollectionConstantData* const ConstantData = ::new FGeometryCollectionConstantData;
 			InitConstantData(ConstantData);
 
-			FGeometryCollectionDynamicData* const DynamicData = ::new FGeometryCollectionDynamicData;
-			InitDynamicData(DynamicData);
+			FGeometryCollectionDynamicData* const DynamicData = InitDynamicData();
 
 			if (LocalSceneProxy->IsNaniteMesh())
 			{
@@ -522,14 +519,14 @@ FPrimitiveSceneProxy* UGeometryCollectionComponent::CreateSceneProxy()
 				ENQUEUE_RENDER_COMMAND(CreateRenderState)(
 					[GeometryCollectionSceneProxy, ConstantData, DynamicData](FRHICommandListImmediate& RHICmdList)
 					{
-						if (GeometryCollectionSceneProxy)
+						GeometryCollectionSceneProxy->SetConstantData_RenderThread(ConstantData);
+						if (DynamicData)
 						{
-							GeometryCollectionSceneProxy->SetConstantData_RenderThread(ConstantData);
 							GeometryCollectionSceneProxy->SetDynamicData_RenderThread(DynamicData);
-							if (GeometryCollectionSceneProxy->GetPrimitiveSceneInfo())
-							{
-								GeometryCollectionSceneProxy->GetPrimitiveSceneInfo()->RequestGPUSceneUpdate();
-							}
+						}
+						if (GeometryCollectionSceneProxy->GetPrimitiveSceneInfo())
+						{
+							GeometryCollectionSceneProxy->GetPrimitiveSceneInfo()->RequestGPUSceneUpdate();
 						}
 					}
 				);
@@ -549,9 +546,9 @@ FPrimitiveSceneProxy* UGeometryCollectionComponent::CreateSceneProxy()
 				ENQUEUE_RENDER_COMMAND(CreateRenderState)(
 					[GeometryCollectionSceneProxy, ConstantData, DynamicData](FRHICommandListImmediate& RHICmdList)
 					{
-						if (GeometryCollectionSceneProxy)
+						GeometryCollectionSceneProxy->SetConstantData_RenderThread(ConstantData);
+						if (DynamicData)
 						{
-							GeometryCollectionSceneProxy->SetConstantData_RenderThread(ConstantData);
 							GeometryCollectionSceneProxy->SetDynamicData_RenderThread(DynamicData);
 						}
 					}
@@ -1292,83 +1289,79 @@ void UGeometryCollectionComponent::InitConstantData(FGeometryCollectionConstantD
 	ConstantData->RestTransforms = MoveTemp(RestMatrices); 
 }
 
-void UGeometryCollectionComponent::InitDynamicData(FGeometryCollectionDynamicData * DynamicData)
+FGeometryCollectionDynamicData* UGeometryCollectionComponent::InitDynamicData()
 {
 	SCOPE_CYCLE_COUNTER(STAT_GCInitDynamicData);
 
-	check(DynamicData);
-	DynamicData->IsDynamic = this->GetIsObjectDynamic() || bShowBoneColors || bEnableBoneSelection;
-	DynamicData->IsLoading = this->GetIsObjectLoading();
+	FGeometryCollectionDynamicData* DynamicData = nullptr;
 
-	if (CachePlayback && CacheParameters.TargetCache && CacheParameters.TargetCache->GetData())
-	{		
-		const TManagedArray<int32> &Parents = GetParentArray();
-		const TManagedArray<TSet<int32>> &Children = GetChildrenArray();
-		const TManagedArray<FTransform> &Transform = RestCollection->GetGeometryCollection()->Transform;
-		const TManagedArray<FTransform> &MassToLocal = RestCollection->GetGeometryCollection()->GetAttribute<FTransform>("MassToLocal", FGeometryCollection::TransformGroup);
+	const bool bEditorMode = bShowBoneColors || bEnableBoneSelection;
+	const bool bIsDynamic = GetIsObjectDynamic() || bEditorMode;
 
-		// #todo(dmp): find a better place to calculate and store this
-		float CacheDt = CacheParameters.TargetCache->GetData()->GetDt();
-
-		// if we want the state before the first cached frame, then the collection doesn't need to be dynamic since it'll render the pre-fractured geometry
-		DynamicData->IsDynamic = DesiredCacheTime > CacheDt;
-
-		// if we are already on the current cached frame, return
+	if (CachePlayback && CacheParameters.TargetCache)
+	{
+		// If we are already on the current cached frame, return
 		if (FMath::IsNearlyEqual(CurrentCacheTime, DesiredCacheTime) && GlobalMatrices.Num() != 0)
 		{
-			DynamicData->PrevTransforms = GlobalMatrices;
-			DynamicData->Transforms = GlobalMatrices;
-
-			// maintaining the cache time means we should consider the transforms equal for dynamic data sending purposes
-			TransformsAreEqual[(TransformsAreEqualIndex++) % TransformsAreEqual.Num()] = true;
-
-			return;
+			return DynamicData;
 		}
 
-		// if the input simulation time to playback is the first frame, reset simulation time
-		if (DesiredCacheTime <= CacheDt ||  FMath::IsNearlyEqual(CurrentCacheTime, FLT_MAX))
-		{
-			CurrentCacheTime = DesiredCacheTime;
-		    
-			GeometryCollectionAlgo::GlobalMatrices(GetTransformArray(), GetParentArray(), GlobalMatrices);
-			DynamicData->PrevTransforms = GlobalMatrices;
-			DynamicData->Transforms = GlobalMatrices;			
+		// #todo(dmp): find a better place to calculate and store this
+		float CacheDt = CacheParameters.TargetCache->GetData()->GetDt();
 
-			EventsPlayed.Empty();
-			EventsPlayed.AddDefaulted(CacheParameters.TargetCache->GetData()->Records.Num());
+		// If we want the state before the first cached frame, then the collection doesn't need to be dynamic since it'll render the pre-fractured geometry
+		const bool bIsCacheDynamic = DesiredCacheTime > CacheDt;
 
-			// reset should send new transforms to the RT
-			TransformsAreEqual[(TransformsAreEqualIndex++) % TransformsAreEqual.Num()] = false;
-		}
-		else if (GlobalMatrices.Num() == 0)
+		// Bad case here.  Sequencer starts at non zero position  We cannot correctly reconstruct the current frame, so give a warning
+		const bool bOutOfSync = GlobalMatrices.Num() == 0;
+
+		// If the input simulation time to playback is the first frame, reset simulation time.
+		if (bOutOfSync || DesiredCacheTime <= CacheDt ||  FMath::IsNearlyEqual(CurrentCacheTime, FLT_MAX))
 		{
-			// bad case here.  Sequencer starts at non zero position  We cannot correctly reconstruct the current frame, so give a warning
 			GeometryCollectionAlgo::GlobalMatrices(GetTransformArray(), GetParentArray(), GlobalMatrices);
+
+			DynamicData = GDynamicDataPool.Allocate();
 			DynamicData->PrevTransforms = GlobalMatrices;
 			DynamicData->Transforms = GlobalMatrices;
+			DynamicData->ChangedCount = GlobalMatrices.Num();
+			DynamicData->IsDynamic = bIsCacheDynamic;
+			DynamicData->IsLoading = GetIsObjectLoading();
 
-			EventsPlayed.Empty();
+			EventsPlayed.Reset();
 			EventsPlayed.AddDefaulted(CacheParameters.TargetCache->GetData()->Records.Num());
 
-			// degenerate case causes and reset should send new transforms to the RT
-			TransformsAreEqual[(TransformsAreEqualIndex++) % TransformsAreEqual.Num()] = false;
-
-			UE_LOG(UGCC_LOG, Warning, TEXT("Cache out of sync - must rewind sequencer to start frame"));
+			if (bOutOfSync)
+			{
+				UE_LOG(UGCC_LOG, Warning, TEXT("Cache out of sync - must rewind sequencer to start frame"));
+			}
+			else
+			{
+				CurrentCacheTime = DesiredCacheTime;
+			}
 		}
 		else if (DesiredCacheTime >= CurrentCacheTime)
-		{	
-			int NumSteps = floor((DesiredCacheTime - CurrentCacheTime) / CacheDt);
+		{
+			const TManagedArray<FTransform>& Transform = RestCollection->GetGeometryCollection()->Transform;
+			const TManagedArray<TSet<int32>>& Children = GetChildrenArray();
+			const TManagedArray<FTransform>& MassToLocal = RestCollection->GetGeometryCollection()->GetAttribute<FTransform>("MassToLocal", FGeometryCollection::TransformGroup);
+
+			int32 NumSteps = floor((DesiredCacheTime - CurrentCacheTime) / CacheDt);
 			float LastDt = FMath::Fmod(DesiredCacheTime - CurrentCacheTime, CacheDt);
 			NumSteps += LastDt > SMALL_NUMBER ? 1 : 0;
 			
 			FTransform ActorToWorld = GetComponentTransform();
 
+			DynamicData = GDynamicDataPool.Allocate();
+			DynamicData->ChangedCount = 0;
+			DynamicData->IsDynamic = bIsCacheDynamic;
+			DynamicData->IsLoading = GetIsObjectLoading();
+
 			bool HasAnyActiveTransforms = false;
 
 			// Jump ahead in increments of CacheDt evaluating the cache until we reach our desired time
-			for (int st = 0; st < NumSteps; ++st)
+			for (int32 Step = 0; Step < NumSteps; ++Step)
 			{
-				float TimeIncrement = st == NumSteps - 1 ? LastDt : CacheDt;
+				float TimeIncrement = Step == NumSteps - 1 ? LastDt : CacheDt;
 				CurrentCacheTime += TimeIncrement;
 			
 				DynamicData->PrevTransforms = GlobalMatrices;
@@ -1379,35 +1372,39 @@ void UGeometryCollectionComponent::InitDynamicData(FGeometryCollectionDynamicDat
 
 				if (FirstFrame && !SecondFrame)
 				{
+					const TArray<FTransform>& FirstFrameTransforms = FirstFrame->Transforms;
+					const TArray<int32>& TransformIndices = FirstFrame->TransformIndices;
 
-					const TArray<FTransform> &xforms = FirstFrame->Transforms;
-					const TArray<int32> &TransformIndices = FirstFrame->TransformIndices;
-
-					const int32 NumActives = FirstFrame->TransformIndices.Num();
+					const int32 NumActives = TransformIndices.Num();
 					
 					if (NumActives > 0)
 					{
 						HasAnyActiveTransforms = true;
 					}
 
-					for (int i = 0; i < NumActives; ++i)
+					for (int32 TransformIndex = 0; TransformIndex < NumActives; ++TransformIndex)
 					{
-						const int32 InternalIndexTmp = TransformIndices[i];
+						const int32 InternalIndexTmp = TransformIndices[TransformIndex];
 						
 						if (InternalIndexTmp >= GlobalMatrices.Num())
 						{
-							UE_LOG(UGCC_LOG, Error, 
-								TEXT("%s: TargetCache (%s) is out of sync with GeometryCollection.  Regenerate the cache."), 
-								*RestCollection->GetName(), *CacheParameters.TargetCache->GetName());
+							UE_LOG
+							(
+								UGCC_LOG, Error,
+								TEXT("%s: TargetCache (%s) is out of sync with GeometryCollection.  Regenerate the cache."),
+								*RestCollection->GetName(), *CacheParameters.TargetCache->GetName()
+							);
+
 							DynamicData->PrevTransforms = GlobalMatrices;
 							DynamicData->Transforms = GlobalMatrices;
-							return;
+							DynamicData->ChangedCount = GlobalMatrices.Num();
+							return DynamicData;
 						}
-												
+
 						// calculate global matrix for current
-						FTransform ParticleToWorld = xforms[i];						
+						FTransform ParticleToWorld = FirstFrameTransforms[TransformIndex];
 
-						FTransform CurrGlobalTransform = MassToLocal[InternalIndexTmp].GetRelativeTransformReverse(ParticleToWorld).GetRelativeTransform(ActorToWorld);						
+						FTransform CurrGlobalTransform = MassToLocal[InternalIndexTmp].GetRelativeTransformReverse(ParticleToWorld).GetRelativeTransform(ActorToWorld);
 						CurrGlobalTransform.NormalizeRotation();
 						GlobalMatrices[InternalIndexTmp] = CurrGlobalTransform.ToMatrixWithScale();
 						
@@ -1418,7 +1415,7 @@ void UGeometryCollectionComponent::InitDynamicData(FGeometryCollectionDynamicDat
 				else if (FirstFrame && SecondFrame && CurrentCacheTime > FirstFrame->Timestamp)
 				{
 					const float Alpha = (CurrentCacheTime - FirstFrame->Timestamp) / (SecondFrame->Timestamp - FirstFrame->Timestamp);
-					check(0 <= Alpha && Alpha <= 1.0f);
+					checkSlow(0 <= Alpha && Alpha <= 1.0f);
 
 					const int32 NumActives = SecondFrame->TransformIndices.Num();
 
@@ -1427,19 +1424,23 @@ void UGeometryCollectionComponent::InitDynamicData(FGeometryCollectionDynamicDat
 						HasAnyActiveTransforms = true;
 					}
 
-					for (int Index = 0; Index < NumActives; ++Index)
+					for (int32 Index = 0; Index < NumActives; ++Index)
 					{
 						const int32 InternalIndexTmp = SecondFrame->TransformIndices[Index];
 						
-						// check if transform index is valid
+						// Check if transform index is valid
 						if (InternalIndexTmp >= GlobalMatrices.Num())
 						{
-							UE_LOG(UGCC_LOG, Error, 
+							UE_LOG
+							(
+								UGCC_LOG, Error,
 								TEXT("%s: TargetCache (%s) is out of sync with GeometryCollection.  Regenerate the cache."), 
-								*RestCollection->GetName(), *CacheParameters.TargetCache->GetName());
+								*RestCollection->GetName(), *CacheParameters.TargetCache->GetName()
+							);
 							DynamicData->PrevTransforms = GlobalMatrices;
 							DynamicData->Transforms = GlobalMatrices;
-							return;
+							DynamicData->ChangedCount = GlobalMatrices.Num();
+							return DynamicData;
 						}
 
 						const int32 PreviousIndexSlot = Index < SecondFrame->PreviousTransformIndices.Num() ? SecondFrame->PreviousTransformIndices[Index] : INDEX_NONE;
@@ -1449,7 +1450,7 @@ void UGeometryCollectionComponent::InitDynamicData(FGeometryCollectionDynamicDat
 							FTransform ParticleToWorld;
 							ParticleToWorld.Blend(FirstFrame->Transforms[PreviousIndexSlot], SecondFrame->Transforms[Index], Alpha);
 							
-							FTransform CurrGlobalTransform = MassToLocal[InternalIndexTmp].GetRelativeTransformReverse(ParticleToWorld).GetRelativeTransform(ActorToWorld);							
+							FTransform CurrGlobalTransform = MassToLocal[InternalIndexTmp].GetRelativeTransformReverse(ParticleToWorld).GetRelativeTransform(ActorToWorld);
 							CurrGlobalTransform.NormalizeRotation();
 							GlobalMatrices[InternalIndexTmp] = CurrGlobalTransform.ToMatrixWithScale();
 
@@ -1460,7 +1461,7 @@ void UGeometryCollectionComponent::InitDynamicData(FGeometryCollectionDynamicDat
 						{
 							FTransform ParticleToWorld = SecondFrame->Transforms[Index];
 							
-							FTransform CurrGlobalTransform = MassToLocal[InternalIndexTmp].GetRelativeTransformReverse(ParticleToWorld).GetRelativeTransform(ActorToWorld);							
+							FTransform CurrGlobalTransform = MassToLocal[InternalIndexTmp].GetRelativeTransformReverse(ParticleToWorld).GetRelativeTransform(ActorToWorld);
 							CurrGlobalTransform.NormalizeRotation();
 							GlobalMatrices[InternalIndexTmp] = CurrGlobalTransform.ToMatrixWithScale();
 
@@ -1472,31 +1473,28 @@ void UGeometryCollectionComponent::InitDynamicData(FGeometryCollectionDynamicDat
 				DynamicData->Transforms = GlobalMatrices;
 			}
 
-			// check if transforms at start of this tick are the same as what is calculated from the cache
-			TransformsAreEqual[(TransformsAreEqualIndex++) % TransformsAreEqual.Num()] = !HasAnyActiveTransforms;
-		}
-		else
-		{			
-			// time is before current cache time so maintain the matrices we have since we can't rewind
-			DynamicData->PrevTransforms = GlobalMatrices;
-			DynamicData->Transforms = GlobalMatrices;
-
-			// reset event means we don't want to consider transforms as being equal between prev and current frame
-			TransformsAreEqual[(TransformsAreEqualIndex++) % TransformsAreEqual.Num()] = true;
+			// Check if transforms at start of this tick are the same as what is calculated from the cache
+			if (HasAnyActiveTransforms)
+			{
+				DynamicData->DetermineChanges();
+			}
 		}
 	}
-	else if (DynamicData->IsDynamic)
+	else if (bIsDynamic)
 	{
+		DynamicData = GDynamicDataPool.Allocate();
+		DynamicData->IsDynamic = true;
+		DynamicData->IsLoading = GetIsObjectLoading();
+
 		// If we have no transforms stored in the dynamic data, then assign both prev and current to the same global matrices
 		if (GlobalMatrices.Num() == 0)
 		{
 			// Copy global matrices over to DynamicData
-			CalculateGlobalMatrices();		
-			DynamicData->PrevTransforms = GlobalMatrices;		
-			DynamicData->Transforms = GlobalMatrices;
+			CalculateGlobalMatrices();
 
-			// reset event means we don't want to consider transforms as being equal between prev and current frame
-			TransformsAreEqual[(TransformsAreEqualIndex++) % TransformsAreEqual.Num()] = false;
+			DynamicData->PrevTransforms = GlobalMatrices;
+			DynamicData->Transforms = GlobalMatrices;
+			DynamicData->ChangedCount = GlobalMatrices.Num();
 		}
 		else
 		{
@@ -1506,18 +1504,35 @@ void UGeometryCollectionComponent::InitDynamicData(FGeometryCollectionDynamicDat
 			// Copy global matrices over to DynamicData
 			CalculateGlobalMatrices();
 
+			bool bComputeChanges = true;
+
 			// if the number of matrices has changed between frames, then sync previous to current
 			if (GlobalMatrices.Num() != DynamicData->PrevTransforms.Num())
 			{
 				DynamicData->PrevTransforms = GlobalMatrices;
+				DynamicData->ChangedCount = GlobalMatrices.Num();
+				bComputeChanges = false; // Optimization to just force all transforms as changed and skip comparison
 			}
 
 			DynamicData->Transforms = GlobalMatrices;
 
-			// check if previous transforms are the same as current
-			TransformsAreEqual[(TransformsAreEqualIndex++) % TransformsAreEqual.Num()] = IsEqual(DynamicData->PrevTransforms, DynamicData->Transforms);
+			// The number of transforms for current and previous should match now
+			check(DynamicData->PrevTransforms.Num() == DynamicData->Transforms.Num());
+
+			if (bComputeChanges)
+			{
+				DynamicData->DetermineChanges();
+			}
 		}
 	}
+
+	if (DynamicData && DynamicData->ChangedCount == 0 && !bEditorMode)
+	{
+		GDynamicDataPool.Release(DynamicData);
+		DynamicData = nullptr;
+	}
+
+	return DynamicData;
 }
 
 void UGeometryCollectionComponent::OnUpdateTransform(EUpdateTransformFlags UpdateTransformFlags, ETeleportType Teleport)
@@ -2002,30 +2017,15 @@ void UGeometryCollectionComponent::SendRenderDynamicData_Concurrent()
 	// Only update the dynamic data if the dynamic collection is dirty
 	if (SceneProxy && ((DynamicCollection && DynamicCollection->IsDirty()) || CachePlayback)) 
 	{
-		FGeometryCollectionDynamicData * DynamicData = ::new FGeometryCollectionDynamicData;
-		InitDynamicData(DynamicData);
+		FGeometryCollectionDynamicData* DynamicData = InitDynamicData();
 
-		//
-		// Only send dynamic data if the transform arrays on the past N frames are different
-		//
-		bool PastFramesTransformsAreEqual = true;
-		for (int32 TransformIndex = 0; TransformIndex < TransformsAreEqual.Num(); ++TransformIndex)
+		if (DynamicData)
 		{
-			if (!TransformsAreEqual[TransformIndex])
-			{
-				PastFramesTransformsAreEqual = false;
-				break;
-			}
-		}
+			INC_DWORD_STAT_BY(STAT_GCTotalTransforms, DynamicData->Transforms.Num());
+			INC_DWORD_STAT_BY(STAT_GCChangedTransforms, DynamicData->ChangedCount);
 
-		if (PastFramesTransformsAreEqual)
-		{
-			delete DynamicData;
-		}
-		else
-		{
 			// #todo (bmiller) Once ISMC changes have been complete, this is the best place to call this method
-			// but we can't currently because it's an inappropraite place to call MarkRenderStateDirty on the ISMC.
+			// but we can't currently because it's an inappropriate place to call MarkRenderStateDirty on the ISMC.
 			// RefreshEmbeddedGeometry();
 
 			// Enqueue command to send to render thread
@@ -2035,11 +2035,8 @@ void UGeometryCollectionComponent::SendRenderDynamicData_Concurrent()
 				ENQUEUE_RENDER_COMMAND(SendRenderDynamicData)(
 					[GeometryCollectionSceneProxy, DynamicData](FRHICommandListImmediate& RHICmdList)
 					{
-						if (GeometryCollectionSceneProxy)
-						{
-							GeometryCollectionSceneProxy->SetDynamicData_RenderThread(DynamicData);
-							GeometryCollectionSceneProxy->GetPrimitiveSceneInfo()->RequestGPUSceneUpdate();
-						}
+						GeometryCollectionSceneProxy->SetDynamicData_RenderThread(DynamicData);
+						GeometryCollectionSceneProxy->GetPrimitiveSceneInfo()->RequestGPUSceneUpdate();
 					}
 				);
 			}
@@ -2055,8 +2052,6 @@ void UGeometryCollectionComponent::SendRenderDynamicData_Concurrent()
 						}
 					}
 				);
-
-				
 			}
 		}
 
@@ -2635,7 +2630,7 @@ void UGeometryCollectionComponent::CalculateGlobalMatrices()
 	if(NumTransforms > 0)
 	{
 		// Just calc from results
-		GlobalMatrices.Empty();
+		GlobalMatrices.Reset();
 		GlobalMatrices.Append(Results->GlobalTransforms);	
 	}
 	else
@@ -2727,24 +2722,6 @@ void UGeometryCollectionComponent::SwitchRenderModels(const AActor* Actor)
 
 }
 
-bool UGeometryCollectionComponent::IsEqual(const TArray<FMatrix> &A, const TArray<FMatrix> &B, const float Tolerance)
-{
-	if (A.Num() != B.Num())
-	{
-		return false;
-	}
-
-	for (int Index = 0; Index < A.Num(); ++Index)
-	{
-		if (!A[Index].Equals(B[Index], Tolerance))
-		{
-			return false;
-		}
-	}
-
-	return true;
-}
-
 #if GEOMETRYCOLLECTION_EDITOR_SELECTION
 void UGeometryCollectionComponent::EnableTransformSelectionMode(bool bEnable)
 {
diff --git a/Engine/Source/Runtime/Experimental/GeometryCollectionEngine/Private/GeometryCollection/GeometryCollectionSceneProxy.cpp b/Engine/Source/Runtime/Experimental/GeometryCollectionEngine/Private/GeometryCollection/GeometryCollectionSceneProxy.cpp
--- a/Engine/Source/Runtime/Experimental/GeometryCollectionEngine/Private/GeometryCollection/GeometryCollectionSceneProxy.cpp
+++ b/Engine/Source/Runtime/Experimental/GeometryCollectionEngine/Private/GeometryCollection/GeometryCollectionSceneProxy.cpp
@@ -54,6 +54,8 @@ FAutoConsoleVariableRef CVarGeometryCollectionTripleBufferUploads(
 
 DEFINE_LOG_CATEGORY_STATIC(FGeometryCollectionSceneProxyLogging, Log, All);
 
+FGeometryCollectionDynamicDataPool GDynamicDataPool;
+
 FGeometryCollectionSceneProxy::FGeometryCollectionSceneProxy(UGeometryCollectionComponent* Component)
 	: FPrimitiveSceneProxy(Component)
 	, MaterialRelevance(Component->GetMaterialRelevance(GetScene().GetFeatureLevel()))
@@ -476,7 +478,7 @@ void FGeometryCollectionSceneProxy::SetDynamicData_RenderThread(FGeometryCollect
 	{
 		if (DynamicData)
 		{
-			delete DynamicData;
+			GDynamicDataPool.Release(DynamicData);
 			DynamicData = nullptr;
 		}
 		DynamicData = NewDynamicData;
@@ -662,7 +664,8 @@ void FGeometryCollectionSceneProxy::GetDynamicMeshElements(const TArray<const FS
 						
 			// render original mesh if it isn't dynamic and there is an unfractured mesh
 			// #todo(dmp): refactor this to share more code later
-			if (!DynamicData->IsDynamic)
+			const bool bIsDynamic = DynamicData && DynamicData->IsDynamic;
+			if (!bIsDynamic)
 			{
 			#if GEOMETRYCOLLECTION_EDITOR_SELECTION
 				const TArray<FGeometryCollectionSection>& SectionArray = bUsesSubSections && SubSections.Num() ? SubSections: ConstantData->OriginalMeshSections;
@@ -1313,5 +1316,65 @@ void FNaniteGeometryCollectionSceneProxy::SetDynamicData_RenderThread(FGeometryC
 		// ...
 	}
 
-	delete NewDynamicData;
+	GDynamicDataPool.Release(NewDynamicData);
+}
+
+FGeometryCollectionDynamicDataPool::FGeometryCollectionDynamicDataPool()
+{
+	FreeList.SetNum(32);
+	for (int32 ListIndex = 0; ListIndex < FreeList.Num(); ++ListIndex)
+	{
+		FreeList[ListIndex] = new FGeometryCollectionDynamicData;
+	}
+}
+
+FGeometryCollectionDynamicDataPool::~FGeometryCollectionDynamicDataPool()
+{
+	FScopeLock ScopeLock(&ListLock);
+
+	for (FGeometryCollectionDynamicData* Entry : FreeList)
+	{
+		delete Entry;
+	}
+
+	for (FGeometryCollectionDynamicData* Entry : UsedList)
+	{
+		delete Entry;
+	}
+
+	FreeList.Empty();
+	UsedList.Empty();
+}
+
+FGeometryCollectionDynamicData* FGeometryCollectionDynamicDataPool::Allocate()
+{
+	FScopeLock ScopeLock(&ListLock);
+
+	FGeometryCollectionDynamicData* NewEntry = nullptr;
+	if (FreeList.Num() > 0)
+	{
+		NewEntry = FreeList.Pop(false /* no shrinking */);
+	}
+
+	if (NewEntry == nullptr)
+	{
+		NewEntry = new FGeometryCollectionDynamicData;
+	}
+
+	NewEntry->Reset();
+	UsedList.Push(NewEntry);
+
+	return NewEntry;
+}
+
+void FGeometryCollectionDynamicDataPool::Release(FGeometryCollectionDynamicData* DynamicData)
+{
+	FScopeLock ScopeLock(&ListLock);
+
+	int32 UsedIndex = UsedList.Find(DynamicData);
+	if (ensure(UsedIndex != INDEX_NONE))
+	{
+		UsedList.RemoveAt(UsedIndex, 1, false /* no shrinking */);
+		FreeList.Push(DynamicData);
+	}
 }
diff --git a/Engine/Source/Runtime/Experimental/GeometryCollectionEngine/Private/GeometryCollection/GeometryCollectionSceneProxy.h b/Engine/Source/Runtime/Experimental/GeometryCollectionEngine/Private/GeometryCollection/GeometryCollectionSceneProxy.h
--- a/Engine/Source/Runtime/Experimental/GeometryCollectionEngine/Private/GeometryCollection/GeometryCollectionSceneProxy.h
+++ b/Engine/Source/Runtime/Experimental/GeometryCollectionEngine/Private/GeometryCollection/GeometryCollectionSceneProxy.h
@@ -3,6 +3,7 @@
 
 #include "DynamicMeshBuilder.h"
 #include "EngineGlobals.h"
+#include "HAL/CriticalSection.h"
 #include "PrimitiveViewRelevance.h"
 #include "PrimitiveSceneProxy.h"
 #include "StaticMeshResources.h"
@@ -106,10 +107,61 @@ struct FGeometryCollectionDynamicData
 {
 	TArray<FMatrix> Transforms;
 	TArray<FMatrix> PrevTransforms;
-	bool IsDynamic;
-	bool IsLoading;
+	uint32 ChangedCount;
+	uint8 IsDynamic : 1;
+	uint8 IsLoading : 1;
 
-	FGeometryCollectionDynamicData() : IsDynamic(false) {}
+	FGeometryCollectionDynamicData()
+	{
+		Reset();
+	}
+
+	void Reset()
+	{
+		Transforms.Reset();
+		PrevTransforms.Reset();
+		IsDynamic = false;
+		IsLoading = false;
+	}
+
+	void DetermineChanges()
+	{
+		// Check if previous transforms are the same as current
+		const float EqualTolerance = 1e-6;
+
+		check(Transforms.Num() == PrevTransforms.Num());
+		if (Transforms.Num() != PrevTransforms.Num())
+		{
+			ChangedCount = Transforms.Num();
+		}
+		else
+		{
+			ChangedCount = 0;
+			for (int32 TransformIndex = 0; TransformIndex < Transforms.Num(); ++TransformIndex)
+			{
+				if (!PrevTransforms[TransformIndex].Equals(Transforms[TransformIndex], EqualTolerance))
+				{
+					++ChangedCount;
+				}
+			}
+		}
+	}
+};
+
+class FGeometryCollectionDynamicDataPool
+{
+public:
+	FGeometryCollectionDynamicDataPool();
+	~FGeometryCollectionDynamicDataPool();
+
+	FGeometryCollectionDynamicData* Allocate();
+	void Release(FGeometryCollectionDynamicData* DynamicData);
+
+private:
+	TArray<FGeometryCollectionDynamicData*> UsedList;
+	TArray<FGeometryCollectionDynamicData*> FreeList;
+
+	FCriticalSection ListLock;
 };
 
 /***
diff --git a/Engine/Source/Runtime/Experimental/GeometryCollectionEngine/Public/GeometryCollection/GeometryCollectionComponent.h b/Engine/Source/Runtime/Experimental/GeometryCollectionEngine/Public/GeometryCollection/GeometryCollectionComponent.h
--- a/Engine/Source/Runtime/Experimental/GeometryCollectionEngine/Public/GeometryCollection/GeometryCollectionComponent.h
+++ b/Engine/Source/Runtime/Experimental/GeometryCollectionEngine/Public/GeometryCollection/GeometryCollectionComponent.h
@@ -664,7 +664,7 @@ protected:
 	void InitConstantData(FGeometryCollectionConstantData* ConstantData) const;
 
 	/** Populate the dynamic particle data for the render thread. */
-	void InitDynamicData(FGeometryCollectionDynamicData* ConstantData);
+	FGeometryCollectionDynamicData* InitDynamicData();
 
 	/** Reset the dynamic collection from the current rest state. */
 	void ResetDynamicCollection();
@@ -744,7 +744,7 @@ private:
 	TArray<FMatrix> GlobalMatrices;
 	FBox LocalBounds;
 	
-	FBoxSphereBounds WorldBounds;		
+	FBoxSphereBounds WorldBounds;
 
 	float CurrentCacheTime;
 	TArray<bool> EventsPlayed;
@@ -771,10 +771,6 @@ private:
 #endif
 	void SwitchRenderModels(const AActor* Actor);
 
-	bool IsEqual(const TArray<FMatrix> &A, const TArray<FMatrix> &B, const float Tolerance = 1e-6);
-	TArray<bool> TransformsAreEqual;	
-	int32 TransformsAreEqualIndex;
-
 	UChaosGameplayEventDispatcher* EventDispatcher;
 
 #if GEOMETRYCOLLECTION_EDITOR_SELECTION
