Chaos - Fix capsule-capsule and capsule-sphere collision in the non-manifold path

diff --git a/Engine/Source/Runtime/Experimental/Chaos/Private/Chaos/Collision/PBDCollisionConstraint.cpp b/Engine/Source/Runtime/Experimental/Chaos/Private/Chaos/Collision/PBDCollisionConstraint.cpp
--- a/Engine/Source/Runtime/Experimental/Chaos/Private/Chaos/Collision/PBDCollisionConstraint.cpp
+++ b/Engine/Source/Runtime/Experimental/Chaos/Private/Chaos/Collision/PBDCollisionConstraint.cpp
@@ -173,6 +173,11 @@ namespace Chaos
 				AddManifoldPoint(ContactPoint, Dt);
 			}
 
+			// Copy currently active point
+			if (ContactPoint.Phi < Manifold.Phi)
+			{
+				SetActiveContactPoint(ContactPoint);
+			}
 		}
 		else 
 		{
@@ -187,11 +192,7 @@ namespace Chaos
 			}
 
 			InitManifoldPoint(ManifoldPoints[0], Dt);
-		}
 
-		// Copy currently active point
-		if (ContactPoint.Phi < Manifold.Phi)
-		{
 			SetActiveContactPoint(ContactPoint);
 		}
 	}
diff --git a/Engine/Source/Runtime/Experimental/Chaos/Private/Chaos/CollisionResolution.cpp b/Engine/Source/Runtime/Experimental/Chaos/Private/Chaos/CollisionResolution.cpp
--- a/Engine/Source/Runtime/Experimental/Chaos/Private/Chaos/CollisionResolution.cpp
+++ b/Engine/Source/Runtime/Experimental/Chaos/Private/Chaos/CollisionResolution.cpp
@@ -1110,17 +1110,18 @@ namespace Chaos
 			{
 				FReal NewPhi = DeltaLen - (A.GetRadius() + B.GetRadius()) - ShapePadding;
 				FVec3 Dir = Delta / DeltaLen;
-				FVec3 Location = A1 + Dir * A.GetRadius();
+				FVec3 LocationA = A1 + Dir * A.GetRadius();
+				FVec3 LocationB = P2 - Dir * B.GetRadius();
 				FVec3 Normal = -Dir;
 				// @todo(chaos): margin
 				Result.ShapeMargins[0] = 0.0f;
 				Result.ShapeMargins[1] = 0.0f;
-				Result.ShapeContactPoints[0] = ATransform.InverseTransformPosition(Location);
-				Result.ShapeContactPoints[1] = BTransform.InverseTransformPosition(Location);
+				Result.ShapeContactPoints[0] = ATransform.InverseTransformPosition(LocationA);
+				Result.ShapeContactPoints[1] = BTransform.InverseTransformPosition(LocationB);
 				Result.ShapeContactNormal = BTransform.InverseTransformVector(Normal);
 				Result.Phi = NewPhi;
 				Result.Normal = Normal;
-				Result.Location = Location;
+				Result.Location = 0.5f * (LocationA + LocationB);
 			}
 
 			return Result;
@@ -1307,16 +1308,17 @@ namespace Chaos
 				FReal NewPhi = DeltaLen - (A.GetRadius() + B.GetRadius()) - ShapePadding;
 				FVec3 Dir = Delta / DeltaLen;
 				FVec3 Normal = -Dir;
-				FVec3 Location = P1 + Dir * A.GetRadius();
+				FVec3 LocationA = P1 + Dir * A.GetRadius();
+				FVec3 LocationB = P2 - Dir * B.GetRadius();
 				// @todo(chaos): margin
 				Result.ShapeMargins[0] = 0.0f;
 				Result.ShapeMargins[1] = 0.0f;
-				Result.ShapeContactPoints[0] = ATransform.InverseTransformPosition(Location);
-				Result.ShapeContactPoints[1] = BTransform.InverseTransformPosition(Location);
+				Result.ShapeContactPoints[0] = ATransform.InverseTransformPosition(LocationA);
+				Result.ShapeContactPoints[1] = BTransform.InverseTransformPosition(LocationB);
 				Result.ShapeContactNormal = BTransform.InverseTransformVector(Normal);
 				Result.Phi = NewPhi;
 				Result.Normal = Normal;
-				Result.Location = Location;
+				Result.Location = 0.5f * (LocationA + LocationB);
 			}
 
 			return Result;
