Chaos : Particle Deletion
- PVAR to help debug partile lifespan issues.

diff --git a/Engine/Source/Runtime/Experimental/Chaos/Private/PBDRigidsSolver.cpp b/Engine/Source/Runtime/Experimental/Chaos/Private/PBDRigidsSolver.cpp
--- a/Engine/Source/Runtime/Experimental/Chaos/Private/PBDRigidsSolver.cpp
+++ b/Engine/Source/Runtime/Experimental/Chaos/Private/PBDRigidsSolver.cpp
@@ -159,6 +159,10 @@ FAutoConsoleVariableRef CVarChaosSolverCollisionUseManifolds(TEXT("p.Chaos.Solve
 int32 ChaosVisualDebuggerEnable = 1;
 FAutoConsoleVariableRef CVarChaosVisualDebuggerEnable(TEXT("p.Chaos.VisualDebuggerEnable"), ChaosVisualDebuggerEnable, TEXT("Enable/Disable pushing/saving data to the visual debugger"));
 
+bool bDisableParticleDeletion = false;
+FAutoConsoleVariableRef CVarDisableParticleDeletion(TEXT("p.Chaos.Solver.DisableParticleDeletion"), bDisableParticleDeletion, TEXT("Only ever use when debugging the particle clean up, this introduces a massive memory leak."));
+
+
 namespace Chaos
 {
 
@@ -179,6 +183,11 @@ namespace Chaos
 
 		void DoWork()
 		{
+			// WARNING
+			// Only ever use DisableParticleDeletion when debugging the particle clean up. 
+			// This introduces a massive memory leak.
+			MSolver->GetParticles().SetDisableParticleDeletion(bDisableParticleDeletion);
+
 			LLM_SCOPE(ELLMTag::Chaos);
 			UE_LOG(LogPBDRigidsSolver, Verbose, TEXT("AdvanceOneTimeStepTask::DoWork()"));
 			MSolver->StartingSceneSimulation();
@@ -485,7 +494,6 @@ namespace Chaos
 			}
 
 		});
-
 	}
 
 	void FPBDRigidsSolver::RegisterObject(FGeometryCollectionPhysicsProxy* InProxy)
diff --git a/Engine/Source/Runtime/Experimental/Chaos/Public/Chaos/PBDRigidsSOAs.h b/Engine/Source/Runtime/Experimental/Chaos/Public/Chaos/PBDRigidsSOAs.h
--- a/Engine/Source/Runtime/Experimental/Chaos/Public/Chaos/PBDRigidsSOAs.h
+++ b/Engine/Source/Runtime/Experimental/Chaos/Public/Chaos/PBDRigidsSOAs.h
@@ -176,6 +176,7 @@ class FPBDRigidsSOAs
 {
 public:
 	FPBDRigidsSOAs()
+		: bDisableParticleDeletion(false)
 	{
 #if CHAOS_DETERMINISTIC
 		BiggestParticleID = 0;
@@ -296,6 +297,12 @@ public:
 		UpdateViews();
 	}
 
+	// WARNING
+	// Only ever use DisableParticleDeletion when debugging the particle clean up. 
+	// This introduces a massive memory leak.
+	bool bDisableParticleDeletion;
+	void SetDisableParticleDeletion(bool bIn) { bDisableParticleDeletion = bIn; }
+
 	void DestroyParticle(FGeometryParticleHandle* Particle)
 	{
 		auto PBDRigid = Particle->CastToRigidParticle();
@@ -333,8 +340,14 @@ public:
 			GetDynamicParticles().GetSleepDataLock().WriteUnlock();
 		}
 
-		ParticleHandles.DestroyHandleSwap(Particle);
-		
+		// WARNING
+		// Only ever use DisableParticleDeletion when debugging the particle clean up. 
+		// This introduces a massive memory leak.
+		if (!bDisableParticleDeletion)
+		{
+			ParticleHandles.DestroyHandleSwap(Particle);
+		}
+
 		UpdateViews();
 	}
 
