Chaos: Constraint Crash Fix

diff --git a/Engine/Source/Runtime/Experimental/Chaos/Private/Chaos/PBDSuspensionConstraints.cpp b/Engine/Source/Runtime/Experimental/Chaos/Private/Chaos/PBDSuspensionConstraints.cpp
--- a/Engine/Source/Runtime/Experimental/Chaos/Private/Chaos/PBDSuspensionConstraints.cpp
+++ b/Engine/Source/Runtime/Experimental/Chaos/Private/Chaos/PBDSuspensionConstraints.cpp
@@ -49,6 +49,7 @@ namespace Chaos
 		ConstrainedParticles.Add(Particle);
 		SuspensionLocalOffset.Add(InSuspensionLocalOffset);
 		ConstraintSettings.Add(InConstraintSettings);
+		ConstraintEnabledStates.Add(true); // Note: assumes always enabled on creation
 		Handles.Add(HandleAllocator.AllocHandle(this, NewIndex));
 		return Handles[NewIndex];
 	}
@@ -68,6 +69,7 @@ namespace Chaos
 		ConstrainedParticles.RemoveAtSwap(ConstraintIndex);
 		SuspensionLocalOffset.RemoveAtSwap(ConstraintIndex);
 		ConstraintSettings.RemoveAtSwap(ConstraintIndex);
+		ConstraintEnabledStates.RemoveAtSwap(ConstraintIndex);
 		Handles.RemoveAtSwap(ConstraintIndex);
 
 		// Update the handle for the constraint that was moved
diff --git a/Engine/Source/Runtime/Experimental/Chaos/Public/Chaos/PBDSuspensionConstraints.h b/Engine/Source/Runtime/Experimental/Chaos/Public/Chaos/PBDSuspensionConstraints.h
--- a/Engine/Source/Runtime/Experimental/Chaos/Public/Chaos/PBDSuspensionConstraints.h
+++ b/Engine/Source/Runtime/Experimental/Chaos/Public/Chaos/PBDSuspensionConstraints.h
@@ -119,6 +119,16 @@ namespace Chaos
 			}
 		}
 
+		bool IsConstraintEnabled(int32 ConstraintIndex) const
+		{
+			return ConstraintEnabledStates[ConstraintIndex];
+		}
+
+		void SetConstraintEnabled(int32 ConstraintIndex, bool bEnabled)
+		{
+			ConstraintEnabledStates[ConstraintIndex] = bEnabled;
+		}
+
 		//
 		// Constraint API
 		//
@@ -218,6 +228,7 @@ namespace Chaos
 		TArray<FGeometryParticleHandle*> ConstrainedParticles;
 		TArray<FVec3> SuspensionLocalOffset;
 		TArray<FPBDSuspensionSettings> ConstraintSettings;
+		TArray<bool> ConstraintEnabledStates;
 
 		FHandles Handles;
 		FConstraintHandleAllocator HandleAllocator;
