Do not skip InitDynamicData result if called from geometry collection initialization. Fixes corrupt rendering in TM-ChaosDestruction.

diff --git a/Engine/Source/Runtime/Experimental/GeometryCollectionEngine/Private/GeometryCollection/GeometryCollectionComponent.cpp b/Engine/Source/Runtime/Experimental/GeometryCollectionEngine/Private/GeometryCollection/GeometryCollectionComponent.cpp
--- a/Engine/Source/Runtime/Experimental/GeometryCollectionEngine/Private/GeometryCollection/GeometryCollectionComponent.cpp
+++ b/Engine/Source/Runtime/Experimental/GeometryCollectionEngine/Private/GeometryCollection/GeometryCollectionComponent.cpp
@@ -501,7 +501,7 @@ FPrimitiveSceneProxy* UGeometryCollectionComponent::CreateSceneProxy()
 			FGeometryCollectionConstantData* const ConstantData = ::new FGeometryCollectionConstantData;
 			InitConstantData(ConstantData);
 
-			FGeometryCollectionDynamicData* const DynamicData = InitDynamicData();
+			FGeometryCollectionDynamicData* const DynamicData = InitDynamicData(true /* initialization */);
 
 			if (LocalSceneProxy->IsNaniteMesh())
 			{
@@ -1286,17 +1286,17 @@ void UGeometryCollectionComponent::InitConstantData(FGeometryCollectionConstantD
 	TArray<FMatrix> RestMatrices;
 	GeometryCollectionAlgo::GlobalMatrices(RestCollection->GetGeometryCollection()->Transform, RestCollection->GetGeometryCollection()->Parent, RestMatrices);
 
-	ConstantData->RestTransforms = MoveTemp(RestMatrices); 
+	ConstantData->RestTransforms = MoveTemp(RestMatrices);
 }
 
-FGeometryCollectionDynamicData* UGeometryCollectionComponent::InitDynamicData()
+FGeometryCollectionDynamicData* UGeometryCollectionComponent::InitDynamicData(bool bInitialization)
 {
 	SCOPE_CYCLE_COUNTER(STAT_GCInitDynamicData);
 
 	FGeometryCollectionDynamicData* DynamicData = nullptr;
 
 	const bool bEditorMode = bShowBoneColors || bEnableBoneSelection;
-	const bool bIsDynamic = GetIsObjectDynamic() || bEditorMode;
+	const bool bIsDynamic  = GetIsObjectDynamic() || bEditorMode || bInitialization;
 
 	if (CachePlayback && CacheParameters.TargetCache)
 	{
@@ -1526,7 +1526,7 @@ FGeometryCollectionDynamicData* UGeometryCollectionComponent::InitDynamicData()
 		}
 	}
 
-	if (DynamicData && DynamicData->ChangedCount == 0 && !bEditorMode)
+	if (DynamicData && DynamicData->ChangedCount == 0 && !bEditorMode && !bInitialization)
 	{
 		GDynamicDataPool.Release(DynamicData);
 		DynamicData = nullptr;
@@ -2017,7 +2017,7 @@ void UGeometryCollectionComponent::SendRenderDynamicData_Concurrent()
 	// Only update the dynamic data if the dynamic collection is dirty
 	if (SceneProxy && ((DynamicCollection && DynamicCollection->IsDirty()) || CachePlayback)) 
 	{
-		FGeometryCollectionDynamicData* DynamicData = InitDynamicData();
+		FGeometryCollectionDynamicData* DynamicData = InitDynamicData(false /* initialization */);
 
 		if (DynamicData)
 		{
diff --git a/Engine/Source/Runtime/Experimental/GeometryCollectionEngine/Public/GeometryCollection/GeometryCollectionComponent.h b/Engine/Source/Runtime/Experimental/GeometryCollectionEngine/Public/GeometryCollection/GeometryCollectionComponent.h
--- a/Engine/Source/Runtime/Experimental/GeometryCollectionEngine/Public/GeometryCollection/GeometryCollectionComponent.h
+++ b/Engine/Source/Runtime/Experimental/GeometryCollectionEngine/Public/GeometryCollection/GeometryCollectionComponent.h
@@ -664,7 +664,7 @@ protected:
 	void InitConstantData(FGeometryCollectionConstantData* ConstantData) const;
 
 	/** Populate the dynamic particle data for the render thread. */
-	FGeometryCollectionDynamicData* InitDynamicData();
+	FGeometryCollectionDynamicData* InitDynamicData(bool bInitialization);
 
 	/** Reset the dynamic collection from the current rest state. */
 	void ResetDynamicCollection();
