Chaos Cloth - Rework Long Range Attachment initializations. [Integrate CL 15987361 from //UE5/Release-Engine-Staging]

CIS - Fix Linux/Mac builds erroring on the TAtomic initialization. [Integrate CL 15987392 from //UE5/Release-Engine-Staging]

Chaos Cloth - Remove one unused line added by mistake during the large world coordinate merge. [Integrate CL 15989334 from //UE5/Release-Engine-Staging]

Chaos Cloth - Fix Long Range Attachement Euclidean distance that was missing the TetherView range update.  [Integrate CL 16022340 from //UE5/Release-Engine-Staging]

Chaos Cloth - Improve Long Range Attachment stiffness.  [Integrate CL 16044944 from //UE5/Release-Engine-Staging]

CIS - Fix  multiply defined symbols found on FPBDStiffness caused by LWC merge. [Integrate CL 16044951 from //UE5/Release-Engine-Staging]

diff --git a/Engine/Plugins/Experimental/ChaosCloth/Source/Chaos/Private/ChaosCloth/ChaosClothConfig.cpp b/Engine/Plugins/Experimental/ChaosCloth/Source/Chaos/Private/ChaosCloth/ChaosClothConfig.cpp
--- a/Engine/Plugins/Experimental/ChaosCloth/Source/Chaos/Private/ChaosCloth/ChaosClothConfig.cpp
+++ b/Engine/Plugins/Experimental/ChaosCloth/Source/Chaos/Private/ChaosCloth/ChaosClothConfig.cpp
@@ -22,7 +22,6 @@
 //  LinearDrag
 //  AngularDrag
 //  StiffnessFrequency
-//  TetherLimit
 
 UChaosClothConfig::UChaosClothConfig()
 {}
@@ -32,6 +31,7 @@ UChaosClothConfig::~UChaosClothConfig()
 
 void UChaosClothConfig::MigrateFrom(const FClothConfig_Legacy& ClothConfig)
 {
+#if WITH_EDITORONLY_DATA
 	const float VerticalStiffness =
 		ClothConfig.VerticalConstraintConfig.Stiffness *
 		ClothConfig.VerticalConstraintConfig.StiffnessMultiplier;
@@ -59,7 +59,8 @@ void UChaosClothConfig::MigrateFrom(const FClothConfig_Legacy& ClothConfig)
 	bUseBendingElements = false;
 	bUseSelfCollisions = (ClothConfig.SelfCollisionRadius > 0.f && ClothConfig.SelfCollisionStiffness > 0.f);
 
-	StrainLimitingStiffness = FMath::Clamp(ClothConfig.TetherStiffness, 0.f, 1.f);
+	TetherStiffness.Low = FMath::Clamp(ClothConfig.TetherStiffness, 0.f, 1.f);
+	TetherStiffness.High = 1.f;
 	LimitScale = FMath::Clamp(ClothConfig.TetherLimit, 0.01f, 10.f);
 	ShapeTargetStiffness = 0.f;
 
@@ -82,10 +83,12 @@ void UChaosClothConfig::MigrateFrom(const FClothConfig_Legacy& ClothConfig)
 	Gravity = ClothConfig.GravityOverride;
 
 	bUseLegacyBackstop = true;
+#endif  // #if WITH_EDITORONLY_DATA
 }
 
 void UChaosClothConfig::MigrateFrom(const UClothSharedConfigCommon* ClothSharedConfig)
 {
+#if WITH_EDITORONLY_DATA
 	if (const UChaosClothSharedSimConfig* const ChaosClothSharedSimConfig = Cast<UChaosClothSharedSimConfig>(ClothSharedConfig))
 	{
 		const int32 ChaosClothConfigCustomVersion = GetLinkerCustomVersion(FChaosClothConfigCustomVersion::GUID);
@@ -106,6 +109,7 @@ void UChaosClothConfig::MigrateFrom(const UClothSharedConfigCommon* ClothSharedC
 			Gravity = ChaosClothSharedSimConfig->Gravity_DEPRECATED;
 		}
 	}
+#endif  // #if WITH_EDITORONLY_DATA
 }
 
 void UChaosClothConfig::Serialize(FArchive& Ar)
@@ -119,6 +123,8 @@ void UChaosClothConfig::Serialize(FArchive& Ar)
 void UChaosClothConfig::PostLoad()
 {
 	Super::PostLoad();
+
+#if WITH_EDITORONLY_DATA
 	const int32 ChaosClothConfigCustomVersion = GetLinkerCustomVersion(FChaosClothConfigCustomVersion::GUID);
 	const int32 PhysicsObjectVersion = GetLinkerCustomVersion(FPhysicsObjectVersion::GUID);
 	const int32 FortniteMainBranchObjectVersion = GetLinkerCustomVersion(FFortniteMainBranchObjectVersion::GUID);
@@ -143,10 +149,26 @@ void UChaosClothConfig::PostLoad()
 		AnimDriveStiffness.Low = 0.f;
 		AnimDriveStiffness.High = FMath::Clamp(FMath::Loge(AnimDriveSpringStiffness_DEPRECATED) / FMath::Loge(1.e3f) + 1.f, 0.f, 1.f);
 	}
+
+	if (TetherMode_DEPRECATED != EChaosClothTetherMode::MaxChaosClothTetherMode)
+	{
+		// Note: MaxChaosClothTetherMode is used here to detect that the TetherMode parameter isn't set to its default value and therefore needs to be migrated.
+		bUseGeodesicDistance = (TetherMode_DEPRECATED != EChaosClothTetherMode::FastTetherFastLength);
+		TetherMode_DEPRECATED = EChaosClothTetherMode::MaxChaosClothTetherMode;
+	}
+
 	if (FortniteMainBranchObjectVersion < FFortniteMainBranchObjectVersion::ChaosClothAddfictitiousforces)
 	{
 		FictitiousAngularScale = 0.f;  // Maintain early behavior with no fictitious forces
 	}
+
+	if (PhysicsObjectVersion < FPhysicsObjectVersion::ChaosClothAddTetherStiffnessWeightMap)
+	{
+		// Note: Unlike AnimDriveStiffness, Low is updated here, because there was no existing weight map before this version
+		TetherStiffness.Low = FMath::Clamp(FMath::Loge(StrainLimitingStiffness_DEPRECATED) / FMath::Loge(1.e3f) + 1.f, 0.f, 1.f);
+		TetherStiffness.High = 0.f;
+	}
+#endif  // #if WITH_EDITORONLY_DATA
 }
 
 float UChaosClothConfig::GetMassValue() const
@@ -168,9 +190,11 @@ UChaosClothSharedSimConfig::~UChaosClothSharedSimConfig()
 
 void UChaosClothSharedSimConfig::MigrateFrom(const FClothConfig_Legacy& ClothConfig)
 {
+#if WITH_EDITORONLY_DATA
 	IterationCount = FMath::Clamp(int32(ClothConfig.SolverFrequency / 60.f), 1, 100);
 
 	bUseDampingOverride_DEPRECATED = false;  // Damping is migrated to per cloth configs
+#endif
 }
 
 void UChaosClothSharedSimConfig::Serialize(FArchive& Ar)
@@ -182,12 +206,15 @@ void UChaosClothSharedSimConfig::Serialize(FArchive& Ar)
 void UChaosClothSharedSimConfig::PostLoad()
 {
 	Super::PostLoad();
+
+#if WITH_EDITORONLY_DATA
 	const int32 ChaosClothSharedConfigCustomVersion = GetLinkerCustomVersion(FChaosClothSharedConfigCustomVersion::GUID);
 
 	if (ChaosClothSharedConfigCustomVersion < FChaosClothSharedConfigCustomVersion::AddGravityOverride)
 	{
 		bUseGravityOverride_DEPRECATED = true;  // Default gravity override would otherwise disable the currently set gravity on older versions
 	}
+#endif
 }
 
 #if WITH_EDITOR
diff --git a/Engine/Plugins/Experimental/ChaosCloth/Source/Chaos/Private/ChaosCloth/ChaosClothConstraints.cpp b/Engine/Plugins/Experimental/ChaosCloth/Source/Chaos/Private/ChaosCloth/ChaosClothConstraints.cpp
--- a/Engine/Plugins/Experimental/ChaosCloth/Source/Chaos/Private/ChaosCloth/ChaosClothConstraints.cpp
+++ b/Engine/Plugins/Experimental/ChaosCloth/Source/Chaos/Private/ChaosCloth/ChaosClothConstraints.cpp
@@ -76,7 +76,7 @@ void FClothConstraints::CreateRules()
 		ConstraintRuleOffset = Evolution->AddConstraintRuleRange(NumConstraintRules, false);
 	}
 
-	TFunction<void(const FPBDParticles&)>* const ConstraintInits = Evolution->ConstraintInits().GetData() + ConstraintInitOffset;
+	TFunction<void(const FPBDParticles&, const FReal)>* const ConstraintInits = Evolution->ConstraintInits().GetData() + ConstraintInitOffset;
 	TFunction<void(FPBDParticles&, const FReal)>* const ConstraintRules = Evolution->ConstraintRules().GetData() + ConstraintRuleOffset;
 
 	int32 ConstraintInitIndex = 0;
@@ -85,7 +85,7 @@ void FClothConstraints::CreateRules()
 	if (XEdgeConstraints)
 	{
 		ConstraintInits[ConstraintInitIndex++] =
-			[this](const FPBDParticles& /*Particles*/)
+			[this](const FPBDParticles& /*Particles*/, const FReal /*Dt*/)
 			{
 				XEdgeConstraints->Init();
 			};
@@ -107,7 +107,7 @@ void FClothConstraints::CreateRules()
 	if (XBendingConstraints)
 	{
 		ConstraintInits[ConstraintInitIndex++] =
-			[this](const FPBDParticles& /*Particles*/)
+			[this](const FPBDParticles& /*Particles*/, const FReal /*Dt*/)
 			{
 				XBendingConstraints->Init();
 			};
@@ -136,7 +136,7 @@ void FClothConstraints::CreateRules()
 	if (XAreaConstraints)
 	{
 		ConstraintInits[ConstraintInitIndex++] =
-			[this](const FPBDParticles& /*Particles*/)
+			[this](const FPBDParticles& /*Particles*/, const FReal /*Dt*/)
 			{
 				XAreaConstraints->Init();
 			};
@@ -173,9 +173,10 @@ void FClothConstraints::CreateRules()
 	if (XLongRangeConstraints)
 	{
 		ConstraintInits[ConstraintInitIndex++] =
-			[this](const FPBDParticles& /*Particles*/)
+			[this](const FPBDParticles& /*Particles*/, const FReal Dt)
 			{
 				XLongRangeConstraints->Init();
+				XLongRangeConstraints->ApplyProperties(Dt, Evolution->GetIterations());
 			};
 		ConstraintRules[ConstraintRuleIndex++] =
 			[this](FPBDParticles& Particles, const FReal Dt)
@@ -185,6 +186,11 @@ void FClothConstraints::CreateRules()
 	}
 	if (LongRangeConstraints)
 	{
+		ConstraintInits[ConstraintInitIndex++] =
+			[this](const FPBDParticles& /*Particles*/, const FReal Dt)
+			{
+				LongRangeConstraints->ApplyProperties(Dt, Evolution->GetIterations());
+			};
 		ConstraintRules[ConstraintRuleIndex++] =
 			[this](FPBDParticles& Particles, const FReal Dt)
 			{
@@ -209,10 +215,15 @@ void FClothConstraints::CreateRules()
 	}
 	if (AnimDriveConstraints)
 	{
+		ConstraintInits[ConstraintInitIndex++] =
+			[this](const FPBDParticles& /*Particles*/, const FReal Dt)
+			{
+				AnimDriveConstraints->ApplyProperties(Dt, Evolution->GetIterations());
+			};
+
 		ConstraintRules[ConstraintRuleIndex++] =
 			[this](FPBDParticles& Particles, const FReal Dt)
 			{
-				AnimDriveConstraints->ApplyProperties(Dt, Evolution->GetIterations());  // TODO: Move to init, there is no reason to update this at every iterations
 				AnimDriveConstraints->Apply(Particles, Dt);
 			};
 	}
@@ -227,7 +238,7 @@ void FClothConstraints::CreateRules()
 	if (SelfCollisionConstraints)
 	{
 		ConstraintInits[ConstraintInitIndex++] =
-			[this](const FPBDParticles& Particles)
+			[this](const FPBDParticles& Particles, const FReal /*Dt*/)
 			{
 				SelfCollisionConstraints->Init(Particles);
 			};
@@ -318,30 +329,37 @@ void FClothConstraints::SetVolumeConstraints(TArray<TVec3<int32>>&& SurfaceEleme
 	++NumConstraintRules;
 }
 
-void FClothConstraints::SetLongRangeConstraints(const TMap<int32, TSet<uint32>>& PointToNeighborsMap, FReal StrainLimitingStiffness, FReal LimitScale, ETetherMode TetherMode, bool bUseXPBDConstraints)
+void FClothConstraints::SetLongRangeConstraints(const TMap<int32, TSet<int32>>& PointToNeighborsMap, const TConstArrayView<FRealSingle>& TetherStiffnessMultipliers, const FVec2& TetherStiffness, FReal LimitScale, ETetherMode TetherMode, bool bUseXPBDConstraints)
 {
 	check(Evolution);
-	check(StrainLimitingStiffness > 0.f && StrainLimitingStiffness <= 1.f);
+
+	static const int32 MaxNumTetherIslands = 4;  // The max number of connected neighbors per particle.
 
 	if (bUseXPBDConstraints)
 	{
 		XLongRangeConstraints = MakeShared<FXPBDLongRangeConstraints>(
 			Evolution->Particles(),
-			PointToNeighborsMap, 
-			10, // The max number of connected neighbors per particle.
-			StrainLimitingStiffness);  // TODO: Add LimitScale to the XPBD constraint
-		++NumConstraintInits;
+			ParticleOffset,
+			NumParticles,
+			PointToNeighborsMap,
+			TetherStiffnessMultipliers,
+			MaxNumTetherIslands,
+			TetherStiffness);  // TODO: Add LimitScale to the XPBD constraint
 	}
 	else
 	{
 		LongRangeConstraints = MakeShared<FPBDLongRangeConstraints>(
 			Evolution->Particles(),
+			ParticleOffset,
+			NumParticles,
 			PointToNeighborsMap,
-			10, // The max number of connected neighbors per particle.
-			StrainLimitingStiffness,
+			TetherStiffnessMultipliers,
+			MaxNumTetherIslands,
+			TetherStiffness,
 			LimitScale,
 			TetherMode);
 	}
+	++NumConstraintInits;  // Uses init to update the property tables
 	++NumConstraintRules;
 }
 
@@ -377,6 +395,7 @@ void FClothConstraints::SetAnimDriveConstraints(const TConstArrayView<FReal>& An
 		*OldAnimationPositions,
 		AnimDriveStiffnessMultipliers,
 		AnimDriveDampingMultipliers);
+	++NumConstraintInits;  // Uses init to update the property tables
 	++NumConstraintRules;
 }
 
@@ -459,7 +478,7 @@ void FClothConstraints::SetVolumeProperties(FReal VolumeStiffness)
 	}
 }
 
-void FClothConstraints::SetLongRangeAttachmentProperties(FReal TetherStiffness)
+void FClothConstraints::SetLongRangeAttachmentProperties(const FVec2& TetherStiffness)
 {
 	if (LongRangeConstraints)
 	{
diff --git a/Engine/Plugins/Experimental/ChaosCloth/Source/Chaos/Private/ChaosCloth/ChaosClothingSimulation.cpp b/Engine/Plugins/Experimental/ChaosCloth/Source/Chaos/Private/ChaosCloth/ChaosClothingSimulation.cpp
--- a/Engine/Plugins/Experimental/ChaosCloth/Source/Chaos/Private/ChaosCloth/ChaosClothingSimulation.cpp
+++ b/Engine/Plugins/Experimental/ChaosCloth/Source/Chaos/Private/ChaosCloth/ChaosClothingSimulation.cpp
@@ -384,9 +384,9 @@ void FClothingSimulation::CreateActor(USkeletalMeshComponent* InOwnerComponent,
 		ClothConfig->AreaStiffness,
 		ClothConfig->VolumeStiffness,
 		ClothConfig->bUseThinShellVolumeConstraints,
-		ClothConfig->StrainLimitingStiffness,
+		TVector<float, 2>(ClothConfig->TetherStiffness.Low, ClothConfig->TetherStiffness.High),  // Animatable
 		ClothConfig->LimitScale,
-		(FClothingSimulationCloth::ETetherMode)ClothConfig->TetherMode,
+		ClothConfig->bUseGeodesicDistance ? FClothingSimulationCloth::ETetherMode::Geodesic : FClothingSimulationCloth::ETetherMode::Euclidean,
 		/*MaxDistancesMultiplier =*/ 1.f,  // Animatable
 		FVec2(ClothConfig->AnimDriveStiffness.Low, ClothConfig->AnimDriveStiffness.High),  // Animatable
 		FVec2(ClothConfig->AnimDriveDamping.Low, ClothConfig->AnimDriveDamping.High),  // Animatable
@@ -745,9 +745,9 @@ void FClothingSimulation::RefreshClothConfig(const IClothingSimulationContext* I
 			ClothConfig->AreaStiffness,
 			ClothConfig->VolumeStiffness,
 			ClothConfig->bUseThinShellVolumeConstraints,
-			ClothConfig->StrainLimitingStiffness,
+			TVector<float, 2>(ClothConfig->TetherStiffness.Low, ClothConfig->TetherStiffness.High),  // Animatable
 			ClothConfig->LimitScale,
-			(FClothingSimulationCloth::ETetherMode)ClothConfig->TetherMode,
+			ClothConfig->bUseGeodesicDistance ? FClothingSimulationCloth::ETetherMode::Geodesic : FClothingSimulationCloth::ETetherMode::Euclidean,
 			/*MaxDistancesMultiplier =*/ 1.f,  // Animatable
 			TVector<float, 2>(ClothConfig->AnimDriveStiffness.Low, ClothConfig->AnimDriveStiffness.High),  // Animatable
 			TVector<float, 2>(ClothConfig->AnimDriveDamping.Low, ClothConfig->AnimDriveDamping.High),  // Animatable
@@ -998,10 +998,13 @@ static void DrawPoint(FPrimitiveDrawInterface* PDI, const FVector& Pos, const FL
 	}
 #endif
 #if WITH_EDITOR
-	const FMatrix& ViewMatrix = PDI->View->ViewMatrices.GetViewMatrix();
-	const FVector XAxis = ViewMatrix.GetColumn(0); // Just using transpose here (orthogonal transform assumed)
-	const FVector YAxis = ViewMatrix.GetColumn(1);
-	DrawDisc(PDI, Pos, XAxis, YAxis, FColor::White, 0.2f, 10, DebugClothMaterialVertex->GetRenderProxy(), SDPG_World);
+	if (DebugClothMaterialVertex)
+	{
+		const FMatrix& ViewMatrix = PDI->View->ViewMatrices.GetViewMatrix();
+		const FVector XAxis = ViewMatrix.GetColumn(0); // Just using transpose here (orthogonal transform assumed)
+		const FVector YAxis = ViewMatrix.GetColumn(1);
+		DrawDisc(PDI, Pos, XAxis, YAxis, Color.ToFColor(true), 0.2f, 10, DebugClothMaterialVertex->GetRenderProxy(), SDPG_World);
+	}
 #endif
 }
 
@@ -1624,6 +1627,8 @@ void FClothingSimulation::DebugDrawMaxDistances(FPrimitiveDrawInterface* PDI) co
 			{
 #if WITH_EDITOR
 				DrawPoint(PDI, Position, FLinearColor::Red, DebugClothMaterialVertex);
+#else
+				DrawPoint(nullptr, Position, FLinearColor::Red, nullptr);
 #endif
 			}
 			else
@@ -1737,11 +1742,11 @@ void FClothingSimulation::DebugDrawLongRangeConstraint(FPrimitiveDrawInterface*
 		const FTriangleMesh& TriangleMesh = Cloth->GetTriangleMesh(Solver.Get());
 		const TConstArrayView<FReal> InvMasses = Cloth->GetParticleInvMasses(Solver.Get());
 
-		const TMap<int32, TSet<uint32>>& PointToNeighborsMap = TriangleMesh.GetPointToNeighborsMap();
+		const TMap<int32, TSet<int32>>& PointToNeighborsMap = TriangleMesh.GetPointToNeighborsMap();
 
-		static TArray<uint32> KinematicIndices;  // Make static to prevent constant allocations
+		static TArray<int32> KinematicIndices;  // Make static to prevent constant allocations
 		KinematicIndices.Reset();
-		for (const TPair<int32, TSet<uint32>>& PointNeighbors : PointToNeighborsMap)
+		for (const TPair<int32, TSet<int32>>& PointNeighbors : PointToNeighborsMap)
 		{
 			const int32 Index = PointNeighbors.Key;
 			if (InvMasses[Index - Offset] == 0.f)  // TODO: Triangle indices should ideally be starting at 0 to avoid these mix-ups
@@ -1750,7 +1755,7 @@ void FClothingSimulation::DebugDrawLongRangeConstraint(FPrimitiveDrawInterface*
 			}
 		}
 
-		const TArray<TArray<uint32>> IslandElements = FPBDLongRangeConstraints::ComputeIslands(PointToNeighborsMap, KinematicIndices);
+		const TArray<TArray<int32>> IslandElements = FPBDLongRangeConstraints::ComputeIslands(PointToNeighborsMap, KinematicIndices);
 
 		// Draw constraints
 		const FClothConstraints& ClothConstraints = Solver->GetClothConstraints(Offset);
@@ -1759,73 +1764,35 @@ void FClothingSimulation::DebugDrawLongRangeConstraint(FPrimitiveDrawInterface*
 
 		if (const FPBDLongRangeConstraints* const LongRangeConstraints = ClothConstraints.GetLongRangeConstraints().Get())
 		{
-			switch (LongRangeConstraints->GetMode())
+			for (const FPBDLongRangeConstraints::FTether& Tether : LongRangeConstraints->GetTethers())
 			{
-			case FPBDLongRangeConstraints::EMode::FastTetherFastLength:
-			case FPBDLongRangeConstraints::EMode::AccurateTetherFastLength:
+				const int32 KinematicIndex = Tether.Start;
+				const int32 DynamicIndex = Tether.End;
+				const FReal RefLength = Tether.RefLength;
+				// Find Island
+				int32 ColorIndex = 0;
+				for (int32 IslandIndex = 0; IslandIndex < IslandElements.Num(); ++IslandIndex)
 				{
-					const TArray<TVec2<uint32>>& Constraints = LongRangeConstraints->GetEuclideanConstraints();
-					for (int32 ConstraintIndex = 0; ConstraintIndex < Constraints.Num(); ++ConstraintIndex)
+					if (IslandElements[IslandIndex].Find(KinematicIndex) != INDEX_NONE)  // TODO: This is O(n^2), it'll be nice to make this faster, even if it is only debug viz. Maybe binary search if the kinematic indices are ordered?
 					{
-						const TVec2<uint32>& Path = Constraints[ConstraintIndex];
-						const uint32 KinematicIndex = Path[0];
-						const uint32 DynamicIndex = Path[Path.Num() - 1];
-
-						// Find Island
-						int32 ColorIndex = 0;
-						for (int32 IslandIndex = 0; IslandIndex < IslandElements.Num(); ++IslandIndex)
-						{
-							if (IslandElements[IslandIndex].Find(KinematicIndex) != INDEX_NONE)  // TODO: This is O(n^2), it'll be nice to make this faster, even if it is only debug viz. Maybe binary search if the kinematic indices are ordered?
-							{
-								ColorIndex = ColorOffset + IslandIndex;
-								break;
-							}
-						}
-
-						// Draw line
-						const FVec3 Pos0 = Positions[KinematicIndex - Offset] + LocalSpaceLocation;
-						const FVec3 Pos1 = Positions[DynamicIndex - Offset] + LocalSpaceLocation;
-						DrawLine(PDI, Pos0, Pos1, PseudoRandomColor(ColorIndex));
+						ColorIndex = ColorOffset + IslandIndex;
+						break;
 					}
 				}
-				break;
-			case FPBDLongRangeConstraints::EMode::AccurateTetherAccurateLength:
-				{
-					const TArray<TArray<uint32>>& Constraints = LongRangeConstraints->GetGeodesicConstraints();
-					for (int32 ConstraintIndex = 0; ConstraintIndex < Constraints.Num(); ++ConstraintIndex)
-					{
-						const TArray<uint32>& Path = Constraints[ConstraintIndex];
-						const uint32 KinematicIndex = Path[0];
+				const FLinearColor Color = PseudoRandomColor(ColorIndex);
+				const FVec3 Pos0 = Positions[KinematicIndex - Offset] + LocalSpaceLocation;
+				const FVec3 Pos1 = Positions[DynamicIndex - Offset] + LocalSpaceLocation;
 
-						// Find Island
-						int32 ColorIndex = 0;
-						for (int32 IslandIndex = 0; IslandIndex < IslandElements.Num(); ++IslandIndex)
-						{
-							if (IslandElements[IslandIndex].Find(KinematicIndex) != INDEX_NONE)  // TODO: This is O(n^2), it'll be nice to make this faster, even if it is only debug viz. Maybe binary search if the kinematic indices are ordered?
-							{
-								ColorIndex = ColorOffset + IslandIndex;
-								break;
-							}
-						}
-						const FLinearColor Color = PseudoRandomColor(ColorIndex);
+				DrawLine(PDI, Pos0, Pos1, Color);
 
-						// Draw lines
-						FVec3 Pos0 = Positions[KinematicIndex - Offset] + LocalSpaceLocation;
-						for (int32 PathIndex = 1; PathIndex < Path.Num(); ++PathIndex)
-						{
-							const uint32 DynamicIndex = Path[PathIndex];
-							const FVec3 Pos1 = Positions[DynamicIndex - Offset] + LocalSpaceLocation;
-							DrawLine(PDI, Pos0, Pos1, Color);
-							Pos0 = Pos1;
-						}
-					}
+				FVec3 Direction = Pos1 - Pos0;
+				const float Length = Direction.SafeNormalize();
+				if (Length > SMALL_NUMBER)
+				{
+					const FVec3 Pos2 = Pos1 + Direction * (Tether.RefLength - Length);
+					DrawLine(PDI, Pos1, Pos2, FLinearColor::Black);  // Color.Desaturate(1.f));
 				}
-				break;
-			default:
-				unimplemented();
-				break;
 			}
-
 		}
 
 		// Draw islands
diff --git a/Engine/Plugins/Experimental/ChaosCloth/Source/Chaos/Private/ChaosCloth/ChaosClothingSimulationCloth.cpp b/Engine/Plugins/Experimental/ChaosCloth/Source/Chaos/Private/ChaosCloth/ChaosClothingSimulationCloth.cpp
--- a/Engine/Plugins/Experimental/ChaosCloth/Source/Chaos/Private/ChaosCloth/ChaosClothingSimulationCloth.cpp
+++ b/Engine/Plugins/Experimental/ChaosCloth/Source/Chaos/Private/ChaosCloth/ChaosClothingSimulationCloth.cpp
@@ -190,13 +190,14 @@ void FClothingSimulationCloth::FLODData::Add(FClothingSimulationSolver* Solver,
 	}
 
 	// Long range constraints
-	if (Cloth->TetherStiffness)
+	if (Cloth->TetherStiffness[0] > 0.f || Cloth->TetherStiffness[1] > 0.f)
 	{
-		check(TriangleMesh.GetNumElements() > 0);
-		// PerFormance note: The Per constraint version of this function is quite a bit faster for smaller assets
-		// There might be a cross-over point where the PerParticle version is faster: To be determined
+		const TMap<int32, TSet<int32>>& PointToNeighborsMap = TriangleMesh.GetPointToNeighborsMap();
+		const TConstArrayView<FRealSingle>& TetherStiffnessMultipliers = WeightMaps[(int32)EChaosWeightMapTarget::TetherStiffness];
+
 		ClothConstraints.SetLongRangeConstraints(
-			TriangleMesh.GetPointToNeighborsMap(),
+			PointToNeighborsMap,
+			TetherStiffnessMultipliers,
 			Cloth->TetherStiffness,
 			Cloth->LimitScale,
 			Cloth->TetherMode,
@@ -337,7 +338,7 @@ FClothingSimulationCloth::FClothingSimulationCloth(
 	FRealSingle InAreaStiffness,
 	FRealSingle InVolumeStiffness,
 	bool bInUseThinShellVolumeConstraints,
-	FRealSingle InTetherStiffness,
+	const FVec2& InTetherStiffness,
 	FRealSingle InLimitScale,
 	ETetherMode InTetherMode,
 	FRealSingle InMaxDistancesMultiplier,
diff --git a/Engine/Plugins/Experimental/ChaosCloth/Source/Chaos/Private/ChaosCloth/ChaosClothingSimulationInteractor.cpp b/Engine/Plugins/Experimental/ChaosCloth/Source/Chaos/Private/ChaosCloth/ChaosClothingSimulationInteractor.cpp
--- a/Engine/Plugins/Experimental/ChaosCloth/Source/Chaos/Private/ChaosCloth/ChaosClothingSimulationInteractor.cpp
+++ b/Engine/Plugins/Experimental/ChaosCloth/Source/Chaos/Private/ChaosCloth/ChaosClothingSimulationInteractor.cpp
@@ -31,7 +31,17 @@ void UChaosClothingInteractor::SetMaterialLinear(float EdgeStiffness, float Bend
 	}));
 }
 
-void UChaosClothingInteractor::SetLongRangeAttachmentLinear(float TetherStiffness)
+void UChaosClothingInteractor::SetLongRangeAttachmentLinear(float TetherStiffnessLinear)
+{
+	// Deprecated
+	Commands.Add(FChaosClothingInteractorCommand::CreateLambda([TetherStiffnessLinear](FClothingSimulationCloth* Cloth)
+	{
+		const FVec2 TetherStiffness((FReal)FMath::Clamp(FMath::Loge(TetherStiffnessLinear) / FMath::Loge(1.e3f) + 1.f, 0.f, 1.f), (FReal)1.f);
+		Cloth->SetLongRangeAttachmentProperties(TetherStiffness);
+	}));
+}
+
+void UChaosClothingInteractor::SetLongRangeAttachment(FVector2D TetherStiffness)
 {
 	Commands.Add(FChaosClothingInteractorCommand::CreateLambda([TetherStiffness](FClothingSimulationCloth* Cloth)
 	{
@@ -73,9 +83,11 @@ void UChaosClothingInteractor::SetGravity(float GravityScale, bool bIsGravityOve
 
 void UChaosClothingInteractor::SetAnimDriveLinear(float AnimDriveStiffnessLinear)
 {
+	// Deprecated
 	Commands.Add(FChaosClothingInteractorCommand::CreateLambda([AnimDriveStiffnessLinear](FClothingSimulationCloth* Cloth)
 	{
-		const FVec2 AnimDriveStiffness((FReal)0.f, (FReal)(FMath::Clamp(FMath::Loge(AnimDriveStiffnessLinear) / FMath::Loge(1.e3f) + 1.f, 0.f, 1.f)));
+		// The Anim Drive stiffness Low value needs to be 0 in order to keep backward compatibility with existing mask (this wouldn't be an issue if this property had no legacy mask)
+		const FVec2 AnimDriveStiffness((FReal)0.f, (FReal)FMath::Clamp(FMath::Loge(AnimDriveStiffnessLinear) / FMath::Loge(1.e3f) + 1.f, 0.f, 1.f));
 		const FVec2 AnimDriveDamping((FReal)0.f, (FReal)1.f);
 		Cloth->SetAnimDriveProperties(AnimDriveStiffness, AnimDriveDamping);
 	}));
diff --git a/Engine/Plugins/Experimental/ChaosCloth/Source/Chaos/Private/ChaosCloth/ChaosWeightMapTarget.h b/Engine/Plugins/Experimental/ChaosCloth/Source/Chaos/Private/ChaosCloth/ChaosWeightMapTarget.h
--- a/Engine/Plugins/Experimental/ChaosCloth/Source/Chaos/Private/ChaosCloth/ChaosWeightMapTarget.h
+++ b/Engine/Plugins/Experimental/ChaosCloth/Source/Chaos/Private/ChaosCloth/ChaosWeightMapTarget.h
@@ -14,6 +14,7 @@ enum class EChaosWeightMapTarget : uint8
 	BackstopDistance   = (uint8)EWeightMapTargetCommon::BackstopDistance,
 	BackstopRadius     = (uint8)EWeightMapTargetCommon::BackstopRadius,
 	AnimDriveStiffness = (uint8)EWeightMapTargetCommon::AnimDriveStiffness,
-	AnimDriveDamping = (uint8)EWeightMapTargetCommon::AnimDriveDamping
+	AnimDriveDamping   = (uint8)EWeightMapTargetCommon::AnimDriveDamping,
 	// Add Chaos specific maps below this line
+	TetherStiffness
 };
diff --git a/Engine/Plugins/Experimental/ChaosCloth/Source/Chaos/Public/ChaosCloth/ChaosClothConfig.h b/Engine/Plugins/Experimental/ChaosCloth/Source/Chaos/Public/ChaosCloth/ChaosClothConfig.h
--- a/Engine/Plugins/Experimental/ChaosCloth/Source/Chaos/Public/ChaosCloth/ChaosClothConfig.h
+++ b/Engine/Plugins/Experimental/ChaosCloth/Source/Chaos/Public/ChaosCloth/ChaosClothConfig.h
@@ -32,7 +32,10 @@ struct FChaosClothWeightedValue
 	float High = 1.f;
 };
 
-/** Long range attachment options. */
+/**
+ * Long range attachment options.
+ * Deprecated.
+ */
 UENUM()
 enum class EChaosClothTetherMode : uint8
 {
@@ -41,7 +44,7 @@ enum class EChaosClothTetherMode : uint8
 	// Accurate Tether Fast Length: Use the accurate geodesic method to setup the tethers and a fast euclidean method to calculate their lengths. Slow initialization times and fast simulation times, but can still be prone to artifacts.
 	AccurateTetherFastLength,
 	// Accurate Tether Accurate Length: Use accurate geodesic method to both setup the tethers and calculate their lengths. Slow initialization and simulation times, but this is the most accurate setting showing the less artifacts.
-	AccurateTetherAccurateLength UMETA(Hidden),  // TODO: Fix Geodesic LRA
+	AccurateTetherAccurateLength UMETA(Hidden),
 	MaxChaosClothTetherMode UMETA(Hidden)
 };
 
@@ -79,11 +82,11 @@ public:
 	UPROPERTY(EditAnywhere, Category = "Mass Properties")
 	EClothMassMode MassMode = EClothMassMode::Density;
 
-	// The value used when the Mass Mode is set to Uniform Mass
+	/** The value used when the Mass Mode is set to Uniform Mass. */
 	UPROPERTY(EditAnywhere, Category = "Mass Properties", meta = (UIMin = "0.000001", UIMax = "0.001", ClampMin = "0", EditCondition = "MassMode == EClothMassMode::UniformMass"))
 	float UniformMass = 0.00015f;
 
-	// The value used when Mass Mode is set to TotalMass
+	/** The value used when Mass Mode is set to TotalMass. */
 	UPROPERTY(EditAnywhere, Category = "Mass Properties", meta = (UIMin = "0.001", UIMax = "10", ClampMin = "0", EditCondition = "MassMode == EClothMassMode::TotalMass"))
 	float TotalMass = 0.5f;
 
@@ -100,95 +103,101 @@ public:
 	UPROPERTY(EditAnywhere, Category = "Mass Properties", meta = (UIMin = "0.001", UIMax = "1", ClampMin = "0", EditCondition = "MassMode == EClothMassMode::Density"))
 	float Density = 0.35f;
 
-	// This is a lower bound to cloth particle masses.
+	/** This is a lower bound to cloth particle masses. */
 	UPROPERTY()
 	float MinPerParticleMass = 0.0001f;
 
-	// The Stiffness of the Edge constraints, only use lower than 1 values for very stretchy materials. Increase the iteration count for stiffer materials.
+	/** The Stiffness of the Edge constraints, only use lower than 1 values for very stretchy materials. Increase the iteration count for stiffer materials. */
 	UPROPERTY(EditAnywhere, Category = "Material Properties", meta = (UIMin = "0", UIMax = "1", ClampMin = "0", ClampMax = "1"))
 	float EdgeStiffness = 1.f;
 
-	// The Stiffness of the bending constraints. Increase the iteration count for stiffer materials.
+	/** The Stiffness of the bending constraints. Increase the iteration count for stiffer materials. Increase the iteration count for stiffer materials. */
 	UPROPERTY(EditAnywhere, Category = "Material Properties", meta = (UIMin = "0", UIMax = "1", ClampMin = "0", ClampMax = "1"))
 	float BendingStiffness = 1.f;
 
-	// Enable the more accurate bending element constraints instead of the faster cross-edge spring constraints used for controlling bending stiffness.
+	/** Enable the more accurate bending element constraints instead of the faster cross-edge spring constraints used for controlling bending stiffness. */
 	UPROPERTY(EditAnywhere, Category = "Material Properties")
 	bool bUseBendingElements = false;
 
-	// The stiffness of the area preservation constraints. Increase the iteration count for stiffer materials.
+	/** The stiffness of the area preservation constraints. Increase the iteration count for stiffer materials. */
 	UPROPERTY(EditAnywhere, Category = "Material Properties", meta = (UIMin = "0", UIMax = "1", ClampMin = "0", ClampMax = "1"))
 	float AreaStiffness = 1.f;
 
-	// The stiffness of the volume preservation constraints.
+	/** The stiffness of the volume preservation constraints. */
 	UPROPERTY()
 	float VolumeStiffness = 0.f;
 
-	// The strain limiting stiffness of the long range attachment constraints (aka tether stiffness).
-	// The long range attachment connects each of the cloth particles to its closest fixed point with a spring constraint.
-	// This can be used to compensate for a lack of stretch resistance when the iterations count is kept low for performance reasons.
-	// Can lead to an unnatural pull string puppet like behavior. Use 0 to disable. 
+	/**
+	 * The tethers' stiffness of the long range attachment constraints.
+	 * The long range attachment connects each of the cloth particles to its closest fixed point with a spring constraint.
+	 * This can be used to compensate for a lack of stretch resistance when the iterations count is kept low for performance reasons.
+	 * Can lead to an unnatural pull string puppet like behavior.
+	 * If an enabled Weight Map (A.K.A. Mask) targeting the "Tether Stiffness" is added to the cloth, 
+	 * then both the Low and High values will be used in conjunction with the per particle Weight stored
+	 * in the Weight Map to interpolate the final value from them.
+	 * Otherwise only the Low value is meaningful and sufficient to enable this constraint.
+	 * Use 0, 0 to disable.
+	 */
 	UPROPERTY(EditAnywhere, Category = "Long Range Attachment", meta = (UIMin = "0", UIMax = "1", ClampMin = "0", ClampMax = "1"))
-	float StrainLimitingStiffness = 0.5f;
+	FChaosClothWeightedValue TetherStiffness = { 1.f, 1.f };
 
-	// The limit scale of the long range attachment constraints (aka tether limit).
+	/** The limit scale of the long range attachment constraints (aka tether limit). */
 	UPROPERTY(EditAnywhere, Category = "Long Range Attachment", meta = (UIMin = "1.", UIMax = "1.1", ClampMin = "0.01", ClampMax = "10"))
 	float LimitScale = 1.f;
 
 	/**
-	 * How the long range attachment tethers are created, and how their length is updated.
-	 * This determines how fast and accurate both the tether initialization and simulation runs.
-	 * -	Fast Tether Fast Length: Use fast euclidean methods to both setup the tethers and calculate their lengths. Fast initialization and simulation times, but is very prone to artifacts.
-	 * -	Accurate Tether Fast Length: Use the accurate geodesic method to setup the tethers and a fast euclidean method to calculate their lengths. Slow initialization times and fast simulation times, but can still be prone to artifacts.
-	 * -	Accurate Tether Accurate Length: Use accurate geodesic method to both setup the tethers and calculate their lengths. Slow initialization and simulation times, but this is the most accurate setting showing the less artifacts.
+	 * Use geodesic instead of euclidean distance calculations for the Long Range Attachment constraint,
+	 * which is slower at setup but more accurate at establishing the correct position and length of the tethers,
+	 * and therefore is less prone to artifacts during the simulation.
 	 */
 	UPROPERTY(EditAnywhere, Category = "Long Range Attachment")
-	EChaosClothTetherMode TetherMode = EChaosClothTetherMode::AccurateTetherFastLength;
-
-	// Use geodesic instead of euclidean distance calculations in the long range attachment constraint,
-	// which is slower at setup but less prone to artifacts during simulation.
-	UPROPERTY()
-	bool bUseGeodesicDistance_DEPRECATED = true;
+	bool bUseGeodesicDistance = true;
 
-	// The stiffness of the shape target constraints
+	/** The stiffness of the shape target constraints. */
 	UPROPERTY()
 	float ShapeTargetStiffness = 0.f;
 
-	// The added thickness of collision shapes.
+	/** The added thickness of collision shapes. */
 	UPROPERTY(EditAnywhere, Category = "Collision Properties", meta = (UIMin = "0", UIMax = "100", ClampMin = "0", ClampMax = "1000"))
 	float CollisionThickness = 1.0f;
 
-	// Friction coefficient for cloth - collider interaction.
+	/** Friction coefficient for cloth - collider interaction. */
 	UPROPERTY(EditAnywhere, Category = "Collision Properties", meta = (UIMin = "0", UIMax = "1", ClampMin = "0", ClampMax = "10"))
 	float FrictionCoefficient = 0.8f;
 
-	// Use continuous collision detection (CCD) to prevent any missed collisions between fast moving particles and colliders.
-	// This has a negative effect on performance compared to when resolving collision without using CCD.
+	/**
+	 * Use continuous collision detection (CCD) to prevent any missed collisions between fast moving particles and colliders.
+	 * This has a negative effect on performance compared to when resolving collision without using CCD.
+	 */
 	UPROPERTY(EditAnywhere, Category = "Collision Properties")
 	bool bUseCCD = false;
 
-	// Enable self collision.
+	/** Enable self collision. */
 	UPROPERTY(EditAnywhere, Category = "Collision Properties", meta = (InlineEditConditionToggle))
 	bool bUseSelfCollisions = false;
 
-	// The radius of the spheres used in self collision 
+	/** The radius of the spheres used in self collision. */
 	UPROPERTY(EditAnywhere, Category = "Collision Properties", meta = (UIMin = "0", UIMax = "100", ClampMin = "0", ClampMax = "1000", EditCondition = "bUseSelfCollisions"))
 	float SelfCollisionThickness = 2.0f;
 
-	// This parameter is automatically set by the migration code. It can be overridden here to use the old way of authoring the backstop distances.
-	// The legacy backstop requires the sphere radius to be included within the painted distance mask, making it difficult to author correctly. In this case the backstop distance is the distance from the animated mesh to the center of the corresponding backstop collision sphere.
-	// The non legacy backstop automatically adds the matching sphere's radius to the distance calculations at runtime to make for a simpler authoring of the backstop distances. In this case the backstop distance is the distance from the animated mesh to the surface of the backstop collision sphere.
-	// In both cases, a positive backstop distance goes against the corresponding animated mesh's normal, and a negative backstop distance goes along the corresponding animated mesh's normal.
+	/**
+	 * This parameter is automatically set by the migration code. It can be overridden here to use the old way of authoring the backstop distances.
+	 * The legacy backstop requires the sphere radius to be included within the painted distance mask, making it difficult to author correctly. In this case the backstop distance is the distance from the animated mesh to the center of the corresponding backstop collision sphere.
+	 * The non legacy backstop automatically adds the matching sphere's radius to the distance calculations at runtime to make for a simpler authoring of the backstop distances. In this case the backstop distance is the distance from the animated mesh to the surface of the backstop collision sphere.
+	 * In both cases, a positive backstop distance goes against the corresponding animated mesh's normal, and a negative backstop distance goes along the corresponding animated mesh's normal.
+	 */
 	UPROPERTY(EditAnywhere, Category = "Collision Properties")
 	bool bUseLegacyBackstop = false;
 
-	// The amount of damping applied to the cloth velocities.
+	/** The amount of damping applied to the cloth velocities. */
 	UPROPERTY(EditAnywhere, Category = "Environmental Properties", meta = (UIMin = "0", UIMax = "1", ClampMin = "0", ClampMax = "1"))
 	float DampingCoefficient = 0.01f;
 
-	// This parameter is automatically set by the migration code. It can be overridden here to use the old deprecated "Legacy" wind model in order to preserve behavior with previous versions of the engine.
-	// The old wind model is not an accurate aerodynamic model and as such should be avoided. Being point based, it doesn't take into account the surface area that gets hit by the wind.
-	// Using this model makes the simulation slightly slower, disables the aerodynamically accurate wind model, and prevents the cloth to interact with the surrounding environment (air, water, ...etc.) even when there is no wind.
+	/**
+	 * This parameter is automatically set by the migration code. It can be overridden here to use the old deprecated "Legacy" wind model in order to preserve behavior with previous versions of the engine.
+	 * The old wind model is not an accurate aerodynamic model and as such should be avoided. Being point based, it doesn't take into account the surface area that gets hit by the wind.
+	 * Using this model makes the simulation slightly slower, disables the aerodynamically accurate wind model, and prevents the cloth to interact with the surrounding environment (air, water, ...etc.) even when there is no wind.
+	 */
 	UPROPERTY(EditAnywhere, Category = "Environmental Properties")
 	bool bUsePointBasedWindModel = false;
 
@@ -212,10 +221,6 @@ public:
 	UPROPERTY(EditAnywhere, Category = "Environmental Properties", meta = (EditCondition = "bUseGravityOverride"))
 	FVector Gravity = { 0.f, 0.f, -980.665f };
 
-	// The strength of the constraint driving the cloth towards the animated goal mesh. Deprecated, use AnimDriveStiffness instead.
-	UPROPERTY()
-	float AnimDriveSpringStiffness_DEPRECATED = 0.f;
-
 	/**
 	 * The strength of the constraint driving the cloth towards the animated goal mesh.
 	 * If an enabled Weight Map (A.K.A. Mask) targeting the "Anim Drive Stiffness" is added to the cloth, 
@@ -236,36 +241,53 @@ public:
 	UPROPERTY(EditAnywhere, Category = "Animation Properties", meta = (UIMin = "0", UIMax = "1", ClampMin = "0", ClampMax = "1"))
 	FChaosClothWeightedValue AnimDriveDamping = { 0.f, 1.f };
 
-	// The amount of linear velocities sent to the local cloth space from the reference bone
-	// (the closest bone to the root on which the cloth section has been skinned, or the root itself if the cloth isn't skinned).
+	/**
+	 * The amount of linear velocities sent to the local cloth space from the reference bone
+	 * (the closest bone to the root on which the cloth section has been skinned, or the root itself if the cloth isn't skinned).
+	 */
 	UPROPERTY(EditAnywhere, Category = "Animation Properties", meta = (UIMin = "0", UIMax = "1", ClampMin = "0", ClampMax = "1"))
 	FVector LinearVelocityScale = { 0.75f, 0.75f, 0.75f };
 
-	// The amount of angular velocities sent to the local cloth space from the reference bone
-	// (the closest bone to the root on which the cloth section has been skinned, or the root itself if the cloth isn't skinned).
+	/**
+	 * The amount of angular velocities sent to the local cloth space from the reference bone
+	 * (the closest bone to the root on which the cloth section has been skinned, or the root itself if the cloth isn't skinned).
+	 */
 	UPROPERTY(EditAnywhere, Category = "Animation Properties", meta = (UIMin = "0", UIMax = "1", ClampMin = "0", ClampMax = "1"))
 	float AngularVelocityScale = 0.75f;
 
-	// The portion of the angular velocity that is used to calculate the strength of all fictitious forces (e.g. centrifugal force).
-	// This parameter is only having an effect on the portion of the reference bone's angular velocity that has been removed from the
-	// simulation via the Angular Velocity Scale parameter. This means it has no effect when AngularVelocityScale is set to 1 in which
-	// case the cloth is simulated with full world space angular velocities and subjected to the true physical world inertial forces.
-	// Values range from 0 to 2, with 0 showing no centrifugal effect, 1 full centrifugal effect, and 2 an overdriven centrifugal effect.
+	/**
+	 * The portion of the angular velocity that is used to calculate the strength of all fictitious forces (e.g. centrifugal force).
+	 * This parameter is only having an effect on the portion of the reference bone's angular velocity that has been removed from the
+	 * simulation via the Angular Velocity Scale parameter. This means it has no effect when AngularVelocityScale is set to 1 in which
+	 * case the cloth is simulated with full world space angular velocities and subjected to the true physical world inertial forces.
+	 * Values range from 0 to 2, with 0 showing no centrifugal effect, 1 full centrifugal effect, and 2 an overdriven centrifugal effect.
+	 */
 	UPROPERTY(EditAnywhere, Category = "Animation Properties", meta = (UIMin = "0", UIMax = "1", ClampMin = "0", ClampMax = "2"))
 	float FictitiousAngularScale = 1.f;
 
-	// Enable tetrahedral constraints.
+	/** Enable tetrahedral constraints. */
 	UPROPERTY()
 	bool bUseTetrahedralConstraints = false;
 
-	// Enable thin shell volume constraints.
+	/** Enable thin shell volume constraints. */
 	UPROPERTY()
 	bool bUseThinShellVolumeConstraints = false;
 
-	// Enable continuous collision detection.
+	/** Enable continuous collision detection. */
 	UPROPERTY()
 	bool bUseContinuousCollisionDetection = false;
 
+	// Deprecated properties
+#if WITH_EDITORONLY_DATA
+	UPROPERTY()
+	EChaosClothTetherMode TetherMode_DEPRECATED = EChaosClothTetherMode::MaxChaosClothTetherMode;
+
+	UPROPERTY()
+	float AnimDriveSpringStiffness_DEPRECATED = 0.f;
+
+	UPROPERTY()
+	float StrainLimitingStiffness_DEPRECATED = 0.5f;
+#endif
 };
 
 /**
@@ -306,6 +328,7 @@ public:
 	UPROPERTY(EditAnywhere, Category = Simulation, meta = (UIMin = "1", UIMax = "10", ClampMin = "1", ClampMax = "100"))
 	int32 SubdivisionCount = 1;
 
+#if WITH_EDITORONLY_DATA
 	// The radius of the spheres used in self collision 
 	UPROPERTY()
 	float SelfCollisionThickness_DEPRECATED  = 2.0f;
@@ -333,6 +356,8 @@ public:
 	// The gravitational acceleration vector [cm/s^2]
 	UPROPERTY()
 	FVector Gravity_DEPRECATED = { 0.f, 0.f, -980.665f };
+#endif  // #if WITH_EDITORONLY_DATA
+
 
 	// Enable local space simulation to help with jitter due to floating point precision errors if the character is far away from the world origin
 	UPROPERTY()
diff --git a/Engine/Plugins/Experimental/ChaosCloth/Source/Chaos/Public/ChaosCloth/ChaosClothConstraints.h b/Engine/Plugins/Experimental/ChaosCloth/Source/Chaos/Public/ChaosCloth/ChaosClothConstraints.h
--- a/Engine/Plugins/Experimental/ChaosCloth/Source/Chaos/Public/ChaosCloth/ChaosClothConstraints.h
+++ b/Engine/Plugins/Experimental/ChaosCloth/Source/Chaos/Public/ChaosCloth/ChaosClothConstraints.h
@@ -50,7 +50,7 @@ namespace Chaos
 		void SetAreaConstraints(TArray<TVec3<int32>>&& SurfaceElements, FReal AreaStiffness, bool bUseXPBDConstraints);
 		void SetVolumeConstraints(TArray<TVec2<int32>>&& DoubleBendingEdges, FReal VolumeStiffness);
 		void SetVolumeConstraints(TArray<TVec3<int32>>&& SurfaceElements, FReal VolumeStiffness);
-		void SetLongRangeConstraints(const TMap<int32, TSet<uint32>>& PointToNeighborsMap, FReal StrainLimitingStiffness, FReal LimitScale, ETetherMode TetherMode, bool bUseXPBDConstraints);
+		void SetLongRangeConstraints(const TMap<int32, TSet<int32>>& PointToNeighborsMap, const TConstArrayView<FReal>& TetherStiffnessMultipliers, const FVec2& TetherStiffness, FReal LimitScale, ETetherMode TetherMode, bool bUseXPBDConstraints);
 		void SetMaximumDistanceConstraints(const TConstArrayView<FReal>& MaxDistances);
 		void SetBackstopConstraints(const TConstArrayView<FReal>& BackstopDistances, const TConstArrayView<FReal>& BackstopRadiuses, bool bUseLegacyBackstop);
 		void SetAnimDriveConstraints(const TConstArrayView<FReal>& AnimDriveStiffnessMultipliers, const TConstArrayView<FReal>& AnimDriveDampingMultipliers);
@@ -65,9 +65,9 @@ namespace Chaos
 		void SetAreaProperties(FReal AreaStiffness);
 		void SetThinShellVolumeProperties(FReal VolumeStiffness);
 		void SetVolumeProperties(FReal VolumeStiffness);
-		void SetLongRangeAttachmentProperties(FReal TetherStiffness);
+		void SetLongRangeAttachmentProperties(const FVec2& TetherStiffness);
 		void SetMaximumDistanceProperties(FReal MaxDistancesMultiplier);
-		void SetAnimDriveProperties(const FVec2& AnimDriveStiffness, const TVector<FReal, 2>& AnimDriveDamping);
+		void SetAnimDriveProperties(const FVec2& AnimDriveStiffness, const FVec2& AnimDriveDamping);
 		void SetSelfCollisionProperties(FReal SelfCollisionThickness);
 		// ---- End of Cloth interface ----
 
diff --git a/Engine/Plugins/Experimental/ChaosCloth/Source/Chaos/Public/ChaosCloth/ChaosClothingSimulationCloth.h b/Engine/Plugins/Experimental/ChaosCloth/Source/Chaos/Public/ChaosCloth/ChaosClothingSimulationCloth.h
--- a/Engine/Plugins/Experimental/ChaosCloth/Source/Chaos/Public/ChaosCloth/ChaosClothingSimulationCloth.h
+++ b/Engine/Plugins/Experimental/ChaosCloth/Source/Chaos/Public/ChaosCloth/ChaosClothingSimulationCloth.h
@@ -39,7 +39,7 @@ namespace Chaos
 			FRealSingle InAreaStiffness,
 			FRealSingle InVolumeStiffness,
 			bool bInUseThinShellVolumeConstraints,
-			FRealSingle InTetherStiffness,
+			const FVec2& InTetherStiffness,
 			FRealSingle InLimitScale,
 			ETetherMode InTetherMode,
 			FRealSingle InMaxDistancesMultiplier,
@@ -77,7 +77,7 @@ namespace Chaos
 		void SetMaxDistancesMultiplier(FRealSingle InMaxDistancesMultiplier) { MaxDistancesMultiplier = InMaxDistancesMultiplier; }
 
 		void SetMaterialProperties(FRealSingle InEdgeStiffness, FRealSingle InBendingStiffness, FRealSingle InAreaStiffness) { EdgeStiffness = InEdgeStiffness; BendingStiffness = InBendingStiffness; AreaStiffness = InAreaStiffness; }
-		void SetLongRangeAttachmentProperties(FRealSingle InTetherStiffness) { TetherStiffness = InTetherStiffness; }
+		void SetLongRangeAttachmentProperties(const FVec2& InTetherStiffness) { TetherStiffness = InTetherStiffness; }
 		void SetCollisionProperties(FRealSingle InCollisionThickness, FRealSingle InFrictionCoefficient, bool bInUseCCD, FRealSingle InSelfCollisionThickness) { CollisionThickness = InCollisionThickness; FrictionCoefficient = InFrictionCoefficient; bUseCCD = bInUseCCD; SelfCollisionThickness = InSelfCollisionThickness; }
 		void SetDampingProperties(FRealSingle InDampingCoefficient) { DampingCoefficient = InDampingCoefficient; }
 		void SetAerodynamicsProperties(FRealSingle InDragCoefficient, FRealSingle InLiftCoefficient) { DragCoefficient = InDragCoefficient; LiftCoefficient = InLiftCoefficient; }
@@ -191,7 +191,7 @@ namespace Chaos
 		FRealSingle AreaStiffness;
 		FRealSingle VolumeStiffness;
 		bool bUseThinShellVolumeConstraints;
-		FRealSingle TetherStiffness;
+		FVec2 TetherStiffness;
 		FRealSingle LimitScale;
 		ETetherMode TetherMode;
 		FRealSingle MaxDistancesMultiplier;
diff --git a/Engine/Plugins/Experimental/ChaosCloth/Source/Chaos/Public/ChaosCloth/ChaosClothingSimulationInteractor.h b/Engine/Plugins/Experimental/ChaosCloth/Source/Chaos/Public/ChaosCloth/ChaosClothingSimulationInteractor.h
--- a/Engine/Plugins/Experimental/ChaosCloth/Source/Chaos/Public/ChaosCloth/ChaosClothingSimulationInteractor.h
+++ b/Engine/Plugins/Experimental/ChaosCloth/Source/Chaos/Public/ChaosCloth/ChaosClothingSimulationInteractor.h
@@ -30,6 +30,9 @@ public:
 	UFUNCTION(BlueprintCallable, Category = ClothingSimulation, Meta = (Keywords = "Chaos Tether Stiffness"))
 	void SetLongRangeAttachmentLinear(float TetherStiffness = 1.f);
 
+	UFUNCTION(BlueprintCallable, Category = ClothingSimulation, Meta = (Keywords = "Chaos Tether Stiffness"))
+	void SetLongRangeAttachment(FVector2D TetherStiffness = FVector2D(0.f, 1.f));
+
 	UFUNCTION(BlueprintCallable, Category = ClothingSimulation, Meta = (Keywords = "Chaos Self Thickness Friction Coefficient"))
 	void SetCollision(float CollisionThickness = 1.f, float FrictionCoefficient = 0.8f, bool bUseCCD = false, float SelfCollisionThickness = 2.f);
 
diff --git a/Engine/Source/Runtime/Core/Public/UObject/PhysicsObjectVersion.h b/Engine/Source/Runtime/Core/Public/UObject/PhysicsObjectVersion.h
--- a/Engine/Source/Runtime/Core/Public/UObject/PhysicsObjectVersion.h
+++ b/Engine/Source/Runtime/Core/Public/UObject/PhysicsObjectVersion.h
@@ -39,6 +39,9 @@ struct CORE_API FPhysicsObjectVersion
 		// Add centrifugal forces for cloth
 		ChaosClothAddfictitiousforces,
 
+		// Added the Long Range Attachment stiffness weight map
+		ChaosClothAddTetherStiffnessWeightMap,
+
 		// -----<new versions can be added above this line>-------------------------------------------------
 		VersionPlusOne,
 		LatestVersion = VersionPlusOne - 1
diff --git a/Engine/Source/Runtime/Experimental/Chaos/Private/Chaos/PBDEvolution.cpp b/Engine/Source/Runtime/Experimental/Chaos/Private/Chaos/PBDEvolution.cpp
--- a/Engine/Source/Runtime/Experimental/Chaos/Private/Chaos/PBDEvolution.cpp
+++ b/Engine/Source/Runtime/Experimental/Chaos/Private/Chaos/PBDEvolution.cpp
@@ -347,9 +347,9 @@ void FPBDEvolution::AdvanceOneTimeStep(const FReal Dt)
 	{
 		SCOPE_CYCLE_COUNTER(STAT_ChaosXPBDConstraintsInit);
 		MConstraintInitsActiveView.SequentialFor(
-			[this](TArray<TFunction<void(const FPBDParticles&)>>& ConstraintInits, int32 Index)
+			[this, Dt](TArray<TFunction<void(const FPBDParticles&, const FReal)>>& ConstraintInits, int32 Index)
 			{
-				ConstraintInits[Index](MParticles);
+				ConstraintInits[Index](MParticles, Dt);
 			});
 	}
 
diff --git a/Engine/Source/Runtime/Experimental/Chaos/Private/Chaos/PBDLongRangeConstraints.cpp b/Engine/Source/Runtime/Experimental/Chaos/Private/Chaos/PBDLongRangeConstraints.cpp
--- a/Engine/Source/Runtime/Experimental/Chaos/Private/Chaos/PBDLongRangeConstraints.cpp
+++ b/Engine/Source/Runtime/Experimental/Chaos/Private/Chaos/PBDLongRangeConstraints.cpp
@@ -12,68 +12,72 @@ FAutoConsoleVariableRef CVarChaosLongRangeISPCEnabled(TEXT("p.Chaos.LongRange.IS
 
 using namespace Chaos;
 
-void FPBDLongRangeConstraints::Apply(FPBDParticles& InParticles, const FReal Dt, const TArray<int32>& ConstraintIndices) const
+void FPBDLongRangeConstraints::Apply(FPBDParticles& Particles, const FReal /*Dt*/, const TArray<int32>& ConstraintIndices) const
 {
 	SCOPE_CYCLE_COUNTER(STAT_PBD_LongRange);
-	switch (GetMode())
+	if (!Stiffness.HasWeightMap())
 	{
-	case EMode::FastTetherFastLength:
-	case EMode::AccurateTetherFastLength:
-		for (int32 i : ConstraintIndices)
+		const FReal ExpStiffnessValue = (FReal)Stiffness;
+		for (const int32 ConstraintIndex : ConstraintIndices)
 		{
-			Apply(MEuclideanConstraints[i], InParticles, Dt, MDists[i]);
+			const FTether& Tether = Tethers[ConstraintIndex];
+			Particles.P(Tether.End) += Stiffness[Tether.End - ParticleOffset] * Tether.GetDelta(Particles);
 		}
-		break;
-	case EMode::AccurateTetherAccurateLength:
-		for (int32 i : ConstraintIndices)
+	}
+	else
+	{
+		for (const int32 ConstraintIndex : ConstraintIndices)
 		{
-			Apply(MGeodesicConstraints[i], InParticles, Dt, MDists[i]);
+			const FTether& Tether = Tethers[ConstraintIndex];
+			const FReal ExpStiffnessValue = Stiffness[Tether.End - ParticleOffset];
+			Particles.P(Tether.End) += Stiffness[Tether.End - ParticleOffset] * Tether.GetDelta(Particles);
 		}
-		break;
-	default:
-		unimplemented();
-		break;
 	}
 }
 
-void FPBDLongRangeConstraints::Apply(FPBDParticles& InParticles, const FReal Dt) const
+void FPBDLongRangeConstraints::Apply(FPBDParticles& Particles, const FReal /*Dt*/) const
 {
 	SCOPE_CYCLE_COUNTER(STAT_PBD_LongRange);
 
-	switch (GetMode())
+	if (!Stiffness.HasWeightMap())
 	{
-	case EMode::FastTetherFastLength:
-	case EMode::AccurateTetherFastLength:
+		const float ExpStiffnessValue = (float)Stiffness;
+#if INTEL_ISPC
 		if (bRealTypeCompatibleWithISPC && bChaos_LongRange_ISPC_Enabled)
 		{
-#if INTEL_ISPC
-			ispc::ApplyLongRangeConstraints(
-				(ispc::FVector*)InParticles.GetP().GetData(),
-				(ispc::FUIntVector2*)MEuclideanConstraints.GetData(),
-				MDists.GetData(),
-				MStiffness,
-				MEuclideanConstraints.Num());
-#endif
+			// Run particles in parallel, and ranges in sequence to avoid a race condition when updating the same particle from different tethers
+			TethersView.RangeFor([this, &Particles, ExpStiffnessValue](TArray<FTether>& InTethers, int32 Offset, int32 Range)
+				{
+					ispc::ApplyLongRangeConstraints(
+						(ispc::FVector*)Particles.GetP().GetData(),
+						(ispc::FTether*)(InTethers.GetData() + Offset),
+						ExpStiffnessValue,
+						Range - Offset);
+				});
 		}
 		else
+#endif
 		{
-			for (int32 i = 0; i < MEuclideanConstraints.Num(); ++i)
-			{
-				Apply(MEuclideanConstraints[i], InParticles, Dt, MDists[i]);
-			}
-		}
-		break;
-
-	case EMode::AccurateTetherAccurateLength:
-		for (int32 i = 0; i < MGeodesicConstraints.Num(); ++i)
-		{
-			Apply(MGeodesicConstraints[i], InParticles, Dt, MDists[i]);
+			// Run particles in parallel, and ranges in sequence to avoid a race condition when updating the same particle from different tethers
+			static const int32 MinParallelSize = 500;
+			TethersView.ParallelFor([this, &Particles, ExpStiffnessValue](TArray<FTether>& InTethers, int32 Index)
+				{
+					const FTether& Tether = InTethers[Index];
+					Particles.P(Tether.End) += ExpStiffnessValue * Tether.GetDelta(Particles);
+				}, MinParallelSize);
 		}
-		break;
+	}
+	else
+	{
+		// TODO: ISPC implementation
 
-	default:
-		unimplemented();
-		break;
+		// Run particles in parallel, and ranges in sequence to avoid a race condition when updating the same particle from different tethers
+		static const int32 MinParallelSize = 500;
+		TethersView.ParallelFor([this, &Particles](TArray<FTether>& InTethers, int32 Index)
+			{
+				const FTether& Tether = InTethers[Index];
+				Particles.P(Tether.End) += Stiffness[Tether.End - ParticleOffset] * Tether.GetDelta(Particles);
+			}, MinParallelSize);
 	}
 }
 
diff --git a/Engine/Source/Runtime/Experimental/Chaos/Private/Chaos/PBDLongRangeConstraints.ispc b/Engine/Source/Runtime/Experimental/Chaos/Private/Chaos/PBDLongRangeConstraints.ispc
--- a/Engine/Source/Runtime/Experimental/Chaos/Private/Chaos/PBDLongRangeConstraints.ispc
+++ b/Engine/Source/Runtime/Experimental/Chaos/Private/Chaos/PBDLongRangeConstraints.ispc
@@ -4,52 +4,63 @@
 
 #include "Math/Vector.isph"
 
+struct FTether
+{
+	int32 Start;
+	int32 End;
+	float RefLength;
+};
+
+static inline varying FTether LoadTetherHelper(const uniform FTether *uniform SrcPtr)
+{
+	varying FTether Tether;
+	const varying FIntVector Loader = VectorLoad((const uniform FIntVector *uniform)SrcPtr);
+	Tether.Start = Loader.V[0];
+	Tether.End = Loader.V[1];
+	Tether.RefLength = floatbits(Loader.V[2]);
+	return Tether;
+}
+
 export void ApplyLongRangeConstraints(uniform FVector P[],
-										const uniform FUIntVector2 Constraints[],
-										const uniform float Dists[],
+										const uniform FTether Tethers[],
 										const uniform float Stiffness,
 										const uniform int32 NumConstraints)
 {
-	varying FUIntVector2 Constraint;
+	varying FTether Tether;
 
 #if HW_GATHER_SUPPORTED == 1
 	if(programIndex < NumConstraints)
 	{
-		Constraint = VectorLoad(&Constraints[0]);
+		Tether = LoadTetherHelper(&Tethers[0]);
 	}
 #endif
 
 	foreach(i = 0 ... NumConstraints)
 	{
 #if HW_GATHER_SUPPORTED == 0
-		Constraint = VectorLoad(&Constraints[extract(i, 0)]);
+		Tether = LoadTetherHelper(&Tethers[extract(i, 0)]);
 #endif
-
-		const varying uint32 i1 = Constraint.V[0];
-		const varying uint32 i2 = Constraint.V[1];
-
+		const varying int32 i1 = Tether.Start;
+		const varying int32 i2 = Tether.End;
+		const varying float RefLength = Tether.RefLength;
 		const varying FVector P1 = VectorGather(&P[i1]);
 		const varying FVector P2 = VectorGather(&P[i2]);
 
-		FVector Direction = P1 - P2;
-
-		// ComputeGeodesicDistance - This function is used for either Euclidian or Geodisic distances
-		const float Distance = VectorSize(Direction);
-		Direction = VectorSelect(Distance < KINDA_SMALL_NUMBER, ForwardVector, Direction / Distance);
-
 #if HW_GATHER_SUPPORTED == 1
 		if(i + programCount < NumConstraints)
 		{
-			Constraint = VectorLoad(&Constraints[extract(i, 0) + programCount]);
+			Tether = LoadTetherHelper(&Tethers[extract(i, 0) + programCount]);
 		}
 #endif
 
-		const varying float Dist = Dists[i];
+		FVector Direction = P1 - P2;
+		const float Length = VectorSize(Direction);
+		Direction = VectorSelect(Length < KINDA_SMALL_NUMBER, ForwardVector, Direction / Length);
 
-		const float Offset = Distance - Dist;
+		const float Offset = Length - RefLength;
 		const FVector ComputedDelta = Stiffness * Offset * Direction;
 
-		const FVector Delta = VectorSelect(Distance < Dist, ZeroVector, ComputedDelta);
+		const FVector Delta = VectorSelect(Offset < 0.f, ZeroVector, ComputedDelta);
 
 		VectorScatter(&P[i2], P2 + Delta);
 	}
diff --git a/Engine/Source/Runtime/Experimental/Chaos/Private/Chaos/PBDLongRangeConstraintsBase.cpp b/Engine/Source/Runtime/Experimental/Chaos/Private/Chaos/PBDLongRangeConstraintsBase.cpp
--- a/Engine/Source/Runtime/Experimental/Chaos/Private/Chaos/PBDLongRangeConstraintsBase.cpp
+++ b/Engine/Source/Runtime/Experimental/Chaos/Private/Chaos/PBDLongRangeConstraintsBase.cpp
@@ -4,68 +4,73 @@
 #include "Chaos/Map.h"
 #include "Chaos/Vector.h"
 #include "Chaos/Framework/Parallel.h"
+#include "ChaosStats.h"
 
-#include <queue>
-#include <unordered_map>
+DECLARE_CYCLE_STAT(TEXT("Chaos Cloth Compute Geodesic Constraints"), STAT_ChaosClothComputeGeodesicConstraints, STATGROUP_Chaos);
 
 using namespace Chaos;
 
 FPBDLongRangeConstraintsBase::FPBDLongRangeConstraintsBase(
-	const FDynamicParticles& InParticles,
-	const TMap<int32, TSet<uint32>>& PointToNeighbors,
-	const int32 NumberOfAttachments,
-	const FReal Stiffness,
+	const FPBDParticles& Particles,
+	const int32 InParticleOffset,
+	const int32 InParticleCount,
+	const TMap<int32, TSet<int32>>& PointToNeighbors,
+	const TConstArrayView<FReal>& StiffnessMultipliers,
+	const int32 MaxNumTetherIslands,
+	const FVec2& InStiffness,
 	const FReal LimitScale,
-	const EMode Mode)
-	: MStiffness(Stiffness)
-	, MMode(Mode)
+	const EMode InMode)
+	: TethersView(Tethers)
+	, Stiffness(StiffnessMultipliers, InStiffness, InParticleCount)
+	, Mode(InMode)
+	, ParticleOffset(InParticleOffset)
+	, ParticleCount(InParticleCount)
 {
-	switch (MMode)
+	switch (Mode)
 	{
-	case EMode::FastTetherFastLength:
-		ComputeEuclideanConstraints(InParticles, PointToNeighbors, NumberOfAttachments);
+	case EMode::Euclidean:
+		ComputeEuclideanConstraints(Particles, PointToNeighbors, MaxNumTetherIslands);
 		break;
-	case EMode::AccurateTetherFastLength:
-	case EMode::AccurateTetherAccurateLength:
-		ComputeGeodesicConstraints(InParticles, PointToNeighbors, NumberOfAttachments);
+	case EMode::Geodesic:
+		ComputeGeodesicConstraints(Particles, PointToNeighbors, MaxNumTetherIslands);
 		break;
 	default:
 		unimplemented();
 		break;
 	}
 
-	// Scale distance limits
-	for (FReal& Dist : MDists)
+	// Scale the tether's reference lengths
+	for (FTether& Tether : Tethers)
 	{
-		Dist *= LimitScale;
+		Tether.RefLength *= LimitScale;
 	}
 }
 
-TArray<TArray<uint32>> FPBDLongRangeConstraintsBase::ComputeIslands(
-    const TMap<int32, TSet<uint32>>& PointToNeighbors,
-    const TArray<uint32>& KinematicParticles)
+TArray<TArray<int32>> FPBDLongRangeConstraintsBase::ComputeIslands(
+    const TMap<int32, TSet<int32>>& PointToNeighbors,
+    const TArray<int32>& KinematicParticles)
 {
 	// Compute Islands
-	uint32 NextIsland = 0;
-	TArray<uint32> FreeIslands;
-	TArray<TArray<uint32>> IslandElements;
+	int32 NextIsland = 0;
+	TArray<int32> FreeIslands;
+	TArray<TArray<int32>> IslandElements;
 
-	TMap<uint32, uint32> ParticleToIslandMap;
+	TMap<int32, int32> ParticleToIslandMap;
 	ParticleToIslandMap.Reserve(KinematicParticles.Num());
 
-	for (const uint32 Element : KinematicParticles)
+	for (const int32 Element : KinematicParticles)
 	{
-		// Assign Element an island, possibly unioning existing islands
-		uint32 Island = TNumericLimits<uint32>::Max();
+		// Assign Element an island, possibly unionizing existing islands
+		int32 Island = TNumericLimits<int32>::Max();
 
 		// KinematicParticles is generated from the keys of PointToNeighbors, so
 		// we don't need to check if Element exists.
-		const TSet<uint32>& Neighbors = PointToNeighbors[Element];
-		for (auto Neighbor : Neighbors)
+		const TSet<int32>& Neighbors = PointToNeighbors[Element];
+		for (const int32 Neighbor : Neighbors)
 		{
 			if (ParticleToIslandMap.Contains(Neighbor))
 			{
-				if (Island == TNumericLimits<uint32>::Max())
+				if (Island == TNumericLimits<int32>::Max())
 				{
 					// We don't have an assigned island yet.  Join with the island
 					// of the neighbor.
@@ -73,12 +78,12 @@ TArray<TArray<uint32>> FPBDLongRangeConstraintsBase::ComputeIslands(
 				}
 				else
 				{
-					const uint32 OtherIsland = ParticleToIslandMap[Neighbor];
+					const int32 OtherIsland = ParticleToIslandMap[Neighbor];
 					if (OtherIsland != Island)
 					{
 						// This kinematic particle is connected to multiple islands.
 						// Union them.
-						for (auto OtherElement : IslandElements[OtherIsland])
+						for (const int32 OtherElement : IslandElements[OtherIsland])
 						{
 							check(ParticleToIslandMap[OtherElement] == OtherIsland);
 							ParticleToIslandMap[OtherElement] = Island;
@@ -93,7 +98,7 @@ TArray<TArray<uint32>> FPBDLongRangeConstraintsBase::ComputeIslands(
 
 		// If no connected Island was found, create a new one (or reuse an old 
 		// vacated one).
-		if (Island == TNumericLimits<uint32>::Max())
+		if (Island == TNumericLimits<int32>::Max())
 		{
 			if (FreeIslands.Num() == 0)
 			{
@@ -111,249 +116,343 @@ TArray<TArray<uint32>> FPBDLongRangeConstraintsBase::ComputeIslands(
 		check(IslandElements.IsValidIndex(Island));
 		IslandElements[Island].Add(Element);
 	}
+
 	// IslandElements may contain empty arrays.
+	for (int32 IslandIndex = 0; IslandIndex < IslandElements.Num(); )
+	{
+		if (!IslandElements[IslandIndex].Num())
+		{
+			IslandElements.RemoveAtSwap(IslandIndex, 1, false);  // RemoveAtSwap takes the last elements to replace the current one, do not increment the index in this case
+		}
+		else
+		{
+			++IslandIndex;
+		}
+	}
+
 	return IslandElements;
 }
 
 void FPBDLongRangeConstraintsBase::ComputeEuclideanConstraints(
-    const FDynamicParticles& InParticles,
-    const TMap<int32, TSet<uint32>>& PointToNeighbors,
-    const int32 NumberOfAttachments)
+    const FPBDParticles& Particles,
+    const TMap<int32, TSet<int32>>& PointToNeighbors,
+    const int32 MaxNumTetherIslands)
 {
-	// TODO(mlentine): Support changing which particles are kinematic during simulation
-	TArray<uint32> KinematicParticles;
-	for (const auto& KV : PointToNeighbors)
+	// Fill up the list of all used indices
+	TArray<int32> Nodes;
+	PointToNeighbors.GenerateKeyArray(Nodes);
+
+	TArray<int32> KinematicParticles;
+	for (const int32 Node : Nodes)
 	{
-		const int32 i = KV.Key;
-		if (InParticles.InvM(i) == 0.0)
+		if (Particles.InvM(Node) == (FReal)0.)
 		{
-			KinematicParticles.Add(i);
+			KinematicParticles.Add(Node);
 		}
 	}
 
 	// Compute the islands of kinematic particles
-	const TArray<TArray<uint32>> IslandElements = ComputeIslands(PointToNeighbors, KinematicParticles);
+	const TArray<TArray<int32>> IslandElements = ComputeIslands(PointToNeighbors, KinematicParticles);
 	int32 NumTotalIslandElements = 0;
-	for(const auto& Elements : IslandElements)
+	for (const TArray<int32>& Elements : IslandElements)
+	{
 		NumTotalIslandElements += Elements.Num();
-	TArray<Pair<FReal, uint32>> ClosestElements;
-	ClosestElements.Reserve(NumTotalIslandElements);
+	}
+
+	TArray<TArray<FTether>> NewTethersSlots;
+	NewTethersSlots.SetNum(Nodes.Num());
+
+	for (TArray<FTether>& NewTethers : NewTethersSlots)
+	{
+		NewTethers.Reserve(MaxNumTetherIslands);
+	}
 
-	//FCriticalSection CriticalSection;
-	//PhysicsParallelFor(InParticles.Size(), [&](int32 i)
-	for (const auto& KV : PointToNeighbors)
+	int32 NumTethers = 0;
+
+	PhysicsParallelFor(Nodes.Num(), [&](int32 Index)
 	{
+		const int32 Node = Nodes[Index];
+
 		// For each non-kinematic particle i...
-		const uint32 i = KV.Key;
-		if (InParticles.InvM(i) == 0.0)
-			continue;
+		if (Particles.InvM(Node) == (FReal)0.)
+		{
+			return;
+		}
 
 		// Measure the distance to all kinematic particles in all islands...
-		ClosestElements.Reset();
-		for (const TArray<uint32>& Elements : IslandElements)
+		TArray<Pair<FReal, int32>> ClosestElements;
+		ClosestElements.Reserve(NumTotalIslandElements);
+
+		for (const TArray<int32>& Elements : IslandElements)
 		{
 			// IslandElements may contain empty arrays.
 			if (Elements.Num() == 0)
+			{
 				continue;
+			}
 
-			uint32 ClosestElement = TNumericLimits<uint32>::Max();
-			FReal ClosestDistance = TNumericLimits<FReal>::Max();
-			for (const uint32 Element : Elements)
+			int32 ClosestElement = TNumericLimits<int32>::Max();
+			FReal ClosestSquareDistance = TNumericLimits<FReal>::Max();
+			for (const int32 Element : Elements)
 			{
-				const FReal Distance = ComputeDistance(InParticles, Element, i);
-				if (Distance < ClosestDistance)
+				const FReal SquareDistance = (Particles.X(Element) - Particles.X(Node)).SizeSquared();
+				if (SquareDistance < ClosestSquareDistance)
 				{
 					ClosestElement = Element;
-					ClosestDistance = Distance;
+					ClosestSquareDistance = SquareDistance;
 				}
 			}
-			ClosestElements.Add(MakePair(ClosestDistance, ClosestElement));
+			ClosestElements.Add(MakePair(FMath::Sqrt(ClosestSquareDistance), ClosestElement));
 		}
 
 		// Order all by distance, smallest first...
 		ClosestElements.Sort();
 
-		// Take the first N-umberOfAttachments.
-		if (NumberOfAttachments < ClosestElements.Num())
+		// Take the first MaxNumTetherIslands.
+		if (MaxNumTetherIslands < ClosestElements.Num())
 		{
-			ClosestElements.SetNum(NumberOfAttachments);
+			ClosestElements.SetNum(MaxNumTetherIslands);
 		}
 
 		// Add constraints between this particle, and the N closest...
-		for (auto Element : ClosestElements)
+		for (const Pair<FReal, int32> Element : ClosestElements)
 		{
-			//CriticalSection.Lock();
-			MEuclideanConstraints.Add({Element.Second, i});
-			MDists.Add(Element.First);
-			//CriticalSection.Unlock();
+			NewTethersSlots[Index].Emplace(Element.Second, Node, Element.First);
+			++NumTethers;
 		}
+	});
+
+	// Consolidate the N slots into one single list of tethers, spreading the nodes across different ranges for parallel processing
+	Tethers.Reset(NumTethers);
+
+	for (int32 Slot = 0; Slot < MaxNumTetherIslands; ++Slot)
+	{
+		int32 Size = 0;
+		for (int32 Index = 0; Index < Nodes.Num(); ++Index)
+		{
+			if (NewTethersSlots[Index].IsValidIndex(Slot))
+			{
+				Tethers.Emplace(NewTethersSlots[Index][Slot]);
+				++Size;
+			}
+		}
+		TethersView.AddRange(Size);
 	}
-	//);
 }
 
 void FPBDLongRangeConstraintsBase::ComputeGeodesicConstraints(
-    const FDynamicParticles& InParticles,
-    const TMap<int32, TSet<uint32>>& PointToNeighbors,
-    const int32 NumberOfAttachments)
+    const FPBDParticles& Particles,
+    const TMap<int32, TSet<int32>>& PointToNeighbors,
+    const int32 MaxNumTetherIslands)
 {
-	TArray<int32> UsedIndices;
-	PointToNeighbors.GenerateKeyArray(UsedIndices);
+	SCOPE_CYCLE_COUNTER(STAT_ChaosClothComputeGeodesicConstraints);
+
+	// Fill up the list of all used indices
+	TArray<int32> Nodes;
+	PointToNeighbors.GenerateKeyArray(Nodes);
 
-	// TODO(mlentine): Support changing which particles are kinematic during simulation
-	TArray<uint32> KinematicParticles;
-	for (const uint32 i : UsedIndices)
+	// Find all kinematic points to use as anchor points for the tethers and seed the path finding
+	TArray<int32> Seeds;  // Kinematic nodes connected to at least one dynamic node
+	Seeds.Reserve(Nodes.Num() / 10);  // Start at 10% seeds to minimize this array's reallocation
+
+	int32 NumDynamicNodes = 0;
+
+	for (const int32 Node : Nodes)
 	{
-		if (InParticles.InvM(i) == 0)
+		if (Particles.InvM(Node) == (FReal)0.)
 		{
-			KinematicParticles.Add(i);
+			const TSet<int32>& Neighbors = PointToNeighbors[Node];
+			for (const int32 Neighbor : Neighbors)
+			{
+				if (Particles.InvM(Neighbor) != (FReal)0.)  // There's no point to keep a node as an anchor, unless at least one of its neighbors is dynamic
+				{
+					Seeds.Add(Node);
+					break;
+				}
+			}
 		}
-	}
-	TArray<TArray<uint32>> IslandElements = ComputeIslands(PointToNeighbors, KinematicParticles);
-	// Store distances for all adjacent vertices
-	TMap<TVector<uint32, 2>, FReal> Distances;
-	for (const uint32 i : UsedIndices)
-	{
-		auto Neighbors = PointToNeighbors[i];
-		for (auto Neighbor : Neighbors)
+		else
 		{
-			Distances.Add(TVector<uint32, 2>(i, Neighbor), ComputeDistance(InParticles, Neighbor, i));
+			++NumDynamicNodes;
 		}
 	}
-	// Start and End Points to path and geodesic distance
-	TMap<TVector<uint32, 2>, Pair<FReal, TArray<uint32>>> GeodesicPaths;
+
 	// Dijkstra for each Kinematic Particle (assume a small number of kinematic points) - note this is N^2 log N with N kinematic points
-	for (const uint32 Element : KinematicParticles)
+	// Find shortest paths from each kinematic node to all dynamic nodes
+	TMap<TVec2<int32>, FReal> GeodesicDistances;
+	GeodesicDistances.Reserve(Seeds.Num() * NumDynamicNodes);
+
+	for (const int32 Seed : Seeds)
 	{
-		GeodesicPaths.Add(TVector<uint32, 2>(Element, Element), {0, {Element}});
-		for (const uint32 i : UsedIndices)
+		for (const int32 Node : Nodes)
 		{
-			if (i != Element)
+			if (Particles.InvM(Node) != (FReal)0.)
 			{
-				GeodesicPaths.Add(TVector<uint32, 2>(Element, i), {FLT_MAX, {}});
+				GeodesicDistances.Emplace(TVec2<int32>(Seed, Node), Node != Seed ? TNumericLimits<FReal>::Max() : (FReal)0.);
 			}
 		}
 	}
-	PhysicsParallelFor(KinematicParticles.Num(), [&](int32 Index)
+
+	PhysicsParallelFor(Seeds.Num(), [&](int32 Index)
 	{
-		const uint32 Element = KinematicParticles[Index];
-		std::priority_queue<Pair<FReal, uint32>, std::vector<Pair<FReal, uint32>>, std::greater<Pair<FReal, uint32>>> q;  // TODO(Kriss.Gossart): Remove use of std container
-		q.push(MakePair((FReal)0., Element));
-		TSet<uint32> Visited;
-		while (!q.empty())
+		const int32 Seed = Seeds[Index];
+
+		// Keep track of all visited nodes in a bit array, will need to be offsetted since the first node index may not be 0
+		TBitArray<> VisitedNodes(false, Nodes.Num());
+
+		// Priority queue based implementation of the Dijkstra algorithm
+#define CHAOS_PBDLONGRANGEATTACHMENT_PROFILE_HEAPSIZE 0
+#if CHAOS_PBDLONGRANGEATTACHMENT_PROFILE_HEAPSIZE
+		static int32 MaxHeapSize = 1;  // Record the max heap size
+#else
+		static const int32 MaxHeapSize = 512;  // Set the queue size to something large enough to avoid reallocations in most cases
+#endif
+		auto LessPredicate = [Seed, &GeodesicDistances](int32 Node1, int32 Node2) -> bool
+			{
+				return GeodesicDistances[TVec2<int32>(Seed, Node1)] < GeodesicDistances[TVec2<int32>(Seed, Node2)];  // Less for node priority
+			};
+
+		TArray<int32> Queue;
+		Queue.Reserve(MaxHeapSize);
+		Queue.Heapify(LessPredicate);  // Turn the array into a priority queue
+
+		// Initiate the graph progression
+		VisitedNodes[Seed - ParticleOffset] = true;
+		Queue.HeapPush(Seed, LessPredicate);
+
+		do
 		{
-			const Pair<FReal, uint32> PairElem = q.top();
-			q.pop();
-			if (Visited.Contains(PairElem.Second))
-				continue;
-			Visited.Add(PairElem.Second);
-			const TVector<uint32, 2> CurrentStartEnd(Element, PairElem.Second);
-			const TSet<uint32>& Neighbors = PointToNeighbors[PairElem.Second];
-			for (const uint32 Neighbor : Neighbors)
+			int32 ParentNode;
+			Queue.HeapPop(ParentNode, LessPredicate, false);
+
+			check(VisitedNodes[ParentNode - ParticleOffset]);
+
+			const FReal ParentDistance = (ParentNode != Seed) ? GeodesicDistances[TVec2<int32>(Seed, ParentNode)] : (FReal)0.;
+
+			const TSet<int32>& NeighborNodes = PointToNeighbors[ParentNode];
+			for (const int32 NeighborNode : NeighborNodes)
 			{
-				if (InParticles.InvM(Neighbor) == (FReal)0.) { continue; }
-				check(Neighbor != PairElem.Second);
-				const TVector<uint32, 2> NeighborStartEnd(Element, Neighbor);
-				const Pair<FReal, TArray<uint32>>& NeighborDistancePath = GeodesicPaths[NeighborStartEnd];
-				// Compute a possible distance for NeighborStartEnd
-				const FReal NewDist = PairElem.First + Distances[TVector<uint32, 2>(PairElem.Second, Neighbor)];
-				if (NewDist < NeighborDistancePath.First)
+				check(NeighborNode != ParentNode);
+
+				// Do not progress onto kinematic nodes
+				if (Particles.InvM(NeighborNode) == (FReal)0.)
 				{
-					TArray<uint32> NewPath = GeodesicPaths[CurrentStartEnd].Second;
-					check(NewPath.Num() > 0 && NewPath[NewPath.Num() - 1] != Neighbor);
-					NewPath.Add(Neighbor);
-					GeodesicPaths[NeighborStartEnd] = {NewDist, NewPath};
-					q.push(MakePair(GeodesicPaths[NeighborStartEnd].First, Neighbor));
+					continue;
+				}
+
+				// Update the geodesic distance if this path is a shorter one
+				const FReal NewDistance = ParentDistance + (Particles.X(NeighborNode) - Particles.X(ParentNode)).Size();
+
+				FReal& GeodesicDistance = GeodesicDistances[TVec2<int32>(Seed, NeighborNode)];
+
+				if (NewDistance < GeodesicDistance)
+				{
+					// Update this path distance
+					GeodesicDistance = NewDistance;
+
+					// Progress to this node position if it hasn't yet been visited
+					if (!VisitedNodes[NeighborNode - ParticleOffset])
+					{
+						VisitedNodes[NeighborNode - ParticleOffset] = true;
+
+						Queue.HeapPush(NeighborNode, LessPredicate);
+
+#if CHAOS_PBDLONGRANGEATTACHMENT_PROFILE_HEAPSIZE
+						MaxHeapSize = FMath::Max(Queue.Num(), MaxHeapSize);
+#endif
+					}
 				}
 			}
-		}
+		} while (Queue.Num());
 	});
-	TArray<TArray<uint32>> NewConstraints;
-	FCriticalSection CriticalSection;
-	PhysicsParallelFor(UsedIndices.Num(), [&](uint32 UsedIndex) {
-		const uint32 i = UsedIndices[UsedIndex];
-		if (InParticles.InvM(i) == 0)
-			return;
-		TArray<Pair<FReal, int32>> ClosestElements;
-		for (const TArray<uint32>& Elements : IslandElements)
+
+	// Compute islands of kinematic points
+	const TArray<TArray<int32>> Islands = ComputeIslands(PointToNeighbors, Seeds);
+
+	// Find the tether constraints starting from each dynamic node
+	TArray<TArray<FTether>> NewTethersSlots;
+	NewTethersSlots.SetNum(Nodes.Num());
+
+	for (TArray<FTether>& NewTethers : NewTethersSlots)
+	{
+		NewTethers.Reserve(MaxNumTetherIslands);
+	}
+
+	TAtomic<int32> NumTethers(0);
+
+	PhysicsParallelFor(Nodes.Num(), [&](int32 Index)
+	{
+		const int32 Node = Nodes[Index];
+		if (Particles.InvM(Node) == (FReal)0.)
 		{
-			if (!Elements.Num()) { continue; }  // Empty island 
+			return;  // No tethers on dynamic particles
+		}
 
-			int32 ClosestElement = INDEX_NONE;
-			FReal ClosestDistance = FLT_MAX;
+		// Find the closest seeds in each island
+		TArray<int32, TInlineAllocator<32>> ClosestSeeds;
+		ClosestSeeds.Reserve(Islands.Num());
 
-			for (const uint32 Element : Elements)
+		for (const TArray<int32>& Seeds : Islands)
+		{
+			if (!Seeds.Num())
 			{
-				const FReal Distance = GeodesicPaths[TVector<uint32, 2>(Element, i)].First;
+				continue;  // Empty island 
+			}
+
+			int32 ClosestSeed = INDEX_NONE;
+			FReal ClosestDistance = TNumericLimits<FReal>::Max();
+
+			for (const int32 Seed : Seeds)
+			{
+				const FReal Distance = GeodesicDistances[TVec2<int32>(Seed, Node)];
 				if (Distance < ClosestDistance)
 				{
 					ClosestDistance = Distance;
-					ClosestElement = Element;
+					ClosestSeed = Seed;
 				}
 			}
-			if (ClosestElement == INDEX_NONE) { continue; }  // Not on this island
-
-			const TVector<uint32, 2> Index(ClosestElement, i);
-			check(GeodesicPaths[Index].First != FLT_MAX);
-			check(GeodesicPaths[Index].Second.Num() > 1);
-			ClosestElements.Add(MakePair(ClosestDistance, ClosestElement));
+			if (ClosestSeed != INDEX_NONE)
+			{
+				ClosestSeeds.Add(ClosestSeed);
+			}
 		}
-		// How to sort based on smalled first value of pair....
-		ClosestElements.Sort();
-		if (NumberOfAttachments < ClosestElements.Num())
+
+		// Sort all the tethers for this node based on smallest distance
+		ClosestSeeds.Sort([&GeodesicDistances, Node](int32 Seed1, int32 Seed2)
+			{
+				return GeodesicDistances[TVec2<int32>(Seed1, Node)] < GeodesicDistances[TVec2<int32>(Seed2, Node)];
+			});
+
+		// Keep only N closest tethers
+		if (MaxNumTetherIslands < ClosestSeeds.Num())
 		{
-			ClosestElements.SetNum(NumberOfAttachments);
+			ClosestSeeds.SetNum(MaxNumTetherIslands);
 		}
-		for (const Pair<FReal, int32>& Element : ClosestElements)
+
+		// Add these tethers to the N (or less) available slots
+		for (const int32 Seed : ClosestSeeds)
 		{
-			const TVector<uint32, 2> Index(Element.Second, i);
-			check(GeodesicPaths[Index].First == Element.First);
-			check(FGenericPlatformMath::Abs(Element.First - ComputeGeodesicDistance(InParticles, GeodesicPaths[Index].Second)) < 1e-4);
-			CriticalSection.Lock();
-			NewConstraints.Add(GeodesicPaths[Index].Second);
-			MDists.Add(Element.First);
-			CriticalSection.Unlock();
+			const FReal RefLength = GeodesicDistances[TVec2<int32>(Seed, Node)];
+			NewTethersSlots[Index].Emplace(Seed, Node, RefLength);
+			++NumTethers;
 		}
+		check(NewTethersSlots[Index].Num() <= MaxNumTetherIslands);
 	});
-	// TODO(mlentine): This should work by just reverse sorting and not needing the filtering but it may not be guaranteed. Work out if this is actually guaranteed or not.
-	NewConstraints.Sort([](const TArray<uint32>& Elem1, const TArray<uint32>& Elem2) { return Elem1.Num() > Elem2.Num(); });
-	TArray<FReal> NewDists;
-	TMap<TVector<uint32, 2>, TArray<uint32>> ProcessedPairs;
-	for (uint32 i = 0; i < static_cast<uint32>(NewConstraints.Num()); ++i)
+
+	// Consolidate the N slots into one single list of tethers, spreading the nodes across different ranges for parallel processing
+	Tethers.Reset(NumTethers);
+
+	for (int32 Slot = 0; Slot < MaxNumTetherIslands; ++Slot)
 	{
-		const TVector<uint32, 2> Traverse(NewConstraints[i][0], NewConstraints[i].Last());
-		if (const TArray<uint32>* TraversePath = ProcessedPairs.Find(Traverse))
-		{
-			check(NewConstraints[i].Num() == TraversePath->Num());
-			for (uint32 j = 0; j < static_cast<uint32>(TraversePath->Num()); ++j)
-			{
-				check((*TraversePath)[j] == NewConstraints[i][j]);
-			}
-			continue;
-		}
-		TArray<uint32> Path;
-		FReal Dist = (FReal)0.;
-		Path.Add(NewConstraints[i][0]);
-		for (uint32 j = 1; j < static_cast<uint32>(NewConstraints[i].Num() - 1); ++j)
+		int32 Size = 0;
+		for (int32 Index = 0; Index < Nodes.Num(); ++Index)
 		{
-			Dist += (InParticles.X(NewConstraints[i][j]) - InParticles.X(NewConstraints[i][j - 1])).Size();
-			Path.Add(NewConstraints[i][j]);
-			switch (MMode)
+			if (NewTethersSlots[Index].IsValidIndex(Slot))
 			{
-			case EMode::AccurateTetherFastLength:
-				MEuclideanConstraints.Add({Path[0], Path[Path.Num() - 1]}) ;                // Use euclidean (beeline) distance
-				NewDists.Add(ComputeDistance(InParticles, Path[0], Path[Path.Num() - 1]));  // Fastest method in ISPC
-				break;
-			case EMode::AccurateTetherAccurateLength:
-				MGeodesicConstraints.Add(Path);
-				NewDists.Add(Dist);
-				break;
-			default:
-				unimplemented();
-				break;
+				Tethers.Emplace(NewTethersSlots[Index][Slot]);
+				++Size;
 			}
-			const TVector<uint32, 2> SubTraverse(NewConstraints[i][0], NewConstraints[i][j]);
-			ProcessedPairs.Add(SubTraverse, Path);
 		}
+		TethersView.AddRange(Size);
 	}
-	MDists = NewDists;
 }
diff --git a/Engine/Source/Runtime/Experimental/Chaos/Private/Chaos/TriangleMesh.cpp b/Engine/Source/Runtime/Experimental/Chaos/Private/Chaos/TriangleMesh.cpp
--- a/Engine/Source/Runtime/Experimental/Chaos/Private/Chaos/TriangleMesh.cpp
+++ b/Engine/Source/Runtime/Experimental/Chaos/Private/Chaos/TriangleMesh.cpp
@@ -128,7 +128,7 @@ void FTriangleMesh::GetVertexSet(TSet<int32>& VertexSet) const
 	}
 }
 
-const TMap<int32, TSet<uint32>>& FTriangleMesh::GetPointToNeighborsMap() const
+const TMap<int32, TSet<int32>>& FTriangleMesh::GetPointToNeighborsMap() const
 {
 	if (MPointToNeighborsMap.Num())
 	{
@@ -137,9 +137,9 @@ const TMap<int32, TSet<uint32>>& FTriangleMesh::GetPointToNeighborsMap() const
 	MPointToNeighborsMap.Reserve(MNumIndices);
 	for (int i = 0; i < MElements.Num(); ++i)
 	{
-		TSet<uint32>& Elems0 = MPointToNeighborsMap.FindOrAdd(MElements[i][0]);
-		TSet<uint32>& Elems1 = MPointToNeighborsMap.FindOrAdd(MElements[i][1]);
-		TSet<uint32>& Elems2 = MPointToNeighborsMap.FindOrAdd(MElements[i][2]);
+		TSet<int32>& Elems0 = MPointToNeighborsMap.FindOrAdd(MElements[i][0]);
+		TSet<int32>& Elems1 = MPointToNeighborsMap.FindOrAdd(MElements[i][1]);
+		TSet<int32>& Elems2 = MPointToNeighborsMap.FindOrAdd(MElements[i][2]);
 		Elems0.Reserve(Elems0.Num() + 2);
 		Elems1.Reserve(Elems1.Num() + 2);
 		Elems2.Reserve(Elems2.Num() + 2);
diff --git a/Engine/Source/Runtime/Experimental/Chaos/Public/Chaos/PBDActiveView.h b/Engine/Source/Runtime/Experimental/Chaos/Public/Chaos/PBDActiveView.h
--- a/Engine/Source/Runtime/Experimental/Chaos/Public/Chaos/PBDActiveView.h
+++ b/Engine/Source/Runtime/Experimental/Chaos/Public/Chaos/PBDActiveView.h
@@ -31,13 +31,17 @@ namespace Chaos
 		// Execute the specified function on all active items.
 		void SequentialFor(TFunctionRef<void(TItemsType&, int32)> Function) const;
 
-		// Execute the specified function in parallel on all active items. Set MinParallelSize to run sequential on the smaller ranges.
+		// Execute the specified function in parallel, on all items for each active range (sequential range, parallel items). Set MinParallelSize to run sequential on the smaller ranges.
 		void ParallelFor(TFunctionRef<void(TItemsType&, int32)> Function, int32 MinParallelSize = TNumericLimits<int32>::Max()) const;
-		void ParallelFor(TFunctionRef<void(TItemsType&, int32)> Function, bool bParallelRange, int32 MinParallelSize = TNumericLimits<int32>::Max()) const;
 
-		// Execute the specified function on all active items. Callee responsible for inner loop.
+		// Execute the specified function in nested parallel for loops, on all items for each active range (parallel range, parallel items). Set MinParallelSize to run sequential on the smaller ranges.
+		void ParallelFor(TFunctionRef<void(TItemsType&, int32)> Function, bool bForceSingleThreadedRange, int32 MinParallelSize = TNumericLimits<int32>::Max()) const;
+
+		// Execute the specified function in sequence for all active range. Callee responsible for inner loop.
 		void RangeFor(TFunctionRef<void(TItemsType&, int32, int32)> Function) const;
-		void RangeFor(TFunctionRef<void(TItemsType&, int32, int32)> Function, bool bParallelRange) const;
+
+		// Execute the specified function in parallel for all active ranges. Callee responsible for inner loop.
+		void RangeFor(TFunctionRef<void(TItemsType&, int32, int32)> Function, bool bForceSingleThreadedRange) const;
 
 		// Remove all ranges above the current given size.
 		void Reset(int32 Offset = 0);
diff --git a/Engine/Source/Runtime/Experimental/Chaos/Public/Chaos/PBDAnimDriveConstraint.h b/Engine/Source/Runtime/Experimental/Chaos/Public/Chaos/PBDAnimDriveConstraint.h
--- a/Engine/Source/Runtime/Experimental/Chaos/Public/Chaos/PBDAnimDriveConstraint.h
+++ b/Engine/Source/Runtime/Experimental/Chaos/Public/Chaos/PBDAnimDriveConstraint.h
@@ -1,18 +1,15 @@
 // Copyright Epic Games, Inc. All Rights Reserved.
 #pragma once
 
-#include "Chaos/ParticleRule.h"
 #include "Chaos/Framework/Parallel.h"
-#include "Containers/ArrayView.h"
-#include "Containers/ContainersFwd.h"
+#include "Chaos/PBDStiffness.h"
 #include "ChaosStats.h"
 
 DECLARE_CYCLE_STAT(TEXT("Chaos PBD Anim Drive Constraint"), STAT_PBD_AnimDriveConstraint, STATGROUP_Chaos);
-DECLARE_CYCLE_STAT(TEXT("Chaos PBD Anim Drive Constraint Set Stiffness"), STAT_PBD_AnimDriveConstraintSetStiffness, STATGROUP_Chaos);
 
 namespace Chaos
 {
-	class FPBDAnimDriveConstraint : public FParticleRule
+	class FPBDAnimDriveConstraint final
 	{
 	public:
 		FPBDAnimDriveConstraint(
@@ -27,127 +24,71 @@ namespace Chaos
 			, OldAnimationPositions(InOldAnimationPositions)
 			, ParticleOffset(InParticleOffset)
 			, ParticleCount(InParticleCount)
-			, Stiffness((FReal)0., (FReal)1.)
-			, Damping((FReal)0., (FReal)1.)
+			, Stiffness(StiffnessMultipliers, FVec2((FReal)0., (FReal)1.), InParticleCount)
+			, Damping(DampingMultipliers, FVec2((FReal)0., (FReal)1.), InParticleCount)
 		{
-			static_assert(StiffnessTableSize > 0 && StiffnessTableSize < 256, "The Stiffness lookup table is restricted to uint8 sized indices");
-			static_assert(DampingTableSize > 0 && DampingTableSize < 256, "The Damping lookup tables is restricted to uint8 sized indices");
-
-			// Convert the weight maps into an array of lookup indices to the stiffness table
-			auto IntializeTableIndices = [this](const TConstArrayView<FReal>& Multipliers, const int32 TableSize, TArray<uint8>& OutIndices)
-			{
-				if (Multipliers.Num() == ParticleCount)
-				{
-					OutIndices.AddUninitialized(ParticleCount);
+		}
 
-					static const FReal TableScale = (FReal)(TableSize - 1);
+		~FPBDAnimDriveConstraint() {}
 
-					for (int32 Index = 0; Index < ParticleCount; ++Index)
-					{
-						OutIndices[Index] = (uint8)(FMath::Clamp(Multipliers[Index], (FReal)0., (FReal)1.) * TableScale);
-					}
-				}
-			};
+		// Return the stiffness input values used by the constraint
+		FVec2 GetStiffness() const { return Stiffness.GetWeightedValue(); }
 
-			IntializeTableIndices(StiffnessMultipliers, StiffnessTableSize, StiffnessIndices);
-			IntializeTableIndices(DampingMultipliers, DampingTableSize, DampingIndices);
+		// Return the damping input values used by the constraint
+		FVec2 GetDamping() const { return Damping.GetWeightedValue(); }
 
-			StiffnessTable.AddZeroed(StiffnessTableSize);
-			DampingTable.AddZeroed(DampingTableSize);
+		inline void SetProperties(const FVec2& InStiffness, const FVec2& InDamping)
+		{
+			Stiffness.SetWeightedValue(InStiffness);
+			Damping.SetWeightedValue(InDamping);
 		}
 
-		virtual ~FPBDAnimDriveConstraint() {}
-
 		// Set stiffness offset and range, as well as the simulation stiffness exponent
 		inline void ApplyProperties(const FReal Dt, const int32 NumIterations)
 		{
-			SCOPE_CYCLE_COUNTER(STAT_PBD_AnimDriveConstraintSetStiffness);
-
-			// Define the stiffness mapping function
-			static const FReal ParameterFitLogBase = FMath::Loge(ParameterFitBase);
-			const FReal Exponent = Dt * ParameterFrequency / (FReal)NumIterations;
-
-			auto SimulationValue = [Exponent](const FReal InValue)->FReal
-			{
-				// Get a very steep exponential curve between the [0, 1] range to make easier to set the parameter
-				// The base has been chosen empirically
-				// ParameterValue = Pow(ParameterFitBase, ParameterValue - 1)
-				const FReal ParameterFit = FMath::Exp(ParameterFitLogBase * (FMath::Clamp(InValue, (FReal)0., (FReal)1.) - (FReal)1.));
-
-				// Use simulation dependent stiffness exponent to alleviate the variations in effect when Dt and NumIterations change
-				// This is based on the Position-Based Simulation Methods paper (page 8),
-				// but uses the delta time in addition of the number of iterations in the calculation of the error term.
-				const FReal LogValue = FMath::Loge((FReal)1. - ParameterFit);
-				return (FReal)1. - FMath::Exp(LogValue * Exponent);
-			};
-
-			auto InitializeTable = [&SimulationValue](const FVec2& WeightedValue, const int32 TableSize, TArray<FReal>& OutTable)
-			{
-				const FReal Offset = WeightedValue[0];
-				const FReal Range = WeightedValue[1] - WeightedValue[0];
-				static const FReal WeightIncrement = (FReal)1. / (FReal)(TableSize - 1); // Must allow full range from 0 to 1 included
-				for (int32 Index = 0; Index < TableSize; ++Index)
-				{
-					const FReal Weight = (FReal)Index * WeightIncrement;
-					OutTable[Index] = SimulationValue(Offset + Weight * Range);
-				}
-			};
-
-			InitializeTable(Stiffness, StiffnessTableSize, StiffnessTable);
-			InitializeTable(Damping, DampingTableSize, DampingTable);
-		}
-
-		// Return the stiffness input values used by the constraint
-		inline FVec2 GetStiffness() const { return Stiffness; }
-
-		// Return the damping input values used by the constraint
-		inline FVec2 GetDamping() const { return Damping; }
-
-		void SetProperties(const FVec2& InStiffness, const FVec2& InDamping)
-		{
-			Stiffness = InStiffness.ClampAxes((FReal)0., (FReal)1.);
-			Damping = InDamping.ClampAxes((FReal)0., (FReal)1.);
+			Stiffness.ApplyValues(Dt, NumIterations);
+			Damping.ApplyValues(Dt, NumIterations);
 		}
 
-		inline virtual void Apply(FPBDParticles& InParticles, const FReal Dt) const override
+		inline void Apply(FPBDParticles& InParticles, const FReal Dt) const
 		{
 			SCOPE_CYCLE_COUNTER(STAT_PBD_AnimDriveConstraint);
 
-			if (StiffnessIndices.Num() == ParticleCount)
+			if (Stiffness.HasWeightMap())
 			{
-				if (DampingIndices.Num() == ParticleCount)
+				if (Damping.HasWeightMap())
 				{
 					PhysicsParallelFor(ParticleCount, [&](int32 Index)  // TODO: profile needed for these parallel loop based on particle count
 					{
-						const FReal ParticleStiffness = StiffnessTable[StiffnessIndices[Index]];
-						const FReal ParticleDamping = DampingTable[DampingIndices[Index]];
+						const FReal ParticleStiffness = Stiffness[Index];
+						const FReal ParticleDamping = Damping[Index];
 						ApplyHelper(InParticles, ParticleStiffness, ParticleDamping, Dt, Index);
 					});
 				}
 				else
 				{
-					const FReal ParticleDamping = DampingTable[0];
+					const FReal ParticleDamping = (FReal)Damping;
 					PhysicsParallelFor(ParticleCount, [&](int32 Index)
 					{
-						const FReal ParticleStiffness = StiffnessTable[StiffnessIndices[Index]];
+						const FReal ParticleStiffness = Stiffness[Index];
 						ApplyHelper(InParticles, ParticleStiffness, ParticleDamping, Dt, Index);
 					});
 				}
 			}
 			else
 			{
-				const FReal ParticleStiffness = StiffnessTable[0];
-				if (DampingIndices.Num() == ParticleCount)
+				const FReal ParticleStiffness = (FReal)Stiffness;
+				if (Damping.HasWeightMap())
 				{
 					PhysicsParallelFor(ParticleCount, [&](int32 Index)
 					{
-						const FReal ParticleDamping = DampingTable[DampingIndices[Index]];
+						const FReal ParticleDamping = Damping[Index];
 						ApplyHelper(InParticles, ParticleStiffness, ParticleDamping, Dt, Index);
 					});
 				}
 				else
 				{
-					const FReal ParticleDamping = DampingTable[0];
+					const FReal ParticleDamping = (FReal)Damping;
 					PhysicsParallelFor(ParticleCount, [&](int32 Index)
 					{
 						ApplyHelper(InParticles, ParticleStiffness, ParticleDamping, Dt, Index);
@@ -177,23 +118,13 @@ namespace Chaos
 		}
 
 	private:
-		static constexpr FReal ParameterFrequency = (FReal)120.;  // 60Hz @ 2 iterations as a root for all stiffness values TODO: Make this a solver parameter
-		static constexpr FReal ParameterFitBase = (FReal)1.e3;  // Logarithm base to use in the stiffness parameter fit function
-		static constexpr int32 StiffnessTableSize = 16;  // Size of the lookup table, can't be more than 256 values
-		static constexpr int32 DampingTableSize = 16;  // Size of the lookup table, can't be more than 256 values
-
 		const TArray<FVec3>& AnimationPositions;  // Use global index (needs adding ParticleOffset)
 		const TArray<FVec3>& OldAnimationPositions;  // Use global index (needs adding ParticleOffset)
 		const int32 ParticleOffset;
 		const int32 ParticleCount;
 
-		TArray<uint8> StiffnessIndices; // Per particle array of index to the stiffness table. Use local index
-		TArray<FReal> StiffnessTable;  // Fixed lookup table of stiffness values, use uint8 indexation
-		TArray<uint8> DampingIndices; // Per particle array of index to the stiffness table. Use local index
-		TArray<FReal> DampingTable;  // Fixed lookup table of stiffness values, use uint8 indexation
-
-		FVec2 Stiffness;  // Input stiffness before parameter fitting
-		FVec2 Damping;  // Input damping before parameter fitting
+		FPBDStiffness Stiffness;
+		FPBDStiffness Damping;
 	};
 
 	template<typename T, int d>
diff --git a/Engine/Source/Runtime/Experimental/Chaos/Public/Chaos/PBDEvolution.h b/Engine/Source/Runtime/Experimental/Chaos/Public/Chaos/PBDEvolution.h
--- a/Engine/Source/Runtime/Experimental/Chaos/Public/Chaos/PBDEvolution.h
+++ b/Engine/Source/Runtime/Experimental/Chaos/Public/Chaos/PBDEvolution.h
@@ -76,8 +76,8 @@ class CHAOS_API FPBDEvolution : public TArrayCollection
 	void ActivateConstraintRuleRange(int32 Offset, bool bActivate) { MConstraintRulesActiveView.ActivateRange(Offset, bActivate); }
 
 	// Constraint accessors
-	const TArray<TFunction<void(const FPBDParticles&)>>& ConstraintInits() const { return MConstraintInits; }
-	TArray<TFunction<void(const FPBDParticles&)>>& ConstraintInits() { return MConstraintInits; }
+	const TArray<TFunction<void(const FPBDParticles&, const FReal)>>& ConstraintInits() const { return MConstraintInits; }
+	TArray<TFunction<void(const FPBDParticles&, const FReal)>>& ConstraintInits() { return MConstraintInits; }
 	const TArray<TFunction<void(FPBDParticles&, const FReal)>>& ConstraintRules() const { return MConstraintRules; }
 	TArray<TFunction<void(FPBDParticles&, const FReal)>>& ConstraintRules() { return MConstraintRules; }
 	
@@ -158,8 +158,8 @@ private:
 	TArrayCollectionArray<FReal> MGroupDampings;
 	TArrayCollectionArray<bool> MGroupUseCCDs;
 	
-	TArray<TFunction<void(const FPBDParticles&)>> MConstraintInits;
-	TPBDActiveView<TArray<TFunction<void(const FPBDParticles&)>>> MConstraintInitsActiveView;
+	TArray<TFunction<void(const FPBDParticles&, const FReal)>> MConstraintInits;
+	TPBDActiveView<TArray<TFunction<void(const FPBDParticles&, const FReal)>>> MConstraintInitsActiveView;
 	TArray<TFunction<void(FPBDParticles&, const FReal)>> MConstraintRules;
 	TPBDActiveView<TArray<TFunction<void(FPBDParticles&, const FReal)>>> MConstraintRulesActiveView;
 
diff --git a/Engine/Source/Runtime/Experimental/Chaos/Public/Chaos/PBDLongRangeConstraints.h b/Engine/Source/Runtime/Experimental/Chaos/Public/Chaos/PBDLongRangeConstraints.h
--- a/Engine/Source/Runtime/Experimental/Chaos/Public/Chaos/PBDLongRangeConstraints.h
+++ b/Engine/Source/Runtime/Experimental/Chaos/Public/Chaos/PBDLongRangeConstraints.h
@@ -2,49 +2,40 @@
 #pragma once
 
 #include "Chaos/PBDLongRangeConstraintsBase.h"
-#include "Chaos/PBDParticles.h"
-#include "Chaos/PBDConstraintContainer.h"
 #include "ChaosStats.h"
 
 DECLARE_CYCLE_STAT(TEXT("Chaos PBD Long Range Constraint"), STAT_PBD_LongRange, STATGROUP_Chaos);
 
 namespace Chaos
 {
-class CHAOS_API FPBDLongRangeConstraints : public FPBDLongRangeConstraintsBase, public FPBDConstraintContainer
+class CHAOS_API FPBDLongRangeConstraints final : public FPBDLongRangeConstraintsBase
 {
 public:
 	typedef FPBDLongRangeConstraintsBase Base;
+	typedef typename Base::FTether FTether;
 	typedef typename Base::EMode EMode;
 
-	using Base::GetMode;
-
 	FPBDLongRangeConstraints(
-		const FDynamicParticles& InParticles,
-		const TMap<int32, TSet<uint32>>& PointToNeighbors,
-		const int32 NumberOfAttachments = 1,
-		const FReal Stiffness = (FReal)1.,
+		const FPBDParticles& Particles,
+		const int32 InParticleOffset,
+		const int32 InParticleCount,
+		const TMap<int32, TSet<int32>>& PointToNeighbors,
+		const TConstArrayView<FReal>& StiffnessMultipliers,
+		const int32 MaxNumTetherIslands = 4,
+		const FVec2& InStiffness = FVec2((FReal)1., (FReal)1.),
 		const FReal LimitScale = (FReal)1.,
-		const EMode Mode = EMode::AccurateTetherFastLength)
-		: FPBDLongRangeConstraintsBase(InParticles, PointToNeighbors, NumberOfAttachments, Stiffness, LimitScale, Mode) {}
+		const EMode InMode = EMode::Geodesic)
+		: FPBDLongRangeConstraintsBase(Particles, InParticleOffset, InParticleCount, PointToNeighbors, StiffnessMultipliers, MaxNumTetherIslands, InStiffness, LimitScale, InMode) {}
 	virtual ~FPBDLongRangeConstraints() {}
 
-	void Apply(FPBDParticles& InParticles, const FReal Dt, const TArray<int32>& InConstraintIndices) const;
-	void Apply(FPBDParticles& InParticles, const FReal Dt) const;
+	void Apply(FPBDParticles& Particles, const FReal Dt, const TArray<int32>& ConstraintIndices) const;
+	void Apply(FPBDParticles& Particles, const FReal Dt) const;
 
 private:
-	using Base::MEuclideanConstraints;
-	using Base::MGeodesicConstraints;
-	using Base::MDists;
-	using Base::GetDelta;
-
-	template<class TConstraintType>
-	void Apply(const TConstraintType& Constraint, FPBDParticles& InParticles, const FReal Dt, const FReal RefDist) const
-	{
-		const int32 i2 = Constraint[Constraint.Num() - 1];
-		checkSlow(InParticles.InvM(i2) > (FReal)0.);
-		InParticles.P(i2) += GetDelta(Constraint, InParticles, RefDist);
-	}
-
+	using Base::Tethers;
+	using Base::TethersView;
+	using Base::Stiffness;
+	using Base::ParticleOffset;
 };
 }
 
diff --git a/Engine/Source/Runtime/Experimental/Chaos/Public/Chaos/PBDLongRangeConstraintsBase.h b/Engine/Source/Runtime/Experimental/Chaos/Public/Chaos/PBDLongRangeConstraintsBase.h
--- a/Engine/Source/Runtime/Experimental/Chaos/Public/Chaos/PBDLongRangeConstraintsBase.h
+++ b/Engine/Source/Runtime/Experimental/Chaos/Public/Chaos/PBDLongRangeConstraintsBase.h
@@ -2,8 +2,9 @@
 #pragma once
 
 #include "Chaos/Core.h"
-#include "Chaos/Array.h"
 #include "Chaos/PBDParticles.h"
+#include "Chaos/PBDActiveView.h"
+#include "Chaos/PBDStiffness.h"
 #include "Containers/Map.h"
 #include "Containers/Set.h"
 
@@ -14,107 +15,81 @@ class CHAOS_API FPBDLongRangeConstraintsBase
 public:
 	enum class EMode : uint8
 	{
-		FastTetherFastLength,
-		AccurateTetherFastLength,
-		AccurateTetherAccurateLength,
-	};
-
-	FPBDLongRangeConstraintsBase(
-		const FDynamicParticles& InParticles,
-		const TMap<int32, TSet<uint32>>& PointToNeighbors,
-		const int32 NumberOfAttachments = 1,
-		const FReal Stiffness = (FReal)1.,
-		const FReal LimitScale = (FReal)1.,
-		const EMode Mode = EMode::AccurateTetherFastLength);
-
-	virtual ~FPBDLongRangeConstraintsBase() {}
-
-	EMode GetMode() const { return MMode; }
+		Euclidean,
+		Geodesic,
 
-	const TArray<TVec2<uint32>>& GetEuclideanConstraints() const { return MEuclideanConstraints; }
-	const TArray<TArray<uint32>>& GetGeodesicConstraints() const { return MGeodesicConstraints; }
-
-	const TArray<FReal>& GetDists() const { return MDists; }
+		// Deprecated modes
+		FastTetherFastLength = Euclidean,
+		AccurateTetherFastLength = Geodesic,
+		AccurateTetherAccurateLength = Geodesic
+	};
 
-	static TArray<TArray<uint32>> ComputeIslands(const TMap<int32, TSet<uint32>>& PointToNeighbors, const TArray<uint32>& KinematicParticles);
+	struct FTether
+	{
+		int32 Start;
+		int32 End;
+		FReal RefLength;
 
-	void SetStiffness(FReal InStiffness) { MStiffness = FMath::Clamp(InStiffness, (FReal)0., (FReal)1.); }
+		FTether(int32 InStart, int32 InEnd, FReal InRefLength): Start(InStart), End(InEnd), RefLength(InRefLength) {}
 
-protected:
-	template<class TConstraintType>
-	inline FVec3 GetDelta(const TConstraintType& Constraint, const FPBDParticles& InParticles, const FReal RefDist) const
-	{
-		checkSlow(Constraint.Num() > 1);
-		const uint32 i2 = Constraint[Constraint.Num() - 1];
-		const uint32 i2m1 = Constraint[Constraint.Num() - 2];
-		checkSlow(InParticles.InvM(Constraint[0]) == (FReal)0.);
-		checkSlow(InParticles.InvM(i2) > (FReal)0.);
-		const FReal Distance = ComputeGeodesicDistance(InParticles, Constraint); // This function is used for either Euclidean or Geodesic distances
-		if (Distance < RefDist)
+		inline FVec3 GetDelta(const FPBDParticles& Particles) const
+		{
+			checkSlow(Particles.InvM(Start) == (FReal)0.);
+			checkSlow(Particles.InvM(End) > (FReal)0.);
+			FVec3 Direction = Particles.P(Start) - Particles.P(End);
+			const FReal Length = Direction.SafeNormalize();
+			const FReal Offset = Length - RefLength;
+			return Offset < (FReal)0. ? FVec3((FReal)0.) : Offset * Direction;
+		};
+
+		inline void GetDelta(const FPBDParticles& Particles, FVec3& OutDirection, FReal& OutOffset) const
 		{
-			return FVec3((FReal)0.);
-		}
+			checkSlow(Particles.InvM(Start) == (FReal)0.);
+			checkSlow(Particles.InvM(End) > (FReal)0.);
+			OutDirection = Particles.P(Start) - Particles.P(End);
+			const FReal Length = OutDirection.SafeNormalize();
+			OutOffset = FMath::Max((FReal)0., Length - RefLength);
+		};
+	};
 
-		//const FVec3 Direction = (InParticles.P(i2m1) - InParticles.P(i2)).GetSafeNormal();
-		FVec3 Direction = InParticles.P(i2m1) - InParticles.P(i2);
-		const FReal DirLen = Direction.SafeNormalize();
+	FPBDLongRangeConstraintsBase(
+		const FPBDParticles& Particles,
+		const int32 InParticleOffset,
+		const int32 InParticleCount,
+		const TMap<int32, TSet<int32>>& PointToNeighbors,
+		const TConstArrayView<FReal>& StiffnessMultipliers,
+		const int32 MaxNumTetherIslands = 4,
+		const FVec2& InStiffness = FVec2((FReal)0., (FReal)1.),
+		const FReal LimitScale = (FReal)1,
+		const EMode InMode = EMode::AccurateTetherFastLength);
 
-		const FReal Offset = Distance - RefDist;
-		const FVec3 Delta = MStiffness * Offset * Direction;
+	virtual ~FPBDLongRangeConstraintsBase() {}
 
-	/*  // ryan - this currently fails:
+	EMode GetMode() const { return Mode; }
 
-		const FReal NewDirLen = (InParticles.P(i2) + Delta - InParticles.P(i2m1)).Size();
-		//FReal Correction = (InParticles.P(i2) - InParticles.P(i2m1)).Size() - (InParticles.P(i2) + Delta - InParticles.P(i2m1)).Size();
-		const FReal Correction = DirLen - NewDirLen;
-		check(Correction >= 0);
+	// Return the stiffness input values used by the constraint
+	FVec2 GetStiffness() const { return Stiffness.GetWeightedValue(); }
 
-		//FReal NewDist = (Distance - (InParticles.P(i2) - InParticles.P(i2m1)).Size() + (InParticles.P(i2) + Delta - InParticles.P(i2m1)).Size());
-		const FReal NewDist = Distance - DirLen + NewDirLen;
-		check(FGenericPlatformMath::Abs(NewDist - RefDist) < 1e-4);
-	*/
-		return Delta;
-	};
+	// Set the stiffness input values used by the constraint
+	void SetStiffness(const FVec2& InStiffness) { Stiffness.SetWeightedValue(InStiffness); }
 
-	static FReal ComputeGeodesicDistance(const FParticles& InParticles, const TArray<uint32>& Path)
-	{
-		FReal distance = 0;
-		for (int32 i = 0; i < Path.Num() - 1; ++i)
-		{
-			distance += (InParticles.X(Path[i]) - InParticles.X(Path[i + 1])).Size();
-		}
-		return distance;
-	}
-	static FReal ComputeGeodesicDistance(const FPBDParticles& InParticles, const TArray<uint32>& Path)
-	{
-		FReal distance = 0;
-		for (int32 i = 0; i < Path.Num() - 1; ++i)
-		{
-			distance += (InParticles.P(Path[i]) - InParticles.P(Path[i + 1])).Size();
-		}
-		return distance;
-	}
+	// Set stiffness offset and range, as well as the simulation stiffness exponent
+	void ApplyProperties(const FReal Dt, const int32 NumIterations) { Stiffness.ApplyValues(Dt, NumIterations); }
 
-	void ComputeEuclideanConstraints(const FDynamicParticles& InParticles, const TMap<int32, TSet<uint32>>& PointToNeighbors, const int32 NumberOfAttachments);
-	void ComputeGeodesicConstraints(const FDynamicParticles& InParticles, const TMap<int32, TSet<uint32>>& PointToNeighbors, const int32 NumberOfAttachments);
+	const TArray<FTether>& GetTethers() const { return Tethers; }
 
-	static FReal ComputeDistance(const FParticles& InParticles, const uint32 i, const uint32 j) { return (InParticles.X(i) - InParticles.X(j)).Size(); }
-	static FReal ComputeDistance(const FPBDParticles& InParticles, const uint32 i, const uint32 j) { return (InParticles.P(i) - InParticles.P(j)).Size(); }
+	static TArray<TArray<int32>> ComputeIslands(const TMap<int32, TSet<int32>>& PointToNeighbors, const TArray<int32>& KinematicParticles);
 
-	static FReal ComputeGeodesicDistance(const FParticles& InParticles, const TVector<uint32, 2>& Path)
-	{
-		return (InParticles.X(Path[0]) - InParticles.X(Path[1])).Size();
-	}
-	static FReal ComputeGeodesicDistance(const FPBDParticles& InParticles, const TVector<uint32, 2>& Path)
-	{
-		return (InParticles.P(Path[0]) - InParticles.P(Path[1])).Size();
-	}
+protected:
+	void ComputeEuclideanConstraints(const FPBDParticles& Particles, const TMap<int32, TSet<int32>>& PointToNeighbors, const int32 NumberOfAttachments);
+	void ComputeGeodesicConstraints(const FPBDParticles& Particles, const TMap<int32, TSet<int32>>& PointToNeighbors, const int32 NumberOfAttachments);
 
 protected:
-	TArray<TVec2<uint32>> MEuclideanConstraints;
-	TArray<TArray<uint32>> MGeodesicConstraints;
-	TArray<FReal> MDists;
-	FReal MStiffness;
-	EMode MMode;
+	TArray<FTether> Tethers;
+	TPBDActiveView<TArray<FTether>> TethersView;
+	FPBDStiffness Stiffness;
+	const EMode Mode;
+	const int32 ParticleOffset;
+	const int32 ParticleCount;
 };
 }
diff --git a/Engine/Source/Runtime/Experimental/Chaos/Public/Chaos/PBDStiffness.h b/Engine/Source/Runtime/Experimental/Chaos/Public/Chaos/PBDStiffness.h
new file mode 100644
--- /dev/null
+++ b/Engine/Source/Runtime/Experimental/Chaos/Public/Chaos/PBDStiffness.h
@@ -0,0 +1,145 @@
+// Copyright Epic Games, Inc. All Rights Reserved.
+#pragma once
+
+#include "Chaos/Core.h"
+#include "Chaos/Framework/Parallel.h"
+#include "Containers/ArrayView.h"
+#include "Containers/ContainersFwd.h"
+#include "ChaosStats.h"
+
+DECLARE_CYCLE_STAT(TEXT("Chaos PBD Stiffness Apply Values"), STAT_PBD_StiffnessApplyValues, STATGROUP_Chaos);
+
+namespace Chaos
+{
+	/**
+	 * Stiffness class for managing real time update to the weight map and low/high value ranges
+	 * and to exponentiate the stiffness value depending on the iterations and Dt.
+	 */
+	class FPBDStiffness final
+	{
+	public:
+		/**
+		 * Weightmap constructor. 
+		 */
+		inline FPBDStiffness(
+			const TConstArrayView<FReal>& Multipliers,
+			const FVec2& InWeightedValue,
+			int32 ParticleCount,
+			int32 TableSize = 16,  // Size of the lookup table, can't be more than 256 values, the larger the table the longer it takes to apply changes to the stiffness values
+			FReal InParameterFitBase = (FReal)1.e3);  // Logarithm base to use in the stiffness parameter fit function
+
+		~FPBDStiffness() {}
+
+		/** Return the number of values stored in the weight map. */
+		int32 Num() const { return Indices.Num(); }
+
+		/** Return whether this object contains weight map values. */
+		bool HasWeightMap() const { return Table.Num() > 1; }
+
+		/**
+		 * Set the low and high values of the weight map.
+		 * The weight map table only gets updated after ApplyValues is called.
+		 */
+		void SetWeightedValue(const FVec2& InWeightedValue) { WeightedValue = InWeightedValue.ClampAxes((FReal)0., (FReal)1.); }
+
+		/**
+		  * Return the low and high values set for this weight map.
+		  * Both values will always be between 0 and 1 due to having been clamped in SetWeightedValue.
+		  */
+		const FVec2& GetWeightedValue() const { return WeightedValue; }
+
+		/** Update the weight map table with the current simulation parameters. */
+		inline void ApplyValues(const FReal Dt, const int32 NumIterations);
+
+		/**
+		 * Lookup for the exponential weighted value at the specified weight map index.
+		 * This function will assert if it is called with a non zero index on an empty weight map.
+		*/
+		FReal operator[](int32 Index) const { return Table[Indices[Index]]; }
+
+		/** Return the exponential value at the Low weight. */
+		FReal GetLow() const { return Table[0]; }
+
+		/** Return the exponential value at the High weight. */
+		FReal GetHigh() const { return Table.Last(); }
+
+		/** Return the exponential stiffness value when the weight map is not used. */
+		explicit operator FReal() const { return GetLow(); }
+
+	private:
+		static constexpr FReal ParameterFrequency = (FReal)120.;  // 60Hz @ 2 iterations as a root for all stiffness values TODO: Make this a global solver parameter
+
+		TArray<uint8> Indices; // Per particle array of index to the stiffness table
+		TArray<FReal> Table;  // Fixed lookup table of stiffness values, use uint8 indexation
+		FVec2 WeightedValue;
+		const FReal ParameterFitLogBase;
+	};
+
+	FPBDStiffness::FPBDStiffness(
+		const TConstArrayView<FReal>& Multipliers,
+		const FVec2& InWeightedValue,
+		int32 ParticleCount,
+		int32 TableSize,
+		FReal InParameterFitBase)
+		: WeightedValue(InWeightedValue)
+		, ParameterFitLogBase(FMath::Loge(InParameterFitBase))
+	{
+		check(TableSize > 0 && TableSize < 256);  // The Stiffness lookup table is restricted to uint8 sized indices
+
+		if (Multipliers.Num() == ParticleCount)
+		{
+			// Convert the weight maps into an array of lookup indices to the stiffness table
+			Indices.AddUninitialized(ParticleCount);
+
+			const FReal TableScale = (FReal)(TableSize - 1);
+
+			for (int32 Index = 0; Index < ParticleCount; ++Index)
+			{
+				Indices[Index] = (uint8)(FMath::Clamp(Multipliers[Index], (FReal)0., (FReal)1.) * TableScale);
+			}
+
+			// Initialize empty table until ApplyValues is called
+			Table.AddZeroed(TableSize);
+		}
+		else
+		{
+			// Initialize with a one element table until ApplyValues is called
+			Indices.AddZeroed(1);
+			Table.AddZeroed(1);
+		}
+	}
+
+	void FPBDStiffness::ApplyValues(const FReal Dt, const int32 NumIterations)
+	{
+		SCOPE_CYCLE_COUNTER(STAT_PBD_StiffnessApplyValues);
+
+		// Calculate the simulation exponent
+		const FReal Exponent = Dt * ParameterFrequency / (FReal)NumIterations;
+
+		// Define the stiffness mapping function
+		auto SimulationValue = [this, Exponent](const FReal InValue)->FReal
+		{
+			// Get a very steep exponential curve between the [0, 1] range to make easier to set the parameter
+			// The base has been chosen empirically
+			// ParameterValue = Pow(ParameterFitBase, ParameterValue - 1)
+			const FReal ParameterFit = FMath::Exp(ParameterFitLogBase * (FMath::Clamp(InValue, (FReal)0., (FReal)1.) - (FReal)1.));
+
+			// Use simulation dependent stiffness exponent to alleviate the variations in effect when Dt and NumIterations change
+			// This is based on the Position-Based Simulation Methods paper (page 8),
+			// but uses the delta time in addition of the number of iterations in the calculation of the error term.
+			const FReal LogValue = FMath::Loge((FReal)1. - ParameterFit);
+			return (FReal)1. - FMath::Exp(LogValue * Exponent);
+		};
+
+		const FReal Offset = WeightedValue[0];
+		const FReal Range = WeightedValue[1] - WeightedValue[0];
+		const int32 TableSize = Table.Num();
+		const FReal WeightIncrement = (TableSize > 1) ? (FReal)1. / (FReal)(TableSize - 1) : (FReal)1.; // Must allow full range from 0 to 1 included
+		for (int32 Index = 0; Index < TableSize; ++Index)
+		{
+			const FReal Weight = (FReal)Index * WeightIncrement;
+			Table[Index] = SimulationValue(Offset + Weight * Range);
+		}
+	}
+
+}
diff --git a/Engine/Source/Runtime/Experimental/Chaos/Public/Chaos/TriangleMesh.h b/Engine/Source/Runtime/Experimental/Chaos/Public/Chaos/TriangleMesh.h
--- a/Engine/Source/Runtime/Experimental/Chaos/Public/Chaos/TriangleMesh.h
+++ b/Engine/Source/Runtime/Experimental/Chaos/Public/Chaos/TriangleMesh.h
@@ -86,8 +86,8 @@ namespace Chaos
 
 		FORCEINLINE int32 GetNumElements() const { return MElements.Num(); }
 
-		CHAOS_API const TMap<int32, TSet<uint32>>& GetPointToNeighborsMap() const;
-		FORCEINLINE const TSet<uint32>& GetNeighbors(const int32 Element) const { return GetPointToNeighborsMap()[Element]; }
+		CHAOS_API const TMap<int32, TSet<int32>>& GetPointToNeighborsMap() const;
+		FORCEINLINE const TSet<int32>& GetNeighbors(const int32 Element) const { return GetPointToNeighborsMap()[Element]; }
 
 		CHAOS_API TConstArrayView<TArray<int32>> GetPointToTriangleMap() const;  // Return an array view using global indexation. Only elements starting at MStartIdx will be valid!
 		FORCEINLINE const TArray<int32>& GetCoincidentTriangles(const int32 Element) const { return GetPointToTriangleMap()[Element]; }
@@ -95,7 +95,7 @@ namespace Chaos
 		FORCEINLINE TSet<int32> GetNRing(const int32 Element, const int32 N) const
 		{
 			TSet<int32> Neighbors;
-			TSet<uint32> LevelNeighbors, PrevLevelNeighbors;
+			TSet<int32> LevelNeighbors, PrevLevelNeighbors;
 			PrevLevelNeighbors = GetNeighbors(Element);
 			for (auto SubElement : PrevLevelNeighbors)
 			{
@@ -278,7 +278,7 @@ namespace Chaos
 		TArray<TVec3<int32>> MElements;
 
 		mutable TArray<TArray<int32>> MPointToTriangleMap;  // !! Unlike the TArrayView returned by GetPointToTriangleMap, this array starts at 0 for the point of index MStartIdx. Use GlobalToLocal to access with a global index. Note that this array's content is always indexed in global index.
-		mutable TMap<int32, TSet<uint32>> MPointToNeighborsMap;
+		mutable TMap<int32, TSet<int32>> MPointToNeighborsMap;
 
 		FSegmentMesh MSegmentMesh;
 		TArray<TVec3<int32>> MFaceToEdges;
diff --git a/Engine/Source/Runtime/Experimental/Chaos/Public/Chaos/XPBDLongRangeConstraints.h b/Engine/Source/Runtime/Experimental/Chaos/Public/Chaos/XPBDLongRangeConstraints.h
--- a/Engine/Source/Runtime/Experimental/Chaos/Public/Chaos/XPBDLongRangeConstraints.h
+++ b/Engine/Source/Runtime/Experimental/Chaos/Public/Chaos/XPBDLongRangeConstraints.h
@@ -3,7 +3,6 @@
 
 #include "Chaos/PBDLongRangeConstraintsBase.h"
 #include "Chaos/PBDParticles.h"
-#include "Chaos/PBDConstraintContainer.h"
 #include "ChaosStats.h"
 
 DECLARE_CYCLE_STAT(TEXT("Chaos XPBD Long Range Constraint"), STAT_XPBD_LongRange, STATGROUP_Chaos);
@@ -13,75 +12,103 @@ namespace Chaos
 // Stiffness is in N/CM^2, so it needs to be adjusted from the PBD stiffness ranging between [0,1]
 static const double XPBDLongRangeMaxCompliance = 1.e-3;
 
-class FXPBDLongRangeConstraints : public FPBDLongRangeConstraintsBase, public FPBDConstraintContainer
+class FXPBDLongRangeConstraints final : public FPBDLongRangeConstraintsBase
 {
-	typedef FPBDLongRangeConstraintsBase Base;
-	using Base::MEuclideanConstraints;
-	using Base::MDists;
-	using Base::MStiffness;
-
 public:
-	FXPBDLongRangeConstraints(const FDynamicParticles& InParticles, const TMap<int32, TSet<uint32>>& PointToNeighbors, const int32 NumberOfAttachments = 1, const FReal Stiffness = (FReal)1.)
-	    : FPBDLongRangeConstraintsBase(InParticles, PointToNeighbors, NumberOfAttachments, Stiffness)
-	{ MLambdas.Init(0.f, MEuclideanConstraints.Num()); }
+	typedef FPBDLongRangeConstraintsBase Base;
+	typedef typename Base::FTether FTether;
+	typedef typename Base::EMode EMode;
+
+	FXPBDLongRangeConstraints(
+		const FPBDParticles& Particles,
+		const int32 InParticleOffset,
+		const int32 InParticleCount,
+		const TMap<int32, TSet<int32>>& PointToNeighbors,
+		const TConstArrayView<FReal>& StiffnessMultipliers,
+		const int32 MaxNumTetherIslands = 4,
+		const FVec2& InStiffness = FVec2((FReal)1., (FReal)1.),
+		const FReal LimitScale = (FReal)1.,
+		const EMode InMode = EMode::Geodesic)
+	    : FPBDLongRangeConstraintsBase(Particles, InParticleOffset, InParticleCount, PointToNeighbors, StiffnessMultipliers, MaxNumTetherIslands, InStiffness, LimitScale, InMode)
+	{
+		Lambdas.Reserve(Tethers.Num());
+	}
 
 	virtual ~FXPBDLongRangeConstraints() {}
 
-	void Init() const { for (FReal& Lambda : MLambdas) { Lambda = (FReal)0.; } }
-
-	void Apply(FPBDParticles& InParticles, const FReal Dt, int32 Index) const
+	void Init() const
 	{
-		const auto& Constraint = MEuclideanConstraints[Index];
-		int32 i2 = Constraint[Constraint.Num() - 1];
-		check(InParticles.InvM(i2) > 0);
-		InParticles.P(i2) += GetDelta(InParticles, Dt, Index);
+		Lambdas.Reset();
+		Lambdas.AddZeroed(Tethers.Num());
 	}
 
-	void Apply(FPBDParticles& InParticles, const FReal Dt) const
+	void Apply(FPBDParticles& Particles, const FReal Dt) const 
 	{
 		SCOPE_CYCLE_COUNTER(STAT_XPBD_LongRange);
-		for (int32 i = 0; i < MEuclideanConstraints.Num(); ++i)
+		// Run particles in parallel, and ranges in sequence to avoid a race condition when updating the same particle from different tethers
+		static const int32 MinParallelSize = 500;
+		if (Stiffness.HasWeightMap())
 		{
-			Apply(InParticles, Dt, i);
+			TethersView.ParallelFor([this, &Particles, Dt](TArray<FTether>& /*InTethers*/, int32 Index)
+				{
+					const FReal ExpStiffnessValue = Stiffness[Index - ParticleOffset];
+					Apply(Particles, Dt, Index, ExpStiffnessValue);
+				}, MinParallelSize);
+		}
+		else
+		{
+			const FReal ExpStiffnessValue = (FReal)Stiffness;
+			TethersView.ParallelFor([this, &Particles, Dt, ExpStiffnessValue](TArray<FTether>& /*InTethers*/, int32 Index)
+				{
+					Apply(Particles, Dt, Index, ExpStiffnessValue);
+				}, MinParallelSize);
 		}
 	}
 
-	void Apply(FPBDParticles& InParticles, const FReal Dt, const TArray<int32>& InConstraintIndices) const
+	void Apply(FPBDParticles& Particles, const FReal Dt, const TArray<int32>& ConstraintIndices) const
 	{
 		SCOPE_CYCLE_COUNTER(STAT_XPBD_LongRange);
-		for (int32 i : InConstraintIndices)
+		if (Stiffness.HasWeightMap())
+		{
+			for (const int32 Index : ConstraintIndices)
+			{
+				const FReal ExpStiffnessValue = Stiffness[Index - ParticleOffset];
+				Apply(Particles, Dt, Index, ExpStiffnessValue);
+			}
+		}
+		else
 		{
-			Apply(InParticles, Dt, i);
+			const FReal ExpStiffnessValue = (FReal)Stiffness;
+			for (const int32 Index : ConstraintIndices)
+			{
+				Apply(Particles, Dt, Index, ExpStiffnessValue);
+			}
 		}
 	}
 
 private:
-	inline FVec3 GetDelta(const FPBDParticles& InParticles, const FReal Dt, const int32 InConstraintIndices) const
+	void Apply(FPBDParticles& Particles, const FReal Dt, int32 Index, const FReal InStiffness) const
 	{
-		const TVector<uint32, 2>& Constraint = MEuclideanConstraints[InConstraintIndices];
-		check(Constraint.Num() > 1);
-		const uint32 i1 = Constraint[0];
-		const uint32 i2 = Constraint[1];
-		check(InParticles.InvM(i1) == 0);
-		check(InParticles.InvM(i2) > 0);
-		const FReal Distance = Base::ComputeGeodesicDistance(InParticles, Constraint); // This function is used for either Euclidean or Geodesic distances
-		if (Distance < MDists[InConstraintIndices]) { return FVec3((FReal)0.); }
-		const FReal Offset = Distance - MDists[InConstraintIndices];
+		const FTether& Tether = Tethers[Index];
 
-		FVec3 Direction = InParticles.P(i1) - InParticles.P(i2);
-		Direction.SafeNormalize();
+		FVec3 Direction;
+		FReal Offset;
+		Tether.GetDelta(Particles, Direction, Offset);
 
-		FReal& Lambda = MLambdas[InConstraintIndices];
-		const FReal Alpha = (FReal)XPBDLongRangeMaxCompliance / (MStiffness * Dt * Dt);
+		FReal& Lambda = Lambdas[Index];
+		const FReal Alpha = (FReal)XPBDLongRangeMaxCompliance / (InStiffness * Dt * Dt);
 
 		const FReal DLambda = (Offset - Alpha * Lambda) / ((FReal)1. + Alpha);
-		const FVec3 Delta = DLambda * Direction;
+		Particles.P(Tether.End) += DLambda * Direction;
 		Lambda += DLambda;
-
-		return Delta;
-	};
+	}
 
 private:
-	mutable TArray<FReal> MLambdas;
+	using Base::Tethers;
+	using Base::TethersView;
+	using Base::Stiffness;
+	using Base::ParticleOffset;
+
+	mutable TArray<FReal> Lambdas;
 };
 }
