MeshModelingTools: Allow programmatic changes to properties to update detail panel properly.

#rnx

diff --git a/Engine/Plugins/Experimental/MeshModelingToolset/Source/MeshModelingTools/Private/DrawAndRevolveTool.cpp b/Engine/Plugins/Experimental/MeshModelingToolset/Source/MeshModelingTools/Private/DrawAndRevolveTool.cpp
--- a/Engine/Plugins/Experimental/MeshModelingToolset/Source/MeshModelingTools/Private/DrawAndRevolveTool.cpp
+++ b/Engine/Plugins/Experimental/MeshModelingToolset/Source/MeshModelingTools/Private/DrawAndRevolveTool.cpp
@@ -133,7 +133,12 @@ void UDrawAndRevolveTool::Setup()
 		{
 			Preview->InvalidateResult();
 		}
-		Settings->bAllowedToEditDrawPlane = (ControlPointsMechanic->GetNumPoints() == 0);
+		bool bAllowedToEditDrawPlane = (ControlPointsMechanic->GetNumPoints() == 0);
+		if (Settings->bAllowedToEditDrawPlane != bAllowedToEditDrawPlane)
+		{
+			Settings->bAllowedToEditDrawPlane = bAllowedToEditDrawPlane;
+			NotifyOfPropertyChangeByTool(Settings);
+		}
 		});
 	// This gets called when we enter/leave curve initialization mode
 	ControlPointsMechanic->OnModeChanged.AddLambda([this]() {
@@ -171,6 +176,7 @@ void UDrawAndRevolveTool::Setup()
 	PlaneMechanic->OnPlaneChanged.AddLambda([this]() {
 		Settings->DrawPlaneOrigin = (FVector)PlaneMechanic->Plane.Origin;
 		Settings->DrawPlaneOrientation = ((FQuat)PlaneMechanic->Plane.Rotation).Rotator();
+		NotifyOfPropertyChangeByTool(Settings);
 		if (ControlPointsMechanic)
 		{
 			ControlPointsMechanic->SetPlane(PlaneMechanic->Plane);
diff --git a/Engine/Plugins/Experimental/ModelingToolsEditorMode/Source/ModelingToolsEditorMode/Private/ModelingToolsEditorModeToolkit.cpp b/Engine/Plugins/Experimental/ModelingToolsEditorMode/Source/ModelingToolsEditorMode/Private/ModelingToolsEditorModeToolkit.cpp
--- a/Engine/Plugins/Experimental/ModelingToolsEditorMode/Source/ModelingToolsEditorMode/Private/ModelingToolsEditorModeToolkit.cpp
+++ b/Engine/Plugins/Experimental/ModelingToolsEditorMode/Source/ModelingToolsEditorMode/Private/ModelingToolsEditorModeToolkit.cpp
@@ -370,6 +370,11 @@ void FModelingToolsEditorModeToolkit::UpdateActiveToolProperties()
 	}
 }
 
+void FModelingToolsEditorModeToolkit::RefreshPropertyObject(UObject* ChangedObject)
+{
+	ModeDetailsView->ForceRefresh();
+}
+
 
 void FModelingToolsEditorModeToolkit::PostNotification(const FText& Message)
 {
@@ -695,6 +700,7 @@ void FModelingToolsEditorModeToolkit::OnToolStarted(UInteractiveToolManager* Man
 
 	UInteractiveTool* CurTool = GetScriptableEditorMode()->GetToolManager()->GetActiveTool(EToolSide::Left);
 	CurTool->OnPropertySetsModified.AddSP(this, &FModelingToolsEditorModeToolkit::UpdateActiveToolProperties);
+	CurTool->OnPropertyModifiedDirectlyByTool.AddSP(this, &FModelingToolsEditorModeToolkit::RefreshPropertyObject);
 
 	ModeHeaderArea->SetVisibility(EVisibility::Collapsed);
 	ActiveToolName = CurTool->GetToolInfo().ToolDisplayName;
@@ -730,6 +736,7 @@ void FModelingToolsEditorModeToolkit::OnToolEnded(UInteractiveToolManager* Manag
 	if ( CurTool )
 	{
 		CurTool->OnPropertySetsModified.RemoveAll(this);
+		CurTool->OnPropertyModifiedDirectlyByTool.RemoveAll(this);
 	}
 
 	// re-enable LOD level picker
diff --git a/Engine/Plugins/Experimental/ModelingToolsEditorMode/Source/ModelingToolsEditorMode/Public/ModelingToolsEditorModeToolkit.h b/Engine/Plugins/Experimental/ModelingToolsEditorMode/Source/ModelingToolsEditorMode/Public/ModelingToolsEditorModeToolkit.h
--- a/Engine/Plugins/Experimental/ModelingToolsEditorMode/Source/ModelingToolsEditorMode/Public/ModelingToolsEditorModeToolkit.h
+++ b/Engine/Plugins/Experimental/ModelingToolsEditorMode/Source/ModelingToolsEditorMode/Public/ModelingToolsEditorModeToolkit.h
@@ -74,6 +74,7 @@ private:
 
 	TSharedPtr<SWidget> ToolkitWidget;
 	void UpdateActiveToolProperties();
+	void RefreshPropertyObject(UObject* ChangedObject);
 
 	TSharedPtr<SWidget> ViewportOverlayWidget;
 
diff --git a/Engine/Source/Runtime/Experimental/InteractiveToolsFramework/Private/InteractiveTool.cpp b/Engine/Source/Runtime/Experimental/InteractiveToolsFramework/Private/InteractiveTool.cpp
--- a/Engine/Source/Runtime/Experimental/InteractiveToolsFramework/Private/InteractiveTool.cpp
+++ b/Engine/Source/Runtime/Experimental/InteractiveToolsFramework/Private/InteractiveTool.cpp
@@ -143,6 +143,11 @@ bool UInteractiveTool::SetToolPropertySourceEnabled(UInteractiveToolPropertySet*
 	return true;
 }
 
+void UInteractiveTool::NotifyOfPropertyChangeByTool(UInteractiveToolPropertySet* PropertySet) const
+{
+	OnPropertyModifiedDirectlyByTool.Broadcast(PropertySet);
+}
+
 TArray<UObject*> UInteractiveTool::GetToolProperties(bool bEnabledOnly) const
 {
 	if (bEnabledOnly == false)
diff --git a/Engine/Source/Runtime/Experimental/InteractiveToolsFramework/Public/InteractiveTool.h b/Engine/Source/Runtime/Experimental/InteractiveToolsFramework/Public/InteractiveTool.h
--- a/Engine/Source/Runtime/Experimental/InteractiveToolsFramework/Public/InteractiveTool.h
+++ b/Engine/Source/Runtime/Experimental/InteractiveToolsFramework/Public/InteractiveTool.h
@@ -354,6 +354,18 @@ public:
 	DECLARE_MULTICAST_DELEGATE(OnInteractiveToolPropertySetsModified);
 	OnInteractiveToolPropertySetsModified OnPropertySetsModified;
 
+	/**
+	 * OnPropertyModifiedDirectlyByTool is broadcast whenever the ToolPropertyObjects array stays the same, but a property
+	 * inside of one of the objects is changed internally by the tool. This allows any external display of such properties
+	 * to properly update. In a DetailsViewPanel, for instance, it refreshes certain cached states such as edit condition
+	 * states for other properties.
+	 * 
+	 * This should only broadcast when the tool itself is responsible for the change, so it typically isn't broadcast
+	 * from the tool's OnPropertyModified function. 
+	 */
+	DECLARE_MULTICAST_DELEGATE_OneParam(OnInteractiveToolPropertyInternallyModified, UObject*);
+	OnInteractiveToolPropertyInternallyModified OnPropertyModifiedDirectlyByTool;
+
 	/**
 	 * Automatically called by UInteractiveToolPropertySet.OnModified delegate to notify Tool of child property set changes
 	 * @param PropertySet which UInteractiveToolPropertySet was modified
@@ -410,6 +422,13 @@ protected:
 	 */
 	virtual bool SetToolPropertySourceEnabled(UInteractiveToolPropertySet* PropertySet, bool bEnabled);
 
+	/**
+	 * Call after changing a propertyset internally in the tool to allow external views of the property
+	 * set to update properly. This is meant as an outward notification mechanism, not a way to to
+	 * pass along notifications, so don't call this if the property is changed externally (i.e., this
+	 * should not usually be called from OnPropertyModified unless the tool adds changes of its own).
+	 */
+	virtual void NotifyOfPropertyChangeByTool(UInteractiveToolPropertySet* PropertySet) const;
 
 	enum EAcceptWarning
 	{
