Fix field update kinematic state crash

diff --git a/Engine/Source/Runtime/Experimental/Chaos/Public/PhysicsProxy/FieldSystemProxyHelper.h b/Engine/Source/Runtime/Experimental/Chaos/Public/PhysicsProxy/FieldSystemProxyHelper.h
--- a/Engine/Source/Runtime/Experimental/Chaos/Public/PhysicsProxy/FieldSystemProxyHelper.h
+++ b/Engine/Source/Runtime/Experimental/Chaos/Public/PhysicsProxy/FieldSystemProxyHelper.h
@@ -182,19 +182,17 @@ namespace Chaos
 		{
 			RigidSolver->GetParticles().UpdateGeometryCollectionViews(true);
 
-			const Chaos::FPBDRigidsSOAs& SolverParticles = RigidSolver->GetParticles();
+			Chaos::FPBDRigidsSOAs& SolverParticles = RigidSolver->GetParticles();
 			auto& Clustering = RigidSolver->GetEvolution()->GetRigidClustering();
 
-			const Chaos::TParticleView<Chaos::FGeometryParticles>& ParticleView =
-				SolverParticles.GetNonDisabledView();
+			FPBDRigidClusteredParticles& ClusteredParticles = SolverParticles.GetClusteredParticles();
 
-			for (Chaos::TParticleIterator<Chaos::FGeometryParticles> It = ParticleView.Begin(), ItEnd = ParticleView.End();
-				It != ItEnd; ++It)
+			for (uint32 ParticleIndex = 0, NumParticles = ClusteredParticles.Size(); ParticleIndex < NumParticles; ++ParticleIndex)
 			{
-				const auto* Clustered = It->Handle()->CastToClustered();
-				if (Clustered && Clustered->ClusterIds().NumChildren)
+				FPBDRigidClusteredParticleHandle* ClusteredHandle = ClusteredParticles.Handle(ParticleIndex);
+				if (ClusteredHandle && ClusteredHandle->ClusterIds().NumChildren && !ClusteredHandle->Disabled())
 				{
-					Chaos::FPBDRigidParticleHandle* RigidHandle = It->Handle()->CastToRigidParticle();
+					Chaos::FPBDRigidParticleHandle* RigidHandle = ClusteredHandle->CastToRigidParticle();
 					check(RigidHandle);
 					Clustering.UpdateKinematicProperties(RigidHandle);
 				}
