Chaos - disable manifolds for levelset collision.

Manifolds don't currently work well when there are very fast rotating overlapping pieces and this is very common for destructible objects.

diff --git a/Engine/Source/Runtime/Experimental/Chaos/Private/Chaos/CollisionResolution.cpp b/Engine/Source/Runtime/Experimental/Chaos/Private/Chaos/CollisionResolution.cpp
--- a/Engine/Source/Runtime/Experimental/Chaos/Private/Chaos/CollisionResolution.cpp
+++ b/Engine/Source/Runtime/Experimental/Chaos/Private/Chaos/CollisionResolution.cpp
@@ -91,6 +91,9 @@ FAutoConsoleVariableRef CVarConstraintsDetailedStats(TEXT("p.Chaos.Constraints.D
 int32 AlwaysAddSweptConstraints = 0;
 FAutoConsoleVariableRef CVarAlwaysAddSweptConstraints(TEXT("p.Chaos.Constraints.AlwaysAddSweptConstraints"), AlwaysAddSweptConstraints, TEXT("Since GJKContactPointSwept returns infinity for it's contact data when not hitting anything, some contacts are discarded prematurely. This flag will cause contact points considered for sweeps to never be discarded."));
 
+bool bChaos_Collision_AllowLevelsetManifolds = false;
+FAutoConsoleVariableRef CVarChaosCollisionAllowLevelsetManifolds(TEXT("p.Chaos.Collision.AllowLevelsetManifolds"), bChaos_Collision_AllowLevelsetManifolds, TEXT("Use incremental manifolds for levelset-levelset collision. This does not work well atm - too much rotation in the small pieces"));
+
 //Chaos::FRealSingle Chaos_Collision_ManifoldFaceAngle = 5.0f;
 //Chaos::FRealSingle Chaos_Collision_ManifoldFaceEpsilon = FMath::Sin(FMath::DegreesToRadians(Chaos_Collision_ManifoldFaceAngle));
 //FConsoleVariableDelegate Chaos_Collision_ManifoldFaceDelegate = FConsoleVariableDelegate::CreateLambda([](IConsoleVariable* CVar) { Chaos_Collision_ManifoldFaceEpsilon = FMath::Sin(FMath::DegreesToRadians(Chaos_Collision_ManifoldFaceAngle)); });
@@ -2073,7 +2076,8 @@ namespace Chaos
 		template<typename T_TRAITS>
 		void ConstructLevelsetLevelsetConstraints(TGeometryParticleHandle<FReal, 3>* Particle0, TGeometryParticleHandle<FReal, 3>* Particle1, const FImplicitObject* Implicit0, const FBVHParticles* Simplicial0, const FImplicitObject* Implicit1, const FBVHParticles* Simplicial1, const FRigidTransform3& LocalTransform0, const FRigidTransform3& LocalTransform1, const FReal CullDistance, const FReal Dt, const FCollisionContext& Context, FCollisionConstraintsArray& NewConstraints)
 		{
-			FRigidBodyPointContactConstraint Constraint = FRigidBodyPointContactConstraint(Particle0, Implicit0, Simplicial0, LocalTransform0, Particle1, Implicit1, Simplicial1, LocalTransform1, CullDistance, EContactShapesType::LevelSetLevelSet, Context.bAllowManifolds);
+			const bool bAllowManifolds = Context.bAllowManifolds && bChaos_Collision_AllowLevelsetManifolds;
+			FRigidBodyPointContactConstraint Constraint = FRigidBodyPointContactConstraint(Particle0, Implicit0, Simplicial0, LocalTransform0, Particle1, Implicit1, Simplicial1, LocalTransform1, CullDistance, EContactShapesType::LevelSetLevelSet, bAllowManifolds);
 
 			bool bIsParticleDynamic0 = Particle0->CastToRigidParticle() && Particle0->ObjectState() == EObjectStateType::Dynamic;
 			int32 P0NumCollisionParticles = Simplicial0 ? Simplicial0->Size() : Particle0->CastToRigidParticle()->CollisionParticlesSize();
