Chaos
- debug draw colliding shapes
- debug draw exploded comvex

diff --git a/Engine/Source/Runtime/Experimental/Chaos/Private/Chaos/ChaosDebugDraw.cpp b/Engine/Source/Runtime/Experimental/Chaos/Private/Chaos/ChaosDebugDraw.cpp
--- a/Engine/Source/Runtime/Experimental/Chaos/Private/Chaos/ChaosDebugDraw.cpp
+++ b/Engine/Source/Runtime/Experimental/Chaos/Private/Chaos/ChaosDebugDraw.cpp
@@ -39,11 +39,18 @@ namespace Chaos
 		bool bChaosDebugDebugDrawInactiveContacts = true;
 		FAutoConsoleVariableRef CVarChaosDebugDrawInactiveContacts(TEXT("p.Chaos.DebugDraw.ShowInactiveContacts"), bChaosDebugDebugDrawInactiveContacts, TEXT("Whether to show inactive contacts (ones that contributed no impulses or pushout)"));
 
+		bool bChaosDebugDebugDrawConvexVertices = false;
 		bool bChaosDebugDebugDrawCoreShapes = false;
 		bool bChaosDebugDebugDrawExactCoreShapes = false;
+		FAutoConsoleVariableRef CVarChaosDebugDrawConvexVertices(TEXT("p.Chaos.DebugDraw.ShowConvexVertices"), bChaosDebugDebugDrawConvexVertices, TEXT("Whether to show the vertices of convex shapes"));
 		FAutoConsoleVariableRef CVarChaosDebugDrawCoreShapes(TEXT("p.Chaos.DebugDraw.ShowCoreShapes"), bChaosDebugDebugDrawCoreShapes, TEXT("Whether to show the core (margin-reduced) shape where applicable"));
 		FAutoConsoleVariableRef CVarChaosDebugDrawExactShapes(TEXT("p.Chaos.DebugDraw.ShowExactCoreShapes"), bChaosDebugDebugDrawExactCoreShapes, TEXT("Whether to show the exact core shape. NOTE: Extremely expensive and should only be used on a small scene with a couple convex shapes in it"));
 
+		bool bChaosDebugDebugDrawContactGraph = false;
+		FAutoConsoleVariableRef CVarChaosDebugDrawContactGraph(TEXT("p.Chaos.DebugDraw.ShowContactGraph"), bChaosDebugDebugDrawContactGraph, TEXT("Whether to show the contactgraph when drawing islands"));
+
+		float ChaosDebugDrawConvexExplodeDistance = 0.0f;
+		FAutoConsoleVariableRef CVarChaosDebugDrawConvexExplodeDistance(TEXT("p.Chaos.DebugDraw.ConvexExplodeDistance"), ChaosDebugDrawConvexExplodeDistance, TEXT("Explode convex edges by this amount (useful for looking at convex integrity)"));
 
 		const FChaosDebugDrawSettings& GetChaosDebugDrawSettings(const FChaosDebugDrawSettings* InSettings)
 		{
@@ -111,8 +118,11 @@ namespace Chaos
 						const FVec3 OuterP0 = ShapeTransform.TransformPosition(Shape->GetVertex(VertexIndex0));
 						const FVec3 OuterP1 = ShapeTransform.TransformPosition(Shape->GetVertex(VertexIndex1));
 
+						const FVec3 N0 = ShapeTransform.TransformVectorNoScale(Shape->GetPlane(PlaneIndex).Normal());
+						const FVec3 ExplodeDelta = ChaosDebugDrawConvexExplodeDistance * N0;
+
 						// Outer shape
-						FDebugDrawQueue::GetInstance().DrawDebugLine(OuterP0, OuterP1, Color, false, -1.f, Settings.DrawPriority, Settings.ShapeThicknesScale * Settings.LineThickness);
+						FDebugDrawQueue::GetInstance().DrawDebugLine(OuterP0 + ExplodeDelta, OuterP1 + ExplodeDelta, Color, false, -1.f, Settings.DrawPriority, Settings.ShapeThicknesScale * Settings.LineThickness);
 
 						// Core shape and lines connecting core to outer
 						if (Margin > 0.0f)
@@ -123,6 +133,12 @@ namespace Chaos
 							FDebugDrawQueue::GetInstance().DrawDebugLine(InnerP0, InnerP1, FColor::Blue, false, -1.f, Settings.DrawPriority, LineThickness);
 							FDebugDrawQueue::GetInstance().DrawDebugLine(InnerP0, OuterP0, FColor::Black, false, -1.f, Settings.DrawPriority, LineThickness);
 						}
+
+						// Vertex and face normal
+						if (bChaosDebugDebugDrawConvexVertices)
+						{
+							FDebugDrawQueue::GetInstance().DrawDebugLine(OuterP0 + ExplodeDelta, OuterP0 + ExplodeDelta + Settings.DrawScale * 20.0f * N0, FColor::Green, false, -1.f, Settings.DrawPriority, Settings.LineThickness);
+						}
 					}
 				}
 			}
@@ -726,6 +742,57 @@ namespace Chaos
 			DrawCollisionImpl(SpaceTransform, ConstraintHandle->GetContact(), ColorScale, Settings);
 		}
 
+		void DrawCollidingShapesImpl(const FRigidTransform3& SpaceTransform, const FPBDCollisionConstraints& Collisions, float ColorScale, const FChaosDebugDrawSettings& Settings)
+		{
+			TArray<const FImplicitObject*> Shapes;
+			TArray<FConstGenericParticleHandle> ShapeParticles;
+			TArray<FRigidTransform3> ShapeTransforms;
+
+			for (int32 ConstraintIndex = 0; ConstraintIndex < Collisions.NumConstraints(); ++ConstraintIndex)
+			{
+				const FRigidBodyPointContactConstraint* PointConstraint = Collisions.GetConstraint(ConstraintIndex).template As<FRigidBodyPointContactConstraint>();
+				if (!PointConstraint->GetDisabled() && (PointConstraint->Manifold.Phi < TNumericLimits<FReal>::Max()))
+				{
+					const FImplicitObject* Implicit0 = PointConstraint->Manifold.Implicit[0];
+					const FImplicitObject* Implicit1 = PointConstraint->Manifold.Implicit[1];
+					if ((Implicit0 != nullptr) && (Implicit1 != nullptr))
+					{
+						FConstGenericParticleHandle Particle0 = PointConstraint->Particle[0];
+						FConstGenericParticleHandle Particle1 = PointConstraint->Particle[1];
+						const FRigidTransform3 WorldActorTransform0 = FParticleUtilities::GetActorWorldTransform(Particle0);
+						const FRigidTransform3 WorldActorTransform1 = FParticleUtilities::GetActorWorldTransform(Particle1);
+						const FRigidTransform3 ShapeWorldTransform0 = PointConstraint->ImplicitTransform[0] * WorldActorTransform0;
+						const FRigidTransform3 ShapeWorldTransform1 = PointConstraint->ImplicitTransform[1] * WorldActorTransform1;
+
+						if (!Shapes.Contains(Implicit0))
+						{
+							Shapes.Add(Implicit0);
+							ShapeParticles.Add(Particle0);
+							ShapeTransforms.Add(ShapeWorldTransform0);
+						}
+
+						if (!Shapes.Contains(Implicit1))
+						{
+							Shapes.Add(Implicit1);
+							ShapeParticles.Add(Particle1);
+							ShapeTransforms.Add(ShapeWorldTransform1);
+						}
+					}
+				}
+			}
+
+			for (int32 ShapeIndex = 0; ShapeIndex < Shapes.Num(); ++ShapeIndex)
+			{
+				DrawShapesImpl(
+					ShapeParticles[ShapeIndex]->Handle(), 
+					ShapeTransforms[ShapeIndex], 
+					Shapes[ShapeIndex], 
+					0.0f, 
+					ShapeParticles[ShapeIndex]->IsDynamic() ? FColor::Yellow : FColor::Red, 
+					Settings);
+			}
+		}
+
 		void DrawJointConstraintImpl(const FRigidTransform3& SpaceTransform, const FVec3& InPa, const FVec3& InCa, const FVec3& InXa, const FMatrix33& Ra, const FVec3& InPb, const FVec3& InCb, const FVec3& InXb, const FMatrix33& Rb, int32 IslandIndex, int32 LevelIndex, int32 ColorIndex, int32 BatchIndex, int32 Index, FReal ColorScale, uint32 FeatureMask, const FChaosDebugDrawSettings& Settings)
 		{
 			using namespace Chaos::DebugDraw;
@@ -876,16 +943,33 @@ namespace Chaos
 			};
 			const int32 NumIslandColors = UE_ARRAY_COUNT(IslandColors);
 
-			auto DrawGraphCollision = [&](const FRigidTransform3& SpaceTransform, const FPBDCollisionConstraintHandle& CollisionConstraint,  int32 IslandIndex, int32 LevelIndex, int32 ColorIndex, const FChaosDebugDrawSettings& Settings)
+			auto DrawGraphCollision = [&](const FRigidTransform3& SpaceTransform, const FPBDCollisionConstraintHandle& CollisionConstraint,  int32 IslandIndex, int32 LevelIndex, int32 ColorIndex, int32 OrderIndex, const FChaosDebugDrawSettings& Settings)
 			{
 				FColor IslandColor = IslandColors[IslandIndex % NumIslandColors];
 
 				const FRigidTransform3 Transform0 = FParticleUtilities::GetCoMWorldTransform(FConstGenericParticleHandle(CollisionConstraint.GetConstrainedParticles()[0])) * SpaceTransform;
 				const FRigidTransform3 Transform1 = FParticleUtilities::GetCoMWorldTransform(FConstGenericParticleHandle(CollisionConstraint.GetConstrainedParticles()[1])) * SpaceTransform;
-				const FVec3 TextPos = 0.5f * (Transform0.GetLocation() + Transform1.GetLocation());
-				FDebugDrawQueue::GetInstance().DrawDebugLine(Transform0.GetLocation(), TextPos, FColor::Black, false, KINDA_SMALL_NUMBER, Settings.DrawPriority, Settings.LineThickness);
-				FDebugDrawQueue::GetInstance().DrawDebugLine(Transform1.GetLocation(), TextPos, FColor::White, false, KINDA_SMALL_NUMBER, Settings.DrawPriority, Settings.LineThickness);
-				FDebugDrawQueue::GetInstance().DrawDebugString(TextPos, FString::Format(TEXT("{0}-{1}-{2}"), { IslandIndex, LevelIndex, ColorIndex }), nullptr, FColor::Yellow, KINDA_SMALL_NUMBER, false, Settings.FontScale);
+
+				FVec3 ContactPos = CollisionConstraint.GetContactLocation();
+				if (CollisionConstraint.GetContact().GetType() == FCollisionConstraintBase::FType::SinglePoint)
+				{
+					const FRigidBodyPointContactConstraint* Constraint = CollisionConstraint.GetContact().template As<const FRigidBodyPointContactConstraint>();
+					if ((Constraint != nullptr) && (Constraint->GetManifoldPoints().Num() > 0))
+					{
+						ContactPos = FVec3(0);
+						for (const FManifoldPoint& ManifoldPoint : Constraint->GetManifoldPoints())
+						{
+							ContactPos += Transform0.TransformPositionNoScale(ManifoldPoint.CoMContactPoints[0]);
+							ContactPos += Transform1.TransformPositionNoScale(ManifoldPoint.CoMContactPoints[1]);
+						}
+						ContactPos /= (FReal)(2 * Constraint->GetManifoldPoints().Num());
+					}
+				}
+
+
+				FDebugDrawQueue::GetInstance().DrawDebugLine(Transform0.GetLocation(), ContactPos, FColor::Black, false, KINDA_SMALL_NUMBER, Settings.DrawPriority, Settings.LineThickness);
+				FDebugDrawQueue::GetInstance().DrawDebugLine(Transform1.GetLocation(), ContactPos, FColor::White, false, KINDA_SMALL_NUMBER, Settings.DrawPriority, Settings.LineThickness);
+				FDebugDrawQueue::GetInstance().DrawDebugString(ContactPos, FString::Format(TEXT("{0}-{1}-{2}"), { LevelIndex, ColorIndex, OrderIndex }), nullptr, FColor::Yellow, KINDA_SMALL_NUMBER, false, Settings.FontScale);
 
 			};
 
@@ -903,8 +987,10 @@ namespace Chaos
 						if (LevelColorConstraintMap[LevelIndex].Contains(ColorIndex) && LevelColorConstraintMap[LevelIndex][ColorIndex].Num())
 						{
 							const TArray<FConstraintHandle*>& ConstraintHandles = LevelColorConstraintMap[LevelIndex][ColorIndex];
-							for (const FConstraintHandle* ConstraintHandle : ConstraintHandles)
+							for (int32 HandleIndex = 0; HandleIndex < ConstraintHandles.Num(); ++HandleIndex)
 							{
+								const FConstraintHandle* ConstraintHandle = ConstraintHandles[HandleIndex];
+
 								if (const FPBDCollisionConstraintHandle* CollisionHandle = ConstraintHandle->As<FPBDCollisionConstraintHandle>())
 								{
 									if (CollisionHandle->GetConstrainedParticles()[0]->HasBounds() && (CollisionHandle->GetConstrainedParticles()[0]->ObjectState() == EObjectStateType::Dynamic))
@@ -915,7 +1001,11 @@ namespace Chaos
 									{
 										IslandAABB.GrowToInclude(CollisionHandle->GetConstrainedParticles()[1]->WorldSpaceInflatedBounds());
 									}
-									//DrawGraphCollision(SpaceTransform, *CollisionHandle, IslandIndex, LevelIndex, ColorIndex, Settings);
+
+									if (bChaosDebugDebugDrawContactGraph)
+									{
+										DrawGraphCollision(SpaceTransform, *CollisionHandle, IslandIndex, LevelIndex, ColorIndex, HandleIndex, Settings);
+									}
 								}
 							}
 
@@ -1121,8 +1211,16 @@ namespace Chaos
 			{
 				DrawConstraintGraphImpl(SpaceTransform, Graph, GetChaosDebugDrawSettings(Settings));
 			}
+		}
 
+		void DrawCollidingShapes(const FRigidTransform3& SpaceTransform, const FPBDCollisionConstraints& Collisions, float ColorScale, const FChaosDebugDrawSettings* Settings)
+		{
+			if (FDebugDrawQueue::IsDebugDrawingEnabled())
+			{
+				DrawCollidingShapesImpl(SpaceTransform, Collisions, ColorScale, GetChaosDebugDrawSettings(Settings));
+			}
 		}
+
 #endif
 	}
 }
diff --git a/Engine/Source/Runtime/Experimental/Chaos/Private/PBDRigidsSolver.cpp b/Engine/Source/Runtime/Experimental/Chaos/Private/PBDRigidsSolver.cpp
--- a/Engine/Source/Runtime/Experimental/Chaos/Private/PBDRigidsSolver.cpp
+++ b/Engine/Source/Runtime/Experimental/Chaos/Private/PBDRigidsSolver.cpp
@@ -53,6 +53,7 @@ CSV_DEFINE_CATEGORY(ChaosCounters, true);
 
 int32 ChaosSolverDebugDrawShapes = CHAOS_SOLVER_ENABLE_DEBUG_DRAW;
 int32 ChaosSolverDebugDrawCollisions = CHAOS_SOLVER_ENABLE_DEBUG_DRAW;
+int32 ChaosSolverDebugDrawCollidingShapes = 0;
 int32 ChaosSolverDebugDrawBounds = 0;
 int32 ChaosSolverDrawTransforms = 0;
 int32 ChaosSolverDrawIslands = 0;
@@ -61,6 +62,7 @@ int32 ChaosSolverDrawShapesShowKinematic = 1;
 int32 ChaosSolverDrawShapesShowDynamic = 1;
 FAutoConsoleVariableRef CVarChaosSolverDrawShapes(TEXT("p.Chaos.Solver.DebugDrawShapes"), ChaosSolverDebugDrawShapes, TEXT("Draw Shapes (0 = never; 1 = end of frame)."));
 FAutoConsoleVariableRef CVarChaosSolverDrawCollisions(TEXT("p.Chaos.Solver.DebugDrawCollisions"), ChaosSolverDebugDrawCollisions, TEXT("Draw Collisions (0 = never; 1 = end of frame)."));
+FAutoConsoleVariableRef CVarChaosSolverDrawCollidingShapes(TEXT("p.Chaos.Solver.DebugDrawCollidingShapes"), ChaosSolverDebugDrawCollidingShapes, TEXT("Draw Shapes that have collisions on them (0 = never; 1 = end of frame)."));
 FAutoConsoleVariableRef CVarChaosSolverDrawBounds(TEXT("p.Chaos.Solver.DebugDrawBounds"), ChaosSolverDebugDrawBounds, TEXT("Draw bounding volumes inside the broadphase (0 = never; 1 = end of frame)."));
 FAutoConsoleVariableRef CVarChaosSolverDrawTransforms(TEXT("p.Chaos.Solver.DebugDrawTransforms"), ChaosSolverDrawTransforms, TEXT("Draw particle transforms (0 = never; 1 = end of frame)."));
 FAutoConsoleVariableRef CVarChaosSolverDrawIslands(TEXT("p.Chaos.Solver.DebugDrawIslands"), ChaosSolverDrawIslands, TEXT("Draw solver islands (0 = never; 1 = end of frame)."));
@@ -1264,6 +1266,10 @@ namespace Chaos
 		{
 			DebugDraw::DrawConstraintGraph(FRigidTransform3(), GetEvolution()->GetCollisionConstraintsRule().GetGraphColor(), &ChaosSolverDebugDebugDrawSettings);
 		}
+		if (ChaosSolverDebugDrawCollidingShapes == 1)
+		{
+			DebugDraw::DrawCollidingShapes(FRigidTransform3(), GetEvolution()->GetCollisionConstraints(), 1.f, &ChaosSolverDebugDebugDrawSettings);
+		}
 #endif
 	}
 
diff --git a/Engine/Source/Runtime/Experimental/Chaos/Public/Chaos/ChaosDebugDraw.h b/Engine/Source/Runtime/Experimental/Chaos/Public/Chaos/ChaosDebugDraw.h
--- a/Engine/Source/Runtime/Experimental/Chaos/Public/Chaos/ChaosDebugDraw.h
+++ b/Engine/Source/Runtime/Experimental/Chaos/Public/Chaos/ChaosDebugDraw.h
@@ -126,6 +126,7 @@ namespace Chaos
 		CHAOS_API void DrawSimulationSpace(const FSimulationSpace& SimSpace, const FChaosDebugDrawSettings* Settings = nullptr);
 		CHAOS_API void DrawShape(const FRigidTransform3& ShapeTransform, const FImplicitObject* Shape, const FColor& Color, const FChaosDebugDrawSettings* Settings = nullptr);
 		CHAOS_API void DrawConstraintGraph(const FRigidTransform3& ShapeTransform, const FPBDConstraintColor& Graph, const FChaosDebugDrawSettings* Settings = nullptr);
+		CHAOS_API void DrawCollidingShapes(const FRigidTransform3& SpaceTransform, const FPBDCollisionConstraints& Collisions, float ColorScale, const FChaosDebugDrawSettings* Settings = nullptr);
 #endif
 	}
 }
