Chaos: Fixed custom solver actors being non-functional in non-editor builds (standalone, cooked). Make sure the solver is registered with the current world on BeginPlay

diff --git a/Engine/Source/Runtime/Experimental/ChaosSolverEngine/Private/Chaos/ChaosSolverActor.cpp b/Engine/Source/Runtime/Experimental/ChaosSolverEngine/Private/Chaos/ChaosSolverActor.cpp
--- a/Engine/Source/Runtime/Experimental/ChaosSolverEngine/Private/Chaos/ChaosSolverActor.cpp
+++ b/Engine/Source/Runtime/Experimental/ChaosSolverEngine/Private/Chaos/ChaosSolverActor.cpp
@@ -271,6 +271,12 @@ void AChaosSolverActor::BeginPlay()
 		return;
 	}
 
+	// Make sure that the solver is registered in the right world
+	if(FChaosSolversModule* Module = FChaosSolversModule::GetModule())
+	{
+		Module->MigrateSolver(GetSolver(), GetWorld());
+	}
+
 	Solver->EnqueueCommandImmediate(
 		[InSolver = Solver, InProps = Properties]()
 		{
