Fixed Nanite geometry collection velocity vectors, and optimizing transform updates for dummy GPU scene instance culling placeholders. Scene proxies can now fully own the concatenation of transform updates and avoid doing incorrect or redundant work during GPU Scene updates.

diff --git a/Engine/Source/Runtime/Engine/Private/PrimitiveSceneProxy.cpp b/Engine/Source/Runtime/Engine/Private/PrimitiveSceneProxy.cpp
--- a/Engine/Source/Runtime/Engine/Private/PrimitiveSceneProxy.cpp
+++ b/Engine/Source/Runtime/Engine/Private/PrimitiveSceneProxy.cpp
@@ -159,6 +159,7 @@ FPrimitiveSceneProxy::FPrimitiveSceneProxy(const UPrimitiveComponent* InComponen
 ,	bStaticElementsAlwaysUseProxyPrimitiveUniformBuffer(false)
 ,	bVFRequiresPrimitiveUniformBuffer(true)
 ,	bSupportsInstanceDataBuffer(false)
+,	bShouldUpdateGPUSceneTransforms(true)
 ,	bAlwaysHasVelocity(false)
 ,	bSupportsDistanceFieldRepresentation(false)
 ,	bSupportsMeshCardRepresentation(false)
diff --git a/Engine/Source/Runtime/Engine/Public/PrimitiveSceneProxy.h b/Engine/Source/Runtime/Engine/Public/PrimitiveSceneProxy.h
--- a/Engine/Source/Runtime/Engine/Public/PrimitiveSceneProxy.h
+++ b/Engine/Source/Runtime/Engine/Public/PrimitiveSceneProxy.h
@@ -610,9 +610,11 @@ public:
 	inline bool NeedsLevelAddedToWorldNotification() const { return bNeedsLevelAddedToWorldNotification; }
 	inline bool IsComponentLevelVisible() const { return bIsComponentLevelVisible; }
 	inline bool ShouldReceiveMobileCSMShadows() const { return bReceiveMobileCSMShadows; }
+	inline bool ShouldUpdateGPUSceneTransforms() const { return bShouldUpdateGPUSceneTransforms; }
 
 	/** Returns whether draws velocity in base pass. */
-	inline bool DrawsVelocity() const {
+	inline bool DrawsVelocity() const
+	{
 		return IsMovable() || IsBeingMovedByEditor();
 	}
 
@@ -1009,6 +1011,9 @@ protected:
 	/** Whether the primitive supports the GPUScene instance data buffer. */
 	uint8 bSupportsInstanceDataBuffer : 1;
 
+	/** Whether the instances on the primitive need to update transforms during GPU Scene update. */
+	uint8 bShouldUpdateGPUSceneTransforms : 1;
+
 	/** Whether the primitive should always be considered to have velocities, even if it hasn't moved. */
 	uint8 bAlwaysHasVelocity : 1;
 
diff --git a/Engine/Source/Runtime/Experimental/GeometryCollectionEngine/Private/GeometryCollection/GeometryCollectionComponent.cpp b/Engine/Source/Runtime/Experimental/GeometryCollectionEngine/Private/GeometryCollection/GeometryCollectionComponent.cpp
--- a/Engine/Source/Runtime/Experimental/GeometryCollectionEngine/Private/GeometryCollection/GeometryCollectionComponent.cpp
+++ b/Engine/Source/Runtime/Experimental/GeometryCollectionEngine/Private/GeometryCollection/GeometryCollectionComponent.cpp
@@ -412,7 +412,7 @@ FBoxSphereBounds UGeometryCollectionComponent::CalcBounds(const FTransform& Loca
 			GeometryCollectionAlgo::GlobalMatrices(Transforms, ParentIndices, TmpGlobalMatrices);
 
 			TManagedArray<FBox>& TransformBounds = HackGeometryCollectionPtr->GetAttribute<FBox>("BoundingBox", "Transform");
-			for (int TransformIndex = 0; TransformIndex < HackGeometryCollectionPtr->NumElements(FGeometryCollection::TransformGroup); TransformIndex++)
+			for (int32 TransformIndex = 0; TransformIndex < HackGeometryCollectionPtr->NumElements(FGeometryCollection::TransformGroup); TransformIndex++)
 			{
 				BoundingBox += TransformBounds[TransformIndex].TransformBy(TmpGlobalMatrices[TransformIndex] * LocalToWorldWithScale);
 			}
@@ -509,24 +509,26 @@ FPrimitiveSceneProxy* UGeometryCollectionComponent::CreateSceneProxy()
 
 				// ...
 
-#if GEOMETRYCOLLECTION_EDITOR_SELECTION
+			#if GEOMETRYCOLLECTION_EDITOR_SELECTION
 				if (bIsTransformSelectionModeEnabled)
 				{
 					// ...
 				}
-#endif
+			#endif
 
 				ENQUEUE_RENDER_COMMAND(CreateRenderState)(
 					[GeometryCollectionSceneProxy, ConstantData, DynamicData](FRHICommandListImmediate& RHICmdList)
 					{
 						GeometryCollectionSceneProxy->SetConstantData_RenderThread(ConstantData);
+						
 						if (DynamicData)
 						{
 							GeometryCollectionSceneProxy->SetDynamicData_RenderThread(DynamicData);
 						}
-						if (GeometryCollectionSceneProxy->GetPrimitiveSceneInfo())
+
+						if (FPrimitiveSceneInfo* PrimitiveSceneInfo = GeometryCollectionSceneProxy->GetPrimitiveSceneInfo())
 						{
-							GeometryCollectionSceneProxy->GetPrimitiveSceneInfo()->RequestGPUSceneUpdate();
+							PrimitiveSceneInfo->RequestGPUSceneUpdate();
 						}
 					}
 				);
@@ -535,13 +537,13 @@ FPrimitiveSceneProxy* UGeometryCollectionComponent::CreateSceneProxy()
 			{
 				FGeometryCollectionSceneProxy* const GeometryCollectionSceneProxy = static_cast<FGeometryCollectionSceneProxy*>(LocalSceneProxy);
 
-#if GEOMETRYCOLLECTION_EDITOR_SELECTION
+			#if GEOMETRYCOLLECTION_EDITOR_SELECTION
 				// Re-init subsections
 				if (bIsTransformSelectionModeEnabled)
 				{
 					GeometryCollectionSceneProxy->UseSubSections(true, false);  // Do not force reinit now, it'll be done in SetConstantData_RenderThread
 				}
-#endif
+			#endif
 
 				ENQUEUE_RENDER_COMMAND(CreateRenderState)(
 					[GeometryCollectionSceneProxy, ConstantData, DynamicData](FRHICommandListImmediate& RHICmdList)
@@ -556,8 +558,8 @@ FPrimitiveSceneProxy* UGeometryCollectionComponent::CreateSceneProxy()
 			}
 		}
 	}
-	return LocalSceneProxy;
 
+	return LocalSceneProxy;
 }
 
 bool UGeometryCollectionComponent::ShouldCreatePhysicsState() const
diff --git a/Engine/Source/Runtime/Experimental/GeometryCollectionEngine/Private/GeometryCollection/GeometryCollectionSceneProxy.cpp b/Engine/Source/Runtime/Experimental/GeometryCollectionEngine/Private/GeometryCollection/GeometryCollectionSceneProxy.cpp
--- a/Engine/Source/Runtime/Experimental/GeometryCollectionEngine/Private/GeometryCollection/GeometryCollectionSceneProxy.cpp
+++ b/Engine/Source/Runtime/Experimental/GeometryCollectionEngine/Private/GeometryCollection/GeometryCollectionSceneProxy.cpp
@@ -1053,6 +1053,9 @@ FNaniteGeometryCollectionSceneProxy::FNaniteGeometryCollectionSceneProxy(UGeomet
 	// Nanite supports the GPUScene instance data buffer.
 	bSupportsInstanceDataBuffer = true;
 
+	// We always have correct instance transforms, skip GPUScene updates
+	bShouldUpdateGPUSceneTransforms = false;
+
 	bSupportsDistanceFieldRepresentation = false;
 	bSupportsMeshCardRepresentation = false;
 
@@ -1126,8 +1129,12 @@ FNaniteGeometryCollectionSceneProxy::FNaniteGeometryCollectionSceneProxy(UGeomet
 		MaterialSections[SectionIndex].MaterialIndex = MeshSection.MaterialID;
 	}
 
-	bool bHasGeometryBoundingBoxes = Collection->HasAttribute("BoundingBox", FGeometryCollection::GeometryGroup) && Collection->NumElements(FGeometryCollection::GeometryGroup);
-	bool bHasTransformBoundingBoxes = Collection->NumElements(FGeometryCollection::TransformGroup) && 
+	const bool bHasGeometryBoundingBoxes = 
+		Collection->HasAttribute("BoundingBox", FGeometryCollection::GeometryGroup) &&
+		Collection->NumElements(FGeometryCollection::GeometryGroup);
+	
+	const bool bHasTransformBoundingBoxes = 
+		Collection->NumElements(FGeometryCollection::TransformGroup) && 
 		Collection->HasAttribute("BoundingBox", FGeometryCollection::TransformGroup) &&
 		Collection->HasAttribute("TransformToGeometryIndex", FGeometryCollection::TransformGroup);
 
@@ -1146,7 +1153,7 @@ FNaniteGeometryCollectionSceneProxy::FNaniteGeometryCollectionSceneProxy(UGeomet
 			Instance.RenderBounds = BoundingBoxes[GeometryIndex];
 		}
 	}
-	else if(bHasTransformBoundingBoxes)
+	else if (bHasTransformBoundingBoxes)
 	{
 		Nanite::FResources& Resource = GeometryCollection->NaniteData->NaniteResource;
 		NumGeometry = Resource.HierarchyRootOffsets.Num();
@@ -1161,14 +1168,13 @@ FNaniteGeometryCollectionSceneProxy::FNaniteGeometryCollectionSceneProxy(UGeomet
 			if (GeometryIndex > INDEX_NONE)
 			{
 				FGeometryNaniteData& Instance = GeometryNaniteData[GeometryIndex];
-				Instance.PrimitiveId = ~uint32(0);
-				Instance.NaniteInfo = GeometryCollection->GetNaniteInfo(GeometryIndex);
-				Instance.RenderBounds = BoundingBoxes[TransformIndex];
+				Instance.PrimitiveId	= ~uint32(0);
+				Instance.NaniteInfo		= GeometryCollection->GetNaniteInfo(GeometryIndex);
+				Instance.RenderBounds	= BoundingBoxes[TransformIndex];
 			}
 		}
 	}
 
-
 	// Need to specify initial instance list, even with just identity transforms, so that the
 	// GPUScene instance data allocator reserves space for the instances early on. The instance
 	// transforms will be corrected during the first frame before any rendering occurs.
@@ -1178,7 +1184,9 @@ FNaniteGeometryCollectionSceneProxy::FNaniteGeometryCollectionSceneProxy(UGeomet
 		FPrimitiveInstance& Instance = Instances[GeometryIndex];
 		Instance.PrimitiveId = ~uint32(0);
 		Instance.InstanceToLocal.SetIdentity();
+		Instance.PrevInstanceToLocal.SetIdentity();
 		Instance.LocalToWorld.SetIdentity();
+		Instance.PrevLocalToWorld.SetIdentity();
 		Instance.RenderBounds = GeometryNaniteData[GeometryIndex].RenderBounds;
 		Instance.LocalBounds = Instance.RenderBounds;
 	}
@@ -1249,6 +1257,34 @@ uint32 FNaniteGeometryCollectionSceneProxy::GetMemoryFootprint() const
 	return sizeof(*this) + GetAllocatedSize();
 }
 
+void FNaniteGeometryCollectionSceneProxy::OnTransformChanged()
+{
+	FMatrix ParentPrevLocalToWorld;
+
+	// Pull out the previous primitive local to world transform
+	{
+		bool bHasPrecomputedVolumetricLightmap;
+		int32 SingleCaptureIndex;
+		bool bOutputVelocity;
+
+		GetScene().GetPrimitiveUniformShaderParameters_RenderThread(
+			GetPrimitiveSceneInfo(),
+			bHasPrecomputedVolumetricLightmap,
+			ParentPrevLocalToWorld,
+			SingleCaptureIndex,
+			bOutputVelocity
+		);
+	}
+
+	// Primitive has moved, so update all instance transforms
+	const FMatrix ParentLocalToWorld = GetLocalToWorld();
+	for (FPrimitiveInstance& Instance : Instances)
+	{
+		Instance.LocalToWorld = Instance.InstanceToLocal * ParentLocalToWorld;
+		Instance.PrevLocalToWorld = Instance.PrevInstanceToLocal * ParentPrevLocalToWorld;
+	}
+}
+
 void FNaniteGeometryCollectionSceneProxy::SetConstantData_RenderThread(FGeometryCollectionConstantData* NewConstantData, bool ForceInit)
 {
 	const TSharedPtr<FGeometryCollection, ESPMode::ThreadSafe> Collection = GeometryCollection->GetGeometryCollection();
@@ -1269,12 +1305,12 @@ void FNaniteGeometryCollectionSceneProxy::SetConstantData_RenderThread(FGeometry
 
 		FPrimitiveInstance& Instance = Instances.Emplace_GetRef();
 
-		Instance.InstanceToLocal	= NewConstantData->RestTransforms[TransformIndex];
-		Instance.LocalToWorld		= Instance.InstanceToLocal;
-		Instance.PrimitiveId		= NaniteData.PrimitiveId;
-		Instance.RenderBounds       = NaniteData.RenderBounds;
-		Instance.LocalBounds        = Instance.RenderBounds.TransformBy(Instance.InstanceToLocal);
-		Instance.NaniteInfo			= NaniteData.NaniteInfo;
+		Instance.InstanceToLocal		= NewConstantData->RestTransforms[TransformIndex];
+		Instance.PrevInstanceToLocal	= NewConstantData->RestTransforms[TransformIndex];
+		Instance.PrimitiveId			= NaniteData.PrimitiveId;
+		Instance.RenderBounds			= NaniteData.RenderBounds;
+		Instance.LocalBounds			= Instance.RenderBounds.TransformBy(Instance.InstanceToLocal);
+		Instance.NaniteInfo				= NaniteData.NaniteInfo;
 	}
 
 	delete NewConstantData;
@@ -1286,18 +1322,41 @@ void FNaniteGeometryCollectionSceneProxy::SetDynamicData_RenderThread(FGeometryC
 	if (NewDynamicData->IsDynamic)
 	{
 		const TSharedPtr<FGeometryCollection, ESPMode::ThreadSafe> Collection = GeometryCollection->GetGeometryCollection();
-		const TManagedArray<int32>& TransformToGeometryIndices = Collection->TransformToGeometryIndex;
-		const TManagedArray<TSet<int32>>& TransformChildren = Collection->Children;
-		const TManagedArray<int32>& SimulationType = Collection->SimulationType;
+		const TManagedArray<int32>& TransformToGeometryIndices	 = Collection->TransformToGeometryIndex;
+		const TManagedArray<TSet<int32>>& TransformChildren		 = Collection->Children;
+		const TManagedArray<int32>& SimulationType				 = Collection->SimulationType;
+
+		const int32 NumTransforms = NewDynamicData->Transforms.Num();
+		check(NumTransforms == TransformToGeometryIndices.Num());
+		check(NumTransforms == TransformChildren.Num());
+		check(NumTransforms == NewDynamicData->PrevTransforms.Num());
+		Instances.Reset(NumTransforms);
 
-		check(NewDynamicData->Transforms.Num() == TransformToGeometryIndices.Num());
-		check(NewDynamicData->Transforms.Num() == TransformChildren.Num());
-		Instances.Reset(NewDynamicData->Transforms.Num());
+		const FMatrix ParentLocalToWorld = GetLocalToWorld();
+		FMatrix ParentPrevLocalToWorld;
 
-		for (int32 TransformIndex = 0; TransformIndex < NewDynamicData->Transforms.Num(); ++TransformIndex)
+		// Pull out the previous primitive local to world transform
+		{
+			bool bHasPrecomputedVolumetricLightmap;
+			int32 SingleCaptureIndex;
+			bool bOutputVelocity;
+
+			if (GetPrimitiveSceneInfo())
+			{
+				GetScene().GetPrimitiveUniformShaderParameters_RenderThread(
+					GetPrimitiveSceneInfo(),
+					bHasPrecomputedVolumetricLightmap,
+					ParentPrevLocalToWorld,
+					SingleCaptureIndex,
+					bOutputVelocity
+				);
+			}
+		}
+
+		for (int32 TransformIndex = 0; TransformIndex < NumTransforms; ++TransformIndex)
 		{
 			const int32 TransformToGeometryIndex = TransformToGeometryIndices[TransformIndex];
-			if( SimulationType[TransformIndex] != FGeometryCollection::ESimulationTypes::FST_Rigid)
+			if (SimulationType[TransformIndex] != FGeometryCollection::ESimulationTypes::FST_Rigid)
 			{
 				continue;
 			}
@@ -1306,12 +1365,14 @@ void FNaniteGeometryCollectionSceneProxy::SetDynamicData_RenderThread(FGeometryC
 
 			FPrimitiveInstance& Instance = Instances.Emplace_GetRef();
 
-			Instance.InstanceToLocal	= NewDynamicData->Transforms[TransformIndex];
-			Instance.LocalToWorld		= Instance.InstanceToLocal;
-			Instance.PrimitiveId		= NaniteData.PrimitiveId;
-			Instance.RenderBounds       = NaniteData.RenderBounds;
-			Instance.LocalBounds        = Instance.RenderBounds.TransformBy(Instance.InstanceToLocal);
-			Instance.NaniteInfo			= NaniteData.NaniteInfo;
+			Instance.InstanceToLocal			= NewDynamicData->Transforms[TransformIndex];
+			Instance.PrevInstanceToLocal		= NewDynamicData->PrevTransforms[TransformIndex];
+			Instance.LocalToWorld				= Instance.InstanceToLocal * ParentLocalToWorld;
+			Instance.PrevLocalToWorld			= Instance.PrevInstanceToLocal * ParentPrevLocalToWorld;
+			Instance.LocalBounds				= Instance.RenderBounds.TransformBy(Instance.InstanceToLocal);
+			Instance.PrimitiveId				= NaniteData.PrimitiveId;
+			Instance.RenderBounds				= NaniteData.RenderBounds;
+			Instance.NaniteInfo					= NaniteData.NaniteInfo;
 		}
 	}
 	else
diff --git a/Engine/Source/Runtime/Experimental/GeometryCollectionEngine/Private/GeometryCollection/GeometryCollectionSceneProxy.h b/Engine/Source/Runtime/Experimental/GeometryCollectionEngine/Private/GeometryCollection/GeometryCollectionSceneProxy.h
--- a/Engine/Source/Runtime/Experimental/GeometryCollectionEngine/Private/GeometryCollection/GeometryCollectionSceneProxy.h
+++ b/Engine/Source/Runtime/Experimental/GeometryCollectionEngine/Private/GeometryCollection/GeometryCollectionSceneProxy.h
@@ -332,6 +332,8 @@ public:
 
 	virtual uint32 GetMemoryFootprint() const override;
 
+	virtual void OnTransformChanged() override;
+
 	/** Called on render thread to setup static geometry for rendering */
 	void SetConstantData_RenderThread(FGeometryCollectionConstantData* NewConstantData, bool ForceInit = false);
 
