Fixed locked previous GC transforms (locked motion blur) by detecting dynamic->static transition and resetting previous transforms to current.

diff --git a/Engine/Source/Runtime/Experimental/GeometryCollectionEngine/Private/GeometryCollection/GeometryCollectionComponent.cpp b/Engine/Source/Runtime/Experimental/GeometryCollectionEngine/Private/GeometryCollection/GeometryCollectionComponent.cpp
--- a/Engine/Source/Runtime/Experimental/GeometryCollectionEngine/Private/GeometryCollection/GeometryCollectionComponent.cpp
+++ b/Engine/Source/Runtime/Experimental/GeometryCollectionEngine/Private/GeometryCollection/GeometryCollectionComponent.cpp
@@ -520,7 +520,7 @@ FPrimitiveSceneProxy* UGeometryCollectionComponent::CreateSceneProxy()
 					[GeometryCollectionSceneProxy, ConstantData, DynamicData](FRHICommandListImmediate& RHICmdList)
 					{
 						GeometryCollectionSceneProxy->SetConstantData_RenderThread(ConstantData);
-						
+
 						if (DynamicData)
 						{
 							GeometryCollectionSceneProxy->SetDynamicData_RenderThread(DynamicData);
@@ -2021,10 +2021,10 @@ void UGeometryCollectionComponent::SendRenderDynamicData_Concurrent()
 	{
 		FGeometryCollectionDynamicData* DynamicData = InitDynamicData(false /* initialization */);
 
-		if (DynamicData)
+		if (DynamicData || SceneProxy->IsNaniteMesh())
 		{
-			INC_DWORD_STAT_BY(STAT_GCTotalTransforms, DynamicData->Transforms.Num());
-			INC_DWORD_STAT_BY(STAT_GCChangedTransforms, DynamicData->ChangedCount);
+			INC_DWORD_STAT_BY(STAT_GCTotalTransforms, DynamicData ? DynamicData->Transforms.Num() : 0);
+			INC_DWORD_STAT_BY(STAT_GCChangedTransforms, DynamicData ? DynamicData->ChangedCount : 0);
 
 			// #todo (bmiller) Once ISMC changes have been complete, this is the best place to call this method
 			// but we can't currently because it's an inappropriate place to call MarkRenderStateDirty on the ISMC.
@@ -2037,8 +2037,15 @@ void UGeometryCollectionComponent::SendRenderDynamicData_Concurrent()
 				ENQUEUE_RENDER_COMMAND(SendRenderDynamicData)(
 					[GeometryCollectionSceneProxy, DynamicData](FRHICommandListImmediate& RHICmdList)
 					{
-						GeometryCollectionSceneProxy->SetDynamicData_RenderThread(DynamicData);
-						GeometryCollectionSceneProxy->GetPrimitiveSceneInfo()->RequestGPUSceneUpdate();
+						if (DynamicData)
+						{
+							GeometryCollectionSceneProxy->SetDynamicData_RenderThread(DynamicData);
+						}
+						else
+						{
+							// No longer dynamic, make sure previous transforms are reset
+							GeometryCollectionSceneProxy->ResetPreviousTransforms_RenderThread();
+						}
 					}
 				);
 			}
diff --git a/Engine/Source/Runtime/Experimental/GeometryCollectionEngine/Private/GeometryCollection/GeometryCollectionSceneProxy.cpp b/Engine/Source/Runtime/Experimental/GeometryCollectionEngine/Private/GeometryCollection/GeometryCollectionSceneProxy.cpp
--- a/Engine/Source/Runtime/Experimental/GeometryCollectionEngine/Private/GeometryCollection/GeometryCollectionSceneProxy.cpp
+++ b/Engine/Source/Runtime/Experimental/GeometryCollectionEngine/Private/GeometryCollection/GeometryCollectionSceneProxy.cpp
@@ -14,6 +14,7 @@
 #if GEOMETRYCOLLECTION_EDITOR_SELECTION
 #include "GeometryCollection/GeometryCollectionHitProxy.h"
 #endif
+
 #if INTEL_ISPC
 #if USING_CODE_ANALYSIS
     MSVC_PRAGMA( warning( push ) )
@@ -36,6 +37,14 @@
 #endif    // USING_CODE_ANALYSIS
 #endif
 
+#include "ComponentReregisterContext.h"
+#include "ComponentRecreateRenderStateContext.h"
+
+static void RecreateGlobalRenderState(IConsoleVariable* Var)
+{
+	FGlobalComponentRecreateRenderStateContext Context;
+}
+
 static int32 GParallelGeometryCollectionBatchSize = 1024;
 static TAutoConsoleVariable<int32> CVarParallelGeometryCollectionBatchSize(
 	TEXT("r.ParallelGeometryCollectionBatchSize"),
@@ -44,7 +53,7 @@ static TAutoConsoleVariable<int32> CVarParallelGeometryCollectionBatchSize(
 	ECVF_Default
 );
 
-int32 GGeometryCollectionTripleBufferUploads = 1;
+static int32 GGeometryCollectionTripleBufferUploads = 1;
 FAutoConsoleVariableRef CVarGeometryCollectionTripleBufferUploads(
 	TEXT("r.GeometryCollectionTripleBufferUploads"),
 	GGeometryCollectionTripleBufferUploads,
@@ -52,6 +61,15 @@ FAutoConsoleVariableRef CVarGeometryCollectionTripleBufferUploads(
 	ECVF_Default
 );
 
+static int32 GGeometryCollectionOptimizedTransforms = 1;
+FAutoConsoleVariableRef CVarGeometryCollectionOptimizedTransforms(
+	TEXT("r.GeometryCollectionOptimizedTransforms"),
+	GGeometryCollectionOptimizedTransforms,
+	TEXT("Whether to optimize transform update by skipping automatic updates in GPUScene."),
+	FConsoleVariableDelegate::CreateStatic(&RecreateGlobalRenderState),
+	ECVF_Scalability | ECVF_RenderThreadSafe
+);
+
 DEFINE_LOG_CATEGORY_STATIC(FGeometryCollectionSceneProxyLogging, Log, All);
 
 FGeometryCollectionDynamicDataPool GDynamicDataPool;
@@ -1053,8 +1071,8 @@ FNaniteGeometryCollectionSceneProxy::FNaniteGeometryCollectionSceneProxy(UGeomet
 	// Nanite supports the GPUScene instance data buffer.
 	bSupportsInstanceDataBuffer = true;
 
-	// We always have correct instance transforms, skip GPUScene updates
-	bShouldUpdateGPUSceneTransforms = false;
+	// We always have correct instance transforms, skip GPUScene updates if allowed.
+	bShouldUpdateGPUSceneTransforms = (GGeometryCollectionOptimizedTransforms == 0);
 
 	bSupportsDistanceFieldRepresentation = false;
 	bSupportsMeshCardRepresentation = false;
@@ -1283,6 +1301,9 @@ void FNaniteGeometryCollectionSceneProxy::OnTransformChanged()
 		Instance.LocalToWorld = Instance.InstanceToLocal * ParentLocalToWorld;
 		Instance.PrevLocalToWorld = Instance.PrevInstanceToLocal * ParentPrevLocalToWorld;
 	}
+
+	// Make this count as a dynamic update wrt. previous transform reset on stasis.
+	bLastUpdateWasDynamic = true;
 }
 
 void FNaniteGeometryCollectionSceneProxy::SetConstantData_RenderThread(FGeometryCollectionConstantData* NewConstantData, bool ForceInit)
@@ -1321,6 +1342,8 @@ void FNaniteGeometryCollectionSceneProxy::SetDynamicData_RenderThread(FGeometryC
 	// Are we currently simulating?
 	if (NewDynamicData->IsDynamic)
 	{
+		bLastUpdateWasDynamic = true;
+
 		const TSharedPtr<FGeometryCollection, ESPMode::ThreadSafe> Collection = GeometryCollection->GetGeometryCollection();
 		const TManagedArray<int32>& TransformToGeometryIndices	 = Collection->TransformToGeometryIndex;
 		const TManagedArray<TSet<int32>>& TransformChildren		 = Collection->Children;
@@ -1384,6 +1407,20 @@ void FNaniteGeometryCollectionSceneProxy::SetDynamicData_RenderThread(FGeometryC
 	GDynamicDataPool.Release(NewDynamicData);
 }
 
+void FNaniteGeometryCollectionSceneProxy::ResetPreviousTransforms_RenderThread()
+{
+	if (bLastUpdateWasDynamic)
+	{
+		// Reset previous transforms to avoid locked motion vectors
+		for (FPrimitiveInstance& Instance : Instances)
+		{
+			Instance.PrevLocalToWorld = Instance.LocalToWorld;
+		}
+
+		bLastUpdateWasDynamic = false;
+	}
+}
+
 FGeometryCollectionDynamicDataPool::FGeometryCollectionDynamicDataPool()
 {
 	FreeList.SetNum(32);
diff --git a/Engine/Source/Runtime/Experimental/GeometryCollectionEngine/Private/GeometryCollection/GeometryCollectionSceneProxy.h b/Engine/Source/Runtime/Experimental/GeometryCollectionEngine/Private/GeometryCollection/GeometryCollectionSceneProxy.h
--- a/Engine/Source/Runtime/Experimental/GeometryCollectionEngine/Private/GeometryCollection/GeometryCollectionSceneProxy.h
+++ b/Engine/Source/Runtime/Experimental/GeometryCollectionEngine/Private/GeometryCollection/GeometryCollectionSceneProxy.h
@@ -340,6 +340,8 @@ public:
 	/** Called on render thread to setup dynamic geometry for rendering */
 	void SetDynamicData_RenderThread(FGeometryCollectionDynamicData* NewDynamicData);
 
+	void ResetPreviousTransforms_RenderThread();
+
 	const FORCEINLINE TArray<Nanite::FResources*>& GetResources() const
 	{
 		return Resources;
@@ -360,11 +362,10 @@ protected:
 	uint32 bCastShadow : 1;
 	uint32 bReverseCulling : 1;
 	uint32 bHasMaterialErrors : 1;
+	uint32 bLastUpdateWasDynamic : 1;
 
 	const UGeometryCollection* GeometryCollection = nullptr;
 
-	//TArray<FPrimitiveInstance> RestInstances;
-
 	struct FGeometryNaniteData
 	{
 		FBoxSphereBounds RenderBounds;
