Fix duplicate symbol unity build error

diff --git a/Engine/Source/Runtime/Engine/Private/StaticMeshRender.cpp b/Engine/Source/Runtime/Engine/Private/StaticMeshRender.cpp
--- a/Engine/Source/Runtime/Engine/Private/StaticMeshRender.cpp
+++ b/Engine/Source/Runtime/Engine/Private/StaticMeshRender.cpp
@@ -98,7 +98,7 @@ static FAutoConsoleCommand GToggleReversedIndexBuffersCmd(
 	FConsoleCommandDelegate::CreateStatic(ToggleReversedIndexBuffers)
 	);
 
-static void RecreateGlobalRenderState(IConsoleVariable* Var)
+static void RecreateSMGlobalRenderState(IConsoleVariable* Var)
 {
 	FGlobalComponentRecreateRenderStateContext Context;
 }
@@ -109,7 +109,7 @@ FAutoConsoleVariableRef CVarRenderNaniteMeshes(
 	TEXT("r.Nanite"),
 	GRenderNaniteMeshes,
 	TEXT("Render static meshes using Nanite."),
-	FConsoleVariableDelegate::CreateStatic(&RecreateGlobalRenderState),
+	FConsoleVariableDelegate::CreateStatic(&RecreateSMGlobalRenderState),
 	ECVF_Scalability | ECVF_RenderThreadSafe
 );
 
diff --git a/Engine/Source/Runtime/Experimental/GeometryCollectionEngine/Private/GeometryCollection/GeometryCollectionComponent.cpp b/Engine/Source/Runtime/Experimental/GeometryCollectionEngine/Private/GeometryCollection/GeometryCollectionComponent.cpp
--- a/Engine/Source/Runtime/Experimental/GeometryCollectionEngine/Private/GeometryCollection/GeometryCollectionComponent.cpp
+++ b/Engine/Source/Runtime/Experimental/GeometryCollectionEngine/Private/GeometryCollection/GeometryCollectionComponent.cpp
@@ -165,7 +165,7 @@ bool FGeometryCollectionRepData::NetSerialize(FArchive& Ar, class UPackageMap* M
 	return true;
 }
 
-static void RecreateGlobalRenderState(IConsoleVariable* Var)
+static void RecreateGCGlobalRenderState(IConsoleVariable* Var)
 {
 	FGlobalComponentRecreateRenderStateContext Context;
 }
@@ -175,7 +175,7 @@ FAutoConsoleVariableRef CVarGeometryCollectionNanite(
 	TEXT("r.GeometryCollection.Nanite"),
 	GGeometryCollectionNanite,
 	TEXT("Render geometry collections using Nanite."),
-	FConsoleVariableDelegate::CreateStatic(&RecreateGlobalRenderState),
+	FConsoleVariableDelegate::CreateStatic(&RecreateGCGlobalRenderState),
 	ECVF_RenderThreadSafe
 );
 
diff --git a/Engine/Source/Runtime/Experimental/GeometryCollectionEngine/Private/GeometryCollection/GeometryCollectionSceneProxy.cpp b/Engine/Source/Runtime/Experimental/GeometryCollectionEngine/Private/GeometryCollection/GeometryCollectionSceneProxy.cpp
--- a/Engine/Source/Runtime/Experimental/GeometryCollectionEngine/Private/GeometryCollection/GeometryCollectionSceneProxy.cpp
+++ b/Engine/Source/Runtime/Experimental/GeometryCollectionEngine/Private/GeometryCollection/GeometryCollectionSceneProxy.cpp
@@ -40,7 +40,7 @@
 #include "ComponentReregisterContext.h"
 #include "ComponentRecreateRenderStateContext.h"
 
-static void RecreateGlobalRenderState(IConsoleVariable* Var)
+static void RecreateGCPGlobalRenderState(IConsoleVariable* Var)
 {
 	FGlobalComponentRecreateRenderStateContext Context;
 }
@@ -66,7 +66,7 @@ FAutoConsoleVariableRef CVarGeometryCollectionOptimizedTransforms(
 	TEXT("r.GeometryCollectionOptimizedTransforms"),
 	GGeometryCollectionOptimizedTransforms,
 	TEXT("Whether to optimize transform update by skipping automatic updates in GPUScene."),
-	FConsoleVariableDelegate::CreateStatic(&RecreateGlobalRenderState),
+	FConsoleVariableDelegate::CreateStatic(&RecreateGCPGlobalRenderState),
 	ECVF_Scalability | ECVF_RenderThreadSafe
 );
 
