Property watcher class that can accept a user-defined equality test. Useful to watch types that don't have an equality operator.

#rnx

diff --git a/Engine/Source/Runtime/Experimental/InteractiveToolsFramework/Public/InteractiveTool.h b/Engine/Source/Runtime/Experimental/InteractiveToolsFramework/Public/InteractiveTool.h
--- a/Engine/Source/Runtime/Experimental/InteractiveToolsFramework/Public/InteractiveTool.h
+++ b/Engine/Source/Runtime/Experimental/InteractiveToolsFramework/Public/InteractiveTool.h
@@ -49,21 +49,57 @@ public:
 	};
 
 	template <typename PropType>
-	class TPropertyWatcher : public FPropertyWatcher
+	class TPropertyWatcherBase : public FPropertyWatcher
 	{
 	public:
 		using FValueGetter = TFunction<PropType(void)>;
 		using FChangedCallback = TFunction<void(const PropType&)>;
 
-		TPropertyWatcher(const PropType& Property,
-						 FChangedCallback OnChangedIn)
-			: GetValue([&Property](){return Property;}),
+		TPropertyWatcherBase(const PropType& Property,
+							 FChangedCallback OnChangedIn)
+			: GetValue([&Property]() {return Property; }), 
+			  OnChanged(MoveTemp(OnChangedIn))
+		{}
+
+		TPropertyWatcherBase(FValueGetter GetValueIn,
+							 FChangedCallback OnChangedIn)
+			: GetValue(MoveTemp(GetValueIn)), 
 			  OnChanged(MoveTemp(OnChangedIn))
 		{}
+
+		void SilentUpdate() final
+		{
+			Cached = GetValue();
+		}
+
+	protected:
+		TOptional<PropType> Cached;
+		FValueGetter GetValue;
+		FChangedCallback OnChanged;
+	};
+
+
+	template <typename PropType>
+	class TPropertyWatcher : public TPropertyWatcherBase<PropType>
+	{
+	public:
+		using typename TPropertyWatcherBase<PropType>::FChangedCallback;
+		using typename TPropertyWatcherBase<PropType>::FValueGetter;
+		using TPropertyWatcherBase<PropType>::Cached;
+		using TPropertyWatcherBase<PropType>::GetValue;
+		using TPropertyWatcherBase<PropType>::OnChanged;
+
+
+		TPropertyWatcher(const PropType& Property,
+						 FChangedCallback OnChangedIn) 
+			: TPropertyWatcherBase<PropType>(Property, OnChangedIn)
+		{}
+
 		TPropertyWatcher(FValueGetter GetValueIn,
-						 FChangedCallback OnChangedIn)
-			: GetValue(MoveTemp(GetValueIn)), OnChanged(MoveTemp(OnChangedIn))
+						 FChangedCallback OnChangedIn) 
+			: TPropertyWatcherBase<PropType>(GetValueIn, OnChangedIn)
 		{}
+
 		void CheckAndUpdate() final
 		{
 			PropType Value = GetValue();
@@ -73,14 +109,50 @@ public:
 				OnChanged(Cached.GetValue());
 			}
 		}
-		void SilentUpdate() final
+	};
+
+	/// Property watcher for types that don't have an equality operator available.
+	/// 
+	/// Note: TPropertyWatcherWithEqualityTest can't inherit directly from TPropertyWatcher, since the presence of 
+	/// the == operator in that template class will cause compile errors for types without the operator.
+	template <typename PropType>
+	class TPropertyWatcherWithEqualityTest : public TPropertyWatcherBase<PropType>
+	{
+	public:
+		using FEqualityCheckFunction = TFunction<bool(const PropType&, const PropType&)>;
+
+		using typename TPropertyWatcherBase<PropType>::FChangedCallback;
+		using typename TPropertyWatcherBase<PropType>::FValueGetter;
+		using TPropertyWatcherBase<PropType>::Cached;
+		using TPropertyWatcherBase<PropType>::GetValue;
+		using TPropertyWatcherBase<PropType>::OnChanged;
+
+		TPropertyWatcherWithEqualityTest(const PropType& Property,
+										 FChangedCallback OnChangedIn,
+										 FEqualityCheckFunction EqualsIn)
+			: TPropertyWatcherBase<PropType>(Property, OnChangedIn),
+			  Equals(EqualsIn)
+		{}
+
+		TPropertyWatcherWithEqualityTest(FValueGetter GetValueIn,
+										 FChangedCallback OnChangedIn,
+										 FEqualityCheckFunction EqualsIn)
+			: TPropertyWatcherBase<PropType>(GetValueIn, OnChangedIn),
+			  Equals(EqualsIn)
+		{}
+
+		void CheckAndUpdate() final
 		{
-			Cached = GetValue();
+			PropType Value = this->GetValue();
+			if ((!this->Cached.IsSet()) || !Equals(this->Cached.GetValue(), Value))
+			{
+				this->Cached = Value;
+				this->OnChanged(this->Cached.GetValue());
+			}
 		}
-	private:
-		TOptional<PropType> Cached;
-		FValueGetter GetValue;
-		FChangedCallback OnChanged;
+
+	protected:
+		FEqualityCheckFunction Equals;
 	};
 
 	FWatchablePropertySet() = default;
@@ -126,6 +198,24 @@ public:
 	{
 		return PropertyWatchers.Emplace(MakeUnique<TPropertyWatcher<PropType>>(GetValueIn, OnChangedIn));
 	}
+
+	/** @return Index of the watcher, which can be used in SilentUpdateWatcherAtIndex */
+	template <typename PropType>
+	int32 WatchProperty(const PropType& ValueIn,
+						typename TPropertyWatcherWithEqualityTest<PropType>::FChangedCallback OnChangedIn,
+						typename TPropertyWatcherWithEqualityTest<PropType>::FEqualityCheckFunction EqualsIn)
+	{
+		return PropertyWatchers.Emplace(MakeUnique<TPropertyWatcherWithEqualityTest<PropType>>(ValueIn, OnChangedIn, EqualsIn));
+	}
+	/** @return Index of the watcher, which can be used in SilentUpdateWatcherAtIndex */
+	template <typename PropType>
+	int32 WatchProperty(typename TPropertyWatcherWithEqualityTest<PropType>::FValueGetter GetValueIn,
+					  typename TPropertyWatcherWithEqualityTest<PropType>::FChangedCallback OnChangedIn,
+					  typename TPropertyWatcherWithEqualityTest<PropType>::FEqualityCheckFunction EqualsIn)
+	{
+		return PropertyWatchers.Emplace(MakeUnique<TPropertyWatcherWithEqualityTest<PropType>>(GetValueIn, OnChangedIn, EqualsIn));
+	}
+
 private:
 	TArray<TUniquePtr<FPropertyWatcher>> PropertyWatchers;
 };
