GeometryCollectionComponent
-- Added state machine logic and scene proxy callbacks for triggering when motion begins and ends.

diff --git a/Engine/Source/Runtime/Experimental/GeometryCollectionEngine/Private/GeometryCollection/GeometryCollectionComponent.cpp b/Engine/Source/Runtime/Experimental/GeometryCollectionEngine/Private/GeometryCollection/GeometryCollectionComponent.cpp
--- a/Engine/Source/Runtime/Experimental/GeometryCollectionEngine/Private/GeometryCollection/GeometryCollectionComponent.cpp
+++ b/Engine/Source/Runtime/Experimental/GeometryCollectionEngine/Private/GeometryCollection/GeometryCollectionComponent.cpp
@@ -249,6 +249,7 @@ UGeometryCollectionComponent::UGeometryCollectionComponent(const FObjectInitiali
 #if GEOMETRYCOLLECTION_EDITOR_SELECTION
 	, bIsTransformSelectionModeEnabled(false)
 #endif  // #if GEOMETRYCOLLECTION_EDITOR_SELECTION
+	, bIsMoving(false)
 {
 	PrimaryComponentTick.bCanEverTick = true;
 	bTickInEditor = true;
@@ -1528,10 +1529,37 @@ FGeometryCollectionDynamicData* UGeometryCollectionComponent::InitDynamicData(bo
 		}
 	}
 
-	if (DynamicData && DynamicData->ChangedCount == 0 && !bEditorMode && !bInitialization)
+	if (!bEditorMode && !bInitialization)
 	{
-		GDynamicDataPool.Release(DynamicData);
-		DynamicData = nullptr;
+		if (DynamicData && DynamicData->ChangedCount == 0)
+		{
+			GDynamicDataPool.Release(DynamicData);
+			DynamicData = nullptr;
+
+			// Change of state?
+			if (bIsMoving)
+			{
+				bIsMoving = false;
+				if (SceneProxy && SceneProxy->IsNaniteMesh())
+				{
+					FNaniteGeometryCollectionSceneProxy* NaniteProxy = static_cast<FNaniteGeometryCollectionSceneProxy*>(SceneProxy);
+					NaniteProxy->OnMotionEnd();
+				}
+			}
+		}
+		else
+		{
+			// Change of state?
+			if (!bIsMoving)
+			{
+				bIsMoving = true;
+				if (SceneProxy && SceneProxy->IsNaniteMesh())
+				{
+					FNaniteGeometryCollectionSceneProxy* NaniteProxy = static_cast<FNaniteGeometryCollectionSceneProxy*>(SceneProxy);
+					NaniteProxy->OnMotionBegin();
+				}
+			}
+		}
 	}
 
 	return DynamicData;
diff --git a/Engine/Source/Runtime/Experimental/GeometryCollectionEngine/Private/GeometryCollection/GeometryCollectionSceneProxy.cpp b/Engine/Source/Runtime/Experimental/GeometryCollectionEngine/Private/GeometryCollection/GeometryCollectionSceneProxy.cpp
--- a/Engine/Source/Runtime/Experimental/GeometryCollectionEngine/Private/GeometryCollection/GeometryCollectionSceneProxy.cpp
+++ b/Engine/Source/Runtime/Experimental/GeometryCollectionEngine/Private/GeometryCollection/GeometryCollectionSceneProxy.cpp
@@ -1421,6 +1421,16 @@ void FNaniteGeometryCollectionSceneProxy::ResetPreviousTransforms_RenderThread()
 	}
 }
 
+void FNaniteGeometryCollectionSceneProxy::OnMotionBegin()
+{
+	// Currently we don't trigger anything
+}
+
+void FNaniteGeometryCollectionSceneProxy::OnMotionEnd()
+{
+	// Currently we don't trigger anything
+}
+
 FGeometryCollectionDynamicDataPool::FGeometryCollectionDynamicDataPool()
 {
 	FreeList.SetNum(32);
diff --git a/Engine/Source/Runtime/Experimental/GeometryCollectionEngine/Private/GeometryCollection/GeometryCollectionSceneProxy.h b/Engine/Source/Runtime/Experimental/GeometryCollectionEngine/Private/GeometryCollection/GeometryCollectionSceneProxy.h
--- a/Engine/Source/Runtime/Experimental/GeometryCollectionEngine/Private/GeometryCollection/GeometryCollectionSceneProxy.h
+++ b/Engine/Source/Runtime/Experimental/GeometryCollectionEngine/Private/GeometryCollection/GeometryCollectionSceneProxy.h
@@ -164,6 +164,8 @@ private:
 	FCriticalSection ListLock;
 };
 
+
+
 /***
 *   FGeometryCollectionSceneProxy
 *    
@@ -352,6 +354,9 @@ public:
 		return Resources;
 	}
 
+	void OnMotionBegin();
+	void OnMotionEnd();
+
 protected:
 	TArray<Nanite::FResources*> Resources;
 
diff --git a/Engine/Source/Runtime/Experimental/GeometryCollectionEngine/Public/GeometryCollection/GeometryCollectionComponent.h b/Engine/Source/Runtime/Experimental/GeometryCollectionEngine/Public/GeometryCollection/GeometryCollectionComponent.h
--- a/Engine/Source/Runtime/Experimental/GeometryCollectionEngine/Public/GeometryCollection/GeometryCollectionComponent.h
+++ b/Engine/Source/Runtime/Experimental/GeometryCollectionEngine/Public/GeometryCollection/GeometryCollectionComponent.h
@@ -785,4 +785,7 @@ private:
 	bool IsEmbeddedGeometryValid() const;
 	void ClearEmbeddedGeometry();
 
+	/** True if GeometryCollection transforms have changed from previous tick. */
+	bool bIsMoving;
+
 };
