Corrected OnMotionBegin/OnMotionEnd to run on the render thread, and refactored the previous in motion logic to use the new events.

diff --git a/Engine/Source/Runtime/Experimental/GeometryCollectionEngine/Private/GeometryCollection/GeometryCollectionComponent.cpp b/Engine/Source/Runtime/Experimental/GeometryCollectionEngine/Private/GeometryCollection/GeometryCollectionComponent.cpp
--- a/Engine/Source/Runtime/Experimental/GeometryCollectionEngine/Private/GeometryCollection/GeometryCollectionComponent.cpp
+++ b/Engine/Source/Runtime/Experimental/GeometryCollectionEngine/Private/GeometryCollection/GeometryCollectionComponent.cpp
@@ -1543,7 +1543,12 @@ FGeometryCollectionDynamicData* UGeometryCollectionComponent::InitDynamicData(bo
 				if (SceneProxy && SceneProxy->IsNaniteMesh())
 				{
 					FNaniteGeometryCollectionSceneProxy* NaniteProxy = static_cast<FNaniteGeometryCollectionSceneProxy*>(SceneProxy);
-					NaniteProxy->OnMotionEnd();
+					ENQUEUE_RENDER_COMMAND(NaniteProxyOnMotionEnd)(
+						[NaniteProxy](FRHICommandListImmediate& RHICmdList)
+						{
+							NaniteProxy->OnMotionEnd();
+						}
+					);
 				}
 			}
 		}
@@ -1556,7 +1561,12 @@ FGeometryCollectionDynamicData* UGeometryCollectionComponent::InitDynamicData(bo
 				if (SceneProxy && SceneProxy->IsNaniteMesh())
 				{
 					FNaniteGeometryCollectionSceneProxy* NaniteProxy = static_cast<FNaniteGeometryCollectionSceneProxy*>(SceneProxy);
-					NaniteProxy->OnMotionBegin();
+					ENQUEUE_RENDER_COMMAND(NaniteProxyOnMotionBegin)(
+						[NaniteProxy](FRHICommandListImmediate& RHICmdList)
+						{
+							NaniteProxy->OnMotionBegin();
+						}
+					);
 				}
 			}
 		}
diff --git a/Engine/Source/Runtime/Experimental/GeometryCollectionEngine/Private/GeometryCollection/GeometryCollectionSceneProxy.cpp b/Engine/Source/Runtime/Experimental/GeometryCollectionEngine/Private/GeometryCollection/GeometryCollectionSceneProxy.cpp
--- a/Engine/Source/Runtime/Experimental/GeometryCollectionEngine/Private/GeometryCollection/GeometryCollectionSceneProxy.cpp
+++ b/Engine/Source/Runtime/Experimental/GeometryCollectionEngine/Private/GeometryCollection/GeometryCollectionSceneProxy.cpp
@@ -1059,6 +1059,7 @@ void FGeometryCollectionSceneProxy::GetPreSkinnedLocalBounds(FBoxSphereBounds& O
 FNaniteGeometryCollectionSceneProxy::FNaniteGeometryCollectionSceneProxy(UGeometryCollectionComponent* Component)
 : Nanite::FSceneProxyBase(Component)
 , GeometryCollection(Component->GetRestCollection())
+, bCurrentlyInMotion(false)
 {
 	LLM_SCOPE_BYTAG(Nanite);
 
@@ -1296,14 +1297,22 @@ void FNaniteGeometryCollectionSceneProxy::OnTransformChanged()
 
 	// Primitive has moved, so update all instance transforms
 	const FMatrix ParentLocalToWorld = GetLocalToWorld();
-	for (FPrimitiveInstance& Instance : Instances)
+	if (bCurrentlyInMotion)
 	{
-		Instance.LocalToWorld = Instance.InstanceToLocal * ParentLocalToWorld;
-		Instance.PrevLocalToWorld = Instance.PrevInstanceToLocal * ParentPrevLocalToWorld;
+		for (FPrimitiveInstance& Instance : Instances)
+		{
+			Instance.LocalToWorld = Instance.InstanceToLocal * ParentLocalToWorld;
+			Instance.PrevLocalToWorld = Instance.PrevInstanceToLocal * ParentPrevLocalToWorld;
+		}
+	}
+	else
+	{
+		for (FPrimitiveInstance& Instance : Instances)
+		{
+			Instance.LocalToWorld = Instance.InstanceToLocal * ParentLocalToWorld;
+			Instance.PrevLocalToWorld = Instance.LocalToWorld;
+		}
 	}
-
-	// Make this count as a dynamic update wrt. previous transform reset on stasis.
-	bLastUpdateWasDynamic = true;
 }
 
 void FNaniteGeometryCollectionSceneProxy::SetConstantData_RenderThread(FGeometryCollectionConstantData* NewConstantData, bool ForceInit)
@@ -1342,8 +1351,6 @@ void FNaniteGeometryCollectionSceneProxy::SetDynamicData_RenderThread(FGeometryC
 	// Are we currently simulating?
 	if (NewDynamicData->IsDynamic)
 	{
-		bLastUpdateWasDynamic = true;
-
 		const TSharedPtr<FGeometryCollection, ESPMode::ThreadSafe> Collection = GeometryCollection->GetGeometryCollection();
 		const TManagedArray<int32>& TransformToGeometryIndices	 = Collection->TransformToGeometryIndex;
 		const TManagedArray<TSet<int32>>& TransformChildren		 = Collection->Children;
@@ -1388,10 +1395,21 @@ void FNaniteGeometryCollectionSceneProxy::SetDynamicData_RenderThread(FGeometryC
 
 			FPrimitiveInstance& Instance = Instances.Emplace_GetRef();
 
-			Instance.InstanceToLocal			= NewDynamicData->Transforms[TransformIndex];
-			Instance.PrevInstanceToLocal		= NewDynamicData->PrevTransforms[TransformIndex];
+			Instance.InstanceToLocal = NewDynamicData->Transforms[TransformIndex];
+			Instance.LocalToWorld    = Instance.InstanceToLocal * ParentLocalToWorld;
+			
+			if (bCurrentlyInMotion)
+			{
+				Instance.PrevInstanceToLocal = NewDynamicData->PrevTransforms[TransformIndex];
+				Instance.PrevLocalToWorld    = Instance.PrevInstanceToLocal * ParentPrevLocalToWorld;
+			}
+			else
+			{
+				Instance.PrevInstanceToLocal = Instance.InstanceToLocal;
+				Instance.PrevLocalToWorld    = Instance.LocalToWorld;
+			}
+			
 			Instance.LocalToWorld				= Instance.InstanceToLocal * ParentLocalToWorld;
-			Instance.PrevLocalToWorld			= Instance.PrevInstanceToLocal * ParentPrevLocalToWorld;
 			Instance.LocalBounds				= Instance.RenderBounds.TransformBy(Instance.InstanceToLocal);
 			Instance.PrimitiveId				= NaniteData.PrimitiveId;
 			Instance.RenderBounds				= NaniteData.RenderBounds;
@@ -1409,26 +1427,24 @@ void FNaniteGeometryCollectionSceneProxy::SetDynamicData_RenderThread(FGeometryC
 
 void FNaniteGeometryCollectionSceneProxy::ResetPreviousTransforms_RenderThread()
 {
-	if (bLastUpdateWasDynamic)
+	// Reset previous transforms to avoid locked motion vectors
+	for (FPrimitiveInstance& Instance : Instances)
 	{
-		// Reset previous transforms to avoid locked motion vectors
-		for (FPrimitiveInstance& Instance : Instances)
-		{
-			Instance.PrevLocalToWorld = Instance.LocalToWorld;
-		}
-
-		bLastUpdateWasDynamic = false;
+		Instance.PrevLocalToWorld = Instance.LocalToWorld;
 	}
 }
 
 void FNaniteGeometryCollectionSceneProxy::OnMotionBegin()
 {
-	// Currently we don't trigger anything
+	bCurrentlyInMotion = true;
+	bCanSkipRedundantTransformUpdates = false;
 }
 
 void FNaniteGeometryCollectionSceneProxy::OnMotionEnd()
 {
-	// Currently we don't trigger anything
+	bCurrentlyInMotion = false;
+	bCanSkipRedundantTransformUpdates = true;
+	ResetPreviousTransforms_RenderThread();
 }
 
 FGeometryCollectionDynamicDataPool::FGeometryCollectionDynamicDataPool()
diff --git a/Engine/Source/Runtime/Experimental/GeometryCollectionEngine/Private/GeometryCollection/GeometryCollectionSceneProxy.h b/Engine/Source/Runtime/Experimental/GeometryCollectionEngine/Private/GeometryCollection/GeometryCollectionSceneProxy.h
--- a/Engine/Source/Runtime/Experimental/GeometryCollectionEngine/Private/GeometryCollection/GeometryCollectionSceneProxy.h
+++ b/Engine/Source/Runtime/Experimental/GeometryCollectionEngine/Private/GeometryCollection/GeometryCollectionSceneProxy.h
@@ -360,17 +360,6 @@ public:
 protected:
 	TArray<Nanite::FResources*> Resources;
 
-	// TODO: Should probably calculate this on the materials array above instead of on the component
-	//       Null and !Opaque are assigned default material unlike the component material relevance.
-	FMaterialRelevance MaterialRelevance;
-
-	uint32 bCastShadow : 1;
-	uint32 bReverseCulling : 1;
-	uint32 bHasMaterialErrors : 1;
-	uint32 bLastUpdateWasDynamic : 1;
-
-	const UGeometryCollection* GeometryCollection = nullptr;
-
 	struct FGeometryNaniteData
 	{
 		FBoxSphereBounds RenderBounds;
@@ -378,4 +367,15 @@ protected:
 		uint32 PrimitiveId;
 	};
 	TArray<FGeometryNaniteData> GeometryNaniteData;
+
+	const UGeometryCollection* GeometryCollection = nullptr;
+
+	// TODO: Should probably calculate this on the materials array above instead of on the component
+	//       Null and !Opaque are assigned default material unlike the component material relevance.
+	FMaterialRelevance MaterialRelevance;
+
+	uint32 bCastShadow : 1;
+	uint32 bReverseCulling : 1;
+	uint32 bHasMaterialErrors : 1;
+	uint32 bCurrentlyInMotion : 1;
 };
