Refactor TPropertyWatcher back into a single class. There are now two types of constructor: one that takes a custom "!=" test function, and one that uses the built in operator for the type. If the type does not have a "!=" operator available, the caller is forced to use the constructor that specifies a custom test function.

#rnx

diff --git a/Engine/Plugins/Experimental/MeshModelingToolset/Source/MeshModelingTools/Private/BakeMeshAttributeMapsTool.cpp b/Engine/Plugins/Experimental/MeshModelingToolset/Source/MeshModelingTools/Private/BakeMeshAttributeMapsTool.cpp
--- a/Engine/Plugins/Experimental/MeshModelingToolset/Source/MeshModelingTools/Private/BakeMeshAttributeMapsTool.cpp
+++ b/Engine/Plugins/Experimental/MeshModelingToolset/Source/MeshModelingTools/Private/BakeMeshAttributeMapsTool.cpp
@@ -488,8 +488,8 @@ void UBakeMeshAttributeMapsTool::Setup()
 	SetToolPropertySourceEnabled(MultiTextureProps, false);
 
 	auto SetDirtyCallback = [this](TMap<int32, UTexture2D*>) { bInputsDirty = true; };
-	auto EqualsCallback = [](const TMap<int32, UTexture2D*>& A, const TMap<int32, UTexture2D*>& B) -> bool { return A.OrderIndependentCompareEqual(B); };
-	MultiTextureProps->WatchProperty(MultiTextureProps->MaterialIDSourceTextureMap, SetDirtyCallback, EqualsCallback);
+	auto NotEqualsCallback = [](const TMap<int32, UTexture2D*>& A, const TMap<int32, UTexture2D*>& B) -> bool { return !(A.OrderIndependentCompareEqual(B)); };
+	MultiTextureProps->WatchProperty(MultiTextureProps->MaterialIDSourceTextureMap, SetDirtyCallback, NotEqualsCallback);
 	MultiTextureProps->WatchProperty(MultiTextureProps->UVLayer, [this](float) { bInputsDirty = true; });
 
 	VisualizationProps = NewObject<UBakedOcclusionMapVisualizationProperties>(this);
diff --git a/Engine/Source/Runtime/Experimental/InteractiveToolsFramework/Public/InteractiveTool.h b/Engine/Source/Runtime/Experimental/InteractiveToolsFramework/Public/InteractiveTool.h
--- a/Engine/Source/Runtime/Experimental/InteractiveToolsFramework/Public/InteractiveTool.h
+++ b/Engine/Source/Runtime/Experimental/InteractiveToolsFramework/Public/InteractiveTool.h
@@ -8,6 +8,9 @@
 #include "Shader.h"
 #include "ToolContextInterfaces.h"
 #include "UObject/UObjectGlobals.h"
+#include "Templates/EnableIf.h"
+#include "Templates/Models.h"
+
 #include "InteractiveTool.generated.h"
 
 class UInteractiveToolManager;
@@ -49,110 +52,88 @@ public:
 	};
 
 	template <typename PropType>
-	class TPropertyWatcherBase : public FPropertyWatcher
+	class TPropertyWatcher : public FPropertyWatcher
 	{
 	public:
 		using FValueGetter = TFunction<PropType(void)>;
 		using FChangedCallback = TFunction<void(const PropType&)>;
+		using FNotEqualTestFunction = TFunction<bool(const PropType&, const PropType&)>;	// Define "!=" for PropType
 
-		TPropertyWatcherBase(const PropType& Property,
-							 FChangedCallback OnChangedIn)
-			: GetValue([&Property]() {return Property; }), 
-			  OnChanged(MoveTemp(OnChangedIn))
+		/**
+		 * Describes a type having a "!=" comparasion operator. (Note: it would probably be cleaner to require operator==,
+		 * however some types have operator!= but no operator==, and we do only use != in the code.)
+		*/
+		struct CInequalityComparable
+		{
+			template <typename T>
+			auto Requires(bool& Result, const T& A, const T& B) -> decltype(
+				Result = A != B
+				);
+		};
+
+		// If PropType is CInequalityComparable, allow the two-argument constructor and default to using the type's 
+		// existing != operator.
+		// If PropType is not CInequalityComparable, the caller must use the three-argument constructor to provide a
+		// "not equal" function.
+
+		// Two-argument constructor, enabled only if PropType is CInequalityComparable
+		template<typename Q = PropType, typename = typename TEnableIf<TModels<CInequalityComparable, Q>::Value>::Type>
+		TPropertyWatcher(const PropType& Property,
+						 FChangedCallback OnChangedIn)
+			: Cached(),
+			GetValue([&Property]() {return Property; }),
+			OnChanged(MoveTemp(OnChangedIn)),
+			NotEqual([](const PropType& A, const PropType& B) { return A != B; })
 		{}
 
-		TPropertyWatcherBase(FValueGetter GetValueIn,
-							 FChangedCallback OnChangedIn)
-			: GetValue(MoveTemp(GetValueIn)), 
-			  OnChanged(MoveTemp(OnChangedIn))
+		// Two-argument constructor, enabled only if PropType is CInequalityComparable
+		template<typename Q = PropType, typename = typename TEnableIf<TModels<CInequalityComparable, Q>::Value>::Type>
+		TPropertyWatcher(FValueGetter GetValueIn,
+						 FChangedCallback OnChangedIn)
+			: Cached(),
+			GetValue(MoveTemp(GetValueIn)),
+			OnChanged(MoveTemp(OnChangedIn)),
+			NotEqual([](const PropType& A, const PropType& B) { return A != B; })
 		{}
-
-		void SilentUpdate() final
-		{
-			Cached = GetValue();
-		}
-
-	protected:
-		TOptional<PropType> Cached;
-		FValueGetter GetValue;
-		FChangedCallback OnChanged;
-	};
-
-
-	template <typename PropType>
-	class TPropertyWatcher : public TPropertyWatcherBase<PropType>
-	{
-	public:
-		using typename TPropertyWatcherBase<PropType>::FChangedCallback;
-		using typename TPropertyWatcherBase<PropType>::FValueGetter;
-		using TPropertyWatcherBase<PropType>::Cached;
-		using TPropertyWatcherBase<PropType>::GetValue;
-		using TPropertyWatcherBase<PropType>::OnChanged;
-
-
+	
 		TPropertyWatcher(const PropType& Property,
-						 FChangedCallback OnChangedIn) 
-			: TPropertyWatcherBase<PropType>(Property, OnChangedIn)
+						 FChangedCallback OnChangedIn,
+						 FNotEqualTestFunction NotEqualIn)
+			: Cached(),
+			GetValue([&Property]() {return Property; }),
+			OnChanged(MoveTemp(OnChangedIn)),
+			NotEqual(MoveTemp(NotEqualIn))
 		{}
 
 		TPropertyWatcher(FValueGetter GetValueIn,
-						 FChangedCallback OnChangedIn) 
-			: TPropertyWatcherBase<PropType>(GetValueIn, OnChangedIn)
+						 FChangedCallback OnChangedIn,
+						 FNotEqualTestFunction NotEqualIn)
+			: Cached(),
+			GetValue(MoveTemp(GetValueIn)),
+			OnChanged(MoveTemp(OnChangedIn)),
+			NotEqual(MoveTemp(NotEqualIn))
 		{}
 
 		void CheckAndUpdate() final
 		{
 			PropType Value = GetValue();
-			if ((!Cached.IsSet()) || (Cached.GetValue() != Value))
+			if ((!Cached.IsSet()) || NotEqual(Cached.GetValue(), Value))
 			{
 				Cached = Value;
 				OnChanged(Cached.GetValue());
 			}
 		}
-	};
-
-	/// Property watcher for types that don't have an equality operator available.
-	/// 
-	/// Note: TPropertyWatcherWithEqualityTest can't inherit directly from TPropertyWatcher, since the presence of 
-	/// the == operator in that template class will cause compile errors for types without the operator.
-	template <typename PropType>
-	class TPropertyWatcherWithEqualityTest : public TPropertyWatcherBase<PropType>
-	{
-	public:
-		using FEqualityCheckFunction = TFunction<bool(const PropType&, const PropType&)>;
-
-		using typename TPropertyWatcherBase<PropType>::FChangedCallback;
-		using typename TPropertyWatcherBase<PropType>::FValueGetter;
-		using TPropertyWatcherBase<PropType>::Cached;
-		using TPropertyWatcherBase<PropType>::GetValue;
-		using TPropertyWatcherBase<PropType>::OnChanged;
-
-		TPropertyWatcherWithEqualityTest(const PropType& Property,
-										 FChangedCallback OnChangedIn,
-										 FEqualityCheckFunction EqualsIn)
-			: TPropertyWatcherBase<PropType>(Property, OnChangedIn),
-			  Equals(EqualsIn)
-		{}
-
-		TPropertyWatcherWithEqualityTest(FValueGetter GetValueIn,
-										 FChangedCallback OnChangedIn,
-										 FEqualityCheckFunction EqualsIn)
-			: TPropertyWatcherBase<PropType>(GetValueIn, OnChangedIn),
-			  Equals(EqualsIn)
-		{}
 
-		void CheckAndUpdate() final
+		void SilentUpdate() final
 		{
-			PropType Value = this->GetValue();
-			if ((!this->Cached.IsSet()) || !Equals(this->Cached.GetValue(), Value))
-			{
-				this->Cached = Value;
-				this->OnChanged(this->Cached.GetValue());
-			}
+			Cached = GetValue();
 		}
 
-	protected:
-		FEqualityCheckFunction Equals;
+	private:
+		TOptional<PropType> Cached;
+		FValueGetter GetValue;
+		FChangedCallback OnChanged;
+		FNotEqualTestFunction NotEqual;
 	};
 
 	FWatchablePropertySet() = default;
@@ -202,18 +183,18 @@ public:
 	/** @return Index of the watcher, which can be used in SilentUpdateWatcherAtIndex */
 	template <typename PropType>
 	int32 WatchProperty(const PropType& ValueIn,
-						typename TPropertyWatcherWithEqualityTest<PropType>::FChangedCallback OnChangedIn,
-						typename TPropertyWatcherWithEqualityTest<PropType>::FEqualityCheckFunction EqualsIn)
+						typename TPropertyWatcher<PropType>::FChangedCallback OnChangedIn,
+						typename TPropertyWatcher<PropType>::FNotEqualTestFunction NotEqualsIn)
 	{
-		return PropertyWatchers.Emplace(MakeUnique<TPropertyWatcherWithEqualityTest<PropType>>(ValueIn, OnChangedIn, EqualsIn));
+		return PropertyWatchers.Emplace(MakeUnique<TPropertyWatcher<PropType>>(ValueIn, OnChangedIn, NotEqualsIn));
 	}
 	/** @return Index of the watcher, which can be used in SilentUpdateWatcherAtIndex */
 	template <typename PropType>
-	int32 WatchProperty(typename TPropertyWatcherWithEqualityTest<PropType>::FValueGetter GetValueIn,
-					  typename TPropertyWatcherWithEqualityTest<PropType>::FChangedCallback OnChangedIn,
-					  typename TPropertyWatcherWithEqualityTest<PropType>::FEqualityCheckFunction EqualsIn)
+	int32 WatchProperty(typename TPropertyWatcher<PropType>::FValueGetter GetValueIn,
+					  typename TPropertyWatcher<PropType>::FChangedCallback OnChangedIn,
+					  typename TPropertyWatcher<PropType>::FNotEqualTestFunction NotEqualsIn)
 	{
-		return PropertyWatchers.Emplace(MakeUnique<TPropertyWatcherWithEqualityTest<PropType>>(GetValueIn, OnChangedIn, EqualsIn));
+		return PropertyWatchers.Emplace(MakeUnique<TPropertyWatcher<PropType>>(GetValueIn, OnChangedIn, NotEqualsIn));
 	}
 
 private:
