GeometryCollectionSceneProxy
-- OnMotionEnd now resets instance transforms as well as component transforms.

diff --git a/Engine/Source/Runtime/Experimental/GeometryCollectionEngine/Private/GeometryCollection/GeometryCollectionSceneProxy.cpp b/Engine/Source/Runtime/Experimental/GeometryCollectionEngine/Private/GeometryCollection/GeometryCollectionSceneProxy.cpp
--- a/Engine/Source/Runtime/Experimental/GeometryCollectionEngine/Private/GeometryCollection/GeometryCollectionSceneProxy.cpp
+++ b/Engine/Source/Runtime/Experimental/GeometryCollectionEngine/Private/GeometryCollection/GeometryCollectionSceneProxy.cpp
@@ -1431,6 +1431,7 @@ void FNaniteGeometryCollectionSceneProxy::ResetPreviousTransforms_RenderThread()
 	for (FPrimitiveInstance& Instance : Instances)
 	{
 		Instance.PrevLocalToWorld = Instance.LocalToWorld;
+		Instance.PrevInstanceToLocal = Instance.InstanceToLocal;
 	}
 }
 
