Interactive tools context now has a store for arbitrary context objects, which can be accessed via tools for shared data or operations, like typed element actions.

diff --git a/Engine/Source/Runtime/Experimental/InteractiveToolsFramework/Private/ContextObjectStore.cpp b/Engine/Source/Runtime/Experimental/InteractiveToolsFramework/Private/ContextObjectStore.cpp
new file mode 100644
--- /dev/null
+++ b/Engine/Source/Runtime/Experimental/InteractiveToolsFramework/Private/ContextObjectStore.cpp
@@ -0,0 +1,44 @@
+// Copyright Epic Games, Inc. All Rights Reserved.
+
+#include "ContextObjectStore.h"
+
+UObject* UContextObjectStore::FindContextByClass(UClass* InClass) const
+{
+	for (UObject* ContextObject : ContextObjects)
+	{
+		if (ContextObject && ContextObject->IsA(InClass))
+		{
+			return ContextObject;
+		}
+	}
+
+	return nullptr;
+}
+
+bool UContextObjectStore::AddContextObject(UObject* InContextObject)
+{
+	if (InContextObject)
+	{
+		return ContextObjects.AddUnique(InContextObject) != INDEX_NONE;
+	}
+
+	return false;
+}
+
+bool UContextObjectStore::RemoveContextObject(UObject* InContextObject)
+{
+	return ContextObjects.RemoveSingle(InContextObject) != 0;
+}
+
+bool UContextObjectStore::RemoveContextObjectsOfType(const UClass* InClass)
+{
+	return (ContextObjects.RemoveAll([InClass](UObject* InObject)
+	{
+		return InObject->IsA(InClass);
+	}) > 0);
+}
+
+void UContextObjectStore::Shutdown()
+{
+	ContextObjects.Empty();
+}
diff --git a/Engine/Source/Runtime/Experimental/InteractiveToolsFramework/Private/InteractiveToolManager.cpp b/Engine/Source/Runtime/Experimental/InteractiveToolsFramework/Private/InteractiveToolManager.cpp
--- a/Engine/Source/Runtime/Experimental/InteractiveToolsFramework/Private/InteractiveToolManager.cpp
+++ b/Engine/Source/Runtime/Experimental/InteractiveToolsFramework/Private/InteractiveToolManager.cpp
@@ -5,6 +5,7 @@
 #include "Engine/Engine.h"
 #include "InteractiveToolsContext.h"
 #include "InteractiveToolsSelectionStoreSubsystem.h"
+#include "ContextObjectStore.h"
 
 #define LOCTEXT_NAMESPACE "UInteractiveToolManager"
 
@@ -367,6 +368,11 @@ UInteractiveGizmoManager* UInteractiveToolManager::GetPairedGizmoManager()
 	return Cast<UInteractiveToolsContext>(GetOuter())->GizmoManager;
 }
 
+UContextObjectStore* UInteractiveToolManager::GetContextObjectStore() const
+{
+	return Cast<UInteractiveToolsContext>(GetOuter())->ContextObjectStore;
+}
+
 void UInteractiveToolManager::DisplayMessage(const FText& Message, EToolMessageLevel Level)
 {
 	TransactionsAPI->DisplayMessage(Message, Level);
diff --git a/Engine/Source/Runtime/Experimental/InteractiveToolsFramework/Private/InteractiveToolsContext.cpp b/Engine/Source/Runtime/Experimental/InteractiveToolsFramework/Private/InteractiveToolsContext.cpp
--- a/Engine/Source/Runtime/Experimental/InteractiveToolsFramework/Private/InteractiveToolsContext.cpp
+++ b/Engine/Source/Runtime/Experimental/InteractiveToolsFramework/Private/InteractiveToolsContext.cpp
@@ -3,6 +3,7 @@
 
 #include "InteractiveToolsContext.h"
 #include "ToolTargetManager.h"
+#include "ContextObjectStore.h"
 
 UInteractiveToolsContext::UInteractiveToolsContext()
 {
@@ -27,6 +28,8 @@ void UInteractiveToolsContext::Initialize(IToolsContextQueriesAPI* QueriesAPI, I
 	GizmoManager->Initialize(QueriesAPI, TransactionsAPI, InputRouter);
 
 	GizmoManager->RegisterDefaultGizmos();
+
+	ContextObjectStore = NewObject<UContextObjectStore>(this);
 }
 
 
@@ -42,6 +45,9 @@ void UInteractiveToolsContext::Shutdown()
 
 	ToolManager->Shutdown();
 	ToolManager = nullptr;
+
+	ContextObjectStore->Shutdown();
+	ContextObjectStore = nullptr;
 }
 
 void UInteractiveToolsContext::DeactivateActiveTool(EToolSide WhichSide, EToolShutdownType ShutdownType)
diff --git a/Engine/Source/Runtime/Experimental/InteractiveToolsFramework/Public/ContextObjectStore.h b/Engine/Source/Runtime/Experimental/InteractiveToolsFramework/Public/ContextObjectStore.h
new file mode 100644
--- /dev/null
+++ b/Engine/Source/Runtime/Experimental/InteractiveToolsFramework/Public/ContextObjectStore.h
@@ -0,0 +1,78 @@
+// Copyright Epic Games, Inc. All Rights Reserved.
+
+#pragma once
+
+#include "UObject/Object.h"
+
+#include "ContextObjectStore.generated.h"
+
+/**
+ * A context object store allows tools to get access to arbitrary objects which expose data or APIs to enable additional functionality.
+  * Some example use cases of context objects: 
+  *   - A tool builder may disallow a particular tool if a needed API object is not present in the context store.
+  *   - A tool may allow extra actions if it has access to a particular API object in the context store.
+  *   - A tool may choose to initialize itself differently based on the presence of a selection-holding data object in the context store.
+ */
+UCLASS(Transient)
+class INTERACTIVETOOLSFRAMEWORK_API UContextObjectStore : public UObject
+{
+	GENERATED_BODY()
+
+public:
+	/**
+	 * Finds the first the context object of the given type. Can return a subclass of the given type.
+	 * @returns the found context object, casted to the given type, or nullptr if none matches.
+	 */
+	template <typename TObjectType>
+	TObjectType* FindContext() const
+	{
+		for (UObject* ContextObject : ContextObjects)
+		{
+			if (TObjectType* CastedObject = Cast<TObjectType>(ContextObject))
+			{
+				return CastedObject;
+			}
+		}
+
+		return nullptr;
+	}
+
+	/**
+	 * Finds the first context object that derives from the given class.
+	 * @returns the found context object, or nullptr if none matches.
+	 */
+	UObject* FindContextByClass(UClass* InClass) const;
+
+	/**
+	 * Adds a data object to the tool manager's set of shared data objects.
+	 * @returns true if the addition is successful.
+	 */
+	bool AddContextObject(UObject* InContextObject);
+
+	/**
+	 * Removes a data object from the tool manager's set of shared data objects.
+	 * @returns true if the removal is successful.
+	 */
+	bool RemoveContextObject(UObject* InContextObject);
+
+	/**
+	 * Removes any data objects from the tool manager's set of shared data objects that are of type @param InClass
+	 * @returns true if any objects are removed.
+	 */
+	bool RemoveContextObjectsOfType(const UClass* InClass);
+
+	template <class TObjectType>
+	bool RemoveContextObjectsOfType()
+	{
+		return RemoveToolDataObjectsOfType(TObjectType::StaticClass());
+	}
+
+	/**
+	 * Shuts down the context object store, releasing hold on any stored content objects.
+	 */
+	void Shutdown();
+
+protected:
+	UPROPERTY()
+	TArray<TObjectPtr<UObject>> ContextObjects;
+};
diff --git a/Engine/Source/Runtime/Experimental/InteractiveToolsFramework/Public/InteractiveToolManager.h b/Engine/Source/Runtime/Experimental/InteractiveToolsFramework/Public/InteractiveToolManager.h
--- a/Engine/Source/Runtime/Experimental/InteractiveToolsFramework/Public/InteractiveToolManager.h
+++ b/Engine/Source/Runtime/Experimental/InteractiveToolsFramework/Public/InteractiveToolManager.h
@@ -12,6 +12,7 @@
 #include "InteractiveToolManager.generated.h"
 
 class UInteractiveToolStorableSelection;
+class UContextObjectStore;
 
 /** A Tool can be activated on a particular input device, currently identified by a "side" */
 UENUM()
@@ -263,6 +264,10 @@ public:
 
 	UInteractiveGizmoManager* GetPairedGizmoManager();
 
+	/**
+	 * @return the context object store from the owning tools context.
+	 */
+	UContextObjectStore* GetContextObjectStore() const;
 
 public:
 	/** Currently-active Left Tool, or null if no Tool is active */
diff --git a/Engine/Source/Runtime/Experimental/InteractiveToolsFramework/Public/InteractiveToolsContext.h b/Engine/Source/Runtime/Experimental/InteractiveToolsFramework/Public/InteractiveToolsContext.h
--- a/Engine/Source/Runtime/Experimental/InteractiveToolsFramework/Public/InteractiveToolsContext.h
+++ b/Engine/Source/Runtime/Experimental/InteractiveToolsFramework/Public/InteractiveToolsContext.h
@@ -8,6 +8,7 @@
 #include "InteractiveToolsContext.generated.h"
 
 class UToolTargetManager;
+class UContextObjectStore;
 
 /**
  * InteractiveToolsContext owns a ToolManager and an InputRouter. This is just a top-level 
@@ -72,6 +73,13 @@ public:
 	UPROPERTY()
 	TObjectPtr<UInteractiveGizmoManager> GizmoManager;
 
+	/** 
+	 * Current Context Object Store for this Context.
+	 * Stores arbitrary objects which share data or expose APIs across interactive tools and managers belonging to this context.
+	 */
+	UPROPERTY()
+	TObjectPtr<UContextObjectStore> ContextObjectStore;
+
 protected:
 	UPROPERTY()
 	TSoftClassPtr<UInteractiveToolManager> ToolManagerClass;
