GeometryCollection : Bounds
- Added a debugging flag to enable single threaded bounds.

diff --git a/Engine/Source/Runtime/Experimental/GeometryCollectionEngine/Private/GeometryCollection/GeometryCollectionComponent.cpp b/Engine/Source/Runtime/Experimental/GeometryCollectionEngine/Private/GeometryCollection/GeometryCollectionComponent.cpp
--- a/Engine/Source/Runtime/Experimental/GeometryCollectionEngine/Private/GeometryCollection/GeometryCollectionComponent.cpp
+++ b/Engine/Source/Runtime/Experimental/GeometryCollectionEngine/Private/GeometryCollection/GeometryCollectionComponent.cpp
@@ -184,6 +184,10 @@ FAutoConsoleVariableRef CVarGeometryCollectionNanite(
 float GGeometryCollectionNavigationSizeThreshold = 20.0f;
 FAutoConsoleVariableRef CVarGeometryCollectionNavigationSizeThreshold(TEXT("p.GeometryCollectionNavigationSizeThreshold"), GGeometryCollectionNavigationSizeThreshold, TEXT("Size in CM used as a threshold for whether a geometry in the collection is collected and exported for navigation purposes. Measured as the diagonal of the leaf node bounds."));
 
+// Single-Threaded Bounds
+bool bGeometryCollectionSingleThreadedBoundsCalculation = false;
+FAutoConsoleVariableRef CVarGeometryCollectionSingleThreadedBoundsCalculation(TEXT("p.GeometryCollectionSingleThreadedBoundsCalculation"), bGeometryCollectionSingleThreadedBoundsCalculation, TEXT("[Debug Only] Single threaded bounds calculation. [def:false]"));
+
 
 
 FGeomComponentCacheParameters::FGeomComponentCacheParameters()
@@ -444,6 +448,19 @@ FBoxSphereBounds UGeometryCollectionComponent::CalcBounds(const FTransform& Loca
 				}
 			}
 		}
+		else if (bGeometryCollectionSingleThreadedBoundsCalculation)
+		{
+			CHAOS_ENSURE(false); // this is slower and only enabled through a pvar debugging, disable bGeometryCollectionSingleThreadedBoundsCalculation in a production environment. 
+			for (int32 BoxIdx = 0; BoxIdx < NumBoxes; ++BoxIdx)
+			{
+				const int32 TransformIndex = TransformIndices[BoxIdx];
+
+				if (RestCollection->GetGeometryCollection()->IsGeometry(TransformIndex))
+				{
+					BoundingBox += BoundingBoxes[BoxIdx].TransformBy(GlobalMatrices[TransformIndex] * LocalToWorldWithScale);
+				}
+			}
+		}
 		else
 		{
 #if INTEL_ISPC
