Fix bad merge resolve which broke rewind unit test

diff --git a/Engine/Source/Programs/HeadlessChaos/Private/HeadlessChaosTestRewind.cpp b/Engine/Source/Programs/HeadlessChaos/Private/HeadlessChaosTestRewind.cpp
--- a/Engine/Source/Programs/HeadlessChaos/Private/HeadlessChaosTestRewind.cpp
+++ b/Engine/Source/Programs/HeadlessChaos/Private/HeadlessChaosTestRewind.cpp
@@ -402,7 +402,7 @@ namespace ChaosTest {
 			});
 	}
 
-	GTEST_TEST(AllTraits, DISABLED_RewindTest_ResimSleepChange)
+	GTEST_TEST(AllTraits, RewindTest_ResimSleepChange)
 	{
 		//change object state on physics thread, and make sure the state change is properly recorded in rewind data
 		TRewindHelper::TestDynamicSphere([](auto* Solver, FReal SimDt, int32 Optimization, auto Proxy, auto Sphere)
diff --git a/Engine/Source/Runtime/Experimental/Chaos/Public/Chaos/PBDRigidsEvolution.h b/Engine/Source/Runtime/Experimental/Chaos/Public/Chaos/PBDRigidsEvolution.h
--- a/Engine/Source/Runtime/Experimental/Chaos/Public/Chaos/PBDRigidsEvolution.h
+++ b/Engine/Source/Runtime/Experimental/Chaos/Public/Chaos/PBDRigidsEvolution.h
@@ -427,14 +427,17 @@ public:
 			Particles.SetDynamicParticleSOA(Particle);
 		}
 
-		if (InitialState == EObjectStateType::Sleeping && InitialState != ObjectState && !Particle->Disabled())
+		if (InitialState != ObjectState && !Particle->Disabled())
 		{
-			if (Particle->Island() != INDEX_NONE)
+			if (InitialState == EObjectStateType::Sleeping)
 			{
-				// GT has forced a wake so have to wake everything in the island
-				IslandsToWake.Enqueue(Particle->Island());
+				if (Particle->Island() != INDEX_NONE)
+				{
+					// GT has forced a wake so have to wake everything in the island
+					IslandsToWake.Enqueue(Particle->Island());
+				}
 			}
-			else if(ObjectState != EObjectStateType::Dynamic)
+			else if (ObjectState != EObjectStateType::Dynamic)
 			{
 				// even though we went to sleep, we should still report info back to GT
 				Particles.MarkTransientDirtyParticle(Particle);
