Temporary simplification of chaos archiver to always serialize Chaos::TVectors as floats.

diff --git a/Engine/Source/Runtime/Experimental/Chaos/Public/Chaos/ChaosArchive.h b/Engine/Source/Runtime/Experimental/Chaos/Public/Chaos/ChaosArchive.h
--- a/Engine/Source/Runtime/Experimental/Chaos/Public/Chaos/ChaosArchive.h
+++ b/Engine/Source/Runtime/Experimental/Chaos/Public/Chaos/ChaosArchive.h
@@ -353,17 +353,18 @@ FChaosArchive& operator<<(FChaosArchive& Ar, TArray<T, TAllocator>& Array)
 
 FORCEINLINE FChaosArchive& operator<<(FChaosArchive& Ar, Chaos::FReal& Real)
 {
+	// LWC_TODO: Serializer
 	// we need to check if we are storing doubles or floats
-	if (!Ar.IsPersistent())
-	{
-		// normal umodified type path 
-		operator<<((FArchive&)Ar, Real);
-	}
-	else
+	//if (!Ar.IsPersistent())
+	//{
+	//	// normal umodified type path 
+	//	operator<<((FArchive&)Ar, Real);
+	//}
+	//else
 	{
 		// in that case data is stored as float and we need to read it as such		
-		ensure(Ar.IsLoading()); // this case should normally only happening when reading 
-		FRealSingle RealSingle; 
+		//ensure(Ar.IsLoading()); // this case should normally only happening when reading 
+		FRealSingle RealSingle = (FRealSingle)Real; 
 		operator<<((FArchive&)Ar, RealSingle);
 		Real = (FReal)RealSingle;
 	}
diff --git a/Engine/Source/Runtime/Experimental/ChaosCore/Public/Chaos/Core.h b/Engine/Source/Runtime/Experimental/ChaosCore/Public/Chaos/Core.h
--- a/Engine/Source/Runtime/Experimental/ChaosCore/Public/Chaos/Core.h
+++ b/Engine/Source/Runtime/Experimental/ChaosCore/Public/Chaos/Core.h
@@ -35,5 +35,5 @@ namespace Chaos
 	// NOTE: if you get a merge conflict on the GUID, you must replace it with a new GUID - do not accept the source or target
 	// or you will likely get DDC version conflicts resulting in crashes during load.
 	// Core version string for Chaos data. Any DDC builder dependent on Chaos for serialization should depend on this version
-	static const FString ChaosVersionString = TEXT("D58C401F-A336-4428-9ADE-304E2B348968");
+	static const FString ChaosVersionString = TEXT("6315F0B9C62B42B1A7929463A7375317");
 }
diff --git a/Engine/Source/Runtime/Experimental/ChaosCore/Public/Chaos/Vector.h b/Engine/Source/Runtime/Experimental/ChaosCore/Public/Chaos/Vector.h
--- a/Engine/Source/Runtime/Experimental/ChaosCore/Public/Chaos/Vector.h
+++ b/Engine/Source/Runtime/Experimental/ChaosCore/Public/Chaos/Vector.h
@@ -948,25 +948,26 @@ namespace Chaos
 	template<typename T, int d>
 	FArchive& operator<<(FArchive& Ar, TVector<T, d>& ValueIn)
 	{
-		constexpr bool bIsTypeReal = (TAreTypesEqual<FReal, T>::Value == true);
-		if(!Ar.IsPersistent())
-		{
-			// unchanged type code path 
-			for (int32 Idx = 0; Idx < d; ++Idx)
-			{
-				Ar << ValueIn[Idx];
-			}
-
-		}
-		else 
+		// LWC_TODO: Serializer
+		//constexpr bool bIsTypeReal = (TAreTypesEqual<FReal, T>::Value == true);
+		//if(!Ar.IsPersistent())
+		//{
+		//	// unchanged type code path 
+		//	for (int32 Idx = 0; Idx < d; ++Idx)
+		//	{
+		//		Ar << ValueIn[Idx];
+		//	}
+
+		//}
+		//else
 		{
 			// in that case data is stored as float and we need to read it as such		
-			ensure(Ar.IsLoading()); // this case should normally only happening when reading 
+			//ensure(Ar.IsLoading()); // this case should normally only happening when reading 
 			for (int32 Idx = 0; Idx < d; ++Idx)
 			{
-				FRealSingle RealSingle;
+				FRealSingle RealSingle = (FRealSingle)ValueIn[Idx];
 				Ar << RealSingle;
-				ValueIn[Idx] = RealSingle;
+				ValueIn[Idx] = (T)RealSingle;
 			}
 		}
 		return Ar;
