fix crash on trying to expand selection to neighbors when a transform w/ no geometry is selected

#rnx

diff --git a/Engine/Source/Runtime/Experimental/GeometryCollectionEngine/Private/GeometryCollection/GeometryCollectionComponent.cpp b/Engine/Source/Runtime/Experimental/GeometryCollectionEngine/Private/GeometryCollection/GeometryCollectionComponent.cpp
--- a/Engine/Source/Runtime/Experimental/GeometryCollectionEngine/Private/GeometryCollection/GeometryCollectionComponent.cpp
+++ b/Engine/Source/Runtime/Experimental/GeometryCollectionEngine/Private/GeometryCollection/GeometryCollectionComponent.cpp
@@ -2375,10 +2375,14 @@ void FScopedColorEdit::SelectBones(GeometryCollection::ESelectionMode SelectionM
 				for (int32 Bone : SelectedBones)
 				{
 					NewSelection.AddUnique(Bone);
-					const TSet<int32> &Neighbors = Proximity[TransformToGeometryIndex[Bone]];
-					for (int32 NeighborGeometryIndex : Neighbors)
+					int32 GeometryIdx = TransformToGeometryIndex[Bone];
+					if (GeometryIdx != INDEX_NONE)
 					{
-						NewSelection.AddUnique(TransformIndex[NeighborGeometryIndex]);
+						const TSet<int32>& Neighbors = Proximity[GeometryIdx];
+						for (int32 NeighborGeometryIndex : Neighbors)
+						{
+							NewSelection.AddUnique(TransformIndex[NeighborGeometryIndex]);
+						}
 					}
 				}
 
