Fix Chaos TManagedArray passing an on-stack variable to TCheckedPointerIterator, triggering Address Sanitizer errors.

diff --git a/Engine/Source/Runtime/Experimental/Chaos/Public/GeometryCollection/ManagedArray.h b/Engine/Source/Runtime/Experimental/Chaos/Public/GeometryCollection/ManagedArray.h
--- a/Engine/Source/Runtime/Experimental/Chaos/Public/GeometryCollection/ManagedArray.h
+++ b/Engine/Source/Runtime/Experimental/Chaos/Public/GeometryCollection/ManagedArray.h
@@ -423,37 +423,19 @@ public:
 		}
 	}
 
-	// @todo Add RangedFor support. 
-
-
-	// TARRAY_RANGED_FOR_CHECKS Is defined in Array.h based on build state.
-#if TARRAY_RANGED_FOR_CHECKS
-	// @todo: What is the appropriate size type?
-	typedef TCheckedPointerIterator<      ElementType, int32> RangedForIteratorType;
-	typedef TCheckedPointerIterator<const ElementType, int32> RangedForConstIteratorType;
-#else
-	typedef       ElementType* RangedForIteratorType;
-	typedef const ElementType* RangedForConstIteratorType;
-#endif
-
-private:
+	typedef typename TArray<InElementType>::RangedForIteratorType		RangedForIteratorType;
+	typedef typename TArray<InElementType>::RangedForConstIteratorType	RangedForConstIteratorType;
 
 	/**
 	* DO NOT USE DIRECTLY
 	* STL-like iterators to enable range-based for loop support.
 	*/
-#if TARRAY_RANGED_FOR_CHECKS
-	FORCEINLINE friend RangedForIteratorType      begin(      TManagedArrayBase& ManagedArray) { return RangedForIteratorType     (ManagedArray.Num(), ManagedArray.GetData()); }
-	FORCEINLINE friend RangedForConstIteratorType begin(const TManagedArrayBase& ManagedArray) { return RangedForConstIteratorType(ManagedArray.Num(), ManagedArray.GetData()); }
-	FORCEINLINE friend RangedForIteratorType      end  (      TManagedArrayBase& ManagedArray) { return RangedForIteratorType     (ManagedArray.Num(), ManagedArray.GetData() + ManagedArray.Num()); }
-	FORCEINLINE friend RangedForConstIteratorType end  (const TManagedArrayBase& ManagedArray) { return RangedForConstIteratorType(ManagedArray.Num(), ManagedArray.GetData() + ManagedArray.Num()); }
-#else
-	FORCEINLINE friend RangedForIteratorType      begin(      TManagedArrayBase& ManagedArray) { return ManagedArray.GetData(); }
-	FORCEINLINE friend RangedForConstIteratorType begin(const TManagedArrayBase& ManagedArray) { return ManagedArray.GetData(); }
-	FORCEINLINE friend RangedForIteratorType      end  (      TManagedArrayBase& ManagedArray) { return ManagedArray.GetData() + ManagedArray.Num(); }
-	FORCEINLINE friend RangedForConstIteratorType end  (const TManagedArrayBase& ManagedArray) { return ManagedArray.GetData() + ManagedArray.Num(); }
-#endif
+	FORCEINLINE RangedForIteratorType      begin()			{ return Array.begin(); }
+	FORCEINLINE RangedForConstIteratorType begin() const	{ return Array.begin(); }
+	FORCEINLINE RangedForIteratorType      end  ()			{ return Array.end(); }
+	FORCEINLINE RangedForConstIteratorType end  () const	{ return Array.end(); }
 
+private:
 	/**
 	* Protected Resize to prevent external resizing of the array
 	*
