Fix plasticity-related issue:
- Move ApplyPlasticityLimit to UnprepareIteration
- Use joint connector XR instead of particle XR for plasticity computations
- Modify LinearDrivePositionTarget instead of JointFrames to achieve linear plasticity
- Add warnings to protect against invalid configurations

diff --git a/Engine/Source/Runtime/Experimental/Chaos/Private/Chaos/PBDJointConstraints.cpp b/Engine/Source/Runtime/Experimental/Chaos/Private/Chaos/PBDJointConstraints.cpp
--- a/Engine/Source/Runtime/Experimental/Chaos/Private/Chaos/PBDJointConstraints.cpp
+++ b/Engine/Source/Runtime/Experimental/Chaos/Private/Chaos/PBDJointConstraints.cpp
@@ -812,6 +812,9 @@ namespace Chaos
 					JointState.LinearImpulse = FVec3(0);
 					JointState.AngularImpulse = FVec3(0);
 				}
+
+				// Plasticity should be applied regardless batching is used or not. But since bChaos_Joint_Batching is always false, so it is ok to put the plasticity code in this block
+				ApplyPlasticityLimits(JointIndex);
 			}
 		}
 	}
@@ -1336,11 +1339,6 @@ namespace Chaos
 			ApplyBreakThreshold(Dt, ConstraintIndex, Solver.GetNetLinearImpulse(), Solver.GetNetAngularImpulse());
 		}
 
-		if ((JointSettings.LinearPlasticityLimit != FLT_MAX) || (JointSettings.AngularPlasticityLimit != FLT_MAX))
-		{
-			ApplyPlasticityLimits(Dt, ConstraintIndex, Particle1->X() - Particle0->X(), Particle0->R().Inverse() * Particle1->R());
-		}
-
 		return Solver.GetIsActive() || !bChaos_Joint_EarlyOut_Enabled;
 	}
 
@@ -1439,30 +1437,85 @@ namespace Chaos
 	}
 
 
-	void FPBDJointConstraints::ApplyPlasticityLimits(const FReal Dt, int32 ConstraintIndex, const FVec3& LinearDisplacement, const FRotation3& AngularDisplacement)
+	void FPBDJointConstraints::ApplyPlasticityLimits(const int32 ConstraintIndex)
 	{
 		FPBDJointSettings& JointSettings = ConstraintSettings[ConstraintIndex];
-		FTransformPair& ConstraintFrame = ConstraintFrames[ConstraintIndex];
+		const bool bHasLinearPlasticityLimit = JointSettings.LinearPlasticityLimit != FLT_MAX;
+		const bool bHasAngularPlasticityLimit = JointSettings.AngularPlasticityLimit != FLT_MAX;
+		const bool bHasPlasticityLimits =  bHasLinearPlasticityLimit || bHasAngularPlasticityLimit;
+		if(!bHasPlasticityLimits)
+		{
+			return;
+		}
+
+		if(!Settings.bEnableDrives)
+		{
+			UE_LOG(LogChaosJoint, Warning, TEXT("Linear or angular drive needs to be enabled in order to use plasticity."));
+			return;
+		}
 
-		if (!FMath::IsNearlyEqual(JointSettings.LinearPlasticityLimit, (FReal)FLT_MAX))
+		bool bApplyLinearPlasticityLimits = bHasLinearPlasticityLimit;
+		if(bHasLinearPlasticityLimit)
 		{
-			FTransform JointTransform = ConstraintFrame[0].GetRelativeTransform(ConstraintFrame[1]);
-			const FReal Delta = LinearDisplacement.Size();
-			const FReal TargetDelta = JointTransform.GetTranslation().Size();
-			if (!FMath::IsNearlyZero(TargetDelta) && !FMath::IsNearlyZero(Delta))
+			const bool bLinearPositionDriveEnabled = JointSettings.bLinearPositionDriveEnabled[0] && JointSettings.bLinearPositionDriveEnabled[1] && JointSettings.bLinearPositionDriveEnabled[2];
+			if(!bLinearPositionDriveEnabled)
 			{
-				FReal Ratio = Delta / TargetDelta;
-				if ((1.f - Ratio) > JointSettings.LinearPlasticityLimit)
-				{
-					JointTransform.ScaleTranslation(Ratio);
-					ConstraintFrame[1] = JointTransform.GetRelativeTransformReverse(ConstraintFrame[0]);
-				}
+				UE_LOG(LogChaosJoint, Warning, TEXT("Plasticity only supports three-dimensional linear position drives."));
+				bApplyLinearPlasticityLimits = false;
+			}
+			const bool bLinearMotionLocked = JointSettings.LinearMotionTypes[0] == EJointMotionType::Locked || JointSettings.LinearMotionTypes[1] == EJointMotionType::Locked || JointSettings.LinearMotionTypes[2] == EJointMotionType::Locked;
+			if(bLinearMotionLocked)
+			{
+				UE_LOG(LogChaosJoint, Warning, TEXT("Plasticity does not support locked linear motion."));
+				bApplyLinearPlasticityLimits = false;
+			}
+		}
+
+		bool bApplyAngularPlasticityLimits = bHasAngularPlasticityLimit;
+		if(bHasAngularPlasticityLimit)
+		{
+			const bool bAngularPositionDriveEnabled = JointSettings.bAngularSLerpPositionDriveEnabled || (JointSettings.bAngularSwingPositionDriveEnabled && JointSettings.bAngularTwistPositionDriveEnabled);
+			if(!bAngularPositionDriveEnabled)
+			{
+				UE_LOG(LogChaosJoint, Warning, TEXT("Plasticity only supports three-dimensional angular position drives"));
+				bApplyAngularPlasticityLimits = false;
+			}
+			const bool bAngularMotionLocked = JointSettings.AngularMotionTypes[(int32)EJointAngularConstraintIndex::Twist] == EJointMotionType::Locked || JointSettings.AngularMotionTypes[(int32)EJointAngularConstraintIndex::Swing1] == EJointMotionType::Locked || JointSettings.AngularMotionTypes[(int32)EJointAngularConstraintIndex::Swing2] == EJointMotionType::Locked;
+			if(bAngularMotionLocked)
+			{
+				UE_LOG(LogChaosJoint, Warning, TEXT("Plasticity does not support locked angular motion."));
+				bApplyAngularPlasticityLimits = false;
 			}
 		}
 
+		if(!bApplyLinearPlasticityLimits && !bApplyAngularPlasticityLimits)
+		{
+			return;
+		}
+
+		int32 Index0, Index1;
+		GetConstrainedParticleIndices(ConstraintIndex, Index0, Index1);
+		FGenericParticleHandle Particle0 = FGenericParticleHandle(ConstraintParticles[ConstraintIndex][Index0]);
+		FGenericParticleHandle Particle1 = FGenericParticleHandle(ConstraintParticles[ConstraintIndex][Index1]);
+
+		const FTransformPair &ConstraintFramesLocal = ConstraintFrames[ConstraintIndex];
+		FTransformPair ConstraintFramesGlobal(ConstraintFramesLocal[Index0] * FRigidTransform3(Particle0->P(), Particle0->Q()), ConstraintFramesLocal[Index1] * FRigidTransform3(Particle1->P(), Particle1->Q()));
+		FQuat Q1 = ConstraintFramesGlobal[1].GetRotation();
+		Q1.EnforceShortestArcWith(ConstraintFramesGlobal[0].GetRotation());
+		ConstraintFramesGlobal[1].SetRotation(Q1);
+
+		if(bApplyLinearPlasticityLimits)
+		{
+			const FVec3 LinearDisplacement = ConstraintFramesGlobal[0].InverseTransformPositionNoScale(ConstraintFramesGlobal[1].GetTranslation());
+			if((LinearDisplacement - JointSettings.LinearDrivePositionTarget).Size() > JointSettings.LinearPlasticityLimit)
+			{
+				JointSettings.LinearDrivePositionTarget = LinearDisplacement;
+			}
+		}
 
-		if (!FMath::IsNearlyEqual(JointSettings.AngularPlasticityLimit, (FReal)FLT_MAX))
+		if(bApplyAngularPlasticityLimits)
 		{
+			const FRotation3 AngularDisplacement = ConstraintFramesGlobal[0].GetRotation().Inverse() * ConstraintFramesGlobal[1].GetRotation();
 			const FReal AngleDeg = JointSettings.AngularDrivePositionTarget.AngularDistance(AngularDisplacement);
 			if (AngleDeg > JointSettings.AngularPlasticityLimit)
 			{
diff --git a/Engine/Source/Runtime/Experimental/Chaos/Public/Chaos/PBDJointConstraints.h b/Engine/Source/Runtime/Experimental/Chaos/Public/Chaos/PBDJointConstraints.h
--- a/Engine/Source/Runtime/Experimental/Chaos/Public/Chaos/PBDJointConstraints.h
+++ b/Engine/Source/Runtime/Experimental/Chaos/Public/Chaos/PBDJointConstraints.h
@@ -252,7 +252,7 @@ namespace Chaos
 		bool ApplySingle(const FReal Dt, const int32 ConstraintIndex, const int32 NumPairIts, const int32 It, const int32 NumIts);
 		bool ApplyPushOutSingle(const FReal Dt, const int32 ConstraintIndex, const int32 NumPairIts, const int32 It, const int32 NumIts);
 		void ApplyBreakThreshold(const FReal Dt, int32 ConstraintIndex, const FVec3& LinearImpulse, const FVec3& AngularImpulse);
-		void ApplyPlasticityLimits(const FReal Dt, int32 ConstraintIndex, const FVec3& LinearDisplacement, const FRotation3& AngularDisplacement);
+		void ApplyPlasticityLimits(const int32 ConstraintIndex);
 
 		FPBDJointSolverSettings Settings;
 
