UGeometryCollection
-- Wrapped GeometrySource in header tool directive to be stripped out of shipping game builds.

diff --git a/Engine/Source/Runtime/Experimental/GeometryCollectionEngine/Public/GeometryCollection/GeometryCollectionObject.h b/Engine/Source/Runtime/Experimental/GeometryCollectionEngine/Public/GeometryCollection/GeometryCollectionObject.h
--- a/Engine/Source/Runtime/Experimental/GeometryCollectionEngine/Public/GeometryCollection/GeometryCollectionObject.h
+++ b/Engine/Source/Runtime/Experimental/GeometryCollectionEngine/Public/GeometryCollection/GeometryCollectionObject.h
@@ -285,9 +285,11 @@ public:
 	UPROPERTY(EditAnywhere, Category = "Clustering")
 	EClusterConnectionTypeEnum ClusterConnectionType;
 
+#if WITH_EDITORONLY_DATA
 	UPROPERTY(BlueprintReadWrite, EditAnywhere, Category = "GeometrySource")
 	TArray<FGeometryCollectionSource> GeometrySource;
-	
+#endif
+
 	UPROPERTY(EditAnywhere, Category = "Materials")
 	TArray<TObjectPtr<UMaterialInterface>> Materials;
 
