GeometryCollectionSceneProxy - Properly initialized LocalToWorld transform on Nanite scene proxy.

diff --git a/Engine/Source/Runtime/Experimental/GeometryCollectionEngine/Private/GeometryCollection/GeometryCollectionSceneProxy.cpp b/Engine/Source/Runtime/Experimental/GeometryCollectionEngine/Private/GeometryCollection/GeometryCollectionSceneProxy.cpp
--- a/Engine/Source/Runtime/Experimental/GeometryCollectionEngine/Private/GeometryCollection/GeometryCollectionSceneProxy.cpp
+++ b/Engine/Source/Runtime/Experimental/GeometryCollectionEngine/Private/GeometryCollection/GeometryCollectionSceneProxy.cpp
@@ -1323,6 +1323,8 @@ void FNaniteGeometryCollectionSceneProxy::SetConstantData_RenderThread(FGeometry
 	check(NewConstantData->RestTransforms.Num() == TransformToGeometryIndices.Num());
 	Instances.Reset(NewConstantData->RestTransforms.Num());
 
+	const FMatrix ParentLocalToWorld = GetLocalToWorld();
+
 	for (int32 TransformIndex = 0; TransformIndex < NewConstantData->RestTransforms.Num(); ++TransformIndex)
 	{
 		const int32 TransformToGeometryIndex = TransformToGeometryIndices[TransformIndex];
@@ -1336,7 +1338,9 @@ void FNaniteGeometryCollectionSceneProxy::SetConstantData_RenderThread(FGeometry
 		FPrimitiveInstance& Instance = Instances.Emplace_GetRef();
 
 		Instance.InstanceToLocal		= NewConstantData->RestTransforms[TransformIndex];
+		Instance.LocalToWorld			= Instance.InstanceToLocal * ParentLocalToWorld;
 		Instance.PrevInstanceToLocal	= NewConstantData->RestTransforms[TransformIndex];
+		Instance.PrevLocalToWorld		= Instance.LocalToWorld;
 		Instance.PrimitiveId			= NaniteData.PrimitiveId;
 		Instance.RenderBounds			= NaniteData.RenderBounds;
 		Instance.LocalBounds			= Instance.RenderBounds.TransformBy(Instance.InstanceToLocal);
