Chaos: Fixing implicit type conversions where there could be data loss.
Mostly casting directly to the correct type but caught a few bugs in selected types that are now fixed)
Also made such conversions warnings within the Chaos and ChaosCore modules, may promote to error once we're sure it's buliding fine.

diff --git a/Engine/Source/Runtime/CoreUObject/Public/UObject/UObjectGlobals.h b/Engine/Source/Runtime/CoreUObject/Public/UObject/UObjectGlobals.h
--- a/Engine/Source/Runtime/CoreUObject/Public/UObject/UObjectGlobals.h
+++ b/Engine/Source/Runtime/CoreUObject/Public/UObject/UObjectGlobals.h
@@ -584,7 +584,7 @@ COREUOBJECT_API void GatherUnreachableObjects(bool bForceSingleThreaded);
  * @param	bUseTimeLimit	whether the time limit parameter should be used
  * @param	TimeLimit		soft time limit for this function call
  */
-COREUOBJECT_API void IncrementalPurgeGarbage( bool bUseTimeLimit, float TimeLimit = 0.002 );
+COREUOBJECT_API void IncrementalPurgeGarbage( bool bUseTimeLimit, float TimeLimit = 0.002f );
 
 /**
  * Create a unique name by combining a base name and an arbitrary number string.
diff --git a/Engine/Source/Runtime/Experimental/Chaos/Chaos.Build.cs b/Engine/Source/Runtime/Experimental/Chaos/Chaos.Build.cs
--- a/Engine/Source/Runtime/Experimental/Chaos/Chaos.Build.cs
+++ b/Engine/Source/Runtime/Experimental/Chaos/Chaos.Build.cs
@@ -39,6 +39,8 @@ namespace UnrealBuildTool.Rules
 			{
 				PublicDefinitions.Add("CHAOS_MEMORY_TRACKING=0");
 			}
-        }
+
+			UnsafeTypeCastWarningLevel = WarningLevel.Warning;
+		}
     }
 }
diff --git a/Engine/Source/Runtime/Experimental/Chaos/Private/Chaos/AABB.cpp b/Engine/Source/Runtime/Experimental/Chaos/Private/Chaos/AABB.cpp
--- a/Engine/Source/Runtime/Experimental/Chaos/Private/Chaos/AABB.cpp
+++ b/Engine/Source/Runtime/Experimental/Chaos/Private/Chaos/AABB.cpp
@@ -89,7 +89,7 @@ bool TAABB<T, d>::Raycast(const TVector<T, d>& StartPoint, const TVector<T, d>&
 	//If the box is rounded we have to consider corners and edges.
 	//Break the box into voronoi regions based on features (corner, edge, face) and see which region the raycast hit
 
-	if (Thickness)
+	if (Thickness != (T)0)
 	{
 		check(d == 3);
 		TVector<T, d> GeomStart;
diff --git a/Engine/Source/Runtime/Experimental/Chaos/Private/Chaos/ChaosDebugDraw.cpp b/Engine/Source/Runtime/Experimental/Chaos/Private/Chaos/ChaosDebugDraw.cpp
--- a/Engine/Source/Runtime/Experimental/Chaos/Private/Chaos/ChaosDebugDraw.cpp
+++ b/Engine/Source/Runtime/Experimental/Chaos/Private/Chaos/ChaosDebugDraw.cpp
@@ -86,7 +86,7 @@ namespace Chaos
 			, VelScale(0.0f)
 			, AngVelScale(0.0f)
 			, ImpulseScale(0.0f)
-			, DrawPriority(10.0f)
+			, DrawPriority(10)
 		{
 		}
 
diff --git a/Engine/Source/Runtime/Experimental/Chaos/Private/Chaos/Collision/PBDCollisionConstraint.cpp b/Engine/Source/Runtime/Experimental/Chaos/Private/Chaos/Collision/PBDCollisionConstraint.cpp
--- a/Engine/Source/Runtime/Experimental/Chaos/Private/Chaos/Collision/PBDCollisionConstraint.cpp
+++ b/Engine/Source/Runtime/Experimental/Chaos/Private/Chaos/Collision/PBDCollisionConstraint.cpp
@@ -98,7 +98,7 @@ namespace Chaos
 				FReal DRLen2 = DR.SizeSquared();
 				if (DRLen2 < DistanceTolerance2)
 				{
-					OutScore = FMath::Clamp(1. - DRLen2 / DistanceTolerance2, 0., 1.);
+					OutScore = FMath::Clamp(1.f - DRLen2 / DistanceTolerance2, 0.f, 1.f);
 					return true;
 				}
 			}
diff --git a/Engine/Source/Runtime/Experimental/Chaos/Private/Chaos/CollisionOneShotManifolds.cpp b/Engine/Source/Runtime/Experimental/Chaos/Private/Chaos/CollisionOneShotManifolds.cpp
--- a/Engine/Source/Runtime/Experimental/Chaos/Private/Chaos/CollisionOneShotManifolds.cpp
+++ b/Engine/Source/Runtime/Experimental/Chaos/Private/Chaos/CollisionOneShotManifolds.cpp
@@ -47,7 +47,7 @@ namespace Chaos
 		// This function will clip the input vertices by a reference shape's planes (Specified by ClippingAxis and Distance for an AABB)
 		// more vertices may be added to outputVertexBuffer by this function
 		// This is the core of the Sutherland-Hodgman algorithm
-		uint32 BoxBoxClipVerticesAgainstPlane(const FVec3* InputVertexBuffer, FVec3* outputVertexBuffer, uint32 ClipPointCount, FReal ClippingAxis, FReal Distance)
+		uint32 BoxBoxClipVerticesAgainstPlane(const FVec3* InputVertexBuffer, FVec3* outputVertexBuffer, uint32 ClipPointCount, int32 ClippingAxis, FReal Distance)
 		{
 
 			auto CalculateIntersect = [=](const FVec3& Point1, const FVec3& Point2) -> FVec3
@@ -183,7 +183,7 @@ namespace Chaos
 				}
 			}
 
-			const FReal SmallBiasToPreventFeatureFlipping = 0.002; // This improves frame coherence by penalizing box 1 in favour of box 2
+			const FReal SmallBiasToPreventFeatureFlipping = 0.002f; // This improves frame coherence by penalizing box 1 in favour of box 2
 			bool ReferenceFaceBox1 = true; // Is the reference face on box1 or box2?
 			if (BestFaceNormalSizeInDirectionBox2 + SmallBiasToPreventFeatureFlipping > BestFaceNormalSizeInDirectionBox1)
 			{
@@ -280,8 +280,8 @@ namespace Chaos
 			{
 				if (Coordinate != RefPlaneCoordinateIndex)
 				{
-					ContactPointCount = BoxBoxClipVerticesAgainstPlane(VertexBuffer1, VertexBuffer2, ContactPointCount, Coordinate, refBoxHalfExtents[Coordinate]);
-					ContactPointCount = BoxBoxClipVerticesAgainstPlane(VertexBuffer2, VertexBuffer1, ContactPointCount, Coordinate, -refBoxHalfExtents[Coordinate]);
+					ContactPointCount = BoxBoxClipVerticesAgainstPlane(VertexBuffer1, VertexBuffer2, ContactPointCount, static_cast<int32>(Coordinate), refBoxHalfExtents[Coordinate]);
+					ContactPointCount = BoxBoxClipVerticesAgainstPlane(VertexBuffer2, VertexBuffer1, ContactPointCount, static_cast<int32>(Coordinate), -refBoxHalfExtents[Coordinate]);
 				}
 			}
 
diff --git a/Engine/Source/Runtime/Experimental/Chaos/Private/Chaos/CollisionOneShotManifolds.h b/Engine/Source/Runtime/Experimental/Chaos/Private/Chaos/CollisionOneShotManifolds.h
--- a/Engine/Source/Runtime/Experimental/Chaos/Private/Chaos/CollisionOneShotManifolds.h
+++ b/Engine/Source/Runtime/Experimental/Chaos/Private/Chaos/CollisionOneShotManifolds.h
@@ -11,7 +11,7 @@ namespace Chaos
 {
 	namespace Collisions
 	{
-		uint32 BoxBoxClipVerticesAgainstPlane(const FVec3* InputVertexBuffer, FVec3* outputVertexBuffer, uint32 ClipPointCount, FReal ClippingAxis, FReal Distance);
+		uint32 BoxBoxClipVerticesAgainstPlane(const FVec3* InputVertexBuffer, FVec3* outputVertexBuffer, uint32 ClipPointCount, int32 ClippingAxis, FReal Distance);
 		uint32 ReduceManifoldContactPoints(FVec3* Points, uint32 PointCount);
 
 		void ConstructBoxBoxOneShotManifold(
diff --git a/Engine/Source/Runtime/Experimental/Chaos/Private/Chaos/CollisionResolution.cpp b/Engine/Source/Runtime/Experimental/Chaos/Private/Chaos/CollisionResolution.cpp
--- a/Engine/Source/Runtime/Experimental/Chaos/Private/Chaos/CollisionResolution.cpp
+++ b/Engine/Source/Runtime/Experimental/Chaos/Private/Chaos/CollisionResolution.cpp
@@ -777,7 +777,7 @@ namespace Chaos
 				{
 					AverageLocation += PotentialConstraints[ConstraintIdx];
 				}
-				Contact.Location = AverageLocation / NumConstraints;
+				Contact.Location = AverageLocation / static_cast<FReal>(NumConstraints);
 			}
 
 			Contact.ShapeMargins[0] = 0.0f;
diff --git a/Engine/Source/Runtime/Experimental/Chaos/Private/Chaos/Evolution/PBDMinEvolution.cpp b/Engine/Source/Runtime/Experimental/Chaos/Private/Chaos/Evolution/PBDMinEvolution.cpp
--- a/Engine/Source/Runtime/Experimental/Chaos/Private/Chaos/Evolution/PBDMinEvolution.cpp
+++ b/Engine/Source/Runtime/Experimental/Chaos/Private/Chaos/Evolution/PBDMinEvolution.cpp
@@ -181,7 +181,7 @@ namespace Chaos
 		{
 			// StepFraction: how much of the remaining time this step represents, used to interpolate kinematic targets
 			// E.g., for 4 steps this will be: 1/4, 1/2, 3/4, 1
-			const FReal StepFraction = (FReal)(Step + 1) / (NumSteps);
+			const FReal StepFraction = (FReal)(Step + 1) / (FReal)(NumSteps);
 
 			UE_LOG(LogChaosMinEvolution, Verbose, TEXT("Advance dt = %f [%d/%d]"), StepDt, Step + 1, NumSteps);
 
diff --git a/Engine/Source/Runtime/Experimental/Chaos/Private/Chaos/FFT.cpp b/Engine/Source/Runtime/Experimental/Chaos/Private/Chaos/FFT.cpp
--- a/Engine/Source/Runtime/Experimental/Chaos/Private/Chaos/FFT.cpp
+++ b/Engine/Source/Runtime/Experimental/Chaos/Private/Chaos/FFT.cpp
@@ -41,10 +41,10 @@ void FFFT3::MakeDivergenceFree(const FUniformGrid& Grid, FArrayNDOfComplex& u, F
 		u(i, 0, 0) = FComplex(0, 0);
 	for (int i = 0; i < Grid.Counts()[0]; ++i)
 	{
-		FReal k1 = Coefficients[0] * (i <= Grid.Counts()[0] / 2 ? i : i - Grid.Counts()[0]);
+		FReal k1 = Coefficients[0] * (i <= Grid.Counts()[0] / 2 ? (FReal)i : (FReal)i - (FReal)(Grid.Counts()[0]));
 		for (int j = 1; j <= Grid.Counts()[1] / 2; ++j)
 		{
-			FReal k2 = Coefficients[1] * j;
+			FReal k2 = Coefficients[1] * (FReal)j;
 			FReal OneOverKSquared = 1 / (k1 * k1 + k2 * k2);
 			FComplex correction = (k1 * u(i, j, 0) + k2 * v(i, j, 0)) * OneOverKSquared;
 			u(i, j, 0) -= correction * k1;
@@ -54,13 +54,13 @@ void FFFT3::MakeDivergenceFree(const FUniformGrid& Grid, FArrayNDOfComplex& u, F
 	// volume
 	for (int i = 0; i < Grid.Counts()[0]; ++i)
 	{
-		FReal k1 = Coefficients[0] * (i <= Grid.Counts()[0] / 2 ? i : i - Grid.Counts()[0]);
+		FReal k1 = Coefficients[0] * (i <= Grid.Counts()[0] / 2 ? (FReal)i : (FReal)i - (FReal)Grid.Counts()[0]);
 		for (int j = 0; j < Grid.Counts()[1]; ++j)
 		{
-			FReal k2 = Coefficients[1] * (j <= Grid.Counts()[1] / 2 ? j : j - Grid.Counts()[1]);
+			FReal k2 = Coefficients[1] * (j <= Grid.Counts()[1] / 2 ? (FReal)j : (FReal)j - (FReal)Grid.Counts()[1]);
 			for (int k = 1; k <= Grid.Counts()[2] / 2; ++k)
 			{
-				FReal k3 = Coefficients[2] * k;
+				FReal k3 = Coefficients[2] * (FReal)k;
 				FReal OneOverKSquared = 1 / (k1 * k1 + k2 * k2 + k3 * k3);
 				FComplex correction = (k1 * u(i, j, k) + k2 * v(i, j, k) + k3 * w(i, j, k)) * OneOverKSquared;
 				u(i, j, k) -= correction * k1;
@@ -123,7 +123,7 @@ void NRFourn(const int isign, const TVector<int32, d>& counts, TArray<FReal>& Da
 		while (ifp1 < ip2)
 		{
 			ifp2 = ifp1 << 1;
-			theta = isign * 6.28318530717959 / (ifp2 / ip1);
+			theta = isign * 6.28318530717959 / (FReal)(ifp2 / ip1);
 			wtemp = sin(.5 * theta);
 			wpr = -2 * wtemp * wtemp;
 			wpi = sin(theta);
diff --git a/Engine/Source/Runtime/Experimental/Chaos/Private/Chaos/Framework/PhysicsSolverBase.cpp b/Engine/Source/Runtime/Experimental/Chaos/Private/Chaos/Framework/PhysicsSolverBase.cpp
--- a/Engine/Source/Runtime/Experimental/Chaos/Private/Chaos/Framework/PhysicsSolverBase.cpp
+++ b/Engine/Source/Runtime/Experimental/Chaos/Private/Chaos/Framework/PhysicsSolverBase.cpp
@@ -311,7 +311,7 @@ namespace Chaos
 
 				InternalDt = AsyncDt;
 				NumSteps = FMath::FloorToInt(AccumulatedTime / InternalDt);
-				AccumulatedTime -= InternalDt * NumSteps;
+				AccumulatedTime -= InternalDt * static_cast<FReal>(NumSteps);
 			}
 		}
 		else if (bSubstepping && InDt > 0)
@@ -326,7 +326,7 @@ namespace Chaos
 			}
 			else
 			{
-				InternalDt = DtWithPause / NumSteps;
+				InternalDt = DtWithPause / static_cast<FReal>(NumSteps);
 			}
 		}
 
diff --git a/Engine/Source/Runtime/Experimental/Chaos/Private/Chaos/HeightField.cpp b/Engine/Source/Runtime/Experimental/Chaos/Private/Chaos/HeightField.cpp
--- a/Engine/Source/Runtime/Experimental/Chaos/Private/Chaos/HeightField.cpp
+++ b/Engine/Source/Runtime/Experimental/Chaos/Private/Chaos/HeightField.cpp
@@ -328,8 +328,8 @@ namespace Chaos
 		const int32 NumHeights = BufferView.Num();
 		OutData.Heights.SetNum(NumHeights);
 
-		OutData.NumRows = NumRows;
-		OutData.NumCols = NumCols;
+		OutData.NumRows = static_cast<uint16>(NumRows);
+		OutData.NumCols = static_cast<uint16>(NumCols);
 		OutData.MinValue = ToRealFunc(BufferView[0]);
 		OutData.MaxValue = ToRealFunc(BufferView[0]);
 		OutData.Scale = InScale;
@@ -540,8 +540,8 @@ namespace Chaos
 	{
 		if (FlatGrid.IsValid(InCoord))
 		{
-			OutBounds.Min = FVec2(InCoord[0], InCoord[1]);
-			OutBounds.Max = FVec2(InCoord[0] + 1, InCoord[1] + 1);
+			OutBounds.Min = FVec2(static_cast<FReal>(InCoord[0]), static_cast<FReal>(InCoord[1]));
+			OutBounds.Max = FVec2(static_cast<FReal>(InCoord[0] + 1), static_cast<FReal>(InCoord[1] + 1));
 			OutBounds.Min -= InInflate;
 			OutBounds.Max += InInflate;
 
@@ -675,8 +675,8 @@ namespace Chaos
 			FVec3 Min,Max;
 			CalcCellBounds3D(InCoord,Min,Max);
 
-			OutMin = FVec3(InCoord[0], InCoord[1], GeomData.GetMinHeight());
-			OutMax = FVec3(InCoord[0] + 1, InCoord[1] + 1, Max[2]);
+			OutMin = FVec3(static_cast<FReal>(InCoord[0]), static_cast<FReal>(InCoord[1]), GeomData.GetMinHeight());
+			OutMax = FVec3(static_cast<FReal>(InCoord[0] + 1), static_cast<FReal>(InCoord[1] + 1), Max[2]);
 			OutMin = OutMin - InInflate;
 			OutMax = OutMax + InInflate;
 
@@ -690,8 +690,8 @@ namespace Chaos
 	{
 		if (FlatGrid.IsValid(InCoord))
 		{
-			OutBounds.Min = FVec2(InCoord[0], InCoord[1]);
-			OutBounds.Max = FVec2(InCoord[0] + 1, InCoord[1] + 1);
+			OutBounds.Min = FVec2(static_cast<FReal>(InCoord[0]), static_cast<FReal>(InCoord[1]));
+			OutBounds.Max = FVec2(static_cast<FReal>(InCoord[0] + 1), static_cast<FReal>(InCoord[1] + 1));
 			OutBounds.Min -= InInflate;
 			OutBounds.Max += InInflate;
 			const FVec2 Scale2D = FVec2(GeomData.Scale[0], GeomData.Scale[1]);
@@ -711,8 +711,8 @@ namespace Chaos
 			FVec3 Min,Max;
 			CalcCellBounds3D(InCoord,Min,Max);
 
-			OutMin = FVec3(InCoord[0], InCoord[1], GeomData.GetMinHeight());
-			OutMax = FVec3(InCoord[0] + 1, InCoord[1] + 1, Max[2]);
+			OutMin = FVec3(static_cast<FReal>(InCoord[0]), static_cast<FReal>(InCoord[1]), GeomData.GetMinHeight());
+			OutMax = FVec3(static_cast<FReal>(InCoord[0] + 1), static_cast<FReal>(InCoord[1] + 1), Max[2]);
 			OutMin = OutMin * GeomData.Scale - InInflate;
 			OutMax = OutMax * GeomData.Scale + InInflate;
 			return true;
@@ -783,12 +783,12 @@ namespace Chaos
 		if(Bounds.RaycastFast(StartPoint, Dir, InvDir, bParallel, Length, InvCurrentLength, TOI, NextStart))
 		{
 			const FVec2 Scale2D(GeomData.Scale[0],GeomData.Scale[1]);
-			TVec2<int32> CellIdx = FlatGrid.Cell(TVec2<int32>(NextStart[0] / Scale2D[0], NextStart[1] / Scale2D[1]));
+			TVec2<int32> CellIdx = FlatGrid.Cell(TVec2<int32>(static_cast<int32>(NextStart[0] / Scale2D[0]), static_cast<int32>(NextStart[1] / Scale2D[1])));
 
 			// Boundaries might push us one cell over
 			CellIdx = FlatGrid.ClampIndex(CellIdx);
 			const FReal ZDx = Bounds.Extents()[2];
-			const FReal ZMidPoint = Bounds.Min()[2] + ZDx * 0.5;
+			const FReal ZMidPoint = Bounds.Min()[2] + ZDx * 0.5f;
 			const FVec3 ScaledDx(FlatGrid.Dx()[0] * Scale2D[0],FlatGrid.Dx()[1] * Scale2D[1],ZDx);
 			const FVec2 ScaledDx2D(ScaledDx[0],ScaledDx[1]);
 			const FVec2 ScaledMin = FlatGrid.MinCorner() * Scale2D;
@@ -815,7 +815,7 @@ namespace Chaos
 				//find next cell
 
 				//We want to know which plane we used to cross into next cell
-				const FVec2 ScaledCellCenter2D = ScaledMin + FVec2(CellIdx[0] + 0.5,CellIdx[1] + 0.5) * ScaledDx2D;
+				const FVec2 ScaledCellCenter2D = ScaledMin + FVec2(static_cast<FReal>(CellIdx[0]) + 0.5f, static_cast<FReal>(CellIdx[1]) + 0.5f) * ScaledDx2D;
 				const FVec3 ScaledCellCenter(ScaledCellCenter2D[0], ScaledCellCenter2D[1], ZMidPoint);
 
 				FReal Times[3];
@@ -896,7 +896,7 @@ namespace Chaos
 			int32 Idx = Coordinate[1] * NumX + Coordinate[0];
 			int32 ByteIdx = Idx / 8;
 			int32 BitIdx = Idx % 8;
-			uint8 Mask = 1 << BitIdx;
+			uint8 Mask = static_cast<uint8>(1 << BitIdx);
 			check(ByteIdx >= 0 && ByteIdx < DataSize);
 			Data[ByteIdx] |= Mask;
 		}
@@ -1069,9 +1069,17 @@ namespace Chaos
 					if(CellCoord.ToI < 0)
 					{
 						// Perform expansion for thickness
-						const int32 ExpandAxis = ThickenDir[0] == 0 ? 1 : 0;
-						const FReal ExpandSize = HalfExtents3D[ExpandAxis];
-						const int32 Steps = FMath::RoundFromZero(ExpandSize / GeomData.Scale[ExpandAxis]);
+						int32 ExpandAxis;
+						if(ThickenDir[0] == 0)
+						{
+							ExpandAxis = 1;
+						}
+						else
+						{
+							ExpandAxis = 0;
+						}
+						FReal ExpandSize = HalfExtents3D[ExpandAxis];
+						int32 Steps = FMath::TruncToInt(FMath::RoundFromZero(ExpandSize / GeomData.Scale[ExpandAxis]));
 
 						Expand(Coord, ThickenDir, Steps);
 						Expand(Coord, -ThickenDir, Steps);
@@ -1747,7 +1755,7 @@ namespace Chaos
 			const FVec3 AC = C - A;
 			FVec3 Normal = FVec3::CrossProduct(AB, AC);
 			const FReal Length = Normal.SafeNormalize();
-			ensure(Length);
+			ensure(Length > 0);
 			return Normal;
 		}
 
@@ -1772,7 +1780,7 @@ namespace Chaos
 		TVec2<int32> Cells(GeomData.NumCols - 1, GeomData.NumRows - 1);
 		
 		FVec2 MinCorner(0, 0);
-		FVec2 MaxCorner(GeomData.NumCols - 1, GeomData.NumRows - 1);
+		FVec2 MaxCorner(static_cast<FReal>(GeomData.NumCols - 1), static_cast<FReal>(GeomData.NumRows - 1));
 		//MaxCorner *= {GeomData.Scale[0], GeomData.Scale[1]};
 
 		FlatGrid = TUniformGrid<FReal, 2>(MinCorner, MaxCorner, Cells);
diff --git a/Engine/Source/Runtime/Experimental/Chaos/Private/Chaos/Joint/PBDJointSolverGaussSeidel.cpp b/Engine/Source/Runtime/Experimental/Chaos/Private/Chaos/Joint/PBDJointSolverGaussSeidel.cpp
--- a/Engine/Source/Runtime/Experimental/Chaos/Private/Chaos/Joint/PBDJointSolverGaussSeidel.cpp
+++ b/Engine/Source/Runtime/Experimental/Chaos/Private/Chaos/Joint/PBDJointSolverGaussSeidel.cpp
@@ -154,7 +154,7 @@ namespace Chaos
 		// the position error is less than the position change from constant external forces
 		// (e.g., gravity). So, we are saying that the tolerance was chosen because the position
 		// error is less that F.dt^2. We need to scale the tolerance to work at our current dt.
-		const FReal ToleranceScale = FMath::Min(1., 60. * 60. * Dt * Dt);
+		const FReal ToleranceScale = FMath::Min(1.f, 60.f * 60.f * Dt * Dt);
 		PositionTolerance = ToleranceScale * SolverSettings.PositionTolerance;
 		AngleTolerance = ToleranceScale * SolverSettings.AngleTolerance;
 
@@ -859,7 +859,7 @@ namespace Chaos
 		const FVec3 DR1 = IA1 * -(Angle / II1);
 		ApplyRotationDelta(DIndex, Stiffness, DR1);
 
-		NetAngularImpulse += (KIndex == 0 ? 1 : -1) * Axis * (Stiffness * Angle / II1);
+		NetAngularImpulse += (KIndex == 0 ? 1.0f : -1.0f) * Axis * (Stiffness * Angle / II1);
 	}
 
 
diff --git a/Engine/Source/Runtime/Experimental/Chaos/Private/Chaos/Levelset.cpp b/Engine/Source/Runtime/Experimental/Chaos/Private/Chaos/Levelset.cpp
--- a/Engine/Source/Runtime/Experimental/Chaos/Private/Chaos/Levelset.cpp
+++ b/Engine/Source/Runtime/Experimental/Chaos/Private/Chaos/Levelset.cpp
@@ -73,8 +73,8 @@ FLevelSet::FLevelSet(FErrorReporter& ErrorReporter, const TUniformGrid<FReal, 3>
 		ErrorReporter.ReportError(TEXT("Error calling FLevelSet::ComputeDistancesNearZeroIsocontour"));
 		return;
 	}
-	FReal StoppingDistance = MBandWidth * MGrid.Dx().Max();
-	if (StoppingDistance)
+	FReal StoppingDistance = static_cast<FReal>(MBandWidth) * MGrid.Dx().Max();
+	if (StoppingDistance != 0)
 	{
 		for (int32 i = 0; i < MGrid.Counts().Product(); ++i)
 		{
@@ -83,7 +83,7 @@ FLevelSet::FLevelSet(FErrorReporter& ErrorReporter, const TUniformGrid<FReal, 3>
 	}
 	CorrectSign(BlockedFaceX, BlockedFaceY, BlockedFaceZ, InterfaceIndices);
 	FillWithFastMarchingMethod(StoppingDistance, InterfaceIndices);
-	if (StoppingDistance)
+	if (StoppingDistance != 0)
 	{
 		for (int32 i = 0; i < MGrid.Counts().Product(); ++i)
 		{
@@ -128,8 +128,8 @@ FLevelSet::FLevelSet(FErrorReporter& ErrorReporter, const TUniformGrid<FReal, 3>
 	}
 	TArray<TVec3<int32>> InterfaceIndices;
 	ComputeDistancesNearZeroIsocontour(InObject, ObjectPhi, InterfaceIndices);
-	FReal StoppingDistance = MBandWidth * MGrid.Dx().Max();
-	if (StoppingDistance)
+	FReal StoppingDistance = static_cast<FReal>(MBandWidth) * MGrid.Dx().Max();
+	if (StoppingDistance != 0)
 	{
 		for (int32 i = 0; i < MGrid.Counts().Product(); ++i)
 		{
@@ -142,7 +142,7 @@ FLevelSet::FLevelSet(FErrorReporter& ErrorReporter, const TUniformGrid<FReal, 3>
 		MPhi[i] *= FMath::Sign(ObjectPhi[i]);
 	}
 	FillWithFastMarchingMethod(StoppingDistance, InterfaceIndices);
-	if (StoppingDistance)
+	if (StoppingDistance != 0)
 	{
 		for (int32 i = 0; i < MGrid.Counts().Product(); ++i)
 		{
@@ -218,8 +218,8 @@ bool FLevelSet::ComputeMassProperties(FReal& OutVolume, FVec3& OutCOM, FMatrix33
 	}
 
 	const int32 NumCellsWithVolume = CellsWithVolume.Num();
-	FReal Volume = NumCellsWithVolume * CellVolume;
-	FMatrix33 Inertia = CellInertia * NumCellsWithVolume;
+	FReal Volume = static_cast<FReal>(NumCellsWithVolume) * CellVolume;
+	FMatrix33 Inertia = CellInertia * (FReal)NumCellsWithVolume;
 	if (Volume > 0)
 	{
 		COM /= Volume;
@@ -297,7 +297,7 @@ FReal FLevelSet::ComputeLevelSetError(const FParticles& InParticles, const TArra
 			//AngleErrorValues[i] = FMath::Acos(FVec3::DotProduct(MeshFaceNormal, GridNormal));
 
 			// triangle area used for weighted average
-			TriangleArea[i] = .5 * sqrt(FVec3::CrossProduct(InParticles.X(CurrMeshFace[1]) - InParticles.X(CurrMeshFace[0]), InParticles.X(CurrMeshFace[2]) - InParticles.X(CurrMeshFace[0])).SizeSquared());
+			TriangleArea[i] = (FReal)0.5 * sqrt(FVec3::CrossProduct(InParticles.X(CurrMeshFace[1]) - InParticles.X(CurrMeshFace[0]), InParticles.X(CurrMeshFace[2]) - InParticles.X(CurrMeshFace[0])).SizeSquared());
 		}
 	});
 
@@ -331,7 +331,7 @@ FReal FLevelSet::ComputeLevelSetError(const FParticles& InParticles, const TArra
 	// dist error is a percentage deviation away from geometry bounds, which
 	// normalizes error metrics with respect to world space size
 	FVec3 BoxExtents = MLocalBoundingBox.Extents();
-	FReal AvgExtents = (BoxExtents[0] + BoxExtents[1] + BoxExtents[2]) / 3.;
+	FReal AvgExtents = (BoxExtents[0] + BoxExtents[1] + BoxExtents[2]) / (FReal)3.0;
 
 	// degenerate case where extents are very small
 	if (AvgExtents < 1e-5)
@@ -661,8 +661,8 @@ bool FLevelSet::ComputeDistancesNearZeroIsocontour(FErrorReporter& ErrorReporter
 		TriangleBounds.GrowToInclude(InParticles.X(Element[2]));
 		MOriginalLocalBoundingBox.GrowToInclude(TriangleBounds); //also save the original bounding box
 
-		TVec3<int32> StartIndex = MGrid.ClampIndex(MGrid.Cell(TriangleBounds.Min() - FVec3((0.5 + KINDA_SMALL_NUMBER) * MGrid.Dx())));
-		TVec3<int32> EndIndex = MGrid.ClampIndex(MGrid.Cell(TriangleBounds.Max() + FVec3((0.5 + KINDA_SMALL_NUMBER) * MGrid.Dx())));
+		TVec3<int32> StartIndex = MGrid.ClampIndex(MGrid.Cell(TriangleBounds.Min() - FVec3((0.5f + KINDA_SMALL_NUMBER) * MGrid.Dx())));
+		TVec3<int32> EndIndex = MGrid.ClampIndex(MGrid.Cell(TriangleBounds.Max() + FVec3((0.5f + KINDA_SMALL_NUMBER) * MGrid.Dx())));
 		for (int32 i = StartIndex[0]; i <= EndIndex[0]; ++i)
 		{
 			for (int32 j = StartIndex[1]; j <= EndIndex[1]; ++j)
@@ -1004,7 +1004,7 @@ void FLevelSet::FillWithFastMarchingMethod(const FReal StoppingDistance, const T
 		Pair<FReal*, TVec3<int32>> Smallest;
 		Heap.HeapPop(Smallest, Compare);
 		check(InHeap(Smallest.Second));
-		if (StoppingDistance && FGenericPlatformMath::Abs(*Smallest.First) > StoppingDistance)
+		if (StoppingDistance != 0 && FGenericPlatformMath::Abs(*Smallest.First) > StoppingDistance)
 		{
 			break;
 		}
@@ -1040,7 +1040,7 @@ void FLevelSet::FillWithFastMarchingMethod(const FReal StoppingDistance, const T
 FReal SolveQuadraticEquation(const FReal Phi, const FReal PhiX, const FReal PhiY, const FReal Dx, const FReal Dy)
 {
 	check(FMath::Sign(PhiX) == FMath::Sign(PhiY) || FMath::Sign(PhiX) == 0 || FMath::Sign(PhiY) == 0);
-	FReal Sign = Phi > 0 ? 1 : -1;
+	FReal Sign = Phi > 0 ? (FReal)1.0 : (FReal)-1.0;
 	if (FMath::Abs(PhiX) >= (FMath::Abs(PhiY) + Dy))
 	{
 		return PhiY + Sign * Dy;
@@ -1093,7 +1093,7 @@ FReal FLevelSet::ComputePhi(const TArrayND<bool, 3>& Done, const TVec3<int32>& C
 	}
 	if (NumberOfAxes == 1)
 	{
-		FReal Sign = MPhi(CellIndex) > 0 ? 1 : -1;
+		FReal Sign = MPhi(CellIndex) > 0 ? (FReal)1.0 : (FReal)-1.0;
 		FReal NewPhi = FGenericPlatformMath::Abs(NeighborPhi[0]) + Dx[0];
 		check(NewPhi <= FGenericPlatformMath::Abs(MPhi(CellIndex)));
 		return Sign * NewPhi;
@@ -1114,7 +1114,7 @@ FReal FLevelSet::ComputePhi(const TArrayND<bool, 3>& Done, const TVec3<int32>& C
 		return QuadraticYZ;
 	}
 	// Cubic
-	FReal Sign = MPhi(CellIndex) > 0 ? 1 : -1;
+	FReal Sign = MPhi(CellIndex) > 0 ? (FReal)1.0 : (FReal)-1;
 	FReal Dx2 = Dx[0] * Dx[0];
 	FReal Dy2 = Dx[1] * Dx[1];
 	FReal Dz2 = Dx[2] * Dx[2];
@@ -1289,8 +1289,8 @@ void FLevelSet::ComputeNormals(const FParticles& InParticles, const FTriangleMes
 		TriangleBounds.GrowToInclude(InParticles.X(Element[2]));
 		MOriginalLocalBoundingBox.GrowToInclude(TriangleBounds); //also save the original bounding box
 
-		TVec3<int32> StartIndex = MGrid.ClampIndex(MGrid.Cell(TriangleBounds.Min() - FVec3((0.5 + KINDA_SMALL_NUMBER) * MGrid.Dx())));
-		TVec3<int32> EndIndex = MGrid.ClampIndex(MGrid.Cell(TriangleBounds.Max() + FVec3((0.5 + KINDA_SMALL_NUMBER) * MGrid.Dx())));
+		TVec3<int32> StartIndex = MGrid.ClampIndex(MGrid.Cell(TriangleBounds.Min() - FVec3((0.5f + KINDA_SMALL_NUMBER) * MGrid.Dx())));
+		TVec3<int32> EndIndex = MGrid.ClampIndex(MGrid.Cell(TriangleBounds.Max() + FVec3((0.5f + KINDA_SMALL_NUMBER) * MGrid.Dx())));
 		for (int32 i = StartIndex[0]; i <= EndIndex[0]; ++i)
 		{
 			for (int32 j = StartIndex[1]; j <= EndIndex[1]; ++j)
@@ -1383,14 +1383,14 @@ FReal FLevelSet::SignedDistance(const FVec3& x) const
 	FVec3 Location = MGrid.ClampMinusHalf(x);
 	FReal SizeSquared = (Location - x).SizeSquared();
 	FReal Phi = MGrid.LinearlyInterpolate(MPhi, Location);
-	return SizeSquared ? (sqrt(SizeSquared) + Phi) : Phi;
+	return SizeSquared > 0 ? (sqrt(SizeSquared) + Phi) : Phi;
 }
 
 FReal FLevelSet::PhiWithNormal(const FVec3& x, FVec3& Normal) const
 {
 	FVec3 Location = MGrid.ClampMinusHalf(x);
 	FReal SizeSquared = (Location - x).SizeSquared();
-	if (SizeSquared)
+	if (SizeSquared > 0)
 	{
 		MLocalBoundingBox.PhiWithNormal(Location, Normal);
 	}
@@ -1409,7 +1409,7 @@ FReal FLevelSet::PhiWithNormal(const FVec3& x, FVec3& Normal) const
 		}
 	}
 	FReal Phi = MGrid.LinearlyInterpolate(MPhi, Location);
-	return SizeSquared ? (sqrt(SizeSquared) + Phi) : Phi;
+	return SizeSquared > 0 ? (sqrt(SizeSquared) + Phi) : Phi;
 }
 
 void GetGeomSurfaceSamples(const TSphere<FReal, 3>& InGeom, TArray<FVec3>& OutSamples)
diff --git a/Engine/Source/Runtime/Experimental/Chaos/Private/Chaos/MassProperties.cpp b/Engine/Source/Runtime/Experimental/Chaos/Private/Chaos/MassProperties.cpp
--- a/Engine/Source/Runtime/Experimental/Chaos/Private/Chaos/MassProperties.cpp
+++ b/Engine/Source/Runtime/Experimental/Chaos/Private/Chaos/MassProperties.cpp
@@ -15,7 +15,7 @@ namespace Chaos
 
 		// Extract Eigenvalues
 		FReal OffDiagSize = FMath::Square(Inertia.M[1][0]) + FMath::Square(Inertia.M[2][0]) + FMath::Square(Inertia.M[2][1]);
-		double Trace = (Inertia.M[0][0] + Inertia.M[1][1] + Inertia.M[2][2]) / 3;
+		FRealDouble Trace = (Inertia.M[0][0] + Inertia.M[1][1] + Inertia.M[2][2]) / 3;
 
 		if (Trace <= SMALL_NUMBER)
 		{
@@ -29,22 +29,28 @@ namespace Chaos
 			return FRotation3::FromElements(FVec3(0), 1);
 		}
 
-		FReal Size = FMath::Sqrt((FMath::Square(Inertia.M[0][0] - Trace) + FMath::Square(Inertia.M[1][1] - Trace) + FMath::Square(Inertia.M[2][2] - Trace) + 2 * OffDiagSize) / 6);
-		FMatrix33 NewMat = (Inertia - FMatrix::Identity * Trace) * (1 / Size);
+		FReal Size = static_cast<FReal>(FMath::Sqrt((FMath::Square(Inertia.M[0][0] - Trace) + FMath::Square(Inertia.M[1][1] - Trace) + FMath::Square(Inertia.M[2][2] - Trace) + 2. * OffDiagSize) / 6.));
+		FMatrix33 NewMat = (Inertia - FMatrix::Identity * static_cast<FReal>(Trace)) * static_cast<FReal>((1 / Size));
 		FReal HalfDeterminant = NewMat.Determinant() / 2;
-		FReal Angle = HalfDeterminant <= -1 ? PI / 3 : (HalfDeterminant >= 1 ? 0 : acos(HalfDeterminant) / 3);
-		FReal m00 = Trace + 2 * Size * cos(Angle), m11 = Trace + 2 * Size * cos(Angle + (2 * PI / 3)), m22 = 3 * Trace - m00 - m11;
+		FReal Angle = HalfDeterminant <= -1 ? PI / 3 : (HalfDeterminant >= 1 ? 0 : FMath::Acos(HalfDeterminant) / 3);
+		
+		FReal m00 = static_cast<FReal>(Trace + 2 * Size * FMath::Cos(Angle));
+		FReal m11 = static_cast<FReal>(Trace + 2 * Size * FMath::Cos(Angle + (2 * PI / 3)));
+		FReal m22 = static_cast<FReal>(3 * Trace - m00 - m11);
 
 		// Extract Eigenvectors
 		bool DoSwap = ((m00 - m11) > (m11 - m22)) ? false : true;
 		FVec3 Eigenvector0 = (Inertia.SubtractDiagonal(DoSwap ? m22 : m00)).SymmetricCofactorMatrix().LargestColumnNormalized();
 		FVec3 Orthogonal = Eigenvector0.GetOrthogonalVector().GetSafeNormal();
+		
 		PMatrix<FReal, 3, 2> Cofactors(Orthogonal, FVec3::CrossProduct(Eigenvector0, Orthogonal));
 		PMatrix<FReal, 3, 2> CofactorsScaled = Inertia * Cofactors;
+		
 		PMatrix<FReal, 2, 2> IR(
 			CofactorsScaled.M[0] * Cofactors.M[0] + CofactorsScaled.M[1] * Cofactors.M[1] + CofactorsScaled.M[2] * Cofactors.M[2],
 			CofactorsScaled.M[3] * Cofactors.M[0] + CofactorsScaled.M[4] * Cofactors.M[1] + CofactorsScaled.M[5] * Cofactors.M[2],
 			CofactorsScaled.M[3] * Cofactors.M[3] + CofactorsScaled.M[4] * Cofactors.M[4] + CofactorsScaled.M[5] * Cofactors.M[5]);
+
 		PMatrix<FReal, 2, 2> IM1 = IR.SubtractDiagonal(DoSwap ? m00 : m22);
 		FReal OffDiag = IM1.M[1] * IM1.M[1];
 		FReal IM1Scale0 = FMath::Max(FReal(0), IM1.M[3] * IM1.M[3] + OffDiag);
@@ -70,6 +76,7 @@ namespace Chaos
 		// Return results
 		Inertia = FMatrix33(m00, 0, 0, m11, 0, m22);
 		FMatrix33 RotationMatrix = DoSwap ? FMatrix33(Eigenvector2, Eigenvector1, -Eigenvector0) : FMatrix33(Eigenvector0, Eigenvector1, Eigenvector2);
+
 		// NOTE: UE Matrix are column-major, so the PMatrix constructor is not setting eigenvectors - we need to transpose it to get a UE rotation matrix.
 		FinalRotation = FRotation3(RotationMatrix.GetTransposed());
 		if (!ensure(FMath::IsNearlyEqual(FinalRotation.Size(), 1.0f, KINDA_SMALL_NUMBER)))
diff --git a/Engine/Source/Runtime/Experimental/Chaos/Private/Chaos/PBDAxialSpringConstraints.cpp b/Engine/Source/Runtime/Experimental/Chaos/Private/Chaos/PBDAxialSpringConstraints.cpp
--- a/Engine/Source/Runtime/Experimental/Chaos/Private/Chaos/PBDAxialSpringConstraints.cpp
+++ b/Engine/Source/Runtime/Experimental/Chaos/Private/Chaos/PBDAxialSpringConstraints.cpp
@@ -42,11 +42,11 @@ void FPBDAxialSpringConstraints::ApplyImp(FPBDParticles& InParticles, const FRea
 		{
 			InParticles.P(i1) -= Multiplier * InParticles.InvM(i1) * Delta;
 		}
-		if (InParticles.InvM(i2))
+		if (InParticles.InvM(i2) > 0)
 		{
 			InParticles.P(i2) += Multiplier * InParticles.InvM(i2) * MBarys[i] * Delta;
 		}
-		if (InParticles.InvM(i3))
+		if (InParticles.InvM(i3) > 0)
 		{
 			InParticles.P(i3) += Multiplier * InParticles.InvM(i3) * (1 - MBarys[i]) * Delta;
 		}
diff --git a/Engine/Source/Runtime/Experimental/Chaos/Private/Chaos/PBDRigidClustering.cpp b/Engine/Source/Runtime/Experimental/Chaos/Private/Chaos/PBDRigidClustering.cpp
--- a/Engine/Source/Runtime/Experimental/Chaos/Private/Chaos/PBDRigidClustering.cpp
+++ b/Engine/Source/Runtime/Experimental/Chaos/Private/Chaos/PBDRigidClustering.cpp
@@ -316,7 +316,7 @@ namespace Chaos
 		Parent->X() /= Parent->M();
 		Parent->V() /= Parent->M();
 		Parent->PreV() = Parent->V();
-		Parent->InvM() = 1. / Parent->M();
+		Parent->InvM() = static_cast<FReal>(1.0) / Parent->M();
 		if (ForceMassOrientation)
 		{
 			Parent->X() = ForceMassOrientation->GetLocation();
@@ -482,7 +482,7 @@ namespace Chaos
 		}
 		if (ChildrenSet.Num())
 		{
-			NewParticle->Strains() /= ChildrenSet.Num();
+			NewParticle->Strains() /= static_cast<FReal>(ChildrenSet.Num());
 		}
 
 		ensureMsgf(!ProxyGeometry || ForceMassOrientation, TEXT("If ProxyGeometry is passed, we must override the mass orientation as they are tied"));
@@ -644,7 +644,7 @@ namespace Chaos
 
 				TArray<FPBDRigidParticleHandle*> ActiveCluster = Group.Value;
 
-				FClusterCreationParameters Parameters(0.3, 100, false, !!UnionsHaveCollisionParticles);
+				FClusterCreationParameters Parameters(0.3f, 100, false, !!UnionsHaveCollisionParticles);
 				Parameters.ConnectionMethod = MClusterUnionConnectionType;
 				TPBDRigidClusteredParticleHandleImp<FReal, 3, true>* Handle = CreateClusterParticle(-ClusterGroupID, MoveTemp(Group.Value), Parameters, TSharedPtr<FImplicitObject, ESPMode::ThreadSafe>());
 				Handle->SetInternalCluster(true);
diff --git a/Engine/Source/Runtime/Experimental/Chaos/Private/Chaos/PBDRigidsEvolution.cpp b/Engine/Source/Runtime/Experimental/Chaos/Private/Chaos/PBDRigidsEvolution.cpp
--- a/Engine/Source/Runtime/Experimental/Chaos/Private/Chaos/PBDRigidsEvolution.cpp
+++ b/Engine/Source/Runtime/Experimental/Chaos/Private/Chaos/PBDRigidsEvolution.cpp
@@ -137,7 +137,7 @@ namespace Chaos
 			{
 				if (ConfigSettings.BroadphaseType == 0)
 				{
-					return MakeUnique<BVType>(Particles, false, 0, ConfigSettings.BVNumCells, ConfigSettings.MaxPayloadSize);
+					return MakeUnique<BVType>(Particles, false, static_cast<FReal>(0), ConfigSettings.BVNumCells, ConfigSettings.MaxPayloadSize);
 				}
 				else if (ConfigSettings.BroadphaseType == 1 || ConfigSettings.BroadphaseType == 3)
 				{
@@ -151,7 +151,7 @@ namespace Chaos
 			case 1:
 			{
 				ensure(ConfigSettings.BroadphaseType == 3 || ConfigSettings.BroadphaseType == 4);
-				return MakeUnique<BVType>(Particles, false, 0, ConfigSettings.BVNumCells, ConfigSettings.MaxPayloadSize);
+				return MakeUnique<BVType>(Particles, false, static_cast<FReal>(0), ConfigSettings.BVNumCells, ConfigSettings.MaxPayloadSize);
 			}
 			default:
 			{
@@ -299,7 +299,7 @@ namespace Chaos
 		{
 			const FSpatialAccelerationIdx SpatialIdx = Itr.Key;
 			const FSpatialAccelerationCache& Cache = *Itr.Value;
-			const uint8 BucketIdx = (1 << SpatialIdx.Bucket) & ActiveBucketsMask ? SpatialIdx.Bucket : 0;
+			const uint8 BucketIdx = static_cast<uint8>((1 << SpatialIdx.Bucket) & ActiveBucketsMask ? SpatialIdx.Bucket : 0);
 			if(AccelerationStructure->GetSubstructure(SpatialIdx) && !AccelerationStructure->GetSubstructure(SpatialIdx)->IsAsyncTimeSlicingComplete())
 			{
 				SCOPE_CYCLE_COUNTER(STAT_AccelerationStructureTimeSlice);
diff --git a/Engine/Source/Runtime/Experimental/Chaos/Private/Chaos/PBDRigidsEvolutionGBF.cpp b/Engine/Source/Runtime/Experimental/Chaos/Private/Chaos/PBDRigidsEvolutionGBF.cpp
--- a/Engine/Source/Runtime/Experimental/Chaos/Private/Chaos/PBDRigidsEvolutionGBF.cpp
+++ b/Engine/Source/Runtime/Experimental/Chaos/Private/Chaos/PBDRigidsEvolutionGBF.cpp
@@ -134,7 +134,7 @@ void FPBDRigidsEvolutionGBF::Advance(const FReal Dt,const FReal MaxStepDt,const
 		{
 			// StepFraction: how much of the remaining time this step represents, used to interpolate kinematic targets
 			// E.g., for 4 steps this will be: 1/4, 1/2, 3/4, 1
-			const FReal StepFraction = (FReal)(Step + 1) / (NumSteps);
+			const FReal StepFraction = (FReal)(Step + 1) / (FReal)(NumSteps);
 		
 			UE_LOG(LogChaos, Verbose, TEXT("Advance dt = %f [%d/%d]"), StepDt, Step + 1, NumSteps);
 
diff --git a/Engine/Source/Runtime/Experimental/Chaos/Private/Chaos/PerParticleDampVelocity.cpp b/Engine/Source/Runtime/Experimental/Chaos/Private/Chaos/PerParticleDampVelocity.cpp
--- a/Engine/Source/Runtime/Experimental/Chaos/Private/Chaos/PerParticleDampVelocity.cpp
+++ b/Engine/Source/Runtime/Experimental/Chaos/Private/Chaos/PerParticleDampVelocity.cpp
@@ -38,10 +38,11 @@ void FPerParticleDampVelocity::UpdatePositionBasedState(const FPBDParticles& Par
 
 		for (int32 Index = Offset; Index < Range; ++Index)
 		{
-			if (!Particles.InvM(Index))
+			if (Particles.InvM(Index) == (FReal)0)
 			{
 				continue;
 			}
+
 			MXcm += Particles.X(Index) * Particles.M(Index);
 			MVcm += Particles.V(Index) * Particles.M(Index);
 			Mcm += Particles.M(Index);
@@ -57,10 +58,11 @@ void FPerParticleDampVelocity::UpdatePositionBasedState(const FPBDParticles& Par
 		FMatrix33 I(0);
 		for (int32 Index = Offset; Index < Range; ++Index)
 		{
-			if (!Particles.InvM(Index))
+			if (Particles.InvM(Index) == (FReal)0)
 			{
 				continue;
 			}
+
 			const FVec3 V = Particles.X(Index) - MXcm;
 			L += FVec3::CrossProduct(V, Particles.M(Index) * Particles.V(Index));
 			const FMatrix33 M(0, V[2], -V[1], -V[2], 0, V[0], V[1], -V[0], 0);
diff --git a/Engine/Source/Runtime/Experimental/Chaos/Private/Chaos/PerParticlePBDCollisionConstraint.cpp b/Engine/Source/Runtime/Experimental/Chaos/Private/Chaos/PerParticlePBDCollisionConstraint.cpp
--- a/Engine/Source/Runtime/Experimental/Chaos/Private/Chaos/PerParticlePBDCollisionConstraint.cpp
+++ b/Engine/Source/Runtime/Experimental/Chaos/Private/Chaos/PerParticlePBDCollisionConstraint.cpp
@@ -63,7 +63,7 @@ void TPerParticlePBDCollisionConstraint<EGeometryParticlesSimType::Other>::Apply
 	const FReal PerGroupFriction = MPerGroupFriction[DynamicGroupId];
 	const FReal PerGroupThickness = MPerGroupThickness[DynamicGroupId];
 
-	const int32 NumBatches = FMath::CeilToInt((Range - Offset) / (FReal)Chaos_PerParticleCollision_ISPC_ParallelBatchSize);
+	const int32 NumBatches = FMath::CeilToInt((FReal)(Range - Offset) / (FReal)Chaos_PerParticleCollision_ISPC_ParallelBatchSize);
 
 	if (Chaos_PerParticleCollision_ISPC_FastFriction)
 	{
diff --git a/Engine/Source/Runtime/Experimental/Chaos/Private/Chaos/SpatialHash.cpp b/Engine/Source/Runtime/Experimental/Chaos/Private/Chaos/SpatialHash.cpp
--- a/Engine/Source/Runtime/Experimental/Chaos/Private/Chaos/SpatialHash.cpp
+++ b/Engine/Source/Runtime/Experimental/Chaos/Private/Chaos/SpatialHash.cpp
@@ -13,7 +13,7 @@ namespace Chaos
 		double Time = 0.0;
 		FDurationTimer Timer(Time);
 
-		MCellSize = 2.0 * Radius;
+		MCellSize = (T)2.0 * Radius;
 
 		MBoundingBox = TAABB<T, 3>(TVector<T, 3>(0.0), Chaos::TVector<T, 3>(0.0));
 		for (int32 Idx = 0; Idx < MParticles.Num(); ++Idx)
@@ -26,13 +26,13 @@ namespace Chaos
 		T SmallestAxisLength = Extents[SmallestAxis()];
 		if (MCellSize > SmallestAxisLength)
 		{
-			MCellSize = SmallestAxisLength / 2.0;
+			MCellSize = SmallestAxisLength / (T)2.0;
 		}
 
 		T PrincipalAxisLength = Extents[MBoundingBox.LargestAxis()];
 		int32 NumberOfCellsOnPrincipalAxis = FMath::CeilToInt(PrincipalAxisLength / MCellSize);
 		MCellSize = PrincipalAxisLength / (T)NumberOfCellsOnPrincipalAxis;
-		T CellSizeInv = 1.0 / MCellSize;
+		T CellSizeInv = (T)1.0 / MCellSize;
 
 		MNumberOfCellsX = FMath::CeilToInt(Extents[0] * CellSizeInv) + 1;
 		MNumberOfCellsY = FMath::CeilToInt(Extents[1] * CellSizeInv) + 1;
@@ -72,12 +72,12 @@ namespace Chaos
 
 		// MCellSize needs to be smaller than smallest axis length
 		T SmallestAxisLength = Extents[SmallestAxis()];
-		MCellSize = SmallestAxisLength / 10.0;
+		MCellSize = SmallestAxisLength / (T)10.0;
 
 		T PrincipalAxisLength = Extents[MBoundingBox.LargestAxis()];
 		int32 NumberOfCellsOnPrincipalAxis = FMath::CeilToInt(PrincipalAxisLength / MCellSize);
 		MCellSize = PrincipalAxisLength / (T)NumberOfCellsOnPrincipalAxis;
-		T CellSizeInv = 1.0 / MCellSize;
+		T CellSizeInv = (T)1.0 / MCellSize;
 
 		MNumberOfCellsX = FMath::CeilToInt(Extents[0] * CellSizeInv) + 1;
 		MNumberOfCellsY = FMath::CeilToInt(Extents[1] * CellSizeInv) + 1;
@@ -370,8 +370,8 @@ namespace Chaos
 	template<class T>
 	void TSpatialHash<T>::ComputeGridXYZ(const TVector<T, 3>& Particle, int32& XIndex, int32& YIndex, int32& ZIndex)
 	{
-		T CellSizeInv = 1.0 / MCellSize;
-		TVector<T, 3> Location = Particle - MBoundingBox.Min() + TVector<T, 3>(0.5 * MCellSize);
+		T CellSizeInv = (T)1.0 / MCellSize;
+		TVector<T, 3> Location = Particle - MBoundingBox.Min() + TVector<T, 3>((T)0.5 * MCellSize);
 		XIndex = (int32)(Location.X * CellSizeInv);
 		YIndex = (int32)(Location.Y * CellSizeInv);
 		ZIndex = (int32)(Location.Z * CellSizeInv);
diff --git a/Engine/Source/Runtime/Experimental/Chaos/Private/Chaos/TriangleMesh.cpp b/Engine/Source/Runtime/Experimental/Chaos/Private/Chaos/TriangleMesh.cpp
--- a/Engine/Source/Runtime/Experimental/Chaos/Private/Chaos/TriangleMesh.cpp
+++ b/Engine/Source/Runtime/Experimental/Chaos/Private/Chaos/TriangleMesh.cpp
@@ -368,17 +368,17 @@ void AddTrianglesToHull(const TConstArrayView<FVec3>& Points, const int32 I0, co
 		FVec3 Normal1 = FVec3::CrossProduct(V1, V2).GetSafeNormal();
 		if (FVec3::DotProduct(Normal1, X2 - X0) > 0)
 		{
-			Normal1 *= -1;
+			Normal1 *= -1.0f;
 		}
 		FVec3 Normal2 = FVec3::CrossProduct(V1, V3).GetSafeNormal();
 		if (FVec3::DotProduct(Normal2, X1 - X0) > 0)
 		{
-			Normal2 *= -1;
+			Normal2 *= -1.0f;
 		}
 		FVec3 Normal3 = FVec3::CrossProduct(V2, V3).GetSafeNormal();
 		if (FVec3::DotProduct(Normal3, X0 - X1) > 0)
 		{
-			Normal3 *= -1;
+			Normal3 *= -1.0f;
 		}
 		TPlane<FReal, 3> NewPlane1(NewX, Normal1);
 		TPlane<FReal, 3> NewPlane2(NewX, Normal2);
@@ -781,7 +781,7 @@ TMap<int32, int32> FTriangleMesh::FindCoincidentVertexRemappings(
 		return Remappings;
 	}
 
-	LocalBBox.Thicken(1.0e-3);
+	LocalBBox.Thicken(1.0e-3f);
 	const FVec3 LocalCenter = LocalBBox.Center();
 	const FVec3& LocalMin = LocalBBox.Min();
 
@@ -793,8 +793,8 @@ TMap<int32, int32> FTriangleMesh::FindCoincidentVertexRemappings(
 	TMap<int64, TSet<int32>> OccupiedCells;
 	OccupiedCells.Reserve(NumPoints);
 
-	const int64 Resolution = static_cast<int64>(floor(MaxBBoxDim / 0.01));
-	const FReal CellSize = MaxBBoxDim / Resolution;
+	const int64 Resolution = static_cast<int64>(floor(MaxBBoxDim / 0.01f));
+	const FReal CellSize = static_cast<FReal>(static_cast<double>(MaxBBoxDim) / static_cast<double>(Resolution));
 	for (int i = 0; i < 2; i++)
 	{
 		OccupiedCells.Reset();
@@ -802,7 +802,7 @@ TMap<int32, int32> FTriangleMesh::FindCoincidentVertexRemappings(
 		// Shift the grid by 1/2 a grid cell the second iteration so that
 		// we don't miss slightly adjacent coincident points across cell
 		// boundaries.
-		const FVec3 GridCenter = LocalCenter - FVec3(i * CellSize / 2);
+		const FVec3 GridCenter = LocalCenter - FVec3(static_cast<FReal>(i) * CellSize / 2.0f);
 		for (int32 LocalIdx = 0; LocalIdx < NumPoints; LocalIdx++)
 		{
 			const int32 Idx = TestIndices[LocalIdx];
@@ -814,9 +814,9 @@ TMap<int32, int32> FTriangleMesh::FindCoincidentVertexRemappings(
 
 			const FVec3& Pos = LocalPoints[LocalIdx];
 			const TVec3<int64> Coord(
-				static_cast<int64>(floor((Pos[0] - GridCenter[0]) / CellSize + Resolution / 2)),
-				static_cast<int64>(floor((Pos[1] - GridCenter[1]) / CellSize + Resolution / 2)),
-				static_cast<int64>(floor((Pos[2] - GridCenter[2]) / CellSize + Resolution / 2)));
+				static_cast<int64>(FMath::Floor((Pos[0] - GridCenter[0]) / CellSize + static_cast<double>(Resolution) / 2.0f)),
+				static_cast<int64>(FMath::Floor((Pos[1] - GridCenter[1]) / CellSize + static_cast<double>(Resolution) / 2.0f)),
+				static_cast<int64>(FMath::Floor((Pos[2] - GridCenter[2]) / CellSize + static_cast<double>(Resolution) / 2.0f)));
 			const int64 FlatIdx =
 				((Coord[0] * Resolution + Coord[1]) * Resolution) + Coord[2];
 
@@ -1012,7 +1012,7 @@ TArray<int32> FTriangleMesh::GetVertexImportanceOrdering(
 		LocalPoints[i] = Points[Offset + i] - Center;
 		LocalBBox.GrowToInclude(LocalPoints[i]);
 	}
-	LocalBBox.Thicken(1.0e-3);
+	LocalBBox.Thicken(1.0e-3f);
 	const FVec3 LocalCenter = LocalBBox.Center();
 	const FVec3& LocalMin = LocalBBox.Min();
 
@@ -1053,7 +1053,7 @@ TArray<int32> FTriangleMesh::GetVertexImportanceOrdering(
 	AscendingPredicate<uint8> AscendingRankPred(Rank, Offset); // low to high
 	{
 		const int64 Resolution = static_cast<int64>(floor(MaxBBoxDim / 0.01));
-		const FReal CellSize = MaxBBoxDim / Resolution;
+		const FReal CellSize = static_cast<FReal>(static_cast<double>(MaxBBoxDim) / static_cast<double>(Resolution));
 		for (int i = 0; i < 2; i++)
 		{
 			OccupiedCells.Reset();
@@ -1062,16 +1062,16 @@ TArray<int32> FTriangleMesh::GetVertexImportanceOrdering(
 			// Shift the grid by 1/2 a grid cell the second iteration so that
 			// we don't miss slightly adjacent coincident points across cell
 			// boundaries.
-			const FVec3 GridCenter = LocalCenter - FVec3(i * CellSize / 2);
+			const FVec3 GridCenter = LocalCenter - FVec3(static_cast<FReal>(i) * CellSize / 2);
 			const int NumCoincidentPrev = NumCoincident;
 			for (int j = 0; j < NumPoints - NumCoincidentPrev; j++)
 			{
 				const int32 Idx = PointOrder[j];
 				const FVec3& Pos = LocalPoints[Idx - Offset];
 				const TVec3<int64> Coord(
-				    static_cast<int64>(floor((Pos[0] - GridCenter[0]) / CellSize + Resolution / 2)),
-				    static_cast<int64>(floor((Pos[1] - GridCenter[1]) / CellSize + Resolution / 2)),
-				    static_cast<int64>(floor((Pos[2] - GridCenter[2]) / CellSize + Resolution / 2)));
+				    static_cast<int64>(FMath::Floor((Pos[0] - GridCenter[0]) / CellSize + static_cast<double>(Resolution) / 2)),
+				    static_cast<int64>(FMath::Floor((Pos[1] - GridCenter[1]) / CellSize + static_cast<double>(Resolution) / 2)),
+				    static_cast<int64>(FMath::Floor((Pos[2] - GridCenter[2]) / CellSize + static_cast<double>(Resolution) / 2)));
 				const int64 FlatIdx =
 				    ((Coord[0] * Resolution + Coord[1]) * Resolution) + Coord[2];
 
@@ -1106,7 +1106,7 @@ TArray<int32> FTriangleMesh::GetVertexImportanceOrdering(
 		const int32 Resolution = i;
 		check(Resolution > 0);
 		check(Resolution % 2 == 0);
-		const FReal CellSize = MaxBBoxDim / Resolution;
+		const FReal CellSize = MaxBBoxDim / static_cast<FReal>(Resolution);
 
 		// The order in which we process these points matters.  Must do
 		// the current highest rank first.
@@ -1116,9 +1116,9 @@ TArray<int32> FTriangleMesh::GetVertexImportanceOrdering(
 			const FVec3& Pos = LocalPoints[Idx - Offset];
 			// grid center co-located at bbox center:
 			const TVec3<int64> Coord(
-			    static_cast<int64>(floor((Pos[0] - LocalCenter[0]) / CellSize)) + Resolution / 2,
-			    static_cast<int64>(floor((Pos[1] - LocalCenter[1]) / CellSize)) + Resolution / 2,
-			    static_cast<int64>(floor((Pos[2] - LocalCenter[2]) / CellSize)) + Resolution / 2);
+			    static_cast<int64>(FMath::Floor((Pos[0] - LocalCenter[0]) / CellSize)) + Resolution / 2,
+			    static_cast<int64>(FMath::Floor((Pos[1] - LocalCenter[1]) / CellSize)) + Resolution / 2,
+			    static_cast<int64>(FMath::Floor((Pos[2] - LocalCenter[2]) / CellSize)) + Resolution / 2);
 			const int64 FlatIdx =
 			    ((Coord[0] * Resolution + Coord[1]) * Resolution) + Coord[2];
 
diff --git a/Engine/Source/Runtime/Experimental/Chaos/Private/Chaos/TriangleMeshImplicitObject.cpp b/Engine/Source/Runtime/Experimental/Chaos/Private/Chaos/TriangleMeshImplicitObject.cpp
--- a/Engine/Source/Runtime/Experimental/Chaos/Private/Chaos/TriangleMeshImplicitObject.cpp
+++ b/Engine/Source/Runtime/Experimental/Chaos/Private/Chaos/TriangleMeshImplicitObject.cpp
@@ -101,7 +101,7 @@ struct FTriangleMeshRaycastVisitor
 	template <ERaycastType SQType>
 	bool Visit(int32 TriIdx, FQueryFastData& CurData)
 	{
-		constexpr FReal Epsilon = 1e-4;
+		constexpr FReal Epsilon = 1e-4f;
 		constexpr FReal Epsilon2 = Epsilon * Epsilon;
 		const FReal Thickness2 = SQType == ERaycastType::Sweep ? Thickness * Thickness : 0;
 		FReal MinTime = 0;	//no need to initialize, but fixes warning
@@ -454,7 +454,7 @@ bool FTriangleMeshImplicitObject::OverlapImp(const TArray<TVec3<IdxType>>& Eleme
 	QueryBounds.Thicken(Thickness);
 	const TArray<int32> PotentialIntersections = BVH.FindAllIntersections(QueryBounds);
 
-	const FReal Epsilon = 1e-4;
+	const FReal Epsilon = 1e-4f;
 	//ensure(Thickness > Epsilon);	//There's no hope for this to work unless thickness is large (really a sphere overlap test)
 	//todo: turn ensure back on, off until some other bug is fixed
 
@@ -849,7 +849,7 @@ int32 FTriangleMeshImplicitObject::FindMostOpposingFace(const TArray<TVec3<IdxTy
 	FAABB3 QueryBounds(Position - FVec3(SearchDist), Position + FVec3(SearchDist));
 
 	const TArray<int32> PotentialIntersections = BVH.FindAllIntersections(QueryBounds);
-	const FReal Epsilon = 1e-4;
+	const FReal Epsilon = 1e-4f;
 
 	FReal MostOpposingDot = TNumericLimits<FReal>::Max();
 	int32 MostOpposingFace = HintFaceIndex;
diff --git a/Engine/Source/Runtime/Experimental/Chaos/Private/Field/FieldSystemNodes.cpp b/Engine/Source/Runtime/Experimental/Chaos/Private/Field/FieldSystemNodes.cpp
--- a/Engine/Source/Runtime/Experimental/Chaos/Private/Field/FieldSystemNodes.cpp
+++ b/Engine/Source/Runtime/Experimental/Chaos/Private/Field/FieldSystemNodes.cpp
@@ -4,6 +4,7 @@
 #include "Field/FieldSystemNoiseAlgo.h"
 #include "Async/ParallelFor.h"
 #include "Chaos/Vector.h"
+#include <type_traits>
 
 FFieldNodeBase * FieldNodeFactory(FFieldNodeBase::EFieldType BaseType,FFieldNodeBase::ESerializationType Type)
 {
@@ -131,7 +132,7 @@ void FRadialIntMask::Evaluate(FFieldContext& Context, TArrayView<int32>& Results
 	{
 		const FFieldContextIndex& Index = Context.SampleIndices[SampleIndex];
 		{
-			float Result;
+			int32 Result;
 			float Delta2 = (Position - Context.Samples[Index.Sample]).SizeSquared();
 
 			if(Delta2 < Radius2)
@@ -148,12 +149,14 @@ void FRadialIntMask::Evaluate(FFieldContext& Context, TArrayView<int32>& Results
 			case ESetMaskConditionType::Field_Set_Always:
 				Results[Index.Result] = Result;
 				break;
+
 			case ESetMaskConditionType::Field_Set_IFF_NOT_Interior:
 				if (Results[Index.Result] != InteriorValue) 
 				{
 					Results[Index.Result] = Result;
 				}
 				break;
+
 			case ESetMaskConditionType::Field_Set_IFF_NOT_Exterior:
 				if (Results[Index.Result] != ExteriorValue) 
 				{
@@ -221,9 +224,9 @@ bool FUniformScalar::operator==(const FFieldNodeBase& Node)
 */
 void FWaveScalar::Evaluate(FFieldContext& Context, TArrayView<float>& Results) const
 {
-	const float Velocity = (Period != 0.0 ) ? Wavelength / Period : 0.0;
+	const float Velocity = (Period != 0.0f ) ? Wavelength / Period : 0.0f;
 
-	const float Wavenumber = (Wavelength != 0.0 ) ? 2.0 * PI / Wavelength : 0.0;
+	const float Wavenumber = (Wavelength != 0.0f ) ? 2.0f * PI / Wavelength : 0.0f;
 	const float DeltaTime = FMath::Max(Context.TimeSeconds, 0.0f);
 	const float Radius = Wavelength * DeltaTime / Period;
 	const float Decay = DeltaTime / Period;
@@ -248,7 +251,7 @@ void FWaveScalar::Evaluate(FFieldContext& Context, TArrayView<float>& Results) c
 		{
 			if (Distance < Radius && Radius > 0)
 			{
-				const float Fraction = (1.0 - Distance / Radius);
+				const float Fraction = (1.0f - Distance / Radius);
 				if (Falloff == EFieldFalloffType::Field_FallOff_None)
 				{
 					Results[Index.Result] = Magnitude;
@@ -261,18 +264,18 @@ void FWaveScalar::Evaluate(FFieldContext& Context, TArrayView<float>& Results) c
 				{
 					Results[Index.Result] =  Magnitude * Fraction * Fraction;
 				}
-				else if (Falloff == EFieldFalloffType::Field_Falloff_Inverse && Fraction > 0.0)
+				else if (Falloff == EFieldFalloffType::Field_Falloff_Inverse && Fraction > 0.0f)
 				{
 					Results[Index.Result] =  Magnitude / Fraction;
 				}
 				else if (Falloff == EFieldFalloffType::Field_Falloff_Logarithmic)
 				{
-					Results[Index.Result] = Magnitude * FMath::Loge(Fraction + 1.0) / FMath::Loge(10.0);
+					Results[Index.Result] = Magnitude * FMath::Loge(Fraction + 1.0f) / FMath::Loge(10.0f);
 				}
 			}
 			else
 			{
-				Results[Index.Result] = 0.0;
+				Results[Index.Result] = 0.0f;
 			}
 		}
 		else if (Function == EWaveFunctionType::Field_Wave_Decay)
@@ -343,13 +346,13 @@ float EvalFalloffFunction<EFieldFalloffType::Field_Falloff_Squared>(const float&
 template<>
 float EvalFalloffFunction<EFieldFalloffType::Field_Falloff_Inverse>(const float& MinRange, const float& DeltaRange, const float& NodeMagnitude, const float& FalloffValue)
 {
-	return (FalloffValue > SMALL_NUMBER) ? ScaleFunctionResult(MinRange, DeltaRange, NodeMagnitude, 1.0 / FalloffValue) : 0.0;
+	return (FalloffValue > SMALL_NUMBER) ? ScaleFunctionResult(MinRange, DeltaRange, NodeMagnitude, 1.0f / FalloffValue) : 0.0f;
 }
 
 template<>
 float EvalFalloffFunction<EFieldFalloffType::Field_Falloff_Logarithmic>(const float& MinRange, const float& DeltaRange, const float& NodeMagnitude, const float& FalloffValue)
 {
-	return ScaleFunctionResult(MinRange, DeltaRange, NodeMagnitude, FMath::LogX(10, FalloffValue + 1.0));
+	return ScaleFunctionResult(MinRange, DeltaRange, NodeMagnitude, FMath::LogX(10.0f, FalloffValue + 1.0f));
 }
 
 /**
@@ -372,7 +375,7 @@ void FRadialFalloff::Evaluator(const FFieldContext& Context, TArrayView<float>&
 
 				if (Delta < Radius)
 				{
-					const float Function = 1.0 - Delta / Radius;
+					const float Function = 1.0f - Delta / Radius;
 					Results[Index.Result] = EvalFalloffFunction<FalloffType>(MinRange, DeltaRange, Magnitude, Function);
 				}
 			}
@@ -450,7 +453,7 @@ void FPlaneFalloff::Evaluator(const FFieldContext& Context, const FPlane& Plane,
 
 				if (Delta < -SMALL_NUMBER && Delta > -Distance)
 				{
-					const float Function = 1.0 + Delta / Distance;
+					const float Function = 1.0f + Delta / Distance;
 					Results[Index.Result] = EvalFalloffFunction<FalloffType>(MinRange, DeltaRange, Magnitude, Function);
 				}
 			}
@@ -608,17 +611,17 @@ FNoiseField::Evaluate(FFieldContext& Context, TArrayView<float>& Results) const
 
 		float Dummy = 0.0f;
 		FVector LocalPoint = Transform.InverseTransformPosition(Context.Samples[Index.Sample]);
-		LocalPoint = FVector(FMath::Modf(LocalPoint.X, &Dummy) * 0.5 + 0.5, 
-							 FMath::Modf(LocalPoint.Y, &Dummy) * 0.5 + 0.5,
-							 FMath::Modf(LocalPoint.Z, &Dummy) * 0.5 + 0.5) * 255;
+		LocalPoint = FVector(FMath::Modf(LocalPoint.X, &Dummy) * 0.5f + 0.5f, 
+							 FMath::Modf(LocalPoint.Y, &Dummy) * 0.5f + 0.5f,
+							 FMath::Modf(LocalPoint.Z, &Dummy) * 0.5f + 0.5f) * 255;
 
 		// Samples for the Perlin noise must be btw 0->255
-		float PerlinValue = 0.0;
+		float PerlinValue = 0.0f;
 		Field::PerlinNoise::Sample(&PerlinValue, LocalPoint.X, LocalPoint.Y, LocalPoint.Z);
 
 		// Perlin noise result is btw -1 -> 1
-		PerlinValue = 0.5 * ( PerlinValue + 1.0 );
-		Results[Index.Result] = ScaleFunctionResult(MinRange, DeltaRange, 1.0, PerlinValue);
+		PerlinValue = 0.5f * ( PerlinValue + 1.0f );
+		Results[Index.Result] = ScaleFunctionResult(MinRange, DeltaRange, 1.0f, PerlinValue);
 	}
 }
 void FNoiseField::Serialize(FArchive& Ar)
@@ -999,10 +1002,13 @@ bool FSumVector::operator==(const FFieldNodeBase& Node)
 template<class InT, class OutT>
 void FConversionField<InT,OutT>::Evaluate(FFieldContext& Context, TArrayView<OutT>& Results) const
 {
+	static_assert(std::is_arithmetic_v<InT>, "Arithmetic types required for field conversion");
+	static_assert(std::is_arithmetic_v<OutT>, "Arithmetic types required for field conversion");
+
 	int32 NumSamples = Context.SampleIndices.Num();
 
 	TArray<InT> Array;
-	Array.Init(0.f, Results.Num());
+	Array.Init(static_cast<InT>(0), Results.Num());
 	TArrayView<InT> ArrayView(&(Array[0]), Array.Num());
 	InputField->Evaluate(Context, ArrayView);
 
diff --git a/Engine/Source/Runtime/Experimental/Chaos/Private/Framework/TimeStep.cpp b/Engine/Source/Runtime/Experimental/Chaos/Private/Framework/TimeStep.cpp
--- a/Engine/Source/Runtime/Experimental/Chaos/Private/Framework/TimeStep.cpp
+++ b/Engine/Source/Runtime/Experimental/Chaos/Private/Framework/TimeStep.cpp
@@ -27,9 +27,9 @@ namespace Chaos
 	void FFixedTimeStep::Update()
 	{
 		CurrentTime = FPlatformTime::Seconds();
-		ActualDt = CurrentTime - LastTime;
+		ActualDt = static_cast<float>(CurrentTime - LastTime);
 
-		if((float)ActualDt > TargetDt)
+		if(ActualDt > TargetDt)
 		{
 			UE_LOG(LogChaosDebug, Verbose, TEXT("PhysAdvance: Exceeded requested Dt of %.3f (%.2fFPS). Ran for %.3f"), TargetDt, 1.0f / TargetDt, ActualDt);
 		}
@@ -39,7 +39,7 @@ namespace Chaos
 
 			// #BG TODO need some way to handle abandonning this when the gamethread requests a sync
 			// Or just running more commands in general otherwise this is dead time.
-			FPlatformProcess::Sleep((float)(TargetDt - ActualDt));
+			FPlatformProcess::Sleep(TargetDt - ActualDt);
 		}
 
 		LastTime = FPlatformTime::Seconds();
@@ -76,7 +76,7 @@ namespace Chaos
 	void FVariableTimeStep::Update()
 	{
 		double CurrentTime = FPlatformTime::Seconds();
-		Dt = FPlatformTime::Seconds() - LastTime;
+		Dt = static_cast<float>(FPlatformTime::Seconds() - LastTime);
 		LastTime = CurrentTime;
 	}
 
@@ -97,11 +97,11 @@ namespace Chaos
 	void FVariableMinimumWithCapTimestep::Update()
 	{
 		CurrentTime = FPlatformTime::Seconds();
-		ActualDt = CurrentTime - LastTime;
+		ActualDt = static_cast<float>(CurrentTime - LastTime);
 
 		if(ActualDt < TargetDt)
 		{
-			FPlatformProcess::Sleep((float)(TargetDt - ActualDt));
+			FPlatformProcess::Sleep(TargetDt - ActualDt);
 			Dt = TargetDt;
 		}
 		else
@@ -139,7 +139,7 @@ namespace Chaos
 	void FVariableWithCapTimestep::Update()
 	{
 		CurrentTime = FPlatformTime::Seconds();
-		ActualDt = CurrentTime - LastTime;
+		ActualDt = static_cast<float>(CurrentTime - LastTime);
 
 		float Cap = CVarVariableTickCap.GetValueOnAnyThread();
 		Dt = FMath::Min(ActualDt, Cap);
diff --git a/Engine/Source/Runtime/Experimental/Chaos/Private/GeometryCollection/GeometryCollection.cpp b/Engine/Source/Runtime/Experimental/Chaos/Private/GeometryCollection/GeometryCollection.cpp
--- a/Engine/Source/Runtime/Experimental/Chaos/Private/GeometryCollection/GeometryCollection.cpp
+++ b/Engine/Source/Runtime/Experimental/Chaos/Private/GeometryCollection/GeometryCollection.cpp
@@ -240,7 +240,11 @@ int32 FGeometryCollection::AppendGeometry(const FGeometryCollection & Element, i
 		{
 			Center += Vertices[vdx];
 		}
-		if (NumNewVertices) Center /= NumNewVertices;
+
+		if(NumNewVertices)
+		{
+			Center /= static_cast<float>(NumNewVertices);
+		}
 
 		//
 		//  Inner/Outer Radius
@@ -270,7 +274,7 @@ int32 FGeometryCollection::AppendGeometry(const FGeometryCollection & Element, i
 				{
 					Centroid += Vertices[FaceIndices[fdx][e]];
 				}
-				Centroid /= 3;
+				Centroid /= 3.0f;
 
 				float Delta = (Center - Centroid).Size();
 				InnerR[GeometryIndex] = FMath::Min(InnerR[GeometryIndex], Delta);
diff --git a/Engine/Source/Runtime/Experimental/Chaos/Private/GeometryCollection/GeometryCollectionAlgo.cpp b/Engine/Source/Runtime/Experimental/Chaos/Private/GeometryCollection/GeometryCollectionAlgo.cpp
--- a/Engine/Source/Runtime/Experimental/Chaos/Private/GeometryCollection/GeometryCollectionAlgo.cpp
+++ b/Engine/Source/Runtime/Experimental/Chaos/Private/GeometryCollection/GeometryCollectionAlgo.cpp
@@ -157,7 +157,7 @@ namespace GeometryCollectionAlgo
 		}
 		if (NumIndices > 1)
 		{
-			Translation /= NumIndices;
+			Translation /= static_cast<float>(NumIndices);
 		}
 		return Translation;
 	}
@@ -520,7 +520,7 @@ namespace GeometryCollectionAlgo
 			{
 				if (SurfaceParticlesCount[i])
 				{
-					CenterOfMass[i] /= SurfaceParticlesCount[i];
+					CenterOfMass[i] /= static_cast<float>(SurfaceParticlesCount[i]);
 
 					FTransform Tmp(CenterOfMass[i]);
 
@@ -531,7 +531,7 @@ namespace GeometryCollectionAlgo
 				}
 			}
 
-			CombinedCenterOfMassWorld /= Transform.Num();
+			CombinedCenterOfMassWorld /= static_cast<float>(Transform.Num());
 
 			for (int i = 0; i < Vertex.Num(); i++)
 			{
diff --git a/Engine/Source/Runtime/Experimental/Chaos/Private/GeometryCollection/GeometryCollectionClusteringUtility.cpp b/Engine/Source/Runtime/Experimental/Chaos/Private/GeometryCollection/GeometryCollectionClusteringUtility.cpp
--- a/Engine/Source/Runtime/Experimental/Chaos/Private/GeometryCollection/GeometryCollectionClusteringUtility.cpp
+++ b/Engine/Source/Runtime/Experimental/Chaos/Private/GeometryCollection/GeometryCollectionClusteringUtility.cpp
@@ -130,7 +130,7 @@ void FGeometryCollectionClusteringUtility::ClusterAllBonesUnderNewRoot(FGeometry
 		}
 		ExplodedTransforms[RootNoneIndex] = Transforms[RootNoneIndex];
 		// This bones offset is the average of all the selected bones
-		ExplodedVectors[RootNoneIndex] = SumOfOffsets / ChildBones.Num();
+		ExplodedVectors[RootNoneIndex] = SumOfOffsets / static_cast<float>(ChildBones.Num());
 	}
 
 	// Selected Bone Setup
diff --git a/Engine/Source/Runtime/Experimental/Chaos/Private/GeometryCollection/GeometryCollectionCollisionStructureManager.cpp b/Engine/Source/Runtime/Experimental/Chaos/Private/GeometryCollection/GeometryCollectionCollisionStructureManager.cpp
--- a/Engine/Source/Runtime/Experimental/Chaos/Private/GeometryCollection/GeometryCollectionCollisionStructureManager.cpp
+++ b/Engine/Source/Runtime/Experimental/Chaos/Private/GeometryCollection/GeometryCollectionCollisionStructureManager.cpp
@@ -52,7 +52,7 @@ FCollisionStructureManager::NewSimplicial(
 		{
 			Extent = Implicit->HasBoundingBox() ? Implicit->BoundingBox().Extents().Size() : 1.f;
 
-			Chaos::FReal Threshold = Extent * 0.01;
+			Chaos::FReal Threshold = Extent * 0.01f;
 
 			//
 			//  Remove particles inside the levelset. (I think this is useless) 
diff --git a/Engine/Source/Runtime/Experimental/Chaos/Private/GeometryCollection/GeometryCollectionProximityUtility.cpp b/Engine/Source/Runtime/Experimental/Chaos/Private/GeometryCollection/GeometryCollectionProximityUtility.cpp
--- a/Engine/Source/Runtime/Experimental/Chaos/Private/GeometryCollection/GeometryCollectionProximityUtility.cpp
+++ b/Engine/Source/Runtime/Experimental/Chaos/Private/GeometryCollection/GeometryCollectionProximityUtility.cpp
@@ -120,7 +120,7 @@ void FGeometryCollectionProximityUtility::UpdateProximity(FGeometryCollection* G
 	TManagedArray<FVector3f>& BreakingRegionNormalArray = GeometryCollection->GetAttribute<FVector3f>("BreakingRegionNormal", FGeometryCollection::BreakingGroup);
 	TManagedArray<float>& BreakingRegionRadiusArray = GeometryCollection->GetAttribute<float>("BreakingRegionRadius", FGeometryCollection::BreakingGroup);
 
-	float DistanceThreshold = 1e-2;
+	float DistanceThreshold = 1e-2f;
 	float DistanceThresholdSquared = DistanceThreshold * DistanceThreshold;
 	TArray<FFaceTransformData> FaceTransformDataArray;
 	FaceTransformDataArray.Empty();
@@ -337,7 +337,7 @@ void FGeometryCollectionProximityUtility::UpdateProximity(FGeometryCollection* G
 					FVector OtherEdge2(VertexInWorldArray[IndicesArray[IdxOtherFace][2]] - VertexInWorldArray[IndicesArray[IdxOtherFace][0]]);
 					FVector OtherFaceN(OtherEdge1 ^ OtherEdge2);
 
-					if (FVector::Parallel(FaceN, OtherFaceN, 1e-1))
+					if (FVector::Parallel(FaceN, OtherFaceN, 1e-1f))
 					{
 						FVector FaceCenter((VertexInWorldArray[IndicesArray[IdxFace][0]] + VertexInWorldArray[IndicesArray[IdxFace][1]] + VertexInWorldArray[IndicesArray[IdxFace][2]]) / 3.f);
 						FVector OtherFaceCenter = (VertexInWorldArray[IndicesArray[IdxOtherFace][0]] + VertexInWorldArray[IndicesArray[IdxOtherFace][1]] + VertexInWorldArray[IndicesArray[IdxOtherFace][2]]) / 3.f;
@@ -376,10 +376,10 @@ void FGeometryCollectionProximityUtility::UpdateProximity(FGeometryCollection* G
 								LocalOverlappingFacePairSet.Add(FOverlappingFacePair{ FMath::Min(IdxFace, IdxOtherFace), FMath::Max(IdxFace, IdxOtherFace) });
 							}
 						}
-						else if (IsPointInsideOfTriangle(FaceCenter, VertexInWorldArray[IndicesArray[IdxOtherFace][0]], VertexInWorldArray[IndicesArray[IdxOtherFace][1]], VertexInWorldArray[IndicesArray[IdxOtherFace][2]], 1e-1) ||
-							IsPointInsideOfTriangle(PointInFace1, VertexInWorldArray[IndicesArray[IdxOtherFace][0]], VertexInWorldArray[IndicesArray[IdxOtherFace][1]], VertexInWorldArray[IndicesArray[IdxOtherFace][2]], 1e-1) ||
-							IsPointInsideOfTriangle(PointInFace2, VertexInWorldArray[IndicesArray[IdxOtherFace][0]], VertexInWorldArray[IndicesArray[IdxOtherFace][1]], VertexInWorldArray[IndicesArray[IdxOtherFace][2]], 1e-1) ||
-							IsPointInsideOfTriangle(PointInFace3, VertexInWorldArray[IndicesArray[IdxOtherFace][0]], VertexInWorldArray[IndicesArray[IdxOtherFace][1]], VertexInWorldArray[IndicesArray[IdxOtherFace][2]], 1e-1))
+						else if (IsPointInsideOfTriangle(FaceCenter, VertexInWorldArray[IndicesArray[IdxOtherFace][0]], VertexInWorldArray[IndicesArray[IdxOtherFace][1]], VertexInWorldArray[IndicesArray[IdxOtherFace][2]], 1e-1f) ||
+							IsPointInsideOfTriangle(PointInFace1, VertexInWorldArray[IndicesArray[IdxOtherFace][0]], VertexInWorldArray[IndicesArray[IdxOtherFace][1]], VertexInWorldArray[IndicesArray[IdxOtherFace][2]], 1e-1f) ||
+							IsPointInsideOfTriangle(PointInFace2, VertexInWorldArray[IndicesArray[IdxOtherFace][0]], VertexInWorldArray[IndicesArray[IdxOtherFace][1]], VertexInWorldArray[IndicesArray[IdxOtherFace][2]], 1e-1f) ||
+							IsPointInsideOfTriangle(PointInFace3, VertexInWorldArray[IndicesArray[IdxOtherFace][0]], VertexInWorldArray[IndicesArray[IdxOtherFace][1]], VertexInWorldArray[IndicesArray[IdxOtherFace][2]], 1e-1f))
 						{
 							if (!LocalOverlappingFacePairSet.Contains(FOverlappingFacePair{ FMath::Min(IdxFace, IdxOtherFace), FMath::Max(IdxFace, IdxOtherFace) }))
 							{
@@ -393,10 +393,10 @@ void FGeometryCollectionProximityUtility::UpdateProximity(FGeometryCollection* G
 							PointInFace3 = (VertexInWorldArray[IndicesArray[IdxOtherFace][2]] + OtherFaceCenter) / 2.f;
 
 							// Check if points in OtherFace are in Face
-							if (IsPointInsideOfTriangle(OtherFaceCenter, VertexInWorldArray[IndicesArray[IdxFace][0]], VertexInWorldArray[IndicesArray[IdxFace][1]], VertexInWorldArray[IndicesArray[IdxFace][2]], 1e-1) ||
-								IsPointInsideOfTriangle(PointInFace1, VertexInWorldArray[IndicesArray[IdxFace][0]], VertexInWorldArray[IndicesArray[IdxFace][1]], VertexInWorldArray[IndicesArray[IdxFace][2]], 1e-1) ||
-								IsPointInsideOfTriangle(PointInFace2, VertexInWorldArray[IndicesArray[IdxFace][0]], VertexInWorldArray[IndicesArray[IdxFace][1]], VertexInWorldArray[IndicesArray[IdxFace][2]], 1e-1) ||
-								IsPointInsideOfTriangle(PointInFace3, VertexInWorldArray[IndicesArray[IdxFace][0]], VertexInWorldArray[IndicesArray[IdxFace][1]], VertexInWorldArray[IndicesArray[IdxFace][2]], 1e-1))
+							if (IsPointInsideOfTriangle(OtherFaceCenter, VertexInWorldArray[IndicesArray[IdxFace][0]], VertexInWorldArray[IndicesArray[IdxFace][1]], VertexInWorldArray[IndicesArray[IdxFace][2]], 1e-1f) ||
+								IsPointInsideOfTriangle(PointInFace1, VertexInWorldArray[IndicesArray[IdxFace][0]], VertexInWorldArray[IndicesArray[IdxFace][1]], VertexInWorldArray[IndicesArray[IdxFace][2]], 1e-1f) ||
+								IsPointInsideOfTriangle(PointInFace2, VertexInWorldArray[IndicesArray[IdxFace][0]], VertexInWorldArray[IndicesArray[IdxFace][1]], VertexInWorldArray[IndicesArray[IdxFace][2]], 1e-1f) ||
+								IsPointInsideOfTriangle(PointInFace3, VertexInWorldArray[IndicesArray[IdxFace][0]], VertexInWorldArray[IndicesArray[IdxFace][1]], VertexInWorldArray[IndicesArray[IdxFace][2]], 1e-1f))
 							{
 								if (!LocalOverlappingFacePairSet.Contains(FOverlappingFacePair{ FMath::Min(IdxFace, IdxOtherFace), FMath::Max(IdxFace, IdxOtherFace) }))
 								{
diff --git a/Engine/Source/Runtime/Experimental/Chaos/Private/GeometryCollection/GeometryCollectionUtility.cpp b/Engine/Source/Runtime/Experimental/Chaos/Private/GeometryCollection/GeometryCollectionUtility.cpp
--- a/Engine/Source/Runtime/Experimental/Chaos/Private/GeometryCollection/GeometryCollectionUtility.cpp
+++ b/Engine/Source/Runtime/Experimental/Chaos/Private/GeometryCollection/GeometryCollectionUtility.cpp
@@ -9,6 +9,7 @@
 
 #include "CoreMinimal.h"
 #include "Templates/SharedPointer.h"
+#include "Chaos/Real.h"
 
 DEFINE_LOG_CATEGORY_STATIC(FGeometryCollectionUtilityLogging, Log, All);
 
@@ -154,7 +155,7 @@ namespace GeometryCollection
 	{
 		check(RestCollectionIn.IsValid());
 
-		float domain = 10;
+		float domain = 10.f;
 		FVector Stack(domain);
 		float numElements = powf(domain, 3);
 
@@ -166,16 +167,16 @@ namespace GeometryCollection
 		FVector MinCorner = -Length * Expansion / 2.f * Stackf;
 
 
-		for (int32 i = 0; i < Stack[0]; ++i)
+		for (int32 i = 0; i < FMath::TruncToInt(Stack[0]); ++i)
 		{
-			for (int32 j = 0; j < Stack[1]; ++j)
+			for (int32 j = 0; j < FMath::TruncToInt(Stack[1]); ++j)
 			{
-				for (int32 k = 0; k < Stack[2]; ++k)
+				for (int32 k = 0; k < FMath::TruncToInt(Stack[2]); ++k)
 				{
 					FVector Delta(j % 2 == 1 ? Length / 2.f : 0.f, 0.f, j % 2 == 1 ? Length / 2.f : 0.f);
-					FVector CenterOfMass = FVector(MinCorner[0] + Expansion * Length * i + Length * (Expansion / 2.f),
-						MinCorner[0] + Expansion * Length * j + Length * (Expansion / 2.f),
-						MinCorner[0] + Expansion * Length * k + Length * (Expansion / 2.f)) + Delta;
+					FVector CenterOfMass = FVector(MinCorner[0] + Expansion * Length * static_cast<Chaos::FReal>(i) + Length * (Expansion / 2.f),
+						MinCorner[0] + Expansion * Length * static_cast<Chaos::FReal>(j) + Length * (Expansion / 2.f),
+						MinCorner[0] + Expansion * Length * static_cast<Chaos::FReal>(k) + Length * (Expansion / 2.f)) + Delta;
 					TSharedPtr<FGeometryCollection> Element = MakeCubeElement(FTransform(CenterOfMass), FVector(Length) );
 					RestCollectionIn->AppendGeometry(*Element);
 				}
@@ -336,7 +337,7 @@ namespace GeometryCollection
 					{
 						if (VertexCount[GeometryIndex])
 						{
-							Center[GeometryIndex] /= VertexCount[GeometryIndex];
+							Center[GeometryIndex] /= static_cast<Chaos::FReal>(VertexCount[GeometryIndex]);
 						}
 					}
 
@@ -387,7 +388,7 @@ namespace GeometryCollection
 						{
 							Centroid += Vertex[FaceIndices[fdx][e]];
 						}
-						Centroid /= 3;
+						Centroid /= 3.0f;
 
 						float Delta = (Center[GeometryIndex] - Centroid).Size();
 						InnerRadius[GeometryIndex] = FMath::Min(InnerRadius[GeometryIndex], Delta);
diff --git a/Engine/Source/Runtime/Experimental/Chaos/Private/GeometryCollection/RecordedTransformTrack.cpp b/Engine/Source/Runtime/Experimental/Chaos/Private/GeometryCollection/RecordedTransformTrack.cpp
--- a/Engine/Source/Runtime/Experimental/Chaos/Private/GeometryCollection/RecordedTransformTrack.cpp
+++ b/Engine/Source/Runtime/Experimental/Chaos/Private/GeometryCollection/RecordedTransformTrack.cpp
@@ -96,8 +96,8 @@ FRecordedTransformTrack FRecordedTransformTrack::ProcessRawRecordedData(const FR
 		UScriptStruct* RecordedTrackStructType = FRecordedTransformTrack::StaticStruct();
 		RecordedTrackStructType->SerializeTaggedProperties(AfterAr, (uint8*)&RecordedData, RecordedTrackStructType, nullptr);
 
-		const int32 ArchiveBeforeSize = BeforeAr.GetNum();
-		const int32 ArchiveAfterSize = AfterAr.GetNum();
+		const int32 ArchiveBeforeSize = static_cast<int32>(BeforeAr.GetNum());
+		const int32 ArchiveAfterSize = static_cast<int32>(AfterAr.GetNum());
 
 		// Dump data reduction stats.
 		FNumberFormattingOptions Opts;
diff --git a/Engine/Source/Runtime/Experimental/Chaos/Private/PBDRigidsSolver.cpp b/Engine/Source/Runtime/Experimental/Chaos/Private/PBDRigidsSolver.cpp
--- a/Engine/Source/Runtime/Experimental/Chaos/Private/PBDRigidsSolver.cpp
+++ b/Engine/Source/Runtime/Experimental/Chaos/Private/PBDRigidsSolver.cpp
@@ -89,7 +89,7 @@ Chaos::DebugDraw::FChaosDebugDrawSettings ChaosSolverDebugDebugDrawSettings(
 	/* AngVelScale =		*/ 0.0f,
 	/* ImpulseScale =		*/ 0.0f,
 	/* InertiaScale =				*/ 1.0f,
-	/* DrawPriority =				*/ 10.0f,
+	/* DrawPriority =				*/ 10,
 	/* bShowSimple =				*/ true,
 	/* bShowComplex =				*/ false,
 	/* bInShowLevelSetCollision =	*/ true
@@ -716,11 +716,11 @@ namespace Chaos
 			{
 				SetCollisionPushOutPairIterations(ChaosSolverCollisionPushOutIterations);
 			}
-			if (ChaosSolverJointPairIterations >= 0.0f)
+			if (ChaosSolverJointPairIterations >= 0)
 			{
 				SetJointPairIterations(ChaosSolverJointPairIterations);
 			}
-			if (ChaosSolverJointPushOutPairIterations >= 0.0f)
+			if (ChaosSolverJointPushOutPairIterations >= 0)
 			{
 				SetJointPushOutPairIterations(ChaosSolverJointPushOutPairIterations);
 			}
@@ -767,7 +767,7 @@ namespace Chaos
 		//TODO: interpolate some data based on num steps
 
 		FPushPhysicsData* PushData = MarshallingManager.GetProducerData_External();
-		const FReal DynamicsWeight = FReal(1) / NumExternalSteps;
+		const FReal DynamicsWeight = FReal(1) / FReal(NumExternalSteps);
 		FDirtySet* DirtyProxiesData = &PushData->DirtyProxiesDataBuffer;
 		FDirtyPropertiesManager* Manager = &PushData->DirtyPropertiesManager;
 
diff --git a/Engine/Source/Runtime/Experimental/Chaos/Private/PhysicsProxy/GeometryCollectionPhysicsProxy.cpp b/Engine/Source/Runtime/Experimental/Chaos/Private/PhysicsProxy/GeometryCollectionPhysicsProxy.cpp
--- a/Engine/Source/Runtime/Experimental/Chaos/Private/PhysicsProxy/GeometryCollectionPhysicsProxy.cpp
+++ b/Engine/Source/Runtime/Experimental/Chaos/Private/PhysicsProxy/GeometryCollectionPhysicsProxy.cpp
@@ -702,8 +702,8 @@ void FGeometryCollectionPhysicsProxy::InitializeBodiesPT(Chaos::FPBDRigidsSolver
 					Mass[TransformGroupIndex],
 					InertiaTensor[TransformGroupIndex],
 					WorldTransform,
-					(uint8)DynamicState[TransformGroupIndex],
-					CollisionGroup[TransformGroupIndex]);
+					static_cast<uint8>(DynamicState[TransformGroupIndex]),
+					static_cast<int16>(CollisionGroup[TransformGroupIndex]));
 
 				if (Parameters.EnableClustering)
 				{
@@ -715,8 +715,8 @@ void FGeometryCollectionPhysicsProxy::InitializeBodiesPT(Chaos::FPBDRigidsSolver
 				CollisionParticles.Reset(Simplicials[TransformGroupIndex]?Simplicials[TransformGroupIndex]->NewCopy():nullptr); // @chaos(optimize) : maybe just move this memory instead. 
 				if (CollisionParticles)
 				{
-					int32 NumCollisionParticles = CollisionParticles->Size();
-					int32 CollisionParticlesSize = FMath::Max(0, FMath::Min(int(NumCollisionParticles * CollisionParticlesPerObjectFraction), NumCollisionParticles));
+					Chaos::FReal NumCollisionParticles = static_cast<Chaos::FReal>(CollisionParticles->Size());
+					int32 CollisionParticlesSize = FMath::TruncToInt(Chaos::FReal(FMath::Max(0, FMath::Min(int(NumCollisionParticles * CollisionParticlesPerObjectFraction), NumCollisionParticles))));
 					CollisionParticles->Resize(CollisionParticlesSize); // Truncates!
 				}
 
@@ -1007,9 +1007,9 @@ FGeometryCollectionPhysicsProxy::BuildClusters(
 	ClusterCreationParameters.ConnectionMethod = Parameters.ClusterConnectionMethod;
 	if (ClusterCreationParameters.CollisionParticles)
 	{
-		const int32 NumCollisionParticles = ClusterCreationParameters.CollisionParticles->Size();
+		const Chaos::FReal NumCollisionParticles = static_cast<Chaos::FReal>(ClusterCreationParameters.CollisionParticles->Size());
 		const int32 ClampedCollisionParticlesSize = 
-			FMath::Max(0, FMath::Min(int(NumCollisionParticles * CollisionParticlesPerObjectFraction), NumCollisionParticles));
+			FMath::TruncToInt(Chaos::FReal(FMath::Max(0, FMath::Min(int(NumCollisionParticles * CollisionParticlesPerObjectFraction), NumCollisionParticles))));
 		ClusterCreationParameters.CollisionParticles->Resize(ClampedCollisionParticlesSize);
 	}
 	TArray<Chaos::TPBDRigidParticleHandle<Chaos::FReal, 3>*> ChildHandlesCopy(ChildHandles);
@@ -1961,10 +1961,10 @@ void FGeometryCollectionPhysicsProxy::InitializeSharedCollisionStructures(
 					{
 						FVector Extents = BoundingBox[GeometryIndex].GetExtent();
 						MassProperties.Volume = Extents.X * Extents.Y * Extents.Z;
-						float ExtentsYZ = Extents.Y * Extents.Y + Extents.Z * Extents.Z;
-						float ExtentsXZ = Extents.X * Extents.X + Extents.Z * Extents.Z;
-						float ExtentsXY = Extents.X * Extents.X + Extents.Y * Extents.Y;
-						MassProperties.InertiaTensor = PMatrix<FReal, 3, 3>(ExtentsYZ / 12., ExtentsXZ / 12., ExtentsXY / 12.);
+						FRealSingle ExtentsYZ = Extents.Y * Extents.Y + Extents.Z * Extents.Z;
+						FRealSingle ExtentsXZ = Extents.X * Extents.X + Extents.Z * Extents.Z;
+						FRealSingle ExtentsXY = Extents.X * Extents.X + Extents.Y * Extents.Y;
+						MassProperties.InertiaTensor = PMatrix<FRealSingle, 3, 3>(ExtentsYZ / 12.f, ExtentsXZ / 12.f, ExtentsXY / 12.f);
 						MassProperties.CenterOfMass = BoundingBox[GeometryIndex].GetCenter();
 						CollectionMassToLocal[TransformGroupIndex] = FTransform(FQuat::Identity, MassProperties.CenterOfMass);
 						InertiaComputationNeeded[GeometryIndex] = false;
@@ -2376,11 +2376,11 @@ void FGeometryCollectionPhysicsProxy::InitializeSharedCollisionStructures(
 					const float ScaleMax = Scale.GetAbsMax();
 					const float ScaleMin = Scale.GetAbsMin();
 
-					float MinResolution = ScaleMin * SizeSpecificData.MinLevelSetResolution;
-					MinResolution = FMath::Clamp<float>(MinResolution, SizeSpecificData.MinLevelSetResolution, SizeSpecificData.MinClusterLevelSetResolution);
+					float MinResolution = ScaleMin * static_cast<float>(SizeSpecificData.MinLevelSetResolution);
+					MinResolution = FMath::Clamp<float>(MinResolution, static_cast<float>(SizeSpecificData.MinLevelSetResolution), static_cast<float>(SizeSpecificData.MinClusterLevelSetResolution));
 
-					float MaxResolution = ScaleMax * SizeSpecificData.MaxLevelSetResolution;
-					MaxResolution = FMath::Clamp<float>(MaxResolution, SizeSpecificData.MaxLevelSetResolution, SizeSpecificData.MaxClusterLevelSetResolution);
+					float MaxResolution = ScaleMax * static_cast<float>(SizeSpecificData.MaxLevelSetResolution);
+					MaxResolution = FMath::Clamp<float>(MaxResolution, static_cast<float>(SizeSpecificData.MaxLevelSetResolution), static_cast<float>(SizeSpecificData.MaxClusterLevelSetResolution));
 
 					//don't support non level-set serialization
 					ErrorReporter.SetPrefix(BaseErrorPrefix + " | Cluster Transform Index: " + FString::FromInt(ClusterTransformIdx));
@@ -2390,8 +2390,8 @@ void FGeometryCollectionPhysicsProxy::InitializeSharedCollisionStructures(
 							MassSpaceParticles,
 							*UnionMesh,
 							InstanceBoundingBox,
-							MinResolution,
-							MaxResolution,
+							FMath::FloorToInt(MinResolution),
+							FMath::FloorToInt(MaxResolution),
 							SizeSpecificData.CollisionObjectReductionPercentage,
 							SizeSpecificData.CollisionType));
 					// Fall back on sphere if level set rasterization failed.
@@ -2630,7 +2630,11 @@ void FGeometryCollectionPhysicsProxy::FieldParameterUpdateCallback(Chaos::FPBDRi
 								Chaos::TPBDRigidParticleHandle<Chaos::FReal, 3>* RigidHandle = ParticleHandles[Index.Sample]->CastToRigidParticle();
 								if (RigidHandle)
 								{
-									const int8 ResultState = ResultsView[Index.Result];  
+									const int32 CurrResult = ResultsView[Index.Result];
+									check(CurrResult <= std::numeric_limits<int8>::max() &&
+										  CurrResult >= std::numeric_limits<int8>::min());
+
+									const int8 ResultState = static_cast<int8>(CurrResult);
 									const int32 TransformIndex = HandleToTransformGroupIndex[RigidHandle];
 
 									// Update of the handles object state. No need to update 
diff --git a/Engine/Source/Runtime/Experimental/Chaos/Private/PhysicsProxy/SkeletalMeshPhysicsProxy.cpp b/Engine/Source/Runtime/Experimental/Chaos/Private/PhysicsProxy/SkeletalMeshPhysicsProxy.cpp
--- a/Engine/Source/Runtime/Experimental/Chaos/Private/PhysicsProxy/SkeletalMeshPhysicsProxy.cpp
+++ b/Engine/Source/Runtime/Experimental/Chaos/Private/PhysicsProxy/SkeletalMeshPhysicsProxy.cpp
@@ -137,7 +137,7 @@ void FSkeletalMeshPhysicsProxy::CreateRigidBodyCallback(FParticlesType& Particle
 			check(WorldTransform);
 
 			Chaos::FReal TotalMass = 0.;
-			const Chaos::FReal DensityKgCm3 = Parameters.Density / 1000.;
+			const Chaos::FReal DensityKgCm3 = Parameters.Density / 1000.f;
 			Chaos::FMassProperties MassProperties = Group->BuildMassProperties(DensityKgCm3, TotalMass);
 
 			const bool DoCollisionGeom = Parameters.CollisionType == ECollisionTypeEnum::Chaos_Surface_Volumetric;
diff --git a/Engine/Source/Runtime/Experimental/Chaos/Private/PhysicsProxy/StaticMeshPhysicsObject.cpp b/Engine/Source/Runtime/Experimental/Chaos/Private/PhysicsProxy/StaticMeshPhysicsObject.cpp
--- a/Engine/Source/Runtime/Experimental/Chaos/Private/PhysicsProxy/StaticMeshPhysicsObject.cpp
+++ b/Engine/Source/Runtime/Experimental/Chaos/Private/PhysicsProxy/StaticMeshPhysicsObject.cpp
@@ -149,9 +149,14 @@ void FStaticMeshPhysicsProxy::CreateRigidBodyCallback(FParticlesType& Particles)
 						ErrorReporter,
 						MeshParticles, *TriangleMesh,
 						Bounds,
-						FVector::Distance(FVector(0.f, 0.f, 0.f), Bounds.GetExtent())*0.5,
-						Parameters.MinRes, Parameters.MaxRes, 0.f,
-						CollisionType, Parameters.ShapeType)));
+						FVector::Distance(FVector(0.f, 0.f, 0.f), 
+						Bounds.GetExtent()) * 0.5f,
+						Parameters.MinRes, 
+						Parameters.MaxRes, 
+						0.f,
+						CollisionType, 
+						Parameters.ShapeType)));
+
 			if (!ensure(Parameters.MeshVertexPositions.Size()))
 			{
 				Parameters.MeshVertexPositions.AddParticles(1);
@@ -185,7 +190,7 @@ void FStaticMeshPhysicsProxy::CreateRigidBodyCallback(FParticlesType& Particles)
 		}
 		else if (Parameters.ShapeType == EImplicitTypeEnum::Chaos_Implicit_Box)
 		{
-			Chaos::FVec3 HalfExtents = Parameters.ShapeParams.BoxExtents / 2;
+			Chaos::FVec3 HalfExtents = Parameters.ShapeParams.BoxExtents / static_cast<Chaos::FReal>(2);
 			Chaos::TBox<Chaos::FReal,3>* Box = new Chaos::TBox<Chaos::FReal, 3>(-HalfExtents, HalfExtents);
 			Bounds.Min = Box->Min();
 			Bounds.Max = Box->Max();
diff --git a/Engine/Source/Runtime/Experimental/Chaos/Private/kDOP.h b/Engine/Source/Runtime/Experimental/Chaos/Private/kDOP.h
--- a/Engine/Source/Runtime/Experimental/Chaos/Private/kDOP.h
+++ b/Engine/Source/Runtime/Experimental/Chaos/Private/kDOP.h
@@ -5,6 +5,7 @@
 #include "CoreMinimal.h"
 #include "Containers/ChunkedArray.h"
 #include "Templates/UnrealTemplate.h"
+#include <limits>
 
 // Indicates how many "k / 2" there are in the k-DOP. 3 == AABB == 6 DOP. The code relies on this being 3.
 #define NUM_PLANES	3
@@ -613,7 +614,9 @@ struct TkDOPNode
 
 			// No need to subdivide further so make this a leaf node
 			bIsLeaf = 1;
-			Occupancy = NumTris;
+
+			check(NumTris <= std::numeric_limits<uint8>::max());
+			Occupancy = static_cast<uint8>(NumTris);
 			
 			// Generate bounding volume for leaf which is passed up the call chain.
 			FBox BoundingVolume(ForceInit);
diff --git a/Engine/Source/Runtime/Experimental/Chaos/Public/BoneHierarchy.h b/Engine/Source/Runtime/Experimental/Chaos/Public/BoneHierarchy.h
--- a/Engine/Source/Runtime/Experimental/Chaos/Public/BoneHierarchy.h
+++ b/Engine/Source/Runtime/Experimental/Chaos/Public/BoneHierarchy.h
@@ -163,9 +163,9 @@ public:
 		FTransform& TargetXf = LocalSpaceTransforms[TransformIndex];
 		// FTransform::Equals() is flakey.
 		//if (!TargetXf.Equals(BoneXf, 1.0e-8))
-		if (!TargetXf.GetTranslation().Equals(BoneXf.GetTranslation(), 1.0e-8) ||
-			!TargetXf.GetRotation().Equals(BoneXf.GetRotation(), 1.0e-8) ||
-			!TargetXf.GetScale3D().Equals(BoneXf.GetScale3D(), 1.0e-8))
+		if (!TargetXf.GetTranslation().Equals(BoneXf.GetTranslation(), 1.0e-8f) ||
+			!TargetXf.GetRotation().Equals(BoneXf.GetRotation(), 1.0e-8f) ||
+			!TargetXf.GetScale3D().Equals(BoneXf.GetScale3D(), 1.0e-8f))
 		{
 			TargetXf = BoneXf;
 			SetAnimLocalDirty(TransformIndex);
@@ -182,9 +182,9 @@ public:
 		// FTransform::Equals() is flakey.
 		//ActorLocalToWorldDirty |= !ActorLocalToWorld.FTransform::Equals(InActorLocalToWorld);
 		ActorLocalToWorldDirty |=
-			!ActorLocalToWorld.GetTranslation().Equals(InActorLocalToWorld.GetTranslation(), 1.0e-8) ||
-			!ActorLocalToWorld.GetRotation().Equals(InActorLocalToWorld.GetRotation(), 1.0e-8) ||
-			!ActorLocalToWorld.GetScale3D().Equals(InActorLocalToWorld.GetScale3D(), 1.0e-8);
+			!ActorLocalToWorld.GetTranslation().Equals(InActorLocalToWorld.GetTranslation(), 1.0e-8f) ||
+			!ActorLocalToWorld.GetRotation().Equals(InActorLocalToWorld.GetRotation(), 1.0e-8f) ||
+			!ActorLocalToWorld.GetScale3D().Equals(InActorLocalToWorld.GetScale3D(), 1.0e-8f);
 		ActorLocalToWorld = InActorLocalToWorld;
 	}
 
diff --git a/Engine/Source/Runtime/Experimental/Chaos/Public/Chaos/AABB.h b/Engine/Source/Runtime/Experimental/Chaos/Public/Chaos/AABB.h
--- a/Engine/Source/Runtime/Experimental/Chaos/Public/Chaos/AABB.h
+++ b/Engine/Source/Runtime/Experimental/Chaos/Public/Chaos/AABB.h
@@ -365,7 +365,7 @@ namespace Chaos
 				ChosenPt[Axis] = Direction[Axis] < 0 ? MMin[Axis] : MMax[Axis];
 			}
 
-			if (Thickness)
+			if (Thickness != (T)0)
 			{
 				//We want N / ||N|| and to avoid inf
 				//So we want N / ||N|| < 1 / eps => N eps < ||N||, but this is clearly true for all eps < 1 and N > 0
@@ -477,7 +477,7 @@ namespace Chaos
 		}
 
 		FORCEINLINE T GetArea() const { return GetArea(Extents()); }
-		FORCEINLINE static T GetArea(const TVector<T, d>& Dim) { return d == 2 ? Dim.Product() : 2. * (Dim[0] * Dim[1] + Dim[0] * Dim[2] + Dim[1] * Dim[2]); }
+		FORCEINLINE static T GetArea(const TVector<T, d>& Dim) { return d == 2 ? Dim.Product() : (T)2. * (Dim[0] * Dim[1] + Dim[0] * Dim[2] + Dim[1] * Dim[2]); }
 
 		FORCEINLINE T GetVolume() const { return GetVolume(Extents()); }
 		FORCEINLINE static T GetVolume(const TVector<T, 3>& Dim) { return Dim.Product(); }
diff --git a/Engine/Source/Runtime/Experimental/Chaos/Public/Chaos/AABBTree.h b/Engine/Source/Runtime/Experimental/Chaos/Public/Chaos/AABBTree.h
--- a/Engine/Source/Runtime/Experimental/Chaos/Public/Chaos/AABBTree.h
+++ b/Engine/Source/Runtime/Experimental/Chaos/Public/Chaos/AABBTree.h
@@ -657,10 +657,10 @@ private:
 		TRACE_CPUPROFILER_EVENT_SCOPE(TAABBTree::ReoptimizeTree);
 		TArray<FElement> AllElements;
 
-		SIZE_T ReserveCount = DirtyElements.Num() + GlobalPayloads.Num();
+		int32 ReserveCount = DirtyElements.Num() + GlobalPayloads.Num();
 		for (const auto& Leaf : Leaves)
 		{
-			ReserveCount += Leaf.GetReserveCount();
+			ReserveCount += static_cast<int32>(Leaf.GetReserveCount());
 		}
 
 		AllElements.Reserve(ReserveCount);
@@ -903,7 +903,7 @@ private:
 
 			if(WorkSnapshot.Elems.Num())
 			{
-				WorkSnapshot.AverageCenter = CenterSum * ((FReal)1 / WorkSnapshot.Elems.Num());
+				WorkSnapshot.AverageCenter = CenterSum * ((FReal)1 / (FReal)WorkSnapshot.Elems.Num());
 			}
 			else
 			{
@@ -1172,8 +1172,8 @@ private:
 				WorkPool[FirstChildIdx].NewNodeIdx = Nodes[NewNodeIdx].ChildrenNodes[0];
 				WorkPool[SecondChildIdx].NewNodeIdx = Nodes[NewNodeIdx].ChildrenNodes[1];
 
-				WorkPool[FirstChildIdx].AverageCenter *= ((FReal)1 / WorkPool[FirstChildIdx].Elems.Num());
-				WorkPool[SecondChildIdx].AverageCenter *= ((FReal)1 / WorkPool[SecondChildIdx].Elems.Num());
+				WorkPool[FirstChildIdx].AverageCenter *= ((FReal)1 / (FReal)WorkPool[FirstChildIdx].Elems.Num());
+				WorkPool[SecondChildIdx].AverageCenter *= ((FReal)1 / (FReal)WorkPool[SecondChildIdx].Elems.Num());
 
 				//push these two new nodes onto the stack
 				WorkStack.Add(SecondChildIdx);
diff --git a/Engine/Source/Runtime/Experimental/Chaos/Public/Chaos/BoundingVolume.h b/Engine/Source/Runtime/Experimental/Chaos/Public/Chaos/BoundingVolume.h
--- a/Engine/Source/Runtime/Experimental/Chaos/Public/Chaos/BoundingVolume.h
+++ b/Engine/Source/Runtime/Experimental/Chaos/Public/Chaos/BoundingVolume.h
@@ -375,11 +375,11 @@ public:
 		}
 	}
 
-	SIZE_T GetReserveCount() const
+	int32 GetReserveCount() const
 	{
 		// Optimize for fewer memory allocations.
-		const auto& Counts = MGrid.Counts(); 
-		const SIZE_T GridCount = Counts[0] * Counts[1] * Counts[2] * MElements.Num();
+		const TVector<int32, d>& Counts = MGrid.Counts();
+		const int32 GridCount = Counts[0] * Counts[1] * Counts[2] * MElements.Num();
 		return MGlobalPayloads.Num() + MDirtyElements.Num() + GridCount;
 	}
 
@@ -554,7 +554,7 @@ private:
 
 				for (int Axis = 0; Axis < d; ++Axis)
 				{
-					constexpr T Epsilon = 1e-2;	//if raycast is slightly off we still count it as hitting the cell surface
+					constexpr T Epsilon = 1e-2f;	//if raycast is slightly off we still count it as hitting the cell surface
 					CellIdx[Axis] += (Times[Axis] <= BestTime + Epsilon) ? (CurData.Dir[Axis] > 0 ? 1 : -1) : 0;
 					if (CellIdx[Axis] < 0 || CellIdx[Axis] >= MGrid.Counts()[Axis])
 					{
@@ -909,7 +909,7 @@ private:
 
 		AllBounds.SetNum(Particles.Num());
 		HasBounds.SetNum(Particles.Num());
-		int32 MaxPayloadBoundsCopy = MaxPayloadBounds;
+		T MaxPayloadBoundsCopy = MaxPayloadBounds;
 		auto GetValidBounds = [MaxPayloadBoundsCopy, bUseVelocity, Dt](const auto& Particle, TAABB<T,d>& OutBounds) -> bool
 		{
 			if (HasBoundingBox(Particle))
@@ -1234,7 +1234,7 @@ private:
 			int32 Idx = (NumY * NumZ) * Coordinate[0] + (NumZ * Coordinate[1]) + Coordinate[2];
 			int32 ByteIdx = Idx / 8;
 			int32 BitIdx = Idx % 8;
-			uint8 Mask = 1 << BitIdx;
+			uint8 Mask = static_cast<uint8>(1 << BitIdx);
 			Data[ByteIdx] |= Mask;
 		}
 
diff --git a/Engine/Source/Runtime/Experimental/Chaos/Public/Chaos/Capsule.h b/Engine/Source/Runtime/Experimental/Chaos/Public/Chaos/Capsule.h
--- a/Engine/Source/Runtime/Experimental/Chaos/Public/Chaos/Capsule.h
+++ b/Engine/Source/Runtime/Experimental/Chaos/Public/Chaos/Capsule.h
@@ -341,10 +341,14 @@ namespace Chaos
 		TSegment<FReal> GetSegment() const { return TSegment<FReal>(GetX1(), GetX2()); }
 
 		FReal GetArea() const { return GetArea(GetHeight(), GetRadius()); }
-		static FReal GetArea(const FReal Height, const FReal Radius) { static const FReal PI2 = 2. * PI; return PI2 * Radius * (Height + 2.*Radius); }
+		static FReal GetArea(const FReal Height, const FReal Radius)
+		{
+			static const FReal PI2 = 2.f * PI;
+			return PI2 * Radius * (Height + 2.f * Radius); 
+		}
 
 		FReal GetVolume() const { return GetVolume(GetHeight(), GetRadius()); }
-		static FReal GetVolume(const FReal Height, const FReal Radius) { static const FReal FourThirds = 4. / 3; return PI * Radius*Radius * (Height + FourThirds * Radius); }
+		static FReal GetVolume(const FReal Height, const FReal Radius) { static const FReal FourThirds = 4.0f / 3.0f; return PI * Radius * Radius * (Height + FourThirds * Radius); }
 
 		FMatrix33 GetInertiaTensor(const FReal Mass) const { return GetInertiaTensor(Mass, GetHeight(), GetRadius()); }
 		static FMatrix33 GetInertiaTensor(const FReal Mass, const FReal Height, const FReal Radius)
@@ -356,9 +360,9 @@ namespace Chaos
 			const FReal HH = H * H;
 
 			// (5H^3 + 20*H^2R + 45HR^2 + 32R^3) / (60H + 80R)
-			const FReal Diag12 = Mass * (5.*HH*H + 20.*HH*R + 45.*H*RR + 32.*RR*R) / (60.*H + 80.*R);
+			const FReal Diag12 = static_cast<FReal>(Mass * (5.*HH*H + 20.*HH*R + 45.*H*RR + 32.*RR*R) / (60.*H + 80.*R));
 			// (R^2 * (15H + 16R) / (30H +40R))
-			const FReal Diag3 = Mass * (RR * (15.*H + 16.*R)) / (30.*H + 40.*R);
+			const FReal Diag3 = static_cast<FReal>(Mass * (RR * (15.*H + 16.*R)) / (30.*H + 40.*R));
 
 			return FMatrix33(Diag12, Diag12, Diag3);
 		}
@@ -424,11 +428,11 @@ namespace Chaos
 			int32 NumPointsEndCap;
 			int32 NumPointsCylinder;
 			const FReal CapArea = 4 * PI * Radius * Radius;
-			const FReal CylArea = 2.0 * PI * Radius * Height;
+			const FReal CylArea = static_cast<FReal>(2.0 * PI * Radius * Height);
 			if (CylArea > KINDA_SMALL_NUMBER)
 			{
 				const FReal AllArea = CylArea + CapArea;
-				NumPointsCylinder = static_cast<int32>(round(CylArea / AllArea * NumPoints));
+				NumPointsCylinder = static_cast<int32>(round(CylArea / AllArea * static_cast<FReal>(NumPoints)));
 				NumPointsCylinder += (NumPoints - NumPointsCylinder) % 2;
 				NumPointsEndCap = (NumPoints - NumPointsCylinder) / 2;
 			}
diff --git a/Engine/Source/Runtime/Experimental/Chaos/Public/Chaos/ChaosDebugDraw.h b/Engine/Source/Runtime/Experimental/Chaos/Public/Chaos/ChaosDebugDraw.h
--- a/Engine/Source/Runtime/Experimental/Chaos/Public/Chaos/ChaosDebugDraw.h
+++ b/Engine/Source/Runtime/Experimental/Chaos/Public/Chaos/ChaosDebugDraw.h
@@ -37,7 +37,7 @@ namespace Chaos
 				FRealSingle InAngVelScale,
 				FRealSingle InImpulseScale,
 				float InInertiaScale,
-				int InDrawPriority,
+				uint8 InDrawPriority,
 				bool bInShowSimpleCollision,
 				bool bInShowComplexCollision,
 				bool bInShowLevelSetCollision
@@ -84,7 +84,7 @@ namespace Chaos
 			FRealSingle AngVelScale;
 			FRealSingle ImpulseScale;
 			FRealSingle InertiaScale;
-			int DrawPriority;
+			uint8 DrawPriority;
 			bool bShowSimpleCollision;
 			bool bShowComplexCollision;
 			bool bShowLevelSetCollision;
diff --git a/Engine/Source/Runtime/Experimental/Chaos/Public/Chaos/Collision/StatsData.h b/Engine/Source/Runtime/Experimental/Chaos/Public/Chaos/Collision/StatsData.h
--- a/Engine/Source/Runtime/Experimental/Chaos/Public/Chaos/Collision/StatsData.h
+++ b/Engine/Source/Runtime/Experimental/Chaos/Public/Chaos/Collision/StatsData.h
@@ -71,10 +71,10 @@ namespace Chaos
 					}
 				}
 
-				const float CountPerChar = MaxBucketCount / 20.f;
+				const float CountPerChar = static_cast<float>(MaxBucketCount) / 20.f;
 				for (int32 Idx = 1; Idx < UE_ARRAY_COUNT(BucketSizes2); ++Idx)
 				{
-					int32 NumChars = BucketCount[Idx] / CountPerChar;
+					int32 NumChars = static_cast<int32>(static_cast<float>(BucketCount[Idx]) / CountPerChar);
 					if (Idx < UE_ARRAY_COUNT(BucketSizes2) - 1)
 					{
 						OutLog += FString::Printf(TEXT("\t[%4d - %4d) (%4d) |"), BucketSizes2[Idx - 1], BucketSizes2[Idx], BucketCount[Idx]);
@@ -106,7 +106,7 @@ namespace Chaos
 
 			void RecordBoundsData(const FAABB3& Box1)
 			{
-				BoundsDistribution.Record(Box1.Extents().GetMax());
+				BoundsDistribution.Record(static_cast<int32>(Box1.Extents().GetMax()));
 			}
 
 			void RecordBroadphasePotentials(int32 Num)
diff --git a/Engine/Source/Runtime/Experimental/Chaos/Public/Chaos/CollisionConvexMesh.h b/Engine/Source/Runtime/Experimental/Chaos/Public/Chaos/CollisionConvexMesh.h
--- a/Engine/Source/Runtime/Experimental/Chaos/Public/Chaos/CollisionConvexMesh.h
+++ b/Engine/Source/Runtime/Experimental/Chaos/Public/Chaos/CollisionConvexMesh.h
@@ -685,7 +685,7 @@ namespace Chaos
 			FHalfEdge* Cur = OldList;
 			while(Cur)
 			{
-				FReal MaxD = 1e-4;
+				FReal MaxD = 1e-4f;
 				int32 MaxIdx = -1;
 				for(int32 Idx = 0; Idx < NumFaces; ++Idx)
 				{
@@ -746,7 +746,7 @@ namespace Chaos
 				return nullptr;
 			}
 
-			constexpr FReal Epsilon = 1e-4;
+			constexpr FReal Epsilon = 1e-4f;
 
 			const int32 NumVertices = InVertices.Num();
 
diff --git a/Engine/Source/Runtime/Experimental/Chaos/Public/Chaos/Convex.h b/Engine/Source/Runtime/Experimental/Chaos/Public/Chaos/Convex.h
--- a/Engine/Source/Runtime/Experimental/Chaos/Public/Chaos/Convex.h
+++ b/Engine/Source/Runtime/Experimental/Chaos/Public/Chaos/Convex.h
@@ -415,7 +415,7 @@ namespace Chaos
 			return FVec3(0.f, 0.f, 1.f);
 		}
 
-		virtual int32 FindClosestFaceAndVertices(const FVec3& Position, TArray<FVec3>& FaceVertices, FReal SearchDist = 0.01) const override;
+		virtual int32 FindClosestFaceAndVertices(const FVec3& Position, TArray<FVec3>& FaceVertices, FReal SearchDist = 0.01f) const override;
 
 		// Returns a winding order multiplier used in the manifold clipping and required when we have negative scales (See ImplicitObjectScaled)
 		FReal GetWindingOrder() const
diff --git a/Engine/Source/Runtime/Experimental/Chaos/Public/Chaos/ConvexStructureData.h b/Engine/Source/Runtime/Experimental/Chaos/Public/Chaos/ConvexStructureData.h
--- a/Engine/Source/Runtime/Experimental/Chaos/Public/Chaos/ConvexStructureData.h
+++ b/Engine/Source/Runtime/Experimental/Chaos/Public/Chaos/ConvexStructureData.h
@@ -115,7 +115,7 @@ namespace Chaos
 			int VertexPlanesArrayStart = 0;
 			for (int32 VertexIndex = 0; VertexIndex < NumVerts; ++VertexIndex)
 			{
-				VertexPlanesOffsetCount[VertexIndex].Key = VertexPlanesArrayStart;
+				VertexPlanesOffsetCount[VertexIndex].Key = static_cast<FOffsetIndex>(VertexPlanesArrayStart);
 				VertexPlanesArrayStart += VertexPlanesOffsetCount[VertexIndex].Value;
 				VertexPlanesOffsetCount[VertexIndex].Value = 0;
 			}
@@ -128,8 +128,8 @@ namespace Chaos
 			int32 PlaneVerticesArrayStart = 0;
 			for (int32 PlaneIndex = 0; PlaneIndex < InPlaneVertices.Num(); ++PlaneIndex)
 			{
-				PlaneVerticesOffsetCount[PlaneIndex].Key = PlaneVerticesArrayStart;
-				PlaneVerticesOffsetCount[PlaneIndex].Value = InPlaneVertices[PlaneIndex].Num();
+				PlaneVerticesOffsetCount[PlaneIndex].Key = static_cast<FOffsetIndex>(PlaneVerticesArrayStart);
+				PlaneVerticesOffsetCount[PlaneIndex].Value = static_cast<FIndex>(InPlaneVertices[PlaneIndex].Num());
 				PlaneVerticesArrayStart += InPlaneVertices[PlaneIndex].Num();
 
 				for (int32 PlaneVertexIndex = 0; PlaneVertexIndex < InPlaneVertices[PlaneIndex].Num(); ++PlaneVertexIndex)
@@ -137,11 +137,11 @@ namespace Chaos
 					const int32 VertexIndex = InPlaneVertices[PlaneIndex][PlaneVertexIndex];
 
 					const int32 PlaneVertexFlatArrayIndex = PlaneVerticesOffsetCount[PlaneIndex].Key + PlaneVertexIndex;
-					PlaneVertices[PlaneVertexFlatArrayIndex] = VertexIndex;
+					PlaneVertices[PlaneVertexFlatArrayIndex] = static_cast<FIndex>(VertexIndex);
 
 					const int32 VertexPlaneFlatArrayIndex = VertexPlanesOffsetCount[VertexIndex].Key + VertexPlanesOffsetCount[VertexIndex].Value;
 					VertexPlanesOffsetCount[VertexIndex].Value++;
-					VertexPlanes[VertexPlaneFlatArrayIndex] = PlaneIndex;
+					VertexPlanes[VertexPlaneFlatArrayIndex] = static_cast<FIndex>(PlaneIndex);
 				}
 			}
 		}
diff --git a/Engine/Source/Runtime/Experimental/Chaos/Public/Chaos/Cylinder.h b/Engine/Source/Runtime/Experimental/Chaos/Public/Chaos/Cylinder.h
--- a/Engine/Source/Runtime/Experimental/Chaos/Public/Chaos/Cylinder.h
+++ b/Engine/Source/Runtime/Experimental/Chaos/Public/Chaos/Cylinder.h
@@ -190,8 +190,8 @@ namespace Chaos
 		{
 			// https://www.wolframalpha.com/input/?i=cylinder
 			const FReal RR = Radius * Radius;
-			const FReal Diag12 = Mass / 12. * (3.*RR + Height*Height);
-			const FReal Diag3 = Mass / 2. * RR;
+			const FReal Diag12 = static_cast<FReal>(Mass / 12. * (3.*RR + Height*Height));
+			const FReal Diag3 = static_cast<FReal>(Mass / 2. * RR);
 			return FMatrix33(Diag12, Diag12, Diag3);
 		}
 
@@ -379,11 +379,11 @@ namespace Chaos
 			if (IncludeEndCaps)
 			{
 				const FReal CapArea = PI * Radius * Radius;
-				const FReal CylArea = 2.0 * PI * Radius * Height;
+				const FReal CylArea = static_cast<FReal>(2.0 * PI * Radius * Height);
 				const FReal AllArea = CylArea + CapArea * 2;
 				if (AllArea > KINDA_SMALL_NUMBER)
 				{
-					NumPointsCylinder = static_cast<int32>(round(CylArea / AllArea * NumPoints));
+					NumPointsCylinder = static_cast<int32>(round(CylArea / AllArea * (FReal)NumPoints));
 					NumPointsCylinder += (NumPoints - NumPointsCylinder) % 2;
 					NumPointsEndCap = (NumPoints - NumPointsCylinder) / 2;
 				}
@@ -421,7 +421,7 @@ namespace Chaos
 			}
 
 			Offset = Points.AddUninitialized(NumPointsCylinder);
-			static const FReal Increment = PI * (1.0 + sqrt(5));
+			static const FReal Increment = static_cast<FReal>(PI * (1.0 + sqrt(5)));
 			for (int32 i = 0; i < NumPointsCylinder; i++)
 			{
 				// In the 2D sphere (disc) case, we vary R so it increases monotonically,
@@ -429,10 +429,10 @@ namespace Chaos
 				//     const T R = FMath::Sqrt((0.5 + Index) / NumPoints) * Radius;
 				// But we're mapping to a cylinder, which means we want to keep R constant.
 				const FReal R = Radius;
-				const FReal Theta = Increment * (0.5 + i + SpiralSeed);
+				const FReal Theta = Increment * (0.5f + (FReal)i + (FReal)SpiralSeed);
 
 				// Map polar coordinates to Cartesian, and vary Z by [-HalfHeight, HalfHeight].
-				const FReal Z = FMath::LerpStable(-HalfHeight, HalfHeight, static_cast<FReal>(i) / (NumPointsCylinder - 1));
+				const FReal Z = FMath::LerpStable(-HalfHeight, HalfHeight, static_cast<FReal>(i) / (static_cast<FReal>(NumPointsCylinder) - 1));
 				Points[i + Offset] =
 				    FVec3(
 				        R * FMath::Cos(Theta),
diff --git a/Engine/Source/Runtime/Experimental/Chaos/Public/Chaos/Defines.h b/Engine/Source/Runtime/Experimental/Chaos/Public/Chaos/Defines.h
--- a/Engine/Source/Runtime/Experimental/Chaos/Public/Chaos/Defines.h
+++ b/Engine/Source/Runtime/Experimental/Chaos/Public/Chaos/Defines.h
@@ -75,7 +75,7 @@ namespace Chaos
 		FChaosPhysicsMaterial()
 			: Friction(0.5)
 			, StaticFriction(0.0)
-			, Restitution(0.1)
+			, Restitution(0.1f)
 			, LinearEtherDrag(0.0)
 			, AngularEtherDrag(0.0)
 			, SleepingLinearThreshold(1)
diff --git a/Engine/Source/Runtime/Experimental/Chaos/Public/Chaos/EPA.h b/Engine/Source/Runtime/Experimental/Chaos/Public/Chaos/EPA.h
--- a/Engine/Source/Runtime/Experimental/Chaos/Public/Chaos/EPA.h
+++ b/Engine/Source/Runtime/Experimental/Chaos/Public/Chaos/EPA.h
@@ -40,7 +40,7 @@ struct TEPAEntry
 		const TVec3<T> V0V2 = V2 - V0;
 		const TVec3<T> Norm = TVec3<T>::CrossProduct(V0V1, V0V2);
 		PlaneNormal = Norm.GetSafeNormal();
-		constexpr T Eps = 1e-4;
+		constexpr T Eps = 1e-4f;
 		if (PlaneNormal.SizeSquared() < Eps)
 		{
 			return false;
@@ -420,7 +420,7 @@ EEPAResult EPA(TArray<TVec3<T>>& VertsABuffer, TArray<TVec3<T>>& VertsBBuffer, c
 	T LowerBound = TNumericLimits<T>::Lowest();
 
 	// @todo(chaos): Should we pass in Eps? Does it need to match anything used in GJK?
-	constexpr T Eps = 1e-2;
+	constexpr T Eps = 1e-2f;
 	constexpr T OriginInsideEps = 0.0f;
 
 	TArray<TEPAEntry<T>> Entries;
diff --git a/Engine/Source/Runtime/Experimental/Chaos/Public/Chaos/GJK.h b/Engine/Source/Runtime/Experimental/Chaos/Public/Chaos/GJK.h
--- a/Engine/Source/Runtime/Experimental/Chaos/Public/Chaos/GJK.h
+++ b/Engine/Source/Runtime/Experimental/Chaos/Public/Chaos/GJK.h
@@ -44,7 +44,7 @@ namespace Chaos
 		T PrevDist2 = FLT_MAX;
 		const T ThicknessA = A.GetMargin() + InThicknessA;
 		const T ThicknessB = B.GetMargin() + InThicknessB;
-		const T Inflation = ThicknessA + ThicknessB + 1e-3;
+		const T Inflation = ThicknessA + ThicknessB + static_cast<T>(1e-3);
 		const T Inflation2 = Inflation * Inflation;
 		do
 		{
@@ -113,7 +113,7 @@ namespace Chaos
 	// the Thickness) use GJKPenetration().
 	//
 	template <bool bNegativePenetrationAllowed = false, typename T, typename TGeometryA, typename TGeometryB>
-	bool GJKPenetration(const TGeometryA& A, const TGeometryB& B, const TRigidTransform<T, 3>& BToATM, T& OutPenetration, TVec3<T>& OutClosestA, TVec3<T>& OutClosestB, TVec3<T>& OutNormal, int32& OutClosestVertexIndexA, int32& OutClosestVertexIndexB, const T InThicknessA = 0.0f, const T InThicknessB = 0.0f, const TVector<T, 3>& InitialDir = TVector<T, 3>(1, 0, 0), const T Epsilon = 1.e-3)
+	bool GJKPenetration(const TGeometryA& A, const TGeometryB& B, const TRigidTransform<T, 3>& BToATM, T& OutPenetration, TVec3<T>& OutClosestA, TVec3<T>& OutClosestB, TVec3<T>& OutNormal, int32& OutClosestVertexIndexA, int32& OutClosestVertexIndexB, const T InThicknessA = 0.0f, const T InThicknessB = 0.0f, const TVector<T, 3>& InitialDir = TVector<T, 3>(1, 0, 0), const T Epsilon = 1.e-3f)
 	{
 		int32 VertexIndexA = INDEX_NONE;
 		int32 VertexIndexB = INDEX_NONE;
@@ -469,7 +469,7 @@ namespace Chaos
 
 		T Barycentric[4] = { -1,-1,-1,-1 };	//not needed, but compiler warns
 		const T Inflation = MarginA + MarginB;
-		const T Inflation2 = Inflation*Inflation + 1e-6;
+		const T Inflation2 = Inflation * Inflation + static_cast<T>(1e-6);
 
 		FSimplex SimplexIDs;
 		const TRotation<T, 3> BToARotation = StartTM.GetRotation();
@@ -499,7 +499,7 @@ namespace Chaos
 		TVector<T, 3> Normal(0,0,1);
 
 		const T InitialPreDist2 = V.SizeSquared();
-		constexpr T Eps2 = 1e-6;
+		constexpr T Eps2 = 1e-6f;
 		//mtd needs to find closest point even in inflation region, so can only skip if we found the closest points
 		bool bCloseEnough = InitialPreDist2 < Inflation2 && (!bComputeMTD || InitialPreDist2 < Eps2);
 		bool bDegenerate = false;
diff --git a/Engine/Source/Runtime/Experimental/Chaos/Public/Chaos/HeightField.h b/Engine/Source/Runtime/Experimental/Chaos/Public/Chaos/HeightField.h
--- a/Engine/Source/Runtime/Experimental/Chaos/Public/Chaos/HeightField.h
+++ b/Engine/Source/Runtime/Experimental/Chaos/Public/Chaos/HeightField.h
@@ -140,7 +140,7 @@ namespace Chaos
 			// to support loading. Cast const away here to get bytes written
 			const_cast<FHeightField*>(this)->Serialize(ChaosAr);
 
-			return FCrc::MemCrc32(Bytes.GetData(), Bytes.GetAllocatedSize());
+			return FCrc::MemCrc32(Bytes.GetData(), (int32)Bytes.GetAllocatedSize());
 		}
 
 		static constexpr EImplicitObjectType StaticType()
diff --git a/Engine/Source/Runtime/Experimental/Chaos/Public/Chaos/ISpatialAcceleration.h b/Engine/Source/Runtime/Experimental/Chaos/Public/Chaos/ISpatialAcceleration.h
--- a/Engine/Source/Runtime/Experimental/Chaos/Public/Chaos/ISpatialAcceleration.h
+++ b/Engine/Source/Runtime/Experimental/Chaos/Public/Chaos/ISpatialAcceleration.h
@@ -16,7 +16,7 @@ struct CHAOS_API FQueryFastData
 		, InvDir( (InDir[0] == 0) ? 0 : 1 / Dir[0], (InDir[1] == 0) ? 0 : 1 / Dir[1], (InDir[2] == 0) ? 0 : 1 / Dir[2])
 		, bParallel{ InDir[0] == 0, InDir[1] == 0, InDir[2] == 0 }
 	{
-		CHAOS_ENSURE(InLength);
+		CHAOS_ENSURE(InLength != 0.0f);
 		SetLength(InLength);
 	}
 
@@ -39,7 +39,7 @@ struct CHAOS_API FQueryFastData
 	{
 		CurrentLength = InLength;
 
-		if(InLength)
+		if(InLength != 0.0f)
 		{
 			InvCurrentLength = 1 / InLength;
 		}
diff --git a/Engine/Source/Runtime/Experimental/Chaos/Public/Chaos/ImplicitObject.h b/Engine/Source/Runtime/Experimental/Chaos/Public/Chaos/ImplicitObject.h
--- a/Engine/Source/Runtime/Experimental/Chaos/Public/Chaos/ImplicitObject.h
+++ b/Engine/Source/Runtime/Experimental/Chaos/Public/Chaos/ImplicitObject.h
@@ -291,7 +291,7 @@ public:
 	@param FaceIndices - Vertices that lie on the face plane.
 	@param SearchDistance - distance to surface [def:0.01]
 	*/
-	virtual int32 FindClosestFaceAndVertices(const FVec3& Position, TArray<FVec3>& FaceVertices, FReal SearchDist = 0.01) const
+	virtual int32 FindClosestFaceAndVertices(const FVec3& Position, TArray<FVec3>& FaceVertices, FReal SearchDist = 0.01f) const
 	{
 		//Many objects have no concept of a face
 		return INDEX_NONE;
diff --git a/Engine/Source/Runtime/Experimental/Chaos/Public/Chaos/ImplicitObjectScaled.h b/Engine/Source/Runtime/Experimental/Chaos/Public/Chaos/ImplicitObjectScaled.h
--- a/Engine/Source/Runtime/Experimental/Chaos/Public/Chaos/ImplicitObjectScaled.h
+++ b/Engine/Source/Runtime/Experimental/Chaos/Public/Chaos/ImplicitObjectScaled.h
@@ -717,8 +717,10 @@ public:
 		// Compute final normal
 		const TVector<T, d> LocalNormal = MObject->FindGeometryOpposingNormal(LocalDenormDir, HintFaceIndex, LocalOriginalNormal);
 		TVector<T, d> Normal = LocalNormal * MInvScale;
-		if (CHAOS_ENSURE(Normal.SafeNormalize(TNumericLimits<T>::Min())) == 0)
+		FReal NormalLength = Normal.SafeNormalize(TNumericLimits<T>::Min());
+		if (NormalLength == 0)
 		{
+			CHAOS_ENSURE(false);
 			Normal = TVector<T, 3>(0, 0, 1);
 		}
 
@@ -748,7 +750,7 @@ public:
 		return ClosestIntersection;
 	}
 
-	virtual int32 FindClosestFaceAndVertices(const FVec3& Position, TArray<FVec3>& FaceVertices, FReal SearchDist = 0.01) const override
+	virtual int32 FindClosestFaceAndVertices(const FVec3& Position, TArray<FVec3>& FaceVertices, FReal SearchDist = 0.01f) const override
 	{
 		const FVec3 UnscaledPoint = MInvScale * Position;
 		const FReal UnscaledSearchDist = SearchDist * MInvScale.Max();	//this is not quite right since it's no longer a sphere, but the whole thing is fuzzy anyway
@@ -783,7 +785,7 @@ public:
 
 	void SetScale(const FVec3& Scale)
 	{
-		constexpr FReal MinMagnitude = 1e-6;
+		constexpr FReal MinMagnitude = 1e-6f;
 		for (int Axis = 0; Axis < 3; ++Axis)
 		{
 			if (!CHAOS_ENSURE(FMath::Abs(Scale[Axis]) >= MinMagnitude))
diff --git a/Engine/Source/Runtime/Experimental/Chaos/Public/Chaos/ImplicitObjectTransformed.h b/Engine/Source/Runtime/Experimental/Chaos/Public/Chaos/ImplicitObjectTransformed.h
--- a/Engine/Source/Runtime/Experimental/Chaos/Public/Chaos/ImplicitObjectTransformed.h
+++ b/Engine/Source/Runtime/Experimental/Chaos/Public/Chaos/ImplicitObjectTransformed.h
@@ -173,7 +173,7 @@ public:
 		return ClosestIntersection;
 	}
 
-	virtual int32 FindClosestFaceAndVertices(const FVec3& Position, TArray<FVec3>& FaceVertices, FReal SearchDist = 0.01) const override
+	virtual int32 FindClosestFaceAndVertices(const FVec3& Position, TArray<FVec3>& FaceVertices, FReal SearchDist = 0.01f) const override
 	{
 		const FVec3 LocalPoint = MTransform.InverseTransformPosition(Position);
 		int32 FaceIndex = MObject->FindClosestFaceAndVertices(LocalPoint, FaceVertices, SearchDist);
diff --git a/Engine/Source/Runtime/Experimental/Chaos/Public/Chaos/Joint/JointSolverConstraints.inl b/Engine/Source/Runtime/Experimental/Chaos/Public/Chaos/Joint/JointSolverConstraints.inl
--- a/Engine/Source/Runtime/Experimental/Chaos/Public/Chaos/Joint/JointSolverConstraints.inl
+++ b/Engine/Source/Runtime/Experimental/Chaos/Public/Chaos/Joint/JointSolverConstraints.inl
@@ -397,7 +397,7 @@ namespace Chaos
 		const FRotation3& RSwing)
 	{
 		const FReal R01SwingYorZ = (FJointConstants::AxisIndex((EJointAngularConstraintIndex)RowData.ConstraintIndex) == 2) ? RSwing.Z : RSwing.Y;	// Can't index a quat :(
-		const FReal Angle = 4. * FMath::Atan2(R01SwingYorZ, (FReal)(1. + RSwing.W));
+		const FReal Angle = (FReal)4. * FMath::Atan2(R01SwingYorZ, (FReal)(1. + RSwing.W));
 		const FVec3& AxisLocal = (RowData.ConstraintIndex == (int32)EJointAngularConstraintIndex::Swing1) ? FJointConstants::Swing1Axis() : FJointConstants::Swing2Axis();
 
 		RowState.Axis = R0 * AxisLocal;
diff --git a/Engine/Source/Runtime/Experimental/Chaos/Public/Chaos/PBDAxialSpringConstraintsBase.h b/Engine/Source/Runtime/Experimental/Chaos/Public/Chaos/PBDAxialSpringConstraintsBase.h
--- a/Engine/Source/Runtime/Experimental/Chaos/Public/Chaos/PBDAxialSpringConstraintsBase.h
+++ b/Engine/Source/Runtime/Experimental/Chaos/Public/Chaos/PBDAxialSpringConstraintsBase.h
@@ -26,9 +26,9 @@ public:
 			FReal Bary2 = FindBary(InParticles, i2, i3, i1);
 			FReal Bary3 = FindBary(InParticles, i3, i1, i2);
 			FReal Bary = Bary1;
-			FReal Bary1dist = FGenericPlatformMath::Abs(Bary1 - 0.5);
-			FReal Bary2dist = FGenericPlatformMath::Abs(Bary2 - 0.5);
-			FReal Bary3dist = FGenericPlatformMath::Abs(Bary3 - 0.5);
+			FReal Bary1dist = FGenericPlatformMath::Abs(Bary1 - 0.5f);
+			FReal Bary2dist = FGenericPlatformMath::Abs(Bary2 - 0.5f);
+			FReal Bary3dist = FGenericPlatformMath::Abs(Bary3 - 0.5f);
 			if (Bary3dist < Bary2dist && Bary3dist < Bary1dist)
 			{
 				Constraint[0] = i3;
diff --git a/Engine/Source/Runtime/Experimental/Chaos/Public/Chaos/PBDCollisionSpringConstraints.h b/Engine/Source/Runtime/Experimental/Chaos/Public/Chaos/PBDCollisionSpringConstraints.h
--- a/Engine/Source/Runtime/Experimental/Chaos/Public/Chaos/PBDCollisionSpringConstraints.h
+++ b/Engine/Source/Runtime/Experimental/Chaos/Public/Chaos/PBDCollisionSpringConstraints.h
@@ -43,15 +43,15 @@ public:
 		{
 			Particles.P(i1) += Multiplier * Particles.InvM(i1) * Delta;
 		}
-		if (Particles.InvM(i2))
+		if (Particles.InvM(i2) > 0)
 		{
 			Particles.P(i2) -= Multiplier * Particles.InvM(i2) * Barys[i][0] * Delta;
 		}
-		if (Particles.InvM(i3))
+		if (Particles.InvM(i3) > 0)
 		{
 			Particles.P(i3) -= Multiplier * Particles.InvM(i3) * Barys[i][1] * Delta;
 		}
-		if (Particles.InvM(i4))
+		if (Particles.InvM(i4) > 0)
 		{
 			Particles.P(i4) -= Multiplier * Particles.InvM(i4) * Barys[i][2] * Delta;
 		}
diff --git a/Engine/Source/Runtime/Experimental/Chaos/Public/Chaos/PBDEvolution.h b/Engine/Source/Runtime/Experimental/Chaos/Public/Chaos/PBDEvolution.h
--- a/Engine/Source/Runtime/Experimental/Chaos/Public/Chaos/PBDEvolution.h
+++ b/Engine/Source/Runtime/Experimental/Chaos/Public/Chaos/PBDEvolution.h
@@ -17,7 +17,7 @@ class CHAOS_API FPBDEvolution : public TArrayCollection
 	using FGravityForces = FPerParticleGravity;
 
 	// TODO(mlentine): Init particles from some type of input
-	FPBDEvolution(FPBDParticles&& InParticles, FKinematicGeometryClothParticles&& InGeometryParticles, TArray<TVec3<int32>>&& CollisionTriangles, int32 NumIterations = 1, FReal CollisionThickness = 0, FReal SelfCollisionsThickness = 0, FReal CoefficientOfFriction = 0, FReal Damping = 0.04);
+	FPBDEvolution(FPBDParticles&& InParticles, FKinematicGeometryClothParticles&& InGeometryParticles, TArray<TVec3<int32>>&& CollisionTriangles, int32 NumIterations = 1, FReal CollisionThickness = 0, FReal SelfCollisionsThickness = 0, FReal CoefficientOfFriction = 0, FReal Damping = 0.04f);
 	~FPBDEvolution() {}
 
 	void AdvanceOneTimeStep(const FReal dt);
diff --git a/Engine/Source/Runtime/Experimental/Chaos/Public/Chaos/PBDRigidClustering.h b/Engine/Source/Runtime/Experimental/Chaos/Public/Chaos/PBDRigidClustering.h
--- a/Engine/Source/Runtime/Experimental/Chaos/Public/Chaos/PBDRigidClustering.h
+++ b/Engine/Source/Runtime/Experimental/Chaos/Public/Chaos/PBDRigidClustering.h
@@ -415,7 +415,7 @@ inline TArray<FVec3> CleanCollisionParticles(
 
 	int32 NumCoincident = 0;
 	const int64 Resolution = static_cast<int64>(floor(MaxBBoxDim / FMath::Max(SnapDistance,(FReal)KINDA_SMALL_NUMBER)));
-	const FReal CellSize = MaxBBoxDim / Resolution;
+	const FReal CellSize = static_cast<FReal>(static_cast<double>(MaxBBoxDim) / static_cast<double>(Resolution));
 	for (int32 i = 0; i < 2; i++)
 	{
 		Redundant.Reset();
@@ -423,14 +423,14 @@ inline TArray<FVec3> CleanCollisionParticles(
 		// Shift the grid by 1/2 a grid cell the second iteration so that
 		// we don't miss slightly adjacent coincident points across cell
 		// boundaries.
-		const FVec3 GridCenter = FVec3(0) - FVec3(i * CellSize / 2);
+		const FVec3 GridCenter = FVec3(0) - FVec3(static_cast<FReal>(i) * CellSize / 2);
 		for (int32 j = 0; j < Points.Num(); j++)
 		{
 			const FVec3 Pos = Points[j] - PointsCenter; // Centered at the origin
 			const TVec3<int64> Coord(
-				static_cast<int64>(floor((Pos[0] - GridCenter[0]) / CellSize + Resolution / 2)),
-				static_cast<int64>(floor((Pos[1] - GridCenter[1]) / CellSize + Resolution / 2)),
-				static_cast<int64>(floor((Pos[2] - GridCenter[2]) / CellSize + Resolution / 2)));
+				static_cast<int64>(FMath::Floor((Pos[0] - GridCenter[0]) / CellSize + static_cast<double>(Resolution) / 2)),
+				static_cast<int64>(FMath::Floor((Pos[1] - GridCenter[1]) / CellSize + static_cast<double>(Resolution) / 2)),
+				static_cast<int64>(FMath::Floor((Pos[2] - GridCenter[2]) / CellSize + static_cast<double>(Resolution) / 2)));
 			const int64 FlatIdx =
 				((Coord[0] * Resolution + Coord[1]) * Resolution) + Coord[2];
 
@@ -504,7 +504,7 @@ inline TArray<FVec3> CleanCollisionParticles(
 	}
 #endif
 
-	CollisionVertices.AddUninitialized(std::min(NumGoodPoints, static_cast<int32>(ceil(NumGoodPoints * Fraction))));
+	CollisionVertices.AddUninitialized(FMath::Min(NumGoodPoints, static_cast<int32>(ceil(static_cast<FReal>(NumGoodPoints) * Fraction))));
 	for (int i = 0; i < CollisionVertices.Num(); i++)
 	{
 		CollisionVertices[i] = Vertices[Ordering[i]];
@@ -525,7 +525,7 @@ inline void CleanCollisionParticles(
 	TArray<int32> CoincidentVertices;
 	const TArray<int32> Ordering = TriMesh.GetVertexImportanceOrdering(Vertices, &CoincidentVertices, true);
 	int32 NumGoodPoints = Ordering.Num() - CoincidentVertices.Num();
-	NumGoodPoints = std::min(NumGoodPoints, static_cast<int32>(ceil(NumGoodPoints * Fraction)));
+	NumGoodPoints = FMath::Min(NumGoodPoints, static_cast<int32>(ceil(static_cast<FReal>(NumGoodPoints) * Fraction)));
 
 	ResultingIndices.Reserve(NumGoodPoints);
 	for (int32 i = 0; i < NumGoodPoints; i++)
diff --git a/Engine/Source/Runtime/Experimental/Chaos/Public/Chaos/PBDRigidDynamicSpringConstraints.h b/Engine/Source/Runtime/Experimental/Chaos/Public/Chaos/PBDRigidDynamicSpringConstraints.h
--- a/Engine/Source/Runtime/Experimental/Chaos/Public/Chaos/PBDRigidDynamicSpringConstraints.h
+++ b/Engine/Source/Runtime/Experimental/Chaos/Public/Chaos/PBDRigidDynamicSpringConstraints.h
@@ -137,7 +137,7 @@ namespace Chaos
 		/**
 		 * Set the maximum number of springs
 		 */
-		void SetMaxSprings(const FReal InMaxSprings)
+		void SetMaxSprings(const int32 InMaxSprings)
 		{
 			MaxSprings = InMaxSprings;
 		}
diff --git a/Engine/Source/Runtime/Experimental/Chaos/Public/Chaos/Plane.h b/Engine/Source/Runtime/Experimental/Chaos/Public/Chaos/Plane.h
--- a/Engine/Source/Runtime/Experimental/Chaos/Public/Chaos/Plane.h
+++ b/Engine/Source/Runtime/Experimental/Chaos/Public/Chaos/Plane.h
@@ -135,7 +135,7 @@ public:
 
 		const TVector<T, d> DirTowardsPlane = SignedDist < 0 ? MNormal : -MNormal;
 		const T RayProjectedTowardsPlane = TVector<T, d>::DotProduct(Dir, DirTowardsPlane);
-		const T Epsilon = 1e-7;
+		const T Epsilon = 1e-7f;
 		if (RayProjectedTowardsPlane < Epsilon)	//moving parallel or away
 		{
 			return false;
@@ -357,7 +357,7 @@ const TVector<T, d> FindClosestPointOnLineSegment(const TVector<T, d>& P0, const
 template<typename T, int d>
 TVector<T, d> FindClosestPointOnTriangle(const TVector<T, d>& ClosestPointOnPlane, const TVector<T, d>& P0, const TVector<T, d>& P1, const TVector<T, d>& P2, const TVector<T, d>& P)
 {
-	const T Epsilon = 1e-4;
+	const T Epsilon = 1e-4f;
 
 	const TVector<T, 2> Bary = ComputeBarycentricInPlane(P0, P1, P2, ClosestPointOnPlane);
 
diff --git a/Engine/Source/Runtime/Experimental/Chaos/Public/Chaos/RigidParticles.h b/Engine/Source/Runtime/Experimental/Chaos/Public/Chaos/RigidParticles.h
--- a/Engine/Source/Runtime/Experimental/Chaos/Public/Chaos/RigidParticles.h
+++ b/Engine/Source/Runtime/Experimental/Chaos/Public/Chaos/RigidParticles.h
@@ -52,11 +52,11 @@ constexpr int8 NumBitsNeeded(const int8 MaxValue)
 // Make a bitmask which covers the lowest NumBits bits with 1's.
 constexpr int8 LowBitsMask(const int8 NumBits)
 {
-	return NumBits == 0 ? 0 : (1 << (NumBits - 1)) | LowBitsMask(NumBits - 1);
+	return NumBits == 0 ? 0 : (int8)((1 << (NumBits - 1)) | LowBitsMask(NumBits - 1));
 }
 
 // Count N, the number of bits needed to store an object state
-static constexpr int8 ObjectStateBitCount = NumBitsNeeded((int8)EObjectStateType::Count - 1);
+static constexpr int8 ObjectStateBitCount = NumBitsNeeded((int8)EObjectStateType::Count - (int8)1);
 
 template<class T, int d>
 class TRigidParticles : public TKinematicGeometryParticles<T, d>
diff --git a/Engine/Source/Runtime/Experimental/Chaos/Public/Chaos/Rotation.h b/Engine/Source/Runtime/Experimental/Chaos/Public/Chaos/Rotation.h
--- a/Engine/Source/Runtime/Experimental/Chaos/Public/Chaos/Rotation.h
+++ b/Engine/Source/Runtime/Experimental/Chaos/Public/Chaos/Rotation.h
@@ -195,7 +195,6 @@ namespace Chaos
 		    const ::Chaos::TVector<FReal, 3>& FinalVector)
 		{
 			typedef Chaos::TVector<FReal, 3> TV;
-			typedef double T;
 			checkSlow(FMath::Abs(InitialVector.Size() - 1.0) < KINDA_SMALL_NUMBER);
 			checkSlow(FMath::Abs(FinalVector.Size() - 1.0) < KINDA_SMALL_NUMBER);
 
@@ -208,11 +207,11 @@ namespace Chaos
 				return TRotation<FReal, 3>::FromElements(InitialVector, 0.f);
 			}
 
-			const T SSquared = .5 * (1.0 + CosTheta); // Uses the half angle formula
-			const T VMagnitudeDesired = sqrt(1.0 - SSquared);
-			V *= (VMagnitudeDesired / VMagnitude);
+			const FRealDouble SSquared = .5 * (1.0 + CosTheta); // Uses the half angle formula
+			const FRealDouble VMagnitudeDesired = sqrt(1.0 - SSquared);
+			V *= static_cast<FReal>(VMagnitudeDesired / VMagnitude);
 
-			return TRotation<FReal, 3>::FromElements(V, sqrt(SSquared));
+			return TRotation<FReal, 3>::FromElements(V, static_cast<FReal>(sqrt(SSquared)));
 		}
 
 		/**
diff --git a/Engine/Source/Runtime/Experimental/Chaos/Public/Chaos/SpatialAccelerationCollection.h b/Engine/Source/Runtime/Experimental/Chaos/Public/Chaos/SpatialAccelerationCollection.h
--- a/Engine/Source/Runtime/Experimental/Chaos/Public/Chaos/SpatialAccelerationCollection.h
+++ b/Engine/Source/Runtime/Experimental/Chaos/Public/Chaos/SpatialAccelerationCollection.h
@@ -89,7 +89,7 @@ struct CHAOS_API TSpatialCollectionBucket
 		}
 		else
 		{
-			Idx = Objects.Add(MoveTemp(Obj));
+			Idx = static_cast<uint16>(Objects.Add(MoveTemp(Obj)));
 		}
 
 		return Idx;
@@ -337,19 +337,19 @@ public:
 		const int32 TypeIdx = GetTypeIdx(AccelPtr);
 		switch (TypeIdx)
 		{
-		case 0: BucketEntry.TypeInnerIdx = GetAccelerationsPerType<0>(Types).Add(&AccelPtr->template AsChecked<typename std::tuple_element<0, std::tuple<TSpatialAccelerationTypes...>>::type>()); break;
-		case 1: BucketEntry.TypeInnerIdx = GetAccelerationsPerType<ClampedIdx(1)>(Types).Add(&AccelPtr->template AsChecked<typename std::tuple_element<ClampedIdx(1), std::tuple<TSpatialAccelerationTypes...>>::type>()); break;
-		case 2: BucketEntry.TypeInnerIdx = GetAccelerationsPerType<ClampedIdx(2)>(Types).Add(&AccelPtr->template AsChecked<typename std::tuple_element<ClampedIdx(2), std::tuple<TSpatialAccelerationTypes...>>::type>()); break;
-		case 3: BucketEntry.TypeInnerIdx = GetAccelerationsPerType<ClampedIdx(3)>(Types).Add(&AccelPtr->template AsChecked<typename std::tuple_element<ClampedIdx(3), std::tuple<TSpatialAccelerationTypes...>>::type>()); break;
-		case 4: BucketEntry.TypeInnerIdx = GetAccelerationsPerType<ClampedIdx(4)>(Types).Add(&AccelPtr->template AsChecked<typename std::tuple_element<ClampedIdx(4), std::tuple<TSpatialAccelerationTypes...>>::type>()); break;
-		case 5: BucketEntry.TypeInnerIdx = GetAccelerationsPerType<ClampedIdx(5)>(Types).Add(&AccelPtr->template AsChecked<typename std::tuple_element<ClampedIdx(5), std::tuple<TSpatialAccelerationTypes...>>::type>()); break;
-		case 6: BucketEntry.TypeInnerIdx = GetAccelerationsPerType<ClampedIdx(6)>(Types).Add(&AccelPtr->template AsChecked<typename std::tuple_element<ClampedIdx(6), std::tuple<TSpatialAccelerationTypes...>>::type>()); break;
-		case 7: BucketEntry.TypeInnerIdx = GetAccelerationsPerType<ClampedIdx(7)>(Types).Add(&AccelPtr->template AsChecked<typename std::tuple_element<ClampedIdx(7), std::tuple<TSpatialAccelerationTypes...>>::type>()); break;
+		case 0: BucketEntry.TypeInnerIdx = static_cast<uint16>(GetAccelerationsPerType<0>(Types).Add(&AccelPtr->template AsChecked<typename std::tuple_element<0, std::tuple<TSpatialAccelerationTypes...>>::type>())); break;
+		case 1: BucketEntry.TypeInnerIdx = static_cast<uint16>(GetAccelerationsPerType<ClampedIdx(1)>(Types).Add(&AccelPtr->template AsChecked<typename std::tuple_element<ClampedIdx(1), std::tuple<TSpatialAccelerationTypes...>>::type>())); break;
+		case 2: BucketEntry.TypeInnerIdx = static_cast<uint16>(GetAccelerationsPerType<ClampedIdx(2)>(Types).Add(&AccelPtr->template AsChecked<typename std::tuple_element<ClampedIdx(2), std::tuple<TSpatialAccelerationTypes...>>::type>())); break;
+		case 3: BucketEntry.TypeInnerIdx = static_cast<uint16>(GetAccelerationsPerType<ClampedIdx(3)>(Types).Add(&AccelPtr->template AsChecked<typename std::tuple_element<ClampedIdx(3), std::tuple<TSpatialAccelerationTypes...>>::type>())); break;
+		case 4: BucketEntry.TypeInnerIdx = static_cast<uint16>(GetAccelerationsPerType<ClampedIdx(4)>(Types).Add(&AccelPtr->template AsChecked<typename std::tuple_element<ClampedIdx(4), std::tuple<TSpatialAccelerationTypes...>>::type>())); break;
+		case 5: BucketEntry.TypeInnerIdx = static_cast<uint16>(GetAccelerationsPerType<ClampedIdx(5)>(Types).Add(&AccelPtr->template AsChecked<typename std::tuple_element<ClampedIdx(5), std::tuple<TSpatialAccelerationTypes...>>::type>())); break;
+		case 6: BucketEntry.TypeInnerIdx = static_cast<uint16>(GetAccelerationsPerType<ClampedIdx(6)>(Types).Add(&AccelPtr->template AsChecked<typename std::tuple_element<ClampedIdx(6), std::tuple<TSpatialAccelerationTypes...>>::type>())); break;
+		case 7: BucketEntry.TypeInnerIdx = static_cast<uint16>(GetAccelerationsPerType<ClampedIdx(7)>(Types).Add(&AccelPtr->template AsChecked<typename std::tuple_element<ClampedIdx(7), std::tuple<TSpatialAccelerationTypes...>>::type>())); break;
 		}
 
 		this->ActiveBucketsMask |= (1 << BucketIdx);
 
-		Result.InnerIdx = Buckets[BucketIdx].Add(MoveTemp(BucketEntry));
+		Result.InnerIdx = static_cast<uint16>(Buckets[BucketIdx].Add(MoveTemp(BucketEntry)));
 		return Result;
 	}
 
diff --git a/Engine/Source/Runtime/Experimental/Chaos/Public/Chaos/SpatialHash.h b/Engine/Source/Runtime/Experimental/Chaos/Public/Chaos/SpatialHash.h
--- a/Engine/Source/Runtime/Experimental/Chaos/Public/Chaos/SpatialHash.h
+++ b/Engine/Source/Runtime/Experimental/Chaos/Public/Chaos/SpatialHash.h
@@ -38,7 +38,7 @@ private:
 	void Init(const T Radius);
 	void Init();
 	
-	T SmallestAxis() const
+	int32 SmallestAxis() const
 	{
 		TVec3<T> Extents = MBoundingBox.Extents();
 
diff --git a/Engine/Source/Runtime/Experimental/Chaos/Public/Chaos/Sphere.h b/Engine/Source/Runtime/Experimental/Chaos/Public/Chaos/Sphere.h
--- a/Engine/Source/Runtime/Experimental/Chaos/Public/Chaos/Sphere.h
+++ b/Engine/Source/Runtime/Experimental/Chaos/Public/Chaos/Sphere.h
@@ -122,7 +122,7 @@ namespace Chaos
 				return false;
 			}
 
-			constexpr T Epsilon = 1e-4;
+			constexpr T Epsilon = 1e-4f;
 			//we early out if starting in sphere, so using first time is always acceptable
 			T FirstTime = QuarterUnderRoot < Epsilon ? -HalfB : -HalfB - FMath::Sqrt(QuarterUnderRoot);
 			if (FirstTime >= 0 && FirstTime <= Length)
@@ -323,8 +323,8 @@ namespace Chaos
 
 		static PMatrix<T, d, d> GetInertiaTensor(const T InMass, const T InRadius, const bool bInThinShell = false)
 		{
-			static const T TwoThirds = 2. / 3;
-			static const T TwoFifths = 2. / 5;
+			static const T TwoThirds = static_cast<T>(2.0 / 3.0);
+			static const T TwoFifths = static_cast<T>(2.0 / 5.0);
 			const T Diagonal = bInThinShell ? TwoThirds * InMass * InRadius * InRadius : TwoFifths * InMass * InRadius * InRadius;
 			return PMatrix<T, d, d>(Diagonal, Diagonal, Diagonal);
 		}
@@ -399,15 +399,15 @@ namespace Chaos
 			// Polar sunflower increment: pi * (1 + sqrt(5))
 
 			// Increment = 10.16640738463053...
-			static const T Increment = PI * (1.0 + sqrt(5));
+			static const T Increment = static_cast<T>(PI * (1.0 + sqrt(5)));
 			for (int32 i = 0; i < NumPoints; i++)
 			{
-				const T Z = 0.5 + i;
+				const T Z = static_cast<T>(0.5 + i);
 				// sqrt((i+0.5) / NumPoints) sampling i = [0, NumPoints) varies: (0, 1).
 				// We then scale to the radius of our Sphere.
-				const T R = FMath::Sqrt(Z / NumPoints) * Radius;
+				const T R = FMath::Sqrt(Z / static_cast<T>(NumPoints)) * Radius;
 				// Theta increases linearly from [Increment/2, Increment*NumPoints)
-				const T Theta = Increment * (Z + SpiralSeed);
+				const T Theta = Increment * (Z + static_cast<T>(SpiralSeed));
 
 				// Convert polar coordinates to Cartesian, offset by the Sphere's location.
 				const int32 Index = i + Offset;
@@ -487,7 +487,7 @@ namespace Chaos
 			// Phi is the angle between the positive Z axis and the line from the origin to the point
 
 			// GRIncrement = 10.16640738463053...
-			static const T GRIncrement = PI * (1.0 + sqrt(5));
+			static const T GRIncrement = static_cast<FReal>(PI * (1.0 + sqrt(5)));
 
 			// If PhiSteps is 2X NumPoints, then we'll only generate half the sphere.
 			//const int32 PhiSteps = TopHalf + BottomHalf == 1 ? NumPoints * 2 : NumPoints;
@@ -500,14 +500,14 @@ namespace Chaos
 			{
 				for (int32 i = 0; i < NumPoints; i++)
 				{
-					const T Sample = 0.5 + i;
+					const T Sample = static_cast<T>(0.5 + i);
 					// ((i + 0.5) / (NumPoints * 2)) varies: (0.0, 0.5)
 					// So, (2 * (i + 0.5) / (NumPoints * 2)) varies: (0.0, 1.0)
 					// So, ((2 * (i + 0.5) / (NumPoints * 2)) - 1) varies: (-1, 0.0)
-					const T V = (2.0 * (0.5 + i) / (2.0 * NumPoints)) - 1.0;
+					const T V = static_cast<T>((2.0 * (0.5 + i) / (2.0 * NumPoints)) - 1.0);
 					const T Phi = FMath::Acos(V);
 					checkSlow(Phi > PI / 2 - KINDA_SMALL_NUMBER);
-					const T Theta = GRIncrement * (Sample + SpiralSeed);
+					const T Theta = GRIncrement * (Sample + static_cast<T>(SpiralSeed));
 
 					// Convert spherical coordinates to Cartesian, scaled by the radius of our Sphere, and offset by its location.
 					const T SinPhi = FMath::Sin(Phi);
@@ -526,11 +526,11 @@ namespace Chaos
 			{
 				for (int32 i = 0; i < NumPoints; i++)
 				{
-					const T Sample = 0.5 + i;
-					const T V = (2.0 * (0.5 + i) / (2.0 * NumPoints)); // varies: (0.0, 1.0)
+					const T Sample = static_cast<T>(0.5 + i);
+					const T V = static_cast<T>((2.0 * (0.5 + i) / (2.0 * NumPoints))); // varies: (0.0, 1.0)
 					const T Phi = FMath::Acos(V);
 					checkSlow(Phi < PI / 2 + KINDA_SMALL_NUMBER);
-					const T Theta = GRIncrement * (Sample + SpiralSeed);
+					const T Theta = GRIncrement * (Sample + static_cast<T>(SpiralSeed));
 
 					// Convert spherical coordinates to Cartesian, scaled by the radius of our Sphere, and offset by its location.
 					const T SinPhi = FMath::Sin(Phi);
@@ -549,15 +549,15 @@ namespace Chaos
 			{
 				for (int32 i = 0; i < NumPoints; i++)
 				{
-					const T Sample = 0.5 + i;
+					const T Sample = static_cast<T>(0.5 + i);
 					// arccos(x), where x = [-1, 1] varies: [PI, 0]
 					// ((i + 0.5) / NumPoints) varies: (0.0, 1.0)
 					// So, (2 * (i + 0.5) / NumPoints) varies: (0.0, 2.0)
 					// So, (1 - (2 * (i + 0.5) / NumPoints) varies: (-1.0, 1.0)
 					// So, Phi varies: (PI, 0) as i varies: [0, NumPoints-1].
-					const T Phi = FMath::Acos(1.0 - 2.0 * Sample / NumPoints);
+					const T Phi = static_cast<T>(FMath::Acos(1.0 - 2.0 * Sample / NumPoints));
 					// Theta varies: [5.0832036..., NumPoints*Increment)
-					const T Theta = GRIncrement * (Sample + SpiralSeed);
+					const T Theta = GRIncrement * (Sample + static_cast<T>(SpiralSeed));
 
 					// Convert spherical coordinates to Cartesian, scaled by the radius of our Sphere, and offset by its location.
 					const T SinPhi = FMath::Sin(Phi);
diff --git a/Engine/Source/Runtime/Experimental/Chaos/Public/Chaos/TaperedCapsule.h b/Engine/Source/Runtime/Experimental/Chaos/Public/Chaos/TaperedCapsule.h
--- a/Engine/Source/Runtime/Experimental/Chaos/Public/Chaos/TaperedCapsule.h
+++ b/Engine/Source/Runtime/Experimental/Chaos/Public/Chaos/TaperedCapsule.h
@@ -132,7 +132,7 @@ namespace Chaos
 			const FReal R2R2 = Radius2 * Radius2;
 			const FReal R1R2 = Radius1 * Radius2;
 			//  compute center of mass as a distance along the axis from the origin as the shape as the axis as a symmetry line 
-			FReal TaperedSectionCenterOfMass = (Height * (R1R1 + 2.*R1R2 + 3.*R2R2) / 4.*(R1R1 + R1R2 + R2R2));
+			FReal TaperedSectionCenterOfMass = (Height * (R1R1 + (FReal)2.0 * R1R2 + (FReal)3.0 * R2R2) / (FReal)4.0 * (R1R1 + R1R2 + R2R2));
 			FReal Hemisphere1CenterOfMass = -((FReal)3.0 * Radius1 / (FReal)8.0);
 			FReal Hemisphere2CenterOfMass = (Height + ((FReal)3.0 * Radius2 / (FReal)8.0));
 
@@ -191,12 +191,12 @@ namespace Chaos
 			const FReal R1R2 = R1 * R2;
 			const FReal R2R2 = R2 * R2;
 
-			const FReal Num1 = 2. * HH * (R1R1 + 3. * R1R2 + 6. * R2R2); // 2H^2 * (R1^2 + 3R1R2 + 6R2^2)
-			const FReal Num2 = 3. * (R1R1 * R1R1 + R1R1 * R1R2 + R1R2 * R1R2 + R1R2 * R2R2 + R2R2 * R2R2); // 3 * (R1^4 + R1^3R2 + R1^2R2^2 + R1R2^3 + R2^4)
+			const FReal Num1 = (FReal)2.0 * HH * (R1R1 + (FReal)3.0 * R1R2 + (FReal)6.0 * R2R2); // 2H^2 * (R1^2 + 3R1R2 + 6R2^2)
+			const FReal Num2 = (FReal)3.0 * (R1R1 * R1R1 + R1R1 * R1R2 + R1R2 * R1R2 + R1R2 * R2R2 + R2R2 * R2R2); // 3 * (R1^4 + R1^3R2 + R1^2R2^2 + R1R2^3 + R2^4)
 			const FReal Den1 = PI * (R1R1 + R1R2 + R2R2); // PI * (R1^2 + R1R2 + R2^2)
 
-			const FReal Diag12 = Mass * (Num1 + Num2) / (20. * Den1);
-			const FReal Diag3 = Mass * Num2 / (10. * Den1);
+			const FReal Diag12 = Mass * (Num1 + Num2) / ((FReal)20.0 * Den1);
+			const FReal Diag3 = Mass * Num2 / ((FReal)10.0 * Den1);
 
 			return FMatrix33(Diag12, Diag12, Diag3);
 		}
@@ -372,8 +372,8 @@ namespace Chaos
 			const FReal AllArea = TaperedSectionArea + Cap1Area + Cap2Area;
 			if (AllArea > KINDA_SMALL_NUMBER)
 			{
-				NumPointsEndCap1 = static_cast<int32>(round(Cap1Area / AllArea * NumPoints));
-				NumPointsEndCap2 = static_cast<int32>(round(Cap2Area / AllArea * NumPoints));
+				NumPointsEndCap1 = static_cast<int32>(round(Cap1Area / AllArea * static_cast<FReal>(NumPoints)));
+				NumPointsEndCap2 = static_cast<int32>(round(Cap2Area / AllArea * static_cast<FReal>(NumPoints)));
 				NumPointsTaperedSection = NumPoints - NumPointsEndCap1 - NumPointsEndCap2;
 			}
 			else
diff --git a/Engine/Source/Runtime/Experimental/Chaos/Public/Chaos/TaperedCylinder.h b/Engine/Source/Runtime/Experimental/Chaos/Public/Chaos/TaperedCylinder.h
--- a/Engine/Source/Runtime/Experimental/Chaos/Public/Chaos/TaperedCylinder.h
+++ b/Engine/Source/Runtime/Experimental/Chaos/Public/Chaos/TaperedCylinder.h
@@ -261,7 +261,7 @@ namespace Chaos
 			const FReal R1R1 = MRadius1 * MRadius1;
 			const FReal R2R2 = MRadius2 * MRadius2;
 			const FReal R1R2 = MRadius1 * MRadius2;
-			return FVec3(0, 0, MHeight*(R1R1 + 2.*R1R2 + 3.*R2R2) / 4.*(R1R1 + R1R2 + R2R2));
+			return FVec3(0, 0, static_cast<FReal>(MHeight*(R1R1 + 2.*R1R2 + 3.*R2R2) / 4.*(R1R1 + R1R2 + R2R2)));
 		}
 		FVec3 GetAxis() const { return (MPlane2.X() - MPlane1.X()).GetSafeNormal(); }
 
@@ -304,12 +304,12 @@ namespace Chaos
 			const FReal R1R2 = R1 * R2;
 			const FReal R2R2 = R2 * R2;
 
-			const FReal Num1 = 2. * HH * (R1R1 + 3. * R1R2 + 6. * R2R2); // 2H^2 * (R1^2 + 3R1R2 + 6R2^2)
-			const FReal Num2 = 3. * (R1R1 * R1R1 + R1R1 * R1R2 + R1R2 * R1R2 + R1R2 * R2R2 + R2R2 * R2R2); // 3 * (R1^4 + R1^3R2 + R1^2R2^2 + R1R2^3 + R2^4)
+			const FReal Num1 = static_cast<FReal>(2. * HH * (R1R1 + 3. * R1R2 + 6. * R2R2)); // 2H^2 * (R1^2 + 3R1R2 + 6R2^2)
+			const FReal Num2 = static_cast<FReal>(3. * (R1R1 * R1R1 + R1R1 * R1R2 + R1R2 * R1R2 + R1R2 * R2R2 + R2R2 * R2R2)); // 3 * (R1^4 + R1^3R2 + R1^2R2^2 + R1R2^3 + R2^4)
 			const FReal Den1 = PI * (R1R1 + R1R2 + R2R2); // PI * (R1^2 + R1R2 + R2^2)
 
-			const FReal Diag12 = Mass * (Num1 + Num2) / (20. * Den1);
-			const FReal Diag3 = Mass * Num2 / (10. * Den1);
+			const FReal Diag12 = Mass * (Num1 + Num2) / (static_cast<FReal>(20.) * Den1);
+			const FReal Diag3 = Mass * Num2 / (static_cast<FReal>(10.) * Den1);
 
 			return FMatrix33(Diag12, Diag12, Diag3);
 		}
@@ -335,7 +335,7 @@ namespace Chaos
 		FReal GetRadius(const FReal& Phi) const
 		{
 			const FReal Alpha = Phi / MHeight;
-			return MRadius1 * (1. - Alpha) + MRadius2 * Alpha;
+			return MRadius1 * (static_cast<FReal>(1.) - Alpha) + MRadius2 * Alpha;
 		}
 
 		TPlane<FReal, 3> MPlane1, MPlane2;
@@ -482,8 +482,8 @@ namespace Chaos
 				const FReal AllArea = CylArea + Cap1Area + Cap2Area;
 				if (AllArea > KINDA_SMALL_NUMBER)
 				{
-					NumPointsEndCap1 = static_cast<int32>(round(Cap1Area / AllArea * NumPoints));
-					NumPointsEndCap2 = static_cast<int32>(round(Cap2Area / AllArea * NumPoints));
+					NumPointsEndCap1 = static_cast<int32>(round(Cap1Area / AllArea * static_cast<FReal>(NumPoints)));
+					NumPointsEndCap2 = static_cast<int32>(round(Cap2Area / AllArea * static_cast<FReal>(NumPoints)));
 					NumPointsCylinder = NumPoints - NumPointsEndCap1 - NumPointsEndCap2;
 				}
 				else
@@ -527,18 +527,18 @@ namespace Chaos
 			}
 			else
 			{
-				static const FReal Increment = PI * (1.0 + sqrt(5));
+				static const FRealSingle Increment = PI * (1.0f + FMath::Sqrt(5.0f));
 				for (int32 i = 0; i < NumPointsCylinder; i++)
 				{
 					// In the 2D sphere (disc) case, we vary R so it increases monotonically,
 					// which spreads points out across the disc:
 					//     const FReal R = FMath::Sqrt((0.5 + Index) / NumPoints) * Radius;
 					// But we're mapping to a cylinder, which means we want to keep R constant.
-					const FReal R = FMath::Lerp(Radius1, Radius2, static_cast<FReal>(i) / (NumPointsCylinder - 1));
-					const FReal Theta = Increment * (0.5 + i + SpiralSeed);
+					const FReal R = FMath::Lerp(Radius1, Radius2, static_cast<FReal>(i) / static_cast<FReal>(NumPointsCylinder - 1));
+					const FReal Theta = Increment * (0.5f + static_cast<FReal>(i + SpiralSeed));
 
 					// Map polar coordinates to Cartesian, and vary Z by [-HalfHeight, HalfHeight].
-					const FReal Z = FMath::LerpStable(-HalfHeight, HalfHeight, static_cast<FReal>(i) / (NumPointsCylinder - 1));
+					const FReal Z = FMath::LerpStable(-HalfHeight, HalfHeight, static_cast<FReal>(i) / static_cast<FReal>(NumPointsCylinder - 1));
 					Points[i + Offset] =
 					    FVec3(
 					        R * FMath::Cos(Theta),
diff --git a/Engine/Source/Runtime/Experimental/Chaos/Public/Chaos/TriangleMesh.h b/Engine/Source/Runtime/Experimental/Chaos/Public/Chaos/TriangleMesh.h
--- a/Engine/Source/Runtime/Experimental/Chaos/Public/Chaos/TriangleMesh.h
+++ b/Engine/Source/Runtime/Experimental/Chaos/Public/Chaos/TriangleMesh.h
@@ -232,9 +232,9 @@ namespace Chaos
 			const int32 Idx = Particles.Size();
 			Particles.AddParticles(3);
 			// Left handed
-			Particles.X(Idx + 0) = FVec3(0., 0.8083, 0.);
-			Particles.X(Idx + 1) = FVec3(0.7, -0.4041, 0.);
-			Particles.X(Idx + 2) = FVec3(-0.7, -0.4041, 0.);
+			Particles.X(Idx + 0) = FVec3((FReal)0., (FReal)0.8083, (FReal)0.);
+			Particles.X(Idx + 1) = FVec3((FReal)0.7, (FReal)-0.4041, (FReal)0.);
+			Particles.X(Idx + 2) = FVec3((FReal)-0.7, (FReal)-0.4041, (FReal)0.);
 
 			TArray<TVec3<int32>> Elements;
 			Elements.SetNum(1);
@@ -247,9 +247,9 @@ namespace Chaos
 			const int32 Idx = Particles.Size();
 			Particles.AddParticles(3);
 			// Left handed
-			Particles.X(Idx + 0) = FVec3(0., 0., 0.8083);
-			Particles.X(Idx + 1) = FVec3(0., 0.7, -0.4041);
-			Particles.X(Idx + 2) = FVec3(0., -0.7, -0.4041);
+			Particles.X(Idx + 0) = FVec3((FReal)0., (FReal)0., (FReal)0.8083);
+			Particles.X(Idx + 1) = FVec3((FReal)0., (FReal)0.7, (FReal)-0.4041);
+			Particles.X(Idx + 2) = FVec3((FReal)0., (FReal)-0.7, (FReal)-0.4041);
 
 			TArray<TVec3<int32>> Elements;
 			Elements.SetNum(1);
diff --git a/Engine/Source/Runtime/Experimental/Chaos/Public/Chaos/UniformGrid.h b/Engine/Source/Runtime/Experimental/Chaos/Public/Chaos/UniformGrid.h
--- a/Engine/Source/Runtime/Experimental/Chaos/Public/Chaos/UniformGrid.h
+++ b/Engine/Source/Runtime/Experimental/Chaos/Public/Chaos/UniformGrid.h
@@ -26,9 +26,9 @@ class CHAOS_API TUniformGridBase
 		MDx = TVector<T, d>(MMaxCorner - MMinCorner) / MCells;
 		if (GhostCells > 0)
 		{
-			MMinCorner -= MDx * GhostCells;
-			MMaxCorner += MDx * GhostCells;
-			MCells += TVector<T, d>(2 * GhostCells);
+			MMinCorner -= MDx * static_cast<T>(GhostCells);
+			MMaxCorner += MDx * static_cast<T>(GhostCells);
+			MCells += TVector<T, d>(2 * static_cast<T>(GhostCells));
 		}
 
 		if (MDx >= TVector<T, d>(SMALL_NUMBER))
diff --git a/Engine/Source/Runtime/Experimental/Chaos/Public/GeometryCollection/GeometryCollectionAlgo.h b/Engine/Source/Runtime/Experimental/Chaos/Public/GeometryCollection/GeometryCollectionAlgo.h
--- a/Engine/Source/Runtime/Experimental/Chaos/Public/GeometryCollection/GeometryCollectionAlgo.h
+++ b/Engine/Source/Runtime/Experimental/Chaos/Public/GeometryCollection/GeometryCollectionAlgo.h
@@ -176,7 +176,7 @@ namespace GeometryCollectionAlgo
 
 	void
 	CHAOS_API
-	TriangulateBoundaries(FGeometryCollection* GeometryCollection, const TArray<TArray<TArray<int32>>> &BoundaryVertexIndices, bool bWoundClockwise = true, float MinTriangleAreaSq = 1e-4);
+	TriangulateBoundaries(FGeometryCollection* GeometryCollection, const TArray<TArray<TArray<int32>>> &BoundaryVertexIndices, bool bWoundClockwise = true, float MinTriangleAreaSq = 1e-4f);
 
 	void
 	CHAOS_API
@@ -196,7 +196,7 @@ namespace GeometryCollectionAlgo
 
 	void
 	CHAOS_API
-	DeleteCoincidentVertices(FGeometryCollection* GeometryCollection, float Tolerance = 1e-2);
+	DeleteCoincidentVertices(FGeometryCollection* GeometryCollection, float Tolerance = 1e-2f);
 
 	void
 	CHAOS_API
@@ -204,7 +204,7 @@ namespace GeometryCollectionAlgo
 
 	void
 	CHAOS_API
-	DeleteZeroAreaFaces(FGeometryCollection* GeometryCollection, float Tolerance = 1e-4);
+	DeleteZeroAreaFaces(FGeometryCollection* GeometryCollection, float Tolerance = 1e-4f);
 
 	void
 	CHAOS_API
diff --git a/Engine/Source/Runtime/Experimental/Chaos/Public/GeometryCollection/GeometryCollectionSimulationCoreTypes.h b/Engine/Source/Runtime/Experimental/Chaos/Public/GeometryCollection/GeometryCollectionSimulationCoreTypes.h
--- a/Engine/Source/Runtime/Experimental/Chaos/Public/GeometryCollection/GeometryCollectionSimulationCoreTypes.h
+++ b/Engine/Source/Runtime/Experimental/Chaos/Public/GeometryCollection/GeometryCollectionSimulationCoreTypes.h
@@ -58,10 +58,10 @@ struct FSharedSimulationParameters
 {
 	FSharedSimulationParameters()
 	: bMassAsDensity(false)
-	, Mass(1.0)
-	, MinimumMassClamp(0.1)								// todo : Expose to users with better initial values
+	, Mass(1.0f)
+	, MinimumMassClamp(0.1f)							// todo : Expose to users with better initial values
 	, MaximumMassClamp(1e5f)							// todo : Expose to users with better initial values
-	, MinimumBoundingExtentClamp(0.1)					// todo : Expose to users with better initial values
+	, MinimumBoundingExtentClamp(0.1f)					// todo : Expose to users with better initial values
 	, MaximumBoundingExtentClamp(1e6f)					// todo : Expose to users with better initial values
 	, MinimumInertiaTensorDiagonalClamp(SMALL_NUMBER)	// todo : Expose to users with better initial values
 	, MaximumInertiaTensorDiagonalClamp(1e20f)			// todo : Expose to users with better initial values
diff --git a/Engine/Source/Runtime/Experimental/Chaos/Public/GeometryCollection/RecordedTransformTrack.h b/Engine/Source/Runtime/Experimental/Chaos/Public/GeometryCollection/RecordedTransformTrack.h
--- a/Engine/Source/Runtime/Experimental/Chaos/Public/GeometryCollection/RecordedTransformTrack.h
+++ b/Engine/Source/Runtime/Experimental/Chaos/Public/GeometryCollection/RecordedTransformTrack.h
@@ -281,11 +281,6 @@ struct CHAOS_API FRecordedTransformTrack
 		return Records.Num() > 1 ? Records[1].Timestamp - Records[0].Timestamp : 0;
 	}
 
-	int32 GetLastTime() const
-	{
-		return Records.Num() > 1 ? Records.Last().Timestamp : 0;
-	}
-
 	/**
 	 * Validates that the time is within the track timestamp bounds (inclusive)
 	 */
diff --git a/Engine/Source/Runtime/Experimental/Chaos/Public/PhysicsProxy/AnalyticImplicitGroup.h b/Engine/Source/Runtime/Experimental/Chaos/Public/PhysicsProxy/AnalyticImplicitGroup.h
--- a/Engine/Source/Runtime/Experimental/Chaos/Public/PhysicsProxy/AnalyticImplicitGroup.h
+++ b/Engine/Source/Runtime/Experimental/Chaos/Public/PhysicsProxy/AnalyticImplicitGroup.h
@@ -200,11 +200,11 @@ public:
 					{
 						const Chaos::FReal VolI = BoxI.GetVolume();
 						const Chaos::FReal VolJ = BoxJ.GetVolume();
-						const Chaos::FReal PctOverlapI = VolI > KINDA_SMALL_NUMBER ? VolIJ / VolI : 0.;
-						const Chaos::FReal PctOverlapJ = VolJ > KINDA_SMALL_NUMBER ? VolIJ / VolJ : 0.;
+						const Chaos::FReal PctOverlapI = VolI > KINDA_SMALL_NUMBER ? VolIJ / VolI : 0.f;
+						const Chaos::FReal PctOverlapJ = VolJ > KINDA_SMALL_NUMBER ? VolIJ / VolJ : 0.f;
 						// Split the overlapping volume between i and j
-						MPArray[i].Volume *= (1.0 - PctOverlapI / 2);
-						MPArray[j].Volume *= (1.0 - PctOverlapJ / 2);
+						MPArray[i].Volume *= static_cast<Chaos::FReal>(1.0 - PctOverlapI / 2.0);
+						MPArray[j].Volume *= static_cast<Chaos::FReal>(1.0 - PctOverlapJ / 2.0);
 					}
 				}
 			}
diff --git a/Engine/Source/Runtime/Experimental/Chaos/Public/PhysicsProxy/FieldSystemProxyHelper.h b/Engine/Source/Runtime/Experimental/Chaos/Public/PhysicsProxy/FieldSystemProxyHelper.h
--- a/Engine/Source/Runtime/Experimental/Chaos/Public/PhysicsProxy/FieldSystemProxyHelper.h
+++ b/Engine/Source/Runtime/Experimental/Chaos/Public/PhysicsProxy/FieldSystemProxyHelper.h
@@ -8,6 +8,7 @@
 #include "Chaos/ParticleHandle.h"
 #include "Chaos/Defines.h"
 #include "Chaos/Particles.h"
+#include <limits>
 
 namespace Chaos
 {
@@ -318,7 +319,10 @@ namespace Chaos
 					Chaos::FPBDRigidParticleHandle* RigidHandle = ParticleHandles[Index.Sample]->CastToRigidParticle();
 					if (RigidHandle)
 					{
-						const int8 ResultState = ResultsView[Index.Result];
+						const int32 CurrResult = ResultsView[Index.Result];
+						check(CurrResult < std::numeric_limits<int8>::max());
+
+						const int8 ResultState = static_cast<int8>(CurrResult);
 						bHasStateChanged |= ReportDynamicStateResult(RigidSolver, static_cast<Chaos::EObjectStateType>(ResultState), RigidHandle,
 							false, Chaos::FVec3(0), false, Chaos::FVec3(0));
 					}
diff --git a/Engine/Source/Runtime/Experimental/Chaos/Public/PhysicsProxy/SkeletalMeshPhysicsProxy.h b/Engine/Source/Runtime/Experimental/Chaos/Public/PhysicsProxy/SkeletalMeshPhysicsProxy.h
--- a/Engine/Source/Runtime/Experimental/Chaos/Public/PhysicsProxy/SkeletalMeshPhysicsProxy.h
+++ b/Engine/Source/Runtime/Experimental/Chaos/Public/PhysicsProxy/SkeletalMeshPhysicsProxy.h
@@ -21,7 +21,7 @@ struct CHAOS_API FSkeletalMeshPhysicsProxyParams
 		, ObjectType(EObjectStateTypeEnum::Chaos_Object_Kinematic)
 
 		, CollisionType(ECollisionTypeEnum::Chaos_Volumetric)
-		, ParticlesPerUnitArea(0.1)
+		, ParticlesPerUnitArea(0.1f)
 		, MinNumParticles(0)
 		, MaxNumParticles(50)
 		, MinRes(5)
@@ -31,11 +31,11 @@ struct CHAOS_API FSkeletalMeshPhysicsProxyParams
 		, bEnableClustering(false)
 		, ClusterGroupIndex(0)
 		, MaxClusterLevel(100)
-		, DamageThreshold(250.)
+		, DamageThreshold(250.f)
 #endif
-		, Density(2.4)
-		, MinMass(0.001)
-		, MaxMass(1.e6)
+		, Density(2.4f)
+		, MinMass(0.001f)
+		, MaxMass(1.e6f)
 
 		, bSimulating(false)
 	{}
diff --git a/Engine/Source/Runtime/Experimental/ChaosCore/ChaosCore.Build.cs b/Engine/Source/Runtime/Experimental/ChaosCore/ChaosCore.Build.cs
--- a/Engine/Source/Runtime/Experimental/ChaosCore/ChaosCore.Build.cs
+++ b/Engine/Source/Runtime/Experimental/ChaosCore/ChaosCore.Build.cs
@@ -34,6 +34,8 @@ namespace UnrealBuildTool.Rules
             {
                 PublicDefinitions.Add("CHAOS_CHECKED=0");
             }
-        }
+
+			UnsafeTypeCastWarningLevel = WarningLevel.Warning;
+		}
     }
 }
diff --git a/Engine/Source/Runtime/Experimental/ChaosCore/Public/Chaos/Matrix.h b/Engine/Source/Runtime/Experimental/ChaosCore/Public/Chaos/Matrix.h
--- a/Engine/Source/Runtime/Experimental/ChaosCore/Public/Chaos/Matrix.h
+++ b/Engine/Source/Runtime/Experimental/ChaosCore/Public/Chaos/Matrix.h
@@ -99,7 +99,7 @@ namespace Chaos
 
 		PMatrix<FReal, 2, 2> Inverse() const
 		{
-			const FReal OneOverDeterminant = 1.0 / (M[0] * M[3] - M[1] * M[2]);
+			const FReal OneOverDeterminant = static_cast<FReal>(1.0) / (M[0] * M[3] - M[1] * M[2]);
 			return PMatrix<FReal, 2, 2>(
 			    OneOverDeterminant * M[3],
 			    -OneOverDeterminant * M[1],
diff --git a/Engine/Source/Runtime/Experimental/ChaosCore/Public/Chaos/Vector.h b/Engine/Source/Runtime/Experimental/ChaosCore/Public/Chaos/Vector.h
--- a/Engine/Source/Runtime/Experimental/ChaosCore/Public/Chaos/Vector.h
+++ b/Engine/Source/Runtime/Experimental/ChaosCore/Public/Chaos/Vector.h
@@ -458,22 +458,22 @@ namespace Chaos
 		template<class T2>
 		TVector<FReal, 3> operator+(const TVector<T2, 3>& Other) const
 		{
-			return TVector<FReal, 3>(X + Other[0], Y + Other[1], Z + Other[2]);
+			return TVector<FReal, 3>(X + static_cast<FReal>(Other[0]), Y + static_cast<FReal>(Other[1]), Z + static_cast<FReal>(Other[2]));
 		}
 		template<class T2>
 		TVector<FReal, 3> operator-(const TVector<T2, 3>& Other) const
 		{
-			return TVector<FReal, 3>(X - Other[0], Y - Other[1], Z - Other[2]);
+			return TVector<FReal, 3>(X - static_cast<FReal>(Other[0]), Y - static_cast<FReal>(Other[1]), Z - static_cast<FReal>(Other[2]));
 		}
 		template<class T2>
 		TVector<FReal, 3> operator*(const TVector<T2, 3>& Other) const
 		{
-			return TVector<FReal, 3>(X * Other[0], Y * Other[1], Z * Other[2]);
+			return TVector<FReal, 3>(X * static_cast<FReal>(Other[0]), Y * static_cast<FReal>(Other[1]), Z * static_cast<FReal>(Other[2]));
 		}
 		template<class T2>
 		TVector<FReal, 3> operator/(const TVector<T2, 3>& Other) const
 		{
-			return TVector<FReal, 3>(X / Other[0], Y / Other[1], Z / Other[2]);
+			return TVector<FReal, 3>(X / static_cast<FReal>(Other[0]), Y / static_cast<FReal>(Other[1]), Z / static_cast<FReal>(Other[2]));
 		}
 		FReal Product() const
 		{
@@ -524,7 +524,7 @@ namespace Chaos
 					return MakePair(max.Z, 2);
 			}
 		}
-		FReal SafeNormalize(FReal Epsilon = 1e-4)
+		FReal SafeNormalize(FReal Epsilon = 1e-4f)
 		{
 			FReal Size = SizeSquared();
 			if (Size < Epsilon)
@@ -639,7 +639,7 @@ namespace Chaos
 		template<class T2>
 		TVector<FReal, 2> operator/(const TVector<T2, 2>& Other) const
 		{
-			return TVector<FReal, 2>(X / Other[0], Y / Other[1]);
+			return TVector<FReal, 2>(X / static_cast<FReal>(Other[0]), Y / static_cast<FReal>(Other[1]));
 		}
 		TVector<FReal, 2> operator/(const FReal Other) const
 		{
