Fix constraint graph ensure firing

diff --git a/Engine/Source/Runtime/Experimental/Chaos/Private/Chaos/PBDConstraintGraph.cpp b/Engine/Source/Runtime/Experimental/Chaos/Private/Chaos/PBDConstraintGraph.cpp
--- a/Engine/Source/Runtime/Experimental/Chaos/Private/Chaos/PBDConstraintGraph.cpp
+++ b/Engine/Source/Runtime/Experimental/Chaos/Private/Chaos/PBDConstraintGraph.cpp
@@ -8,6 +8,7 @@
 #include "Chaos/PBDRigidParticles.h"
 #include "Containers/Queue.h"
 #include "Chaos/ParticleHandle.h"
+#include "Chaos/ConstraintHandle.h"
 #include "Chaos/PBDRigidsSOAs.h"
 #include "ChaosLog.h"
 
@@ -61,10 +62,6 @@ int32 FPBDConstraintGraph::ReserveParticles(const int32 Num)
 	return NumToAdd;
 }
 
-/**
- * Bill added this.
- * Adds new Node to Nodes array when a new particle is created
- */
 
 void FPBDConstraintGraph::ParticleAdd(FGeometryParticleHandle* AddedParticle)
 {
@@ -89,10 +86,6 @@ void FPBDConstraintGraph::ParticleAdd(FGeometryParticleHandle* AddedParticle)
 	}
 }
 
-/**
- * Bill added this
- * Removes Node from Nodes array - marking it an unused, also clears ParticleToNodeIndex
- */
 
 void FPBDConstraintGraph::ParticleRemove(FGeometryParticleHandle* RemovedParticle)
 {
@@ -102,6 +95,30 @@ void FPBDConstraintGraph::ParticleRemove(FGeometryParticleHandle* RemovedParticl
 		FreeIndexList.Push(NodeIdx);
 
 		FGraphNode& NodeRemoved = Nodes[NodeIdx];
+
+		// remove associated edges
+		for (int32 EdgeIndex : NodeRemoved.Edges)
+		{
+			// go to other particle and remove its edge reference
+			FGraphEdge& Edge = Edges[EdgeIndex];
+			if (Edge.FirstNode != Edge.SecondNode)
+			{
+				if (Edge.FirstNode == NodeIdx && Edge.SecondNode != INDEX_NONE)
+				{
+					FGraphNode& OtherNode = Nodes[Edge.SecondNode];
+					OtherNode.Edges.Remove(EdgeIndex);
+				}
+				else if (Edge.SecondNode == NodeIdx && Edge.FirstNode != INDEX_NONE)
+				{
+					FGraphNode& OtherNode = Nodes[Edge.FirstNode];
+					OtherNode.Edges.Remove(EdgeIndex);
+				}
+			}
+
+			// do we want to free the edge here also or leave it for later?
+			//Edges[EdgeIndex] = FGraphEdge();
+		}
+
 		NodeRemoved.Edges.Empty();
 		NodeRemoved.Particle = nullptr;
 		NodeRemoved.Island = INDEX_NONE;
@@ -112,7 +129,6 @@ void FPBDConstraintGraph::ParticleRemove(FGeometryParticleHandle* RemovedParticl
 }
 
 
-
 int32 FPBDConstraintGraph::GetNextNodeIndex()
 {
 	int32 NewNodeIndex = Nodes.Num();
@@ -245,8 +261,10 @@ void FPBDConstraintGraph::AddConstraint(const uint32 InContainerId, FConstraintH
 	// Must have at least one constrained particle
 	check((ConstrainedParticles[0]) || (ConstrainedParticles[1]));
 
-	// todo(chaos) : we shoudl probably have a list of indices of removed edges, so we can reuse them instead of adding to the end ( see RemoveConstraint logic )
 	const int32 NewEdgeIndex = Edges.Num();
+	
+	InConstraintHandle->SetConstraintGraphIndex(NewEdgeIndex);
+
 	FGraphEdge NewEdge;
 	NewEdge.Data = { InContainerId, InConstraintHandle };
 
@@ -286,55 +304,34 @@ void FPBDConstraintGraph::RemoveConstraint(const uint32 InContainerId, FConstrai
 {
 	check(InConstraintHandle);
 
+	int32 EdgeIndex = InConstraintHandle->GetConstraintGraphIndex();
+
+	int32 Index0 = INDEX_NONE, Index1 = INDEX_NONE;
+
 	int32* PNodeIndex0 = (ConstrainedParticles[0]) ? ParticleToNodeIndex.Find(ConstrainedParticles[0]) : nullptr;
 	int32* PNodeIndex1 = (ConstrainedParticles[1]) ? ParticleToNodeIndex.Find(ConstrainedParticles[1]) : nullptr;
 	if (ensure(PNodeIndex0 || PNodeIndex1))
 	{
-		// we need to find the edge index that matches this constraint
-		// we could go through the entire edge list and find the matching InContainerId, 
-		// but each node has an edge array that will contain it 
-		int32 EdgeIndex = INDEX_NONE;
+		// we should have an valid index at this point
 		if (PNodeIndex0)
 		{
-			EdgeIndex = GetConstraintIndexFromNodeAndConstraintHandle(Nodes[*PNodeIndex0], InConstraintHandle);
+			Nodes[*PNodeIndex0].Edges.Remove(EdgeIndex);
+			Index0 = *PNodeIndex0;
 		}
-		if (PNodeIndex1 && EdgeIndex == INDEX_NONE)
+		if (PNodeIndex1)
 		{
-			EdgeIndex = GetConstraintIndexFromNodeAndConstraintHandle(Nodes[*PNodeIndex1], InConstraintHandle);
-		}
-		// we should have an valid index at this point
-		if (EdgeIndex != INDEX_NONE)
-		{
-			int32 Index0 = INDEX_NONE, Index1 = INDEX_NONE;
-			if (PNodeIndex0)
-			{
-				Nodes[*PNodeIndex0].Edges.Remove(EdgeIndex);
-				Index0 = *PNodeIndex0;
-			}
-			if (PNodeIndex1)
-			{
-				Nodes[*PNodeIndex1].Edges.Remove(EdgeIndex);
-				Index1 = *PNodeIndex1;
-			}
-			if (Edges[EdgeIndex].FirstNode == Index0 && Edges[EdgeIndex].SecondNode == Index1)
-			{
-				Edges[EdgeIndex] = FGraphEdge();
-			}
+			Nodes[*PNodeIndex1].Edges.Remove(EdgeIndex);
+			Index1 = *PNodeIndex1;
 		}
 	}
-}
 
-int32 FPBDConstraintGraph::GetConstraintIndexFromNodeAndConstraintHandle(const FGraphNode& Node, const FConstraintHandle* ConstraintHandle)
-{
-	for (const int32 NodeEdgeIndex : Node.Edges)
+	// we need to find the edge index that matches this constraint
+	// if the particles are being removed first and the constraint is being disabled
+	if (ensure(EdgeIndex != INDEX_NONE))
 	{
-		const FConstraintData& EdgeData = Edges[NodeEdgeIndex].Data;
-		if (EdgeData.GetConstraintHandle() == ConstraintHandle)
-		{
-			return NodeEdgeIndex;
-		}
+		InConstraintHandle->SetConstraintGraphIndex(INDEX_NONE);
+		Edges[EdgeIndex] = FGraphEdge();
 	}
-	return INDEX_NONE;
 }
 
 const typename FPBDConstraintGraph::FConstraintData& FPBDConstraintGraph::GetConstraintData(int32 ConstraintDataIndex) const
diff --git a/Engine/Source/Runtime/Experimental/Chaos/Public/Chaos/ConstraintHandle.h b/Engine/Source/Runtime/Experimental/Chaos/Public/Chaos/ConstraintHandle.h
--- a/Engine/Source/Runtime/Experimental/Chaos/Public/Chaos/ConstraintHandle.h
+++ b/Engine/Source/Runtime/Experimental/Chaos/Public/Chaos/ConstraintHandle.h
@@ -30,8 +30,8 @@ namespace Chaos
 
 		using FGeometryParticleHandle = TGeometryParticleHandle<FReal, 3>;
 
-		FConstraintHandle() : Type(EType::Invalid), ConstraintIndex(INDEX_NONE) { }
-		FConstraintHandle(EType InType, int32 InConstraintIndex): Type(InType), ConstraintIndex(InConstraintIndex) {}
+		FConstraintHandle() : Type(EType::Invalid), ConstraintIndex(INDEX_NONE), ConstraintGraphIndex(INDEX_NONE) { }
+		FConstraintHandle(EType InType, int32 InConstraintIndex): Type(InType), ConstraintIndex(InConstraintIndex), ConstraintGraphIndex(INDEX_NONE) {}
 		virtual ~FConstraintHandle() {}
 
 		bool IsValid() const
@@ -44,6 +44,17 @@ namespace Chaos
 			return ConstraintIndex;
 		}
 
+		int32 GetConstraintGraphIndex() const
+		{
+			return ConstraintGraphIndex;
+		}
+
+		void SetConstraintGraphIndex(int32 InIndex)
+		{
+			ConstraintGraphIndex = InIndex;
+		}
+
+
 		virtual void SetEnabled(bool InEnabled) = 0;
 		virtual bool IsEnabled() const = 0;
 
@@ -55,6 +66,7 @@ namespace Chaos
 
 		EType Type;
 		int32 ConstraintIndex;
+		int32 ConstraintGraphIndex; // @todo(chaos): move constraint graph index to base constraint container
 	};
 
 
diff --git a/Engine/Source/Runtime/Experimental/Chaos/Public/Chaos/PBDConstraintGraph.h b/Engine/Source/Runtime/Experimental/Chaos/Public/Chaos/PBDConstraintGraph.h
--- a/Engine/Source/Runtime/Experimental/Chaos/Public/Chaos/PBDConstraintGraph.h
+++ b/Engine/Source/Runtime/Experimental/Chaos/Public/Chaos/PBDConstraintGraph.h
@@ -245,8 +245,6 @@ namespace Chaos
 		int32 GetNextNodeIndex();
 		const TArray<int32>& GetUpdatedNodes() const { return UpdatedNodes; }
 
-		int32 GetConstraintIndexFromNodeAndConstraintHandle(const FGraphNode& Node, const FConstraintHandle* ConstraintHandle);
-
 		TArray<FGraphNode> Nodes;
 		TArray<FGraphEdge> Edges;
 		TArray<FIslandData> IslandToData;
diff --git a/Engine/Source/Runtime/Experimental/Chaos/Public/Chaos/PBDRigidsEvolution.h b/Engine/Source/Runtime/Experimental/Chaos/Public/Chaos/PBDRigidsEvolution.h
--- a/Engine/Source/Runtime/Experimental/Chaos/Public/Chaos/PBDRigidsEvolution.h
+++ b/Engine/Source/Runtime/Experimental/Chaos/Public/Chaos/PBDRigidsEvolution.h
@@ -475,15 +475,19 @@ public:
 
 	/** remove a list of constraints from the constraint graph (see AddConstraintsToConstraintGraph) */
 	CHAOS_API void RemoveConstraintsFromConstraintGraph(const TArray<FConstraintHandle*>& Constraints)
-		{
+	{
 		for (FConstraintHandle* BaseConstraintHandle : Constraints)
+		{
+			if (FPBDJointConstraintHandle* ConstraintHandle = BaseConstraintHandle->As<FPBDJointConstraintHandle>())
 			{
-				if (FPBDJointConstraintHandle* ConstraintHandle = BaseConstraintHandle->As<FPBDJointConstraintHandle>())
+				// if it is already disabled then it will have already been removed from the graph
+				if (ConstraintHandle->IsConstraintEnabled())
 				{
-						ConstraintGraph.RemoveConstraint(ConstraintHandle->GetConstraintIndex(), ConstraintHandle, ConstraintHandle->GetConstrainedParticles());
-					}
+					ConstraintGraph.RemoveConstraint(ConstraintHandle->GetConstraintIndex(), ConstraintHandle, ConstraintHandle->GetConstrainedParticles());
 				}
 			}
+		}
+	}
 
 	/** Add a list of constraints to the constraint graph (see RemoveConstraintsFromConstraintGraph) */
 	CHAOS_API void AddConstraintsToConstraintGraph(const TArray<FConstraintHandle*>& Constraints)
@@ -495,7 +499,7 @@ public:
 				ConstraintGraph.AddConstraint(ConstraintHandle->GetConstraintIndex(), ConstraintHandle, ConstraintHandle->GetConstrainedParticles());
 			}
 		}
-		}
+	}
 
 	/** Disconnect constraints from a set of particles to be removed (or destroyed) 
 	* this will set the constraints to Enbaled = false and set their respective bodies handles to nullptr
@@ -514,7 +518,7 @@ public:
 	}
 
 	/** Disconnect constraints from a set of particles to be removed (or destroyed)
-	* this will set the constraints to Enbaled = false, but leave connections to the 
+	* this will set the constraints to Enabled = false, but leave connections to the 
 	* particles. 
 	*/
 	CHAOS_API void DisableConstraints(const TSet<FGeometryParticleHandle*>& DisabledParticles)
