LWC: Fix PMatrix component type for FGeometryCollectionPhysicsProxy inertial tensor to match LWC type.

diff --git a/Engine/Source/Runtime/Experimental/Chaos/Private/PhysicsProxy/GeometryCollectionPhysicsProxy.cpp b/Engine/Source/Runtime/Experimental/Chaos/Private/PhysicsProxy/GeometryCollectionPhysicsProxy.cpp
--- a/Engine/Source/Runtime/Experimental/Chaos/Private/PhysicsProxy/GeometryCollectionPhysicsProxy.cpp
+++ b/Engine/Source/Runtime/Experimental/Chaos/Private/PhysicsProxy/GeometryCollectionPhysicsProxy.cpp
@@ -1961,10 +1961,10 @@ void FGeometryCollectionPhysicsProxy::InitializeSharedCollisionStructures(
 					{
 						FVector Extents = BoundingBox[GeometryIndex].GetExtent();
 						MassProperties.Volume = Extents.X * Extents.Y * Extents.Z;
-						FRealSingle ExtentsYZ = Extents.Y * Extents.Y + Extents.Z * Extents.Z;
-						FRealSingle ExtentsXZ = Extents.X * Extents.X + Extents.Z * Extents.Z;
-						FRealSingle ExtentsXY = Extents.X * Extents.X + Extents.Y * Extents.Y;
-						MassProperties.InertiaTensor = PMatrix<FRealSingle, 3, 3>(ExtentsYZ / 12.f, ExtentsXZ / 12.f, ExtentsXY / 12.f);
+						FReal ExtentsYZ = Extents.Y * Extents.Y + Extents.Z * Extents.Z;
+						FReal ExtentsXZ = Extents.X * Extents.X + Extents.Z * Extents.Z;
+						FReal ExtentsXY = Extents.X * Extents.X + Extents.Y * Extents.Y;
+						MassProperties.InertiaTensor = PMatrix<FReal, 3, 3>(ExtentsYZ / 12.f, ExtentsXZ / 12.f, ExtentsXY / 12.f);
 						MassProperties.CenterOfMass = BoundingBox[GeometryIndex].GetCenter();
 						CollectionMassToLocal[TransformGroupIndex] = FTransform(FQuat::Identity, MassProperties.CenterOfMass);
 						InertiaComputationNeeded[GeometryIndex] = false;
