Chaos: Clustering
- Clamp infinite inertia bounds to the identity.

diff --git a/Engine/Source/Runtime/Experimental/Chaos/Private/Chaos/PBDRigidClustering.cpp b/Engine/Source/Runtime/Experimental/Chaos/Private/Chaos/PBDRigidClustering.cpp
--- a/Engine/Source/Runtime/Experimental/Chaos/Private/Chaos/PBDRigidClustering.cpp
+++ b/Engine/Source/Runtime/Experimental/Chaos/Private/Chaos/PBDRigidClustering.cpp
@@ -392,6 +392,13 @@ namespace Chaos
 		{
 			Parent->R() = ForceMassOrientation->GetRotation();
 		}
+
+		if (Parent->R().ContainsNaN())
+		{
+			InertiaTensor = PMatrix<T, d, d>(1.f, 1.f, 1.f);
+			Parent->R() = TRotation<T, d>(FMatrix::Identity);
+		}
+
 		Parent->Q() = Parent->R();
 		Parent->InvI() = Parent->I().Inverse();
 	}
