Chaos : Fix for CIS

diff --git a/Engine/Source/Runtime/Experimental/Chaos/Private/Chaos/PBDRigidClustering.cpp b/Engine/Source/Runtime/Experimental/Chaos/Private/Chaos/PBDRigidClustering.cpp
--- a/Engine/Source/Runtime/Experimental/Chaos/Private/Chaos/PBDRigidClustering.cpp
+++ b/Engine/Source/Runtime/Experimental/Chaos/Private/Chaos/PBDRigidClustering.cpp
@@ -395,8 +395,8 @@ namespace Chaos
 
 		if (Parent->R().ContainsNaN())
 		{
-			InertiaTensor = PMatrix<T, d, d>(1.f, 1.f, 1.f);
-			Parent->R() = TRotation<T, d>(FMatrix::Identity);
+			InertiaTensor = PMatrix<FReal, 3,3>(1.f, 1.f, 1.f);
+			Parent->R() = TRotation<FReal,3>(FMatrix::Identity);
 		}
 
 		Parent->Q() = Parent->R();
