Manual merge from RES.

Chaos - fix joint NumIterationsAtMaxStiffness which was meant to be an int, but was a float.

diff --git a/Engine/Source/Runtime/Experimental/Chaos/Private/PBDRigidsSolver.cpp b/Engine/Source/Runtime/Experimental/Chaos/Private/PBDRigidsSolver.cpp
--- a/Engine/Source/Runtime/Experimental/Chaos/Private/PBDRigidsSolver.cpp
+++ b/Engine/Source/Runtime/Experimental/Chaos/Private/PBDRigidsSolver.cpp
@@ -174,7 +174,7 @@ FAutoConsoleVariableRef CVarChaosSolverCollisionUseManifolds(TEXT("p.Chaos.Solve
 // Joint cvars
 float ChaosSolverJointMinSolverStiffness = 1.0f;
 float ChaosSolverJointMaxSolverStiffness = 1.0f;
-float ChaosSolverJointNumIterationsAtMaxSolverStiffness = 1;
+int32 ChaosSolverJointNumIterationsAtMaxSolverStiffness = 1;
 FAutoConsoleVariableRef CVarChaosSolverJointMinSolverStiffness(TEXT("p.Chaos.Solver.Joint.MinSolverStiffness"), ChaosSolverJointMinSolverStiffness, TEXT("Solver stiffness on first iteration, increases each iteration toward MaxSolverStiffness."));
 FAutoConsoleVariableRef CVarChaosSolverJointMaxSolverStiffness(TEXT("p.Chaos.Solver.Joint.MaxSolverStiffness"), ChaosSolverJointMaxSolverStiffness, TEXT("Solver stiffness on last iteration, increases each iteration from MinSolverStiffness."));
 FAutoConsoleVariableRef CVarChaosSolverJointNumIterationsAtMaxSolverStiffness(TEXT("p.Chaos.Solver.Joint.NumIterationsAtMaxSolverStiffness"), ChaosSolverJointNumIterationsAtMaxSolverStiffness, TEXT("How many iterations we want at MaxSolverStiffness."));
diff --git a/Engine/Source/Runtime/Experimental/Chaos/Public/Chaos/PBDJointConstraintTypes.h b/Engine/Source/Runtime/Experimental/Chaos/Public/Chaos/PBDJointConstraintTypes.h
--- a/Engine/Source/Runtime/Experimental/Chaos/Public/Chaos/PBDJointConstraintTypes.h
+++ b/Engine/Source/Runtime/Experimental/Chaos/Public/Chaos/PBDJointConstraintTypes.h
@@ -190,7 +190,7 @@ namespace Chaos
 		// Solver Stiffness (increases over iterations)
 		FReal MinSolverStiffness;
 		FReal MaxSolverStiffness;
-		FReal NumIterationsAtMaxSolverStiffness;
+		int32 NumIterationsAtMaxSolverStiffness;
 
 		// @todo(ccaulfield): remove these TEMP overrides for testing
 		bool bEnableTwistLimits;
