Chaos Cloth - Add paintable TetherScale, Drag, and Lift weightmaps.
Chaos Cloth - Change anim mesh visualization color to better differentiate it against the physics mesh.
[Integrate CL 16297075 via 16313688]

Cloth - Add a ClearColllisionSources method so that all collision sources can be removed without knowledge of any of the original sources. [Integrate CL 16305946 via 16313688]

diff --git a/Engine/Plugins/Experimental/ChaosCloth/Source/Chaos/Private/ChaosCloth/ChaosClothConfig.cpp b/Engine/Plugins/Experimental/ChaosCloth/Source/Chaos/Private/ChaosCloth/ChaosClothConfig.cpp
--- a/Engine/Plugins/Experimental/ChaosCloth/Source/Chaos/Private/ChaosCloth/ChaosClothConfig.cpp
+++ b/Engine/Plugins/Experimental/ChaosCloth/Source/Chaos/Private/ChaosCloth/ChaosClothConfig.cpp
@@ -61,12 +61,12 @@ void UChaosClothConfig::MigrateFrom(const FClothConfig_Legacy& ClothConfig)
 
 	TetherStiffness.Low = FMath::Clamp(ClothConfig.TetherStiffness, 0.f, 1.f);
 	TetherStiffness.High = 1.f;
-	LimitScale = FMath::Clamp(ClothConfig.TetherLimit, 0.01f, 10.f);
+	TetherScale.Low = TetherScale.High = FMath::Clamp(ClothConfig.TetherLimit, 0.01f, 10.f);
 	ShapeTargetStiffness = 0.f;
 
 	bUsePointBasedWindModel = (ClothConfig.WindMethod == EClothingWindMethod_Legacy::Legacy);
-	DragCoefficient = bUsePointBasedWindModel ? 0.07f  : ClothConfig.WindDragCoefficient;  // Only Accurate wind uses the WindDragCoefficient
-	LiftCoefficient = bUsePointBasedWindModel ? 0.035f : ClothConfig.WindLiftCoefficient;  // Only Accurate wind uses the WindLiftCoefficient
+	Drag.Low = Drag.High = bUsePointBasedWindModel ? 0.07f  : ClothConfig.WindDragCoefficient;  // Only Accurate wind uses the WindDragCoefficient
+	Lift.Low = Lift.High = bUsePointBasedWindModel ? 0.035f : ClothConfig.WindLiftCoefficient;  // Only Accurate wind uses the WindLiftCoefficient
 
 	const float Damping = (ClothConfig.Damping.X + ClothConfig.Damping.Y + ClothConfig.Damping.Z) / 3.f;
 	DampingCoefficient = FMath::Clamp(Damping * Damping * 0.7f, 0.f, 1.f);  // Nv Cloth seems to have a different damping formulation.
@@ -131,7 +131,7 @@ void UChaosClothConfig::PostLoad()
 
 	if (ChaosClothConfigCustomVersion < FChaosClothConfigCustomVersion::UpdateDragDefault)
 	{
-		DragCoefficient = 0.07f;  // Reset to a more appropriate default for chaos cloth assets saved before this custom version
+		DragCoefficient_DEPRECATED = 0.07f;  // Reset to a more appropriate default for chaos cloth assets saved before this custom version that had too high drag
 	}
 
 	if (ChaosClothConfigCustomVersion < FChaosClothConfigCustomVersion::RemoveInternalConfigParameters)
@@ -168,7 +168,14 @@ void UChaosClothConfig::PostLoad()
 	{
 		// Note: Unlike AnimDriveStiffness, Low is updated here, because there was no existing weight map before this version
 		TetherStiffness.Low = FMath::Clamp(FMath::Loge(StrainLimitingStiffness_DEPRECATED) / FMath::Loge(1.e3f) + 1.f, 0.f, 1.f);
-		TetherStiffness.High = 0.f;
+		TetherStiffness.High = 1.f;
+	}
+
+	if (FortniteMainBranchObjectVersion < FFortniteMainBranchObjectVersion::ChaosClothAddTetherScaleAndDragLiftWeightMaps)
+	{
+		TetherScale.Low = TetherScale.High = LimitScale_DEPRECATED;
+		Drag.Low = Drag.High = DragCoefficient_DEPRECATED;
+		Lift.Low = Lift.High = LiftCoefficient_DEPRECATED;
 	}
 #endif  // #if WITH_EDITORONLY_DATA
 }
diff --git a/Engine/Plugins/Experimental/ChaosCloth/Source/Chaos/Private/ChaosCloth/ChaosClothConstraints.cpp b/Engine/Plugins/Experimental/ChaosCloth/Source/Chaos/Private/ChaosCloth/ChaosClothConstraints.cpp
--- a/Engine/Plugins/Experimental/ChaosCloth/Source/Chaos/Private/ChaosCloth/ChaosClothConstraints.cpp
+++ b/Engine/Plugins/Experimental/ChaosCloth/Source/Chaos/Private/ChaosCloth/ChaosClothConstraints.cpp
@@ -329,7 +329,10 @@ void FClothConstraints::SetVolumeConstraints(TArray<TVec3<int32>>&& SurfaceEleme
 	++NumConstraintRules;
 }
 
-void FClothConstraints::SetLongRangeConstraints(const TMap<int32, TSet<int32>>& PointToNeighborsMap, const TConstArrayView<FRealSingle>& TetherStiffnessMultipliers, const FVec2& TetherStiffness, FReal LimitScale, ETetherMode TetherMode, bool bUseXPBDConstraints)
+void FClothConstraints::SetLongRangeConstraints(const TMap<int32, TSet<int32>>& PointToNeighborsMap,
+	const TConstArrayView<FRealSingle>& TetherStiffnessMultipliers, const TConstArrayView<FRealSingle>& TetherScaleMultipliers,
+	const FVec2& TetherStiffness, const FVec2& TetherScale,
+	ETetherMode TetherMode, bool bUseXPBDConstraints)
 {
 	check(Evolution);
 
@@ -343,8 +346,11 @@ void FClothConstraints::SetLongRangeConstraints(const TMap<int32, TSet<int32>>&
 			NumParticles,
 			PointToNeighborsMap,
 			TetherStiffnessMultipliers,
+			TetherScaleMultipliers,
 			MaxNumTetherIslands,
-			TetherStiffness);  // TODO: Add LimitScale to the XPBD constraint
+			TetherStiffness,
+			TetherScale,
+			TetherMode);
 	}
 	else
 	{
@@ -354,9 +360,10 @@ void FClothConstraints::SetLongRangeConstraints(const TMap<int32, TSet<int32>>&
 			NumParticles,
 			PointToNeighborsMap,
 			TetherStiffnessMultipliers,
+			TetherScaleMultipliers,
 			MaxNumTetherIslands,
 			TetherStiffness,
-			LimitScale,
+			TetherScale,
 			TetherMode);
 	}
 	++NumConstraintInits;  // Uses init to update the property tables
diff --git a/Engine/Plugins/Experimental/ChaosCloth/Source/Chaos/Private/ChaosCloth/ChaosClothingSimulation.cpp b/Engine/Plugins/Experimental/ChaosCloth/Source/Chaos/Private/ChaosCloth/ChaosClothingSimulation.cpp
--- a/Engine/Plugins/Experimental/ChaosCloth/Source/Chaos/Private/ChaosCloth/ChaosClothingSimulation.cpp
+++ b/Engine/Plugins/Experimental/ChaosCloth/Source/Chaos/Private/ChaosCloth/ChaosClothingSimulation.cpp
@@ -385,7 +385,7 @@ void FClothingSimulation::CreateActor(USkeletalMeshComponent* InOwnerComponent,
 		ClothConfig->VolumeStiffness,
 		ClothConfig->bUseThinShellVolumeConstraints,
 		TVec2<FRealSingle>(ClothConfig->TetherStiffness.Low, ClothConfig->TetherStiffness.High),  // Animatable
-		ClothConfig->LimitScale,
+		TVec2<FRealSingle>(ClothConfig->TetherScale.Low, ClothConfig->TetherScale.High),  // Animatable
 		ClothConfig->bUseGeodesicDistance ? FClothingSimulationCloth::ETetherMode::Geodesic : FClothingSimulationCloth::ETetherMode::Euclidean,
 		/*MaxDistancesMultiplier =*/ 1.f,  // Animatable
 		TVec2<FRealSingle>(ClothConfig->AnimDriveStiffness.Low, ClothConfig->AnimDriveStiffness.High),  // Animatable
@@ -398,8 +398,8 @@ void FClothingSimulation::CreateActor(USkeletalMeshComponent* InOwnerComponent,
 		ClothConfig->LinearVelocityScale,
 		ClothConfig->AngularVelocityScale,
 		ClothConfig->FictitiousAngularScale,
-		ClothConfig->DragCoefficient,
-		ClothConfig->LiftCoefficient,
+		TVec2<FRealSingle>(ClothConfig->Drag.Low, ClothConfig->Drag.High),  // Animatable
+		TVec2<FRealSingle>(ClothConfig->Lift.Low, ClothConfig->Lift.High),  // Animatable
 		ClothConfig->bUsePointBasedWindModel,
 		ClothConfig->DampingCoefficient,
 		ClothConfig->CollisionThickness,
@@ -746,7 +746,7 @@ void FClothingSimulation::RefreshClothConfig(const IClothingSimulationContext* I
 			ClothConfig->VolumeStiffness,
 			ClothConfig->bUseThinShellVolumeConstraints,
 			TVec2<FRealSingle>(ClothConfig->TetherStiffness.Low, ClothConfig->TetherStiffness.High),  // Animatable
-			ClothConfig->LimitScale,
+			TVec2<FRealSingle>(ClothConfig->TetherScale.Low, ClothConfig->TetherScale.High),  // Animatable
 			ClothConfig->bUseGeodesicDistance ? FClothingSimulationCloth::ETetherMode::Geodesic : FClothingSimulationCloth::ETetherMode::Euclidean,
 			/*MaxDistancesMultiplier =*/ 1.f,  // Animatable
 			TVec2<FRealSingle>(ClothConfig->AnimDriveStiffness.Low, ClothConfig->AnimDriveStiffness.High),  // Animatable
@@ -759,8 +759,8 @@ void FClothingSimulation::RefreshClothConfig(const IClothingSimulationContext* I
 			ClothConfig->LinearVelocityScale,
 			ClothConfig->AngularVelocityScale,
 			ClothConfig->FictitiousAngularScale,
-			ClothConfig->DragCoefficient,
-			ClothConfig->LiftCoefficient,
+			TVec2<FRealSingle>(ClothConfig->Drag.Low, ClothConfig->Drag.High),  // Animatable
+			TVec2<FRealSingle>(ClothConfig->Lift.Low, ClothConfig->Lift.High),  // Animatable
 			ClothConfig->bUsePointBasedWindModel,
 			ClothConfig->DampingCoefficient,
 			ClothConfig->CollisionThickness,
@@ -1263,7 +1263,6 @@ void FClothingSimulation::DebugDrawPhysMeshWired(FPrimitiveDrawInterface* PDI) c
 
 void FClothingSimulation::DebugDrawAnimMeshWired(FPrimitiveDrawInterface* PDI) const
 {
-	static const FLinearColor DynamicColor = FColor::White;
 	static const FLinearColor KinematicColor = FColor::Purple;
 
 	const FVec3& LocalSpaceLocation = Solver->GetLocalSpaceLocation();
@@ -1289,13 +1288,9 @@ void FClothingSimulation::DebugDrawAnimMeshWired(FPrimitiveDrawInterface* PDI) c
 			const FVector Pos1 = LocalSpaceLocation + Positions[Element.Y - Offset];
 			const FVector Pos2 = LocalSpaceLocation + Positions[Element.Z - Offset];
 
-			const bool bIsKinematic0 = (InvMasses[Element.X - Offset] == 0.f);
-			const bool bIsKinematic1 = (InvMasses[Element.Y - Offset] == 0.f);
-			const bool bIsKinematic2 = (InvMasses[Element.Z - Offset] == 0.f);
-
-			DrawLine(PDI, Pos0, Pos1, bIsKinematic0 && bIsKinematic1 ? KinematicColor : DynamicColor);
-			DrawLine(PDI, Pos1, Pos2, bIsKinematic1 && bIsKinematic2 ? KinematicColor : DynamicColor);
-			DrawLine(PDI, Pos2, Pos0, bIsKinematic2 && bIsKinematic0 ? KinematicColor : DynamicColor);
+			DrawLine(PDI, Pos0, Pos1, KinematicColor);
+			DrawLine(PDI, Pos1, Pos2, KinematicColor);
+			DrawLine(PDI, Pos2, Pos0, KinematicColor);
 		}
 	}
 }
diff --git a/Engine/Plugins/Experimental/ChaosCloth/Source/Chaos/Private/ChaosCloth/ChaosClothingSimulationCloth.cpp b/Engine/Plugins/Experimental/ChaosCloth/Source/Chaos/Private/ChaosCloth/ChaosClothingSimulationCloth.cpp
--- a/Engine/Plugins/Experimental/ChaosCloth/Source/Chaos/Private/ChaosCloth/ChaosClothingSimulationCloth.cpp
+++ b/Engine/Plugins/Experimental/ChaosCloth/Source/Chaos/Private/ChaosCloth/ChaosClothingSimulationCloth.cpp
@@ -194,12 +194,14 @@ void FClothingSimulationCloth::FLODData::Add(FClothingSimulationSolver* Solver,
 	{
 		const TMap<int32, TSet<int32>>& PointToNeighborsMap = TriangleMesh.GetPointToNeighborsMap();
 		const TConstArrayView<FRealSingle>& TetherStiffnessMultipliers = WeightMaps[(int32)EChaosWeightMapTarget::TetherStiffness];
+		const TConstArrayView<FRealSingle>& TetherScaleMultipliers = WeightMaps[(int32)EChaosWeightMapTarget::TetherScale];
 
 		ClothConstraints.SetLongRangeConstraints(
 			PointToNeighborsMap,
 			TetherStiffnessMultipliers,
+			TetherScaleMultipliers,
 			FVec2((FReal)Cloth->TetherStiffness[0], (FReal)Cloth->TetherStiffness[1]),
-			(FReal)Cloth->LimitScale,
+			FVec2((FReal)Cloth->TetherScale[0], (FReal)Cloth->TetherScale[1]),
 			Cloth->TetherMode,
 			Cloth->bUseXPBDConstraints);
 	}
@@ -341,7 +343,7 @@ FClothingSimulationCloth::FClothingSimulationCloth(
 	FRealSingle InVolumeStiffness,
 	bool bInUseThinShellVolumeConstraints,
 	const TVec2<FRealSingle>& InTetherStiffness,
-	FRealSingle InLimitScale,
+	const TVec2<FRealSingle>& InTetherScale,
 	ETetherMode InTetherMode,
 	FRealSingle InMaxDistancesMultiplier,
 	const TVec2<FRealSingle>& InAnimDriveStiffness,
@@ -354,8 +356,8 @@ FClothingSimulationCloth::FClothingSimulationCloth(
 	const FVec3& InLinearVelocityScale,
 	FRealSingle InAngularVelocityScale,
 	FRealSingle InFictitiousAngularScale,
-	FRealSingle InDragCoefficient,
-	FRealSingle InLiftCoefficient,
+	const TVec2<FRealSingle>& InDrag,
+	const TVec2<FRealSingle>& InLift,
 	bool bInUseLegacyWind,
 	FRealSingle InDampingCoefficient,
 	FRealSingle InCollisionThickness,
@@ -379,7 +381,7 @@ FClothingSimulationCloth::FClothingSimulationCloth(
 	, VolumeStiffness(InVolumeStiffness)
 	, bUseThinShellVolumeConstraints(bInUseThinShellVolumeConstraints)
 	, TetherStiffness(InTetherStiffness)
-	, LimitScale(InLimitScale)
+	, TetherScale(InTetherScale)
 	, TetherMode(InTetherMode)
 	, MaxDistancesMultiplier(InMaxDistancesMultiplier)
 	, AnimDriveStiffness(InAnimDriveStiffness)
@@ -392,9 +394,10 @@ FClothingSimulationCloth::FClothingSimulationCloth(
 	, LinearVelocityScale(InLinearVelocityScale)
 	, AngularVelocityScale(InAngularVelocityScale)
 	, FictitiousAngularScale(InFictitiousAngularScale)
-	, DragCoefficient(InDragCoefficient)
-	, LiftCoefficient(InLiftCoefficient)
+	, Drag(InDrag)
+	, Lift(InLift)
 	, WindVelocity((FReal)0., (FReal)0., (FReal)0.)  // Set by clothing interactor
+	, AirDensity(1.225e-6f)  // Set by clothing interactor
 	, bUseLegacyWind(bInUseLegacyWind)
 	, DampingCoefficient(InDampingCoefficient)
 	, CollisionThickness(InCollisionThickness)
@@ -712,6 +715,12 @@ void FClothingSimulationCloth::Update(FClothingSimulationSolver* Solver)
 				Solver->GetParticlePs(Offset),
 				Solver->GetParticleXs(Offset),
 				Solver->GetParticleVs(Offset));
+
+				// Update the wind velocity field for the new LOD mesh
+				const TConstArrayView<FRealSingle>& DragMultipliers = LODData[LODIndex].WeightMaps[(int32)EChaosWeightMapTarget::Drag];
+				const TConstArrayView<FRealSingle>& LiftMultipliers = LODData[LODIndex].WeightMaps[(int32)EChaosWeightMapTarget::Lift];
+
+				Solver->SetWindGeometry(GroupId, GetTriangleMesh(Solver), DragMultipliers, LiftMultipliers);
 		}
 		else
 		{
@@ -736,11 +745,11 @@ void FClothingSimulationCloth::Update(FClothingSimulationSolver* Solver)
 
 		if (bUseLegacyWind && ChaosClothingSimulationClothConsoleVariables::CVarLegacyDisablesAccurateWind.GetValueOnAnyThread())
 		{
-			Solver->SetWindVelocityField(GroupId, 0.f, 0.f, &GetTriangleMesh(Solver));
+			Solver->SetWindProperties(GroupId, FVec2::ZeroVector, FVec2::ZeroVector, (FReal)0.);  // Disable the wind velocity field
 		}
 		else
 		{
-			Solver->SetWindVelocityField(GroupId, DragCoefficient, LiftCoefficient, &GetTriangleMesh(Solver));
+			Solver->SetWindProperties(GroupId, FVec2((FReal)Drag[0], (FReal)Drag[1]), FVec2((FReal)Lift[0], (FReal)Lift[1]), (FReal)AirDensity);  // TODO: Pass AirDensity too
 		}
 		Solver->SetWindVelocity(GroupId, WindVelocity + Solver->GetWindVelocity());
 
diff --git a/Engine/Plugins/Experimental/ChaosCloth/Source/Chaos/Private/ChaosCloth/ChaosClothingSimulationInteractor.cpp b/Engine/Plugins/Experimental/ChaosCloth/Source/Chaos/Private/ChaosCloth/ChaosClothingSimulationInteractor.cpp
--- a/Engine/Plugins/Experimental/ChaosCloth/Source/Chaos/Private/ChaosCloth/ChaosClothingSimulationInteractor.cpp
+++ b/Engine/Plugins/Experimental/ChaosCloth/Source/Chaos/Private/ChaosCloth/ChaosClothingSimulationInteractor.cpp
@@ -69,7 +69,16 @@ void UChaosClothingInteractor::SetAerodynamics(float DragCoefficient, float Lift
 {
 	Commands.Add(FChaosClothingInteractorCommand::CreateLambda([DragCoefficient, LiftCoefficient, WindVelocity](FClothingSimulationCloth* Cloth)
 	{
-		Cloth->SetAerodynamicsProperties(DragCoefficient, LiftCoefficient, WindVelocity);
+		constexpr FRealSingle AirDensity = 1.225e-6f;
+		Cloth->SetAerodynamicsProperties(TVec2<FRealSingle>(DragCoefficient, DragCoefficient), TVec2<FRealSingle>(LiftCoefficient, LiftCoefficient), AirDensity, WindVelocity);
+	}));
+}
+
+void UChaosClothingInteractor::SetWind(FVector2D Drag, FVector2D Lift, float AirDensity, FVector WindVelocity)
+{
+	Commands.Add(FChaosClothingInteractorCommand::CreateLambda([Drag, Lift, AirDensity, WindVelocity](FClothingSimulationCloth* Cloth)
+	{
+		Cloth->SetAerodynamicsProperties(TVec2<FRealSingle>(Drag[0], Drag[1]), TVec2<FRealSingle>(Lift[0], Lift[1]), AirDensity, WindVelocity);
 	}));
 }
 
diff --git a/Engine/Plugins/Experimental/ChaosCloth/Source/Chaos/Private/ChaosCloth/ChaosClothingSimulationSolver.cpp b/Engine/Plugins/Experimental/ChaosCloth/Source/Chaos/Private/ChaosCloth/ChaosClothingSimulationSolver.cpp
--- a/Engine/Plugins/Experimental/ChaosCloth/Source/Chaos/Private/ChaosCloth/ChaosClothingSimulationSolver.cpp
+++ b/Engine/Plugins/Experimental/ChaosCloth/Source/Chaos/Private/ChaosCloth/ChaosClothingSimulationSolver.cpp
@@ -41,7 +41,6 @@ namespace ChaosClothingSimulationSolverDefault
 {
 	static const FVec3 Gravity(0.f, 0.f, -980.665f);  // cm/s^2
 	static const FVec3 WindVelocity(0.f);
-	static const FRealSingle WindFluidDensity = 1.225e-6f;  // kg/cm^3
 	static const int32 NumIterations = 1;
 	static const int32 NumSubsteps = 1;
 	static const FRealSingle SelfCollisionThickness = 2.f;
@@ -67,8 +66,7 @@ FClothingSimulationSolver::FClothingSimulationSolver()
 	, CollisionParticlesSize(0)
 	, Gravity(ChaosClothingSimulationSolverDefault::Gravity)
 	, WindVelocity(ChaosClothingSimulationSolverDefault::WindVelocity)
-	, LegacyWindAdaption(0.f)
-	, WindFluidDensity(ChaosClothingSimulationSolverDefault::WindFluidDensity)
+	, LegacyWindAdaption((FReal)0.)
 	, bIsClothGravityOverrideEnabled(false)
 {
 	FPBDParticles LocalParticles;
@@ -609,7 +607,7 @@ void FClothingSimulationSolver::SetGravity(uint32 GroupId, const FVec3& InGravit
 void FClothingSimulationSolver::SetWindVelocity(const FVec3& InWindVelocity, FRealSingle InLegacyWindAdaption)
 {
 	WindVelocity = InWindVelocity * ChaosClothingSimulationSolverConstant::WorldScale;
-	LegacyWindAdaption = InLegacyWindAdaption;
+	LegacyWindAdaption = (FReal)InLegacyWindAdaption;
 }
 
 void FClothingSimulationSolver::SetWindVelocity(uint32 GroupId, const FVec3& InWindVelocity)
@@ -618,11 +616,16 @@ void FClothingSimulationSolver::SetWindVelocity(uint32 GroupId, const FVec3& InW
 	VelocityField.SetVelocity(InWindVelocity);
 }
 
-void FClothingSimulationSolver::SetWindVelocityField(uint32 GroupId, FRealSingle DragCoefficient, FRealSingle LiftCoefficient, const FTriangleMesh* TriangleMesh)
+void FClothingSimulationSolver::SetWindGeometry(uint32 GroupId, const FTriangleMesh& TriangleMesh, const TConstArrayView<FRealSingle>& DragMultipliers, const TConstArrayView<FRealSingle>& LiftMultipliers)
 {
 	FVelocityField& VelocityField = Evolution->GetVelocityField(GroupId);
-	VelocityField.SetGeometry(TriangleMesh);
-	VelocityField.SetCoefficients(DragCoefficient, LiftCoefficient);
+	VelocityField.SetGeometry(&TriangleMesh, DragMultipliers, LiftMultipliers);
+}
+
+void FClothingSimulationSolver::SetWindProperties(uint32 GroupId, const FVec2& Drag, const FVec2& Lift, FReal AirDensity)
+{
+	FVelocityField& VelocityField = Evolution->GetVelocityField(GroupId);
+	VelocityField.SetProperties(Drag, Lift, AirDensity);
 }
 
 const FVelocityField& FClothingSimulationSolver::GetWindVelocityField(uint32 GroupId)
@@ -810,7 +813,6 @@ void FClothingSimulationSolver::Update(FReal InDeltaTime)
 			// Pre-update overridable solver properties first
 			Evolution->GetGravityForces(GroupId).SetAcceleration(Gravity);
 			Evolution->GetVelocityField(GroupId).SetVelocity(WindVelocity);
-			Evolution->GetVelocityField(GroupId).SetFluidDensity(WindFluidDensity);
 
 			Cloth->Update(this);
 		}, /*bForceSingleThreaded =*/ !bChaosClothSolverParallelClothUpdate);
diff --git a/Engine/Plugins/Experimental/ChaosCloth/Source/Chaos/Private/ChaosCloth/ChaosWeightMapTarget.h b/Engine/Plugins/Experimental/ChaosCloth/Source/Chaos/Private/ChaosCloth/ChaosWeightMapTarget.h
--- a/Engine/Plugins/Experimental/ChaosCloth/Source/Chaos/Private/ChaosCloth/ChaosWeightMapTarget.h
+++ b/Engine/Plugins/Experimental/ChaosCloth/Source/Chaos/Private/ChaosCloth/ChaosWeightMapTarget.h
@@ -16,5 +16,8 @@ enum class EChaosWeightMapTarget : uint8
 	AnimDriveStiffness = (uint8)EWeightMapTargetCommon::AnimDriveStiffness,
 	// Add Chaos specific maps below this line
 	AnimDriveDamping,
-	TetherStiffness
+	TetherStiffness,
+	TetherScale,
+	Drag,
+	Lift
 };
diff --git a/Engine/Plugins/Experimental/ChaosCloth/Source/Chaos/Public/ChaosCloth/ChaosClothConfig.h b/Engine/Plugins/Experimental/ChaosCloth/Source/Chaos/Public/ChaosCloth/ChaosClothConfig.h
--- a/Engine/Plugins/Experimental/ChaosCloth/Source/Chaos/Public/ChaosCloth/ChaosClothConfig.h
+++ b/Engine/Plugins/Experimental/ChaosCloth/Source/Chaos/Public/ChaosCloth/ChaosClothConfig.h
@@ -141,9 +141,15 @@ public:
 	UPROPERTY(EditAnywhere, Category = "Long Range Attachment", meta = (UIMin = "0", UIMax = "1", ClampMin = "0", ClampMax = "1"))
 	FChaosClothWeightedValue TetherStiffness = { 1.f, 1.f };
 
-	/** The limit scale of the long range attachment constraints (aka tether limit). */
+	/**
+	 * The limit scale of the long range attachment constraints (aka tether limit).
+	 * If an enabled Weight Map (A.K.A. Mask) targeting the "Tether Scale" is added to the cloth, 
+	 * then both the Low and High values will be used in conjunction with the per particle Weight stored
+	 * in the Weight Map to interpolate the final value from them.
+	 * Otherwise only the Low value is meaningful and sufficient to set the tethers' scale.
+	*/
 	UPROPERTY(EditAnywhere, Category = "Long Range Attachment", meta = (UIMin = "1.", UIMax = "1.1", ClampMin = "0.01", ClampMax = "10"))
-	float LimitScale = 1.f;
+	FChaosClothWeightedValue TetherScale = { 1.f, 1.f };
 
 	/**
 	 * Use geodesic instead of euclidean distance calculations for the Long Range Attachment constraint,
@@ -201,13 +207,25 @@ public:
 	UPROPERTY(EditAnywhere, Category = "Environmental Properties")
 	bool bUsePointBasedWindModel = false;
 
-	// The aerodynamic coefficient of drag applying on each particle.
+	/**
+	 * The aerodynamic coefficient of drag applying on each particle.
+	 * If an enabled Weight Map (A.K.A. Mask) targeting the "Drag" is added to the cloth, 
+	 * then both the Low and High values will be used in conjunction with the per particle Weight stored
+	 * in the Weight Map to interpolate the final value from them.
+	 * Otherwise only the Low value is meaningful and sufficient to set the aerodynamic drag.
+	 */
 	UPROPERTY(EditAnywhere, Category = "Environmental Properties", meta = (UIMin = "0", UIMax = "1", ClampMin = "0", ClampMax = "10", EditCondition = "!bUsePointBasedWindModel"))
-	float DragCoefficient = 0.07f;
+	FChaosClothWeightedValue Drag = { 0.035f, 1.f };
 
-	// The aerodynamic coefficient of lift applying on each particle.
+	/**
+	 * The aerodynamic coefficient of lift applying on each particle.
+	 * If an enabled Weight Map (A.K.A. Mask) targeting the "Lift" is added to the cloth, 
+	 * then both the Low and High values will be used in conjunction with the per particle Weight stored
+	 * in the Weight Map to interpolate the final value from them.
+	 * Otherwise only the Low value is meaningful and sufficient to set the aerodynamic lift.
+	 */
 	UPROPERTY(EditAnywhere, Category = "Environmental Properties", meta = (UIMin = "0", UIMax = "1", ClampMin = "0", ClampMax = "10", EditCondition = "!bUsePointBasedWindModel"))
-	float LiftCoefficient = 0.035f;
+	FChaosClothWeightedValue Lift = { 0.035f, 1.f };
 
 	// Use the config gravity value instead of world gravity.
 	UPROPERTY(EditAnywhere, Category = "Environmental Properties", meta = (InlineEditConditionToggle))
@@ -282,6 +300,15 @@ public:
 	UPROPERTY()
 	EChaosClothTetherMode TetherMode_DEPRECATED = EChaosClothTetherMode::MaxChaosClothTetherMode;
 
+	UPROPERTY()
+	float LimitScale_DEPRECATED = 1.f;
+
+	UPROPERTY()
+	float DragCoefficient_DEPRECATED = 0.07f;
+
+	UPROPERTY()
+	float LiftCoefficient_DEPRECATED = 0.035f;
+
 	UPROPERTY()
 	float AnimDriveSpringStiffness_DEPRECATED = 0.f;
 
diff --git a/Engine/Plugins/Experimental/ChaosCloth/Source/Chaos/Public/ChaosCloth/ChaosClothConstraints.h b/Engine/Plugins/Experimental/ChaosCloth/Source/Chaos/Public/ChaosCloth/ChaosClothConstraints.h
--- a/Engine/Plugins/Experimental/ChaosCloth/Source/Chaos/Public/ChaosCloth/ChaosClothConstraints.h
+++ b/Engine/Plugins/Experimental/ChaosCloth/Source/Chaos/Public/ChaosCloth/ChaosClothConstraints.h
@@ -50,7 +50,9 @@ namespace Chaos
 		void SetAreaConstraints(TArray<TVec3<int32>>&& SurfaceElements, FReal AreaStiffness, bool bUseXPBDConstraints);
 		void SetVolumeConstraints(TArray<TVec2<int32>>&& DoubleBendingEdges, FReal VolumeStiffness);
 		void SetVolumeConstraints(TArray<TVec3<int32>>&& SurfaceElements, FReal VolumeStiffness);
-		void SetLongRangeConstraints(const TMap<int32, TSet<int32>>& PointToNeighborsMap, const TConstArrayView<FRealSingle>& TetherStiffnessMultipliers, const FVec2& TetherStiffness, FReal LimitScale, ETetherMode TetherMode, bool bUseXPBDConstraints);
+		void SetLongRangeConstraints(const TMap<int32, TSet<int32>>& PointToNeighborsMap,
+			const TConstArrayView<FRealSingle>& TetherStiffnessMultipliers, const TConstArrayView<FRealSingle>& TetherScaleMultipliers,
+			const FVec2& TetherStiffness, const FVec2& TetherScale, ETetherMode TetherMode, bool bUseXPBDConstraints);
 		void SetMaximumDistanceConstraints(const TConstArrayView<FRealSingle>& MaxDistances);
 		void SetBackstopConstraints(const TConstArrayView<FRealSingle>& BackstopDistances, const TConstArrayView<FRealSingle>& BackstopRadiuses, bool bUseLegacyBackstop);
 		void SetAnimDriveConstraints(const TConstArrayView<FRealSingle>& AnimDriveStiffnessMultipliers, const TConstArrayView<FRealSingle>& AnimDriveDampingMultipliers);
diff --git a/Engine/Plugins/Experimental/ChaosCloth/Source/Chaos/Public/ChaosCloth/ChaosClothingSimulationCloth.h b/Engine/Plugins/Experimental/ChaosCloth/Source/Chaos/Public/ChaosCloth/ChaosClothingSimulationCloth.h
--- a/Engine/Plugins/Experimental/ChaosCloth/Source/Chaos/Public/ChaosCloth/ChaosClothingSimulationCloth.h
+++ b/Engine/Plugins/Experimental/ChaosCloth/Source/Chaos/Public/ChaosCloth/ChaosClothingSimulationCloth.h
@@ -40,7 +40,7 @@ namespace Chaos
 			FRealSingle InVolumeStiffness,
 			bool bInUseThinShellVolumeConstraints,
 			const TVec2<FRealSingle>& InTetherStiffness,
-			FRealSingle InLimitScale,
+			const TVec2<FRealSingle>& InTetherScale,
 			ETetherMode InTetherMode,
 			FRealSingle InMaxDistancesMultiplier,
 			const TVec2<FRealSingle>& InAnimDriveStiffness,
@@ -53,8 +53,8 @@ namespace Chaos
 			const FVec3& InLinearVelocityScale,
 			FRealSingle InAngularVelocityScale,
 			FRealSingle InFictitiousAngularScale,
-			FRealSingle InDragCoefficient,
-			FRealSingle InLiftCoefficient,
+			const TVec2<FRealSingle>& InDrag,
+			const TVec2<FRealSingle>& InLift,
 			bool bInUseLegacyWind,
 			FRealSingle InDampingCoefficient,
 			FRealSingle InCollisionThickness,
@@ -80,7 +80,7 @@ namespace Chaos
 		void SetLongRangeAttachmentProperties(const TVec2<FRealSingle>& InTetherStiffness) { TetherStiffness = InTetherStiffness; }
 		void SetCollisionProperties(FRealSingle InCollisionThickness, FRealSingle InFrictionCoefficient, bool bInUseCCD, FRealSingle InSelfCollisionThickness) { CollisionThickness = InCollisionThickness; FrictionCoefficient = InFrictionCoefficient; bUseCCD = bInUseCCD; SelfCollisionThickness = InSelfCollisionThickness; }
 		void SetDampingProperties(FRealSingle InDampingCoefficient) { DampingCoefficient = InDampingCoefficient; }
-		void SetAerodynamicsProperties(FRealSingle InDragCoefficient, FRealSingle InLiftCoefficient, const FVec3& InWindVelocity) { DragCoefficient = InDragCoefficient; LiftCoefficient = InLiftCoefficient; WindVelocity = InWindVelocity; }
+		void SetAerodynamicsProperties(const TVec2<FRealSingle>& InDrag, const TVec2<FRealSingle>& InLift, FRealSingle InAirDensity, const FVec3& InWindVelocity) { Drag = InDrag; Lift = InLift; InAirDensity = AirDensity; WindVelocity = InWindVelocity; }
 		void SetGravityProperties(FRealSingle InGravityScale, bool bInIsGravityOverridden, const FVec3& InGravityOverride) { GravityScale = InGravityScale; bIsGravityOverridden = bInIsGravityOverridden; GravityOverride = InGravityOverride; }
 		void SetAnimDriveProperties(const TVec2<FRealSingle>& InAnimDriveStiffness, const TVec2<FRealSingle>& InAnimDriveDamping) { AnimDriveStiffness = InAnimDriveStiffness; AnimDriveDamping = InAnimDriveDamping; }
 		void GetAnimDriveProperties(TVec2<FRealSingle>& OutAnimDriveStiffness, TVec2<FRealSingle>& OutAnimDriveDamping) { OutAnimDriveStiffness = AnimDriveStiffness; OutAnimDriveDamping = AnimDriveDamping; }
@@ -191,7 +191,7 @@ namespace Chaos
 		FRealSingle VolumeStiffness;
 		bool bUseThinShellVolumeConstraints;
 		TVec2<FRealSingle> TetherStiffness;
-		FRealSingle LimitScale;
+		TVec2<FRealSingle> TetherScale;
 		ETetherMode TetherMode;
 		FRealSingle MaxDistancesMultiplier;  // Animatable
 		TVec2<FRealSingle> AnimDriveStiffness;  // Animatable
@@ -204,9 +204,10 @@ namespace Chaos
 		FVec3 LinearVelocityScale;  // Linear ratio applied to the reference bone transforms
 		FRealSingle AngularVelocityScale;  // Angular ratio factor applied to the reference bone transforms
 		FRealSingle FictitiousAngularScale;
-		FRealSingle DragCoefficient;
-		FRealSingle LiftCoefficient;
+		TVec2<FRealSingle> Drag;
+		TVec2<FRealSingle> Lift;
 		FVec3 WindVelocity;
+		FRealSingle AirDensity;
 		bool bUseLegacyWind;
 		FRealSingle DampingCoefficient;
 		FRealSingle CollisionThickness;
diff --git a/Engine/Plugins/Experimental/ChaosCloth/Source/Chaos/Public/ChaosCloth/ChaosClothingSimulationInteractor.h b/Engine/Plugins/Experimental/ChaosCloth/Source/Chaos/Public/ChaosCloth/ChaosClothingSimulationInteractor.h
--- a/Engine/Plugins/Experimental/ChaosCloth/Source/Chaos/Public/ChaosCloth/ChaosClothingSimulationInteractor.h
+++ b/Engine/Plugins/Experimental/ChaosCloth/Source/Chaos/Public/ChaosCloth/ChaosClothingSimulationInteractor.h
@@ -39,9 +39,13 @@ public:
 	UFUNCTION(BlueprintCallable, Category = ClothingSimulation, Meta = (Keywords = "Chaos Coefficient"))
 	void SetDamping(float DampingCoefficient = 0.01f);
 
+	// Deprecated. This function cannot set different Low and High values for the Drag and Lift weight maps. Use SetWind instead.
 	UFUNCTION(BlueprintCallable, Category = ClothingSimulation, Meta = (Keywords = "Chaos Wind Drag Lift Coefficient"))
 	void SetAerodynamics(float DragCoefficient = 0.07f, float LiftCoefficient = 0.035f, FVector WindVelocity = FVector(0.f, 0.f, 0.f));
 
+	UFUNCTION(BlueprintCallable, Category = ClothingSimulation, Meta = (Keywords = "Chaos Wind Drag Lift Coefficient Air Density Velocity"))
+	void SetWind(FVector2D Drag = FVector2D(0.07f, 0.5f), FVector2D Lift = FVector2D(0.07f, 0.5f), float AirDensity = 1.225e-6f, FVector WindVelocity = FVector(0.f, 0.f, 0.f));
+
 	UFUNCTION(BlueprintCallable, Category = ClothingSimulation, Meta = (Keywords = "Chaos Scale Override"))
 	void SetGravity(float GravityScale = 1.f, bool bIsGravityOverridden = false, FVector GravityOverride = FVector(0.f, 0.f, -981.f));
 
diff --git a/Engine/Plugins/Experimental/ChaosCloth/Source/Chaos/Public/ChaosCloth/ChaosClothingSimulationSolver.h b/Engine/Plugins/Experimental/ChaosCloth/Source/Chaos/Public/ChaosCloth/ChaosClothingSimulationSolver.h
--- a/Engine/Plugins/Experimental/ChaosCloth/Source/Chaos/Public/ChaosCloth/ChaosClothingSimulationSolver.h
+++ b/Engine/Plugins/Experimental/ChaosCloth/Source/Chaos/Public/ChaosCloth/ChaosClothingSimulationSolver.h
@@ -39,7 +39,9 @@ namespace Chaos
 
 		void SetWindVelocity(const FVec3& InWindVelocity, FRealSingle InLegacyWindAdaption = 0.f);
 		const FVec3& GetWindVelocity() const { return WindVelocity; }
-		void SetWindFluidDensity(FRealSingle InWindFluidDensity) { WindFluidDensity = InWindFluidDensity; }
+
+		UE_DEPRECATED(4.27, "Use SetWindProperties instead.")
+		void SetWindFluidDensity(FRealSingle /*WindFluidDensity*/) {}
 
 		void SetNumIterations(int32 InNumIterations) { NumIterations = InNumIterations; }
 		int32 GetNumIterations() const { return NumIterations; }
@@ -74,7 +76,6 @@ namespace Chaos
 		void SetParticleMassFromTotalMass(int32 Offset, FReal TotalMass, FReal MinPerParticleMass, const FTriangleMesh& Mesh, const TFunctionRef<bool(int32)>& KinematicPredicate);
 		void SetParticleMassFromDensity(int32 Offset, FReal Density, FReal MinPerParticleMass, const FTriangleMesh& Mesh, const TFunctionRef<bool(int32)>& KinematicPredicate);
 
-
 		// Set the amount of velocity allowed to filter from the given change in reference space transform, including local simulation space.
 		void SetReferenceVelocityScale(uint32 GroupId,
 			const FRigidTransform3& OldReferenceSpaceTransform,
@@ -95,8 +96,27 @@ namespace Chaos
 		// Set per group wind velocity, used to override solver's wind velocity. Must be called during cloth update.
 		void SetWindVelocity(uint32 GroupId, const FVec3& InWindVelocity);
 
-		// Set the geometry affected by wind, or disable if TriangleMesh is null.
-		void SetWindVelocityField(uint32 GroupId, FRealSingle DragCoefficient, FRealSingle LiftCoefficient, const FTriangleMesh* TriangleMesh = nullptr);
+		UE_DEPRECATED(4.27, "Use SetWindGeometry and SetWindPropertiesinstead.")
+		void SetWindVelocityField(uint32 GroupId, FRealSingle DragCoefficient, FRealSingle LiftCoefficient, const FTriangleMesh* TriangleMesh = nullptr)
+		{
+			if (TriangleMesh)
+			{
+				SetWindGeometry(GroupId, *TriangleMesh, TConstArrayView<FRealSingle>(), TConstArrayView<FRealSingle>());
+				SetWindProperties(GroupId, FVec2((FReal)DragCoefficient), FVec2((FReal)LiftCoefficient));
+			}
+			else
+			{
+				SetWindProperties(GroupId, FVec2((FReal)0.), FVec2((FReal)0.));
+			}
+		}
+
+		// Set the geometry affected by the wind.
+		void SetWindGeometry(uint32 GroupId, const FTriangleMesh& TriangleMesh, const TConstArrayView<FRealSingle>& DragMultipliers, const TConstArrayView<FRealSingle>& LiftMultipliers);
+
+		// Set the wind properties.
+		void SetWindProperties(uint32 GroupId, const FVec2& Drag, const FVec2& Lift, FReal AirDensity = (FReal)1.225e-6);
+
+		// Return the wind velocity field associated with a given group id.
 		const FVelocityField&  GetWindVelocityField(uint32 GroupId);
 
 		// Add external forces to the particles
@@ -199,8 +219,7 @@ namespace Chaos
 		// Solver parameters
 		FVec3 Gravity;
 		FVec3 WindVelocity;
-		FRealSingle LegacyWindAdaption;
-		FRealSingle WindFluidDensity;
+		FReal LegacyWindAdaption;
 		bool bIsClothGravityOverrideEnabled;
 
 		// Field system unique to the cloth solver
diff --git a/Engine/Source/Runtime/Core/Public/Math/Vector.isph b/Engine/Source/Runtime/Core/Public/Math/Vector.isph
--- a/Engine/Source/Runtime/Core/Public/Math/Vector.isph
+++ b/Engine/Source/Runtime/Core/Public/Math/Vector.isph
@@ -7,6 +7,19 @@
 #include "Math/Soaaos.isph"
 
 // Opt-in for explicit vector structs. Can assist with alignment and codegen in some cases.
+#ifndef EXPLICIT_VECTOR2
+#define EXPLICIT_VECTOR2 0
+#endif
+
+struct FVector2
+{
+#if EXPLICIT_VECTOR2 == 0
+	float V[2];
+#else
+	float<2> V;
+#endif
+};
+
 #ifndef EXPLICIT_VECTOR
 #define EXPLICIT_VECTOR 0
 #endif
@@ -239,6 +252,36 @@ inline void VectorStore(uniform FIntVector *uniform DstPtr, const varying FIntVe
 	}
 }
 
+inline varying FVector2 VectorLoad(const uniform FVector2 *uniform SrcPtr)
+{
+	varying FVector2 Result;
+
+	if(((1<<TARGET_WIDTH)-1 ^ lanemask()) == 0)
+	{
+		aos_to_soa2_ispc((uniform float * uniform)SrcPtr, &Result.V[0], &Result.V[1]);
+	}
+	else
+	{
+		#pragma ignore warning(perf)
+		Result = SrcPtr[programIndex];
+	}
+
+	return Result;
+}
+
+inline void VectorStore(uniform FVector2 *uniform DstPtr, const varying FVector2 &V)
+{
+	if(((1<<TARGET_WIDTH)-1 ^ lanemask()) == 0)
+	{
+		soa_to_aos2_ispc(V.V[0], V.V[1], (uniform float *uniform)DstPtr);
+	}
+	else
+	{
+		#pragma ignore warning(perf)
+		DstPtr[programIndex] = V;
+	}
+}
+
 inline varying FVector VectorLoad(const uniform FVector *uniform SrcPtr)
 {
 	varying FVector Result;
@@ -299,6 +342,18 @@ inline void VectorStore(uniform FVector4 *uniform DstPtr, const varying FVector4
 	}
 }
 
+inline varying FVector2 VectorGather(const uniform FVector2 *varying SrcPtr)
+{
+	#pragma ignore warning(perf)
+	return *SrcPtr;
+}
+
+inline void VectorScatter(uniform FVector2 *varying DstPtr, const varying FVector2 &V)
+{
+	#pragma ignore warning(perf)
+	*DstPtr = V;
+}
+
 inline varying FVector VectorGather(const uniform FVector *varying SrcPtr)
 {
 	#pragma ignore warning(perf)
diff --git a/Engine/Source/Runtime/Core/Public/UObject/FortniteMainBranchObjectVersion.h b/Engine/Source/Runtime/Core/Public/UObject/FortniteMainBranchObjectVersion.h
--- a/Engine/Source/Runtime/Core/Public/UObject/FortniteMainBranchObjectVersion.h
+++ b/Engine/Source/Runtime/Core/Public/UObject/FortniteMainBranchObjectVersion.h
@@ -161,6 +161,9 @@ struct CORE_API FFortniteMainBranchObjectVersion
 		// Fix corrupted LOD transition maps
 		ChaosClothFixLODTransitionMaps,
 
+		// Enable a few more weight maps to better art direct the cloth simulation
+		ChaosClothAddTetherScaleAndDragLiftWeightMaps,
+
 		// -----<new versions can be added above this line>-------------------------------------------------
 		VersionPlusOne,
 		LatestVersion = VersionPlusOne - 1
diff --git a/Engine/Source/Runtime/Engine/Classes/Components/SkeletalMeshComponent.h b/Engine/Source/Runtime/Engine/Classes/Components/SkeletalMeshComponent.h
--- a/Engine/Source/Runtime/Engine/Classes/Components/SkeletalMeshComponent.h
+++ b/Engine/Source/Runtime/Engine/Classes/Components/SkeletalMeshComponent.h
@@ -2175,6 +2175,9 @@ public:
 	/** Remove a cloth collision source defined by both a component and a physics asset */
 	void RemoveClothCollisionSource(USkeletalMeshComponent* InSourceComponent, UPhysicsAsset* InSourcePhysicsAsset);
 
+	/** Remove all cloth collision sources */
+	void ResetClothCollisionSources();
+
 protected:
 	/** copy cloth collision sources to this, where parent means components above it in the hierarchy */
 	void CopyClothCollisionSources();
@@ -2197,6 +2200,7 @@ public:
 	void AddClothCollisionSource(USkeletalMeshComponent* InSourceComponent, UPhysicsAsset* InSourcePhysicsAsset) {}
 	void RemoveClothCollisionSource(USkeletalMeshComponent* InSourceComponent) {}
 	void RemoveClothCollisionSource(USkeletalMeshComponent* InSourceComponent, UPhysicsAsset* InSourcePhysicsAsset) {}
+	void ResetClothCollisionSources() {}
 
 #endif
 
diff --git a/Engine/Source/Runtime/Engine/Private/SkeletalMeshComponentPhysics.cpp b/Engine/Source/Runtime/Engine/Private/SkeletalMeshComponentPhysics.cpp
--- a/Engine/Source/Runtime/Engine/Private/SkeletalMeshComponentPhysics.cpp
+++ b/Engine/Source/Runtime/Engine/Private/SkeletalMeshComponentPhysics.cpp
@@ -3221,6 +3221,11 @@ void USkeletalMeshComponent::RemoveClothCollisionSource(USkeletalMeshComponent*
 	}
 }
 
+void USkeletalMeshComponent::ResetClothCollisionSources()
+{
+	ClothCollisionSources.Reset();
+}
+
 #endif// #if WITH_CLOTH_COLLISION_DETECTION
 
 void USkeletalMeshComponent::EndPhysicsTickComponent(FSkeletalMeshComponentEndPhysicsTickFunction& ThisTickFunction)
diff --git a/Engine/Source/Runtime/Experimental/Chaos/Private/Chaos/PBDLongRangeConstraintsBase.cpp b/Engine/Source/Runtime/Experimental/Chaos/Private/Chaos/PBDLongRangeConstraintsBase.cpp
--- a/Engine/Source/Runtime/Experimental/Chaos/Private/Chaos/PBDLongRangeConstraintsBase.cpp
+++ b/Engine/Source/Runtime/Experimental/Chaos/Private/Chaos/PBDLongRangeConstraintsBase.cpp
@@ -16,9 +16,10 @@ FPBDLongRangeConstraintsBase::FPBDLongRangeConstraintsBase(
 	const int32 InParticleCount,
 	const TMap<int32, TSet<int32>>& PointToNeighbors,
 	const TConstArrayView<FRealSingle>& StiffnessMultipliers,
+	const TConstArrayView<FRealSingle>& ScaleMultipliers,
 	const int32 MaxNumTetherIslands,
 	const FVec2& InStiffness,
-	const FReal LimitScale,
+	const FVec2& Scale,
 	const EMode InMode)
 	: TethersView(Tethers)
 	, Stiffness(StiffnessMultipliers, InStiffness, InParticleCount)
@@ -40,9 +41,23 @@ FPBDLongRangeConstraintsBase::FPBDLongRangeConstraintsBase(
 	}
 
 	// Scale the tether's reference lengths
-	for (FTether& Tether : Tethers)
+	const FReal ScaleOffset = FMath::Clamp(Scale[0], (FReal)0.01, (FReal)10.);
+
+	if (ScaleMultipliers.Num() == InParticleCount)
+	{
+		const FReal ScaleRange = FMath::Clamp(Scale[1], (FReal)0.01, (FReal)10.) - Scale[0];
+
+		for (FTether& Tether : Tethers)
+		{
+			Tether.RefLength *= ScaleOffset + ScaleRange * FMath::Clamp((FReal)ScaleMultipliers[Tether.End], (FReal)0., (FReal)1.);
+		}
+	}
+	else
 	{
-		Tether.RefLength *= LimitScale;
+		for (FTether& Tether : Tethers)
+		{
+			Tether.RefLength *= ScaleOffset;
+		}
 	}
 }
 
diff --git a/Engine/Source/Runtime/Experimental/Chaos/Private/Chaos/VelocityField.cpp b/Engine/Source/Runtime/Experimental/Chaos/Private/Chaos/VelocityField.cpp
--- a/Engine/Source/Runtime/Experimental/Chaos/Private/Chaos/VelocityField.cpp
+++ b/Engine/Source/Runtime/Experimental/Chaos/Private/Chaos/VelocityField.cpp
@@ -13,13 +13,85 @@ FAutoConsoleVariableRef CVarChaosVelocityFieldISPCEnabled(TEXT("p.Chaos.Velocity
 
 using namespace Chaos;
 
+void FVelocityField::SetGeometry(const FTriangleMesh* TriangleMesh, const TConstArrayView<FRealSingle>& DragMultipliers, const TConstArrayView<FRealSingle>& LiftMultipliers)
+{
+	if (TriangleMesh)
+	{
+		PointToTriangleMap = TriangleMesh->GetPointToTriangleMap();
+		Elements = TriangleMesh->GetElements();
+				
+		const TVec2<int32> Range = TriangleMesh->GetVertexRange();
+		Offset = Range[0];
+		NumParticles = 1 + Range[1] - Offset;
+
+		Forces.SetNumUninitialized(Elements.Num());
+
+		enum class EWeighMaps : uint8
+		{
+			None = 0,
+			DragMultipliers = 1 << 0,
+			LiftMultipliers = 1 << 1
+		};
+
+		const EWeighMaps WeighMaps = (EWeighMaps)(
+			(uint8)(DragMultipliers.Num() == NumParticles ? EWeighMaps::DragMultipliers : EWeighMaps::None) |
+			(uint8)(LiftMultipliers.Num() == NumParticles ? EWeighMaps::LiftMultipliers : EWeighMaps::None));
+
+		if (WeighMaps == EWeighMaps::None)
+		{
+			Multipliers.Reset();
+		}
+		else
+		{
+			static const FReal OneThird = (FReal)1. / (FReal)3.;
+
+			Multipliers.SetNumUninitialized(Elements.Num());
+
+			for (int32 ElementIndex = 0; ElementIndex < Elements.Num(); ++ElementIndex)
+			{
+				const TVec3<int32>& Element = Elements[ElementIndex];
+				const int32 I0 = Element[0] - Offset;
+				const int32 I1 = Element[1] - Offset;
+				const int32 I2 = Element[2] - Offset;
+
+				switch (WeighMaps)
+				{
+				case EWeighMaps::DragMultipliers:
+					Multipliers[ElementIndex] = FVec2((DragMultipliers[I0] + DragMultipliers[I1] + DragMultipliers[I2]) * OneThird, (FReal)0.); break;
+					break;
+				case EWeighMaps::LiftMultipliers:
+					Multipliers[ElementIndex] = FVec2((FReal)0., (LiftMultipliers[I0] + LiftMultipliers[I1] + LiftMultipliers[I2]) * OneThird); break;
+					break;
+				default:
+					Multipliers[ElementIndex] = FVec2(
+						DragMultipliers[I0] + DragMultipliers[I1] + DragMultipliers[I2],
+						LiftMultipliers[I0] + LiftMultipliers[I1] + LiftMultipliers[I2]) * OneThird;
+					break;
+				}
+			}
+		}
+	}
+	else
+	{
+		PointToTriangleMap = TArrayView<TArray<int32>>();
+		Elements = TArrayView<TVector<int32, 3>>();
+		Offset = 0;
+		NumParticles = 0;
+		Forces.Reset();
+		Multipliers.Reset();
+		SetProperties(FVec2::ZeroVector, FVec2::ZeroVector, (FReal)0.);
+	}
+}
+
+
+
 void FVelocityField::UpdateForces(const FPBDParticles& InParticles, const FReal /*Dt*/)
 {
-	if (!GetVelocity)
+	if (!Multipliers.Num())
 	{
+#if INTEL_ISPC
 		if (bRealTypeCompatibleWithISPC && bChaos_VelocityField_ISPC_Enabled)
 		{
-#if INTEL_ISPC
 			ispc::UpdateField(
 				(ispc::FVector*)Forces.GetData(),
 				(const ispc::FIntVector*)Elements.GetData(),
@@ -27,32 +99,49 @@ void FVelocityField::UpdateForces(const FPBDParticles& InParticles, const FReal
 				(const ispc::FVector*)InParticles.XArray().GetData(),
 				(const ispc::FVector&)Velocity,
 				QuarterRho,
-				Cd,
-				Cl,
+				DragBase,
+				LiftBase,
 				Elements.Num());
-#endif
 		}
 		else
+#endif
 		{
 			for (int32 ElementIndex = 0; ElementIndex < Elements.Num(); ++ElementIndex)
 			{
-				UpdateField(InParticles, ElementIndex, Velocity);
+				UpdateField(InParticles, ElementIndex, Velocity, DragBase, LiftBase);
 			}
 		}
 	}
 	else
 	{
-		for (int32 ElementIndex = 0; ElementIndex < Elements.Num(); ++ElementIndex)
+#if INTEL_ISPC
+		if (bRealTypeCompatibleWithISPC && bChaos_VelocityField_ISPC_Enabled)
 		{
-			const TVector<int32, 3>& Element = Elements[ElementIndex];
-
-			// Get the triangle's position
-			const FVec3& SurfacePosition = (FReal)(1. / 3.) * (
-				InParticles.X(Element[0]) +
-				InParticles.X(Element[1]) +
-				InParticles.X(Element[2]));
+			ispc::UpdateFieldWithWeightMaps(
+				(ispc::FVector*)Forces.GetData(),
+				(const ispc::FIntVector*)Elements.GetData(),
+				(const ispc::FVector*)InParticles.GetV().GetData(),
+				(const ispc::FVector*)InParticles.XArray().GetData(),
+				(const ispc::FVector2*)Multipliers.GetData(),
+				(const ispc::FVector&)Velocity,
+				QuarterRho,
+				DragBase,
+				DragRange,
+				LiftBase,
+				LiftRange,
+				Elements.Num());
+		}
+		else
+#endif
+		{
+			for (int32 ElementIndex = 0; ElementIndex < Elements.Num(); ++ElementIndex)
+			{
+				const FVec2& Multiplier = Multipliers[ElementIndex];
+				const FReal Cd = DragBase + DragRange * Multiplier[0];
+				const FReal Cl = LiftBase + LiftRange * Multiplier[1];
 
-			UpdateField(InParticles, ElementIndex, GetVelocity(SurfacePosition));
+				UpdateField(InParticles, ElementIndex, Velocity, Cd, Cl);
+			}
 		}
 	}
 }
diff --git a/Engine/Source/Runtime/Experimental/Chaos/Private/Chaos/VelocityField.ispc b/Engine/Source/Runtime/Experimental/Chaos/Private/Chaos/VelocityField.ispc
--- a/Engine/Source/Runtime/Experimental/Chaos/Private/Chaos/VelocityField.ispc
+++ b/Engine/Source/Runtime/Experimental/Chaos/Private/Chaos/VelocityField.ispc
@@ -74,3 +74,74 @@ export void UpdateField(uniform FVector Forces[],
 		Offset += programCount;
 	}
 }
+
+export void UpdateFieldWithWeightMaps(uniform FVector Forces[],
+								const uniform FIntVector Elements[],
+								const uniform FVector V[],
+								const uniform FVector X[],
+								const uniform FVector2 Multipliers[],
+								const uniform FVector& Velocity,
+								const uniform float QuarterRho,
+								const uniform float DragBase,
+								const uniform float DragRange,
+								const uniform float LiftBase,
+								const uniform float LiftRange,
+								const uniform int32 NumElements)
+{
+	varying FIntVector Element;
+	uniform int Offset = 0;
+
+#if HW_GATHER_SUPPORTED == 1
+	if(programIndex < NumElements)
+	{
+		Element = VectorLoad(&Elements[Offset]);
+	}
+#endif
+
+	foreach(ElementIndex = 0 ... NumElements)
+	{
+#if HW_GATHER_SUPPORTED == 0
+		Element = VectorLoad(&Elements[Offset]);
+#endif
+		const FVector X0 = VectorGather(&X[Element.V[0]]);
+		const FVector X1 = VectorGather(&X[Element.V[1]]);
+		const FVector X2 = VectorGather(&X[Element.V[2]]);
+
+		const FVector2 Multiplier = VectorGather(&Multipliers[ElementIndex]);
+
+		// Calculate the normal and the area of the surface exposed to the flow
+		FVector N = VectorCross(X1 - X0, X2 - X0);
+		const float DoubleArea = SafeNormalize(N);
+
+		const FVector V0 = VectorGather(&V[Element.V[0]]);
+		const FVector V1 = VectorGather(&V[Element.V[1]]);
+		const FVector V2 = VectorGather(&V[Element.V[2]]);
+
+		// Calculate the direction and the relative velocity of the triangle to the flow
+		const FVector SurfaceVelocity = (1.0f/3.0f) * (V0 + V1 + V2);
+		const FVector RelVelocity = Velocity - SurfaceVelocity;
+
+		// Set the aerodynamic forces
+		const float VDotN = VectorDot(RelVelocity, N);
+		const float VSquare = VectorDot(RelVelocity, RelVelocity);
+
+		const float Cd = DragBase + DragRange * Multiplier.V[0];
+		const float Cl = LiftBase + LiftRange * Multiplier.V[1];
+
+#if HW_GATHER_SUPPORTED == 1
+		if(ElementIndex + programCount < NumElements)
+		{
+			Element = VectorLoad(&Elements[Offset + programCount]);
+		}
+#endif
+
+		const FVector EvenForce = (Cd - Cl) * VDotN * RelVelocity + Cl * VSquare * N;
+		const FVector OddForce = (Cl - Cd) * VDotN * RelVelocity - Cl * VSquare * N;
+
+		const FVector Force = QuarterRho * DoubleArea * VectorSelect(VDotN >= 0.0f, EvenForce, OddForce);
+
+		VectorStore(&Forces[Offset], Force);
+
+		Offset += programCount;
+	}
+}
diff --git a/Engine/Source/Runtime/Experimental/Chaos/Public/Chaos/PBDLongRangeConstraints.h b/Engine/Source/Runtime/Experimental/Chaos/Public/Chaos/PBDLongRangeConstraints.h
--- a/Engine/Source/Runtime/Experimental/Chaos/Public/Chaos/PBDLongRangeConstraints.h
+++ b/Engine/Source/Runtime/Experimental/Chaos/Public/Chaos/PBDLongRangeConstraints.h
@@ -21,11 +21,12 @@ public:
 		const int32 InParticleCount,
 		const TMap<int32, TSet<int32>>& PointToNeighbors,
 		const TConstArrayView<FRealSingle>& StiffnessMultipliers,
+		const TConstArrayView<FRealSingle>& ScaleMultipliers,
 		const int32 MaxNumTetherIslands = 4,
 		const FVec2& InStiffness = FVec2((FReal)1., (FReal)1.),
-		const FReal LimitScale = (FReal)1.,
+		const FVec2& Scale = FVec2((FReal)1.),
 		const EMode InMode = EMode::Geodesic)
-		: FPBDLongRangeConstraintsBase(Particles, InParticleOffset, InParticleCount, PointToNeighbors, StiffnessMultipliers, MaxNumTetherIslands, InStiffness, LimitScale, InMode) {}
+		: FPBDLongRangeConstraintsBase(Particles, InParticleOffset, InParticleCount, PointToNeighbors, StiffnessMultipliers, ScaleMultipliers, MaxNumTetherIslands, InStiffness, Scale, InMode) {}
 	virtual ~FPBDLongRangeConstraints() {}
 
 	void Apply(FPBDParticles& Particles, const FReal Dt, const TArray<int32>& ConstraintIndices) const;
diff --git a/Engine/Source/Runtime/Experimental/Chaos/Public/Chaos/PBDLongRangeConstraintsBase.h b/Engine/Source/Runtime/Experimental/Chaos/Public/Chaos/PBDLongRangeConstraintsBase.h
--- a/Engine/Source/Runtime/Experimental/Chaos/Public/Chaos/PBDLongRangeConstraintsBase.h
+++ b/Engine/Source/Runtime/Experimental/Chaos/Public/Chaos/PBDLongRangeConstraintsBase.h
@@ -58,9 +58,10 @@ public:
 		const int32 InParticleCount,
 		const TMap<int32, TSet<int32>>& PointToNeighbors,
 		const TConstArrayView<FRealSingle>& StiffnessMultipliers,
+		const TConstArrayView<FRealSingle>& ScaleMultipliers,
 		const int32 MaxNumTetherIslands = 4,
-		const FVec2& InStiffness = FVec2((FReal)0., (FReal)1.),
-		const FReal LimitScale = (FReal)1,
+		const FVec2& InStiffness = FVec2((FReal)1.),
+		const FVec2& Scale = FVec2((FReal)1.),
 		const EMode InMode = EMode::AccurateTetherFastLength);
 
 	virtual ~FPBDLongRangeConstraintsBase() {}
diff --git a/Engine/Source/Runtime/Experimental/Chaos/Public/Chaos/VelocityField.h b/Engine/Source/Runtime/Experimental/Chaos/Public/Chaos/VelocityField.h
--- a/Engine/Source/Runtime/Experimental/Chaos/Public/Chaos/VelocityField.h
+++ b/Engine/Source/Runtime/Experimental/Chaos/Public/Chaos/VelocityField.h
@@ -8,7 +8,7 @@ namespace Chaos
 {
 	// Velocity field basic implementation
 	// TODO: Add lift
-	class CHAOS_API FVelocityField
+	class CHAOS_API FVelocityField final
 	{
 	public:
 		static constexpr FReal DefaultDragCoefficient = (FReal)0.5;
@@ -17,12 +17,14 @@ namespace Chaos
 
 		// Construct an uninitialized field. Mesh, properties, and velocity will have to be set for this field to be valid.
 		FVelocityField()
-			: Range(-1)
+			: Offset(INDEX_NONE)
+			, NumParticles(0)
 		{
-			SetCoefficients((FReal)0., (FReal)0.);
+			SetProperties(FVec2::ZeroVector, FVec2::ZeroVector, (FReal)0.);
 		}
 
 		// Construct a uniform field.
+		UE_DEPRECATED(4.27, "Use Chaos fields instead.")
 		FVelocityField(
 			const FTriangleMesh& TriangleMesh,
 			const FVec3& InVelocity,
@@ -32,38 +34,38 @@ namespace Chaos
 			: PointToTriangleMap(TriangleMesh.GetPointToTriangleMap())
 			, Elements(TriangleMesh.GetElements())
 			, Velocity(InVelocity)
-			, Range(TriangleMesh.GetVertexRange())
+			, Offset(TriangleMesh.GetVertexRange()[0])
+			, NumParticles(TriangleMesh.GetVertexRange()[1] - TriangleMesh.GetVertexRange()[0] + 1)
 		{
 			Forces.SetNumUninitialized(Elements.Num());
-			SetCoefficients(InDragCoefficient, InLiftCoefficient);
-			SetFluidDensity(InFluidDensity);
+			SetProperties(InDragCoefficient, InLiftCoefficient, InFluidDensity);
 		}
 
 		// Construct a vector field.
+		UE_DEPRECATED(4.27, "Use Chaos fields instead.")
 		FVelocityField(
 			const FTriangleMesh& TriangleMesh,
-			TFunction<FVec3(const FVec3&)> InGetVelocity,
+			TFunction<FVec3(const FVec3&)> GetVelocity,
 			const FReal InDragCoefficient = DefaultDragCoefficient,
 			const FReal InLiftCoefficient = DefaultLiftCoefficient,
 			const FReal InFluidDensity = DefaultFluidDensity)
 			: PointToTriangleMap(TriangleMesh.GetPointToTriangleMap())
 			, Elements(TriangleMesh.GetElements())
-			, Velocity((FReal)0.)
-			, GetVelocity(InGetVelocity)
-			, Range(TriangleMesh.GetVertexRange())
+			, Velocity(GetVelocity(FVec3::ZeroVector))
+			, Offset(TriangleMesh.GetVertexRange()[0])
+			, NumParticles(TriangleMesh.GetVertexRange()[1] - TriangleMesh.GetVertexRange()[0] + 1)
 		{
 			Forces.SetNumUninitialized(Elements.Num());
-			SetCoefficients(InDragCoefficient, InLiftCoefficient);
-			SetFluidDensity(InFluidDensity);
+			SetProperties(InDragCoefficient, InLiftCoefficient, InFluidDensity);
 		}
 
-		virtual ~FVelocityField() {}
+		~FVelocityField() {}
 
 		void UpdateForces(const FPBDParticles& InParticles, const FReal /*Dt*/);
 
 		inline void Apply(FPBDParticles& InParticles, const FReal Dt, const int32 Index) const
 		{
-			checkSlow(Index >= Range[0] && Index <= Range[1]);  // The index should always match the original triangle mesh range
+			checkSlow(Index >= Offset && Index < Offset + NumParticles);  // The index should always match the original triangle mesh range
 
 			const TArray<int32>& ElementIndices = PointToTriangleMap[Index];
 			for (const int32 ElementIndex : ElementIndices)
@@ -72,56 +74,48 @@ namespace Chaos
 			}
 		}
 
+		UE_DEPRECATED(4.27, "Use SetProperties instead.")
 		void SetFluidDensity(const FReal InFluidDensity)
 		{
 			QuarterRho = (FReal)0.25 * InFluidDensity;
 		}
 
+		UE_DEPRECATED(4.27, "Use SetProperties instead.")
 		void SetCoefficients(const FReal InDragCoefficient, const FReal InLiftCoefficient)
 		{
-			Cd = InDragCoefficient;
-			Cl = InLiftCoefficient;
+			SetProperties(FVec2(InDragCoefficient), FVec2(InLiftCoefficient), (FReal)4. * QuarterRho);
 		}
 
-		bool IsActive() const
+		void SetProperties(const FVec2& Drag, const FVec2& Lift, const FReal FluidDensity)
 		{
-			return Cd > (FReal)0. || Cl > (FReal)0.;
+			constexpr FReal OneQuarter = (FReal)0.25;
+			QuarterRho = FluidDensity * OneQuarter;
+
+			constexpr FReal MinCoefficient = (FReal)0.;
+			constexpr FReal MaxCoefficient = (FReal)10.;
+			DragBase = FMath::Clamp(Drag[0], MinCoefficient, MaxCoefficient);
+			DragRange = FMath::Clamp(Drag[1], MinCoefficient, MaxCoefficient) - DragBase;
+			LiftBase = FMath::Clamp(Lift[0], MinCoefficient, MaxCoefficient);
+			LiftRange = FMath::Clamp(Lift[1], MinCoefficient, MaxCoefficient) - LiftBase;
 		}
 
-		void SetGeometry(const FTriangleMesh* TriangleMesh)
+		bool IsActive() const 
 		{
-			if (TriangleMesh)
-			{
-				PointToTriangleMap = TriangleMesh->GetPointToTriangleMap();
-				Elements = TriangleMesh->GetElements();
-				Range = TriangleMesh->GetVertexRange();
-				Forces.SetNumUninitialized(Elements.Num());
-			}
-			else
-			{
-				PointToTriangleMap = TArrayView<TArray<int32>>();
-				Elements = TArrayView<TVector<int32, 3>>();
-				Range = TVector<int32, 2>(-1);
-				Forces.Reset();
-			}
+			return (DragBase > (FReal)0. || DragRange != (FReal)0.) || (LiftBase > (FReal)0. || LiftRange != (FReal)0.);  // Note: range can be a negative value (although not when base is zero)
 		}
 
-		void SetVelocity(const FVec3& InVelocity)
-		{
-			Velocity = InVelocity;
-			GetVelocity = TFunction<FVec3(const FVec3&)>();
-		}
+		void SetGeometry(const FTriangleMesh* TriangleMesh, const TConstArrayView<FRealSingle>& DragMultipliers, const TConstArrayView<FRealSingle>& LiftMultipliers);
 
-		void SetVelocity(TFunction<FVec3(const FVec3&)> InGetVelocity)
-		{
-			GetVelocity = InGetVelocity;
-		}
+		void SetVelocity(const FVec3& InVelocity) { Velocity = InVelocity; }
+
+		UE_DEPRECATED(4.27, "Use SetVeloccity(const FVec3&) instead.")
+		void SetVelocity(TFunction<FVec3(const FVec3&)> GetVelocity)  { Velocity = GetVelocity(FVec3::ZeroVector); }
 
 		const TConstArrayView<TVector<int32, 3>>& GetElements() const { return Elements; }
 		TConstArrayView<FVec3> GetForces() const { return TConstArrayView<FVec3>(Forces); }
 
 	private:
-		inline void UpdateField(const FPBDParticles& InParticles, int32 ElementIndex, const FVec3& InVelocity)
+		void UpdateField(const FPBDParticles& InParticles, int32 ElementIndex, const FVec3& InVelocity, const FReal Cd, const FReal Cl)
 		{
 			const TVec3<int32>& Element = Elements[ElementIndex];
 
@@ -150,13 +144,16 @@ namespace Chaos
 	private:
 		TConstArrayView<TArray<int32>> PointToTriangleMap;
 		TConstArrayView<TVec3<int32>> Elements;
-		FVec3 Velocity;
-		TFunction<FVec3(const FVec3&)> GetVelocity;
 		TArray<FVec3> Forces;
+		TArray<FVec2> Multipliers;
+		FVec3 Velocity;
+		FReal DragBase;
+		FReal DragRange;
+		FReal LiftBase;
+		FReal LiftRange;
 		FReal QuarterRho;
-		FReal Cd;
-		FReal Cl;
-		TVec2<int32> Range;  // TODO: Remove? It is used by the check only
+		int32 Offset;
+		int32 NumParticles;
 	};
 }
 
diff --git a/Engine/Source/Runtime/Experimental/Chaos/Public/Chaos/XPBDLongRangeConstraints.h b/Engine/Source/Runtime/Experimental/Chaos/Public/Chaos/XPBDLongRangeConstraints.h
--- a/Engine/Source/Runtime/Experimental/Chaos/Public/Chaos/XPBDLongRangeConstraints.h
+++ b/Engine/Source/Runtime/Experimental/Chaos/Public/Chaos/XPBDLongRangeConstraints.h
@@ -25,11 +25,12 @@ public:
 		const int32 InParticleCount,
 		const TMap<int32, TSet<int32>>& PointToNeighbors,
 		const TConstArrayView<FRealSingle>& StiffnessMultipliers,
+		const TConstArrayView<FRealSingle>& ScaleMultipliers,
 		const int32 MaxNumTetherIslands = 4,
-		const FVec2& InStiffness = FVec2((FReal)1., (FReal)1.),
-		const FReal LimitScale = (FReal)1.,
+		const FVec2& InStiffness = FVec2((FReal)1.),
+		const FVec2& Scale = FVec2((FReal)1.),
 		const EMode InMode = EMode::Geodesic)
-	    : FPBDLongRangeConstraintsBase(Particles, InParticleOffset, InParticleCount, PointToNeighbors, StiffnessMultipliers, MaxNumTetherIslands, InStiffness, LimitScale, InMode)
+	    : FPBDLongRangeConstraintsBase(Particles, InParticleOffset, InParticleCount, PointToNeighbors, StiffnessMultipliers, ScaleMultipliers, MaxNumTetherIslands, InStiffness, Scale, InMode)
 	{
 		Lambdas.Reserve(Tethers.Num());
 	}
