Chaos - fix debug draw switching off (caused by multiple cvars referencing same underlying bool)

diff --git a/Engine/Source/Runtime/Experimental/ChaosSolverEngine/Private/Chaos/ChaosDebugDrawComponent.cpp b/Engine/Source/Runtime/Experimental/ChaosSolverEngine/Private/Chaos/ChaosDebugDrawComponent.cpp
--- a/Engine/Source/Runtime/Experimental/ChaosSolverEngine/Private/Chaos/ChaosDebugDrawComponent.cpp
+++ b/Engine/Source/Runtime/Experimental/ChaosSolverEngine/Private/Chaos/ChaosDebugDrawComponent.cpp
@@ -18,9 +18,6 @@ void ChaosDebugDraw_Enabled_Changed(IConsoleVariable* CVar)
 bool bChaosDebugDraw_Enabled = false;
 FAutoConsoleVariableRef CVarChaos_DebugDraw_Enabled(TEXT("p.Chaos.DebugDraw.Enabled"), bChaosDebugDraw_Enabled, TEXT("Whether to debug draw low level physics solver information"), FConsoleVariableDelegate::CreateStatic(ChaosDebugDraw_Enabled_Changed));
 
-// Deprecated, but widely used...
-FAutoConsoleVariableRef CVarChaos_DebugDraw_Enabled_Deprecated(TEXT("p.Chaos.DebugDrawing"), bChaosDebugDraw_Enabled, TEXT("Deprecated. Please use p.Chaos.DebugDraw.Enabled"), FConsoleVariableDelegate::CreateStatic(ChaosDebugDraw_Enabled_Changed));
-
 int ChaosDebugDraw_MaxElements = 20000;
 FAutoConsoleVariableRef CVarChaos_DebugDraw_MaxElements(TEXT("p.Chaos.DebugDraw.MaxLines"), ChaosDebugDraw_MaxElements, TEXT("Set the maximum number of debug draw lines that can be rendered (to limit perf drops)"));
 
