Chaos
- add margin the Geometry Collection Boxes
- fix feature search in OneShotManifold when margin is zero

diff --git a/Engine/Source/Runtime/Experimental/Chaos/Private/Chaos/CollisionOneShotManifolds.cpp b/Engine/Source/Runtime/Experimental/Chaos/Private/Chaos/CollisionOneShotManifolds.cpp
--- a/Engine/Source/Runtime/Experimental/Chaos/Private/Chaos/CollisionOneShotManifolds.cpp
+++ b/Engine/Source/Runtime/Experimental/Chaos/Private/Chaos/CollisionOneShotManifolds.cpp
@@ -33,6 +33,9 @@ namespace Chaos
 	FRealSingle Chaos_Collision_Manifold_MaxMargin = -1.0f;
 	FAutoConsoleVariableRef CVarChaosCollisioConvexManifoldMaxMargin(TEXT("p.Chaos.Collision.Manifold.MaxMargin"), Chaos_Collision_Manifold_MaxMargin, TEXT(""));
 
+	FRealSingle Chaos_Collision_Manifold_MinFaceSearchDistance = 1.0f;
+	FAutoConsoleVariableRef CVarChaosCollisioConvexManifoldMinFaceSearchDistance(TEXT("p.Chaos.Collision.Manifold.MinFaceSearchDistance"), Chaos_Collision_Manifold_MinFaceSearchDistance, TEXT(""));
+
 	namespace Collisions
 	{
 		// Forward delarations we need from CollisionRestitution.cpp
@@ -570,8 +573,13 @@ namespace Chaos
 			const ConvexImplicitType& Convex,
 			const FVec3 X,
 			const FVec3 N,
-			const FReal MaxDistance)
+			const FReal InMaxDistance)
 		{
+			// Handle InMaxDistance = 0. We expect that the X is actually on the surface in this case, so the search distance just needs to be some reasonable tolerance.
+			// @todo(chaos): this should probable be dependent on the size of the objects...
+			const FReal MinFaceSearchDistance = Chaos_Collision_Manifold_MinFaceSearchDistance;
+			const FReal MaxDistance = FMath::Max(InMaxDistance, MinFaceSearchDistance);
+
 			int32 BestPlaneIndex = INDEX_NONE;
 			FReal BestPlaneDot = 1.0f;
 
diff --git a/Engine/Source/Runtime/Experimental/Chaos/Private/GeometryCollection/GeometryCollectionCollisionStructureManager.cpp b/Engine/Source/Runtime/Experimental/Chaos/Private/GeometryCollection/GeometryCollectionCollisionStructureManager.cpp
--- a/Engine/Source/Runtime/Experimental/Chaos/Private/GeometryCollection/GeometryCollectionCollisionStructureManager.cpp
+++ b/Engine/Source/Runtime/Experimental/Chaos/Private/GeometryCollection/GeometryCollectionCollisionStructureManager.cpp
@@ -221,7 +221,14 @@ FCollisionStructureManager::NewImplicitBox(
 {
 	const FVector HalfExtents = CollisionBounds.GetExtent() * (1 - CollisionObjectReduction / 100.f);
 	const FVector Center = CollisionBounds.GetCenter();
-	Chaos::FImplicitObject* Implicit = new Chaos::TBox<Chaos::FReal, 3>(Center - HalfExtents, Center + HalfExtents);
+
+	// Margin settings are in UPhysicsSettingsCore which we can't access here
+	// @todo(chaos): pass margin settings into the collision manager?
+	float CollisionMarginFraction = 0.1f;// FMath::Max(0.0f, UPhysicsSettingsCore::Get()->SolverOptions.CollisionMarginFraction);
+	float CollisionMarginMax = 10.0f;// FMath::Max(0.0f, UPhysicsSettingsCore::Get()->SolverOptions.CollisionMarginMax);
+	const float Margin = FMath::Min(CollisionMarginFraction * 0.5f * HalfExtents.GetMin(), CollisionMarginMax);
+
+	Chaos::FImplicitObject* Implicit = new Chaos::TBox<Chaos::FReal, 3>(Center - HalfExtents, Center + HalfExtents, Margin);
 	UpdateImplicitFlags(Implicit, CollisionType);
 	return Implicit;
 }
