Chaos: Fix re-enabling of constraints when constrained particles are re-enabled.

diff --git a/Engine/Source/Runtime/Experimental/Chaos/Private/Chaos/PBDJointConstraints.cpp b/Engine/Source/Runtime/Experimental/Chaos/Private/Chaos/PBDJointConstraints.cpp
--- a/Engine/Source/Runtime/Experimental/Chaos/Private/Chaos/PBDJointConstraints.cpp
+++ b/Engine/Source/Runtime/Experimental/Chaos/Private/Chaos/PBDJointConstraints.cpp
@@ -526,7 +526,23 @@ namespace Chaos
 
 	void FPBDJointConstraints::SetConstraintEnabled(int32 ConstraintIndex, bool bEnabled)
 	{
-		ConstraintStates[ConstraintIndex].bDisabled = !bEnabled;
+		const FGenericParticleHandle Particle0 = FGenericParticleHandle(ConstraintParticles[ConstraintIndex][0]);
+		const FGenericParticleHandle Particle1 = FGenericParticleHandle(ConstraintParticles[ConstraintIndex][1]);
+
+		if (bEnabled)
+		{ 
+			// only enable constraint if the particles are valid and not disabled
+			if (Particle0->Handle() != nullptr && !Particle0->Disabled()
+				&& Particle1->Handle() != nullptr && !Particle1->Disabled())
+			{
+				ConstraintStates[ConstraintIndex].bDisabled = false;
+			}
+		}
+		else
+		{ 
+			// desirable to allow disabling no matter what state the endpoints
+			ConstraintStates[ConstraintIndex].bDisabled = true;
+		}
 	}
 
 
diff --git a/Engine/Source/Runtime/Experimental/Chaos/Public/Chaos/PBDConstraintRule.h b/Engine/Source/Runtime/Experimental/Chaos/Public/Chaos/PBDConstraintRule.h
--- a/Engine/Source/Runtime/Experimental/Chaos/Public/Chaos/PBDConstraintRule.h
+++ b/Engine/Source/Runtime/Experimental/Chaos/Public/Chaos/PBDConstraintRule.h
@@ -152,14 +152,17 @@ namespace Chaos
 
 		virtual void SetUseContactGraph(const bool InUseContactGraph) {}
 
-		/** Disable all constraints associated with the specified particles */
+		/** Change enabled state on all constraints associated with the specified particles */
 		virtual void SetConstraintsEnabled(const TSet<TGeometryParticleHandle<FReal, 3>*>& DisabledParticles, bool bInEnabled) 
 		{ 
 			for (TGeometryParticleHandle<FReal, 3>*DisabledParticle : DisabledParticles)
 			{
 				for( FConstraintHandle* Constraint : DisabledParticle->ParticleConstraints() )
 				{
-					Constraint->SetEnabled(bInEnabled);
+					if (Constraint->IsEnabled() != bInEnabled)
+					{
+						Constraint->SetEnabled(bInEnabled);
+					}
 				}
 			}
 		}
diff --git a/Engine/Source/Runtime/Experimental/Chaos/Public/Chaos/PBDRigidsEvolution.h b/Engine/Source/Runtime/Experimental/Chaos/Public/Chaos/PBDRigidsEvolution.h
--- a/Engine/Source/Runtime/Experimental/Chaos/Public/Chaos/PBDRigidsEvolution.h
+++ b/Engine/Source/Runtime/Experimental/Chaos/Public/Chaos/PBDRigidsEvolution.h
@@ -317,6 +317,7 @@ public:
 	CHAOS_API void EnableParticle(FGeometryParticleHandle* Particle, const FGeometryParticleHandle* ParentParticle)
 	{
 		Particles.EnableParticle(Particle);
+		EnableConstraints({ Particle });
 		ConstraintGraph.EnableParticle(Particle, ParentParticle);
 		// enable constraint after disabling the particle in the graph because of dependencies 
 		AddConstraintsToConstraintGraph(Particle->ParticleConstraints());
@@ -511,6 +512,15 @@ public:
 		}
 	}
 
+	/** Enable constraints from the enabled particles; constraints will only become enabled if their particle end points are valid */
+	CHAOS_API void EnableConstraints(const TSet<FGeometryParticleHandle*>& EnabledParticles)
+	{
+		for (FPBDConstraintGraphRule* ConstraintRule : ConstraintRules)
+		{
+			ConstraintRule->SetConstraintsEnabled(EnabledParticles, true);
+		}
+	}
+
 	CHAOS_API void ResetConstraints()
 	{
 		for(FPBDConstraintGraphRule* ConstraintRule : ConstraintRules)
diff --git a/Engine/Source/Runtime/Experimental/Chaos/Public/Chaos/PBDSuspensionConstraints.h b/Engine/Source/Runtime/Experimental/Chaos/Public/Chaos/PBDSuspensionConstraints.h
--- a/Engine/Source/Runtime/Experimental/Chaos/Public/Chaos/PBDSuspensionConstraints.h
+++ b/Engine/Source/Runtime/Experimental/Chaos/Public/Chaos/PBDSuspensionConstraints.h
@@ -126,7 +126,22 @@ namespace Chaos
 
 		void SetConstraintEnabled(int32 ConstraintIndex, bool bEnabled)
 		{
-			ConstraintEnabledStates[ConstraintIndex] = bEnabled;
+			const FGenericParticleHandle Particle = FGenericParticleHandle(ConstrainedParticles[ConstraintIndex]);
+
+			if (bEnabled)
+			{
+				// only enable constraint if the particle is valid and not disabled
+				if (Particle->Handle() != nullptr && !Particle->Disabled())
+				{
+					ConstraintEnabledStates[ConstraintIndex] = true;
+				}
+			}
+			else
+			{
+				// desirable to allow disabling no matter what state the endpoint
+				ConstraintEnabledStates[ConstraintIndex] = false;
+			}
+
 		}
 
 		//
