Fix LevelSet not loading in LWC mode
- kept levelset using double types, and adapted serialization system to handle it for now
- fixed serialization for chaos vectors using integer based types

diff --git a/Engine/Source/Runtime/Experimental/ChaosCore/Public/Chaos/ArrayND.h b/Engine/Source/Runtime/Experimental/ChaosCore/Public/Chaos/ArrayND.h
--- a/Engine/Source/Runtime/Experimental/ChaosCore/Public/Chaos/ArrayND.h
+++ b/Engine/Source/Runtime/Experimental/ChaosCore/Public/Chaos/ArrayND.h
@@ -22,9 +22,67 @@ inline void TryBulkSerializeArrayNDBase(FArchive& Ar, TArray<float>& Array)
 	Array.BulkSerialize(Ar);
 }
 
-inline void TryBulkSerializeArrayNDBase(FArchive& Ar, TArray<FVec3>& Array)
+inline void TryBulkSerializeArrayNDBase(FArchive& Ar, TArray<TVec3<FRealSingle>>& Array)
 {
-	Array.BulkSerialize(Ar);
+	Array.BulkSerialize(Ar); 
+}
+
+inline float ConvertDoubleToFloat(double DoubleValue)
+{
+	return (float)DoubleValue; // LWC_TODO : Perf pessimization 
+}
+
+inline TVec3<float> ConvertDoubleToFloat(TVec3<double> DoubleValue)
+{
+	return TVec3<float>((float)DoubleValue.X, (float)DoubleValue.Y, (float)DoubleValue.Z); // LWC_TODO : Perf pessimization 
+}
+
+inline double ConvertFloatToDouble(float FloatValue)
+{
+	return (double)FloatValue;
+}
+
+inline TVec3<double> ConvertFloatToDouble(TVec3<float> FloatValue)
+{
+	return TVec3<double>((double)FloatValue.X, (double)FloatValue.Y, (double)FloatValue.Z); 
+}
+
+// LWC_TODO : Perf pessimization : this is sub-optimal but will do until we sort the serialization out
+template<typename DOUBLE_T, typename FLOAT_T>
+inline void TryBulkSerializeArrayNDBaseForDoubles(FArchive& Ar, TArray<DOUBLE_T>& DoubleTypedArray)
+{
+	TArray<FLOAT_T> FloatTypedArray;
+	if (Ar.IsSaving())
+	{
+		FloatTypedArray.SetNumUninitialized(DoubleTypedArray.Num());
+		for (int i = 0; i < DoubleTypedArray.Num(); ++i)
+		{
+			FloatTypedArray[i] = ConvertDoubleToFloat(DoubleTypedArray[i]);
+		}
+	}
+
+	TryBulkSerializeArrayNDBase(Ar, FloatTypedArray);
+
+	if (Ar.IsLoading())
+	{
+		DoubleTypedArray.SetNumUninitialized(FloatTypedArray.Num());
+		for (int i = 0; i < FloatTypedArray.Num(); ++i)
+		{
+			DoubleTypedArray[i] = ConvertFloatToDouble(FloatTypedArray[i]);
+		}
+	}
+}
+
+// LWC_TODO : Perf pessimization : this is sub-optimal but will do until we sort the serialization out
+inline void TryBulkSerializeArrayNDBase(FArchive& Ar, TArray<double>& Array)
+{
+	TryBulkSerializeArrayNDBaseForDoubles<double, float>(Ar, Array);
+}
+
+// LWC_TODO : Perf pessimization : this is sub-optimal but will do until we sort the serialization out
+inline void TryBulkSerializeArrayNDBase(FArchive& Ar, TArray<TVec3<FRealDouble>>& Array)
+{
+	TryBulkSerializeArrayNDBaseForDoubles<TVec3<FRealDouble>, TVec3<FRealSingle>>(Ar, Array);
 }
 
 template<class T_DERIVED, class T, int d>
diff --git a/Engine/Source/Runtime/Experimental/ChaosCore/Public/Chaos/Core.h b/Engine/Source/Runtime/Experimental/ChaosCore/Public/Chaos/Core.h
--- a/Engine/Source/Runtime/Experimental/ChaosCore/Public/Chaos/Core.h
+++ b/Engine/Source/Runtime/Experimental/ChaosCore/Public/Chaos/Core.h
@@ -35,5 +35,5 @@ namespace Chaos
 	// NOTE: if you get a merge conflict on the GUID, you must replace it with a new GUID - do not accept the source or target
 	// or you will likely get DDC version conflicts resulting in crashes during load.
 	// Core version string for Chaos data. Any DDC builder dependent on Chaos for serialization should depend on this version
-	static const FString ChaosVersionString = TEXT("6315F0B9C62B42B1A7929463A7375317");
+	static const FString ChaosVersionString = TEXT("2010BDD8-2309-49A0-919D-6BE431CFEE3D");
 }
diff --git a/Engine/Source/Runtime/Experimental/ChaosCore/Public/Chaos/Vector.h b/Engine/Source/Runtime/Experimental/ChaosCore/Public/Chaos/Vector.h
--- a/Engine/Source/Runtime/Experimental/ChaosCore/Public/Chaos/Vector.h
+++ b/Engine/Source/Runtime/Experimental/ChaosCore/Public/Chaos/Vector.h
@@ -963,30 +963,40 @@ namespace Chaos
 		return Seed;
 	}
 
+	// this is used as 
+	template<typename T, int d>
+	inline FArchive& SerializeReal(FArchive& Ar, TVector<T, d>& ValueIn)
+	{
+		static_assert(TIsSame<T, float>::Value || TIsSame<T, double>::Value, "only float or double are supoprted by this function");
+		for (int32 Idx = 0; Idx < d; ++Idx)
+		{
+			FRealSingle RealSingle = (FRealSingle)ValueIn[Idx];
+			Ar << RealSingle;
+			ValueIn[Idx] = (T)RealSingle;
+		}
+		return Ar;
+	}
+
+	template<int d>
+	FArchive& operator<<(FArchive& Ar, TVector<FRealSingle, d>& ValueIn) 
+	{
+		return SerializeReal(Ar, ValueIn);
+	}
+
+	template<int d>
+	FArchive& operator<<(FArchive& Ar, TVector<FRealDouble, d>& ValueIn)
+	{
+		return SerializeReal(Ar, ValueIn);
+	}
+
+	// general for for all other vectors
 	template<typename T, int d>
 	FArchive& operator<<(FArchive& Ar, TVector<T, d>& ValueIn)
 	{
-		// LWC_TODO: Serializer
-		//constexpr bool bIsTypeReal = (TAreTypesEqual<FReal, T>::Value == true);
-		//if(!Ar.IsPersistent())
-		//{
-		//	// unchanged type code path 
-		//	for (int32 Idx = 0; Idx < d; ++Idx)
-		//	{
-		//		Ar << ValueIn[Idx];
-		//	}
-
-		//}
-		//else
-		{
-			// in that case data is stored as float and we need to read it as such		
-			//ensure(Ar.IsLoading()); // this case should normally only happening when reading 
-			for (int32 Idx = 0; Idx < d; ++Idx)
-			{
-				FRealSingle RealSingle = (FRealSingle)ValueIn[Idx];
-				Ar << RealSingle;
-				ValueIn[Idx] = (T)RealSingle;
-			}
+		// unchanged type code path 
+		for (int32 Idx = 0; Idx < d; ++Idx)
+		{
+			Ar << ValueIn[Idx];
 		}
 		return Ar;
 	}
