Fix rendering of GeometryCollection with LWC on
- change matrix types for dynamic data to be explicitly using float
- wrapped update of vertex buffer to ensure type check
- When LWC mode is on, convert double based matrices to they float equivalent ( sub-optimal, but necessary for now until we have a better way )

diff --git a/Engine/Source/Runtime/Experimental/GeometryCollectionEngine/Private/GeometryCollection/GeometryCollectionComponent.cpp b/Engine/Source/Runtime/Experimental/GeometryCollectionEngine/Private/GeometryCollection/GeometryCollectionComponent.cpp
--- a/Engine/Source/Runtime/Experimental/GeometryCollectionEngine/Private/GeometryCollection/GeometryCollectionComponent.cpp
+++ b/Engine/Source/Runtime/Experimental/GeometryCollectionEngine/Private/GeometryCollection/GeometryCollectionComponent.cpp
@@ -1311,7 +1311,7 @@ void UGeometryCollectionComponent::InitConstantData(FGeometryCollectionConstantD
 	TArray<FMatrix> RestMatrices;
 	GeometryCollectionAlgo::GlobalMatrices(RestCollection->GetGeometryCollection()->Transform, RestCollection->GetGeometryCollection()->Parent, RestMatrices);
 
-	ConstantData->RestTransforms = MoveTemp(RestMatrices);
+	ConstantData->SetRestTransforms(RestMatrices);
 }
 
 FGeometryCollectionDynamicData* UGeometryCollectionComponent::InitDynamicData(bool bInitialization)
@@ -1346,9 +1346,7 @@ FGeometryCollectionDynamicData* UGeometryCollectionComponent::InitDynamicData(bo
 			GeometryCollectionAlgo::GlobalMatrices(GetTransformArray(), GetParentArray(), GlobalMatrices);
 
 			DynamicData = GDynamicDataPool.Allocate();
-			DynamicData->PrevTransforms = GlobalMatrices;
-			DynamicData->Transforms = GlobalMatrices;
-			DynamicData->ChangedCount = GlobalMatrices.Num();
+			DynamicData->SetAllTransforms(GlobalMatrices);
 			DynamicData->IsDynamic = bIsCacheDynamic;
 			DynamicData->IsLoading = GetIsObjectLoading();
 
@@ -1389,7 +1387,7 @@ FGeometryCollectionDynamicData* UGeometryCollectionComponent::InitDynamicData(bo
 				float TimeIncrement = Step == NumSteps - 1 ? LastDt : CacheDt;
 				CurrentCacheTime += TimeIncrement;
 			
-				DynamicData->PrevTransforms = GlobalMatrices;
+				DynamicData->SetPrevTransforms(GlobalMatrices);
 
 				const FRecordedFrame* FirstFrame = nullptr;
 				const FRecordedFrame* SecondFrame = nullptr;
@@ -1420,9 +1418,7 @@ FGeometryCollectionDynamicData* UGeometryCollectionComponent::InitDynamicData(bo
 								*RestCollection->GetName(), *CacheParameters.TargetCache->GetName()
 							);
 
-							DynamicData->PrevTransforms = GlobalMatrices;
-							DynamicData->Transforms = GlobalMatrices;
-							DynamicData->ChangedCount = GlobalMatrices.Num();
+							DynamicData->SetAllTransforms(GlobalMatrices);
 							return DynamicData;
 						}
 
@@ -1462,9 +1458,7 @@ FGeometryCollectionDynamicData* UGeometryCollectionComponent::InitDynamicData(bo
 								TEXT("%s: TargetCache (%s) is out of sync with GeometryCollection.  Regenerate the cache."), 
 								*RestCollection->GetName(), *CacheParameters.TargetCache->GetName()
 							);
-							DynamicData->PrevTransforms = GlobalMatrices;
-							DynamicData->Transforms = GlobalMatrices;
-							DynamicData->ChangedCount = GlobalMatrices.Num();
+							DynamicData->SetAllTransforms(GlobalMatrices);
 							return DynamicData;
 						}
 
@@ -1495,7 +1489,7 @@ FGeometryCollectionDynamicData* UGeometryCollectionComponent::InitDynamicData(bo
 						}
 					}
 				}
-				DynamicData->Transforms = GlobalMatrices;
+				DynamicData->SetTransforms(GlobalMatrices);
 			}
 
 			// Check if transforms at start of this tick are the same as what is calculated from the cache
@@ -1517,14 +1511,12 @@ FGeometryCollectionDynamicData* UGeometryCollectionComponent::InitDynamicData(bo
 			// Copy global matrices over to DynamicData
 			CalculateGlobalMatrices();
 
-			DynamicData->PrevTransforms = GlobalMatrices;
-			DynamicData->Transforms = GlobalMatrices;
-			DynamicData->ChangedCount = GlobalMatrices.Num();
+			DynamicData->SetAllTransforms(GlobalMatrices);
 		}
 		else
 		{
 			// Copy existing global matrices into prev transforms
-			DynamicData->PrevTransforms = GlobalMatrices;
+			DynamicData->SetPrevTransforms(GlobalMatrices);
 
 			// Copy global matrices over to DynamicData
 			CalculateGlobalMatrices();
@@ -1534,12 +1526,12 @@ FGeometryCollectionDynamicData* UGeometryCollectionComponent::InitDynamicData(bo
 			// if the number of matrices has changed between frames, then sync previous to current
 			if (GlobalMatrices.Num() != DynamicData->PrevTransforms.Num())
 			{
-				DynamicData->PrevTransforms = GlobalMatrices;
+				DynamicData->SetPrevTransforms(GlobalMatrices);
 				DynamicData->ChangedCount = GlobalMatrices.Num();
 				bComputeChanges = false; // Optimization to just force all transforms as changed and skip comparison
 			}
 
-			DynamicData->Transforms = GlobalMatrices;
+			DynamicData->SetTransforms(GlobalMatrices);
 
 			// The number of transforms for current and previous should match now
 			check(DynamicData->PrevTransforms.Num() == DynamicData->Transforms.Num());
diff --git a/Engine/Source/Runtime/Experimental/GeometryCollectionEngine/Private/GeometryCollection/GeometryCollectionSceneProxy.cpp b/Engine/Source/Runtime/Experimental/GeometryCollectionEngine/Private/GeometryCollection/GeometryCollectionSceneProxy.cpp
--- a/Engine/Source/Runtime/Experimental/GeometryCollectionEngine/Private/GeometryCollection/GeometryCollectionSceneProxy.cpp
+++ b/Engine/Source/Runtime/Experimental/GeometryCollectionEngine/Private/GeometryCollection/GeometryCollectionSceneProxy.cpp
@@ -522,48 +522,33 @@ void FGeometryCollectionSceneProxy::SetDynamicData_RenderThread(FGeometryCollect
 				}
 			}
 
-			const EResourceLockMode LockMode = bLocalGeometryCollectionTripleBufferUploads ? RLM_WriteOnly_NoOverwrite : RLM_WriteOnly;
-
 			// Copy the transform data over to the vertex buffer	
-			if (DynamicData->IsDynamic)
 			{
-				CycleTransformBuffers(bLocalGeometryCollectionTripleBufferUploads);
-				FGeometryCollectionTransformBuffer& TransformBuffer = GetCurrentTransformBuffer();
-				FGeometryCollectionTransformBuffer& PrevTransformBuffer = GetCurrentPrevTransformBuffer();
-				VertexFactory.SetBoneTransformSRV(TransformBuffer.VertexBufferSRV);
-				VertexFactory.SetBonePrevTransformSRV(PrevTransformBuffer.VertexBufferSRV);
-
-				check(TransformBuffer.NumTransforms == DynamicData->Transforms.Num());
-				check(PrevTransformBuffer.NumTransforms == DynamicData->PrevTransforms.Num());
-
-				void* VertexBufferData = RHILockBuffer(TransformBuffer.VertexBufferRHI, 0, DynamicData->Transforms.Num() * sizeof(FMatrix), LockMode);
-				FMemory::Memcpy(VertexBufferData, DynamicData->Transforms.GetData(), DynamicData->Transforms.Num() * sizeof(FMatrix));
-				RHIUnlockBuffer(TransformBuffer.VertexBufferRHI);
+				const EResourceLockMode LockMode = bLocalGeometryCollectionTripleBufferUploads ? RLM_WriteOnly_NoOverwrite : RLM_WriteOnly;
 
-				void* PrevVertexBufferData = RHILockBuffer(PrevTransformBuffer.VertexBufferRHI, 0, DynamicData->PrevTransforms.Num() * sizeof(FMatrix), LockMode);
-				FMemory::Memcpy(PrevVertexBufferData, DynamicData->PrevTransforms.GetData(), DynamicData->PrevTransforms.Num() * sizeof(FMatrix));
-				RHIUnlockBuffer(PrevTransformBuffer.VertexBufferRHI);
-
-				TransformVertexBuffersContainsOriginalMesh = false;
-			}
-			else if (!TransformVertexBuffersContainsOriginalMesh)
-			{		
 				CycleTransformBuffers(bLocalGeometryCollectionTripleBufferUploads);
+
 				FGeometryCollectionTransformBuffer& TransformBuffer = GetCurrentTransformBuffer();
 				FGeometryCollectionTransformBuffer& PrevTransformBuffer = GetCurrentPrevTransformBuffer();
+
 				VertexFactory.SetBoneTransformSRV(TransformBuffer.VertexBufferSRV);
 				VertexFactory.SetBonePrevTransformSRV(PrevTransformBuffer.VertexBufferSRV);
 
-				// if we are rendering the base mesh geometry, then use rest transforms rather than the simulated one for both current and previous transforms
-				void* VertexBufferData = RHILockBuffer(TransformBuffer.VertexBufferRHI, 0, ConstantData->RestTransforms.Num() * sizeof(FMatrix), LockMode);
-				FMemory::Memcpy(VertexBufferData, ConstantData->RestTransforms.GetData(), ConstantData->RestTransforms.Num() * sizeof(FMatrix));
-				RHIUnlockBuffer(TransformBuffer.VertexBufferRHI);
+				if (DynamicData->IsDynamic)
+				{
+					TransformBuffer.UpdateDynamicData(DynamicData->Transforms, LockMode);
+					PrevTransformBuffer.UpdateDynamicData(DynamicData->PrevTransforms, LockMode);
 
-				void* PrevVertexBufferData = RHILockBuffer(PrevTransformBuffer.VertexBufferRHI, 0, ConstantData->RestTransforms.Num() * sizeof(FMatrix), LockMode);
-				FMemory::Memcpy(PrevVertexBufferData, ConstantData->RestTransforms.GetData(), ConstantData->RestTransforms.Num() * sizeof(FMatrix));
-				RHIUnlockBuffer(PrevTransformBuffer.VertexBufferRHI);
+					TransformVertexBuffersContainsOriginalMesh = false;
+				}
+				else if (!TransformVertexBuffersContainsOriginalMesh)
+				{
+					// if we are rendering the base mesh geometry, then use rest transforms rather than the simulated one for both current and previous transforms
+					TransformBuffer.UpdateDynamicData(ConstantData->RestTransforms, LockMode);
+					PrevTransformBuffer.UpdateDynamicData(ConstantData->RestTransforms, LockMode);
 
-				TransformVertexBuffersContainsOriginalMesh = true;
+					TransformVertexBuffersContainsOriginalMesh = true;
+				}
 			}
 		}
 		else
@@ -602,7 +587,7 @@ void FGeometryCollectionSceneProxy::SetDynamicData_RenderThread(FGeometryCollect
 
 				if (ThisBatchSize > 0)
 				{
-					const FMatrix* RESTRICT BoneTransformsPtr = DynamicData->IsDynamic ? DynamicData->Transforms.GetData() : ConstantData->RestTransforms.GetData();
+					const FMatrix44f* RESTRICT BoneTransformsPtr = DynamicData->IsDynamic ? DynamicData->Transforms.GetData() : ConstantData->RestTransforms.GetData();
 	#if INTEL_ISPC
 					uint8* VertexBufferOffset = (uint8*)VertexBufferData + (IndexOffset * VertexBuffer.GetStride());
 					ispc::SetDynamicData_RenderThread(
@@ -1529,3 +1514,12 @@ void FGeometryCollectionDynamicDataPool::Release(FGeometryCollectionDynamicData*
 		FreeList.Push(DynamicData);
 	}
 }
+
+void FGeometryCollectionTransformBuffer::UpdateDynamicData(const TArray<FMatrix44f>& Transforms, EResourceLockMode LockMode)
+{
+	check(NumTransforms == Transforms.Num());
+
+	void* VertexBufferData = RHILockBuffer(VertexBufferRHI, 0, Transforms.Num() * sizeof(FMatrix44f), LockMode);
+	FMemory::Memcpy(VertexBufferData, Transforms.GetData(), Transforms.Num() * sizeof(FMatrix44f));
+	RHIUnlockBuffer(VertexBufferRHI);
+}
diff --git a/Engine/Source/Runtime/Experimental/GeometryCollectionEngine/Private/GeometryCollection/GeometryCollectionSceneProxy.h b/Engine/Source/Runtime/Experimental/GeometryCollectionEngine/Private/GeometryCollection/GeometryCollectionSceneProxy.h
--- a/Engine/Source/Runtime/Experimental/GeometryCollectionEngine/Private/GeometryCollection/GeometryCollectionSceneProxy.h
+++ b/Engine/Source/Runtime/Experimental/GeometryCollectionEngine/Private/GeometryCollection/GeometryCollectionSceneProxy.h
@@ -73,11 +73,28 @@ public:
 		VertexBufferSRV = RHICreateShaderResourceView(VertexBufferRHI, 16, PF_A32B32G32R32F);
 	}
 
+	void UpdateDynamicData(const TArray<FMatrix44f>& Transforms, EResourceLockMode LockMode);
+
 	int32 NumTransforms;
 
 	FShaderResourceViewRHIRef VertexBufferSRV;
 };
 
+
+inline void CopyTransformsWithConversionWhenNeeded(TArray<FMatrix44f>& DstTransforms, const TArray<FMatrix>& SrcTransforms)
+{
+#if UE_LARGE_WORLD_COORDINATES_DISABLED
+	DstTransforms = SrcTransforms; // add move option?
+#else
+	// LWC_TODO : we have no choice but to convert each element at this point to avoid changing GeometryCollectionAlgo::GlobalMatrices that is used all over the place
+	DstTransforms.SetNumUninitialized(SrcTransforms.Num());
+	for (int TransformIndex = 0; TransformIndex < SrcTransforms.Num(); ++TransformIndex)
+	{
+		DstTransforms[TransformIndex] = FMatrix44f(SrcTransforms[TransformIndex]); // LWC_TODO: Perf pessimization
+	}
+#endif
+}
+
 /** Immutable rendering data (kind of) */
 struct FGeometryCollectionConstantData
 {
@@ -99,14 +116,20 @@ struct FGeometryCollectionConstantData
 	TArray<FIntVector> OriginalMeshIndices;
 	TArray<FGeometryCollectionSection> OriginalMeshSections;
 
-	TArray<FMatrix> RestTransforms;
+	TArray<FMatrix44f> RestTransforms;
+
+	void SetRestTransforms(const TArray<FMatrix>& InTransforms)
+	{
+		// use for LWC as FMatrix and FMatrix44f are different when LWC is on 
+		CopyTransformsWithConversionWhenNeeded(RestTransforms, InTransforms);
+	}
 };
 
 /** Mutable rendering data */
 struct FGeometryCollectionDynamicData
 {
-	TArray<FMatrix> Transforms;
-	TArray<FMatrix> PrevTransforms;
+	TArray<FMatrix44f> Transforms;
+	TArray<FMatrix44f> PrevTransforms;
 	uint32 ChangedCount;
 	uint8 IsDynamic : 1;
 	uint8 IsLoading : 1;
@@ -124,6 +147,25 @@ struct FGeometryCollectionDynamicData
 		IsLoading = false;
 	}
 
+	void SetTransforms(const TArray<FMatrix>& InTransforms)
+	{
+		// use for LWC as FMatrix and FMatrix44f are different when LWC is on 
+		CopyTransformsWithConversionWhenNeeded(Transforms, InTransforms);
+	}
+
+	void SetPrevTransforms(const TArray<FMatrix>& InTransforms)
+	{
+		// use for LWC as FMatrix and FMatrix44f are different when LWC is on 
+		CopyTransformsWithConversionWhenNeeded(PrevTransforms, InTransforms);
+	}
+
+	void SetAllTransforms(const TArray<FMatrix>& InTransforms)
+	{
+		SetTransforms(InTransforms);
+		PrevTransforms = Transforms;
+		ChangedCount = Transforms.Num();
+	}
+
 	void DetermineChanges()
 	{
 		// Check if previous transforms are the same as current
