Support locked linear/angular motion for plasticity
Allow linear plasticity to be bigger than 1

diff --git a/Engine/Source/Runtime/Experimental/Chaos/Private/Chaos/PBDJointConstraints.cpp b/Engine/Source/Runtime/Experimental/Chaos/Private/Chaos/PBDJointConstraints.cpp
--- a/Engine/Source/Runtime/Experimental/Chaos/Private/Chaos/PBDJointConstraints.cpp
+++ b/Engine/Source/Runtime/Experimental/Chaos/Private/Chaos/PBDJointConstraints.cpp
@@ -1487,42 +1487,6 @@ namespace Chaos
 			return;
 		}
 
-		// todo(chaos): provide constraint names for these warnings
-		bool bApplyLinearPlasticityLimits = bHasLinearPlasticityLimit;
-		if(bHasLinearPlasticityLimit)
-		{
-			const bool bLinearPositionDriveEnabled = JointSettings.bLinearPositionDriveEnabled[0] && JointSettings.bLinearPositionDriveEnabled[1] && JointSettings.bLinearPositionDriveEnabled[2];
-			if(!bLinearPositionDriveEnabled)
-			{
-				bApplyLinearPlasticityLimits = false;
-			}
-			const bool bLinearMotionLocked = JointSettings.LinearMotionTypes[0] == EJointMotionType::Locked || JointSettings.LinearMotionTypes[1] == EJointMotionType::Locked || JointSettings.LinearMotionTypes[2] == EJointMotionType::Locked;
-			if(bLinearMotionLocked)
-			{
-				bApplyLinearPlasticityLimits = false;
-			}
-		}
-
-		bool bApplyAngularPlasticityLimits = bHasAngularPlasticityLimit;
-		if(bHasAngularPlasticityLimit)
-		{
-			const bool bAngularPositionDriveEnabled = JointSettings.bAngularSLerpPositionDriveEnabled || (JointSettings.bAngularSwingPositionDriveEnabled && JointSettings.bAngularTwistPositionDriveEnabled);
-			if(!bAngularPositionDriveEnabled)
-			{
-				bApplyAngularPlasticityLimits = false;
-			}
-			const bool bAngularMotionLocked = JointSettings.AngularMotionTypes[(int32)EJointAngularConstraintIndex::Twist] == EJointMotionType::Locked || JointSettings.AngularMotionTypes[(int32)EJointAngularConstraintIndex::Swing1] == EJointMotionType::Locked || JointSettings.AngularMotionTypes[(int32)EJointAngularConstraintIndex::Swing2] == EJointMotionType::Locked;
-			if(bAngularMotionLocked)
-			{
-				bApplyAngularPlasticityLimits = false;
-			}
-		}
-
-		if(!bApplyLinearPlasticityLimits && !bApplyAngularPlasticityLimits)
-		{
-			return;
-		}
-
 		int32 Index0, Index1;
 		GetConstrainedParticleIndices(ConstraintIndex, Index0, Index1);
 		FGenericParticleHandle Particle0 = FGenericParticleHandle(ConstraintParticles[ConstraintIndex][Index0]);
@@ -1534,18 +1498,38 @@ namespace Chaos
 		Q1.EnforceShortestArcWith(ConstraintFramesGlobal[0].GetRotation());
 		ConstraintFramesGlobal[1].SetRotation(Q1);
 
-		if(bApplyLinearPlasticityLimits)
+		if(bHasLinearPlasticityLimit)
 		{
-			const FVec3 LinearDisplacement = ConstraintFramesGlobal[0].InverseTransformPositionNoScale(ConstraintFramesGlobal[1].GetTranslation());
-			if((LinearDisplacement - JointSettings.LinearDrivePositionTarget).Size() > JointSettings.LinearPlasticityLimit)
+			FVec3 LinearDisplacement = ConstraintFramesGlobal[0].InverseTransformPositionNoScale(ConstraintFramesGlobal[1].GetTranslation());
+
+			for(int32 Axis = 0; Axis < 3; Axis++)
+			{
+				if(!JointSettings.bLinearPositionDriveEnabled[Axis] || JointSettings.LinearMotionTypes[Axis] == EJointMotionType::Locked)
+				{
+					LinearDisplacement[Axis] = 0;
+				}
+			}
+			if((LinearDisplacement - JointSettings.LinearDrivePositionTarget).SizeSquared() > JointSettings.LinearPlasticityLimit * JointSettings.LinearPlasticityLimit)
 			{
 				JointSettings.LinearDrivePositionTarget = LinearDisplacement;
 			}
 		}
 
-		if(bApplyAngularPlasticityLimits)
+		if(bHasAngularPlasticityLimit)
 		{
-			const FRotation3 AngularDisplacement = ConstraintFramesGlobal[0].GetRotation().Inverse() * ConstraintFramesGlobal[1].GetRotation();
+			FRotation3 Swing, Twist; FPBDJointUtilities::DecomposeSwingTwistLocal(ConstraintFramesGlobal[0].GetRotation(), ConstraintFramesGlobal[1].GetRotation(), Swing, Twist);
+
+			if((!JointSettings.bAngularSLerpPositionDriveEnabled && !JointSettings.bAngularTwistPositionDriveEnabled) || JointSettings.AngularMotionTypes[(int32)EJointAngularConstraintIndex::Twist] == EJointMotionType::Locked)
+			{
+				Twist = FRotation3::Identity;
+			}
+			// @todo(chaos): clamp rotation if only swing1(swing2) is locked
+			if((!JointSettings.bAngularSLerpPositionDriveEnabled && !JointSettings.bAngularSwingPositionDriveEnabled) || (JointSettings.AngularMotionTypes[(int32)EJointAngularConstraintIndex::Swing1] == EJointMotionType::Locked && JointSettings.AngularMotionTypes[(int32)EJointAngularConstraintIndex::Swing2] == EJointMotionType::Locked))
+			{
+				Swing = FRotation3::Identity;
+			}
+			
+			const FRotation3 AngularDisplacement = Swing * Twist;
 			const FReal AngleDeg = JointSettings.AngularDrivePositionTarget.AngularDistance(AngularDisplacement);
 			if (AngleDeg > JointSettings.AngularPlasticityLimit)
 			{
diff --git a/Engine/Source/Runtime/Experimental/Chaos/Private/PhysicsProxy/JointConstraintProxy.cpp b/Engine/Source/Runtime/Experimental/Chaos/Private/PhysicsProxy/JointConstraintProxy.cpp
--- a/Engine/Source/Runtime/Experimental/Chaos/Private/PhysicsProxy/JointConstraintProxy.cpp
+++ b/Engine/Source/Runtime/Experimental/Chaos/Private/PhysicsProxy/JointConstraintProxy.cpp
@@ -299,7 +299,7 @@ void FJointConstraintPhysicsProxy::PushStateOnPhysicsThread(Chaos::FPBDRigidsSol
 			if (DirtyFlagsBuffer.IsDirty(Chaos::EJointConstraintFlags::LinearBreakForce))
 			{
 				ConstraintSettings.LinearBreakForce = JointSettingsBuffer.LinearBreakForce;
-				ConstraintSettings.LinearPlasticityLimit = FMath::Clamp((float)JointSettingsBuffer.LinearPlasticityLimit, 0.f, 1.f);
+				ConstraintSettings.LinearPlasticityLimit = FMath::Clamp((float)JointSettingsBuffer.LinearPlasticityLimit, 0.f, FLT_MAX);
 			}
 
 			if (DirtyFlagsBuffer.IsDirty(Chaos::EJointConstraintFlags::AngularBreakTorque))
