Add comments to clearly state the assumptions the todos for joint plasticity.

diff --git a/Engine/Source/Runtime/Experimental/Chaos/Private/Chaos/PBDJointConstraints.cpp b/Engine/Source/Runtime/Experimental/Chaos/Private/Chaos/PBDJointConstraints.cpp
--- a/Engine/Source/Runtime/Experimental/Chaos/Private/Chaos/PBDJointConstraints.cpp
+++ b/Engine/Source/Runtime/Experimental/Chaos/Private/Chaos/PBDJointConstraints.cpp
@@ -1502,6 +1502,7 @@ namespace Chaos
 		{
 			FVec3 LinearDisplacement = ConstraintFramesGlobal[0].InverseTransformPositionNoScale(ConstraintFramesGlobal[1].GetTranslation());
 
+			// @todo(chaos): still need to warn against the case where all position drives are not enabled or all dimensions are locked. Warning should print out the joint names and should only print out once to avoid spamming.
 			for(int32 Axis = 0; Axis < 3; Axis++)
 			{
 				if(!JointSettings.bLinearPositionDriveEnabled[Axis] || JointSettings.LinearMotionTypes[Axis] == EJointMotionType::Locked)
@@ -1509,6 +1510,7 @@ namespace Chaos
 					LinearDisplacement[Axis] = 0;
 				}
 			}
+			// Assuming that the dimensions which are locked or have no targets are 0. in LinearDrivePositionTarget
 			if((LinearDisplacement - JointSettings.LinearDrivePositionTarget).SizeSquared() > JointSettings.LinearPlasticityLimit * JointSettings.LinearPlasticityLimit)
 			{
 				JointSettings.LinearDrivePositionTarget = LinearDisplacement;
@@ -1519,6 +1521,7 @@ namespace Chaos
 		{
 			FRotation3 Swing, Twist; FPBDJointUtilities::DecomposeSwingTwistLocal(ConstraintFramesGlobal[0].GetRotation(), ConstraintFramesGlobal[1].GetRotation(), Swing, Twist);
 
+			// @todo(chaos): still need to warn against the case where all position drives are not enabled or all dimensions are locked. Warning should print out the joint names and should only print out once to avoid spamming.
 			if((!JointSettings.bAngularSLerpPositionDriveEnabled && !JointSettings.bAngularTwistPositionDriveEnabled) || JointSettings.AngularMotionTypes[(int32)EJointAngularConstraintIndex::Twist] == EJointMotionType::Locked)
 			{
 				Twist = FRotation3::Identity;
@@ -1530,6 +1533,7 @@ namespace Chaos
 			}
 			
 			const FRotation3 AngularDisplacement = Swing * Twist;
+			// Assuming that the dimensions which are locked or have no targets are 0. in AngularDrivePositionTarget
 			const FReal AngleDeg = JointSettings.AngularDrivePositionTarget.AngularDistance(AngularDisplacement);
 			if (AngleDeg > JointSettings.AngularPlasticityLimit)
 			{
