add a "make convex hulls" tool to fracture editor w/ algo to make non-overlapping convex hulls
+ change geometry collection's convex storage to be TSets of convex hulls per transform, to support representing transforms w/ multiple hulls

diff --git a/Engine/Plugins/Experimental/ChaosEditor/Source/FractureEditor/Private/FractureEditorCommands.h b/Engine/Plugins/Experimental/ChaosEditor/Source/FractureEditor/Private/FractureEditorCommands.h
--- a/Engine/Plugins/Experimental/ChaosEditor/Source/FractureEditor/Private/FractureEditorCommands.h
+++ b/Engine/Plugins/Experimental/ChaosEditor/Source/FractureEditor/Private/FractureEditorCommands.h
@@ -64,6 +64,7 @@ class FFractureEditorCommands : public TCommands<FFractureEditorCommands>
 		TSharedPtr< FUICommandInfo > Resample;
 		TSharedPtr< FUICommandInfo > ConvertToMesh;
 		TSharedPtr< FUICommandInfo > Validate;
+		TSharedPtr< FUICommandInfo > MakeConvex;
 
 		// Property Commands
 		TSharedPtr< FUICommandInfo > SetInitialDynamicState;
diff --git a/Engine/Plugins/Experimental/ChaosEditor/Source/FractureEditor/Private/FractureEditorModeToolkit.cpp b/Engine/Plugins/Experimental/ChaosEditor/Source/FractureEditor/Private/FractureEditorModeToolkit.cpp
--- a/Engine/Plugins/Experimental/ChaosEditor/Source/FractureEditor/Private/FractureEditorModeToolkit.cpp
+++ b/Engine/Plugins/Experimental/ChaosEditor/Source/FractureEditor/Private/FractureEditorModeToolkit.cpp
@@ -528,6 +528,7 @@ void FFractureEditorModeToolkit::BuildToolPalette(FName PaletteIndex, class FToo
 		ToolbarBuilder.AddToolBarButton(Commands.Resample);
 		ToolbarBuilder.AddToolBarButton(Commands.ConvertToMesh);
 		ToolbarBuilder.AddToolBarButton(Commands.Validate);
+		ToolbarBuilder.AddToolBarButton(Commands.MakeConvex);
 	}
 	else if (PaletteIndex == TEXT("Properties"))
 	{
diff --git a/Engine/Plugins/Experimental/ChaosEditor/Source/FractureEditor/Private/FractureToolConvex.cpp b/Engine/Plugins/Experimental/ChaosEditor/Source/FractureEditor/Private/FractureToolConvex.cpp
new file mode 100644
--- /dev/null
+++ b/Engine/Plugins/Experimental/ChaosEditor/Source/FractureEditor/Private/FractureToolConvex.cpp
@@ -0,0 +1,161 @@
+// Copyright Epic Games, Inc. All Rights Reserved.
+
+#include "FractureToolConvex.h"
+
+#include "FractureToolContext.h"
+
+#include "GeometryCollection/GeometryCollectionObject.h"
+#include "GeometryCollection/GeometryCollectionComponent.h"
+#include "GeometryCollection/GeometryCollectionAlgo.h"
+#include "GeometryCollection/GeometryCollectionConvexUtility.h"
+
+#define LOCTEXT_NAMESPACE "FractureToolConvex"
+
+UFractureToolConvex::UFractureToolConvex(const FObjectInitializer& ObjInit)
+	: Super(ObjInit)
+{
+	ConvexSettings = NewObject<UFractureConvexSettings>(GetTransientPackage(), UFractureConvexSettings::StaticClass());
+	ConvexSettings->OwnerTool = this;
+}
+
+bool UFractureToolConvex::CanExecute() const
+{
+	if (!IsGeometryCollectionSelected())
+	{
+		return false;
+	}
+
+	return true;
+}
+
+FText UFractureToolConvex::GetDisplayText() const
+{
+	return FText(LOCTEXT("FractureToolConvex", "Make Convex Hulls"));
+}
+
+FText UFractureToolConvex::GetTooltipText() const
+{
+	return FText(LOCTEXT("FractureToolConvexTooltip", "This tool creates (non-overlapping) convex hulls for the bones of geometry collections"));
+}
+
+FSlateIcon UFractureToolConvex::GetToolIcon() const
+{
+	return FSlateIcon("FractureEditorStyle", "FractureEditor.Convex");
+}
+
+void UFractureToolConvex::RegisterUICommand(FFractureEditorCommands* BindingContext)
+{
+	UI_COMMAND_EXT(BindingContext, UICommandInfo, "Convex", "Convex", "Make Convex Hulls", EUserInterfaceActionType::ToggleButton, FInputChord());
+	BindingContext->MakeConvex = UICommandInfo;
+}
+
+TArray<UObject*> UFractureToolConvex::GetSettingsObjects() const
+{
+	TArray<UObject*> Settings;
+	Settings.Add(ConvexSettings);
+	return Settings;
+}
+
+void UFractureToolConvex::FractureContextChanged()
+{
+	TArray<FFractureToolContext> FractureContexts = GetFractureToolContexts();
+
+	HullPoints.Reset();
+	HullEdges.Reset();
+
+	for (const FFractureToolContext& FractureContext : FractureContexts)
+	{
+		FGeometryCollection& Collection = *FractureContext.GetGeometryCollection();
+
+		if (!Collection.HasAttribute("ConvexHull", "Convex") ||
+			!Collection.HasAttribute("TransformToConvexIndices", FGeometryCollection::TransformGroup))
+		{
+			continue;
+		}
+
+		FTransform OuterTransform = FractureContext.GetTransform();
+		for (int32 TransformIdx : FractureContext.GetSelection())
+		{
+			FTransform InnerTransform = GeometryCollectionAlgo::GlobalMatrix(Collection.Transform, Collection.Parent, TransformIdx);
+			if (Collection.HasAttribute("ExplodedVector", FGeometryCollection::TransformGroup))
+			{
+				TManagedArray<FVector3f>& ExplodedVectors = Collection.GetAttribute<FVector3f>("ExplodedVector", FGeometryCollection::TransformGroup);
+				InnerTransform = InnerTransform * FTransform(ExplodedVectors[TransformIdx]);
+			}
+
+			FTransform CombinedTransform = InnerTransform * OuterTransform;
+
+			TManagedArray<TSet<int32>>& TransformToConvexIndices = Collection.GetAttribute<TSet<int32>>("TransformToConvexIndices", FTransformCollection::TransformGroup);
+			TManagedArray<TUniquePtr<Chaos::FConvex>>& ConvexHull = Collection.GetAttribute<TUniquePtr<Chaos::FConvex>>("ConvexHull", "Convex");
+
+			for (int32 ConvexIdx : TransformToConvexIndices[TransformIdx])
+			{
+				int32 HullPtsStart = HullPoints.Num();
+				for (Chaos::FVec3 Pt : ConvexHull[ConvexIdx]->GetVertices())
+				{
+					HullPoints.Add(CombinedTransform.TransformPosition(Pt));
+				}
+				int32 NumPlanes = ConvexHull[ConvexIdx]->NumPlanes();
+				const Chaos::FConvexStructureData& HullData = ConvexHull[ConvexIdx]->GetStructureData();
+				for (int PlaneIdx = 0; PlaneIdx < NumPlanes; PlaneIdx++)
+				{
+					int32 NumPlaneVerts = HullData.NumPlaneVertices(PlaneIdx);
+					for (int32 PlaneVertexIdx = 0; PlaneVertexIdx < NumPlaneVerts; PlaneVertexIdx++)
+					{
+						int32 V0 = HullPtsStart + HullData.GetPlaneVertex(PlaneIdx, PlaneVertexIdx);
+						int32 V1 = HullPtsStart + HullData.GetPlaneVertex(PlaneIdx, (PlaneVertexIdx + 1) % NumPlaneVerts);
+						HullEdges.Add(TPair<int32, int32>(V0, V1));
+					}
+				}
+			}
+		}
+	}
+}
+
+void UFractureToolConvex::Render(const FSceneView* View, FViewport* Viewport, FPrimitiveDrawInterface* PDI)
+{
+	for (TPair<int32,int32> Edge : HullEdges)
+	{
+		FVector P1 = HullPoints[Edge.Key];
+		FVector P2 = HullPoints[Edge.Value];
+		PDI->DrawLine(P1, P2, FLinearColor::Green, SDPG_Foreground, 0.0f, 0.001f);
+	}
+}
+
+void UFractureToolConvex::PostEditChangeChainProperty(struct FPropertyChangedChainEvent& PropertyChangedEvent)
+{
+	// update any cached data 
+}
+
+TArray<FFractureToolContext> UFractureToolConvex::GetFractureToolContexts() const
+{
+	TArray<FFractureToolContext> Contexts;
+
+	// A context is gathered for each selected GeometryCollection component, or for each individual bone if Group Fracture is not used.
+	TSet<UGeometryCollectionComponent*> GeomCompSelection;
+	GetSelectedGeometryCollectionComponents(GeomCompSelection);
+
+	for (UGeometryCollectionComponent* GeometryCollectionComponent : GeomCompSelection)
+	{
+		// Generate a context for each selected node
+		FFractureToolContext FullSelection(GeometryCollectionComponent);
+
+		Contexts.Add(FullSelection);
+	}
+
+	return Contexts;
+}
+
+
+int32 UFractureToolConvex::ExecuteFracture(const FFractureToolContext& FractureContext)
+{
+	if (FractureContext.GetGeometryCollection().IsValid())
+	{
+		FGeometryCollection& Collection = *FractureContext.GetGeometryCollection();
+		FGeometryCollectionConvexUtility::CreateNonOverlappingConvexHullData(&Collection, ConvexSettings->FractionAllowRemove);
+	}
+
+	return INDEX_NONE;
+}
+
+#undef LOCTEXT_NAMESPACE
diff --git a/Engine/Plugins/Experimental/ChaosEditor/Source/FractureEditor/Private/FractureToolConvex.h b/Engine/Plugins/Experimental/ChaosEditor/Source/FractureEditor/Private/FractureToolConvex.h
new file mode 100644
--- /dev/null
+++ b/Engine/Plugins/Experimental/ChaosEditor/Source/FractureEditor/Private/FractureToolConvex.h
@@ -0,0 +1,68 @@
+// Copyright Epic Games, Inc. All Rights Reserved.
+
+#pragma once
+
+#include "FractureTool.h"
+
+#include "FractureToolConvex.generated.h"
+
+class FFractureToolContext;
+
+/** Settings controlling how convex hulls are generated for geometry collections */
+UCLASS(config = EditorPerProjectUserSettings)
+class UFractureConvexSettings : public UFractureToolSettings
+{
+public:
+
+	GENERATED_BODY()
+
+	UFractureConvexSettings(const FObjectInitializer& ObjInit)
+		: Super(ObjInit)
+	{}
+
+	/** Fraction of the convex hulls for a transform that we can remove before instead using the hulls of the children */
+	UPROPERTY(EditAnywhere, Category = MapSettings, meta = (ClampMin = ".01", ClampMax = "1"))
+	double FractionAllowRemove = .5;
+};
+
+
+UCLASS(DisplayName = "Convex Tool", Category = "FractureTools")
+class UFractureToolConvex : public UFractureModalTool
+{
+public:
+	GENERATED_BODY()
+
+	UFractureToolConvex(const FObjectInitializer& ObjInit);
+
+	///
+	/// UFractureModalTool Interface
+	///
+
+	/** This is the Text that will appear on the tool button to execute the tool **/
+	virtual FText GetDisplayText() const override;
+	virtual FText GetTooltipText() const override;
+	/** This is the Text that will appear on the button to execute the fracture **/
+	virtual FText GetApplyText() const override { return FText(NSLOCTEXT("Convex", "ExecuteConvex", "Make Convex Hulls")); }
+	virtual FSlateIcon GetToolIcon() const override;
+	virtual void RegisterUICommand(FFractureEditorCommands* BindingContext) override;
+	virtual TArray<UObject*> GetSettingsObjects() const override;
+	virtual void FractureContextChanged() override;
+	virtual void Render(const FSceneView* View, FViewport* Viewport, FPrimitiveDrawInterface* PDI) override;
+	/** Executes function that generates new geometry. Returns the first new geometry index. */
+	virtual int32 ExecuteFracture(const FFractureToolContext& FractureContext) override;
+	virtual bool CanExecute() const override;
+	virtual void PostEditChangeChainProperty(struct FPropertyChangedChainEvent& PropertyChangedEvent) override;
+
+	/** Gets the UI command info for this command */
+	const TSharedPtr<FUICommandInfo>& GetUICommandInfo() const;
+
+	virtual TArray<FFractureToolContext> GetFractureToolContexts() const override;
+
+
+protected:
+	UPROPERTY(EditAnywhere, Category = Slicing)
+	TObjectPtr<UFractureConvexSettings> ConvexSettings;
+
+	TArray<FVector> HullPoints;
+	TArray<TPair<int32, int32>> HullEdges;
+};
diff --git a/Engine/Plugins/Experimental/ChaosEditor/Source/FractureEditor/Private/FractureToolEmbed.cpp b/Engine/Plugins/Experimental/ChaosEditor/Source/FractureEditor/Private/FractureToolEmbed.cpp
--- a/Engine/Plugins/Experimental/ChaosEditor/Source/FractureEditor/Private/FractureToolEmbed.cpp
+++ b/Engine/Plugins/Experimental/ChaosEditor/Source/FractureEditor/Private/FractureToolEmbed.cpp
@@ -192,6 +192,7 @@ void UFractureToolAutoEmbedGeometry::Execute(TWeakPtr<FFractureEditorModeToolkit
 				FGeometryCollectionConvexUtility::FGeometryCollectionConvexData ConvexData = FGeometryCollectionConvexUtility::GetValidConvexHullData(GeometryCollection);
 				const TManagedArray<FTransform>& Transform = GeometryCollection->Transform;
 				const TManagedArray<int32>& Parent = GeometryCollection->Parent;
+				const TManagedArray<int32>& SimulationType = GeometryCollection->SimulationType;
 				TArray<FTransform> BoneGlobalTransforms;
 				GeometryCollectionAlgo::GlobalMatrices(Transform, Parent, BoneGlobalTransforms);
 
@@ -199,23 +200,27 @@ void UFractureToolAutoEmbedGeometry::Execute(TWeakPtr<FFractureEditorModeToolkit
 				FVector ComponentSpaceLocation = WorldToComponent.TransformPosition(SMLocation);
 
 				int32 NumConvex = ConvexData.ConvexHull.Num();
-				for (int32 ConvexIndex = 0; ConvexIndex < NumConvex; ++ConvexIndex)
+				for (int32 TransformIndex = 0; TransformIndex < Transform.Num(); TransformIndex++)
 				{
-					// transform into bone space where the convex hull is described
-					int32 TransformIndex = ConvexData.TransformToConvexIndex.Find(ConvexIndex);
-					FTransform ComponentToBone = BoneGlobalTransforms[TransformIndex].Inverse();
-					FVector BoneSpaceLocation = ComponentToBone.TransformPosition(ComponentSpaceLocation);
-
-					Chaos::FVec3 Normal;
-					Chaos::FReal Phi = ConvexData.ConvexHull[ConvexIndex]->PhiWithNormal(BoneSpaceLocation, Normal);
-					UE_LOG(LogTemp, Warning, TEXT(" Cone %d Phi %f"), ConvexData.TransformToConvexIndex.Find(ConvexIndex), Phi);
-					if (Phi < ClosestPhi)
+					if (SimulationType[TransformIndex] != FGeometryCollection::ESimulationTypes::FST_Rigid)
 					{
-						ClosestPhi = Phi;
-						ClosestGeometryCollection = Context.GetGeometryCollection();
-						ClosestConvex = ConvexIndex;
-						ClosestComponent = Context.GetGeometryCollectionComponent();
-						ClosestContext = &Context;
+						continue;
+					}
+
+					FVector BoneSpaceLocation = BoneGlobalTransforms[TransformIndex].InverseTransformPosition(ComponentSpaceLocation);
+					for (int32 ConvexIndex : ConvexData.TransformToConvexIndices[TransformIndex])
+					{
+						Chaos::FVec3 Normal;
+						Chaos::FReal Phi = ConvexData.ConvexHull[ConvexIndex]->PhiWithNormal(BoneSpaceLocation, Normal);
+						UE_LOG(LogTemp, Warning, TEXT(" Cone %d Phi %f"), TransformIndex, Phi);
+						if (Phi < ClosestPhi)
+						{
+							ClosestPhi = Phi;
+							ClosestGeometryCollection = Context.GetGeometryCollection();
+							ClosestConvex = ConvexIndex;
+							ClosestComponent = Context.GetGeometryCollectionComponent();
+							ClosestContext = &Context;
+						}
 					}
 				}
 			}
@@ -225,10 +230,10 @@ void UFractureToolAutoEmbedGeometry::Execute(TWeakPtr<FFractureEditorModeToolkit
 				// Which bone points to the closest convex?
 				const TManagedArray<FTransform>& Transform = ClosestGeometryCollection->Transform;
 				const TManagedArray<int32>& Parent = ClosestGeometryCollection->Parent;
-				TManagedArray<int32>& TransformToConvexIndex =
-					ClosestGeometryCollection->GetAttribute<int32>("TransformToConvexIndex", FTransformCollection::TransformGroup);
+				TManagedArray<int32>& TransformToConvexIndices =
+					ClosestGeometryCollection->GetAttribute<int32>("TransformToConvexIndices", FTransformCollection::TransformGroup);
 
-				const int32 BoneIndex = TransformToConvexIndex.Find(ClosestConvex);
+				const int32 BoneIndex = TransformToConvexIndices.Find(ClosestConvex);
 
 				if (BoneIndex > INDEX_NONE)
 				{
diff --git a/Engine/Source/Runtime/Experimental/Chaos/Private/GeometryCollection/GeometryCollection.cpp b/Engine/Source/Runtime/Experimental/Chaos/Private/GeometryCollection/GeometryCollection.cpp
--- a/Engine/Source/Runtime/Experimental/Chaos/Private/GeometryCollection/GeometryCollection.cpp
+++ b/Engine/Source/Runtime/Experimental/Chaos/Private/GeometryCollection/GeometryCollection.cpp
@@ -1007,6 +1007,30 @@ void FGeometryCollection::Serialize(Chaos::FChaosArchive& Ar)
 			Version = 6;
 		}
 
+		if (Version < 7)
+		{
+			if (HasAttribute("TransformToConvexIndex", FTransformCollection::TransformGroup))
+			{
+				TManagedArray<int32> TransformToConvexIndex = MoveTemp(GetAttribute<int32>("TransformToConvexIndex", FTransformCollection::TransformGroup));
+				RemoveAttribute("TransformToConvexIndex", FTransformCollection::TransformGroup);
+				// if we don't already have the one-to-many version, convert the previous one-to-one mapping to the new format
+				if (!HasAttribute("TransformToConvexIndices", FTransformCollection::TransformGroup))
+				{
+					FManagedArrayCollection::FConstructionParameters ConvexDependency("Convex");
+					TManagedArray<TSet<int32>>& IndexSets = AddAttribute<TSet<int32>>("TransformToConvexIndices", FTransformCollection::TransformGroup, ConvexDependency);
+					for (int32 TransformIdx = 0; TransformIdx < TransformToConvexIndex.Num(); TransformIdx++)
+					{
+						int32 ConvexIdx = TransformToConvexIndex[TransformIdx];
+						if (ConvexIdx != INDEX_NONE)
+						{
+							IndexSets[TransformIdx].Add(ConvexIdx);
+						}
+					}
+				}
+			}
+			Version = 7;
+		}
+
 	}
 }
 
diff --git a/Engine/Source/Runtime/Experimental/Chaos/Private/GeometryCollection/GeometryCollectionConvexUtility.cpp b/Engine/Source/Runtime/Experimental/Chaos/Private/GeometryCollection/GeometryCollectionConvexUtility.cpp
--- a/Engine/Source/Runtime/Experimental/Chaos/Private/GeometryCollection/GeometryCollectionConvexUtility.cpp
+++ b/Engine/Source/Runtime/Experimental/Chaos/Private/GeometryCollection/GeometryCollectionConvexUtility.cpp
@@ -2,7 +2,9 @@
 
 #include "GeometryCollection/GeometryCollectionConvexUtility.h"
 #include "Chaos/Convex.h"
+#include "Chaos/GJK.h"
 #include "GeometryCollection/GeometryCollection.h"
+#include "GeometryCollection/GeometryCollectionAlgo.h"
 
 
 FGeometryCollectionConvexUtility::FGeometryCollectionConvexData FGeometryCollectionConvexUtility::GetValidConvexHullData(FGeometryCollection* GeometryCollection)
@@ -14,11 +16,10 @@ FGeometryCollectionConvexUtility::FGeometryCollectionConvexData FGeometryCollect
 		GeometryCollection->AddGroup("Convex");
 	}
 
-	if (!GeometryCollection->HasAttribute("TransformToConvexIndex", FTransformCollection::TransformGroup))
+	if (!GeometryCollection->HasAttribute("TransformToConvexIndices", FTransformCollection::TransformGroup))
 	{
 		FManagedArrayCollection::FConstructionParameters ConvexDependency("Convex");
-		TManagedArray<int32>& IndexArray = GeometryCollection->AddAttribute<int32>("TransformToConvexIndex", FTransformCollection::TransformGroup, ConvexDependency);
-		IndexArray.Fill(INDEX_NONE);
+		TManagedArray<TSet<int32>>& IndexArray = GeometryCollection->AddAttribute<TSet<int32>>("TransformToConvexIndices", FTransformCollection::TransformGroup, ConvexDependency);
 	}
 
 	if (!GeometryCollection->HasAttribute("ConvexHull", "Convex"))
@@ -29,7 +30,7 @@ FGeometryCollectionConvexUtility::FGeometryCollectionConvexData FGeometryCollect
 	// Check for correct population. All (and only) rigid nodes should have a convex associated.
 	const TManagedArray<int32>& SimulationType = GeometryCollection->GetAttribute<int32>("SimulationType", FTransformCollection::TransformGroup);
 	const TManagedArray<int32>& TransformToGeometryIndex = GeometryCollection->GetAttribute<int32>("TransformToGeometryIndex", FTransformCollection::TransformGroup);
-	TManagedArray<int32>& TransformToConvexIndex = GeometryCollection->GetAttribute<int32>("TransformToConvexIndex", FTransformCollection::TransformGroup);
+	TManagedArray<TSet<int32>>& TransformToConvexIndices = GeometryCollection->GetAttribute<TSet<int32>>("TransformToConvexIndices", FTransformCollection::TransformGroup);
 	TManagedArray<TUniquePtr<Chaos::FConvex>>& ConvexHull = GeometryCollection->GetAttribute<TUniquePtr<Chaos::FConvex>>("ConvexHull", "Convex");
 
 	TArray<int32> ProduceConvexHulls;
@@ -39,12 +40,11 @@ FGeometryCollectionConvexUtility::FGeometryCollectionConvexData FGeometryCollect
 
 	for (int32 Idx = 0; Idx < SimulationType.Num(); ++Idx)
 	{
-		if ((SimulationType[Idx] == FGeometryCollection::ESimulationTypes::FST_Rigid) && (TransformToConvexIndex[Idx] == INDEX_NONE))
+		if ((SimulationType[Idx] == FGeometryCollection::ESimulationTypes::FST_Rigid) && (TransformToConvexIndices[Idx].Num() == 0))
 		{
 			ProduceConvexHulls.Add(Idx);
-			
 		}
-		else if ((SimulationType[Idx] != FGeometryCollection::ESimulationTypes::FST_Rigid) && (TransformToConvexIndex[Idx] > INDEX_NONE))
+		else if ((SimulationType[Idx] != FGeometryCollection::ESimulationTypes::FST_Rigid) && (TransformToConvexIndices[Idx].Num() > 0))
 		{
 			EliminateConvexHulls.Add(Idx);
 		}
@@ -55,7 +55,7 @@ FGeometryCollectionConvexUtility::FGeometryCollectionConvexData FGeometryCollect
 		GeometryCollection->RemoveElements("Convex", EliminateConvexHulls);
 		for (int32 Idx : EliminateConvexHulls)
 		{
-			TransformToConvexIndex[Idx] = INDEX_NONE;
+			TransformToConvexIndices[Idx].Reset();
 		}
 	}
 
@@ -66,15 +66,640 @@ FGeometryCollectionConvexUtility::FGeometryCollectionConvexData FGeometryCollect
 		{
 			int32 GeometryIdx = TransformToGeometryIndex[ProduceConvexHulls[Idx]];
 			ConvexHull[NewConvexIndexStart + Idx] = FindConvexHull(GeometryCollection, GeometryIdx);
-			TransformToConvexIndex[ProduceConvexHulls[Idx]] = NewConvexIndexStart + Idx;
+			TransformToConvexIndices[ProduceConvexHulls[Idx]].Reset();
+			TransformToConvexIndices[ProduceConvexHulls[Idx]].Add(NewConvexIndexStart + Idx);
 		}
 	}
 
-	return { TransformToConvexIndex, ConvexHull };
+	return { TransformToConvexIndices, ConvexHull };
+}
+
+
+
+namespace
+{
+
+typedef Chaos::TPlaneConcrete<Chaos::FReal, 3> FChaosPlane;
+
+/// Cut hull with plane, generating the point set of a new hull
+/// @return false if plane did not cut any points on the hull
+bool CutHull(const Chaos::FConvex& HullIn, FChaosPlane Plane, bool KeepSide, TArray<Chaos::FVec3>& HullPtsOut)
+{
+	const TArray<Chaos::FVec3>& Vertices = HullIn.GetVertices();
+	const Chaos::FConvexStructureData& HullData = HullIn.GetStructureData();
+	bool bHasOutside = false;
+	for (int VertIdx = 0; VertIdx < Vertices.Num(); VertIdx++)
+	{
+		const Chaos::FVec3& V = Vertices[VertIdx];
+		if ((Plane.SignedDistance(V) < 0) == KeepSide)
+		{
+			HullPtsOut.Add(V);
+		}
+		else
+		{
+			bHasOutside = true;
+		}
+	}
+
+	if (!bHasOutside)
+	{
+		return false;
+	}
+
+	int32 NumPlanes = HullIn.NumPlanes();
+	for (int PlaneIdx = 0; PlaneIdx < NumPlanes; PlaneIdx++)
+	{
+		int32 NumPlaneVerts = HullData.NumPlaneVertices(PlaneIdx);
+		for (int32 PlaneVertexIdx = 0; PlaneVertexIdx < NumPlaneVerts; PlaneVertexIdx++)
+		{
+			int32 NextVertIdx = (PlaneVertexIdx + 1) % NumPlaneVerts;
+			const Chaos::FVec3& V0 = Vertices[HullData.GetPlaneVertex(PlaneIdx, PlaneVertexIdx)];
+			const Chaos::FVec3& V1 = Vertices[HullData.GetPlaneVertex(PlaneIdx, NextVertIdx)];
+			if ((Plane.SignedDistance(V0) < 0) != (Plane.SignedDistance(V1) < 0))
+			{
+				Chaos::Pair<Chaos::FVec3, bool> Res = Plane.FindClosestIntersection(V0, V1, 0);
+				if (Res.Second)
+				{
+					HullPtsOut.Add(Res.First);
+				}
+			}
+		}
+	}
+
+	return true;
+}
+
+
+/// Cut hull with plane, keeping both sides and generating the point set of both new hulls
+/// @return false if plane did not cut any points on the hull
+bool SplitHull(const Chaos::FConvex& HullIn, FChaosPlane Plane, bool KeepSide, TArray<Chaos::FVec3>& InsidePtsOut, TArray<Chaos::FVec3>& OutsidePtsOut)
+{
+	const TArray<Chaos::FVec3>& Vertices = HullIn.GetVertices();
+	const Chaos::FConvexStructureData& HullData = HullIn.GetStructureData();
+	bool bHasOutside = false;
+	for (int VertIdx = 0; VertIdx < Vertices.Num(); VertIdx++)
+	{
+		const Chaos::FVec3& V = Vertices[VertIdx];
+		if ((Plane.SignedDistance(V) < 0) == KeepSide)
+		{
+			InsidePtsOut.Add(V);
+		}
+		else
+		{
+			OutsidePtsOut.Add(V);
+			bHasOutside = true;
+		}
+	}
+
+	if (!bHasOutside)
+	{
+		return false;
+	}
+
+	int32 NumPlanes = HullIn.NumPlanes();
+	for (int PlaneIdx = 0; PlaneIdx < NumPlanes; PlaneIdx++)
+	{
+		int32 NumPlaneVerts = HullData.NumPlaneVertices(PlaneIdx);
+		for (int32 PlaneVertexIdx = 0; PlaneVertexIdx < NumPlaneVerts; PlaneVertexIdx++)
+		{
+			int32 NextVertIdx = (PlaneVertexIdx + 1) % NumPlaneVerts;
+			const Chaos::FVec3& V0 = Vertices[HullData.GetPlaneVertex(PlaneIdx, PlaneVertexIdx)];
+			const Chaos::FVec3& V1 = Vertices[HullData.GetPlaneVertex(PlaneIdx, NextVertIdx)];
+			if ((Plane.SignedDistance(V0) < 0) != (Plane.SignedDistance(V1) < 0))
+			{
+				Chaos::Pair<Chaos::FVec3, bool> Res = Plane.FindClosestIntersection(V0, V1, 0);
+				if (Res.Second)
+				{
+					InsidePtsOut.Add(Res.First);
+					OutsidePtsOut.Add(Res.First);
+				}
+			}
+		}
+	}
+
+	return true;
+}
+
+
+/// Assumptions: 
+///		Convexes is initialized to one convex hull for each leaf geometry in a SHARED coordinate space
+///		TransformToConvexIndices is initialized to point to the existing convex hulls
+///		Parents, GeoProximity, and GeometryToTransformIndex are all initialized from the geometry collection
+void CreateNonoverlappingConvexHulls(
+	TArray<TUniquePtr<Chaos::FConvex>>& Convexes,
+	TArray<TSet<int32>>& TransformToConvexIndices,
+	const TManagedArray<int32>& SimulationType,
+	int32 LeafType,
+	const TManagedArray<int32>& Parents,
+	const TManagedArray<TSet<int32>>& GeoProximity,
+	const TManagedArray<int32>& GeometryToTransformIndex,
+	double FracAllowRemove
+)
+{
+	int32 NumBones = TransformToConvexIndices.Num();
+	check(Parents.Num() == NumBones);
+
+	auto OnlyConvex = [&TransformToConvexIndices](int32 Bone) -> int32
+	{
+		ensure(TransformToConvexIndices[Bone].Num() == 1);
+		for (int32 ConvexIndex : TransformToConvexIndices[Bone])
+		{
+			return ConvexIndex;
+		}
+		return INDEX_NONE;
+	};
+
+	TArray<TSet<int32>> LeafProximity;
+	LeafProximity.SetNum(NumBones);
+	for (int32 GeomIdx = 0; GeomIdx < GeoProximity.Num(); GeomIdx++)
+	{
+		int32 TransformIdx = GeometryToTransformIndex[GeomIdx];
+		for (int32 NbrGeomIdx : GeoProximity[GeomIdx])
+		{
+			LeafProximity[TransformIdx].Add(GeometryToTransformIndex[NbrGeomIdx]);
+		}
+	}
+
+	auto IsColliding2 = [&Convexes](int32 ConvexA, int32 ConvexB)
+	{
+		if (!ensure(Convexes[ConvexA]->NumVertices() > 0 && Convexes[ConvexB]->NumVertices() > 0))
+		{
+			// at least one of the convex hulls was empty, so cannot be colliding
+			// TODO: eventually remove this ensure; and add more handling for empty hulls
+			return false;
+		}
+		const Chaos::TRigidTransform<Chaos::FReal, 3> IdentityTransform = Chaos::TRigidTransform<Chaos::FReal, 3>::Identity;
+		return GJKIntersection(*Convexes[ConvexA], *Convexes[ConvexB], IdentityTransform);
+	};
+
+	auto FixCollisionWithCut2 = [&Convexes](int32 ConvexA, int32 ConvexB)
+	{
+		if (!ensure(Convexes[ConvexA]->NumVertices() > 0 && Convexes[ConvexB]->NumVertices() > 0))
+		{
+			// at least one of the convex hulls was empty, so cannot be colliding
+			// TODO: eventually remove this ensure; and add more handling for empty hulls
+			return false;
+		}
+		Chaos::FReal Depth;
+		Chaos::FVec3 CloseA, CloseB, Normal;
+		int32 OutIdxA, OutIdxB;
+		const Chaos::TRigidTransform<Chaos::FReal, 3> IdentityTransform = Chaos::TRigidTransform<Chaos::FReal, 3>::Identity;
+		bool bCollide = Chaos::GJKPenetration(*Convexes[ConvexA], *Convexes[ConvexB], IdentityTransform, Depth, CloseA, CloseB, Normal, OutIdxA, OutIdxB);
+		if (bCollide)
+		{
+			FChaosPlane CutPlane((CloseA + CloseB) * .5, Normal);
+			TArray<Chaos::FVec3> CutHullPts;
+			if (CutHull(*Convexes[ConvexA], CutPlane, true, CutHullPts))
+			{
+				*Convexes[ConvexA] = Chaos::FConvex(CutHullPts, KINDA_SMALL_NUMBER);
+			}
+			CutHullPts.Reset();
+			if (CutHull(*Convexes[ConvexB], CutPlane, false, CutHullPts))
+			{
+				*Convexes[ConvexB] = Chaos::FConvex(CutHullPts, KINDA_SMALL_NUMBER);
+			}
+		}
+		return bCollide;
+	};
+
+	// Initialize Children and Depths of tree
+	// Fix collisions between input hulls using the input proximity relationships
+
+	int32 MaxDepth = 0;
+	TArray<int32> Depths;
+	TArray<TArray<int32>> Children;
+	Children.SetNum(NumBones);
+	Depths.SetNumZeroed(NumBones);
+	for (int HullIdx = 0; HullIdx < NumBones; HullIdx++)
+	{
+		if (Parents[HullIdx] != INDEX_NONE)
+		{
+			if (SimulationType[Parents[HullIdx]] != LeafType)
+			{
+				Children[Parents[HullIdx]].Add(HullIdx);
+			}
+			else
+			{
+				Depths[HullIdx] = -1; // child-of-leaf == embedded geometry, just ignore it
+				continue;
+			}
+		}
+		int32 Depth = 0, WalkParent = HullIdx;
+		while (Parents[WalkParent] != INDEX_NONE)
+		{
+			Depth++;
+			WalkParent = Parents[WalkParent];
+		}
+		Depths[HullIdx] = Depth;
+		MaxDepth = FMath::Max(Depth, MaxDepth);
+
+		if (TransformToConvexIndices[HullIdx].Num() > 0)
+		{
+			const TSet<int32>& Neighbors = LeafProximity[HullIdx];
+			for (int32 NbrIdx : Neighbors)
+			{
+				if (NbrIdx < HullIdx && TransformToConvexIndices[NbrIdx].Num() > 0)
+				{
+					FixCollisionWithCut2(OnlyConvex(HullIdx), OnlyConvex(NbrIdx));
+				}
+			}
+		}
+	}
+
+	auto AddLeaves = [&Children, &SimulationType, LeafType](int32 Bone, TArray<int32>& Leaves)
+	{
+		TArray<int32> ToExpand;
+		ToExpand.Add(Bone);
+		while (ToExpand.Num() > 0)
+		{
+			int32 ToProcess = ToExpand.Pop(false);
+			if (Children[ToProcess].Num() > 0 && SimulationType[ToProcess] != LeafType)
+			{
+				ToExpand.Append(Children[ToProcess]);
+			}
+			else
+			{
+				Leaves.Add(ToProcess);
+			}
+		}
+	};
+
+	auto FixLeafCollisions = [&AddLeaves, &FixCollisionWithCut2, &OnlyConvex](int32 BoneA, int32 BoneB)
+	{
+		TArray<int32> LeavesA, LeavesB;
+		AddLeaves(BoneA, LeavesA);
+		AddLeaves(BoneB, LeavesB);
+		bool bAnyCollides = false;
+		for (int32 LeafA : LeavesA)
+		{
+			for (int32 LeafB : LeavesB)
+			{
+				bool bCollides = FixCollisionWithCut2(OnlyConvex(LeafA), OnlyConvex(LeafB));
+				bAnyCollides |= bCollides;
+			}
+		}
+		return bAnyCollides;
+	};
+
+	// Compute initial hulls at all levels, and use these to fill out the full proximity links
+	// Fix collisions between any two leaf hulls
+
+	TArray<TSet<int32>> ClusterProximity;
+	ClusterProximity.SetNum(NumBones);
+	int ProcessDepth = MaxDepth;
+	while (--ProcessDepth >= 0)
+	{
+		TSet<int32> ToProcess;
+		for (int Bone = 0; Bone < Parents.Num(); Bone++)
+		{
+			if (Depths[Bone] == ProcessDepth)
+			{
+				if (TransformToConvexIndices[Bone].Num() == 0)
+				{
+					TArray<Chaos::FVec3> JoinedHullPts;
+					for (int32 Child : Children[Bone])
+					{
+						for (int32 ConvexIdx : TransformToConvexIndices[Child])
+						{
+							JoinedHullPts.Append(Convexes[ConvexIdx]->GetVertices());
+						}
+					}
+					if (JoinedHullPts.Num() > 0)
+					{
+						int32 ConvexIdx = Convexes.Add(MakeUnique<Chaos::FConvex>(JoinedHullPts, KINDA_SMALL_NUMBER));
+						TransformToConvexIndices[Bone].Add(ConvexIdx);
+						ToProcess.Add(Bone);
+					}
+				}
+				else
+				{
+					ToProcess.Add(Bone);
+				}
+			}
+		}
+
+		ensure(ToProcess.Num() > 0);
+
+		// We don't have proximity at the cluster level, so reconstruct from n^2 collision tests for now
+		for (int32 BoneA : ToProcess)
+		{
+			for (int32 BoneB : ToProcess)
+			{
+				if (BoneB < BoneA)
+				{
+					if (IsColliding2(OnlyConvex(BoneA), OnlyConvex(BoneB)))
+					{
+						bool bLeavesCollide = FixLeafCollisions(BoneA, BoneB);
+						int32 Bones[2]{ BoneA, BoneB };
+						for (int32 BoneIdx = 0; BoneIdx < 2; BoneIdx++)
+						{
+							int32 ParentBone = Bones[BoneIdx];
+							int32 OtherBone = Bones[1 - BoneIdx];
+							// if leaves changed flag that all hulls using those leaves must also change
+							TArray<int32> TraverseBones;
+							TraverseBones.Append(Children[ParentBone]);
+							while (TraverseBones.Num() > 0)
+							{
+								int32 ToProc = TraverseBones.Pop(false);
+								if (IsColliding2(OnlyConvex(OtherBone), OnlyConvex(ToProc)))
+								{
+									ClusterProximity[OtherBone].Add(ToProc);
+									ClusterProximity[ToProc].Add(OtherBone);
+									TraverseBones.Append(Children[ToProc]);
+								}
+							}
+						}
+
+						ClusterProximity[BoneA].Add(BoneB);
+						ClusterProximity[BoneB].Add(BoneA);
+					}
+				}
+			}
+		}
+	}
+
+	// Now that leaves don't intersect, re-compute all non-leaf hulls on the clipped leaves
+	// and record volumes for each hull
+
+	TArray<double> NonLeafVolumes; // Original volumes of non-leaf hulls (to compare against progressively cut-down volume as intersections are removed)
+	NonLeafVolumes.SetNumZeroed(Convexes.Num());
+	ProcessDepth = MaxDepth;
+	while (--ProcessDepth > 0)
+	{
+		TSet<int32> ToProcess;
+		for (int32 Bone = 0; Bone < NumBones; Bone++)
+		{
+			if (Depths[Bone] == ProcessDepth && Children[Bone].Num() > 0)
+			{
+				TArray<Chaos::FVec3> JoinedHullPts;
+				for (int32 Child : Children[Bone])
+				{
+					for (int32 ConvexIdx : TransformToConvexIndices[Child])
+					{
+						JoinedHullPts.Append(Convexes[ConvexIdx]->GetVertices());
+					}
+				}
+				int32 ConvexIdx = OnlyConvex(Bone);
+				*Convexes[ConvexIdx] = Chaos::FConvex(JoinedHullPts, KINDA_SMALL_NUMBER);
+				NonLeafVolumes[ConvexIdx] = Convexes[ConvexIdx]->GetVolume();
+			}
+		}
+	}
+
+	auto CutIfOk = [&Convexes, &OnlyConvex, &NonLeafVolumes, &FracAllowRemove](bool bOneSidedCut, int32 ConvexA, int32 ConvexB) -> bool
+	{
+		Chaos::FReal Depth;
+		Chaos::FVec3 CloseA, CloseB, Normal;
+		int32 OutIdxA, OutIdxB;
+		const Chaos::TRigidTransform<Chaos::FReal, 3> IdentityTransform = Chaos::TRigidTransform<Chaos::FReal, 3>::Identity;
+		bool bCollide = Chaos::GJKPenetration(*Convexes[ConvexA], *Convexes[ConvexB], IdentityTransform, Depth, CloseA, CloseB, Normal, OutIdxA, OutIdxB);
+		if (bCollide)
+		{
+			Chaos::FVec3 CutPt = CloseB;
+			if (!bOneSidedCut)
+			{
+				CutPt = (CloseA + CloseB) * .5;
+			}
+			FChaosPlane CutPlane(CutPt, Normal);
+
+			// Tentatively create the clipped hulls
+			Chaos::FConvex CutHullA, CutHullB;
+			bool bCreatedA = false, bCreatedB = false;
+			TArray<Chaos::FVec3> CutHullPts;
+			if (CutHull(*Convexes[ConvexA], CutPlane, true, CutHullPts))
+			{
+				CutHullA = Chaos::FConvex(CutHullPts, KINDA_SMALL_NUMBER);
+				bCreatedA = true;
+			}
+			if (!bOneSidedCut)
+			{
+				CutHullPts.Reset();
+				if (CutHull(*Convexes[ConvexB], CutPlane, false, CutHullPts))
+				{
+					CutHullB = Chaos::FConvex(CutHullPts, KINDA_SMALL_NUMBER);
+					bCreatedB = true;
+				}
+			}
+
+			// Test if the clipped hulls have become too small vs the original volumes
+			if (ensure(ConvexA < NonLeafVolumes.Num()) && NonLeafVolumes[ConvexA] > 0 && bCreatedA && CutHullA.GetVolume() / NonLeafVolumes[ConvexA] < 1 - FracAllowRemove)
+			{
+				return false;
+			}
+			if (!bOneSidedCut)
+			{
+				if (ensure(ConvexB < NonLeafVolumes.Num()) && NonLeafVolumes[ConvexB] > 0 && bCreatedB && CutHullB.GetVolume() / NonLeafVolumes[ConvexB] < 1 - FracAllowRemove)
+				{
+					return false;
+				}
+			}
+
+			// If the clipped hulls were large enough, go ahead and set them as the new hulls
+			if (bCreatedA)
+			{
+				*Convexes[ConvexA] = MoveTemp(CutHullA);
+			}
+			if (bCreatedB)
+			{
+				*Convexes[ConvexB] = MoveTemp(CutHullB);
+			}
+
+			return true;
+		}
+		else
+		{
+			return true; // no cut needed, so was ok
+		}
+	};
+
+	// re-process all non-leaf bones
+	ProcessDepth = MaxDepth;
+	while (--ProcessDepth > 0)
+	{
+		TSet<int32> ToProcess;
+		for (int32 Bone = 0; Bone < NumBones; Bone++)
+		{
+			if (Depths[Bone] == ProcessDepth && Children[Bone].Num() > 0)
+			{
+				ToProcess.Add(Bone);
+			}
+		}
+
+		TSet<int32> WasNotOk;
+
+		for (int32 Bone : ToProcess)
+		{
+			if (WasNotOk.Contains(Bone) || TransformToConvexIndices[Bone].Num() == 0)
+			{
+				continue;
+			}
+			for (int32 Nbr : ClusterProximity[Bone])
+			{
+				if (WasNotOk.Contains(Nbr) || TransformToConvexIndices[Nbr].Num() == 0)
+				{
+					continue;
+				}
+				bool bAllOk = true;
+				for (int32 ConvexBone : TransformToConvexIndices[Bone])
+				{
+					for (int32 ConvexNbr : TransformToConvexIndices[Nbr])
+					{
+						bool bOneSidedCut = Depths[Bone] != Depths[Nbr] || Children[Nbr].Num() == 0;
+						// if the neighbor is less deep and not a leaf, skip processing this to favor processing the neighbor instead
+						if (Depths[Bone] > Depths[Nbr] && Children[Nbr].Num() > 0)
+						{
+							continue;
+						}
+
+						bool bWasOk = CutIfOk(bOneSidedCut, ConvexBone, ConvexNbr);
+
+						// cut would have removed too much; just fall back to using the hulls of children
+						// TODO: attempt splitting hulls before fully falling back to this
+						if (!bWasOk)
+						{
+							bAllOk = false;
+							auto ResetBoneHulls = [&WasNotOk, &TransformToConvexIndices](int32 ToReset)
+							{
+								WasNotOk.Add(ToReset);
+								TransformToConvexIndices[ToReset].Reset();
+							};
+							ResetBoneHulls(Bone);
+							if (!bOneSidedCut)
+							{
+								ResetBoneHulls(Nbr);
+							}
+							break;
+						}
+					}
+					if (!bAllOk)
+					{
+						break;
+					}
+				}
+			}
+		}
+	}
+
+	return;
+}
+
+
+/// Helper to get convex hulls from a geometry collection in the format required by CreateNonoverlappingConvexHulls
+void HullsFromGeometry(
+	const FGeometryCollection& Geometry,
+	TArray<FTransform>& GlobalTransformArray,
+	TArray<TUniquePtr<Chaos::FConvex>>& Convexes,
+	TArray<TSet<int32>>& TransformToConvexIndices,
+	const TManagedArray<int32>& SimulationType,
+	int32 RigidType
+)
+{
+	TArray<FVector> GlobalVertices;
 	
+	GlobalVertices.SetNum(Geometry.Vertex.Num());
+	for (int32 Idx = 0; Idx < GlobalVertices.Num(); Idx++)
+	{
+		GlobalVertices[Idx] = GlobalTransformArray[Geometry.BoneMap[Idx]].TransformPosition(Geometry.Vertex[Idx]);
+	}
+
+	int32 NumBones = Geometry.TransformToGeometryIndex.Num();
+	TransformToConvexIndices.SetNum(NumBones);
+	for (int32 Idx = 0; Idx < NumBones; Idx++)
+	{
+		if (SimulationType[Idx] != RigidType)
+		{
+			continue;
+		}
+		int32 GeomIdx = Geometry.TransformToGeometryIndex[Idx];
+		if (GeomIdx != INDEX_NONE)
+		{
+			int32 VStart = Geometry.VertexStart[GeomIdx];
+			int32 VCount = Geometry.VertexCount[GeomIdx];
+			int32 VEnd = VStart + VCount;
+			TArray<Chaos::FVec3> CutHullPts;
+			CutHullPts.Reserve(VCount);
+			for (int32 VIdx = VStart; VIdx < VEnd; VIdx++)
+			{
+				CutHullPts.Add(GlobalVertices[VIdx]);
+			}
+			ensure(CutHullPts.Num() > 0);
+			int32 ConvexIdx = Convexes.Add(MakeUnique<Chaos::FConvex>(CutHullPts, KINDA_SMALL_NUMBER));
+			TransformToConvexIndices[Idx].Add(ConvexIdx);
+		}
+	}
+}
+
+void TransformHullsToLocal(
+	TArray<FTransform>& GlobalTransformArray,
+	TArray<TUniquePtr<Chaos::FConvex>>& Convexes,
+	TArray<TSet<int32>>& TransformToConvexIndices
+)
+{
+	for (int32 Bone = 0; Bone < TransformToConvexIndices.Num(); Bone++)
+	{
+		FTransform& Transform = GlobalTransformArray[Bone];
+		TArray<Chaos::FVec3> HullPts;
+		for (int32 ConvexIdx : TransformToConvexIndices[Bone])
+		{
+			HullPts.Reset();
+			for (const Chaos::FVec3& P : Convexes[ConvexIdx]->GetVertices())
+			{
+				HullPts.Add(Transform.InverseTransformPosition(P));
+			}
+			*Convexes[ConvexIdx] = Chaos::FConvex(HullPts, KINDA_SMALL_NUMBER);
+		}
+	}
+}
+
+
 }
 
 
+FGeometryCollectionConvexUtility::FGeometryCollectionConvexData FGeometryCollectionConvexUtility::CreateNonOverlappingConvexHullData(FGeometryCollection* GeometryCollection, double FracAllowRemove)
+{
+	check(GeometryCollection);
+
+	// TODO: add fallback in case we don't have proximity data
+	bool bHasProximity = GeometryCollection->HasAttribute("Proximity", FGeometryCollection::GeometryGroup);
+	ensure(bHasProximity);
+	
+	const TManagedArray<TSet<int32>>& GCProximity = GeometryCollection->GetAttribute<TSet<int32>>("Proximity", FGeometryCollection::GeometryGroup);
+	TArray<TUniquePtr<Chaos::FConvex>> Convexes;
+	TArray<TSet<int32>> TransformToConvexIndexArr;
+	TArray<FTransform> GlobalTransformArray;
+	GeometryCollectionAlgo::GlobalMatrices(GeometryCollection->Transform, GeometryCollection->Parent, GlobalTransformArray);
+	HullsFromGeometry(*GeometryCollection, GlobalTransformArray, Convexes, TransformToConvexIndexArr, GeometryCollection->SimulationType, FGeometryCollection::ESimulationTypes::FST_Rigid);
+
+	CreateNonoverlappingConvexHulls(Convexes, TransformToConvexIndexArr, GeometryCollection->SimulationType, FGeometryCollection::ESimulationTypes::FST_Rigid, GeometryCollection->Parent, GCProximity, GeometryCollection->TransformIndex, FracAllowRemove);
+
+	TransformHullsToLocal(GlobalTransformArray, Convexes, TransformToConvexIndexArr);
+
+	if (!GeometryCollection->HasGroup("Convex"))
+	{
+		GeometryCollection->AddGroup("Convex");
+	}
+
+	if (!GeometryCollection->HasAttribute("TransformToConvexIndices", FTransformCollection::TransformGroup))
+	{
+		FManagedArrayCollection::FConstructionParameters ConvexDependency("Convex");
+		TManagedArray<TSet<int32>>& IndexArray = GeometryCollection->AddAttribute<TSet<int32>>("TransformToConvexIndices", FTransformCollection::TransformGroup, ConvexDependency);
+	}
+
+	if (!GeometryCollection->HasAttribute("ConvexHull", "Convex"))
+	{
+		GeometryCollection->AddAttribute<TUniquePtr<Chaos::FConvex>>("ConvexHull", "Convex");
+	}
+
+	TManagedArray<TSet<int32>>& TransformToConvexIndices = GeometryCollection->GetAttribute<TSet<int32>>("TransformToConvexIndices", FTransformCollection::TransformGroup);
+	TManagedArray<TUniquePtr<Chaos::FConvex>>& ConvexHull = GeometryCollection->GetAttribute<TUniquePtr<Chaos::FConvex>>("ConvexHull", "Convex");
+	TransformToConvexIndices = MoveTemp(TransformToConvexIndexArr);
+	ConvexHull = MoveTemp(Convexes);
+
+	return { TransformToConvexIndices, ConvexHull };
+}
+
+
+
 TUniquePtr<Chaos::FConvex> FGeometryCollectionConvexUtility::FindConvexHull(const FGeometryCollection* GeometryCollection, int32 GeometryIndex)
 {
 	check(GeometryCollection);
@@ -95,16 +720,16 @@ TUniquePtr<Chaos::FConvex> FGeometryCollectionConvexUtility::FindConvexHull(cons
 
 void FGeometryCollectionConvexUtility::RemoveConvexHulls(FGeometryCollection* GeometryCollection, const TArray<int32>& SortedTransformDeletes)
 {
-	if (GeometryCollection->HasGroup("Convex") && GeometryCollection->HasAttribute("TransformToConvexIndex", FTransformCollection::TransformGroup))
+	if (GeometryCollection->HasGroup("Convex") && GeometryCollection->HasAttribute("TransformToConvexIndices", FTransformCollection::TransformGroup))
 	{
-		TManagedArray<int32>& TransformToConvexIndex = GeometryCollection->GetAttribute<int32>("TransformToConvexIndex", FTransformCollection::TransformGroup);
+		TManagedArray<int32>& TransformToConvexIndices = GeometryCollection->GetAttribute<int32>("TransformToConvexIndices", FTransformCollection::TransformGroup);
 		TArray<int32> ConvexIndices;
 		for (int32 TransformIdx : SortedTransformDeletes)
 		{
-			if (TransformToConvexIndex[TransformIdx] > INDEX_NONE)
+			if (TransformToConvexIndices[TransformIdx] > INDEX_NONE)
 			{
-				ConvexIndices.Add(TransformToConvexIndex[TransformIdx]);
-				TransformToConvexIndex[TransformIdx] = INDEX_NONE;
+				ConvexIndices.Add(TransformToConvexIndices[TransformIdx]);
+				TransformToConvexIndices[TransformIdx] = INDEX_NONE;
 			}
 		}
 
@@ -121,13 +746,13 @@ void FGeometryCollectionConvexUtility::SetDefaults(FGeometryCollection* Geometry
 {
 	if (Group == FTransformCollection::TransformGroup)
 	{
-		if (GeometryCollection->HasAttribute("TransformToConvexIndex", FTransformCollection::TransformGroup))
+		if (GeometryCollection->HasAttribute("TransformToConvexIndices", FTransformCollection::TransformGroup))
 		{
-			TManagedArray<int32>& TransformToConvexIndex = GeometryCollection->GetAttribute<int32>("TransformToConvexIndex", FTransformCollection::TransformGroup);
+			TManagedArray<int32>& TransformToConvexIndices = GeometryCollection->GetAttribute<int32>("TransformToConvexIndices", FTransformCollection::TransformGroup);
 
 			for (uint32 Idx = StartSize; Idx < StartSize + NumElements; ++Idx)
 			{
-				TransformToConvexIndex[Idx] = INDEX_NONE;
+				TransformToConvexIndices[Idx] = INDEX_NONE;
 			}
 		}
 	}
diff --git a/Engine/Source/Runtime/Experimental/Chaos/Public/GeometryCollection/GeometryCollectionConvexUtility.h b/Engine/Source/Runtime/Experimental/Chaos/Public/GeometryCollection/GeometryCollectionConvexUtility.h
--- a/Engine/Source/Runtime/Experimental/Chaos/Public/GeometryCollection/GeometryCollectionConvexUtility.h
+++ b/Engine/Source/Runtime/Experimental/Chaos/Public/GeometryCollection/GeometryCollectionConvexUtility.h
@@ -18,13 +18,16 @@ public:
 
 	struct FGeometryCollectionConvexData
 	{
-		TManagedArray<int32>& TransformToConvexIndex;
+		TManagedArray<TSet<int32>>& TransformToConvexIndices;
 		TManagedArray<TUniquePtr<Chaos::FConvex>>& ConvexHull;
 	};
 
 	/** Ensure that convex hull data exists for the Geometry Collection and construct it if not (or if some data is missing. */
 	static FGeometryCollectionConvexData GetValidConvexHullData(FGeometryCollection* GeometryCollection);
 
+	/** Create non-overlapping convex hull data for all transforms in the geometry collection (except transforms where it would be better to just use the hulls of the children) */
+	static FGeometryCollectionConvexData CreateNonOverlappingConvexHullData(FGeometryCollection* GeometryCollection, double FractionAllowRemove = .3);
+
 	/** Returns the convex hull of the vertices contained in the specified geometry. */
 	static TUniquePtr<Chaos::FConvex> FindConvexHull(const FGeometryCollection* GeometryCollection, int32 GeometryIndex);
 
