Chaos: Fixed rewind data update accessing non-existent shapes when the shapes array is expanded at runtime (like when welding a shape when it begins simulating physics)

diff --git a/Engine/Source/Runtime/Experimental/Chaos/Private/RewindData.cpp b/Engine/Source/Runtime/Experimental/Chaos/Private/RewindData.cpp
--- a/Engine/Source/Runtime/Experimental/Chaos/Private/RewindData.cpp
+++ b/Engine/Source/Runtime/Experimental/Chaos/Private/RewindData.cpp
@@ -284,6 +284,7 @@ void FGeometryParticleStateBase::RecordPreDirtyData(const FDirtyProxy& Dirty, co
 	}
 
 	//shape properties
+	const int32 NumCurrentShapes = Handle->ShapesArray().Num();
 	for (int32 ShapeDataIdx : Dirty.ShapeDataIndices)
 	{
 		const FShapeDirtyData& ShapeData = ShapeDirtyData[ShapeDataIdx];
@@ -291,13 +292,28 @@ void FGeometryParticleStateBase::RecordPreDirtyData(const FDirtyProxy& Dirty, co
 		
 		FPerShapeDataStateBase& ShapeState = ShapesArrayState.FindOrAdd(ShapeIdx);
 
-		//TODO: why is this needed?
-		if (Handle->ShapesArray()[ShapeIdx] == nullptr) { continue; }
+		if(ShapeIdx >= NumCurrentShapes)
+		{
+			// We are updating the rewind data here before the handle was dirty but the dirty data may
+			// contain shapes we do not yet know about (if an actor had shapes added to it within the
+			// last update - this can happen when welding). If we encounter a shape we don't know about
+			// we skip the shape. Can't fully early out as the dirty indices may not be ordered.
+			// The actual shapes array in the handle will be updated after this step when we get into
+			// the proxy push function.
+			continue;
+		}
+
+		if (Handle->ShapesArray()[ShapeIdx] == nullptr)
+		{ 
+			//TODO: why is this needed?
+			continue;
+		}
 
 		if (ShapeData.IsDirty<EShapeProperty::CollisionData>())
 		{
 			ShapeState.CollisionData.Write(Handle->ShapesArray()[ShapeIdx]->GetCollisionData());
 		}
+
 		if (ShapeData.IsDirty<EShapeProperty::Materials>())
 		{
 			ShapeState.MaterialData.Write(Handle->ShapesArray()[ShapeIdx]->GetMaterialData());
