Chaos Cloth - Add material stiffness weight maps, make the stiffness value exponential, and tidy up solver files. [Merging CL 16355390 via CL 16369406]

CIS - Fix non unity build erroring on Chaos cloth constraints. [ Merging CL 16360199 via CL 16369406]

Chaos Cloth - Add ISPC weight maps implementation for material constraints. [Merging CL 16367050 via CL 16369406]

diff --git a/Engine/Plugins/Experimental/ChaosCloth/Source/Chaos/Private/ChaosCloth/ChaosClothConfig.cpp b/Engine/Plugins/Experimental/ChaosCloth/Source/Chaos/Private/ChaosCloth/ChaosClothConfig.cpp
--- a/Engine/Plugins/Experimental/ChaosCloth/Source/Chaos/Private/ChaosCloth/ChaosClothConfig.cpp
+++ b/Engine/Plugins/Experimental/ChaosCloth/Source/Chaos/Private/ChaosCloth/ChaosClothConfig.cpp
@@ -38,13 +38,13 @@ void UChaosClothConfig::MigrateFrom(const FClothConfig_Legacy& ClothConfig)
 	const float HorizontalStiffness =
 		ClothConfig.HorizontalConstraintConfig.Stiffness *
 		ClothConfig.HorizontalConstraintConfig.StiffnessMultiplier;
-	EdgeStiffness = FMath::Clamp((VerticalStiffness + HorizontalStiffness) * 0.5f, 0.f, 1.f);
+	EdgeStiffnessWeighted.Low = EdgeStiffnessWeighted.High = FMath::Clamp((VerticalStiffness + HorizontalStiffness) * 0.5f, 0.f, 1.f);
 
-	BendingStiffness = FMath::Clamp(
+	BendingStiffnessWeighted.Low = BendingStiffnessWeighted.High = FMath::Clamp(
 		ClothConfig.BendConstraintConfig.Stiffness *
 		ClothConfig.BendConstraintConfig.StiffnessMultiplier, 0.f, 1.f);
 
-	AreaStiffness = FMath::Clamp(
+	AreaStiffnessWeighted.Low = AreaStiffnessWeighted.High = FMath::Clamp(
 		ClothConfig.ShearConstraintConfig.Stiffness *
 		ClothConfig.ShearConstraintConfig.StiffnessMultiplier, 0.f, 1.f);
 
@@ -177,6 +177,13 @@ void UChaosClothConfig::PostLoad()
 		Drag.Low = Drag.High = DragCoefficient_DEPRECATED;
 		Lift.Low = Lift.High = LiftCoefficient_DEPRECATED;
 	}
+
+	if (FortniteMainBranchObjectVersion < FFortniteMainBranchObjectVersion::ChaosClothAddMaterialWeightMaps)
+	{
+		EdgeStiffnessWeighted.Low = EdgeStiffnessWeighted.High = EdgeStiffness_DEPRECATED;
+		BendingStiffnessWeighted.Low = BendingStiffnessWeighted.High = BendingStiffness_DEPRECATED;
+		AreaStiffnessWeighted.Low = AreaStiffnessWeighted.High = AreaStiffness_DEPRECATED;
+	}
 #endif  // #if WITH_EDITORONLY_DATA
 }
 
diff --git a/Engine/Plugins/Experimental/ChaosCloth/Source/Chaos/Private/ChaosCloth/ChaosClothConstraints.cpp b/Engine/Plugins/Experimental/ChaosCloth/Source/Chaos/Private/ChaosCloth/ChaosClothConstraints.cpp
--- a/Engine/Plugins/Experimental/ChaosCloth/Source/Chaos/Private/ChaosCloth/ChaosClothConstraints.cpp
+++ b/Engine/Plugins/Experimental/ChaosCloth/Source/Chaos/Private/ChaosCloth/ChaosClothConstraints.cpp
@@ -85,9 +85,10 @@ void FClothConstraints::CreateRules()
 	if (XEdgeConstraints)
 	{
 		ConstraintInits[ConstraintInitIndex++] =
-			[this](const FPBDParticles& /*Particles*/, const FReal /*Dt*/)
+			[this](const FPBDParticles& /*Particles*/, const FReal Dt)
 			{
 				XEdgeConstraints->Init();
+				XEdgeConstraints->ApplyProperties(Dt, Evolution->GetIterations());
 			};
 
 		ConstraintRules[ConstraintRuleIndex++] = 
@@ -98,6 +99,11 @@ void FClothConstraints::CreateRules()
 	}
 	if (EdgeConstraints)
 	{
+		ConstraintInits[ConstraintInitIndex++] =
+			[this](const FPBDParticles& /*Particles*/, const FReal Dt)
+			{
+				EdgeConstraints->ApplyProperties(Dt, Evolution->GetIterations());
+			};
 		ConstraintRules[ConstraintRuleIndex++] =
 			[this](FPBDParticles& Particles, const FReal Dt)
 			{
@@ -107,9 +113,10 @@ void FClothConstraints::CreateRules()
 	if (XBendingConstraints)
 	{
 		ConstraintInits[ConstraintInitIndex++] =
-			[this](const FPBDParticles& /*Particles*/, const FReal /*Dt*/)
+			[this](const FPBDParticles& /*Particles*/, const FReal Dt)
 			{
 				XBendingConstraints->Init();
+				XBendingConstraints->ApplyProperties(Dt, Evolution->GetIterations());
 			};
 		ConstraintRules[ConstraintRuleIndex++] =
 			[this](FPBDParticles& Particles, const FReal Dt)
@@ -119,6 +126,11 @@ void FClothConstraints::CreateRules()
 	}
 	if (BendingConstraints)
 	{
+		ConstraintInits[ConstraintInitIndex++] =
+			[this](const FPBDParticles& /*Particles*/, const FReal Dt)
+			{
+				BendingConstraints->ApplyProperties(Dt, Evolution->GetIterations());
+			};
 		ConstraintRules[ConstraintRuleIndex++] =
 			[this](FPBDParticles& Particles, const FReal Dt)
 			{
@@ -136,9 +148,10 @@ void FClothConstraints::CreateRules()
 	if (XAreaConstraints)
 	{
 		ConstraintInits[ConstraintInitIndex++] =
-			[this](const FPBDParticles& /*Particles*/, const FReal /*Dt*/)
+			[this](const FPBDParticles& /*Particles*/, const FReal Dt)
 			{
 				XAreaConstraints->Init();
+				XAreaConstraints->ApplyProperties(Dt, Evolution->GetIterations());
 			};
 		ConstraintRules[ConstraintRuleIndex++] =
 			[this](FPBDParticles& Particles, const FReal Dt)
@@ -148,6 +161,11 @@ void FClothConstraints::CreateRules()
 	}
 	if (AreaConstraints)
 	{
+		ConstraintInits[ConstraintInitIndex++] =
+			[this](const FPBDParticles& /*Particles*/, const FReal Dt)
+			{
+				AreaConstraints->ApplyProperties(Dt, Evolution->GetIterations());
+			};
 		ConstraintRules[ConstraintRuleIndex++] =
 			[this](FPBDParticles& Particles, const FReal Dt)
 			{
@@ -252,71 +270,114 @@ void FClothConstraints::CreateRules()
 	check(ConstraintRuleIndex == NumConstraintRules);
 }
 
-void FClothConstraints::SetEdgeConstraints(const TArray<TVec3<int32>>& SurfaceElements, FReal EdgeStiffness, bool bUseXPBDConstraints)
+void FClothConstraints::SetEdgeConstraints(const TArray<TVec3<int32>>& SurfaceElements, const TConstArrayView<FRealSingle>& StiffnessMultipliers, bool bUseXPBDConstraints)
 {
 	check(Evolution);
-	check(EdgeStiffness > 0.f && EdgeStiffness <= 1.f);
 
 	if (bUseXPBDConstraints)
 	{
-		XEdgeConstraints = MakeShared<FXPBDSpringConstraints>(Evolution->Particles(), SurfaceElements, EdgeStiffness, /*bStripKinematicConstraints =*/ true);
-		++NumConstraintInits;
+		XEdgeConstraints = MakeShared<FXPBDSpringConstraints>(
+			Evolution->Particles(),
+			ParticleOffset, NumParticles,
+			SurfaceElements,
+			StiffnessMultipliers,
+			/*InStiffness =*/ FVec2::UnitVector,
+			/*bTrimKinematicConstraints =*/ true);
 	}
 	else
 	{
-		EdgeConstraints = MakeShared<FPBDSpringConstraints>(Evolution->Particles(), SurfaceElements, EdgeStiffness, /*bStripKinematicConstraints =*/ true);
+		EdgeConstraints = MakeShared<FPBDSpringConstraints>(
+			Evolution->Particles(),
+			ParticleOffset,
+			NumParticles,
+			SurfaceElements,
+			StiffnessMultipliers,
+			/*InStiffness =*/ FVec2::UnitVector,
+			/*bTrimKinematicConstraints =*/ true);
 	}
+	++NumConstraintInits;  // Uses init to update the property tables
 	++NumConstraintRules;
 }
 
-void FClothConstraints::SetBendingConstraints(TArray<TVec2<int32>>&& Edges, FReal BendingStiffness, bool bUseXPBDConstraints)
+void FClothConstraints::SetBendingConstraints(const TArray<TVec2<int32>>& Edges, const TConstArrayView<FRealSingle>& StiffnessMultipliers, bool bUseXPBDConstraints)
 {
 	check(Evolution);
 
 	if (bUseXPBDConstraints)
 	{
-		XBendingConstraints = MakeShared<FXPBDSpringConstraints>(Evolution->Particles(), MoveTemp(Edges), BendingStiffness, /*bStripKinematicConstraints =*/ true);
-		++NumConstraintInits;
+		XBendingConstraints = MakeShared<FXPBDSpringConstraints>(
+			Evolution->Particles(),
+			ParticleOffset, NumParticles,
+			Edges,
+			StiffnessMultipliers,
+			/*InStiffness =*/ FVec2::UnitVector,
+			/*bTrimKinematicConstraints =*/ true);
 	}
 	else
 	{
-		BendingConstraints = MakeShared<FPBDSpringConstraints>(Evolution->Particles(), MoveTemp(Edges), BendingStiffness, /*bStripKinematicConstraints =*/ true);
+		BendingConstraints = MakeShared<FPBDSpringConstraints>(
+			Evolution->Particles(),
+			ParticleOffset,
+			NumParticles,
+			Edges,
+			StiffnessMultipliers,
+			/*InStiffness =*/ FVec2::UnitVector,
+			/*bTrimKinematicConstraints =*/ true);
 	}
+	++NumConstraintInits;  // Uses init to update the property tables
 	++NumConstraintRules;
 }
 
 void FClothConstraints::SetBendingConstraints(TArray<TVec4<int32>>&& BendingElements, FReal BendingStiffness)
 {
 	check(Evolution);
-	check(BendingStiffness > 0.f && BendingStiffness <= 1.f);
 
 	BendingElementConstraints = MakeShared<FPBDBendingConstraints>(Evolution->Particles(), MoveTemp(BendingElements), BendingStiffness);  // TODO: Strip kinematic constraints
 	++NumConstraintRules;
 }
 
-void FClothConstraints::SetAreaConstraints(TArray<TVec3<int32>>&& SurfaceElements, FReal AreaStiffness, bool bUseXPBDConstraints)
+void FClothConstraints::SetAreaConstraints(const TArray<TVec3<int32>>& SurfaceElements, const TConstArrayView<FRealSingle>& StiffnessMultipliers, bool bUseXPBDConstraints)
 {
 	check(Evolution);
-	check(AreaStiffness > 0.f && AreaStiffness <= 1.f);
 
 	if (bUseXPBDConstraints)
 	{
-		XAreaConstraints = MakeShared<FXPBDAxialSpringConstraints>(Evolution->Particles(), MoveTemp(SurfaceElements), AreaStiffness);
-		++NumConstraintInits;
+		XAreaConstraints = MakeShared<FXPBDAxialSpringConstraints>(
+			Evolution->Particles(),
+			ParticleOffset,
+			NumParticles,
+			SurfaceElements,
+			StiffnessMultipliers,
+			/*InStiffness =*/ FVec2::UnitVector,
+			/*bTrimKinematicConstraints =*/ true);
 	}
 	else
 	{
-		AreaConstraints = MakeShared<FPBDAxialSpringConstraints>(Evolution->Particles(), MoveTemp(SurfaceElements), AreaStiffness);
+		AreaConstraints = MakeShared<FPBDAxialSpringConstraints>(
+			Evolution->Particles(),
+			ParticleOffset,
+			NumParticles,
+			SurfaceElements,
+			StiffnessMultipliers,
+			/*InStiffness =*/ FVec2::UnitVector,
+			/*bTrimKinematicConstraints =*/ true);
 	}
+	++NumConstraintInits;  // Uses init to update the property tables
 	++NumConstraintRules;
 }
 
-void FClothConstraints::SetVolumeConstraints(TArray<TVec2<int32>>&& DoubleBendingEdges, FReal VolumeStiffness)
+void FClothConstraints::SetVolumeConstraints(const TArray<TVec2<int32>>& DoubleBendingEdges, FReal VolumeStiffness)
 {
 	check(Evolution);
-	check(VolumeStiffness > 0.f && VolumeStiffness <= 1.f);
 
-	ThinShellVolumeConstraints = MakeShared<FPBDSpringConstraints>(Evolution->Particles(), MoveTemp(DoubleBendingEdges), VolumeStiffness);
+	ThinShellVolumeConstraints = MakeShared<FPBDSpringConstraints>(
+		Evolution->Particles(),
+		ParticleOffset,
+		NumParticles,
+		DoubleBendingEdges,
+		TConstArrayView<FRealSingle>(),
+		VolumeStiffness,
+		/*bTrimKinematicConstraints =*/ true);
 	++NumConstraintRules;
 }
 
@@ -331,8 +392,7 @@ void FClothConstraints::SetVolumeConstraints(TArray<TVec3<int32>>&& SurfaceEleme
 
 void FClothConstraints::SetLongRangeConstraints(const TMap<int32, TSet<int32>>& PointToNeighborsMap,
 	const TConstArrayView<FRealSingle>& TetherStiffnessMultipliers, const TConstArrayView<FRealSingle>& TetherScaleMultipliers,
-	const FVec2& TetherStiffness, const FVec2& TetherScale,
-	ETetherMode TetherMode, bool bUseXPBDConstraints)
+	const FVec2& TetherScale, ETetherMode TetherMode, bool bUseXPBDConstraints)
 {
 	check(Evolution);
 
@@ -348,7 +408,7 @@ void FClothConstraints::SetLongRangeConstraints(const TMap<int32, TSet<int32>>&
 			TetherStiffnessMultipliers,
 			TetherScaleMultipliers,
 			MaxNumTetherIslands,
-			TetherStiffness,
+			/*InStiffness =*/ FVec2::UnitVector,
 			TetherScale,
 			TetherMode);
 	}
@@ -362,7 +422,7 @@ void FClothConstraints::SetLongRangeConstraints(const TMap<int32, TSet<int32>>&
 			TetherStiffnessMultipliers,
 			TetherScaleMultipliers,
 			MaxNumTetherIslands,
-			TetherStiffness,
+			/*InStiffness =*/ FVec2::UnitVector,
 			TetherScale,
 			TetherMode);
 	}
@@ -433,39 +493,39 @@ void FClothConstraints::SetSelfCollisionConstraints(const TArray<TVec3<int32>>&
 	++NumConstraintRules;  // and a rule
 }
 
-void FClothConstraints::SetEdgeProperties(FReal EdgeStiffness)
+void FClothConstraints::SetEdgeProperties(const FVec2& EdgeStiffness)
 {
 	if (EdgeConstraints)
 	{
-		EdgeConstraints->SetStiffness(EdgeStiffness);
+		EdgeConstraints->SetProperties(EdgeStiffness);
 	}
 	if (XEdgeConstraints)
 	{
-		XEdgeConstraints->SetStiffness(EdgeStiffness);
+		XEdgeConstraints->SetProperties(EdgeStiffness);
 	}
 }
 
-void FClothConstraints::SetBendingProperties(FReal BendingStiffness)
+void FClothConstraints::SetBendingProperties(const FVec2& BendingStiffness)
 {
 	if (BendingConstraints)
 	{
-		BendingConstraints->SetStiffness(BendingStiffness);
+		BendingConstraints->SetProperties(BendingStiffness);
 	}
 	if (XBendingConstraints)
 	{
-		XBendingConstraints->SetStiffness(BendingStiffness);
+		XBendingConstraints->SetProperties(BendingStiffness);
 	}
 }
 
-void FClothConstraints::SetAreaProperties(FReal AreaStiffness)
+void FClothConstraints::SetAreaProperties(const FVec2& AreaStiffness)
 {
 	if (AreaConstraints)
 	{
-		AreaConstraints->SetStiffness(AreaStiffness);
+		AreaConstraints->SetProperties(AreaStiffness);
 	}
 	if (XAreaConstraints)
 	{
-		XAreaConstraints->SetStiffness(AreaStiffness);
+		XAreaConstraints->SetProperties(AreaStiffness);
 	}
 }
 
@@ -473,7 +533,7 @@ void FClothConstraints::SetThinShellVolumeProperties(FReal VolumeStiffness)
 {
 	if (ThinShellVolumeConstraints)
 	{
-		ThinShellVolumeConstraints->SetStiffness(VolumeStiffness);
+		ThinShellVolumeConstraints->SetProperties(VolumeStiffness);
 	}
 }
 
diff --git a/Engine/Plugins/Experimental/ChaosCloth/Source/Chaos/Private/ChaosCloth/ChaosClothingSimulation.cpp b/Engine/Plugins/Experimental/ChaosCloth/Source/Chaos/Private/ChaosCloth/ChaosClothingSimulation.cpp
--- a/Engine/Plugins/Experimental/ChaosCloth/Source/Chaos/Private/ChaosCloth/ChaosClothingSimulation.cpp
+++ b/Engine/Plugins/Experimental/ChaosCloth/Source/Chaos/Private/ChaosCloth/ChaosClothingSimulation.cpp
@@ -378,10 +378,10 @@ void FClothingSimulation::CreateActor(USkeletalMeshComponent* InOwnerComponent,
 		(FClothingSimulationCloth::EMassMode)ClothConfig->MassMode,
 		ClothConfig->GetMassValue(),
 		ClothConfig->MinPerParticleMass,
-		ClothConfig->EdgeStiffness,
-		ClothConfig->BendingStiffness,
+		TVec2<FRealSingle>(ClothConfig->EdgeStiffnessWeighted.Low, ClothConfig->EdgeStiffnessWeighted.High),
+		TVec2<FRealSingle>(ClothConfig->BendingStiffnessWeighted.Low, ClothConfig->BendingStiffnessWeighted.High),
 		ClothConfig->bUseBendingElements,
-		ClothConfig->AreaStiffness,
+		TVec2<FRealSingle>(ClothConfig->AreaStiffnessWeighted.Low, ClothConfig->AreaStiffnessWeighted.High),
 		ClothConfig->VolumeStiffness,
 		ClothConfig->bUseThinShellVolumeConstraints,
 		TVec2<FRealSingle>(ClothConfig->TetherStiffness.Low, ClothConfig->TetherStiffness.High),  // Animatable
@@ -739,10 +739,10 @@ void FClothingSimulation::RefreshClothConfig(const IClothingSimulationContext* I
 			(FClothingSimulationCloth::EMassMode)ClothConfig->MassMode,
 			ClothConfig->GetMassValue(),
 			ClothConfig->MinPerParticleMass,
-			ClothConfig->EdgeStiffness,
-			ClothConfig->BendingStiffness,
+			TVec2<FRealSingle>(ClothConfig->EdgeStiffnessWeighted.Low, ClothConfig->EdgeStiffnessWeighted.High),
+			TVec2<FRealSingle>(ClothConfig->BendingStiffnessWeighted.Low, ClothConfig->BendingStiffnessWeighted.High),
 			ClothConfig->bUseBendingElements,
-			ClothConfig->AreaStiffness,
+			TVec2<FRealSingle>(ClothConfig->AreaStiffnessWeighted.Low, ClothConfig->AreaStiffnessWeighted.High),
 			ClothConfig->VolumeStiffness,
 			ClothConfig->bUseThinShellVolumeConstraints,
 			TVec2<FRealSingle>(ClothConfig->TetherStiffness.Low, ClothConfig->TetherStiffness.High),  // Animatable
diff --git a/Engine/Plugins/Experimental/ChaosCloth/Source/Chaos/Private/ChaosCloth/ChaosClothingSimulationCloth.cpp b/Engine/Plugins/Experimental/ChaosCloth/Source/Chaos/Private/ChaosCloth/ChaosClothingSimulationCloth.cpp
--- a/Engine/Plugins/Experimental/ChaosCloth/Source/Chaos/Private/ChaosCloth/ChaosClothingSimulationCloth.cpp
+++ b/Engine/Plugins/Experimental/ChaosCloth/Source/Chaos/Private/ChaosCloth/ChaosClothingSimulationCloth.cpp
@@ -118,31 +118,33 @@ void FClothingSimulationCloth::FLODData::Add(FClothingSimulationSolver* Solver,
 	}
 
 	// Edge constraints
-	if (Cloth->EdgeStiffness)
+	const TConstArrayView<FRealSingle>& EdgeStiffnessMultipliers = WeightMaps[(int32)EChaosWeightMapTarget::EdgeStiffness];
+	if (Cloth->EdgeStiffness[0] > (FReal)0. || (Cloth->EdgeStiffness[1] > (FReal)0. && EdgeStiffnessMultipliers.Num()))
 	{
-		ClothConstraints.SetEdgeConstraints(SurfaceElements, (FReal)Cloth->EdgeStiffness, Cloth->bUseXPBDConstraints);
+		ClothConstraints.SetEdgeConstraints(SurfaceElements, EdgeStiffnessMultipliers, Cloth->bUseXPBDConstraints);
 	}
 
 	// Bending constraints
-	if (Cloth->BendingStiffness > 0.f)
+	const TConstArrayView<FRealSingle>& BendingStiffnessMultipliers = WeightMaps[(int32)EChaosWeightMapTarget::BendingStiffness];
+	if (Cloth->BendingStiffness[0] > 0.f || (Cloth->BendingStiffness[1] > (FReal)0. && BendingStiffnessMultipliers.Num()))
 	{
 		if (Cloth->bUseBendingElements)
 		{
 			TArray<Chaos::TVec4<int32>> BendingElements = TriangleMesh.GetUniqueAdjacentElements();
-			ClothConstraints.SetBendingConstraints(MoveTemp(BendingElements), (FReal)Cloth->BendingStiffness);
+			ClothConstraints.SetBendingConstraints(MoveTemp(BendingElements), Cloth->BendingStiffness[0]);  // TODO: Add support for weight maps
 		}
 		else
 		{
-			TArray<Chaos::TVec2<int32>> Edges = TriangleMesh.GetUniqueAdjacentPoints();
-			ClothConstraints.SetBendingConstraints(MoveTemp(Edges), (FReal)Cloth->BendingStiffness, Cloth->bUseXPBDConstraints);
+			const TArray<Chaos::TVec2<int32>> Edges = TriangleMesh.GetUniqueAdjacentPoints();
+			ClothConstraints.SetBendingConstraints(Edges, BendingStiffnessMultipliers, Cloth->bUseXPBDConstraints);
 		}
 	}
 
 	// Area constraints
-	if (Cloth->AreaStiffness)
+	const TConstArrayView<FRealSingle>& AreaStiffnessMultipliers = WeightMaps[(int32)EChaosWeightMapTarget::AreaStiffness];
+	if (Cloth->AreaStiffness[0] > 0.f || (Cloth->AreaStiffness[1] > (FReal)0. && AreaStiffnessMultipliers.Num()))
 	{
-		TArray<Chaos::TVec3<int32>> SurfaceConstraints = SurfaceElements;
-		ClothConstraints.SetAreaConstraints(MoveTemp(SurfaceConstraints), (FReal)Cloth->AreaStiffness, Cloth->bUseXPBDConstraints);
+		ClothConstraints.SetAreaConstraints(SurfaceElements, AreaStiffnessMultipliers, Cloth->bUseXPBDConstraints);
 	}
 
 	// Volume constraints
@@ -180,7 +182,7 @@ void FClothingSimulationCloth::FLODData::Add(FClothingSimulationSolver* Solver,
 				}
 			}
 
-			ClothConstraints.SetVolumeConstraints(MoveTemp(DoubleBendingConstraints), (FReal)Cloth->VolumeStiffness);
+			ClothConstraints.SetVolumeConstraints(DoubleBendingConstraints, (FReal)Cloth->VolumeStiffness);
 		}
 		else
 		{
@@ -200,7 +202,6 @@ void FClothingSimulationCloth::FLODData::Add(FClothingSimulationSolver* Solver,
 			PointToNeighborsMap,
 			TetherStiffnessMultipliers,
 			TetherScaleMultipliers,
-			FVec2((FReal)Cloth->TetherStiffness[0], (FReal)Cloth->TetherStiffness[1]),
 			FVec2((FReal)Cloth->TetherScale[0], (FReal)Cloth->TetherScale[1]),
 			Cloth->TetherMode,
 			Cloth->bUseXPBDConstraints);
@@ -270,9 +271,9 @@ void FClothingSimulationCloth::FLODData::Update(FClothingSimulationSolver* Solve
 	// Update the animatable constraint parameters
 	FClothConstraints& ClothConstraints = Solver->GetClothConstraints(Offset);
 	ClothConstraints.SetMaximumDistanceProperties((FReal)Cloth->MaxDistancesMultiplier);
-	ClothConstraints.SetEdgeProperties((FReal)Cloth->EdgeStiffness);
-	ClothConstraints.SetBendingProperties((FReal)Cloth->BendingStiffness);
-	ClothConstraints.SetAreaProperties((FReal)Cloth->AreaStiffness);
+	ClothConstraints.SetEdgeProperties(FVec2((FReal)Cloth->EdgeStiffness[0], (FReal)Cloth->EdgeStiffness[1]));
+	ClothConstraints.SetBendingProperties(FVec2((FReal)Cloth->BendingStiffness[0], (FReal)Cloth->BendingStiffness[1]));
+	ClothConstraints.SetAreaProperties(FVec2((FReal)Cloth->AreaStiffness[0], (FReal)Cloth->AreaStiffness[1]));
 	ClothConstraints.SetLongRangeAttachmentProperties(FVec2((FReal)Cloth->TetherStiffness[0], (FReal)Cloth->TetherStiffness[1]));
 	ClothConstraints.SetSelfCollisionProperties((FReal)Cloth->SelfCollisionThickness);
 	ClothConstraints.SetAnimDriveProperties(
@@ -336,10 +337,10 @@ FClothingSimulationCloth::FClothingSimulationCloth(
 	EMassMode InMassMode,
 	FRealSingle InMassValue,
 	FRealSingle InMinPerParticleMass,
-	FRealSingle InEdgeStiffness,
-	FRealSingle InBendingStiffness,
+	const TVec2<FRealSingle>& InEdgeStiffness,
+	const TVec2<FRealSingle>& InBendingStiffness,
 	bool bInUseBendingElements,
-	FRealSingle InAreaStiffness,
+	const TVec2<FRealSingle>& InAreaStiffness,
 	FRealSingle InVolumeStiffness,
 	bool bInUseThinShellVolumeConstraints,
 	const TVec2<FRealSingle>& InTetherStiffness,
diff --git a/Engine/Plugins/Experimental/ChaosCloth/Source/Chaos/Private/ChaosCloth/ChaosClothingSimulationInteractor.cpp b/Engine/Plugins/Experimental/ChaosCloth/Source/Chaos/Private/ChaosCloth/ChaosClothingSimulationInteractor.cpp
--- a/Engine/Plugins/Experimental/ChaosCloth/Source/Chaos/Private/ChaosCloth/ChaosClothingSimulationInteractor.cpp
+++ b/Engine/Plugins/Experimental/ChaosCloth/Source/Chaos/Private/ChaosCloth/ChaosClothingSimulationInteractor.cpp
@@ -6,6 +6,11 @@
 
 using namespace Chaos;
 
+namespace ChaosClothingInteractor
+{
+	static const float InvStiffnessLogBase = 1.f / FMath::Loge(1.e3f);  // Log base for updating old linear stiffnesses to the new stiffness exponentiation
+}
+
 void UChaosClothingInteractor::Sync(IClothingSimulation* Simulation)
 {
 	check(Simulation);
@@ -23,7 +28,19 @@ void UChaosClothingInteractor::Sync(IClothingSimulation* Simulation)
 	UClothingInteractor::Sync(Simulation);
 }
 
-void UChaosClothingInteractor::SetMaterialLinear(float EdgeStiffness, float BendingStiffness, float AreaStiffness)
+void UChaosClothingInteractor::SetMaterialLinear(float EdgeStiffnessLinear, float BendingStiffnessLinear, float AreaStiffnessLinear)
+{
+	const TVec2<FRealSingle> EdgeStiffness((FMath::Clamp(FMath::Loge(EdgeStiffnessLinear) * ChaosClothingInteractor::InvStiffnessLogBase + 1.f, 0.f, 1.f)), 1.f);
+	const TVec2<FRealSingle> BendingStiffness((FMath::Clamp(FMath::Loge(BendingStiffnessLinear) * ChaosClothingInteractor::InvStiffnessLogBase + 1.f, 0.f, 1.f)), 1.f);
+	const TVec2<FRealSingle> AreaStiffness((FMath::Clamp(FMath::Loge(AreaStiffnessLinear) * ChaosClothingInteractor::InvStiffnessLogBase + 1.f, 0.f, 1.f)), 1.f);
+
+	Commands.Add(FChaosClothingInteractorCommand::CreateLambda([EdgeStiffness, BendingStiffness, AreaStiffness](FClothingSimulationCloth* Cloth)
+	{
+		Cloth->SetMaterialProperties(EdgeStiffness, BendingStiffness, AreaStiffness);
+	}));
+}
+
+void UChaosClothingInteractor::SetMaterial(FVector2D EdgeStiffness, FVector2D BendingStiffness, FVector2D AreaStiffness)
 {
 	Commands.Add(FChaosClothingInteractorCommand::CreateLambda([EdgeStiffness, BendingStiffness, AreaStiffness](FClothingSimulationCloth* Cloth)
 	{
@@ -34,9 +51,9 @@ void UChaosClothingInteractor::SetMaterialLinear(float EdgeStiffness, float Bend
 void UChaosClothingInteractor::SetLongRangeAttachmentLinear(float TetherStiffnessLinear)
 {
 	// Deprecated
-	Commands.Add(FChaosClothingInteractorCommand::CreateLambda([TetherStiffnessLinear](FClothingSimulationCloth* Cloth)
+	const TVec2<FRealSingle> TetherStiffness((FMath::Clamp(FMath::Loge(TetherStiffnessLinear) * ChaosClothingInteractor::InvStiffnessLogBase + 1.f, 0.f, 1.f)), 1.f);
+	Commands.Add(FChaosClothingInteractorCommand::CreateLambda([TetherStiffness](FClothingSimulationCloth* Cloth)
 	{
-		const TVec2<FRealSingle> TetherStiffness((FMath::Clamp(FMath::Loge(TetherStiffnessLinear) / FMath::Loge(1.e3f) + 1.f, 0.f, 1.f)), 1.f);
 		Cloth->SetLongRangeAttachmentProperties(TetherStiffness);
 	}));
 }
@@ -45,7 +62,7 @@ void UChaosClothingInteractor::SetLongRangeAttachment(FVector2D TetherStiffness)
 {
 	Commands.Add(FChaosClothingInteractorCommand::CreateLambda([TetherStiffness](FClothingSimulationCloth* Cloth)
 	{
-		Cloth->SetLongRangeAttachmentProperties(TVec2<FRealSingle>(TetherStiffness.X, TetherStiffness.Y));
+		Cloth->SetLongRangeAttachmentProperties(TVec2<FRealSingle>(TetherStiffness[0], TetherStiffness[1]));
 	}));
 }
 
@@ -93,11 +110,11 @@ void UChaosClothingInteractor::SetGravity(float GravityScale, bool bIsGravityOve
 void UChaosClothingInteractor::SetAnimDriveLinear(float AnimDriveStiffnessLinear)
 {
 	// Deprecated
-	Commands.Add(FChaosClothingInteractorCommand::CreateLambda([AnimDriveStiffnessLinear](FClothingSimulationCloth* Cloth)
+	const TVec2<FRealSingle> AnimDriveStiffness(0.f, FMath::Clamp(FMath::Loge(AnimDriveStiffnessLinear) * ChaosClothingInteractor::InvStiffnessLogBase + 1.f, 0.f, 1.f));
+	Commands.Add(FChaosClothingInteractorCommand::CreateLambda([AnimDriveStiffness](FClothingSimulationCloth* Cloth)
 	{
 		// The Anim Drive stiffness Low value needs to be 0 in order to keep backward compatibility with existing mask (this wouldn't be an issue if this property had no legacy mask)
-		const TVec2<FRealSingle> AnimDriveStiffness(0.f, FMath::Clamp(FMath::Loge(AnimDriveStiffnessLinear) / FMath::Loge(1.e3f) + 1.f, 0.f, 1.f));
-		const TVec2<FRealSingle> AnimDriveDamping(0.f, 1.f);
+		static const TVec2<FRealSingle> AnimDriveDamping(0.f, 1.f);
 		Cloth->SetAnimDriveProperties(AnimDriveStiffness, AnimDriveDamping);
 	}));
 }
diff --git a/Engine/Plugins/Experimental/ChaosCloth/Source/Chaos/Private/ChaosCloth/ChaosWeightMapTarget.h b/Engine/Plugins/Experimental/ChaosCloth/Source/Chaos/Private/ChaosCloth/ChaosWeightMapTarget.h
--- a/Engine/Plugins/Experimental/ChaosCloth/Source/Chaos/Private/ChaosCloth/ChaosWeightMapTarget.h
+++ b/Engine/Plugins/Experimental/ChaosCloth/Source/Chaos/Private/ChaosCloth/ChaosWeightMapTarget.h
@@ -19,5 +19,8 @@ enum class EChaosWeightMapTarget : uint8
 	TetherStiffness,
 	TetherScale,
 	Drag,
-	Lift
+	Lift,
+	EdgeStiffness,
+	BendingStiffness,
+	AreaStiffness
 };
diff --git a/Engine/Plugins/Experimental/ChaosCloth/Source/Chaos/Public/ChaosCloth/ChaosClothConfig.h b/Engine/Plugins/Experimental/ChaosCloth/Source/Chaos/Public/ChaosCloth/ChaosClothConfig.h
--- a/Engine/Plugins/Experimental/ChaosCloth/Source/Chaos/Public/ChaosCloth/ChaosClothConfig.h
+++ b/Engine/Plugins/Experimental/ChaosCloth/Source/Chaos/Public/ChaosCloth/ChaosClothConfig.h
@@ -107,21 +107,36 @@ public:
 	UPROPERTY()
 	float MinPerParticleMass = 0.0001f;
 
-	/** The Stiffness of the Edge constraints, only use lower than 1 values for very stretchy materials. Increase the iteration count for stiffer materials. */
-	UPROPERTY(EditAnywhere, Category = "Material Properties", meta = (UIMin = "0", UIMax = "1", ClampMin = "0", ClampMax = "1"))
-	float EdgeStiffness = 1.f;
+	/**
+	 * The Stiffness of segments constraints. Increase the iteration count for stiffer materials.
+	 * If an enabled Weight Map (Mask with values in the range [0;1]) targeting the "Edge Stiffness" is added to the cloth, 
+	 * then both the Low and High values will be used in conjunction with the per particle Weight stored in the Weight Map to interpolate the final value from them.
+	 * Otherwise only the Low value is meaningful and sufficient to enable this constraint.
+	 */
+	UPROPERTY(EditAnywhere, Category = "Material Properties", DisplayName = "Edge Stiffness", meta = (UIMin = "0", UIMax = "1", ClampMin = "0", ClampMax = "1"))
+	FChaosClothWeightedValue EdgeStiffnessWeighted = { 1.f, 1.f };
 
-	/** The Stiffness of the bending constraints. Increase the iteration count for stiffer materials. Increase the iteration count for stiffer materials. */
-	UPROPERTY(EditAnywhere, Category = "Material Properties", meta = (UIMin = "0", UIMax = "1", ClampMin = "0", ClampMax = "1"))
-	float BendingStiffness = 1.f;
+	/**
+	 * The Stiffness of cross segments and bending elements constraints. Increase the iteration count for stiffer materials.
+	 * If an enabled Weight Map (Mask with values in the range [0;1]) targeting the "Bend Stiffness" is added to the cloth, 
+	 * then both the Low and High values will be used in conjunction with the per particle Weight stored in the Weight Map to interpolate the final value from them.
+	 * Otherwise only the Low value is meaningful and sufficient to enable this constraint.
+	 */
+	UPROPERTY(EditAnywhere, Category = "Material Properties", DisplayName = "Bending Stiffness", meta = (UIMin = "0", UIMax = "1", ClampMin = "0", ClampMax = "1"))
+	FChaosClothWeightedValue BendingStiffnessWeighted = { 1.f, 1.f };
 
 	/** Enable the more accurate bending element constraints instead of the faster cross-edge spring constraints used for controlling bending stiffness. */
 	UPROPERTY(EditAnywhere, Category = "Material Properties")
 	bool bUseBendingElements = false;
 
-	/** The stiffness of the area preservation constraints. Increase the iteration count for stiffer materials. */
-	UPROPERTY(EditAnywhere, Category = "Material Properties", meta = (UIMin = "0", UIMax = "1", ClampMin = "0", ClampMax = "1"))
-	float AreaStiffness = 1.f;
+	/**
+	 * The stiffness of the surface area preservation constraints. Increase the iteration count for stiffer materials.
+	 * If an enabled Weight Map (Mask with values in the range [0;1]) targeting the "Bend Stiffness" is added to the cloth, 
+	 * then both the Low and High values will be used in conjunction with the per particle Weight stored in the Weight Map to interpolate the final value from them.
+	 * Otherwise only the Low value is meaningful and sufficient to enable this constraint.
+	 */
+	UPROPERTY(EditAnywhere, Category = "Material Properties", DisplayName = "Area Stiffness", meta = (UIMin = "0", UIMax = "1", ClampMin = "0", ClampMax = "1"))
+	FChaosClothWeightedValue AreaStiffnessWeighted = { 1.f, 1.f };
 
 	/** The stiffness of the volume preservation constraints. */
 	UPROPERTY()
@@ -132,7 +147,7 @@ public:
 	 * The long range attachment connects each of the cloth particles to its closest fixed point with a spring constraint.
 	 * This can be used to compensate for a lack of stretch resistance when the iterations count is kept low for performance reasons.
 	 * Can lead to an unnatural pull string puppet like behavior.
-	 * If an enabled Weight Map (A.K.A. Mask) targeting the "Tether Stiffness" is added to the cloth, 
+	 * If an enabled Weight Map (Mask with values in the range [0;1]) targeting the "Tether Stiffness" is added to the cloth, 
 	 * then both the Low and High values will be used in conjunction with the per particle Weight stored
 	 * in the Weight Map to interpolate the final value from them.
 	 * Otherwise only the Low value is meaningful and sufficient to enable this constraint.
@@ -143,7 +158,7 @@ public:
 
 	/**
 	 * The limit scale of the long range attachment constraints (aka tether limit).
-	 * If an enabled Weight Map (A.K.A. Mask) targeting the "Tether Scale" is added to the cloth, 
+	 * If an enabled Weight Map (Mask with values in the range [0;1]) targeting the "Tether Scale" is added to the cloth, 
 	 * then both the Low and High values will be used in conjunction with the per particle Weight stored
 	 * in the Weight Map to interpolate the final value from them.
 	 * Otherwise only the Low value is meaningful and sufficient to set the tethers' scale.
@@ -209,7 +224,7 @@ public:
 
 	/**
 	 * The aerodynamic coefficient of drag applying on each particle.
-	 * If an enabled Weight Map (A.K.A. Mask) targeting the "Drag" is added to the cloth, 
+	 * If an enabled Weight Map (Mask with values in the range [0;1]) targeting the "Drag" is added to the cloth, 
 	 * then both the Low and High values will be used in conjunction with the per particle Weight stored
 	 * in the Weight Map to interpolate the final value from them.
 	 * Otherwise only the Low value is meaningful and sufficient to set the aerodynamic drag.
@@ -219,7 +234,7 @@ public:
 
 	/**
 	 * The aerodynamic coefficient of lift applying on each particle.
-	 * If an enabled Weight Map (A.K.A. Mask) targeting the "Lift" is added to the cloth, 
+	 * If an enabled Weight Map (Mask with values in the range [0;1]) targeting the "Lift" is added to the cloth, 
 	 * then both the Low and High values will be used in conjunction with the per particle Weight stored
 	 * in the Weight Map to interpolate the final value from them.
 	 * Otherwise only the Low value is meaningful and sufficient to set the aerodynamic lift.
@@ -241,7 +256,7 @@ public:
 
 	/**
 	 * The strength of the constraint driving the cloth towards the animated goal mesh.
-	 * If an enabled Weight Map (A.K.A. Mask) targeting the "Anim Drive Stiffness" is added to the cloth, 
+	 * If an enabled Weight Map (Mask with values in the range [0;1]) targeting the "Anim Drive Stiffness" is added to the cloth, 
 	 * then both the Low and High values will be used in conjunction with the per particle Weight stored
 	 * in the Weight Map to interpolate the final value from them.
 	 * Otherwise only the Low value is meaningful and sufficient to enable this constraint.
@@ -251,7 +266,7 @@ public:
 
 	/**
 	 * The damping amount of the anim drive.
-	 * If an enabled Weight Map (A.K.A. Mask) targeting the "Anim Drive Damping" is added to the cloth, 
+	 * If an enabled Weight Map (Mask with values in the range [0;1]) targeting the "Anim Drive Damping" is added to the cloth, 
 	 * then both the Low and High values will be used in conjunction with the per particle Weight stored
 	 * in the Weight Map to interpolate the final value from them.
 	 * Otherwise only the Low value is sufficient to work on this constraint.
@@ -297,6 +312,17 @@ public:
 
 	// Deprecated properties
 #if WITH_EDITORONLY_DATA
+	UPROPERTY()
+	float EdgeStiffness_DEPRECATED = 1.f;
+
+	/** The Stiffness of the bending constraints. Increase the iteration count for stiffer materials. Increase the iteration count for stiffer materials. */
+	UPROPERTY()
+	float BendingStiffness_DEPRECATED = 1.f;
+
+	/** The stiffness of the area preservation constraints. Increase the iteration count for stiffer materials. */
+	UPROPERTY()
+	float AreaStiffness_DEPRECATED = 1.f;
+
 	UPROPERTY()
 	EChaosClothTetherMode TetherMode_DEPRECATED = EChaosClothTetherMode::MaxChaosClothTetherMode;
 
diff --git a/Engine/Plugins/Experimental/ChaosCloth/Source/Chaos/Public/ChaosCloth/ChaosClothConstraints.h b/Engine/Plugins/Experimental/ChaosCloth/Source/Chaos/Public/ChaosCloth/ChaosClothConstraints.h
--- a/Engine/Plugins/Experimental/ChaosCloth/Source/Chaos/Public/ChaosCloth/ChaosClothConstraints.h
+++ b/Engine/Plugins/Experimental/ChaosCloth/Source/Chaos/Public/ChaosCloth/ChaosClothConstraints.h
@@ -44,15 +44,15 @@ namespace Chaos
 		// ---- End of Solver interface ----
 
 		// ---- Cloth interface ----
-		void SetEdgeConstraints(const TArray<TVec3<int32>>& SurfaceElements, FReal EdgeStiffness, bool bUseXPBDConstraints);
-		void SetBendingConstraints(TArray<TVec2<int32>>&& Edges, FReal BendingStiffness, bool bUseXPBDConstraints);
+		void SetEdgeConstraints(const TArray<TVec3<int32>>& SurfaceElements, const TConstArrayView<FRealSingle>& StiffnessMultipliers, bool bUseXPBDConstraints);
+		void SetBendingConstraints(const TArray<TVec2<int32>>& Edges, const TConstArrayView<FRealSingle>& StiffnessMultipliers, bool bUseXPBDConstraints);
 		void SetBendingConstraints(TArray<TVec4<int32>>&& BendingElements, FReal BendingStiffness);
-		void SetAreaConstraints(TArray<TVec3<int32>>&& SurfaceElements, FReal AreaStiffness, bool bUseXPBDConstraints);
-		void SetVolumeConstraints(TArray<TVec2<int32>>&& DoubleBendingEdges, FReal VolumeStiffness);
+		void SetAreaConstraints(const TArray<TVec3<int32>>& SurfaceElements, const TConstArrayView<FRealSingle>& StiffnessMultipliers, bool bUseXPBDConstraints);
+		void SetVolumeConstraints(const TArray<TVec2<int32>>& DoubleBendingEdges, FReal VolumeStiffness);
 		void SetVolumeConstraints(TArray<TVec3<int32>>&& SurfaceElements, FReal VolumeStiffness);
 		void SetLongRangeConstraints(const TMap<int32, TSet<int32>>& PointToNeighborsMap,
 			const TConstArrayView<FRealSingle>& TetherStiffnessMultipliers, const TConstArrayView<FRealSingle>& TetherScaleMultipliers,
-			const FVec2& TetherStiffness, const FVec2& TetherScale, ETetherMode TetherMode, bool bUseXPBDConstraints);
+			const FVec2& TetherScale, ETetherMode TetherMode, bool bUseXPBDConstraints);
 		void SetMaximumDistanceConstraints(const TConstArrayView<FRealSingle>& MaxDistances);
 		void SetBackstopConstraints(const TConstArrayView<FRealSingle>& BackstopDistances, const TConstArrayView<FRealSingle>& BackstopRadiuses, bool bUseLegacyBackstop);
 		void SetAnimDriveConstraints(const TConstArrayView<FRealSingle>& AnimDriveStiffnessMultipliers, const TConstArrayView<FRealSingle>& AnimDriveDampingMultipliers);
@@ -62,9 +62,9 @@ namespace Chaos
 		void CreateRules();
 		void Enable(bool bEnable);
 
-		void SetEdgeProperties(FReal EdgeStiffness);
-		void SetBendingProperties(FReal BendingStiffness);
-		void SetAreaProperties(FReal AreaStiffness);
+		void SetEdgeProperties(const FVec2& EdgeStiffness);
+		void SetBendingProperties(const FVec2& BendingStiffness);
+		void SetAreaProperties(const FVec2& AreaStiffness);
 		void SetThinShellVolumeProperties(FReal VolumeStiffness);
 		void SetVolumeProperties(FReal VolumeStiffness);
 		void SetLongRangeAttachmentProperties(const FVec2& TetherStiffness);
diff --git a/Engine/Plugins/Experimental/ChaosCloth/Source/Chaos/Public/ChaosCloth/ChaosClothingSimulationCloth.h b/Engine/Plugins/Experimental/ChaosCloth/Source/Chaos/Public/ChaosCloth/ChaosClothingSimulationCloth.h
--- a/Engine/Plugins/Experimental/ChaosCloth/Source/Chaos/Public/ChaosCloth/ChaosClothingSimulationCloth.h
+++ b/Engine/Plugins/Experimental/ChaosCloth/Source/Chaos/Public/ChaosCloth/ChaosClothingSimulationCloth.h
@@ -33,10 +33,10 @@ namespace Chaos
 			EMassMode InMassMode,
 			FRealSingle InMassValue,
 			FRealSingle InMinPerParticleMass,
-			FRealSingle InEdgeStiffness,
-			FRealSingle InBendingStiffness,
+			const TVec2<FRealSingle>& InEdgeStiffness,
+			const TVec2<FRealSingle>& InBendingStiffness,
 			bool bInUseBendingElements,
-			FRealSingle InAreaStiffness,
+			const TVec2<FRealSingle>& InAreaStiffness,
 			FRealSingle InVolumeStiffness,
 			bool bInUseThinShellVolumeConstraints,
 			const TVec2<FRealSingle>& InTetherStiffness,
@@ -76,7 +76,7 @@ namespace Chaos
 		// ---- Animatable property setters ----
 		void SetMaxDistancesMultiplier(FRealSingle InMaxDistancesMultiplier) { MaxDistancesMultiplier = InMaxDistancesMultiplier; }
 
-		void SetMaterialProperties(FRealSingle InEdgeStiffness, FRealSingle InBendingStiffness, FRealSingle InAreaStiffness) { EdgeStiffness = InEdgeStiffness; BendingStiffness = InBendingStiffness; AreaStiffness = InAreaStiffness; }
+		void SetMaterialProperties(const TVec2<FRealSingle>& InEdgeStiffness, const TVec2<FRealSingle>& InBendingStiffness, const TVec2<FRealSingle>& InAreaStiffness) { EdgeStiffness = InEdgeStiffness; BendingStiffness = InBendingStiffness; AreaStiffness = InAreaStiffness; }
 		void SetLongRangeAttachmentProperties(const TVec2<FRealSingle>& InTetherStiffness) { TetherStiffness = InTetherStiffness; }
 		void SetCollisionProperties(FRealSingle InCollisionThickness, FRealSingle InFrictionCoefficient, bool bInUseCCD, FRealSingle InSelfCollisionThickness) { CollisionThickness = InCollisionThickness; FrictionCoefficient = InFrictionCoefficient; bUseCCD = bInUseCCD; SelfCollisionThickness = InSelfCollisionThickness; }
 		void SetDampingProperties(FRealSingle InDampingCoefficient) { DampingCoefficient = InDampingCoefficient; }
@@ -184,10 +184,10 @@ namespace Chaos
 		EMassMode MassMode;
 		FRealSingle MassValue;
 		FRealSingle MinPerParticleMass;
-		FRealSingle EdgeStiffness;
-		FRealSingle BendingStiffness;
+		TVec2<FRealSingle> EdgeStiffness;
+		TVec2<FRealSingle> BendingStiffness;
 		bool bUseBendingElements;
-		FRealSingle AreaStiffness;
+		TVec2<FRealSingle> AreaStiffness;
 		FRealSingle VolumeStiffness;
 		bool bUseThinShellVolumeConstraints;
 		TVec2<FRealSingle> TetherStiffness;
diff --git a/Engine/Plugins/Experimental/ChaosCloth/Source/Chaos/Public/ChaosCloth/ChaosClothingSimulationInteractor.h b/Engine/Plugins/Experimental/ChaosCloth/Source/Chaos/Public/ChaosCloth/ChaosClothingSimulationInteractor.h
--- a/Engine/Plugins/Experimental/ChaosCloth/Source/Chaos/Public/ChaosCloth/ChaosClothingSimulationInteractor.h
+++ b/Engine/Plugins/Experimental/ChaosCloth/Source/Chaos/Public/ChaosCloth/ChaosClothingSimulationInteractor.h
@@ -27,6 +27,9 @@ public:
 	UFUNCTION(BlueprintCallable, Category = ClothingSimulation, Meta = (Keywords = "Chaos Edge Bending Area Stiffness"))
 	void SetMaterialLinear(float EdgeStiffness = 1.f, float BendingStiffness = 1.f, float AreaStiffness = 1.f);
 
+	UFUNCTION(BlueprintCallable, Category = ClothingSimulation, Meta = (Keywords = "Chaos Edge Bending Area Stiffness"))
+	void SetMaterial(FVector2D EdgeStiffness = FVector2D(1.f, 1.f), FVector2D BendingStiffness = FVector2D(1.f, 1.f), FVector2D AreaStiffness = FVector2D(1.f, 1.f));
+
 	UFUNCTION(BlueprintCallable, Category = ClothingSimulation, Meta = (Keywords = "Chaos Tether Stiffness"))
 	void SetLongRangeAttachmentLinear(float TetherStiffness = 1.f);
 
diff --git a/Engine/Source/Runtime/Core/Public/UObject/FortniteMainBranchObjectVersion.h b/Engine/Source/Runtime/Core/Public/UObject/FortniteMainBranchObjectVersion.h
--- a/Engine/Source/Runtime/Core/Public/UObject/FortniteMainBranchObjectVersion.h
+++ b/Engine/Source/Runtime/Core/Public/UObject/FortniteMainBranchObjectVersion.h
@@ -164,6 +164,9 @@ struct CORE_API FFortniteMainBranchObjectVersion
 		// Enable a few more weight maps to better art direct the cloth simulation
 		ChaosClothAddTetherScaleAndDragLiftWeightMaps,
 
+		// Enable material (edge, bending, and area stiffness) weight maps
+		ChaosClothAddMaterialWeightMaps,
+
 		// -----<new versions can be added above this line>-------------------------------------------------
 		VersionPlusOne,
 		LatestVersion = VersionPlusOne - 1
diff --git a/Engine/Source/Runtime/Experimental/Chaos/Private/Chaos/PBDAxialSpringConstraints.cpp b/Engine/Source/Runtime/Experimental/Chaos/Private/Chaos/PBDAxialSpringConstraints.cpp
--- a/Engine/Source/Runtime/Experimental/Chaos/Private/Chaos/PBDAxialSpringConstraints.cpp
+++ b/Engine/Source/Runtime/Experimental/Chaos/Private/Chaos/PBDAxialSpringConstraints.cpp
@@ -1,5 +1,6 @@
 // Copyright Epic Games, Inc. All Rights Reserved.
 #include "Chaos/PBDAxialSpringConstraints.h"
+#include "Chaos/GraphColoring.h"
 #include "ChaosStats.h"
 #include "ChaosLog.h"
 #if INTEL_ISPC
@@ -19,73 +20,128 @@ using namespace Chaos;
 int32 Chaos_AxialSpring_ParallelConstraintCount = 100;
 FAutoConsoleVariableRef CVarChaosAxialSpringParallelConstraintCount(TEXT("p.Chaos.AxialSpring.ParallelConstraintCount"), Chaos_AxialSpring_ParallelConstraintCount, TEXT("If we have more constraints than this, use parallel-for in Apply."));
 
-void FPBDAxialSpringConstraints::InitColor(const FDynamicParticles& InParticles)
+void FPBDAxialSpringConstraints::InitColor(const FPBDParticles& InParticles)
 {
 	// In dev builds we always color so we can tune the system without restarting. See Apply()
 #if UE_BUILD_SHIPPING || UE_BUILD_TEST
-	if (MConstraints.Num() > Chaos_AxialSpring_ParallelConstraintCount)
+	if (Constraints.Num() > Chaos_AxialSpring_ParallelConstraintCount)
 #endif
 	{
-		MConstraintsPerColor = FGraphColoring::ComputeGraphColoring(MConstraints, InParticles);
+		ConstraintsPerColor = FGraphColoring::ComputeGraphColoring(Constraints, InParticles);
 	}
 }
 
-void FPBDAxialSpringConstraints::ApplyImp(FPBDParticles& InParticles, const FReal Dt, const int32 i) const
+void FPBDAxialSpringConstraints::ApplyHelper(FPBDParticles& Particles, const FReal Dt, const int32 ConstraintIndex, const FReal ExpStiffnessValue) const
 {
-		const auto& constraint = MConstraints[i];
-		int32 i1 = constraint[0];
-		int32 i2 = constraint[1];
-		int32 i3 = constraint[2];
-		auto Delta = Base::GetDelta(InParticles, i);
-		FReal Multiplier = 2 / (FMath::Max(MBarys[i], 1 - MBarys[i]) + 1);
-		if (InParticles.InvM(i1) > 0)
+		const auto& constraint = Constraints[ConstraintIndex];
+		const int32 i1 = constraint[0];
+		const int32 i2 = constraint[1];
+		const int32 i3 = constraint[2];
+		const FVec3 Delta = Base::GetDelta(Particles, ConstraintIndex, ExpStiffnessValue);
+		const FReal Multiplier = 2 / (FMath::Max(Barys[ConstraintIndex], 1 - Barys[ConstraintIndex]) + 1);
+		if (Particles.InvM(i1) > (FReal)0.)
 		{
-			InParticles.P(i1) -= Multiplier * InParticles.InvM(i1) * Delta;
+			Particles.P(i1) -= Multiplier * Particles.InvM(i1) * Delta;
 		}
-		if (InParticles.InvM(i2) > 0)
+		if (Particles.InvM(i2) > (FReal)0.)
 		{
-			InParticles.P(i2) += Multiplier * InParticles.InvM(i2) * MBarys[i] * Delta;
+			Particles.P(i2) += Multiplier * Particles.InvM(i2) * Barys[ConstraintIndex] * Delta;
 		}
-		if (InParticles.InvM(i3) > 0)
+		if (Particles.InvM(i3) > (FReal)0.)
 		{
-			InParticles.P(i3) += Multiplier * InParticles.InvM(i3) * (1 - MBarys[i]) * Delta;
+			Particles.P(i3) += Multiplier * Particles.InvM(i3) * ((FReal)1. - Barys[ConstraintIndex]) * Delta;
 		}
 }
 
-void FPBDAxialSpringConstraints::Apply(FPBDParticles& InParticles, const FReal Dt) const
+void FPBDAxialSpringConstraints::Apply(FPBDParticles& Particles, const FReal Dt) const
 {
 	SCOPE_CYCLE_COUNTER(STAT_PBD_AxialSpring);
-	if ((MConstraintsPerColor.Num() > 0) && (MConstraints.Num() > Chaos_AxialSpring_ParallelConstraintCount))
+	if (ConstraintsPerColor.Num() > 0 && Constraints.Num() > Chaos_AxialSpring_ParallelConstraintCount)
 	{
-		for (const auto& Constraints : MConstraintsPerColor)
+		if (!Stiffness.HasWeightMap())
 		{
+			const FReal ExpStiffnessValue = (FReal)Stiffness;
+
+#if INTEL_ISPC
 			if (bRealTypeCompatibleWithISPC && bChaos_AxialSpring_ISPC_Enabled)
 			{
-#if INTEL_ISPC
-				ispc::ApplyAxialSpringConstraints(
-				    (ispc::FVector*)&InParticles.GetP()[0],
-				    (ispc::FIntVector*)&MConstraints.GetData()[0],
-				    &Constraints.GetData()[0],
-				    &InParticles.GetInvM().GetData()[0],
-				    &MBarys.GetData()[0],
-				    &MDists.GetData()[0],
-				    MStiffness,
-				    Constraints.Num());
+				for (const TArray<int32>& ConstraintBatch : ConstraintsPerColor)
+				{
+					ispc::ApplyAxialSpringConstraints(
+						(ispc::FVector*)&Particles.GetP()[0],
+						(ispc::FIntVector*)&Constraints.GetData()[0],
+						&ConstraintBatch.GetData()[0],
+						&Particles.GetInvM().GetData()[0],
+						&Barys.GetData()[0],
+						&Dists.GetData()[0],
+						ExpStiffnessValue,
+						ConstraintBatch.Num());
+				}
+			}
+			else
 #endif
+			{
+				for (const TArray<int32>& ConstraintBatch : ConstraintsPerColor)
+				{
+					PhysicsParallelFor(ConstraintBatch.Num(), [&](const int32 Index)
+					{
+						const int32 ConstraintIndex = ConstraintBatch[Index];
+						ApplyHelper(Particles, Dt, ConstraintIndex, ExpStiffnessValue);
+					});
+				}
+			}
+		}
+		else  // Has weight maps
+		{
+#if INTEL_ISPC
+			if (bRealTypeCompatibleWithISPC && bChaos_AxialSpring_ISPC_Enabled)
+			{
+				for (const TArray<int32>& ConstraintBatch : ConstraintsPerColor)
+				{
+					ispc::ApplyAxialSpringConstraintsWithWeightMaps(
+						(ispc::FVector*)&Particles.GetP()[0],
+						(ispc::FIntVector*)&Constraints.GetData()[0],
+						&ConstraintBatch.GetData()[0],
+						&Particles.GetInvM().GetData()[0],
+						&Barys.GetData()[0],
+						&Dists.GetData()[0],
+						&Stiffness.GetIndices().GetData()[0],
+						&Stiffness.GetTable().GetData()[0],
+						ConstraintBatch.Num());
+				}
 			}
 			else
+#endif
 			{
-				PhysicsParallelFor(Constraints.Num(), [&](const int32 Index) {
-					ApplyImp(InParticles, Dt, Constraints[Index]);
-				});
+				for (const TArray<int32>& ConstraintBatch : ConstraintsPerColor)
+				{
+					PhysicsParallelFor(ConstraintBatch.Num(), [&](const int32 Index)
+					{
+						const int32 ConstraintIndex = ConstraintBatch[Index];
+						const FReal ExpStiffnessValue = Stiffness[ConstraintIndex];
+						ApplyHelper(Particles, Dt, ConstraintIndex, ExpStiffnessValue);
+					});
+				}
 			}
 		}
 	}
 	else
 	{
-		for (int32 i = 0; i < MConstraints.Num(); ++i)
+		if (!Stiffness.HasWeightMap())
 		{
-			ApplyImp(InParticles, Dt, i);
+			const FReal ExpStiffnessValue = (FReal)Stiffness;
+			for (int32 ConstraintIndex = 0; ConstraintIndex < Constraints.Num(); ++ConstraintIndex)
+			{
+				ApplyHelper(Particles, Dt, ConstraintIndex, ExpStiffnessValue);
+			}
+		}
+		else
+		{
+			for (int32 ConstraintIndex = 0; ConstraintIndex < Constraints.Num(); ++ConstraintIndex)
+			{
+				const FReal ExpStiffnessValue = Stiffness[ConstraintIndex];
+				ApplyHelper(Particles, Dt, ConstraintIndex, ExpStiffnessValue);
+			}
 		}
 	}
 
diff --git a/Engine/Source/Runtime/Experimental/Chaos/Private/Chaos/PBDAxialSpringConstraints.ispc b/Engine/Source/Runtime/Experimental/Chaos/Private/Chaos/PBDAxialSpringConstraints.ispc
--- a/Engine/Source/Runtime/Experimental/Chaos/Private/Chaos/PBDAxialSpringConstraints.ispc
+++ b/Engine/Source/Runtime/Experimental/Chaos/Private/Chaos/PBDAxialSpringConstraints.ispc
@@ -92,3 +92,98 @@ export void ApplyAxialSpringConstraints(uniform FVector P[],
 		}
 	}
 }
+
+export void ApplyAxialSpringConstraintsWithWeightMaps(uniform FVector P[],
+										const uniform FIntVector AllConstraints[],
+										const uniform int32 ActiveConstraints[],
+										const uniform float InvM[],
+										const uniform float Barys[],
+										const uniform float Dists[],
+										const uniform uint8 StiffnessIndices[],
+										const uniform float StiffnessTable[],
+										const uniform int32 NumConstraints)
+{
+	varying int32 ActiveConstraint;
+	varying FIntVector Constraint;
+
+#if HW_GATHER_SUPPORTED == 1
+	if(programIndex < NumConstraints)
+	{
+		ActiveConstraint = ActiveConstraints[programIndex];
+		#pragma ignore warning(perf)
+		Constraint = AllConstraints[ActiveConstraint];
+	}
+#endif
+
+	foreach(i = 0 ... NumConstraints)
+	{
+#if HW_GATHER_SUPPORTED == 0
+		ActiveConstraint = ActiveConstraints[i];
+		#pragma ignore warning(perf)
+		Constraint = AllConstraints[ActiveConstraint];
+#endif
+
+		#pragma ignore warning(perf)
+		const varying float Bary = Barys[ActiveConstraint];
+		#pragma ignore warning(perf)
+		const varying float Dist = Dists[ActiveConstraint];
+		#pragma ignore warning(perf)
+		const varying uint8 StiffnessIndex = StiffnessIndices[ActiveConstraint];
+
+		const varying int32 i1 = Constraint.V[0];
+		const varying int32 i2 = Constraint.V[1];
+		const varying int32 i3 = Constraint.V[2];
+
+		const varying FVector P2 = VectorGather(&P[i2]);
+		const varying FVector P3 = VectorGather(&P[i3]);
+
+#if HW_GATHER_SUPPORTED == 1
+		if(i + programCount < NumConstraints)
+		{
+			ActiveConstraint = ActiveConstraints[i+programCount];
+			#pragma ignore warning(perf)
+			Constraint = AllConstraints[ActiveConstraint];
+		}
+#endif
+
+		const varying FVector P1 = VectorGather(&P[i1]);
+
+		const FVector P0 = (P2 - P3) * Bary + P3;
+		const FVector Difference = P1 - P0;
+
+		const float Distance = VectorSize(Difference);
+		const float Multiplier = 2 / (max(Bary, 1 - Bary) + 1);
+
+		#pragma ignore warning(perf)
+		const varying float M1 = InvM[i1];
+		#pragma ignore warning(perf)
+		const varying float M2 = InvM[i2];
+		#pragma ignore warning(perf)
+		const varying float M3 = InvM[i3];
+
+		#pragma ignore warning(perf)
+		const varying float Stiffness = StiffnessTable[StiffnessIndex];
+
+		const FVector Direction = Difference / Distance;
+		const FVector InnerDelta = (Distance - Dist) * Direction;
+
+		const float PInvMass = M3 * (1 - Bary) + M2 * Bary;
+		const float CombinedInvMass = PInvMass + M1;
+		const FVector ComputedDelta = Stiffness * InnerDelta / CombinedInvMass;
+
+		const FVector Delta = VectorSelect(((M1 == 0 && PInvMass == 0) || Distance <= SMALL_NUMBER), ZeroVector, ComputedDelta);
+
+		if (M1 > 0)
+		{
+			VectorScatter(&P[i1], P1 - (Multiplier * M1 * Delta));
+		}
+		if (M2)
+		{
+			VectorScatter(&P[i2], P2 + (Multiplier * M2 * Bary * Delta));
+		}
+		if (M3)
+		{
+			VectorScatter(&P[i3], P3 + (Multiplier * M3 * (1 - Bary) * Delta));
+		}
+	}
+}
diff --git a/Engine/Source/Runtime/Experimental/Chaos/Private/Chaos/PBDLongRangeConstraintsBase.cpp b/Engine/Source/Runtime/Experimental/Chaos/Private/Chaos/PBDLongRangeConstraintsBase.cpp
--- a/Engine/Source/Runtime/Experimental/Chaos/Private/Chaos/PBDLongRangeConstraintsBase.cpp
+++ b/Engine/Source/Runtime/Experimental/Chaos/Private/Chaos/PBDLongRangeConstraintsBase.cpp
@@ -22,7 +22,7 @@ FPBDLongRangeConstraintsBase::FPBDLongRangeConstraintsBase(
 	const FVec2& Scale,
 	const EMode InMode)
 	: TethersView(Tethers)
-	, Stiffness(StiffnessMultipliers, InStiffness, InParticleCount)
+	, Stiffness(InStiffness, StiffnessMultipliers, InParticleCount)
 	, Mode(InMode)
 	, ParticleOffset(InParticleOffset)
 	, ParticleCount(InParticleCount)
diff --git a/Engine/Source/Runtime/Experimental/Chaos/Private/Chaos/PBDSpringConstraints.cpp b/Engine/Source/Runtime/Experimental/Chaos/Private/Chaos/PBDSpringConstraints.cpp
--- a/Engine/Source/Runtime/Experimental/Chaos/Private/Chaos/PBDSpringConstraints.cpp
+++ b/Engine/Source/Runtime/Experimental/Chaos/Private/Chaos/PBDSpringConstraints.cpp
@@ -2,6 +2,7 @@
 #include "Chaos/PBDSpringConstraints.h"
 #include "Chaos/PBDParticles.h"
 #include "Chaos/PBDRigidParticles.h"
+#include "Chaos/GraphColoring.h"
 #include "Chaos/Framework/Parallel.h"
 #include "ChaosStats.h"
 #include "ChaosLog.h"
@@ -26,83 +27,120 @@ static int32 Chaos_Spring_ParallelConstraintCount = 100;
 FAutoConsoleVariableRef CVarChaosSpringParallelConstraintCount(TEXT("p.Chaos.Spring.ParallelConstraintCount"), Chaos_Spring_ParallelConstraintCount, TEXT("If we have more constraints than this, use parallel-for in Apply."));
 #endif
 
-void FPBDSpringConstraints::InitColor(const FDynamicParticles& InParticles)
+void FPBDSpringConstraints::InitColor(const FPBDParticles& Particles)
 {
 	// In dev builds we always color so we can tune the system without restarting. See Apply()
 #if UE_BUILD_SHIPPING || UE_BUILD_TEST
-	if (MConstraints.Num() > Chaos_Spring_ParallelConstraintCount)
+	if (Constraints.Num() > Chaos_Spring_ParallelConstraintCount)
 #endif
 	{
-		MConstraintsPerColor = FGraphColoring::ComputeGraphColoring(MConstraints, InParticles);
+		ConstraintsPerColor = FGraphColoring::ComputeGraphColoring(Constraints, Particles);
 	}
 }
 
-template<class T_PARTICLES>
-void FPBDSpringConstraints::Apply(T_PARTICLES& InParticles, const FReal Dt, const int32 InConstraintIndex) const
+void FPBDSpringConstraints::ApplyHelper(FPBDParticles& Particles, const FReal Dt, const int32 ConstraintIndex, const FReal ExpStiffnessValue) const
 {
-	const int32 i = InConstraintIndex;
+	const TVec2<int32>& Constraint = Constraints[ConstraintIndex];
+	const int32 i1 = Constraint[0];
+	const int32 i2 = Constraint[1];
+	const FVec3 Delta =  Base::GetDelta(Particles, ConstraintIndex, ExpStiffnessValue);
+	if (Particles.InvM(i1) > (FReal)0.)
 	{
-		const auto& Constraint = MConstraints[i];
-		const int32 i1 = Constraint[0];
-		const int32 i2 = Constraint[1];
-		auto Delta = Base::GetDelta(InParticles, i);
-		if (InParticles.InvM(i1) > 0)
-		{
-			InParticles.P(i1) -= InParticles.InvM(i1) * Delta;
-		}
-		if (InParticles.InvM(i2) > 0)
-		{
-			InParticles.P(i2) += InParticles.InvM(i2) * Delta;
-		}
+		Particles.P(i1) -= Particles.InvM(i1) * Delta;
+	}
+	if (Particles.InvM(i2) > (FReal)0.)
+	{
+		Particles.P(i2) += Particles.InvM(i2) * Delta;
 	}
 }
 
-void FPBDSpringConstraints::Apply(FPBDParticles& InParticles, const FReal Dt) const
+void FPBDSpringConstraints::Apply(FPBDParticles& Particles, const FReal Dt) const
 {
 	SCOPE_CYCLE_COUNTER(STAT_PBD_Spring);
-	if ((MConstraintsPerColor.Num() > 0) && (MConstraints.Num() > Chaos_Spring_ParallelConstraintCount))
+	if ((ConstraintsPerColor.Num() > 0) && (Constraints.Num() > Chaos_Spring_ParallelConstraintCount))
 	{
-		for (const auto& Constraints : MConstraintsPerColor)
+		if (!Stiffness.HasWeightMap())
 		{
+			const FReal ExpStiffnessValue = (FReal)Stiffness;
+
+#if INTEL_ISPC
 			if (bRealTypeCompatibleWithISPC && bChaos_Spring_ISPC_Enabled)
 			{
-#if INTEL_ISPC
-				ispc::ApplySpringConstraints(
-				    (ispc::FVector*)&InParticles.GetP()[0],
-				    (ispc::FIntVector2*)&MConstraints.GetData()[0],
-				    &Constraints.GetData()[0],
-				    &InParticles.GetInvM().GetData()[0],
-				    &MDists.GetData()[0],
-				    MStiffness,
-				    Constraints.Num());
+				for (const TArray<int32>& ConstraintBatch : ConstraintsPerColor)
+				{
+					ispc::ApplySpringConstraints(
+						(ispc::FVector*) & Particles.GetP()[0],
+						(ispc::FIntVector2*) & Constraints.GetData()[0],
+						&ConstraintBatch.GetData()[0],
+						&Particles.GetInvM().GetData()[0],
+						&Dists.GetData()[0],
+						ExpStiffnessValue,
+						ConstraintBatch.Num());
+				}
+			}
+			else
 #endif
+			{
+				for (const TArray<int32>& ConstraintBatch : ConstraintsPerColor)
+				{
+					PhysicsParallelFor(ConstraintBatch.Num(), [&](const int32 Index)
+					{
+						const int32 ConstraintIndex = ConstraintBatch[Index];
+						ApplyHelper(Particles, Dt, ConstraintIndex, ExpStiffnessValue);
+					});
+				}
+			}
+		}
+		else  // Has weight maps
+		{
+#if INTEL_ISPC
+			if (bRealTypeCompatibleWithISPC && bChaos_Spring_ISPC_Enabled)
+			{
+				for (const TArray<int32>& ConstraintBatch : ConstraintsPerColor)
+				{
+					ispc::ApplySpringConstraintsWithWeightMaps(
+						(ispc::FVector*) & Particles.GetP()[0],
+						(ispc::FIntVector2*) & Constraints.GetData()[0],
+						&ConstraintBatch.GetData()[0],
+						&Particles.GetInvM().GetData()[0],
+						&Dists.GetData()[0],
+						&Stiffness.GetIndices().GetData()[0],
+						&Stiffness.GetTable().GetData()[0],
+						ConstraintBatch.Num());
+				}
 			}
 			else
+#endif
 			{
-				PhysicsParallelFor(Constraints.Num(), [&](const int32 Index) {
-					Apply(InParticles, Dt, Constraints[Index]);
-				});
+				for (const TArray<int32>& ConstraintBatch : ConstraintsPerColor)
+				{
+					PhysicsParallelFor(ConstraintBatch.Num(), [&](const int32 Index)
+					{
+						const int32 ConstraintIndex = ConstraintBatch[Index];
+						const FReal ExpStiffnessValue = Stiffness[ConstraintIndex];
+						ApplyHelper(Particles, Dt, ConstraintIndex, ExpStiffnessValue);
+					});
+				}
 			}
 		}
 	}
 	else
 	{
-		for (int32 i = 0; i < MConstraints.Num(); ++i)
+		if (!Stiffness.HasWeightMap())
 		{
-			Apply(InParticles, Dt, i);
+			const FReal ExpStiffnessValue = (FReal)Stiffness;
+			for (int32 ConstraintIndex = 0; ConstraintIndex < Constraints.Num(); ++ConstraintIndex)
+			{
+				ApplyHelper(Particles, Dt, ConstraintIndex, ExpStiffnessValue);
+			}
+		}
+		else
+		{
+			for (int32 ConstraintIndex = 0; ConstraintIndex < Constraints.Num(); ++ConstraintIndex)
+			{
+				const FReal ExpStiffnessValue = Stiffness[ConstraintIndex];
+				ApplyHelper(Particles, Dt, ConstraintIndex, ExpStiffnessValue);
+			}
 		}
-	}
-}
-
-void FPBDSpringConstraints::Apply(TPBDRigidParticles<FReal, 3>& InParticles, const FReal Dt, const TArray<int32>& InConstraintIndices) const
-{
-	SCOPE_CYCLE_COUNTER(STAT_PBD_Spring);
-	for (int32 i : InConstraintIndices)
-	{
-		const auto& Constraint = MConstraints[i];
-		const int32 i1 = Constraint[0];
-		const int32 i2 = Constraint[1];
-		check(InParticles.Island(i1) == InParticles.Island(i2) || InParticles.Island(i1) == INDEX_NONE || InParticles.Island(i2) == INDEX_NONE);
-		Apply(InParticles, Dt, i);
 	}
 }
diff --git a/Engine/Source/Runtime/Experimental/Chaos/Private/Chaos/PBDSpringConstraints.ispc b/Engine/Source/Runtime/Experimental/Chaos/Private/Chaos/PBDSpringConstraints.ispc
--- a/Engine/Source/Runtime/Experimental/Chaos/Private/Chaos/PBDSpringConstraints.ispc
+++ b/Engine/Source/Runtime/Experimental/Chaos/Private/Chaos/PBDSpringConstraints.ispc
@@ -82,3 +82,82 @@ export void ApplySpringConstraints(uniform FVector P[],
 		}
 	}
 }
+
+export void ApplySpringConstraintsWithWeightMaps(uniform FVector P[],
+									const uniform FIntVector2 AllConstraints[],
+									const uniform int32 ActiveConstraints[],
+									const uniform float InvM[],
+									const uniform float Dists[],
+									const uniform uint8 StiffnessIndices[],
+									const uniform float StiffnessTable[],
+									const uniform int32 NumConstraints)
+{
+	varying int32 ActiveConstraint;
+	varying FIntVector2 Constraint;
+
+#if HW_GATHER_SUPPORTED == 1
+	if(programIndex < NumConstraints)
+	{
+		ActiveConstraint = ActiveConstraints[programIndex];
+		#pragma ignore warning(perf)
+		Constraint = AllConstraints[ActiveConstraint];
+	}
+#endif
+
+	foreach(i = 0 ... NumConstraints)
+	{
+#if HW_GATHER_SUPPORTED == 0
+		ActiveConstraint = ActiveConstraints[i];
+		#pragma ignore warning(perf)
+		Constraint = AllConstraints[ActiveConstraint];
+#endif
+
+		const varying int32 i1 = Constraint.V[0];
+		const varying int32 i2 = Constraint.V[1];
+
+		const varying FVector P1 = VectorGather(&P[i1]);
+		const varying FVector P2 = VectorGather(&P[i2]);
+
+		FVector Direction = P1 - P2;
+		const float Distance = SafeNormalize(Direction);
+
+		#pragma ignore warning(perf)
+		const varying float Dist = Dists[ActiveConstraint];
+
+		#pragma ignore warning(perf)
+		const varying uint8 StiffnessIndex = StiffnessIndices[ActiveConstraint];
+
+#if HW_GATHER_SUPPORTED == 1
+		if(i + programCount < NumConstraints)
+		{
+			ActiveConstraint = ActiveConstraints[i+programCount];
+			#pragma ignore warning(perf)
+			Constraint = AllConstraints[ActiveConstraint];
+		}
+#endif
+
+		const FVector InnerDelta = (Distance - Dist) * Direction;
+
+		#pragma ignore warning(perf)
+		const varying float InvM1 = InvM[i1];
+		#pragma ignore warning(perf)
+		const varying float InvM2 = InvM[i2];
+
+		#pragma ignore warning(perf)
+		const varying float Stiffness = StiffnessTable[StiffnessIndex];
+
+		const float CombinedInvMass = InvM2 + InvM1;
+		const FVector ComputedDelta = Stiffness * InnerDelta / CombinedInvMass;
+
+		const FVector Delta = VectorSelect((InvM2 == 0 && InvM1 == 0), ZeroVector, ComputedDelta);
+
+		if (InvM1 > 0)
+		{
+			VectorScatter(&P[i1], P1 - (InvM1 * Delta));
+		}
+		if (InvM2 > 0)
+		{
+			VectorScatter(&P[i2], P2 + (InvM2 * Delta));
+		}
+	}
+}
diff --git a/Engine/Source/Runtime/Experimental/Chaos/Public/Chaos/PBDAnimDriveConstraint.h b/Engine/Source/Runtime/Experimental/Chaos/Public/Chaos/PBDAnimDriveConstraint.h
--- a/Engine/Source/Runtime/Experimental/Chaos/Public/Chaos/PBDAnimDriveConstraint.h
+++ b/Engine/Source/Runtime/Experimental/Chaos/Public/Chaos/PBDAnimDriveConstraint.h
@@ -24,8 +24,8 @@ namespace Chaos
 			, OldAnimationPositions(InOldAnimationPositions)
 			, ParticleOffset(InParticleOffset)
 			, ParticleCount(InParticleCount)
-			, Stiffness(StiffnessMultipliers, FVec2((FReal)0., (FReal)1.), InParticleCount)
-			, Damping(DampingMultipliers, FVec2((FReal)0., (FReal)1.), InParticleCount)
+			, Stiffness(FVec2::UnitVector, StiffnessMultipliers, InParticleCount)
+			, Damping(FVec2::UnitVector, DampingMultipliers, InParticleCount)
 		{
 		}
 
diff --git a/Engine/Source/Runtime/Experimental/Chaos/Public/Chaos/PBDAxialSpringConstraints.h b/Engine/Source/Runtime/Experimental/Chaos/Public/Chaos/PBDAxialSpringConstraints.h
--- a/Engine/Source/Runtime/Experimental/Chaos/Public/Chaos/PBDAxialSpringConstraints.h
+++ b/Engine/Source/Runtime/Experimental/Chaos/Public/Chaos/PBDAxialSpringConstraints.h
@@ -1,38 +1,41 @@
 // Copyright Epic Games, Inc. All Rights Reserved.
 #pragma once
 
-#include "Chaos/Array.h"
 #include "Chaos/PBDAxialSpringConstraintsBase.h"
-#include "Chaos/PBDParticles.h"
-#include "Chaos/PerParticleRule.h"
-#include "Chaos/GraphColoring.h"
-#include "Chaos/Core.h"
-#include "ChaosStats.h"
 
 namespace Chaos
 {
 
-class CHAOS_API FPBDAxialSpringConstraints : public FParticleRule, public FPBDAxialSpringConstraintsBase
+class CHAOS_API FPBDAxialSpringConstraints : public FPBDAxialSpringConstraintsBase
 {
 	typedef FPBDAxialSpringConstraintsBase Base;
-	using Base::MBarys;
-	using Base::MConstraints;
-
-  public:
-	FPBDAxialSpringConstraints(const FDynamicParticles& InParticles, TArray<TVector<int32, 3>>&& Constraints, const FReal Stiffness = (FReal)1.)
-		: FPBDAxialSpringConstraintsBase(InParticles, MoveTemp(Constraints), Stiffness)
+	using Base::Barys;
+	using Base::Constraints;
+
+public:
+	FPBDAxialSpringConstraints(
+		const FPBDParticles& Particles,
+		int32 ParticleOffset,
+		int32 ParticleCount,
+		const TArray<TVec3<int32>>& InConstraints,
+		const TConstArrayView<FRealSingle>& StiffnessMultipliers,
+		const FVec2& InStiffness,
+		bool bTrimKinematicConstraints)
+		: Base(Particles, ParticleOffset, ParticleCount, InConstraints, StiffnessMultipliers, InStiffness, bTrimKinematicConstraints)
 	{
-		InitColor(InParticles);
+		InitColor(Particles);
 	}
+
 	virtual ~FPBDAxialSpringConstraints() {}
 
-  private:
-	void InitColor(const FDynamicParticles& InParticles);
-	void ApplyImp(FPBDParticles& InParticles, const FReal Dt, const int32 i) const;
-  public:
-	void Apply(FPBDParticles& InParticles, const FReal Dt) const override; //-V762
+	void Apply(FPBDParticles& InParticles, const FReal Dt) const;
+
+private:
+	void InitColor(const FPBDParticles& InParticles);
+	void ApplyHelper(FPBDParticles& InParticles, const FReal Dt, const int32 ConstraintIndex, const FReal ExpStiffnessValue) const;
 
-	TArray<TArray<int32>> MConstraintsPerColor;
+private:
+	TArray<TArray<int32>> ConstraintsPerColor;
 };
 
 }
diff --git a/Engine/Source/Runtime/Experimental/Chaos/Public/Chaos/PBDAxialSpringConstraintsBase.h b/Engine/Source/Runtime/Experimental/Chaos/Public/Chaos/PBDAxialSpringConstraintsBase.h
--- a/Engine/Source/Runtime/Experimental/Chaos/Public/Chaos/PBDAxialSpringConstraintsBase.h
+++ b/Engine/Source/Runtime/Experimental/Chaos/Public/Chaos/PBDAxialSpringConstraintsBase.h
@@ -1,34 +1,126 @@
 // Copyright Epic Games, Inc. All Rights Reserved.
 #pragma once
 
-#include "Chaos/Array.h"
 #include "Chaos/PBDParticles.h"
-#include "Chaos/PerParticleRule.h"
-
-#include <cmath>
-#include <functional>
+#include "Chaos/PBDStiffness.h"
+#include "Containers/Array.h"
 
 namespace Chaos
 {
+
 class FPBDAxialSpringConstraintsBase
 {
 public:
-	FPBDAxialSpringConstraintsBase(const FDynamicParticles& InParticles, TArray<TVec3<int32>>&& Constraints, const FReal Stiffness = (FReal)1.)
-	    : MConstraints(MoveTemp(Constraints)), MStiffness(Stiffness)
+	FPBDAxialSpringConstraintsBase(
+		const FPBDParticles& Particles,
+		int32 ParticleOffset,
+		int32 ParticleCount,
+		const TArray<TVec3<int32>>& InConstraints,
+		const TConstArrayView<FRealSingle>& StiffnessMultipliers,
+		const FVec2& InStiffness,
+		bool bTrimKinematicConstraints)
+		: Constraints(TrimConstraints(InConstraints, 
+			[&Particles, bTrimKinematicConstraints](int32 Index0, int32 Index1, int32 Index2)
+			{
+				return bTrimKinematicConstraints && Particles.InvM(Index0) == (FReal)0. && Particles.InvM(Index1) == (FReal)0. && Particles.InvM(Index2) == (FReal)0.;
+			}))
+		, Stiffness(InStiffness, StiffnessMultipliers, TConstArrayView<TVec3<int32>>(Constraints), ParticleOffset, ParticleCount)
+	{
+		Init(Particles);
+	}
+
+	virtual ~FPBDAxialSpringConstraintsBase() {}
+
+	void SetProperties(const FVec2& InStiffness) { Stiffness.SetWeightedValue(InStiffness); }
+
+	void ApplyProperties(const FReal Dt, const int32 NumIterations) { Stiffness.ApplyValues(Dt, NumIterations); }
+
+protected:
+	inline FVec3 GetDelta(const FPBDParticles& Particles, const int32 ConstraintIndex, const FReal ExpStiffnessValue) const
+	{
+		const TVec3<int32>& Constraint = Constraints[ConstraintIndex];
+		const int32 i1 = Constraint[0];
+		const int32 i2 = Constraint[1];
+		const int32 i3 = Constraint[2];
+		const FReal PInvMass = Particles.InvM(i3) * ((FReal)1. - Barys[ConstraintIndex]) + Particles.InvM(i2) * Barys[ConstraintIndex];
+		if (Particles.InvM(i1) == (FReal)0. && PInvMass == (FReal)0.)
+		{
+			return FVec3((FReal)0.);
+		}
+		const FVec3& P1 = Particles.P(i1);
+		const FVec3& P2 = Particles.P(i2);
+		const FVec3& P3 = Particles.P(i3);
+		const FVec3 P = (P2 - P3) * Barys[ConstraintIndex] + P3;
+		const FVec3 Difference = P1 - P;
+		const FReal Distance = Difference.Size();
+		if (UNLIKELY(Distance <= (FReal)SMALL_NUMBER))
+		{
+			return FVec3((FReal)0.);
+		}
+		const FVec3 Direction = Difference / Distance;
+		const FVec3 Delta = (Distance - Dists[ConstraintIndex]) * Direction;
+		const FReal CombinedInvMass = PInvMass + Particles.InvM(i1);
+		checkSlow(CombinedInvMass > (FReal)1e-7);
+		return ExpStiffnessValue * Delta / CombinedInvMass;
+	}
+
+private:
+	FReal FindBary(const FPBDParticles& Particles, const int32 i1, const int32 i2, const int32 i3)
+	{
+		const FVec3& P1 = Particles.X(i1);
+		const FVec3& P2 = Particles.X(i2);
+		const FVec3& P3 = Particles.X(i3);
+		const FVec3& P32 = P3 - P2;
+		const FReal Bary = FVec3::DotProduct(P32, P3 - P1) / P32.SizeSquared();
+		return FMath::Clamp(Bary, (FReal)0., (FReal)1.);
+	}
+
+	template<typename Predicate>
+	TArray<TVec3<int32>> TrimConstraints(const TArray<TVec3<int32>>& InConstraints, Predicate TrimPredicate)
 	{
-		for (auto& Constraint : MConstraints)
+		TSet<TVec3<int32>> TrimmedConstraints;
+		TrimmedConstraints.Reserve(InConstraints.Num());
+
+		for (const TVec3<int32>& Constraint : InConstraints)
+		{
+			const int32 Index0 = Constraint[0];
+			const int32 Index1 = Constraint[1];
+			const int32 Index2 = Constraint[2];
+
+			if (!TrimPredicate(Index0, Index1, Index2))
+			{
+				TrimmedConstraints.Add(
+					Index0 <= Index1 ?
+						Index1 <= Index2 ? TVec3<int32>(Index0, Index1, Index2) :
+						Index0 <= Index2 ? TVec3<int32>(Index0, Index2, Index1) :
+										   TVec3<int32>(Index2, Index0, Index1) :
+					// Index1 < Index0
+						Index0 <= Index2 ? TVec3<int32>(Index1, Index0, Index2) :
+						Index1 <= Index2 ? TVec3<int32>(Index1, Index2, Index0) :
+										   TVec3<int32>(Index2, Index1, Index0));
+			}
+		}
+		return TrimmedConstraints.Array();
+	}
+
+	void Init(const FPBDParticles& Particles)
+	{
+		Barys.Reset(Constraints.Num());
+		Dists.Reset(Constraints.Num());
+
+		for (TVec3<int32>& Constraint : Constraints)
 		{
 			int32 i1 = Constraint[0];
 			int32 i2 = Constraint[1];
 			int32 i3 = Constraint[2];
 			// Find Bary closest to 0.5
-			FReal Bary1 = FindBary(InParticles, i1, i2, i3);
-			FReal Bary2 = FindBary(InParticles, i2, i3, i1);
-			FReal Bary3 = FindBary(InParticles, i3, i1, i2);
+			const FReal Bary1 = FindBary(Particles, i1, i2, i3);
+			const FReal Bary2 = FindBary(Particles, i2, i3, i1);
+			const FReal Bary3 = FindBary(Particles, i3, i1, i2);
 			FReal Bary = Bary1;
-			FReal Bary1dist = FGenericPlatformMath::Abs(Bary1 - 0.5f);
-			FReal Bary2dist = FGenericPlatformMath::Abs(Bary2 - 0.5f);
-			FReal Bary3dist = FGenericPlatformMath::Abs(Bary3 - 0.5f);
+			const FReal Bary1dist = FGenericPlatformMath::Abs(Bary1 - (FReal)0.5);
+			const FReal Bary2dist = FGenericPlatformMath::Abs(Bary2 - (FReal)0.5);
+			const FReal Bary3dist = FGenericPlatformMath::Abs(Bary3 - (FReal)0.5);
 			if (Bary3dist < Bary2dist && Bary3dist < Bary1dist)
 			{
 				Constraint[0] = i3;
@@ -47,61 +139,20 @@ public:
 			i1 = Constraint[0];
 			i2 = Constraint[1];
 			i3 = Constraint[2];
-			const FVec3& P1 = InParticles.X(i1);
-			const FVec3& P2 = InParticles.X(i2);
-			const FVec3& P3 = InParticles.X(i3);
+			const FVec3& P1 = Particles.X(i1);
+			const FVec3& P2 = Particles.X(i2);
+			const FVec3& P3 = Particles.X(i3);
 			const FVec3 P = (P2 - P3) * Bary + P3;
-			MBarys.Add(Bary);
-			MDists.Add((P1 - P).Size());
+			Barys.Add(Bary);
+			Dists.Add((P1 - P).Size());
 		}
 	}
-	virtual ~FPBDAxialSpringConstraintsBase() {}
-
-	inline FVec3 GetDelta(const FPBDParticles& InParticles, const int32 i) const
-	{
-		const auto& Constraint = MConstraints[i];
-		int32 i1 = Constraint[0];
-		int32 i2 = Constraint[1];
-		int32 i3 = Constraint[2];
-		FReal PInvMass = InParticles.InvM(i3) * (1 - MBarys[i]) + InParticles.InvM(i2) * MBarys[i];
-		if (InParticles.InvM(i1) == 0 && PInvMass == 0)
-			return FVec3(0);
-		const FVec3& P1 = InParticles.P(i1);
-		const FVec3& P2 = InParticles.P(i2);
-		const FVec3& P3 = InParticles.P(i3);
-		const FVec3 P = (P2 - P3) * MBarys[i] + P3;
-		FVec3 Difference = P1 - P;
-		FReal Distance = Difference.Size();
-		if (UNLIKELY(Distance <= SMALL_NUMBER))
-			return FVec3((FReal) 0);
-		FVec3 Direction = Difference / Distance;
-		FVec3 Delta = (Distance - MDists[i]) * Direction;
-		FReal CombinedInvMass = PInvMass + InParticles.InvM(i1);
-		ensure(CombinedInvMass > 1e-7);
-		return MStiffness * Delta / CombinedInvMass;
-	}
-
-	void SetStiffness(FReal InStiffness) { MStiffness = FMath::Clamp(InStiffness, (FReal)0., (FReal)1.); }
-
-private:
-	  FReal FindBary(const FDynamicParticles& InParticles, const int32 i1, const int32 i2, const int32 i3)
-	{
-		const FVec3& P1 = InParticles.X(i1);
-		const FVec3& P2 = InParticles.X(i2);
-		const FVec3& P3 = InParticles.X(i3);
-		const FVec3& P32 = P3 - P2;
-		FReal Bary = FVec3::DotProduct(P32, P3 - P1) / P32.SizeSquared();
-		if (Bary > (FReal)1.)
-			Bary = (FReal)1.;
-		if (Bary < (FReal)0.)
-			Bary = (FReal)0.;
-		return Bary;
-	}
 
 protected:
-	TArray<TVec3<int32>> MConstraints;
-	TArray<FReal> MBarys;
-	TArray<FReal> MDists;
-	FReal MStiffness;
+	TArray<TVec3<int32>> Constraints;
+	TArray<FReal> Barys;
+	TArray<FReal> Dists;
+	FPBDStiffness Stiffness;
 };
+
 }
diff --git a/Engine/Source/Runtime/Experimental/Chaos/Public/Chaos/PBDLongRangeConstraints.h b/Engine/Source/Runtime/Experimental/Chaos/Public/Chaos/PBDLongRangeConstraints.h
--- a/Engine/Source/Runtime/Experimental/Chaos/Public/Chaos/PBDLongRangeConstraints.h
+++ b/Engine/Source/Runtime/Experimental/Chaos/Public/Chaos/PBDLongRangeConstraints.h
@@ -23,8 +23,8 @@ public:
 		const TConstArrayView<FRealSingle>& StiffnessMultipliers,
 		const TConstArrayView<FRealSingle>& ScaleMultipliers,
 		const int32 MaxNumTetherIslands = 4,
-		const FVec2& InStiffness = FVec2((FReal)1., (FReal)1.),
-		const FVec2& Scale = FVec2((FReal)1.),
+		const FVec2& InStiffness = FVec2::UnitVector,
+		const FVec2& Scale = FVec2::UnitVector,
 		const EMode InMode = EMode::Geodesic)
 		: FPBDLongRangeConstraintsBase(Particles, InParticleOffset, InParticleCount, PointToNeighbors, StiffnessMultipliers, ScaleMultipliers, MaxNumTetherIslands, InStiffness, Scale, InMode) {}
 	virtual ~FPBDLongRangeConstraints() {}
diff --git a/Engine/Source/Runtime/Experimental/Chaos/Public/Chaos/PBDLongRangeConstraintsBase.h b/Engine/Source/Runtime/Experimental/Chaos/Public/Chaos/PBDLongRangeConstraintsBase.h
--- a/Engine/Source/Runtime/Experimental/Chaos/Public/Chaos/PBDLongRangeConstraintsBase.h
+++ b/Engine/Source/Runtime/Experimental/Chaos/Public/Chaos/PBDLongRangeConstraintsBase.h
@@ -60,8 +60,8 @@ public:
 		const TConstArrayView<FRealSingle>& StiffnessMultipliers,
 		const TConstArrayView<FRealSingle>& ScaleMultipliers,
 		const int32 MaxNumTetherIslands = 4,
-		const FVec2& InStiffness = FVec2((FReal)1.),
-		const FVec2& Scale = FVec2((FReal)1.),
+		const FVec2& InStiffness = FVec2::UnitVector,
+		const FVec2& Scale = FVec2::UnitVector,
 		const EMode InMode = EMode::AccurateTetherFastLength);
 
 	virtual ~FPBDLongRangeConstraintsBase() {}
diff --git a/Engine/Source/Runtime/Experimental/Chaos/Public/Chaos/PBDSpringConstraints.h b/Engine/Source/Runtime/Experimental/Chaos/Public/Chaos/PBDSpringConstraints.h
--- a/Engine/Source/Runtime/Experimental/Chaos/Public/Chaos/PBDSpringConstraints.h
+++ b/Engine/Source/Runtime/Experimental/Chaos/Public/Chaos/PBDSpringConstraints.h
@@ -1,70 +1,43 @@
 // Copyright Epic Games, Inc. All Rights Reserved.
 #pragma once
 
-#include "Chaos/Array.h"
 #include "Chaos/PBDSpringConstraintsBase.h"
-#include "Chaos/PBDConstraintContainer.h"
-#include "Chaos/GraphColoring.h"
 
 namespace Chaos
 {
 
-template<class T, int d>
-class TPBDParticles;
-
-template<class T, int d>
-class TRigidParticles;
-
-template<class T, int d>
-class TPBDRigidParticles;
-
-class CHAOS_API FPBDSpringConstraints : public FPBDSpringConstraintsBase, public FPBDConstraintContainer
+class CHAOS_API FPBDSpringConstraints : public FPBDSpringConstraintsBase
 {
 	typedef FPBDSpringConstraintsBase Base;
-	using Base::MConstraints;
-
-  public:
-	FPBDSpringConstraints(const FReal Stiffness = (FReal)1)
-	    : FPBDSpringConstraintsBase(Stiffness)
-	{}
-
-	FPBDSpringConstraints(const FDynamicParticles& InParticles, TArray<TVec2<int32>>&& Constraints, const FReal Stiffness = (FReal)1, bool bStripKinematicConstraints = false)
-	    : FPBDSpringConstraintsBase(InParticles, MoveTemp(Constraints), Stiffness, bStripKinematicConstraints)
-	{
-		InitColor(InParticles);
-	}
-	FPBDSpringConstraints(const TRigidParticles<FReal, 3>& InParticles, TArray<TVec2<int32>>&& Constraints, const FReal Stiffness = (FReal)1, bool bStripKinematicConstraints = false)
-	    : FPBDSpringConstraintsBase(InParticles, MoveTemp(Constraints), Stiffness, bStripKinematicConstraints) 
-	{}
-
-	FPBDSpringConstraints(const FDynamicParticles& InParticles, const TArray<TVec3<int32>>& Constraints, const FReal Stiffness = (FReal)1, bool bStripKinematicConstraints = false)
-	    : FPBDSpringConstraintsBase(InParticles, Constraints, Stiffness, bStripKinematicConstraints)
-	{
-		InitColor(InParticles);
-	}
-
-	FPBDSpringConstraints(const FDynamicParticles& InParticles, const TArray<TVec4<int32>>& Constraints, const FReal Stiffness = (FReal)1, bool bStripKinematicConstraints = false)
-	    : FPBDSpringConstraintsBase(InParticles, Constraints, Stiffness, bStripKinematicConstraints)
+	using Base::Constraints;
+	using Base::Stiffness;
+
+public:
+	template<int32 Valence>
+	FPBDSpringConstraints(
+		const FPBDParticles& Particles,
+		int32 ParticleOffset,
+		int32 ParticleCount,
+		const TArray<TVector<int32, Valence>>& InConstraints,
+		const TConstArrayView<FRealSingle>& StiffnessMultipliers,
+		const FVec2& InStiffness,
+		bool bTrimKinematicConstraints = false,
+		typename TEnableIf<Valence >= 2 && Valence <= 4>::Type* = nullptr)
+		: Base(Particles, ParticleOffset, ParticleCount, InConstraints, StiffnessMultipliers, InStiffness, bTrimKinematicConstraints)
 	{
-		InitColor(InParticles);
+		InitColor(Particles);
 	}
 
 	virtual ~FPBDSpringConstraints() {}
 
-	const TArray<TVec2<int32>>& GetConstraints() const { return MConstraints; }
-	TArray<TVec2<int32>>& GetConstraints() { return MConstraints; }
-	TArray<TVec2<int32>>& Constraints() { return MConstraints; }
-
-	template<class T_PARTICLES>
-	void Apply(T_PARTICLES& InParticles, const FReal Dt, const int32 InConstraintIndex) const;
-	void Apply(TPBDParticles<FReal, 3>& InParticles, const FReal Dt) const;
-
-	void Apply(TPBDRigidParticles<FReal, 3>& InParticles, const FReal Dt, const TArray<int32>& InConstraintIndices) const;
+	void Apply(FPBDParticles& Particles, const FReal Dt) const;
 
 private:
-	void InitColor(const FDynamicParticles& InParticles);
+	void InitColor(const FPBDParticles& InParticles);
+	void ApplyHelper(FPBDParticles& Particles, const FReal Dt, const int32 ConstraintIndex, const FReal ExpStiffnessValue) const;
 
-	TArray<TArray<int32>> MConstraintsPerColor;
+private:
+	TArray<TArray<int32>> ConstraintsPerColor;
 };
 
 }
diff --git a/Engine/Source/Runtime/Experimental/Chaos/Public/Chaos/PBDSpringConstraintsBase.h b/Engine/Source/Runtime/Experimental/Chaos/Public/Chaos/PBDSpringConstraintsBase.h
--- a/Engine/Source/Runtime/Experimental/Chaos/Public/Chaos/PBDSpringConstraintsBase.h
+++ b/Engine/Source/Runtime/Experimental/Chaos/Public/Chaos/PBDSpringConstraintsBase.h
@@ -1,181 +1,110 @@
 // Copyright Epic Games, Inc. All Rights Reserved.
 #pragma once
 
-#include "Chaos/Array.h"
+#include "Chaos/Core.h"
 #include "Chaos/Utilities.h"
-#include "Chaos/RigidParticles.h"
-#include "Chaos/DynamicParticles.h"
-
+#include "Chaos/PBDParticles.h"
+#include "Chaos/PBDStiffness.h"
+#include "Containers/Array.h"
 #include "Templates/EnableIf.h"
 
-#include <functional>
-
 namespace Chaos
 {
-	class FPBDSpringConstraintsBase
+
+class FPBDSpringConstraintsBase
+{
+public:
+	template<int32 Valence>
+	FPBDSpringConstraintsBase(
+		const FPBDParticles& Particles,
+		int32 ParticleOffset,
+		int32 ParticleCount,
+		const TArray<TVector<int32, Valence>>& InConstraints,
+		const TConstArrayView<FRealSingle>& StiffnessMultipliers,
+		const FVec2& InStiffness,
+		bool bTrimKinematicConstraints = false,
+		typename TEnableIf<Valence >= 2 && Valence <= 4>::Type* = nullptr)
+		: Constraints(TrimConstraints(InConstraints, 
+			[&Particles, bTrimKinematicConstraints](int32 Index0, int32 Index1)
+			{
+				return bTrimKinematicConstraints && Particles.InvM(Index0) == (FReal)0. && Particles.InvM(Index1) == (FReal)0.;
+			}))
+		, Stiffness(InStiffness, StiffnessMultipliers, TConstArrayView<TVec2<int32>>(Constraints), ParticleOffset, ParticleCount)
 	{
-	public:
-		FPBDSpringConstraintsBase(const FReal Stiffness = (FReal)1.)
-		    : MStiffness(Stiffness)
-		{}
-		FPBDSpringConstraintsBase(const FDynamicParticles& InParticles, TArray<TVec2<int32>>&& Constraints, const FReal Stiffness = (FReal)1., bool bStripKinematicConstraints = false)
-		    : MConstraints(MoveTemp(Constraints)), MStiffness(Stiffness)
-		{
-			RemoveRedundantConstraints(InParticles, bStripKinematicConstraints);
-			UpdateDistances(InParticles);
-		}
-		FPBDSpringConstraintsBase(const TRigidParticles<FReal, 3>& InParticles, TArray<TVec2<int32>>&& Constraints, const FReal Stiffness = (FReal)1., bool bStripKinematicConstraints = false)
-		    : MConstraints(MoveTemp(Constraints)), MStiffness(Stiffness)
-		{
-			RemoveRedundantConstraints(InParticles, bStripKinematicConstraints);
-			UpdateDistances(InParticles);
-		}
-		FPBDSpringConstraintsBase(const FDynamicParticles& InParticles, const TArray<TVec3<int32>>& Constraints, const FReal Stiffness = (FReal)1., bool bStripKinematicConstraints = false)
-		    : MStiffness(Stiffness)
+		// Update distances
+		Dists.Reset(Constraints.Num());
+		for (const TVec2<int32>& Constraint : Constraints)
 		{
-			Init<3>(Constraints);
-			RemoveRedundantConstraints(InParticles, bStripKinematicConstraints);
-			UpdateDistances(InParticles);
+			const FVec3& P0 = Particles.X(Constraint[0]);
+			const FVec3& P1 = Particles.X(Constraint[1]);
+			Dists.Add((P1 - P0).Size());
 		}
-		FPBDSpringConstraintsBase(const FDynamicParticles& InParticles, const TArray<TVector<int32, 4>>& Constraints, const FReal Stiffness = (FReal)1., bool bStripKinematicConstraints = false)
-		    : MStiffness(Stiffness)
-		{
-			Init<4>(Constraints);
-			RemoveRedundantConstraints(InParticles, bStripKinematicConstraints);
-			UpdateDistances(InParticles);
-		}
-		virtual ~FPBDSpringConstraintsBase()
-		{}
+	}
 
-		void SetStiffness(FReal InStiffness) { MStiffness = FMath::Clamp(InStiffness, (FReal)0., (FReal)1.); }
+	virtual ~FPBDSpringConstraintsBase()
+	{}
 
-	protected:
-		template<class T_PARTICLES>
-		inline FVec3 GetDelta(const T_PARTICLES& InParticles, const int32 i) const
-		{
-			const auto& Constraint = MConstraints[i];
-			const int32 i1 = Constraint[0];
-			const int32 i2 = Constraint[1];
+	// Update stiffness values
+	void SetProperties(const FVec2& InStiffness) { Stiffness.SetWeightedValue(InStiffness.ClampAxes((FReal)0., (FReal)1.)); }
 
-			if (InParticles.InvM(i2) == 0 && InParticles.InvM(i1) == 0)
-				return FVec3(0);
-			const FReal CombinedMass = InParticles.InvM(i2) + InParticles.InvM(i1);
+	// Update stiffness table, as well as the simulation stiffness exponent
+	inline void ApplyProperties(const FReal Dt, const int32 NumIterations) { Stiffness.ApplyValues(Dt, NumIterations); }
 
-			const FVec3& P1 = InParticles.P(i1);
-			const FVec3& P2 = InParticles.P(i2);
-			FVec3 Direction = P1 - P2;
-			const FReal Distance = Direction.SafeNormalize();
+	const TArray<TVec2<int32>>& GetConstraints() const { return Constraints; }
 
-			const FVec3 Delta = (Distance - MDists[i]) * Direction;
-			return MStiffness * Delta / CombinedMass;
-		}
+protected:
+	FVec3 GetDelta(const FPBDParticles& Particles, const int32 ConstraintIndex, const FReal ExpStiffnessValue) const
+	{
+		const auto& Constraint = Constraints[ConstraintIndex];
+		const int32 i1 = Constraint[0];
+		const int32 i2 = Constraint[1];
 
-		//! Same as \c GetDelta(), but doesn't check for a zero length vector between 
-		// the dynamic particle positions prior to normalizing. Use this if you happen
-		// to know that the particle positions aren't coincident.
-		template<class T_PARTICLES>
-		inline FVec3 GetUnsafeDelta(const T_PARTICLES& InParticles, const int32 i) const
+		if (Particles.InvM(i2) == (FReal)0. && Particles.InvM(i1) == (FReal)0.)
 		{
-			const auto& Constraint = MConstraints[i];
-			const int32 i1 = Constraint[0];
-			const int32 i2 = Constraint[1];
-
-			if (InParticles.InvM(i2) == 0 && InParticles.InvM(i1) == 0)
-				return FVec3(0);
-			const FReal CombinedMass = InParticles.InvM(i2) + InParticles.InvM(i1);
+			return FVec3((FReal)0.);
+		}
+		const FReal CombinedMass = Particles.InvM(i2) + Particles.InvM(i1);
 
-			const FVec3& P1 = InParticles.P(i1);
-			const FVec3& P2 = InParticles.P(i2);
-			FVec3 Direction = P1 - P2;
-			const FReal Distance = Direction.Normalize();
+		const FVec3& P1 = Particles.P(i1);
+		const FVec3& P2 = Particles.P(i2);
+		FVec3 Direction = P1 - P2;
+		const FReal Distance = Direction.SafeNormalize();
 
-			const FVec3 Delta = (Distance - MDists[i]) * Direction;
-			return MStiffness * Delta / CombinedMass;
-		}
+		const FVec3 Delta = (Distance - Dists[ConstraintIndex]) * Direction;
+		return ExpStiffnessValue * Delta / CombinedMass;
+	}
 
-	private:
-		template<int32 Valence>
-		typename TEnableIf<Valence == 2, void>::Type
-		Init(TArray<TVector<int32, Valence>>&& Constraints)
-		{
-			MConstraints = MoveTemp(Constraints);
-			MDists.Reset();
-		}
-		template<int32 Valence>
-		typename TEnableIf<Valence == 2, void>::Type
-		Init(const TArray<TVector<int32, Valence>>& Constraints)
-		{
-			MConstraints = Constraints;
-			MDists.Reset();
-		}
+private:
+	template<int32 Valence, typename Predicate>
+	typename TEnableIf<Valence >= 2 && Valence <= 4, TArray<TVector<int32, 2>>>::Type
+	TrimConstraints(const TArray<TVector<int32, Valence>>& InConstraints, Predicate TrimPredicate)
+	{
+		TSet<TVec2<int32>> TrimmedConstraints;
+		TrimmedConstraints.Reserve(Valence == 2 ? InConstraints.Num() : InConstraints.Num() * Chaos::Utilities::NChooseR(Valence, 2));
 
-		template<int32 Valence>
-		typename TEnableIf<Valence != 2, void>::Type
-		Init(const TArray<TVector<int32, Valence>>& Constraints)
-		{
-			MConstraints.Reset();
-			MConstraints.Reserve(Constraints.Num() * Chaos::Utilities::NChooseR(Valence, 2));
-			MDists.Reset();
-			for (auto Constraint : Constraints)
-			{
-				for (int32 i = 0; i < Valence-1; i++)
-				{
-					for(int32 j=i+1; j < Valence; j++)
-					{
-						const int32 i1 = Constraint[i];
-						const int32 i2 = Constraint[j];
-						MConstraints.Add(TVec2<int32>(i1, i2));
-					}
-				}
-			}
-		}
-		
-		template<class T_PARTICLES>
-		uint32 RemoveRedundantConstraints(const T_PARTICLES& InParticles, bool bStripKinematicConstraints)
+		for (const TVector<int32, Valence>& ConstraintV : InConstraints)
 		{
-			const uint32 OriginalSize = MConstraints.Num();
-			TArray<TVec2<int32>> TrimmedConstraints;
-			TSet<TVec2<int32>>   ConstraintsAlreadyAdded;
-			TrimmedConstraints.Reserve(MConstraints.Num());
-			ConstraintsAlreadyAdded.Reserve(MConstraints.Num());
-			for (TVec2<int32> Constraint : MConstraints)
+			for (int32 i = 0; i < Valence - 1; ++i)
 			{
-				if (Constraint[0] > Constraint[1])
-				{
-					Swap(Constraint[0], Constraint[1]);
-				}
-				if (!ConstraintsAlreadyAdded.Contains(Constraint))
+				for (int32 j = i + 1; j < Valence; ++j)
 				{
-					ConstraintsAlreadyAdded.Add(Constraint);
+					const int32 IndexI = ConstraintV[i];
+					const int32 IndexJ = ConstraintV[j];
 
-					if (!bStripKinematicConstraints || InParticles.InvM(Constraint[0]) > 0 || InParticles.InvM(Constraint[1]) > 0)  // Only add constraint to kinematic/dynamic or dynamic/dynamic particles
+					if (!TrimPredicate(IndexI, IndexJ))
 					{
-						TrimmedConstraints.Add(Constraint);
+						TrimmedConstraints.Add(IndexI <= IndexJ ? TVec2<int32>(IndexI, IndexJ) : TVec2<int32>(IndexJ, IndexI));
 					}
 				}
 			}
-			MConstraints = MoveTemp(TrimmedConstraints);
-			return OriginalSize - MConstraints.Num();
 		}
+		return TrimmedConstraints.Array();
+	}
 
-		template<class T_PARTICLES>
-		void UpdateDistances(const T_PARTICLES& InParticles)
-		{
-			MDists.Reset();
-			MDists.Reserve(MConstraints.Num());
-			for (auto Constraint : MConstraints)
-			{
-				const int32 i1 = Constraint[0];
-				const int32 i2 = Constraint[1];
-				const FVec3& P1 = InParticles.X(i1);
-				const FVec3& P2 = InParticles.X(i2);
-				MDists.Add((P1 - P2).Size());
-			}
-		}
+protected:
+	TArray<TVec2<int32>> Constraints;
+	TArray<FReal> Dists;
+	FPBDStiffness Stiffness;
+};
 
-	protected:
-		TArray<TVec2<int32>> MConstraints;
-		TArray<FReal> MDists;
-		FReal MStiffness;
-	};
 }
diff --git a/Engine/Source/Runtime/Experimental/Chaos/Public/Chaos/PBDStiffness.h b/Engine/Source/Runtime/Experimental/Chaos/Public/Chaos/PBDStiffness.h
--- a/Engine/Source/Runtime/Experimental/Chaos/Public/Chaos/PBDStiffness.h
+++ b/Engine/Source/Runtime/Experimental/Chaos/Public/Chaos/PBDStiffness.h
@@ -11,135 +11,206 @@ DECLARE_CYCLE_STAT(TEXT("Chaos PBD Stiffness Apply Values"), STAT_PBD_StiffnessA
 
 namespace Chaos
 {
+
+/**
+ * Stiffness class for managing real time update to the weight map and low/high value ranges
+ * and to exponentiate the stiffness value depending on the iterations and Dt.
+ */
+class FPBDStiffness final
+{
+public:
 	/**
-	 * Stiffness class for managing real time update to the weight map and low/high value ranges
-	 * and to exponentiate the stiffness value depending on the iterations and Dt.
+	 * Weightmap particle constructor. 
 	 */
-	class FPBDStiffness final
-	{
-	public:
-		/**
-		 * Weightmap constructor. 
-		 */
-		inline FPBDStiffness(
-			const TConstArrayView<FRealSingle>& Multipliers,
-			const FVec2& InWeightedValue,
-			int32 ParticleCount,
-			int32 TableSize = 16,  // Size of the lookup table, can't be more than 256 values, the larger the table the longer it takes to apply changes to the stiffness values
-			FReal InParameterFitBase = (FReal)1.e3);  // Logarithm base to use in the stiffness parameter fit function
-
-		~FPBDStiffness() {}
-
-		/** Return the number of values stored in the weight map. */
-		int32 Num() const { return Indices.Num(); }
-
-		/** Return whether this object contains weight map values. */
-		bool HasWeightMap() const { return Table.Num() > 1; }
-
-		/**
-		 * Set the low and high values of the weight map.
-		 * The weight map table only gets updated after ApplyValues is called.
-		 */
-		void SetWeightedValue(const FVec2& InWeightedValue) { WeightedValue = InWeightedValue.ClampAxes((FReal)0., (FReal)1.); }
-
-		/**
-		  * Return the low and high values set for this weight map.
-		  * Both values will always be between 0 and 1 due to having been clamped in SetWeightedValue.
-		  */
-		const FVec2& GetWeightedValue() const { return WeightedValue; }
-
-		/** Update the weight map table with the current simulation parameters. */
-		inline void ApplyValues(const FReal Dt, const int32 NumIterations);
-
-		/**
-		 * Lookup for the exponential weighted value at the specified weight map index.
-		 * This function will assert if it is called with a non zero index on an empty weight map.
-		*/
-		FReal operator[](int32 Index) const { return Table[Indices[Index]]; }
-
-		/** Return the exponential value at the Low weight. */
-		FReal GetLow() const { return Table[0]; }
-
-		/** Return the exponential value at the High weight. */
-		FReal GetHigh() const { return Table.Last(); }
-
-		/** Return the exponential stiffness value when the weight map is not used. */
-		explicit operator FReal() const { return GetLow(); }
-
-	private:
-		static constexpr FReal ParameterFrequency = (FReal)120.;  // 60Hz @ 2 iterations as a root for all stiffness values TODO: Make this a global solver parameter
-
-		TArray<uint8> Indices; // Per particle array of index to the stiffness table
-		TArray<FReal> Table;  // Fixed lookup table of stiffness values, use uint8 indexation
-		FVec2 WeightedValue;
-		const FReal ParameterFitLogBase;
-	};
+	inline FPBDStiffness(
+		const FVec2& InWeightedValue,
+		const TConstArrayView<FRealSingle>& Multipliers = TConstArrayView<FRealSingle>(),
+		int32 ParticleCount = 0,  // A value of 0 also disables the map
+		int32 TableSize = 16,  // Size of the lookup table, can't be more than 256 values, the larger the table the longer it takes to apply changes to the stiffness values
+		FReal InParameterFitBase = (FReal)1.e3);  // Logarithm base to use in the stiffness parameter fit function
 
-	FPBDStiffness::FPBDStiffness(
-		const TConstArrayView<FRealSingle>& Multipliers,
+	/**
+	 * Weightmap constraint constructor. 
+	 */
+	template<int32 Valence>
+	inline FPBDStiffness(
 		const FVec2& InWeightedValue,
-		int32 ParticleCount,
-		int32 TableSize,
-		FReal InParameterFitBase)
-		: WeightedValue(InWeightedValue)
-		, ParameterFitLogBase(FMath::Loge(InParameterFitBase))
-	{
-		check(TableSize > 0 && TableSize < 256);  // The Stiffness lookup table is restricted to uint8 sized indices
+		const TConstArrayView<FRealSingle>& Multipliers = TConstArrayView<FRealSingle>(),
+		const TConstArrayView<TVector<int32, Valence>>& Constraints = TConstArrayView<TVector<int32, Valence>>(),
+		int32 ParticleOffset = INDEX_NONE, // Constraints have usually a particle offset added to them compared to the weight maps that always starts at index 0
+		int32 ParticleCount = 0,  // A value of 0 also disables the map
+		int32 TableSize = 16,  // Size of the lookup table, can't be more than 256 values, the larger the table the longer it takes to apply changes to the stiffness values
+		FReal InParameterFitBase = (FReal)1.e3,  // Logarithm base to use in the stiffness parameter fit function
+		typename TEnableIf<Valence >= 2 && Valence <= 4>::Type* = nullptr);  // Prevents incorrect valence, the value is actually unused
 
-		if (Multipliers.Num() == ParticleCount)
-		{
-			// Convert the weight maps into an array of lookup indices to the stiffness table
-			Indices.AddUninitialized(ParticleCount);
+	~FPBDStiffness() {}
 
-			const FRealSingle TableScale = (FRealSingle)(TableSize - 1);
+	/** Return the number of values stored in the weight map. */
+	int32 Num() const { return Indices.Num(); }
 
-			for (int32 Index = 0; Index < ParticleCount; ++Index)
-			{
-				Indices[Index] = (uint8)(FMath::Clamp(Multipliers[Index], (FRealSingle)0., (FRealSingle)1.) * TableScale);
-			}
+	/** Return whether this object contains weight map values. */
+	bool HasWeightMap() const { return Table.Num() > 1 && FMath::Abs(WeightedValue[0] - WeightedValue[1]) > KINDA_SMALL_NUMBER; }
 
-			// Initialize empty table until ApplyValues is called
-			Table.AddZeroed(TableSize);
-		}
-		else
+	/**
+	 * Set the low and high values of the weight map.
+	 * The weight map table only gets updated after ApplyValues is called.
+	 */
+	void SetWeightedValue(const FVec2& InWeightedValue) { WeightedValue = InWeightedValue.ClampAxes((FReal)0., (FReal)1.); }
+
+	/**
+	 * Return the low and high values set for this weight map.
+	 * Both values will always be between 0 and 1 due to having been clamped in SetWeightedValue.
+	 */
+	const FVec2& GetWeightedValue() const { return WeightedValue; }
+
+	/** Update the weight map table with the current simulation parameters. */
+	inline void ApplyValues(const FReal Dt, const int32 NumIterations);
+
+	/**
+	 * Lookup for the exponential weighted value at the specified weight map index.
+	 * This function will assert if it is called with a non zero index on an empty weight map.
+	*/
+	FReal operator[](int32 Index) const { return Table[Indices[Index]]; }
+
+	/** Return the exponential value at the Low weight. */
+	FReal GetLow() const { return Table[0]; }
+
+	/** Return the exponential value at the High weight. */
+	FReal GetHigh() const { return Table.Last(); }
+
+	/** Return the exponential stiffness value when the weight map is not used. */
+	explicit operator FReal() const { return GetLow(); }
+
+	/** Return the particles/constraints indices to the stiffness table as a read only array. */
+	TConstArrayView<uint8> GetIndices() const { return TConstArrayView<uint8>(Indices); }
+
+	/** Return the table of stiffnesses as a read only array. */
+	TConstArrayView<FReal> GetTable() const { return TConstArrayView<FReal>(Table); }
+
+private:
+	static constexpr FReal ParameterFrequency = (FReal)120.;  // 60Hz @ 2 iterations as a root for all stiffness values TODO: Make this a global solver parameter
+
+	TArray<uint8> Indices; // Per particle/constraints array of index to the stiffness table
+	TArray<FReal> Table;  // Fixed lookup table of stiffness values, use uint8 indexation
+	FVec2 WeightedValue;
+	const FReal ParameterFitLogBase;
+};
+
+FPBDStiffness::FPBDStiffness(
+	const FVec2& InWeightedValue,
+	const TConstArrayView<FRealSingle>& Multipliers,
+	int32 ParticleCount,
+	int32 TableSize,
+	FReal InParameterFitBase)
+	: WeightedValue(InWeightedValue.ClampAxes((FReal)0., (FReal)1.))
+	, ParameterFitLogBase(FMath::Loge(InParameterFitBase))
+{
+	check(TableSize > 0 && TableSize < 256);  // The Stiffness lookup table is restricted to uint8 sized indices
+
+	if (Multipliers.Num() == ParticleCount && ParticleCount > 0)
+	{
+		// Convert the weight maps into an array of lookup indices to the stiffness table
+		Indices.AddUninitialized(ParticleCount);
+
+		const FRealSingle TableScale = (FRealSingle)(TableSize - 1);
+
+		for (int32 Index = 0; Index < ParticleCount; ++Index)
 		{
-			// Initialize with a one element table until ApplyValues is called
-			Indices.AddZeroed(1);
-			Table.AddZeroed(1);
+			Indices[Index] = (uint8)(FMath::Clamp(Multipliers[Index], (FRealSingle)0., (FRealSingle)1.) * TableScale);
 		}
+
+		// Initialize empty table until ApplyValues is called
+		Table.AddZeroed(TableSize);
 	}
+	else
+	{
+		// Initialize with a one element table until ApplyValues is called
+		Indices.AddZeroed(1);
+		Table.AddZeroed(1);
+	}
+}
 
-	void FPBDStiffness::ApplyValues(const FReal Dt, const int32 NumIterations)
+template<int32 Valence>
+FPBDStiffness::FPBDStiffness(
+	const FVec2& InWeightedValue,
+	const TConstArrayView<FRealSingle>& Multipliers,
+	const TConstArrayView<TVector<int32, Valence>>& Constraints,
+	int32 ParticleOffset,
+	int32 ParticleCount,
+	int32 TableSize,
+	FReal InParameterFitBase,
+	typename TEnableIf<Valence >= 2 && Valence <= 4>::Type*)
+	: WeightedValue(InWeightedValue.ClampAxes((FReal)0., (FReal)1.))
+	, ParameterFitLogBase(FMath::Loge(InParameterFitBase))
+{
+	check(TableSize > 0 && TableSize < 256);  // The Stiffness lookup table is restricted to uint8 sized indices
+
+	const int32 ConstraintCount = Constraints.Num();
+
+	if (Multipliers.Num() == ParticleCount && ParticleCount > 0 && ConstraintCount > 0)
 	{
-		SCOPE_CYCLE_COUNTER(STAT_PBD_StiffnessApplyValues);
+		// Convert the weight maps into an array of lookup indices to the stiffness table
+		Indices.AddUninitialized(ConstraintCount);
 
-		// Calculate the simulation exponent
-		const FReal Exponent = Dt * ParameterFrequency / (FReal)NumIterations;
+		const FRealSingle TableScale = (FRealSingle)(TableSize - 1);
 
-		// Define the stiffness mapping function
-		auto SimulationValue = [this, Exponent](const FReal InValue)->FReal
-		{
-			// Get a very steep exponential curve between the [0, 1] range to make easier to set the parameter
-			// The base has been chosen empirically
-			// ParameterValue = Pow(ParameterFitBase, ParameterValue - 1)
-			const FReal ParameterFit = FMath::Exp(ParameterFitLogBase * (FMath::Clamp(InValue, (FReal)0., (FReal)1.) - (FReal)1.));
-
-			// Use simulation dependent stiffness exponent to alleviate the variations in effect when Dt and NumIterations change
-			// This is based on the Position-Based Simulation Methods paper (page 8),
-			// but uses the delta time in addition of the number of iterations in the calculation of the error term.
-			const FReal LogValue = FMath::Loge((FReal)1. - ParameterFit);
-			return (FReal)1. - FMath::Exp(LogValue * Exponent);
-		};
-
-		const FReal Offset = WeightedValue[0];
-		const FReal Range = WeightedValue[1] - WeightedValue[0];
-		const int32 TableSize = Table.Num();
-		const FReal WeightIncrement = (TableSize > 1) ? (FReal)1. / (FReal)(TableSize - 1) : (FReal)1.; // Must allow full range from 0 to 1 included
-		for (int32 Index = 0; Index < TableSize; ++Index)
+		for (int32 ConstraintIndex = 0; ConstraintIndex < ConstraintCount; ++ConstraintIndex)
 		{
-			const FReal Weight = (FReal)Index * WeightIncrement;
-			Table[Index] = SimulationValue(Offset + Weight * Range);
+			const TVector<int32, Valence>& Constraint = Constraints[ConstraintIndex];
+
+			FRealSingle Weight = 0.f;
+			for (int32 Index = 0; Index < Valence; ++Index)
+			{
+				Weight += FMath::Clamp(Multipliers[Constraint[Index] - ParticleOffset], (FRealSingle)0., (FRealSingle)1.);
+			}
+			Weight /= (FRealSingle)Valence;
+
+			Indices[ConstraintIndex] = (uint8)(Weight * TableScale);
 		}
+
+		// Initialize empty table until ApplyValues is called
+		Table.AddZeroed(TableSize);
+	}
+	else
+	{
+		// Initialize with a one element table until ApplyValues is called
+		Indices.AddZeroed(1);
+		Table.AddZeroed(1);
 	}
+}
+
+
+void FPBDStiffness::ApplyValues(const FReal Dt, const int32 NumIterations)
+{
+	SCOPE_CYCLE_COUNTER(STAT_PBD_StiffnessApplyValues);
+
+	// Calculate the simulation exponent
+	const FReal Exponent = Dt * ParameterFrequency / (FReal)NumIterations;
+
+	// Define the stiffness mapping function
+	auto SimulationValue = [this, Exponent](const FReal InValue)->FReal
+	{
+		// Get a very steep exponential curve between the [0, 1] range to make easier to set the parameter
+		// The base has been chosen empirically
+		// ParameterValue = Pow(ParameterFitBase, ParameterValue - 1)
+		const FReal ParameterFit = FMath::Exp(ParameterFitLogBase * (FMath::Clamp(InValue, (FReal)0., (FReal)1.) - (FReal)1.));
+
+		// Use simulation dependent stiffness exponent to alleviate the variations in effect when Dt and NumIterations change
+		// This is based on the Position-Based Simulation Methods paper (page 8),
+		// but uses the delta time in addition of the number of iterations in the calculation of the error term.
+		const FReal LogValue = FMath::Loge((FReal)1. - ParameterFit);
+		return (FReal)1. - FMath::Exp(LogValue * Exponent);
+	};
+
+	const FReal Offset = WeightedValue[0];
+	const FReal Range = WeightedValue[1] - WeightedValue[0];
+	const int32 TableSize = Table.Num();
+	const FReal WeightIncrement = (TableSize > 1) ? (FReal)1. / (FReal)(TableSize - 1) : (FReal)1.; // Must allow full range from 0 to 1 included
+	for (int32 Index = 0; Index < TableSize; ++Index)
+	{
+		const FReal Weight = (FReal)Index * WeightIncrement;
+		Table[Index] = SimulationValue(Offset + Weight * Range);
+	}
+}
 
 }
diff --git a/Engine/Source/Runtime/Experimental/Chaos/Public/Chaos/PerParticlePBDSpringConstraints.h b/Engine/Source/Runtime/Experimental/Chaos/Public/Chaos/PerParticlePBDSpringConstraints.h
--- a/Engine/Source/Runtime/Experimental/Chaos/Public/Chaos/PerParticlePBDSpringConstraints.h
+++ b/Engine/Source/Runtime/Experimental/Chaos/Public/Chaos/PerParticlePBDSpringConstraints.h
@@ -11,27 +11,28 @@ namespace Chaos
 class FPerParticlePBDSpringConstraints : public FPerParticleRule, public FPBDSpringConstraintsBase
 {
 	typedef FPBDSpringConstraintsBase Base;
-	using Base::MConstraints;
+	using Base::Constraints;
+	using Base::Stiffness;
 
   public:
-	FPerParticlePBDSpringConstraints(const FDynamicParticles& InParticles, TArray<TVec2<int32>>&& Constraints, const FReal Stiffness = (FReal)1.)
-	    : Base(InParticles, MoveTemp(Constraints), Stiffness)
+	FPerParticlePBDSpringConstraints(const FPBDParticles& InParticles, const TArray<TVec2<int32>>& InConstraints, const FReal InStiffness = (FReal)1.)
+	    : Base(InParticles, 0, 0, InConstraints, TConstArrayView<FRealSingle>(), FVec2(InStiffness))
 	{
-		for (int32 i = 0; i < MConstraints.Num(); ++i)
+		for (int32 i = 0; i < Constraints.Num(); ++i)
 		{
-			const auto& Constraint = MConstraints[i];
+			const auto& Constraint = Constraints[i];
 			int32 i1 = Constraint[0];
 			int32 i2 = Constraint[1];
-			if (i1 >= MParticleToConstraints.Num())
+			if (i1 >= ParticleToConstraints.Num())
 			{
-				MParticleToConstraints.SetNum(i1 + 1);
+				ParticleToConstraints.SetNum(i1 + 1);
 			}
-			if (i2 >= MParticleToConstraints.Num())
+			if (i2 >= ParticleToConstraints.Num())
 			{
-				MParticleToConstraints.SetNum(i2 + 1);
+				ParticleToConstraints.SetNum(i2 + 1);
 			}
-			MParticleToConstraints[i1].Add(i);
-			MParticleToConstraints[i2].Add(i);
+			ParticleToConstraints[i1].Add(i);
+			ParticleToConstraints[i2].Add(i);
 		}
 	}
 	virtual ~FPerParticlePBDSpringConstraints() {}
@@ -39,25 +40,25 @@ class FPerParticlePBDSpringConstraints : public FPerParticleRule, public FPBDSpr
 	// TODO(mlentine): We likely need to use time n positions here
 	void Apply(FPBDParticles& InParticles, const FReal Dt, const int32 Index) const override //-V762
 	{
-		for (int32 i = 0; i < MParticleToConstraints[Index].Num(); ++i)
+		for (int32 i = 0; i < ParticleToConstraints[Index].Num(); ++i)
 		{
-			int32 CIndex = MParticleToConstraints[Index][i];
-			const auto& Constraint = MConstraints[CIndex];
+			int32 CIndex = ParticleToConstraints[Index][i];
+			const auto& Constraint = Constraints[CIndex];
 			int32 i1 = Constraint[0];
 			int32 i2 = Constraint[1];
 			if (Index == i1 && InParticles.InvM(i1) > 0)
 			{
-				InParticles.P(i1) -= InParticles.InvM(i1) * Base::GetDelta(InParticles, CIndex);
+				InParticles.P(i1) -= InParticles.InvM(i1) * Base::GetDelta(InParticles, CIndex, (FReal)Stiffness);
 			}
 			else if (InParticles.InvM(i2) > 0)
 			{
 				check(Index == i2);
-				InParticles.P(i2) += InParticles.InvM(i2) * Base::GetDelta(InParticles, CIndex);
+				InParticles.P(i2) += InParticles.InvM(i2) * Base::GetDelta(InParticles, CIndex, (FReal)Stiffness);
 			}
 		}
 	}
 
   private:
-	TArray<TArray<int32>> MParticleToConstraints;
+	TArray<TArray<int32>> ParticleToConstraints;
 };
 }
diff --git a/Engine/Source/Runtime/Experimental/Chaos/Public/Chaos/XPBDAxialSpringConstraints.h b/Engine/Source/Runtime/Experimental/Chaos/Public/Chaos/XPBDAxialSpringConstraints.h
--- a/Engine/Source/Runtime/Experimental/Chaos/Public/Chaos/XPBDAxialSpringConstraints.h
+++ b/Engine/Source/Runtime/Experimental/Chaos/Public/Chaos/XPBDAxialSpringConstraints.h
@@ -1,83 +1,118 @@
 // Copyright Epic Games, Inc. All Rights Reserved.
 #pragma once
 
-#include "Chaos/Array.h"
 #include "Chaos/PBDAxialSpringConstraintsBase.h"
-#include "Chaos/PBDParticles.h"
-#include "Chaos/PerParticleRule.h"
 #include "ChaosStats.h"
 
-#include <algorithm>
-
 DECLARE_CYCLE_STAT(TEXT("Chaos XPBD Axial Spring Constraint"), STAT_XPBD_AxialSpring, STATGROUP_Chaos);
+
 namespace Chaos
 {
+
 // Stiffness is in N/CM^2, so it needs to be adjusted from the PBD stiffness ranging between [0,1]
 static const double XPBDAxialSpringMaxCompliance = 1e-7;  // Max stiffness: 1e+11 N/M^2 = 1e+7 N/CM^2 -> Max compliance: 1e-7 CM^2/N
 
-class FXPBDAxialSpringConstraints : public FParticleRule, public FPBDAxialSpringConstraintsBase
+class FXPBDAxialSpringConstraints : public FPBDAxialSpringConstraintsBase
 {
 	typedef FPBDAxialSpringConstraintsBase Base;
-	using Base::MBarys;
-	using Base::MConstraints;
-	using Base::MDists;
-	using Base::MStiffness;
+	using Base::Barys;
+	using Base::Constraints;
+	using Base::Dists;
+	using Base::Stiffness;
 
 public:
-	FXPBDAxialSpringConstraints(const FDynamicParticles& InParticles, TArray<TVector<int32, 3>>&& Constraints, const FReal Stiffness = (FReal)1.)
-	    : FPBDAxialSpringConstraintsBase(InParticles, MoveTemp(Constraints), Stiffness)
-	{ MLambdas.Init(0.f, MConstraints.Num()); }
+	FXPBDAxialSpringConstraints(
+		const FPBDParticles& Particles,
+		int32 ParticleOffset,
+		int32 ParticleCount,
+		const TArray<TVec3<int32>>& InConstraints,
+		const TConstArrayView<FRealSingle>& StiffnessMultipliers,
+		const FVec2& InStiffness,
+		bool bTrimKinematicConstraints)
+		: Base(Particles, ParticleOffset, ParticleCount, InConstraints, StiffnessMultipliers, InStiffness, bTrimKinematicConstraints)
+	{
+		Lambdas.Init(0.f, Constraints.Num());
+	}
 
 	virtual ~FXPBDAxialSpringConstraints() {}
 
-	void Init() const { for (FReal& Lambda : MLambdas) { Lambda = (FReal)0.; } }
+	void Init() const { for (FReal& Lambda : Lambdas) { Lambda = (FReal)0.; } }
 
-	virtual void Apply(FPBDParticles& InParticles, const FReal Dt) const override //-V762
+	void Apply(FPBDParticles& Particles, const FReal Dt) const
 	{
 		SCOPE_CYCLE_COUNTER(STAT_XPBD_AxialSpring);
-		for (int32 i = 0; i < MConstraints.Num(); ++i)
+		if (!Stiffness.HasWeightMap())
 		{
-			const TVector<int32, 3>& constraint = MConstraints[i];
-			const int32 i1 = constraint[0];
-			const int32 i2 = constraint[1];
-			const int32 i3 = constraint[2];
-			const FVec3 Delta = GetDelta(InParticles, Dt, i);
-			const FReal Multiplier = (FReal)2. / (FMath::Max(MBarys[i], (FReal)1. - MBarys[i]) + (FReal)1.);
-			if (InParticles.InvM(i1) > 0)
+			const FReal ExpStiffnessValue = (FReal)Stiffness;
+			for (int32 ConstraintIndex = 0; ConstraintIndex < Constraints.Num(); ++ConstraintIndex)
 			{
-				InParticles.P(i1) -= Multiplier * InParticles.InvM(i1) * Delta;
+				const TVector<int32, 3>& constraint = Constraints[ConstraintIndex];
+				const int32 i1 = constraint[0];
+				const int32 i2 = constraint[1];
+				const int32 i3 = constraint[2];
+				const FVec3 Delta = GetDelta(Particles, Dt, ConstraintIndex, ExpStiffnessValue);
+				const FReal Multiplier = (FReal)2. / (FMath::Max(Barys[ConstraintIndex], (FReal)1. - Barys[ConstraintIndex]) + (FReal)1.);
+				if (Particles.InvM(i1) > 0)
+				{
+					Particles.P(i1) -= Multiplier * Particles.InvM(i1) * Delta;
+				}
+				if (Particles.InvM(i2))
+				{
+					Particles.P(i2) += Multiplier * Particles.InvM(i2) * Barys[ConstraintIndex] * Delta;
+				}
+				if (Particles.InvM(i3))
+				{
+					Particles.P(i3) += Multiplier * Particles.InvM(i3) * ((FReal)1. - Barys[ConstraintIndex]) * Delta;
+				}
 			}
-			if (InParticles.InvM(i2))
-			{
-				InParticles.P(i2) += Multiplier * InParticles.InvM(i2) * MBarys[i] * Delta;
-			}
-			if (InParticles.InvM(i3))
+		}
+		else
+		{
+			for (int32 ConstraintIndex = 0; ConstraintIndex < Constraints.Num(); ++ConstraintIndex)
 			{
-				InParticles.P(i3) += Multiplier * InParticles.InvM(i3) * (1 - MBarys[i]) * Delta;
+				const FReal ExpStiffnessValue = Stiffness[ConstraintIndex];
+				const TVector<int32, 3>& constraint = Constraints[ConstraintIndex];
+				const int32 i1 = constraint[0];
+				const int32 i2 = constraint[1];
+				const int32 i3 = constraint[2];
+				const FVec3 Delta = GetDelta(Particles, Dt, ConstraintIndex, ExpStiffnessValue);
+				const FReal Multiplier = (FReal)2. / (FMath::Max(Barys[ConstraintIndex], (FReal)1. - Barys[ConstraintIndex]) + (FReal)1.);
+				if (Particles.InvM(i1) > 0)
+				{
+					Particles.P(i1) -= Multiplier * Particles.InvM(i1) * Delta;
+				}
+				if (Particles.InvM(i2))
+				{
+					Particles.P(i2) += Multiplier * Particles.InvM(i2) * Barys[ConstraintIndex] * Delta;
+				}
+				if (Particles.InvM(i3))
+				{
+					Particles.P(i3) += Multiplier * Particles.InvM(i3) * ((FReal)1. - Barys[ConstraintIndex]) * Delta;
+				}
 			}
 		}
 	}
 
 private:
-	inline FVec3 GetDelta(const FPBDParticles& InParticles, const FReal Dt, const int32 InConstraintIndex) const
+	FVec3 GetDelta(const FPBDParticles& Particles, const FReal Dt, const int32 InConstraintIndex, const FReal ExpStiffnessValue) const
 	{
-		const TVector<int32, 3>& Constraint = MConstraints[InConstraintIndex];
+		const TVector<int32, 3>& Constraint = Constraints[InConstraintIndex];
 		const int32 i1 = Constraint[0];
 		const int32 i2 = Constraint[1];
 		const int32 i3 = Constraint[2];
 
-		const FReal Bary = MBarys[InConstraintIndex];
-		const FReal PInvMass = InParticles.InvM(i3) * ((FReal)1. - Bary) + InParticles.InvM(i2) * Bary;
-		if (InParticles.InvM(i1) == (FReal)0. && PInvMass == (FReal)0.)
+		const FReal Bary = Barys[InConstraintIndex];
+		const FReal PInvMass = Particles.InvM(i3) * ((FReal)1. - Bary) + Particles.InvM(i2) * Bary;
+		if (Particles.InvM(i1) == (FReal)0. && PInvMass == (FReal)0.)
 		{
 			return FVec3((FReal)0.);
 		}
-		const FReal CombinedInvMass = PInvMass + InParticles.InvM(i1);
+		const FReal CombinedInvMass = PInvMass + Particles.InvM(i1);
 		ensure(CombinedInvMass > (FReal)SMALL_NUMBER);
 
-		const FVec3& P1 = InParticles.P(i1);
-		const FVec3& P2 = InParticles.P(i2);
-		const FVec3& P3 = InParticles.P(i3);
+		const FVec3& P1 = Particles.P(i1);
+		const FVec3& P2 = Particles.P(i2);
+		const FVec3& P3 = Particles.P(i3);
 		const FVec3 P = (P2 - P3) * Bary + P3;
 
 		const FVec3 Difference = P1 - P;
@@ -87,10 +122,10 @@ private:
 			return FVec3((FReal)0.);
 		}
 		const FVec3 Direction = Difference / Distance;
-		const FReal Offset = (Distance - MDists[InConstraintIndex]);
+		const FReal Offset = (Distance - Dists[InConstraintIndex]);
 
-		FReal& Lambda = MLambdas[InConstraintIndex];
-		const FReal Alpha = (FReal)XPBDAxialSpringMaxCompliance / (MStiffness * Dt * Dt);
+		FReal& Lambda = Lambdas[InConstraintIndex];
+		const FReal Alpha = (FReal)XPBDAxialSpringMaxCompliance / (ExpStiffnessValue * Dt * Dt);
 
 		const FReal DLambda = (Offset - Alpha * Lambda) / (CombinedInvMass + Alpha);
 		const FVec3 Delta = DLambda * Direction;
@@ -100,7 +135,7 @@ private:
 	}
 
 private:
-	mutable TArray<FReal> MLambdas;
+	mutable TArray<FReal> Lambdas;
 };
 
 template<class T, int d>
diff --git a/Engine/Source/Runtime/Experimental/Chaos/Public/Chaos/XPBDLongRangeConstraints.h b/Engine/Source/Runtime/Experimental/Chaos/Public/Chaos/XPBDLongRangeConstraints.h
--- a/Engine/Source/Runtime/Experimental/Chaos/Public/Chaos/XPBDLongRangeConstraints.h
+++ b/Engine/Source/Runtime/Experimental/Chaos/Public/Chaos/XPBDLongRangeConstraints.h
@@ -27,8 +27,8 @@ public:
 		const TConstArrayView<FRealSingle>& StiffnessMultipliers,
 		const TConstArrayView<FRealSingle>& ScaleMultipliers,
 		const int32 MaxNumTetherIslands = 4,
-		const FVec2& InStiffness = FVec2((FReal)1.),
-		const FVec2& Scale = FVec2((FReal)1.),
+		const FVec2& InStiffness = FVec2::UnitVector,
+		const FVec2& Scale = FVec2::UnitVector,
 		const EMode InMode = EMode::Geodesic)
 	    : FPBDLongRangeConstraintsBase(Particles, InParticleOffset, InParticleCount, PointToNeighbors, StiffnessMultipliers, ScaleMultipliers, MaxNumTetherIslands, InStiffness, Scale, InMode)
 	{
diff --git a/Engine/Source/Runtime/Experimental/Chaos/Public/Chaos/XPBDSpringConstraints.h b/Engine/Source/Runtime/Experimental/Chaos/Public/Chaos/XPBDSpringConstraints.h
--- a/Engine/Source/Runtime/Experimental/Chaos/Public/Chaos/XPBDSpringConstraints.h
+++ b/Engine/Source/Runtime/Experimental/Chaos/Public/Chaos/XPBDSpringConstraints.h
@@ -1,113 +1,104 @@
 // Copyright Epic Games, Inc. All Rights Reserved.
 #pragma once
 
-#include "Chaos/Array.h"
-#include "Chaos/PBDParticles.h"
 #include "Chaos/PBDSpringConstraintsBase.h"
-#include "Chaos/PBDConstraintContainer.h"
 #include "ChaosStats.h"
 
-#include "Templates/EnableIf.h"
-
 DECLARE_CYCLE_STAT(TEXT("Chaos XPBD Spring Constraint"), STAT_XPBD_Spring, STATGROUP_Chaos);
 
 namespace Chaos
 {
+
 // Stiffness is in N/CM^2, so it needs to be adjusted from the PBD stiffness ranging between [0,1]
 static const double XPBDSpringMaxCompliance = 1e-7;  // Max stiffness: 1e+11 N/M^2 = 1e+7 N/CM^2 -> Max compliance: 1e-7 CM^2/N
 
-class FXPBDSpringConstraints : public FPBDSpringConstraintsBase, public FPBDConstraintContainer
+class FXPBDSpringConstraints : public FPBDSpringConstraintsBase
 {
 	typedef FPBDSpringConstraintsBase Base;
-	using Base::MConstraints;
-	using Base::MDists;
-	using Base::MStiffness;
+	using Base::Constraints;
+	using Base::Dists;
+	using Base::Stiffness;
 
 public:
-	FXPBDSpringConstraints(const FReal Stiffness = (FReal)1.)
-	    : FPBDSpringConstraintsBase(Stiffness)
-	{}
-
-	FXPBDSpringConstraints(const FDynamicParticles& InParticles, TArray<TVec2<int32>>&& Constraints, const FReal Stiffness = (FReal)1., bool bStripKinematicConstraints = false)
-		: FPBDSpringConstraintsBase(InParticles, MoveTemp(Constraints), Stiffness, bStripKinematicConstraints)
-	{
-		MLambdas.Init((FReal)0., MConstraints.Num());
-	}
-
-	FXPBDSpringConstraints(const TRigidParticles<FReal, 3>& InParticles, TArray<TVec2<int32>>&& Constraints, const FReal Stiffness = (FReal)1., bool bStripKinematicConstraints = false)
-		: FPBDSpringConstraintsBase(InParticles, MoveTemp(Constraints), Stiffness, bStripKinematicConstraints)
-	{
-		MLambdas.Init((FReal)0., MConstraints.Num());
-	}
-
-	FXPBDSpringConstraints(const FDynamicParticles& InParticles, const TArray<TVec3<int32>>& Constraints, const FReal Stiffness = (FReal)1., bool bStripKinematicConstraints = false)
-		: FPBDSpringConstraintsBase(InParticles, Constraints, Stiffness, bStripKinematicConstraints)
-	{
-		MLambdas.Init((FReal)0., MConstraints.Num());
-	}
-
-	FXPBDSpringConstraints(const FDynamicParticles& InParticles, const TArray<TVec4<int32>>& Constraints, const FReal Stiffness = (FReal)1., bool bStripKinematicConstraints = false)
-		: FPBDSpringConstraintsBase(InParticles, Constraints, Stiffness, bStripKinematicConstraints)
+	template<int32 Valence>
+	FXPBDSpringConstraints(
+		const FPBDParticles& Particles,
+		int32 ParticleOffset,
+		int32 ParticleCount,
+		const TArray<TVector<int32, Valence>>& InConstraints,
+		const TConstArrayView<FRealSingle>& StiffnessMultipliers,
+		const FVec2& InStiffness,
+		bool bTrimKinematicConstraints = false,
+		typename TEnableIf<Valence >= 2 && Valence <= 4>::Type* = nullptr)
+		: Base(Particles, ParticleOffset, ParticleCount, InConstraints, StiffnessMultipliers, InStiffness, bTrimKinematicConstraints)
 	{
-		MLambdas.Init((FReal)0., MConstraints.Num());
+		Lambdas.Init((FReal)0., Constraints.Num());
 	}
 
 	virtual ~FXPBDSpringConstraints() {}
 
-	const TArray<TVec2<int32>>& GetConstraints() const { return MConstraints; }
-	TArray<TVec2<int32>>& GetConstraints() { return MConstraints; }
-	TArray<TVec2<int32>>& Constraints() { return MConstraints; }
+	void Init() const { for (FReal& Lambda : Lambdas) { Lambda = (FReal)0.; } }
 
-	void Init() const { for (FReal& Lambda : MLambdas) { Lambda = (FReal)0.; } }
-
-	void Apply(FPBDParticles& InParticles, const FReal Dt, const int32 InConstraintIndex) const
+	void Apply(FPBDParticles& Particles, const FReal Dt, const int32 ConstraintIndex, const FReal ExpStiffnessValue) const
 	{
-		const int32 i = InConstraintIndex;
+		const TVec2<int32>& Constraint = Constraints[ConstraintIndex];
+		const int32 i1 = Constraint[0];
+		const int32 i2 = Constraint[1];
+		const FVec3 Delta = GetDelta(Particles, Dt, ConstraintIndex, ExpStiffnessValue);
+		if (Particles.InvM(i1) > (FReal)0.)
 		{
-			const auto& Constraint = MConstraints[i];
-			const int32 i1 = Constraint[0];
-			const int32 i2 = Constraint[1];
-			const FVec3 Delta = GetDelta(InParticles, Dt, i);
-			if (InParticles.InvM(i1) > 0)
-			{
-				InParticles.P(i1) -= InParticles.InvM(i1) * Delta;
-			}
-			if (InParticles.InvM(i2) > 0)
-			{
-				InParticles.P(i2) += InParticles.InvM(i2) * Delta;
-			}
+			Particles.P(i1) -= Particles.InvM(i1) * Delta;
+		}
+		if (Particles.InvM(i2) > (FReal)0.)
+		{
+			Particles.P(i2) += Particles.InvM(i2) * Delta;
 		}
 	}
 
-	void Apply(FPBDParticles& InParticles, const FReal Dt) const
+	void Apply(FPBDParticles& Particles, const FReal Dt) const
 	{
 		SCOPE_CYCLE_COUNTER(STAT_XPBD_Spring);
-		for (int32 i = 0; i < MConstraints.Num(); ++i)
+
+		if (!Stiffness.HasWeightMap())
 		{
-			Apply(InParticles, Dt, i);
+			const FReal ExpStiffnessValue = (FReal)Stiffness;
+			for (int32 ConstraintIndex = 0; ConstraintIndex < Constraints.Num(); ++ConstraintIndex)
+			{
+				Apply(Particles, Dt, ConstraintIndex, ExpStiffnessValue);
+			}
+		}
+		else
+		{
+			for (int32 ConstraintIndex = 0; ConstraintIndex < Constraints.Num(); ++ConstraintIndex)
+			{
+				const FReal ExpStiffnessValue = Stiffness[ConstraintIndex];
+				Apply(Particles, Dt, ConstraintIndex, ExpStiffnessValue);
+			}
 		}
 	}
 
 private:
-	inline FVec3 GetDelta(const FPBDParticles& InParticles, const FReal Dt, const int32 InConstraintIndex) const
+	FVec3 GetDelta(const FPBDParticles& Particles, const FReal Dt, const int32 ConstraintIndex, const FReal ExpStiffnessValue) const
 	{
-		const TVec2<int32>& Constraint = MConstraints[InConstraintIndex];
+		const TVec2<int32>& Constraint = Constraints[ConstraintIndex];
 
 		const int32 i1 = Constraint[0];
 		const int32 i2 = Constraint[1];
 
-		if (InParticles.InvM(i2) == 0 && InParticles.InvM(i1) == 0)
-			return FVec3(0.);
-		const FReal CombinedInvMass = InParticles.InvM(i2) + InParticles.InvM(i1);
+		if (Particles.InvM(i2) == (FReal)0. && Particles.InvM(i1) == (FReal)0.)
+		{
+			return FVec3((FReal)0.);
+		}
+		const FReal CombinedInvMass = Particles.InvM(i2) + Particles.InvM(i1);
 
-		const FVec3& P1 = InParticles.P(i1);
-		const FVec3& P2 = InParticles.P(i2);
+		const FVec3& P1 = Particles.P(i1);
+		const FVec3& P2 = Particles.P(i2);
 		FVec3 Direction = P1 - P2;
 		const FReal Distance = Direction.SafeNormalize();
-		const FReal Offset = Distance - MDists[InConstraintIndex];
+		const FReal Offset = Distance - Dists[ConstraintIndex];
 
-		FReal& Lambda = MLambdas[InConstraintIndex];
-		const FReal Alpha = (FReal)XPBDSpringMaxCompliance / (MStiffness * Dt * Dt);
+		FReal& Lambda = Lambdas[ConstraintIndex];
+		const FReal Alpha = (FReal)XPBDSpringMaxCompliance / (ExpStiffnessValue * Dt * Dt);
 
 		const FReal DLambda = (Offset - Alpha * Lambda) / (CombinedInvMass + Alpha);
 		const FVec3 Delta = DLambda * Direction;
@@ -117,6 +108,7 @@ private:
 	}
 
 private:
-	mutable TArray<FReal> MLambdas;
+	mutable TArray<FReal> Lambdas;
 };
+
 }
