GeometryCollection
-- Fixed namespace collision

diff --git a/Engine/Source/Runtime/Experimental/Chaos/Private/GeometryCollection/GeometryCollection.cpp b/Engine/Source/Runtime/Experimental/Chaos/Private/GeometryCollection/GeometryCollection.cpp
--- a/Engine/Source/Runtime/Experimental/Chaos/Private/GeometryCollection/GeometryCollection.cpp
+++ b/Engine/Source/Runtime/Experimental/Chaos/Private/GeometryCollection/GeometryCollection.cpp
@@ -173,7 +173,7 @@ int32 FGeometryCollection::AppendGeometry(const FGeometryCollection & Element, i
 	{
 		Vertices[VerticesIndex + vdx] = ElementVertices[vdx];
 		Normals[VerticesIndex + vdx] = ElementNormals[vdx];
-		for (int UVChannelIndex = 0; UVChannelIndex < Chaos::GeometryCollection::MAX_NUM_UV_CHANNELS; UVChannelIndex++)
+		for (int UVChannelIndex = 0; UVChannelIndex < GeometryCollectionUV::MAX_NUM_UV_CHANNELS; UVChannelIndex++)
 		{
 			AllUVs[VerticesIndex + vdx][UVChannelIndex] = ElementUVs[vdx][UVChannelIndex];
 		}
@@ -1258,7 +1258,7 @@ FGeometryCollection* FGeometryCollection::NewGeometryCollection(const TArray<flo
 		Vertices[Idx] = FVector3f(RawVertexArray[3 * Idx], RawVertexArray[3 * Idx + 1], RawVertexArray[3 * Idx + 2]);
 		TempVertices += Vertices[Idx];
 
-		UVs[Idx].SetNumZeroed(Chaos::GeometryCollection::MAX_NUM_UV_CHANNELS);
+		UVs[Idx].SetNumZeroed(GeometryCollectionUV::MAX_NUM_UV_CHANNELS);
 		Colors[Idx] = FLinearColor::White;
 	}
 
@@ -1577,7 +1577,7 @@ FGeometryCollection* FGeometryCollection::NewGeometryCollection(const TArray<flo
 		Vertices[Idx] = FVector3f(RawVertexArray[3 * Idx], RawVertexArray[3 * Idx + 1], RawVertexArray[3 * Idx + 2]);
 		BoneMap[Idx] = RawBoneMapArray[Idx];
 
-		UVs[Idx].SetNumZeroed(Chaos::GeometryCollection::MAX_NUM_UV_CHANNELS);
+		UVs[Idx].SetNumZeroed(GeometryCollectionUV::MAX_NUM_UV_CHANNELS);
 		Colors[Idx] = FLinearColor::White;
 	}
 
diff --git a/Engine/Source/Runtime/Experimental/Chaos/Public/GeometryCollection/GeometryCollection.h b/Engine/Source/Runtime/Experimental/Chaos/Public/GeometryCollection/GeometryCollection.h
--- a/Engine/Source/Runtime/Experimental/Chaos/Public/GeometryCollection/GeometryCollection.h
+++ b/Engine/Source/Runtime/Experimental/Chaos/Public/GeometryCollection/GeometryCollection.h
@@ -9,14 +9,14 @@
 namespace Chaos
 {
 	class FChaosArchive;
+}
 
-	namespace GeometryCollection
+namespace GeometryCollectionUV
+{
+	enum
 	{
-		enum 
-		{
-			MAX_NUM_UV_CHANNELS = 8,
-		};
-	}
+		MAX_NUM_UV_CHANNELS = 8,
+	};
 }
 
 /**
diff --git a/Engine/Source/Runtime/Experimental/GeometryCollectionEngine/Private/GeometryCollection/GeometryCollectionSceneProxy.cpp b/Engine/Source/Runtime/Experimental/GeometryCollectionEngine/Private/GeometryCollection/GeometryCollectionSceneProxy.cpp
--- a/Engine/Source/Runtime/Experimental/GeometryCollectionEngine/Private/GeometryCollection/GeometryCollectionSceneProxy.cpp
+++ b/Engine/Source/Runtime/Experimental/GeometryCollectionEngine/Private/GeometryCollection/GeometryCollectionSceneProxy.cpp
@@ -218,7 +218,7 @@ void FGeometryCollectionSceneProxy::InitResources()
 	
 	// Init buffers
 	VertexBuffers.PositionVertexBuffer.Init(NumVertices);
-	VertexBuffers.StaticMeshVertexBuffer.Init(NumVertices, Chaos::GeometryCollection::MAX_NUM_UV_CHANNELS);
+	VertexBuffers.StaticMeshVertexBuffer.Init(NumVertices, GeometryCollectionUV::MAX_NUM_UV_CHANNELS);
 	VertexBuffers.ColorVertexBuffer.Init(NumVertices);
 
 	// Init resources
@@ -388,7 +388,7 @@ void FGeometryCollectionSceneProxy::SetConstantData_RenderThread(FGeometryCollec
 
 			VertexBuffers.PositionVertexBuffer.VertexPosition(i) = Vertex.Position;
 			VertexBuffers.StaticMeshVertexBuffer.SetVertexTangents(i, Vertex.TangentX.ToFVector(), Vertex.GetTangentY(), Vertex.TangentZ.ToFVector());
-			for (int UVChannelIndex = 0; UVChannelIndex < Chaos::GeometryCollection::MAX_NUM_UV_CHANNELS; UVChannelIndex++)
+			for (int UVChannelIndex = 0; UVChannelIndex < GeometryCollectionUV::MAX_NUM_UV_CHANNELS; UVChannelIndex++)
 			{
 				VertexBuffers.StaticMeshVertexBuffer.SetVertexUV(i, UVChannelIndex, Vertex.TextureCoordinate[UVChannelIndex]);
 			}
