GeometryCollectionUtility
-- fixed FVector2D constructor call that was breaking builds.

diff --git a/Engine/Source/Runtime/Experimental/Chaos/Public/GeometryCollection/GeometryCollectionUtility.h b/Engine/Source/Runtime/Experimental/Chaos/Public/GeometryCollection/GeometryCollectionUtility.h
--- a/Engine/Source/Runtime/Experimental/Chaos/Public/GeometryCollection/GeometryCollectionUtility.h
+++ b/Engine/Source/Runtime/Experimental/Chaos/Public/GeometryCollection/GeometryCollectionUtility.h
@@ -58,7 +58,7 @@ namespace GeometryCollection
 		{
 			Vertices[Idx] = GeoXf.TransformPosition(FVector(PointsIn[Idx][0], PointsIn[Idx][1], PointsIn[Idx][2])); // transform points by GeoXf
 			Normals[Idx] = NormalsIn.Num() > Idx ? FVector(NormalsIn[Idx][0], NormalsIn[Idx][1], NormalsIn[Idx][2]) : FVector(0);
-			UVs[Idx] = UVsIn.Num() > Idx ? FVector2D(UVsIn[Idx][0], UVsIn[Idx][1]) : FVector2D(0);
+			UVs[Idx] = UVsIn.Num() > Idx ? FVector2D{ UVsIn[Idx][0], UVsIn[Idx][1] } : FVector2D{ 0.0f, 0.0f };
 			Colors[Idx] = FLinearColor::White;
 		}
 
