GeometryCollection
-- Corrected issues causing HeadlessChaos to fail due to addition of multiple UVs

diff --git a/Engine/Source/Runtime/Experimental/Chaos/Private/GeometryCollection/GeometryCollection.cpp b/Engine/Source/Runtime/Experimental/Chaos/Private/GeometryCollection/GeometryCollection.cpp
--- a/Engine/Source/Runtime/Experimental/Chaos/Private/GeometryCollection/GeometryCollection.cpp
+++ b/Engine/Source/Runtime/Experimental/Chaos/Private/GeometryCollection/GeometryCollection.cpp
@@ -173,10 +173,14 @@ int32 FGeometryCollection::AppendGeometry(const FGeometryCollection & Element, i
 	{
 		Vertices[VerticesIndex + vdx] = ElementVertices[vdx];
 		Normals[VerticesIndex + vdx] = ElementNormals[vdx];
-		for (int UVChannelIndex = 0; UVChannelIndex < GeometryCollectionUV::MAX_NUM_UV_CHANNELS; UVChannelIndex++)
+		
+		const int32 NumUVLayers = ElementUVs[vdx].Num();
+		AllUVs[VerticesIndex + vdx].SetNum(NumUVLayers);
+		for (int UVChannelIndex = 0; UVChannelIndex < NumUVLayers; UVChannelIndex++)
 		{
 			AllUVs[VerticesIndex + vdx][UVChannelIndex] = ElementUVs[vdx][UVChannelIndex];
 		}
+
 		Colors[VerticesIndex + vdx] = ElementColors[vdx];
 		TangentUs[VerticesIndex + vdx] = ElementTangentUs[vdx];
 		TangentVs[VerticesIndex + vdx] = ElementTangentVs[vdx];
diff --git a/Engine/Source/Runtime/Experimental/Chaos/Private/GeometryCollection/GeometryCollectionUtility.cpp b/Engine/Source/Runtime/Experimental/Chaos/Private/GeometryCollection/GeometryCollectionUtility.cpp
--- a/Engine/Source/Runtime/Experimental/Chaos/Private/GeometryCollection/GeometryCollectionUtility.cpp
+++ b/Engine/Source/Runtime/Experimental/Chaos/Private/GeometryCollection/GeometryCollectionUtility.cpp
@@ -66,15 +66,14 @@ namespace GeometryCollection
 		Normals[6] = FVector3f(-1.f, 1.f, 1.f).GetSafeNormal();
 		Normals[7] = FVector3f(1.f, 1.f, 1.f).GetSafeNormal();
 
-		const int32 UVChannel = 0;
-		UVs[0][UVChannel] = FVector2D(0, 0);
-		UVs[1][UVChannel] = FVector2D(1, 0);
-		UVs[2][UVChannel] = FVector2D(0, 1);
-		UVs[3][UVChannel] = FVector2D(1, 1);
-		UVs[4][UVChannel] = FVector2D(0, 0);
-		UVs[5][UVChannel] = FVector2D(1, 0);
-		UVs[6][UVChannel] = FVector2D(0, 1);
-		UVs[7][UVChannel] = FVector2D(1, 1);
+		UVs[0].Add(FVector2D(0, 0));
+		UVs[1].Add(FVector2D(1, 0));
+		UVs[2].Add(FVector2D(0, 1));
+		UVs[3].Add(FVector2D(1, 1));
+		UVs[4].Add(FVector2D(0, 0));
+		UVs[5].Add(FVector2D(1, 0));
+		UVs[6].Add(FVector2D(0, 1));
+		UVs[7].Add(FVector2D(1, 1));
 
 		Colors[0] = FLinearColor::White;
 		Colors[1] = FLinearColor::White;
diff --git a/Engine/Source/Runtime/Experimental/Chaos/Public/GeometryCollection/GeometryCollectionUtility.h b/Engine/Source/Runtime/Experimental/Chaos/Public/GeometryCollection/GeometryCollectionUtility.h
--- a/Engine/Source/Runtime/Experimental/Chaos/Public/GeometryCollection/GeometryCollectionUtility.h
+++ b/Engine/Source/Runtime/Experimental/Chaos/Public/GeometryCollection/GeometryCollectionUtility.h
@@ -58,7 +58,10 @@ namespace GeometryCollection
 		{
 			Vertices[Idx] = GeoXf.TransformPosition(FVector(PointsIn[Idx][0], PointsIn[Idx][1], PointsIn[Idx][2])); // transform points by GeoXf
 			Normals[Idx] = NormalsIn.Num() > Idx ? FVector(NormalsIn[Idx][0], NormalsIn[Idx][1], NormalsIn[Idx][2]) : FVector(0);
-			UVs[Idx] = UVsIn.Num() > Idx ? FVector2D{ UVsIn[Idx][0], UVsIn[Idx][1] } : FVector2D{ 0.0f, 0.0f };
+
+			FVector2D UV = UVsIn.Num() > Idx ? FVector2D(UVsIn[Idx][0], UVsIn[Idx][1]) : FVector2D(0);
+			UVs[Idx].Add(UV);
+			
 			Colors[Idx] = FLinearColor::White;
 		}
 
