Fix the big inertia tensor scaling issue
- Use bInertiaScaleIncludeMass=true for world solver
- Diagonalize inertia tensors for RBAN solver
- Fix ColorConstraints
- Copy default settings from RBAN to world for PositionTolerance, AngleTolerance, MinParentMassRatio, MaxInertiaRatio
Note that one also needs to use 1 joint pushout iteration in order to have the same default RBAN behavior.

diff --git a/Engine/Source/Runtime/Engine/Private/PhysicsEngine/BodyInstance.cpp b/Engine/Source/Runtime/Engine/Private/PhysicsEngine/BodyInstance.cpp
--- a/Engine/Source/Runtime/Engine/Private/PhysicsEngine/BodyInstance.cpp
+++ b/Engine/Source/Runtime/Engine/Private/PhysicsEngine/BodyInstance.cpp
@@ -3209,7 +3209,7 @@ void FBodyInstance::UpdateMassProperties()
 						// The component scale is already built into the geometry, but if the user has set up a CoM
 						// modifier, it will need to be transformed by the component scale.
 						FTransform MassModifierTransform(FQuat::Identity, FVector(0.f, 0.f, 0.f), Scale3D);
-						TotalMassProperties = BodyUtils::ComputeMassProperties(this, Shapes, MassModifierTransform);
+						TotalMassProperties = BodyUtils::ComputeMassProperties(this, Shapes, MassModifierTransform, false/*bInertiaScaleIncludeMass*/);
 
 #if WITH_CHAOS
 						// @todo(chaos): this is starting to get messy - we should probably just split off the PhysX implementation now
diff --git a/Engine/Source/Runtime/Engine/Private/PhysicsEngine/ImmediatePhysics/ImmediatePhysicsChaos/ImmediatePhysicsActorHandle_Chaos.cpp b/Engine/Source/Runtime/Engine/Private/PhysicsEngine/ImmediatePhysics/ImmediatePhysicsChaos/ImmediatePhysicsActorHandle_Chaos.cpp
--- a/Engine/Source/Runtime/Engine/Private/PhysicsEngine/ImmediatePhysics/ImmediatePhysicsChaos/ImmediatePhysicsActorHandle_Chaos.cpp
+++ b/Engine/Source/Runtime/Engine/Private/PhysicsEngine/ImmediatePhysics/ImmediatePhysicsChaos/ImmediatePhysicsActorHandle_Chaos.cpp
@@ -182,6 +182,7 @@ namespace ImmediatePhysics_Chaos
 			// as well as the dimension change even though we don't actually change the mass.
 			const bool bInertaScaleIncludeMass = true;
 			FMassProperties MassProperties = BodyUtils::ComputeMassProperties(BodyInstance, Shapes, bContributesToMass, FTransform::Identity, bInertaScaleIncludeMass);
+			MassProperties.RotationOfMass = Chaos::TransformToLocalSpace(MassProperties.InertiaTensor);
 			OutMass = MassProperties.Mass;
 			OutInertia = MassProperties.InertiaTensor.GetDiagonal();
 			OutCoMTransform = FTransform(MassProperties.RotationOfMass, MassProperties.CenterOfMass);
diff --git a/Engine/Source/Runtime/Engine/Private/PhysicsEngine/ImmediatePhysics/ImmediatePhysicsChaos/ImmediatePhysicsSimulation_Chaos.cpp b/Engine/Source/Runtime/Engine/Private/PhysicsEngine/ImmediatePhysics/ImmediatePhysicsChaos/ImmediatePhysicsSimulation_Chaos.cpp
--- a/Engine/Source/Runtime/Engine/Private/PhysicsEngine/ImmediatePhysics/ImmediatePhysicsChaos/ImmediatePhysicsSimulation_Chaos.cpp
+++ b/Engine/Source/Runtime/Engine/Private/PhysicsEngine/ImmediatePhysics/ImmediatePhysicsChaos/ImmediatePhysicsSimulation_Chaos.cpp
@@ -107,7 +107,7 @@ FAutoConsoleVariableRef CVarChaosImmPhysJointPairIterations(TEXT("p.Chaos.ImmPhy
 FAutoConsoleVariableRef CVarChaosImmPhysJointPushOutPairIterations(TEXT("p.Chaos.ImmPhys.Joint.PushOutPairIterations"), ChaosImmediate_Joint_PushOutPairIterations, TEXT("Override joint push-out pair iterations (if >= 0)"));
 FAutoConsoleVariableRef CVarChaosImmPhysJointSwingTwistAngleTolerance(TEXT("p.Chaos.ImmPhys.Joint.SwingTwistAngleTolerance"), ChaosImmediate_Joint_SwingTwistAngleTolerance, TEXT("SwingTwistAngleTolerance."));
 FAutoConsoleVariableRef CVarChaosImmPhysJointPositionTolerance(TEXT("p.Chaos.ImmPhys.Joint.PositionTolerance"), ChaosImmediate_Joint_PositionTolerance, TEXT("PositionTolerance."));
-FAutoConsoleVariableRef CVarChaosImmPhysJointAngleTolerance(TEXT("p.Chaos.ImmPhys.Joint.AngleTolerance"), ChaosImmediate_Joint_AngleTolerance, TEXT("PositionTolerance."));
+FAutoConsoleVariableRef CVarChaosImmPhysJointAngleTolerance(TEXT("p.Chaos.ImmPhys.Joint.AngleTolerance"), ChaosImmediate_Joint_AngleTolerance, TEXT("AngleTolerance."));
 FAutoConsoleVariableRef CVarChaosImmPhysJointEnableTwistLimits(TEXT("p.Chaos.ImmPhys.Joint.EnableTwistLimits"), ChaosImmediate_Joint_EnableTwistLimits, TEXT("EnableTwistLimits."));
 FAutoConsoleVariableRef CVarChaosImmPhysJointEnableSwingLimits(TEXT("p.Chaos.ImmPhys.Joint.EnableSwingLimits"), ChaosImmediate_Joint_EnableSwingLimits, TEXT("EnableSwingLimits."));
 FAutoConsoleVariableRef CVarChaosImmPhysJointEnableDrives(TEXT("p.Chaos.ImmPhys.Joint.EnableDrives"), ChaosImmediate_Joint_EnableDrives, TEXT("EnableDrives."));
diff --git a/Engine/Source/Runtime/Experimental/Chaos/Private/Chaos/PBDJointConstraints.cpp b/Engine/Source/Runtime/Experimental/Chaos/Private/Chaos/PBDJointConstraints.cpp
--- a/Engine/Source/Runtime/Experimental/Chaos/Private/Chaos/PBDJointConstraints.cpp
+++ b/Engine/Source/Runtime/Experimental/Chaos/Private/Chaos/PBDJointConstraints.cpp
@@ -1560,14 +1560,16 @@ namespace Chaos
 
 			TPBDRigidParticleHandle<FReal, 3>* Particle0 = ConstraintParticles[ConstraintIndex][0]->CastToRigidParticle();
 			TPBDRigidParticleHandle<FReal, 3>* Particle1 = ConstraintParticles[ConstraintIndex][1]->CastToRigidParticle();
+			bool IsParticle0Dynamic = (Particle0 != nullptr) && (Particle0->ObjectState() == EObjectStateType::Dynamic || Particle0->ObjectState() == EObjectStateType::Sleeping);
+			bool IsParticle1Dynamic = (Particle1 != nullptr) && (Particle1->ObjectState() == EObjectStateType::Dynamic || Particle1->ObjectState() == EObjectStateType::Sleeping);
 
-			bool bContainsDynamic = (Particle0 != nullptr) || (Particle1 != nullptr);
+			bool bContainsDynamic = IsParticle0Dynamic || IsParticle1Dynamic;
 			if (bContainsDynamic)
 			{
 				ConstraintVertices[ConstraintIndex] = Graph.AddVertex();
 
 				// Set kinematic-connected constraints to level 0 to initialize level calculation
-				bool bContainsKinematic = (Particle0 == nullptr) || (Particle1 == nullptr);
+				bool bContainsKinematic = !IsParticle0Dynamic || !IsParticle1Dynamic;
 				if (bContainsKinematic)
 				{
 					Graph.SetVertexLevel(ConstraintVertices[ConstraintIndex], 0);
@@ -1608,10 +1610,18 @@ namespace Chaos
 			{
 				int32 ConstraintIndex0 = ParticleConstraintIndices[ParticleConstraintIndex0];
 				int32 VertexIndex0 = ConstraintVertices[ConstraintIndex0];
+				if(VertexIndex0 == INDEX_NONE)
+				{
+					continue;
+				}
 				for (int32 ParticleConstraintIndex1 = ParticleConstraintIndex0 + 1; ParticleConstraintIndex1 < NumParticleConstraintIndices; ++ParticleConstraintIndex1)
 				{
 					int32 ConstraintIndex1 = ParticleConstraintIndices[ParticleConstraintIndex1];
 					int32 VertexIndex1 = ConstraintVertices[ConstraintIndex1];
+					if(VertexIndex1 == INDEX_NONE)
+					{
+						continue;
+					}
 					Graph.AddEdge(VertexIndex0, VertexIndex1);
 				}
 			}
diff --git a/Engine/Source/Runtime/Experimental/Chaos/Private/PBDRigidsSolver.cpp b/Engine/Source/Runtime/Experimental/Chaos/Private/PBDRigidsSolver.cpp
--- a/Engine/Source/Runtime/Experimental/Chaos/Private/PBDRigidsSolver.cpp
+++ b/Engine/Source/Runtime/Experimental/Chaos/Private/PBDRigidsSolver.cpp
@@ -154,6 +154,15 @@ FAutoConsoleVariableRef CVarChaosSolverJointPairIterations(TEXT("p.Chaos.Solver.
 int32 ChaosSolverJointPushOutPairIterations = -1;
 FAutoConsoleVariableRef CVarChaosSolverJointPushOutPairIterations(TEXT("p.Chaos.Solver.Joint.PushOutPairIterations"), ChaosSolverJointPushOutPairIterations, TEXT("Override number of push out iterations per joint during a solver iteration (-1 to use config)"));
 
+// Copied from RBAN
+Chaos::FRealSingle ChaosSolverJointPositionTolerance = 0.025f;
+FAutoConsoleVariableRef CVarChaosSolverJointPositionTolerance(TEXT("p.Chaos.Solver.Joint.PositionTolerance"), ChaosSolverJointPositionTolerance, TEXT("PositionTolerance."));
+Chaos::FRealSingle ChaosSolverJointAngleTolerance = 0.001f;
+FAutoConsoleVariableRef CVarChaosSolverJointAngleTolerance(TEXT("p.Chaos.Solver.Joint.AngleTolerance"), ChaosSolverJointAngleTolerance, TEXT("AngleTolerance."));
+Chaos::FRealSingle ChaosSolverJointMinParentMassRatio = 0.2f;
+FAutoConsoleVariableRef CVarChaosSolverJointMinParentMassRatio(TEXT("p.Chaos.Solver.Joint.MinParentMassRatio"), ChaosSolverJointMinParentMassRatio, TEXT("6Dof joint MinParentMassRatio (if > 0)"));
+Chaos::FRealSingle ChaosSolverJointMaxInertiaRatio = 5.0f;
+FAutoConsoleVariableRef CVarChaosSolverJointMaxInertiaRatio(TEXT("p.Chaos.Solver.Joint.MaxInertiaRatio"), ChaosSolverJointMaxInertiaRatio, TEXT("6Dof joint MaxInertiaRatio (if > 0)"));
 
 // Collision detection cvars
 // These override the engine config if >= 0
@@ -756,6 +765,12 @@ namespace Chaos
 			{
 				SetCollisionCullDistance(ChaosSolverCullDistance);
 			}
+			// Joint settings overrides
+			JointsSettings.PositionTolerance = ChaosSolverJointPositionTolerance;
+			JointsSettings.AngleTolerance = ChaosSolverJointAngleTolerance;
+			JointsSettings.MinParentMassRatio = ChaosSolverJointMinParentMassRatio;
+			JointsSettings.MaxInertiaRatio = ChaosSolverJointMaxInertiaRatio;
+			JointConstraints.SetSettings(JointsSettings);
 		}
 
 		UE_LOG(LogPBDRigidsSolver, Verbose, TEXT("PBDRigidsSolver::Tick(%3.5f)"), DeltaTime);
