GeometryCollectionComponent
-- Removed comment in code preventing kinematic fields from expressing.

diff --git a/Engine/Source/Runtime/Experimental/GeometryCollectionEngine/Private/GeometryCollection/GeometryCollectionComponent.cpp b/Engine/Source/Runtime/Experimental/GeometryCollectionEngine/Private/GeometryCollection/GeometryCollectionComponent.cpp
--- a/Engine/Source/Runtime/Experimental/GeometryCollectionEngine/Private/GeometryCollection/GeometryCollectionComponent.cpp
+++ b/Engine/Source/Runtime/Experimental/GeometryCollectionEngine/Private/GeometryCollection/GeometryCollectionComponent.cpp
@@ -2577,22 +2577,18 @@ void FScopedColorEdit::UpdateBoneColors()
 
 void UGeometryCollectionComponent::ApplyKinematicField(float Radius, FVector Position)
 {
-	/*
 	FName TargetName = GetGeometryCollectionPhysicsTypeName(EGeometryCollectionPhysicsTypeEnum::Chaos_DynamicState);
 	DispatchFieldCommand({ TargetName,new FRadialIntMask(Radius, Position, (int32)Chaos::EObjectStateType::Dynamic,
 		(int32)Chaos::EObjectStateType::Kinematic, ESetMaskConditionType::Field_Set_IFF_NOT_Interior) });
-		*/
 }
 
 void UGeometryCollectionComponent::ApplyPhysicsField(bool Enabled, EGeometryCollectionPhysicsTypeEnum Target, UFieldSystemMetaData* MetaData, UFieldNodeBase* Field)
 {
-	/*
 	if (Enabled && Field)
 	{
 		FFieldSystemCommand Command = FFieldObjectCommands::CreateFieldCommand(GetGeometryCollectionPhysicsTypeName(Target), Field, MetaData);
 		DispatchFieldCommand(Command);
 	}
-	*/
 }
 
 void UGeometryCollectionComponent::DispatchFieldCommand(const FFieldSystemCommand& InCommand)
