GeometryCollectionComponent
-- Corrected merge issue that prevented build

diff --git a/Engine/Source/Runtime/Experimental/GeometryCollectionEngine/Private/GeometryCollection/GeometryCollectionComponent.cpp b/Engine/Source/Runtime/Experimental/GeometryCollectionEngine/Private/GeometryCollection/GeometryCollectionComponent.cpp
--- a/Engine/Source/Runtime/Experimental/GeometryCollectionEngine/Private/GeometryCollection/GeometryCollectionComponent.cpp
+++ b/Engine/Source/Runtime/Experimental/GeometryCollectionEngine/Private/GeometryCollection/GeometryCollectionComponent.cpp
@@ -2577,20 +2577,21 @@ void FScopedColorEdit::UpdateBoneColors()
 
 void UGeometryCollectionComponent::ApplyKinematicField(float Radius, FVector Position)
 {
-	FName TargetName = GetGeometryCollectionPhysicsTypeName(EGeometryCollectionPhysicsTypeEnum::Chaos_DynamicState);
-	DispatchFieldCommand({ TargetName,new FRadialIntMask(Radius, Position, (int32)Chaos::EObjectStateType::Dynamic,
-		(int32)Chaos::EObjectStateType::Kinematic, ESetMaskConditionType::Field_Set_IFF_NOT_Interior) });
+	FFieldSystemCommand Command = FFieldObjectCommands::CreateFieldCommand(EFieldPhysicsType::Field_DynamicState, new FRadialIntMask(Radius, Position, (int32)Chaos::EObjectStateType::Dynamic,
+		(int32)Chaos::EObjectStateType::Kinematic, ESetMaskConditionType::Field_Set_IFF_NOT_Interior));
+	DispatchFieldCommand(Command);
 }
 
 void UGeometryCollectionComponent::ApplyPhysicsField(bool Enabled, EGeometryCollectionPhysicsTypeEnum Target, UFieldSystemMetaData* MetaData, UFieldNodeBase* Field)
 {
 	if (Enabled && Field)
 	{
-		FFieldSystemCommand Command = FFieldObjectCommands::CreateFieldCommand(GetGeometryCollectionPhysicsTypeName(Target), Field, MetaData);
+		FFieldSystemCommand Command = FFieldObjectCommands::CreateFieldCommand(GetGeometryCollectionPhysicsType(Target), Field, MetaData);
 		DispatchFieldCommand(Command);
 	}
 }
 
+
 void UGeometryCollectionComponent::DispatchFieldCommand(const FFieldSystemCommand& InCommand)
 {
 	if (PhysicsProxy)
