fix crash w/ gc->sm->gc path creating out-of-bounds material IDs

#rnx

diff --git a/Engine/Source/Runtime/Experimental/GeometryCollectionEngine/Private/GeometryCollection/GeometryCollectionObject.cpp b/Engine/Source/Runtime/Experimental/GeometryCollectionEngine/Private/GeometryCollection/GeometryCollectionObject.cpp
--- a/Engine/Source/Runtime/Experimental/GeometryCollectionEngine/Private/GeometryCollection/GeometryCollectionObject.cpp
+++ b/Engine/Source/Runtime/Experimental/GeometryCollectionEngine/Private/GeometryCollection/GeometryCollectionObject.cpp
@@ -251,8 +251,19 @@ void UGeometryCollection::InitializeMaterials()
 	// Last Material is the selection one
 	UMaterialInterface* BoneSelectedMaterial = LoadObject<UMaterialInterface>(nullptr, GetSelectedMaterialPath(), nullptr, LOAD_None, nullptr);
 
-	// Skip selection materials
-	Materials.Remove(BoneSelectedMaterial);
+	TManagedArray<int32>& MaterialID = GeometryCollection->MaterialID;
+
+	// normally we filter out instances of the selection material ID, but if it's actually used on any face we have to keep it
+	bool bBoneSelectedMaterialIsUsed = false;
+	for (int32 FaceIdx = 0; FaceIdx < MaterialID.Num(); ++FaceIdx)
+	{
+		int32 FaceMaterialID = MaterialID[FaceIdx];
+		if (FaceMaterialID < Materials.Num() && Materials[FaceMaterialID] == BoneSelectedMaterial)
+		{
+			bBoneSelectedMaterialIsUsed = true;
+			break;
+		}
+	}
 	
 	// We're assuming that all materials are arranged in pairs, so first we collect these.
 	using FMaterialPair = TPair<UMaterialInterface*, UMaterialInterface*>;
@@ -260,12 +271,21 @@ void UGeometryCollection::InitializeMaterials()
 	for (int32 MaterialIndex = 0; MaterialIndex < Materials.Num(); ++MaterialIndex)
 	{
 		UMaterialInterface* ExteriorMaterial = Materials[MaterialIndex];
+		if (ExteriorMaterial == BoneSelectedMaterial && !bBoneSelectedMaterialIsUsed) // skip unused bone selected material
+		{
+			continue;
+		}
 		
 		// If we have an odd number of materials, the last material duplicates itself.
 		UMaterialInterface* InteriorMaterial = Materials[MaterialIndex];
-		if (++MaterialIndex < Materials.Num())
+		while (++MaterialIndex < Materials.Num())
 		{
+			if (Materials[MaterialIndex] == BoneSelectedMaterial && !bBoneSelectedMaterialIsUsed) // skip bone selected material
+			{
+				continue;
+			}
 			InteriorMaterial = Materials[MaterialIndex];
+			break;
 		}
 
 		MaterialSet.Add(FMaterialPair(ExteriorMaterial, InteriorMaterial));
@@ -287,8 +307,6 @@ void UGeometryCollection::InitializeMaterials()
 		InteriorMaterialPtrToArrayIndex.Add(InteriorTuple);
 	}
 
-	TManagedArray<int32>& MaterialID = GeometryCollection->MaterialID;
-
 	// Reassign material ID for each face given the new consolidated array of materials
 	for (int32 Material = 0; Material < MaterialID.Num(); ++Material)
 	{
