Fix wrong field boxfalloff bbox + fix field particlehandles construction + fix field commands removal

diff --git a/Engine/Source/Runtime/Engine/Private/PhysicsField/PhysicsFieldComponent.cpp b/Engine/Source/Runtime/Engine/Private/PhysicsField/PhysicsFieldComponent.cpp
--- a/Engine/Source/Runtime/Engine/Private/PhysicsField/PhysicsFieldComponent.cpp
+++ b/Engine/Source/Runtime/Engine/Private/PhysicsField/PhysicsFieldComponent.cpp
@@ -1062,8 +1062,8 @@ void FPhysicsFieldInstance::BuildNodeBounds(FFieldNodeBase* FieldNode, FVector&
 		{
 			FBoxFalloff* LocalNode = StaticCast<FBoxFalloff*>(FieldNode);
 
-			MinBounds = (LocalNode->Default == 0) ? LocalNode->Transform.GetTranslation() - LocalNode->Transform.GetScale3D() : FVector(-FLT_MAX);
-			MaxBounds = (LocalNode->Default == 0) ? LocalNode->Transform.GetTranslation() + LocalNode->Transform.GetScale3D() : FVector(FLT_MAX);
+			MinBounds = (LocalNode->Default == 0) ? LocalNode->Transform.GetTranslation() - 50 * LocalNode->Transform.GetScale3D() : FVector(-FLT_MAX);
+			MaxBounds = (LocalNode->Default == 0) ? LocalNode->Transform.GetTranslation() + 50 * LocalNode->Transform.GetScale3D() : FVector(FLT_MAX);
 		}
 		else if (FieldNode->SerializationType() == FFieldNodeBase::ESerializationType::FieldNode_FSumScalar)
 		{
diff --git a/Engine/Source/Runtime/Experimental/Chaos/Private/PhysicsProxy/GeometryCollectionPhysicsProxy.cpp b/Engine/Source/Runtime/Experimental/Chaos/Private/PhysicsProxy/GeometryCollectionPhysicsProxy.cpp
--- a/Engine/Source/Runtime/Experimental/Chaos/Private/PhysicsProxy/GeometryCollectionPhysicsProxy.cpp
+++ b/Engine/Source/Runtime/Experimental/Chaos/Private/PhysicsProxy/GeometryCollectionPhysicsProxy.cpp
@@ -2618,132 +2618,133 @@ void FGeometryCollectionPhysicsProxy::FieldParameterUpdateCallback(Chaos::FPBDRi
 		for (int32 CommandIndex = 0; CommandIndex < NumCommands; CommandIndex++)
 		{
 			FFieldSystemCommand& FieldCommand = Commands[CommandIndex];
-			if (Chaos::BuildFieldSamplePoints(this, RigidSolver, FieldCommand, ExecutionDatas, PrevResolutionType, PrevFilterType, PrevObjectType, PrevPositionType))
+			if (IsParameterFieldValid(FieldCommand) || FieldCommand.PhysicsType == EFieldPhysicsType::Field_InitialLinearVelocity || FieldCommand.PhysicsType == EFieldPhysicsType::Field_InitialAngularVelocity)
 			{
-				const float TimeSeconds = RigidSolver->GetSolverTime() - FieldCommand.TimeCreation;
+				if (Chaos::BuildFieldSamplePoints(this, RigidSolver, FieldCommand, ExecutionDatas, PrevResolutionType, PrevFilterType, PrevObjectType, PrevPositionType))
+				{
+					const float TimeSeconds = RigidSolver->GetSolverTime() - FieldCommand.TimeCreation;
 
-				FFieldContext FieldContext(
-					ExecutionDatas,
-					FieldCommand.MetaData,
-					TimeSeconds);
+					FFieldContext FieldContext(
+						ExecutionDatas,
+						FieldCommand.MetaData,
+						TimeSeconds);
 
-				TArray<Chaos::FGeometryParticleHandle*>& ParticleHandles = ExecutionDatas.ParticleHandles[(uint8)EFieldCommandHandlesType::InsideHandles];
+					TArray<Chaos::FGeometryParticleHandle*>& ParticleHandles = ExecutionDatas.ParticleHandles[(uint8)EFieldCommandHandlesType::InsideHandles];
 
-				if (FieldCommand.RootNode->Type() == FFieldNodeBase::EFieldType::EField_Int32)
-				{
-					TArray<int32>& FinalResults = ExecutionDatas.IntegerResults[(uint8)EFieldCommandResultType::FinalResult];
-					ResetResultsArray < int32 >(ExecutionDatas.SamplePositions.Num(), FinalResults, 0);
+					if (FieldCommand.RootNode->Type() == FFieldNodeBase::EFieldType::EField_Int32)
+					{
+						TArray<int32>& FinalResults = ExecutionDatas.IntegerResults[(uint8)EFieldCommandResultType::FinalResult];
+						ResetResultsArray < int32 >(ExecutionDatas.SamplePositions.Num(), FinalResults, 0);
 
-					TArrayView<int32> ResultsView(&(FinalResults[0]), FinalResults.Num());
+						TArrayView<int32> ResultsView(&(FinalResults[0]), FinalResults.Num());
 
-					if (FieldCommand.PhysicsType == EFieldPhysicsType::Field_DynamicState)
-					{
-						SCOPE_CYCLE_COUNTER(STAT_ParamUpdateField_DynamicState);
+						if (FieldCommand.PhysicsType == EFieldPhysicsType::Field_DynamicState)
 						{
-							bool bHasStateChanged = false;
-							InitDynamicStateResults(ParticleHandles, FieldContext, FinalResults);
-
-							static_cast<const FFieldNode<int32>*>(FieldCommand.RootNode.Get())->Evaluate(FieldContext, ResultsView);
-							for (const FFieldContextIndex& Index : FieldContext.GetEvaluatedSamples())
+							SCOPE_CYCLE_COUNTER(STAT_ParamUpdateField_DynamicState);
 							{
-								Chaos::TPBDRigidParticleHandle<Chaos::FReal, 3>* RigidHandle = ParticleHandles[Index.Sample]->CastToRigidParticle();
-								if (RigidHandle)
-								{
-									const int32 CurrResult = ResultsView[Index.Result];
-									check(CurrResult <= std::numeric_limits<int8>::max() &&
-										  CurrResult >= std::numeric_limits<int8>::min());
+								bool bHasStateChanged = false;
+								InitDynamicStateResults(ParticleHandles, FieldContext, FinalResults);
 
-									const int8 ResultState = static_cast<int8>(CurrResult);
-									const int32 TransformIndex = HandleToTransformGroupIndex[RigidHandle];
-
-									// Update of the handles object state. No need to update 
-									// the initial velocities since it is done after this function call in InitializeBodiesPT
-									if (bUpdateViews && (Parameters.InitialVelocityType == EInitialVelocityTypeEnum::Chaos_Initial_Velocity_User_Defined))
-									{
-										bHasStateChanged |= ReportDynamicStateResult(RigidSolver, static_cast<Chaos::EObjectStateType>(ResultState), RigidHandle,
-											true, Collection.InitialLinearVelocity[TransformIndex], true, Collection.InitialAngularVelocity[TransformIndex]);
-									}
-									else
+								static_cast<const FFieldNode<int32>*>(FieldCommand.RootNode.Get())->Evaluate(FieldContext, ResultsView);
+								for (const FFieldContextIndex& Index : FieldContext.GetEvaluatedSamples())
+								{
+									Chaos::TPBDRigidParticleHandle<Chaos::FReal, 3>* RigidHandle = ParticleHandles[Index.Sample]->CastToRigidParticle();
+									if (RigidHandle)
 									{
-										bHasStateChanged |= ReportDynamicStateResult(RigidSolver, static_cast<Chaos::EObjectStateType>(ResultState), RigidHandle,
-											false, Chaos::FVec3(0), false, Chaos::FVec3(0));
+										const int32 CurrResult = ResultsView[Index.Result];
+										check(CurrResult <= std::numeric_limits<int8>::max() &&
+											CurrResult >= std::numeric_limits<int8>::min());
+
+										const int8 ResultState = static_cast<int8>(CurrResult);
+										const int32 TransformIndex = HandleToTransformGroupIndex[RigidHandle];
+
+										// Update of the handles object state. No need to update 
+										// the initial velocities since it is done after this function call in InitializeBodiesPT
+										if (bUpdateViews && (Parameters.InitialVelocityType == EInitialVelocityTypeEnum::Chaos_Initial_Velocity_User_Defined))
+										{
+											bHasStateChanged |= ReportDynamicStateResult(RigidSolver, static_cast<Chaos::EObjectStateType>(ResultState), RigidHandle,
+												true, Collection.InitialLinearVelocity[TransformIndex], true, Collection.InitialAngularVelocity[TransformIndex]);
+										}
+										else
+										{
+											bHasStateChanged |= ReportDynamicStateResult(RigidSolver, static_cast<Chaos::EObjectStateType>(ResultState), RigidHandle,
+												false, Chaos::FVec3(0), false, Chaos::FVec3(0));
+										}
+										// Update of the Collection dynamic state. It will be used just after to set the initial velocity
+										Collection.DynamicState[TransformIndex] = ResultState;
 									}
-									// Update of the Collection dynamic state. It will be used just after to set the initial velocity
-									Collection.DynamicState[TransformIndex] = ResultState;
+								}
+								if (bUpdateViews)
+								{
+									UpdateSolverParticlesState(RigidSolver, bHasStateChanged);
 								}
 							}
-							if (bUpdateViews)
-							{
-								UpdateSolverParticlesState(RigidSolver, bHasStateChanged);
-							}
 						}
-						CommandsToRemove.Add(CommandIndex);
+						else
+						{
+							Chaos::FieldIntegerParameterUpdate(RigidSolver, FieldCommand, ExecutionDatas.ParticleHandles[(uint8)EFieldCommandHandlesType::InsideHandles],
+								FieldContext, PositionTarget, TargetedParticles, FinalResults);
+						}
 					}
-					else
+					else if (FieldCommand.RootNode->Type() == FFieldNodeBase::EFieldType::EField_FVector)
 					{
-						Chaos::FieldIntegerParameterUpdate(RigidSolver, FieldCommand, ExecutionDatas.ParticleHandles[(uint8)EFieldCommandHandlesType::InsideHandles],
-							FieldContext, CommandsToRemove, PositionTarget, TargetedParticles, CommandIndex, FinalResults);
-					}
-				}
-				else if (FieldCommand.RootNode->Type() == FFieldNodeBase::EFieldType::EField_FVector)
-				{
-					TArray<FVector>& FinalResults = ExecutionDatas.VectorResults[(uint8)EFieldCommandResultType::FinalResult];
-					ResetResultsArray < FVector >(ExecutionDatas.SamplePositions.Num(), FinalResults, FVector::ZeroVector);
+						TArray<FVector>& FinalResults = ExecutionDatas.VectorResults[(uint8)EFieldCommandResultType::FinalResult];
+						ResetResultsArray < FVector >(ExecutionDatas.SamplePositions.Num(), FinalResults, FVector::ZeroVector);
 
-					TArrayView<FVector> ResultsView(&(FinalResults[0]), FinalResults.Num());
+						TArrayView<FVector> ResultsView(&(FinalResults[0]), FinalResults.Num());
 
-					if (FieldCommand.PhysicsType == EFieldPhysicsType::Field_InitialLinearVelocity)
-					{
-						if (Parameters.InitialVelocityType == EInitialVelocityTypeEnum::Chaos_Initial_Velocity_User_Defined)
+						if (FieldCommand.PhysicsType == EFieldPhysicsType::Field_InitialLinearVelocity)
 						{
-							SCOPE_CYCLE_COUNTER(STAT_ParamUpdateField_LinearVelocity);
+							if (Parameters.InitialVelocityType == EInitialVelocityTypeEnum::Chaos_Initial_Velocity_User_Defined)
 							{
-								static_cast<const FFieldNode<FVector>*>(FieldCommand.RootNode.Get())->Evaluate(FieldContext, ResultsView);
-								for (const FFieldContextIndex& Index : FieldContext.GetEvaluatedSamples())
+								SCOPE_CYCLE_COUNTER(STAT_ParamUpdateField_LinearVelocity);
 								{
-									Chaos::TPBDRigidParticleHandle<Chaos::FReal, 3>* RigidHandle = ParticleHandles[Index.Sample]->CastToRigidParticle();
-									if (RigidHandle)
+									static_cast<const FFieldNode<FVector>*>(FieldCommand.RootNode.Get())->Evaluate(FieldContext, ResultsView);
+									for (const FFieldContextIndex& Index : FieldContext.GetEvaluatedSamples())
 									{
-										Collection.InitialLinearVelocity[HandleToTransformGroupIndex[RigidHandle]] = ResultsView[Index.Result];
+										Chaos::TPBDRigidParticleHandle<Chaos::FReal, 3>* RigidHandle = ParticleHandles[Index.Sample]->CastToRigidParticle();
+										if (RigidHandle)
+										{
+											Collection.InitialLinearVelocity[HandleToTransformGroupIndex[RigidHandle]] = ResultsView[Index.Result];
+										}
 									}
 								}
 							}
-							CommandsToRemove.Add(CommandIndex);
-						}
-						else
-						{
-							UE_LOG(LogChaos, Error, TEXT("Field based evaluation of the simulations 'InitialLinearVelocity' requires the geometry collection be set to User Defined Initial Velocity"));
+							else
+							{
+								UE_LOG(LogChaos, Error, TEXT("Field based evaluation of the simulations 'InitialLinearVelocity' requires the geometry collection be set to User Defined Initial Velocity"));
+							}
 						}
-					}
-					else if (FieldCommand.PhysicsType == EFieldPhysicsType::Field_InitialAngularVelocity)
-					{
-						if (Parameters.InitialVelocityType == EInitialVelocityTypeEnum::Chaos_Initial_Velocity_User_Defined)
+						else if (FieldCommand.PhysicsType == EFieldPhysicsType::Field_InitialAngularVelocity)
 						{
-							SCOPE_CYCLE_COUNTER(STAT_ParamUpdateField_AngularVelocity);
+							if (Parameters.InitialVelocityType == EInitialVelocityTypeEnum::Chaos_Initial_Velocity_User_Defined)
 							{
-								static_cast<const FFieldNode<FVector>*>(FieldCommand.RootNode.Get())->Evaluate(FieldContext, ResultsView);
-								for (const FFieldContextIndex& Index : FieldContext.GetEvaluatedSamples())
+								SCOPE_CYCLE_COUNTER(STAT_ParamUpdateField_AngularVelocity);
 								{
-									Chaos::TPBDRigidParticleHandle<Chaos::FReal, 3>* RigidHandle = ParticleHandles[Index.Sample]->CastToRigidParticle();
-									if (RigidHandle)
+									static_cast<const FFieldNode<FVector>*>(FieldCommand.RootNode.Get())->Evaluate(FieldContext, ResultsView);
+									for (const FFieldContextIndex& Index : FieldContext.GetEvaluatedSamples())
 									{
-										Collection.InitialAngularVelocity[HandleToTransformGroupIndex[RigidHandle]] = ResultsView[Index.Result];
+										Chaos::TPBDRigidParticleHandle<Chaos::FReal, 3>* RigidHandle = ParticleHandles[Index.Sample]->CastToRigidParticle();
+										if (RigidHandle)
+										{
+											Collection.InitialAngularVelocity[HandleToTransformGroupIndex[RigidHandle]] = ResultsView[Index.Result];
+										}
 									}
 								}
 							}
-							CommandsToRemove.Add(CommandIndex);
+							else
+							{
+								UE_LOG(LogChaos, Error, TEXT("Field based evaluation of the simulations 'InitialAngularVelocity' requires the geometry collection be set to User Defined Initial Velocity"));
+							}
 						}
 						else
 						{
-							UE_LOG(LogChaos, Error, TEXT("Field based evaluation of the simulations 'InitialAngularVelocity' requires the geometry collection be set to User Defined Initial Velocity"));
+							Chaos::FieldVectorParameterUpdate(RigidSolver, FieldCommand, ParticleHandles,
+								FieldContext, PositionTarget, TargetedParticles, FinalResults);
 						}
 					}
-					else 
-					{
-						Chaos::FieldVectorParameterUpdate(RigidSolver, FieldCommand, ParticleHandles,
-							FieldContext, CommandsToRemove, PositionTarget, TargetedParticles, CommandIndex, FinalResults);
-					}
 				}
+				CommandsToRemove.Add(CommandIndex);
 			}
 		}		
 		
@@ -2772,28 +2773,31 @@ void FGeometryCollectionPhysicsProxy::FieldForcesUpdateCallback(Chaos::FPBDRigid
 		for (int32 CommandIndex = 0; CommandIndex < NumCommands; CommandIndex++)
 		{
 			const FFieldSystemCommand& FieldCommand = Commands[CommandIndex];
-
-			if (Chaos::BuildFieldSamplePoints(this, RigidSolver, FieldCommand, ExecutionDatas, PrevResolutionType, PrevFilterType, PrevObjectType, PrevPositionType))
+			if (IsForceFieldValid(FieldCommand))
 			{
-				const float TimeSeconds = RigidSolver->GetSolverTime() - FieldCommand.TimeCreation;
+				if (Chaos::BuildFieldSamplePoints(this, RigidSolver, FieldCommand, ExecutionDatas, PrevResolutionType, PrevFilterType, PrevObjectType, PrevPositionType))
+				{
+					const float TimeSeconds = RigidSolver->GetSolverTime() - FieldCommand.TimeCreation;
 
-				FFieldContext FieldContext(
-					ExecutionDatas,
-					FieldCommand.MetaData,
-					TimeSeconds);
+					FFieldContext FieldContext(
+						ExecutionDatas,
+						FieldCommand.MetaData,
+						TimeSeconds);
 
-				TArray<Chaos::FGeometryParticleHandle*>& ParticleHandles = ExecutionDatas.ParticleHandles[(uint8)EFieldCommandHandlesType::InsideHandles];
+					TArray<Chaos::FGeometryParticleHandle*>& ParticleHandles = ExecutionDatas.ParticleHandles[(uint8)EFieldCommandHandlesType::InsideHandles];
 
-				if (FieldCommand.RootNode->Type() == FFieldNode<FVector>::StaticType())
-				{
-					TArray<FVector>& FinalResults = ExecutionDatas.VectorResults[(uint8)EFieldCommandResultType::FinalResult];
-					ResetResultsArray < FVector >(ExecutionDatas.SamplePositions.Num(), FinalResults, FVector::ZeroVector);
+					if (FieldCommand.RootNode->Type() == FFieldNode<FVector>::StaticType())
+					{
+						TArray<FVector>& FinalResults = ExecutionDatas.VectorResults[(uint8)EFieldCommandResultType::FinalResult];
+						ResetResultsArray < FVector >(ExecutionDatas.SamplePositions.Num(), FinalResults, FVector::ZeroVector);
 
-					TArrayView<FVector> ResultsView(&(FinalResults[0]), FinalResults.Num());
+						TArrayView<FVector> ResultsView(&(FinalResults[0]), FinalResults.Num());
 
-					Chaos::FieldVectorForceUpdate(RigidSolver, FieldCommand, ParticleHandles,
-						FieldContext, CommandsToRemove, CommandIndex, FinalResults);
+						Chaos::FieldVectorForceUpdate(RigidSolver, FieldCommand, ParticleHandles,
+							FieldContext, FinalResults);
+					}
 				}
+				CommandsToRemove.Add(CommandIndex);
 			}
 		}
 		for (int32 Index = CommandsToRemove.Num() - 1; Index >= 0; --Index)
diff --git a/Engine/Source/Runtime/Experimental/Chaos/Private/PhysicsProxy/PerSolverFieldSystem.cpp b/Engine/Source/Runtime/Experimental/Chaos/Private/PhysicsProxy/PerSolverFieldSystem.cpp
--- a/Engine/Source/Runtime/Experimental/Chaos/Private/PhysicsProxy/PerSolverFieldSystem.cpp
+++ b/Engine/Source/Runtime/Experimental/Chaos/Private/PhysicsProxy/PerSolverFieldSystem.cpp
@@ -25,31 +25,6 @@ void ResetIndicesArray(TArray<int32>& IndicesArray, int32 Size)
 // FPerSolverFieldSystem
 //==============================================================================
 
-FORCEINLINE bool IsParameterFieldValid(const FFieldSystemCommand& FieldCommand)
-{
-	if (FieldCommand.RootNode->Type() == FFieldNodeBase::EFieldType::EField_Int32)
-	{
-		return (FieldCommand.PhysicsType == EFieldPhysicsType::Field_DynamicState) ||
-			   (FieldCommand.PhysicsType == EFieldPhysicsType::Field_ActivateDisabled) ||
-			   (FieldCommand.PhysicsType == EFieldPhysicsType::Field_CollisionGroup);
-
-	}
-	else if (FieldCommand.RootNode->Type() == FFieldNodeBase::EFieldType::EField_Float)
-	{
-		return (FieldCommand.PhysicsType == EFieldPhysicsType::Field_ExternalClusterStrain) ||
-			   (FieldCommand.PhysicsType == EFieldPhysicsType::Field_Kill) ||
-			   (FieldCommand.PhysicsType == EFieldPhysicsType::Field_SleepingThreshold) ||
-			   (FieldCommand.PhysicsType == EFieldPhysicsType::Field_DisableThreshold) ||
-			   (FieldCommand.PhysicsType == EFieldPhysicsType::Field_InternalClusterStrain);
-	}
-	else if (FieldCommand.RootNode->Type() == FFieldNodeBase::EFieldType::EField_FVector)
-	{
-		return (FieldCommand.PhysicsType == EFieldPhysicsType::Field_LinearVelocity) ||
-			   (FieldCommand.PhysicsType == EFieldPhysicsType::Field_AngularVelociy);
-	}
-	return false;
-}
-
 void FPerSolverFieldSystem::FieldParameterUpdateInternal(
 	Chaos::FPBDRigidsSolver* RigidSolver,
 	Chaos::FPBDPositionConstraints& PositionTarget,
@@ -72,41 +47,45 @@ void FPerSolverFieldSystem::FieldParameterUpdateInternal(
 		for (int32 CommandIndex = 0; CommandIndex < NumCommands; CommandIndex++)
 		{
 			const FFieldSystemCommand& FieldCommand = Commands[CommandIndex];
-			if(IsParameterFieldValid(FieldCommand) && Chaos::BuildFieldSamplePoints(this, RigidSolver, FieldCommand, ExecutionDatas, PrevResolutionType, PrevFilterType, PrevObjectType, PrevPositionType) )
+			if(IsParameterFieldValid(FieldCommand))
 			{
-				const float TimeSeconds = RigidSolver->GetSolverTime() - FieldCommand.TimeCreation;
+				if (Chaos::BuildFieldSamplePoints(this, RigidSolver, FieldCommand, ExecutionDatas, PrevResolutionType, PrevFilterType, PrevObjectType, PrevPositionType))
+				{
+					const float TimeSeconds = RigidSolver->GetSolverTime() - FieldCommand.TimeCreation;
 
-				FFieldContext FieldContext(
-					ExecutionDatas,
-					FieldCommand.MetaData,
-					TimeSeconds);
+					FFieldContext FieldContext(
+						ExecutionDatas,
+						FieldCommand.MetaData,
+						TimeSeconds);
 
-				TArray<Chaos::FGeometryParticleHandle*>& ParticleHandles = ExecutionDatas.ParticleHandles[(uint8)EFieldCommandHandlesType::InsideHandles];
+					TArray<Chaos::FGeometryParticleHandle*>& ParticleHandles = ExecutionDatas.ParticleHandles[(uint8)EFieldCommandHandlesType::InsideHandles];
 
-				if (FieldCommand.RootNode->Type() == FFieldNodeBase::EFieldType::EField_Int32)
-				{
-					TArray<int32>& FinalResults = ExecutionDatas.IntegerResults[(uint8)EFieldCommandResultType::FinalResult];
-					ResetResultsArray < int32 >(ExecutionDatas.SamplePositions.Num(), FinalResults, 0);
+					if (FieldCommand.RootNode->Type() == FFieldNodeBase::EFieldType::EField_Int32)
+					{
+						TArray<int32>& FinalResults = ExecutionDatas.IntegerResults[(uint8)EFieldCommandResultType::FinalResult];
+						ResetResultsArray < int32 >(ExecutionDatas.SamplePositions.Num(), FinalResults, 0);
 
-					Chaos::FieldIntegerParameterUpdate(RigidSolver, FieldCommand, ParticleHandles,
-						FieldContext, CommandsToRemove, PositionTarget, TargetedParticles, CommandIndex, FinalResults);
-				}
-				else if (FieldCommand.RootNode->Type() == FFieldNodeBase::EFieldType::EField_Float)
-				{
-					TArray<float>& FinalResults = ExecutionDatas.ScalarResults[(uint8)EFieldCommandResultType::FinalResult];
-					ResetResultsArray < float >(ExecutionDatas.SamplePositions.Num(), FinalResults, 0.0);
+						Chaos::FieldIntegerParameterUpdate(RigidSolver, FieldCommand, ParticleHandles,
+							FieldContext, PositionTarget, TargetedParticles, FinalResults);
+					}
+					else if (FieldCommand.RootNode->Type() == FFieldNodeBase::EFieldType::EField_Float)
+					{
+						TArray<float>& FinalResults = ExecutionDatas.ScalarResults[(uint8)EFieldCommandResultType::FinalResult];
+						ResetResultsArray < float >(ExecutionDatas.SamplePositions.Num(), FinalResults, 0.0);
 
-					Chaos::FieldScalarParameterUpdate(RigidSolver, FieldCommand, ParticleHandles,
-						FieldContext, CommandsToRemove, PositionTarget, TargetedParticles, CommandIndex, FinalResults);
-				}
-				else if (FieldCommand.RootNode->Type() == FFieldNodeBase::EFieldType::EField_FVector)
-				{
-					TArray<FVector>& FinalResults = ExecutionDatas.VectorResults[(uint8)EFieldCommandResultType::FinalResult];
-					ResetResultsArray < FVector >(ExecutionDatas.SamplePositions.Num(), FinalResults, FVector::ZeroVector);
+						Chaos::FieldScalarParameterUpdate(RigidSolver, FieldCommand, ParticleHandles,
+							FieldContext, PositionTarget, TargetedParticles, FinalResults);
+					}
+					else if (FieldCommand.RootNode->Type() == FFieldNodeBase::EFieldType::EField_FVector)
+					{
+						TArray<FVector>& FinalResults = ExecutionDatas.VectorResults[(uint8)EFieldCommandResultType::FinalResult];
+						ResetResultsArray < FVector >(ExecutionDatas.SamplePositions.Num(), FinalResults, FVector::ZeroVector);
 
-					Chaos::FieldVectorParameterUpdate(RigidSolver, FieldCommand, ParticleHandles,
-						FieldContext, CommandsToRemove, PositionTarget, TargetedParticles, CommandIndex, FinalResults);
+						Chaos::FieldVectorParameterUpdate(RigidSolver, FieldCommand, ParticleHandles,
+							FieldContext, PositionTarget, TargetedParticles, FinalResults);
+					}
 				}
+				CommandsToRemove.Add(CommandIndex);
 			}
 		}
 		if (IsTransient)
@@ -119,15 +98,7 @@ void FPerSolverFieldSystem::FieldParameterUpdateInternal(
 	}
 }
 
-FORCEINLINE bool IsForceFieldValid(const FFieldSystemCommand& FieldCommand)
-{
-	if (FieldCommand.RootNode->Type() == FFieldNodeBase::EFieldType::EField_FVector)
-	{
-		return (FieldCommand.PhysicsType == EFieldPhysicsType::Field_LinearForce) ||
-			   (FieldCommand.PhysicsType == EFieldPhysicsType::Field_AngularTorque);
-	}
-	return false;
-}
+
 
 void FPerSolverFieldSystem::FieldParameterUpdateCallback(
 	Chaos::FPBDRigidsSolver* InSolver,
@@ -159,25 +130,29 @@ void FPerSolverFieldSystem::FieldForcesUpdateInternal(
 		{
 			const FFieldSystemCommand& FieldCommand = Commands[CommandIndex];
 
-			if (IsForceFieldValid(FieldCommand) && Chaos::BuildFieldSamplePoints(this, RigidSolver, FieldCommand, ExecutionDatas, PrevResolutionType, PrevFilterType, PrevObjectType, PrevPositionType))
+			if (IsForceFieldValid(FieldCommand))
 			{
-				const float TimeSeconds = RigidSolver->GetSolverTime() - FieldCommand.TimeCreation;
+				if (Chaos::BuildFieldSamplePoints(this, RigidSolver, FieldCommand, ExecutionDatas, PrevResolutionType, PrevFilterType, PrevObjectType, PrevPositionType))
+				{
+					const float TimeSeconds = RigidSolver->GetSolverTime() - FieldCommand.TimeCreation;
 
-				FFieldContext FieldContext(
-					ExecutionDatas,
-					FieldCommand.MetaData,
-					TimeSeconds);
+					FFieldContext FieldContext(
+						ExecutionDatas,
+						FieldCommand.MetaData,
+						TimeSeconds);
 
-				TArray<Chaos::FGeometryParticleHandle*>& ParticleHandles = ExecutionDatas.ParticleHandles[(uint8)EFieldCommandHandlesType::InsideHandles];
+					TArray<Chaos::FGeometryParticleHandle*>& ParticleHandles = ExecutionDatas.ParticleHandles[(uint8)EFieldCommandHandlesType::InsideHandles];
 
-				if (FieldCommand.RootNode->Type() == FFieldNodeBase::EFieldType::EField_FVector)
-				{
-					TArray<FVector>& FinalResults = ExecutionDatas.VectorResults[(uint8)EFieldCommandResultType::FinalResult];
-					ResetResultsArray < FVector >(ExecutionDatas.SamplePositions.Num(), FinalResults, FVector::ZeroVector);
+					if (FieldCommand.RootNode->Type() == FFieldNodeBase::EFieldType::EField_FVector)
+					{
+						TArray<FVector>& FinalResults = ExecutionDatas.VectorResults[(uint8)EFieldCommandResultType::FinalResult];
+						ResetResultsArray < FVector >(ExecutionDatas.SamplePositions.Num(), FinalResults, FVector::ZeroVector);
 
-					Chaos::FieldVectorForceUpdate(RigidSolver, FieldCommand, ParticleHandles,
-						FieldContext, CommandsToRemove, CommandIndex, FinalResults);
+						Chaos::FieldVectorForceUpdate(RigidSolver, FieldCommand, ParticleHandles,
+							FieldContext, FinalResults);
+					}
 				}
+				CommandsToRemove.Add(CommandIndex);
 			}
 		}
 		if (IsTransient)
diff --git a/Engine/Source/Runtime/Experimental/Chaos/Public/PhysicsProxy/FieldSystemProxyHelper.h b/Engine/Source/Runtime/Experimental/Chaos/Public/PhysicsProxy/FieldSystemProxyHelper.h
--- a/Engine/Source/Runtime/Experimental/Chaos/Public/PhysicsProxy/FieldSystemProxyHelper.h
+++ b/Engine/Source/Runtime/Experimental/Chaos/Public/PhysicsProxy/FieldSystemProxyHelper.h
@@ -55,7 +55,7 @@ namespace Chaos
 		TArray<Chaos::FGeometryParticleHandle*>& FilteredHandles = ExecutionDatas.ParticleHandles[(uint8)EFieldCommandHandlesType::FilteredHandles];
 		TArray<Chaos::FGeometryParticleHandle*>& InsideHandles = ExecutionDatas.ParticleHandles[(uint8)EFieldCommandHandlesType::InsideHandles];
 
-		if (LocalProxy && ( (PrevResolutionType != ResolutionType) || (PrevFilterType != FilterType) || (PrevObjectType != ObjectType) || (PrevPositionType != PositionType) || FilteredHandles.Num() == 0))
+		//if (LocalProxy && ( (PrevResolutionType != ResolutionType) || (PrevFilterType != FilterType) || (PrevObjectType != ObjectType) || (PrevPositionType != PositionType) || FilteredHandles.Num() == 0))
 		{
 			if (FilterType != EFieldFilterType::Field_Filter_Max)
 			{
@@ -124,7 +124,7 @@ namespace Chaos
 	 * @param    FieldContext Field context to retrieve the evaluated samples
 	 * @param    LocalResults Array to store the dynamic state
 	 */
-	FORCEINLINE void InitDynamicStateResults(const TArray<Chaos::FGeometryParticleHandle*>& ParticleHandles, FFieldContext& FieldContext, TArray<int32>& LocalResults)
+	static void InitDynamicStateResults(const TArray<Chaos::FGeometryParticleHandle*>& ParticleHandles, FFieldContext& FieldContext, TArray<int32>& LocalResults)
 	{
 		for (const FFieldContextIndex& Index : FieldContext.GetEvaluatedSamples())
 		{
@@ -140,7 +140,7 @@ namespace Chaos
 	 * @param    FieldContext Field context to retrieve the evaluated samples
 	 * @param    LocalResults Array to store the enable/disable boolean
 	 */
-	FORCEINLINE void InitActivateDisabledResults(const TArray<Chaos::FGeometryParticleHandle*>& ParticleHandles, FFieldContext& FieldContext, TArray<int32>& LocalResults)
+	static void InitActivateDisabledResults(const TArray<Chaos::FGeometryParticleHandle*>& ParticleHandles, FFieldContext& FieldContext, TArray<int32>& LocalResults)
 	{
 		for (const FFieldContextIndex& Index : FieldContext.GetEvaluatedSamples())
 		{
@@ -158,7 +158,7 @@ namespace Chaos
 	 * @param    FieldState Field state that will be set on the handle
 	 * @param    RigidHandle Particle hadle on which the state will be set
 	 */
-	FORCEINLINE void SetParticleDynamicState(Chaos::FPBDRigidsSolver* RigidSolver,
+	static void SetParticleDynamicState(Chaos::FPBDRigidsSolver* RigidSolver,
 		const Chaos::EObjectStateType FieldState, Chaos::FPBDRigidParticleHandle* RigidHandle)
 	{
 		const bool bIsGC = (RigidHandle->GetParticleType() == Chaos::EParticleType::GeometryCollection) ||
@@ -184,7 +184,7 @@ namespace Chaos
 	 * @param    HasInitialAngularVelocity Boolean to check if we have to set the initial angular velocity 
 	 * @param    InitialAngularVelocity Initial angular velocity to potentially set onto he handle
 	 */
-	FORCEINLINE bool ReportDynamicStateResult(Chaos::FPBDRigidsSolver* RigidSolver,
+	static bool ReportDynamicStateResult(Chaos::FPBDRigidsSolver* RigidSolver,
 		const Chaos::EObjectStateType FieldState, Chaos::FPBDRigidParticleHandle* RigidHandle,
 		const bool HasInitialLinearVelocity, const Chaos::FVec3& InitialLinearVelocity,
 		const bool HasInitialAngularVelocity, const Chaos::FVec3& InitialAngularVelocity)
@@ -224,7 +224,7 @@ namespace Chaos
 	 * @param    Rigidsolver Rigid solver owning the particle handle
 	 * @param    bHasStateChanged Boolean to check before updating the handle state
 	 */
-	FORCEINLINE void UpdateSolverParticlesState(Chaos::FPBDRigidsSolver* RigidSolver, const bool bHasStateChanged)
+	static void UpdateSolverParticlesState(Chaos::FPBDRigidsSolver* RigidSolver, const bool bHasStateChanged)
 	{
 		if (bHasStateChanged)
 		{
@@ -253,7 +253,7 @@ namespace Chaos
 	 * @param    Rigidsolver Rigid solver owning the breaking model
 	 * @param    ExternalStrain Strain to be used to update the breaking model
 	 */
-	FORCEINLINE void UpdateSolverBreakingModel(Chaos::FPBDRigidsSolver* RigidSolver, TMap<Chaos::FGeometryParticleHandle*, float>& ExternalStrain)
+	static void UpdateSolverBreakingModel(Chaos::FPBDRigidsSolver* RigidSolver, TMap<Chaos::FGeometryParticleHandle*, float>& ExternalStrain)
 	{
 		// Capture the results from the breaking model to post-process
 		TMap<Chaos::FPBDRigidClusteredParticleHandle*, TSet<Chaos::FPBDRigidParticleHandle*>> BreakResults =
@@ -279,7 +279,7 @@ namespace Chaos
 	 * @param    RigidHandle Particle handle on which the threshold will be updated
 	 * @param    ResultThreshold Threshoild to be set onto the handle
 	 */
-	FORCEINLINE void UpdateMaterialSleepingThreshold(Chaos::FPBDRigidsSolver* RigidSolver, Chaos::FPBDRigidParticleHandle* RigidHandle, const float ResultThreshold)
+	static void UpdateMaterialSleepingThreshold(Chaos::FPBDRigidsSolver* RigidSolver, Chaos::FPBDRigidParticleHandle* RigidHandle, const float ResultThreshold)
 	{
 		// if no per particle physics material is set, make one
 		if (!RigidSolver->GetEvolution()->GetPerParticlePhysicsMaterial(RigidHandle).IsValid())
@@ -309,7 +309,7 @@ namespace Chaos
 	 * @param    RigidHandle Particle handle on which the threshold will be updated
 	 * @param    ResultThreshold Threshoild to be set onto the handle
 	 */
-	FORCEINLINE void UpdateMaterialDisableThreshold(Chaos::FPBDRigidsSolver* RigidSolver, Chaos::FPBDRigidParticleHandle* RigidHandle, const float ResultThreshold)
+	static void UpdateMaterialDisableThreshold(Chaos::FPBDRigidsSolver* RigidSolver, Chaos::FPBDRigidParticleHandle* RigidHandle, const float ResultThreshold)
 	{
 		// if no per particle physics material is set, make one
 		if (!RigidSolver->GetEvolution()->GetPerParticlePhysicsMaterial(RigidHandle).IsValid())
@@ -339,16 +339,14 @@ namespace Chaos
 	 * @param    FieldCommand Field command to be used for the parameter field evaluatation
 	 * @param    ParticleHandles List of particle handles extracted from the field command meta data
 	 * @param	 FieldContext Field context that will be used for field evaluation
-	 * @param    CommandsToRemove List of commands that will be removed after evaluation
 	 * @param    PositionTarget Chaos position contraint in which each target will be added
 	 * @param    TargetedParticles List of particles (source/target) that will be filled by the PositionTarget/static parameter 
-	 * @param    CommandIndex Command index that we are evaluating
 	 * @param    FinalResults Array in which will be stored the field nodes evaulation
 	 */
-	FORCEINLINE void FieldIntegerParameterUpdate(Chaos::FPBDRigidsSolver* RigidSolver, const FFieldSystemCommand& FieldCommand,
-		TArray<Chaos::FGeometryParticleHandle*>& ParticleHandles, FFieldContext& FieldContext, TArray<int32>& CommandsToRemove,
+	static void FieldIntegerParameterUpdate(Chaos::FPBDRigidsSolver* RigidSolver, const FFieldSystemCommand& FieldCommand,
+		TArray<Chaos::FGeometryParticleHandle*>& ParticleHandles, FFieldContext& FieldContext, 
 		Chaos::FPBDPositionConstraints& PositionTarget,
-		TMap<int32, int32>& TargetedParticles, const int32 CommandIndex, TArray<int32>& FinalResults)
+		TMap<int32, int32>& TargetedParticles, TArray<int32>& FinalResults)
 	{
 		TArrayView<int32> ResultsView(&(FinalResults[0]), FinalResults.Num());
 
@@ -375,7 +373,6 @@ namespace Chaos
 				}
 				UpdateSolverParticlesState(RigidSolver, bHasStateChanged);
 			}
-			CommandsToRemove.Add(CommandIndex);
 		}
 		else if (FieldCommand.PhysicsType == EFieldPhysicsType::Field_ActivateDisabled)
 		{
@@ -394,7 +391,6 @@ namespace Chaos
 					}
 				}
 			}
-			CommandsToRemove.Add(CommandIndex);
 		}
 		else if (FieldCommand.PhysicsType == EFieldPhysicsType::Field_CollisionGroup)
 		{
@@ -410,7 +406,6 @@ namespace Chaos
 					}
 				}
 			}
-			CommandsToRemove.Add(CommandIndex);
 		}
 		else if (FieldCommand.PhysicsType == EFieldPhysicsType::Field_PositionStatic)
 		{
@@ -436,7 +431,6 @@ namespace Chaos
 					}
 				}
 			}
-			CommandsToRemove.Add(CommandIndex);
 		}
 		else if (FieldCommand.PhysicsType == EFieldPhysicsType::Field_DynamicConstraint)
 		{
@@ -444,7 +438,6 @@ namespace Chaos
 			{
 				UE_LOG(LogChaos, Error, TEXT("Dynamic constraint target currently not supported by chaos"));
 			}
-			CommandsToRemove.Add(CommandIndex);
 		}
 	}
 
@@ -454,16 +447,14 @@ namespace Chaos
 	 * @param    FieldCommand Field command to be used for the parameter field evaluatation
 	 * @param    ParticleHandles List of particle handles extracted from the field command meta data
 	 * @param	 FieldContext Field context that will be used for field evaluation
-	 * @param    CommandsToRemove List of commands that will be removed after evaluation
 	 * @param    PositionTarget Chaos position contraint in which each target will be added
 	 * @param    TargetedParticles List of particles (source/target) that will be filled by the PositionTarget/Static parameter
-	 * @param    CommandIndex Command index that we are evaluating
 	 * @param    FinalResults Array in which will be stored the field nodes evaulation
 	 */
-	FORCEINLINE void FieldScalarParameterUpdate(Chaos::FPBDRigidsSolver* RigidSolver, const FFieldSystemCommand& FieldCommand,
-		TArray<Chaos::FGeometryParticleHandle*>& ParticleHandles, FFieldContext& FieldContext, TArray<int32>& CommandsToRemove,
+	static void FieldScalarParameterUpdate(Chaos::FPBDRigidsSolver* RigidSolver, const FFieldSystemCommand& FieldCommand,
+		TArray<Chaos::FGeometryParticleHandle*>& ParticleHandles, FFieldContext& FieldContext, 
 		Chaos::FPBDPositionConstraints& PositionTarget,
-		TMap<int32, int32>& TargetedParticles, const int32 CommandIndex, TArray<float>& FinalResults)
+		TMap<int32, int32>& TargetedParticles, TArray<float>& FinalResults)
 	{
 		TArrayView<float> ResultsView(&(FinalResults[0]), FinalResults.Num());
 
@@ -483,7 +474,6 @@ namespace Chaos
 				}
 				UpdateSolverBreakingModel(RigidSolver, ExternalStrain);
 			}
-			CommandsToRemove.Add(CommandIndex);
 		}
 		else if (FieldCommand.PhysicsType == EFieldPhysicsType::Field_Kill)
 		{
@@ -499,7 +489,6 @@ namespace Chaos
 					}
 				}
 			}
-			CommandsToRemove.Add(CommandIndex);
 		}
 		else if (FieldCommand.PhysicsType == EFieldPhysicsType::Field_SleepingThreshold)
 		{
@@ -515,7 +504,6 @@ namespace Chaos
 					}
 				}
 			}
-			CommandsToRemove.Add(CommandIndex);
 		}
 		else if (FieldCommand.PhysicsType == EFieldPhysicsType::Field_DisableThreshold)
 		{
@@ -532,7 +520,6 @@ namespace Chaos
 					}
 				}
 			}
-			CommandsToRemove.Add(CommandIndex);
 		}
 		else if (FieldCommand.PhysicsType == EFieldPhysicsType::Field_InternalClusterStrain)
 		{
@@ -548,7 +535,6 @@ namespace Chaos
 					}
 				}
 			}
-			CommandsToRemove.Add(CommandIndex);
 		}
 	}
 
@@ -558,16 +544,14 @@ namespace Chaos
 	 * @param    FieldCommand Field command to be used for the parameter field evaluatation
 	 * @param    ParticleHandles List of particle handles extracted from the field command meta data
 	 * @param	 FieldContext Field context that will be used for field evaluation
-	 * @param    CommandsToRemove List of commands that will be removed after evaluation
 	 * @param    PositionTarget Chaos position contraint in which each target will be added
 	 * @param    TargetedParticles List of particles (source/target) that will be filled by the PositionTarget/Static parameter
-	 * @param    CommandIndex Command index that we are evaluating
 	 * @param    FinalResults Array in which will be stored the field nodes evaulation
 	 */
-	FORCEINLINE void FieldVectorParameterUpdate(Chaos::FPBDRigidsSolver* RigidSolver, const FFieldSystemCommand& FieldCommand,
-		TArray<Chaos::FGeometryParticleHandle*>& ParticleHandles, FFieldContext& FieldContext, TArray<int32>& CommandsToRemove,
+	static void FieldVectorParameterUpdate(Chaos::FPBDRigidsSolver* RigidSolver, const FFieldSystemCommand& FieldCommand,
+		TArray<Chaos::FGeometryParticleHandle*>& ParticleHandles, FFieldContext& FieldContext, 
 		Chaos::FPBDPositionConstraints& PositionTarget,
-		TMap<int32, int32>& TargetedParticles, const int32 CommandIndex, TArray<FVector>& FinalResults)
+		TMap<int32, int32>& TargetedParticles, TArray<FVector>& FinalResults)
 	{
 		TArrayView<FVector> ResultsView(&(FinalResults[0]), FinalResults.Num());
 
@@ -585,7 +569,6 @@ namespace Chaos
 					}
 				}
 			}
-			CommandsToRemove.Add(CommandIndex);
 		}
 		else if (FieldCommand.PhysicsType == EFieldPhysicsType::Field_AngularVelociy)
 		{
@@ -602,7 +585,6 @@ namespace Chaos
 					}
 				}
 			}
-			CommandsToRemove.Add(CommandIndex);
 		}
 		else if (FieldCommand.PhysicsType == EFieldPhysicsType::Field_PositionTarget)
 		{
@@ -628,7 +610,6 @@ namespace Chaos
 					}
 				}
 			}
-			CommandsToRemove.Add(CommandIndex);
 		}
 		else if (FieldCommand.PhysicsType == EFieldPhysicsType::Field_PositionAnimated)
 		{
@@ -636,14 +617,12 @@ namespace Chaos
 			{
 				UE_LOG(LogChaos, Error, TEXT("Position Animated target currently not supported by chaos"));
 			}
-			CommandsToRemove.Add(CommandIndex);
 		}
 	}
 
 
-	FORCEINLINE void FieldVectorForceUpdate(Chaos::FPBDRigidsSolver* RigidSolver, const FFieldSystemCommand& FieldCommand,
-		TArray<Chaos::FGeometryParticleHandle*>& ParticleHandles, FFieldContext& FieldContext, TArray<int32>& CommandsToRemove,
-		const int32 CommandIndex, TArray<FVector>& FinalResults)
+	static void FieldVectorForceUpdate(Chaos::FPBDRigidsSolver* RigidSolver, const FFieldSystemCommand& FieldCommand,
+		TArray<Chaos::FGeometryParticleHandle*>& ParticleHandles, FFieldContext& FieldContext, TArray<FVector>& FinalResults)
 	{
 		TArrayView<FVector> ResultsView(&(FinalResults[0]), FinalResults.Num());
 
@@ -666,7 +645,6 @@ namespace Chaos
 					}
 				}
 			}
-			CommandsToRemove.Add(CommandIndex);
 		}
 		else if (FieldCommand.PhysicsType == EFieldPhysicsType::Field_AngularTorque)
 		{
@@ -685,8 +663,42 @@ namespace Chaos
 					}
 				}
 			}
-			CommandsToRemove.Add(CommandIndex);
 		}
 	}
 }
+
+FORCEINLINE bool IsParameterFieldValid(const FFieldSystemCommand& FieldCommand)
+{
+	if (FieldCommand.RootNode->Type() == FFieldNodeBase::EFieldType::EField_Int32)
+	{
+		return (FieldCommand.PhysicsType == EFieldPhysicsType::Field_DynamicState) ||
+			(FieldCommand.PhysicsType == EFieldPhysicsType::Field_ActivateDisabled) ||
+			(FieldCommand.PhysicsType == EFieldPhysicsType::Field_CollisionGroup);
+
+	}
+	else if (FieldCommand.RootNode->Type() == FFieldNodeBase::EFieldType::EField_Float)
+	{
+		return (FieldCommand.PhysicsType == EFieldPhysicsType::Field_ExternalClusterStrain) ||
+			(FieldCommand.PhysicsType == EFieldPhysicsType::Field_Kill) ||
+			(FieldCommand.PhysicsType == EFieldPhysicsType::Field_SleepingThreshold) ||
+			(FieldCommand.PhysicsType == EFieldPhysicsType::Field_DisableThreshold) ||
+			(FieldCommand.PhysicsType == EFieldPhysicsType::Field_InternalClusterStrain);
+	}
+	else if (FieldCommand.RootNode->Type() == FFieldNodeBase::EFieldType::EField_FVector)
+	{
+		return (FieldCommand.PhysicsType == EFieldPhysicsType::Field_LinearVelocity) ||
+			(FieldCommand.PhysicsType == EFieldPhysicsType::Field_AngularVelociy);
+	}
+	return false;
+}
+
+FORCEINLINE bool IsForceFieldValid(const FFieldSystemCommand& FieldCommand)
+{
+	if (FieldCommand.RootNode->Type() == FFieldNodeBase::EFieldType::EField_FVector)
+	{
+		return (FieldCommand.PhysicsType == EFieldPhysicsType::Field_LinearForce) ||
+			(FieldCommand.PhysicsType == EFieldPhysicsType::Field_AngularTorque);
+	}
+	return false;
+}
 	
