Fix ArrayView crash due to array resizing

diff --git a/Engine/Plugins/Experimental/ChaosNiagara/Source/ChaosNiagara/Private/NiagaraDataInterfacePhysicsField.cpp b/Engine/Plugins/Experimental/ChaosNiagara/Source/ChaosNiagara/Private/NiagaraDataInterfacePhysicsField.cpp
--- a/Engine/Plugins/Experimental/ChaosNiagara/Source/ChaosNiagara/Private/NiagaraDataInterfacePhysicsField.cpp
+++ b/Engine/Plugins/Experimental/ChaosNiagara/Source/ChaosNiagara/Private/NiagaraDataInterfacePhysicsField.cpp
@@ -491,7 +491,7 @@ void EvaluateFieldNodes(TArray<FFieldSystemCommand>& FieldCommands, const EField
 	if (FieldCommands.Num() > 0 && ResultsArray.Num() == MaxArray.Num())
 	{
 		SCOPE_CYCLE_COUNTER(STAT_NiagaraUpdateField_Object);
-		TArrayView<DataType> ResultsView(&(ResultsArray.operator[](0)), ResultsArray.Num());
+		TFieldArrayView<DataType> ResultsView(ResultsArray, 0, ResultsArray.Num());
 
 		const float TimeSeconds = FieldContext.TimeSeconds;
 		for (int32 CommandIndex = 0; CommandIndex < FieldCommands.Num(); ++CommandIndex)
diff --git a/Engine/Source/Runtime/Experimental/Chaos/Private/Field/FieldSystemNodes.cpp b/Engine/Source/Runtime/Experimental/Chaos/Private/Field/FieldSystemNodes.cpp
--- a/Engine/Source/Runtime/Experimental/Chaos/Private/Field/FieldSystemNodes.cpp
+++ b/Engine/Source/Runtime/Experimental/Chaos/Private/Field/FieldSystemNodes.cpp
@@ -94,7 +94,7 @@ void SerializeInternal(FArchive& Ar, Enum& Var)
 /**
 * FUniformInteger
 */
-void FUniformInteger::Evaluate(FFieldContext& Context, TArrayView<int32>& Results) const
+void FUniformInteger::Evaluate(FFieldContext& Context, TFieldArrayView<int32>& Results) const
 {
 	int32 MagnitudeVal = Magnitude;
 
@@ -123,7 +123,7 @@ bool FUniformInteger::operator==(const FFieldNodeBase& Node)
 /**
 * FRadialIntMask
 */
-void FRadialIntMask::Evaluate(FFieldContext& Context, TArrayView<int32>& Results) const
+void FRadialIntMask::Evaluate(FFieldContext& Context, TFieldArrayView<int32>& Results) const
 {
 	float Radius2 = Radius * Radius;
 
@@ -194,7 +194,7 @@ bool FRadialIntMask::operator==(const FFieldNodeBase& Node)
 /**
 * FUniformScalar
 */
-void FUniformScalar::Evaluate(FFieldContext& Context, TArrayView<float>& Results) const
+void FUniformScalar::Evaluate(FFieldContext& Context, TFieldArrayView<float>& Results) const
 {
 	float MagnitudeVal = Magnitude;
 
@@ -222,7 +222,7 @@ bool FUniformScalar::operator==(const FFieldNodeBase& Node)
 /**
 * FWaveScalar
 */
-void FWaveScalar::Evaluate(FFieldContext& Context, TArrayView<float>& Results) const
+void FWaveScalar::Evaluate(FFieldContext& Context, TFieldArrayView<float>& Results) const
 {
 	const float Velocity = (Period != 0.0f ) ? Wavelength / Period : 0.0f;
 
@@ -364,7 +364,7 @@ float EvalFalloffFunction<EFieldFalloffType::Field_Falloff_Logarithmic>(const fl
 */
 
 template<EFieldFalloffType FalloffType>
-void FRadialFalloff::Evaluator(const FFieldContext& Context, TArrayView<float>& Results) const  
+void FRadialFalloff::Evaluator(const FFieldContext& Context, TFieldArrayView<float>& Results) const  
 {
 	if (Radius > 0.f)
 	{
@@ -388,7 +388,7 @@ void FRadialFalloff::Evaluator(const FFieldContext& Context, TArrayView<float>&
 	}
 }
 
-void FRadialFalloff::Evaluate(FFieldContext& Context, TArrayView<float>& Results) const
+void FRadialFalloff::Evaluate(FFieldContext& Context, TFieldArrayView<float>& Results) const
 {
 	switch (Falloff)
 	{
@@ -443,7 +443,7 @@ bool FRadialFalloff::operator==(const FFieldNodeBase& Node)
 */
 
 template<EFieldFalloffType FalloffType>
-void FPlaneFalloff::Evaluator(const FFieldContext& Context, const FPlane& Plane, TArrayView<float>& Results) const
+void FPlaneFalloff::Evaluator(const FFieldContext& Context, const FPlane& Plane, TFieldArrayView<float>& Results) const
 {
 	if (Distance > 0.f)
 	{
@@ -467,7 +467,7 @@ void FPlaneFalloff::Evaluator(const FFieldContext& Context, const FPlane& Plane,
 }
 
 void
-FPlaneFalloff::Evaluate(FFieldContext& Context, TArrayView<float>& Results) const
+FPlaneFalloff::Evaluate(FFieldContext& Context, TFieldArrayView<float>& Results) const
 {
 	FPlane Plane(Position, Normal);
 	switch (Falloff)
@@ -525,7 +525,7 @@ bool FPlaneFalloff::operator==(const FFieldNodeBase& Node)
 */
 
 template<EFieldFalloffType FalloffType>
-void FBoxFalloff::Evaluator(const FFieldContext& Context, TArrayView<float>& Results) const
+void FBoxFalloff::Evaluator(const FFieldContext& Context, TFieldArrayView<float>& Results) const
 {
 	const float DeltaRange = (MaxRange - MinRange);
 
@@ -552,7 +552,7 @@ void FBoxFalloff::Evaluator(const FFieldContext& Context, TArrayView<float>& Res
 }
 
 void
-FBoxFalloff::Evaluate(FFieldContext& Context, TArrayView<float>& Results) const
+FBoxFalloff::Evaluate(FFieldContext& Context, TFieldArrayView<float>& Results) const
 {
 	switch (Falloff)
 	{
@@ -605,7 +605,7 @@ bool FBoxFalloff::operator==(const FFieldNodeBase& Node)
 * FNoiseField
 */
 void
-FNoiseField::Evaluate(FFieldContext& Context, TArrayView<float>& Results) const
+FNoiseField::Evaluate(FFieldContext& Context, TFieldArrayView<float>& Results) const
 {
 	const float DeltaRange = (MaxRange - MinRange);
 	const int32 NumSamples = Context.SampleIndices.Num();
@@ -653,7 +653,7 @@ bool FNoiseField::operator==(const FFieldNodeBase& Node)
 /**
 * FUniformVector
 */
-void FUniformVector::Evaluate(FFieldContext& Context, TArrayView<FVector>& Results) const
+void FUniformVector::Evaluate(FFieldContext& Context, TFieldArrayView<FVector>& Results) const
 {
 	const FVector DirectionVal = Direction;
 	const float MagnitudeVal = Magnitude;
@@ -685,7 +685,7 @@ bool FUniformVector::operator==(const FFieldNodeBase& Node)
 /**
 * FRadialVector
 */
-void FRadialVector::Evaluate(FFieldContext& Context, TArrayView<FVector>& Results) const
+void FRadialVector::Evaluate(FFieldContext& Context, TFieldArrayView<FVector>& Results) const
 {
 	const int32 NumSamples = Context.SampleIndices.Num();
 	for (int32 SampleIndex = 0; SampleIndex < NumSamples; SampleIndex++)
@@ -718,7 +718,7 @@ bool FRadialVector::operator==(const FFieldNodeBase& Node)
 /**
 * FRandomVector
 */
-void FRandomVector::Evaluate(FFieldContext& Context, TArrayView<FVector>& Results) const
+void FRandomVector::Evaluate(FFieldContext& Context, TFieldArrayView<FVector>& Results) const
 {
 	const int32 NumSamples = Context.SampleIndices.Num();
 	for (int32 SampleIndex = 0; SampleIndex < NumSamples; SampleIndex++)
@@ -753,7 +753,7 @@ bool FRandomVector::operator==(const FFieldNodeBase& Node)
 * FSumScalar 
 * 
 */
-void FSumScalar::Evaluate(FFieldContext& ContextIn, TArrayView<float>& Results) const
+void FSumScalar::Evaluate(FFieldContext& ContextIn, TFieldArrayView<float>& Results) const
 {
 	TUniquePtr<FFieldSystemMetaDataResults<float> > ResultsData(new FFieldSystemMetaDataResults<float>(Results));
 	FScopedFieldContextMetaData ScopedMetaData(ContextIn, ResultsData.Get());
@@ -770,9 +770,9 @@ void FSumScalar::Evaluate(FFieldContext& ContextIn, TArrayView<float>& Results)
 		const uint32 BufferOffset = ContextIn.ScalarResults.Num();
 		ContextIn.ScalarResults.SetNum(BufferOffset + 2* NumResults, false);
 
-		TArrayView<float> Buffers[2] = {
-			TArrayView<float>(&ContextIn.ScalarResults[BufferOffset],NumResults),
-			TArrayView<float>(&ContextIn.ScalarResults[BufferOffset+NumResults],NumResults),
+		TFieldArrayView<float> Buffers[2] = {
+			TFieldArrayView<float>(ContextIn.ScalarResults,BufferOffset,NumResults),
+			TFieldArrayView<float>(ContextIn.ScalarResults,BufferOffset+NumResults,NumResults),
 		};
 
 		TArray<const FFieldNode<float> * > FieldNodes = { LeftField,RightField };
@@ -876,7 +876,7 @@ bool FSumScalar::operator==(const FFieldNodeBase& Node)
 /**
 * FSumVector
 */
-void FSumVector::Evaluate(FFieldContext& ContextIn, TArrayView<FVector>& Results) const
+void FSumVector::Evaluate(FFieldContext& ContextIn, TFieldArrayView<FVector>& Results) const
 {
 	TUniquePtr<FFieldSystemMetaDataResults<FVector>> ResultsData(new FFieldSystemMetaDataResults<FVector>(Results));
 	FScopedFieldContextMetaData ScopedMetaData(ContextIn, ResultsData.Get());
@@ -894,9 +894,9 @@ void FSumVector::Evaluate(FFieldContext& ContextIn, TArrayView<FVector>& Results
 		const uint32 BufferOffset = ContextIn.VectorResults.Num();
 		ContextIn.VectorResults.SetNum(BufferOffset + 2 * NumResults, false);
 
-		TArrayView<FVector> Buffers[2] = {
-			TArrayView<FVector>(&ContextIn.VectorResults[BufferOffset], NumResults),
-			TArrayView<FVector>(&ContextIn.VectorResults[BufferOffset + NumResults], NumResults),
+		TFieldArrayView<FVector> Buffers[2] = {
+			TFieldArrayView<FVector>(ContextIn.VectorResults, BufferOffset, NumResults),
+			TFieldArrayView<FVector>(ContextIn.VectorResults, BufferOffset + NumResults, NumResults),
 		};
 
 		LeftVectorField->Evaluate(ContextIn, Buffers[0]);
@@ -957,7 +957,7 @@ void FSumVector::Evaluate(FFieldContext& ContextIn, TArrayView<FVector>& Results
 		const uint32 BufferOffset = ContextIn.ScalarResults.Num();
 		ContextIn.ScalarResults.SetNum(BufferOffset + NumResults, false);
 
-		TArrayView<float> Buffer(&ContextIn.ScalarResults[BufferOffset], NumResults);
+		TFieldArrayView<float> Buffer(ContextIn.ScalarResults, BufferOffset, NumResults);
 
 		ScalarField->Evaluate(ContextIn, Buffer);
 
@@ -1012,7 +1012,7 @@ bool FSumVector::operator==(const FFieldNodeBase& Node)
 * FConversionField<InT,OutT>
 */
 template<class InT, class OutT>
-void FConversionField<InT,OutT>::Evaluate(FFieldContext& Context, TArrayView<OutT>& Results) const
+void FConversionField<InT,OutT>::Evaluate(FFieldContext& Context, TFieldArrayView<OutT>& Results) const
 {
 	static_assert(std::is_arithmetic_v<InT>, "Arithmetic types required for field conversion");
 	static_assert(std::is_arithmetic_v<OutT>, "Arithmetic types required for field conversion");
@@ -1026,7 +1026,7 @@ void FConversionField<InT,OutT>::Evaluate(FFieldContext& Context, TArrayView<Out
 	ResultsArray.SetNum(BufferOffset + NumResults, false);
 	FMemory::Memzero(&ResultsArray[BufferOffset], sizeof(InT) * NumResults);
 
-	TArrayView<InT> BufferView(&ResultsArray[BufferOffset], NumResults);
+	TFieldArrayView<InT> BufferView(ResultsArray, BufferOffset, NumResults);
 	InputField->Evaluate(Context, BufferView);
 
 	for (int32 SampleIndex = 0; SampleIndex < NumSamples; SampleIndex++)
@@ -1061,7 +1061,7 @@ template class FConversionField<float, int32>;
 *  FCullingField<T>
 */
 template<class T>
-void FCullingField<T>::Evaluate(FFieldContext& Context, TArrayView<T>& Results) const
+void FCullingField<T>::Evaluate(FFieldContext& Context, TFieldArrayView<T>& Results) const
 {
 	int32 NumResults = Results.Num();
 	int32 NumSamples = Context.SampleIndices.Num();
@@ -1076,14 +1076,14 @@ void FCullingField<T>::Evaluate(FFieldContext& Context, TArrayView<T>& Results)
 		{
 			const uint32 BufferOffset = Context.ScalarResults.Num();
 			Context.ScalarResults.SetNum(BufferOffset + NumResults, false);
-			TArrayView<float> CullingBuffer(&Context.ScalarResults[BufferOffset], NumResults);
+			TFieldArrayView<float> CullingBuffer(Context.ScalarResults, BufferOffset, NumResults);
 
 			FMemory::Memzero(&Context.ScalarResults[BufferOffset], sizeof(float) * NumResults);
 			CullingField->Evaluate(Context, CullingBuffer);
 
 			const uint32 IndexOffset = Context.IndexResults.Num();
 			Context.IndexResults.SetNum(IndexOffset + NumResults, false);
-			TArrayView<FFieldContextIndex> IndexBuffer(&Context.IndexResults[IndexOffset], NumResults);
+			TFieldArrayView<FFieldContextIndex> IndexBuffer(Context.IndexResults, IndexOffset, NumResults);
 
 			int NewEvaluationSize = 0;
 			for (int32 SampleIndex = 0; SampleIndex < NumSamples; SampleIndex++)
@@ -1122,7 +1122,7 @@ void FCullingField<T>::Evaluate(FFieldContext& Context, TArrayView<T>& Results)
 
 					if (InputField)
 					{
-						TArrayView<FFieldContextIndex> LocalIndices(&Context.IndexResults[IndexOffset], NewEvaluationSize);
+						TFieldArrayView<FFieldContextIndex> LocalIndices(Context.IndexResults, IndexOffset, NewEvaluationSize);
 						FFieldContext LocalContext(LocalIndices, Context.SamplePositions, Context.MetaData, Context.TimeSeconds,
 							Context.VectorResults, Context.ScalarResults, Context.IntegerResults, Context.IndexResults, CullingData->CullingIndices);
 						InputField->Evaluate(LocalContext, Results);
@@ -1167,7 +1167,7 @@ template class CHAOS_API FCullingField<FVector>;
 * FReturnResultsTerminal<T>
 */
 template<class T>
-void FReturnResultsTerminal<T>::Evaluate(FFieldContext& Context, TArrayView<T>& Results) const
+void FReturnResultsTerminal<T>::Evaluate(FFieldContext& Context, TFieldArrayView<T>& Results) const
 {
 	if (ensureMsgf(Context.MetaData.Contains(FFieldSystemMetaData::EMetaType::ECommandData_Results),
 		TEXT("Return results nodes can only be used upstream from a 'results expector', for example as an input "
diff --git a/Engine/Source/Runtime/Experimental/Chaos/Private/PhysicsProxy/GeometryCollectionPhysicsProxy.cpp b/Engine/Source/Runtime/Experimental/Chaos/Private/PhysicsProxy/GeometryCollectionPhysicsProxy.cpp
--- a/Engine/Source/Runtime/Experimental/Chaos/Private/PhysicsProxy/GeometryCollectionPhysicsProxy.cpp
+++ b/Engine/Source/Runtime/Experimental/Chaos/Private/PhysicsProxy/GeometryCollectionPhysicsProxy.cpp
@@ -2636,7 +2636,7 @@ void FGeometryCollectionPhysicsProxy::FieldParameterUpdateCallback(Chaos::FPBDRi
 						TArray<int32>& FinalResults = ExecutionDatas.IntegerResults[(uint8)EFieldCommandResultType::FinalResult];
 						ResetResultsArray < int32 >(ExecutionDatas.SamplePositions.Num(), FinalResults, 0);
 
-						TArrayView<int32> ResultsView(&(FinalResults[0]), FinalResults.Num());
+						TFieldArrayView<int32> ResultsView(FinalResults, 0, FinalResults.Num());
 
 						if (FieldCommand.PhysicsType == EFieldPhysicsType::Field_DynamicState)
 						{
@@ -2691,7 +2691,7 @@ void FGeometryCollectionPhysicsProxy::FieldParameterUpdateCallback(Chaos::FPBDRi
 						TArray<FVector>& FinalResults = ExecutionDatas.VectorResults[(uint8)EFieldCommandResultType::FinalResult];
 						ResetResultsArray < FVector >(ExecutionDatas.SamplePositions.Num(), FinalResults, FVector::ZeroVector);
 
-						TArrayView<FVector> ResultsView(&(FinalResults[0]), FinalResults.Num());
+						TFieldArrayView<FVector> ResultsView(FinalResults, 0, FinalResults.Num());
 
 						if (FieldCommand.PhysicsType == EFieldPhysicsType::Field_InitialLinearVelocity)
 						{
diff --git a/Engine/Source/Runtime/Experimental/Chaos/Private/PhysicsProxy/PerSolverFieldSystem.cpp b/Engine/Source/Runtime/Experimental/Chaos/Private/PhysicsProxy/PerSolverFieldSystem.cpp
--- a/Engine/Source/Runtime/Experimental/Chaos/Private/PhysicsProxy/PerSolverFieldSystem.cpp
+++ b/Engine/Source/Runtime/Experimental/Chaos/Private/PhysicsProxy/PerSolverFieldSystem.cpp
@@ -175,7 +175,7 @@ void FPerSolverFieldSystem::FieldForcesUpdateCallback(
 FORCEINLINE void EvaluateImpulseField(
 	const FFieldSystemCommand& FieldCommand,
 	FFieldContext& FieldContext,
-	TArrayView<FVector>& ResultsView,
+	TFieldArrayView<FVector>& ResultsView,
 	TArray<FVector>& OutputImpulse)
 {
 	static_cast<const FFieldNode<FVector>*>(FieldCommand.RootNode.Get())->Evaluate(FieldContext, ResultsView);
@@ -237,7 +237,7 @@ void ComputeFieldRigidImpulseInternal(
 					TArray<FVector>& FinalResults = ExecutionDatas.VectorResults[(uint8)EFieldCommandResultType::FinalResult];
 					ResetResultsArray < FVector >(ExecutionDatas.SamplePositions.Num(), FinalResults, FVector::ZeroVector);
 
-					TArrayView<FVector> ResultsView(&(FinalResults[0]), FinalResults.Num());
+					TFieldArrayView<FVector> ResultsView(FinalResults, 0, FinalResults.Num());
 
 					if (FieldCommand.PhysicsType == EFieldPhysicsType::Field_LinearVelocity)
 					{
@@ -327,7 +327,7 @@ void ComputeFieldLinearImpulseInternal(
 					TArray<FVector>& FinalResults = ExecutionDatas.VectorResults[(uint8)EFieldCommandResultType::FinalResult];
 					ResetResultsArray < FVector >(ExecutionDatas.SamplePositions.Num(), FinalResults, FVector::ZeroVector);
 
-					TArrayView<FVector> ResultsView(&(FinalResults[0]), FinalResults.Num());
+					TFieldArrayView<FVector> ResultsView(FinalResults, 0, FinalResults.Num());
 
 					if (FieldCommand.PhysicsType == EFieldPhysicsType::Field_LinearVelocity)
 					{
diff --git a/Engine/Source/Runtime/Experimental/Chaos/Public/Field/FieldArrayView.h b/Engine/Source/Runtime/Experimental/Chaos/Public/Field/FieldArrayView.h
new file mode 100644
--- /dev/null
+++ b/Engine/Source/Runtime/Experimental/Chaos/Public/Field/FieldArrayView.h
@@ -0,0 +1,139 @@
+// Copyright Epic Games, Inc. All Rights Reserved.
+
+#pragma once
+
+#include "CoreMinimal.h"
+#include "Containers/Array.h"
+
+/**
+ * TFieldArrayView : Templated fixed-size view of another array
+ *
+ * A statically sized view of an TArray of typed elements. The reference
+ * TArray could be resized and the TFieldArrayView will still be valid 
+ * For now only used for the fields. Should be extended if required to be used somewhere else. 
+*/
+
+template<typename InElementType>
+class TFieldArrayView
+{
+public:
+	using ElementType = InElementType;
+	using SizeType = int32;
+
+	/**
+	 * Constructor.
+	 */
+	TFieldArrayView(TArray<ElementType>& InElementArray, const SizeType InArrayOffset, const SizeType InArrayNum)
+		: ElementArray(InElementArray)
+		, ArrayOffset(InArrayOffset)
+		, ArrayNum(InArrayNum)
+	{
+		checkf(InArrayOffset + InArrayNum <= ElementArray.Num(), TEXT("Field array size out of bounds: %i from an array of size %i"), InArrayOffset + InArrayNum, ElementArray.Num());
+	}
+
+	/**
+	 * Helper function for returning a typed pointer to the first array entry.
+	 *
+	 * @returns Pointer to first array entry or nullptr if ArrayMax == 0.
+	 */
+	FORCEINLINE ElementType* GetData() const
+	{
+		return ElementArray.GetData() + ArrayOffset;
+	}
+
+	/**
+	 * Helper function returning the size of the inner type.
+	 *
+	 * @returns Size in bytes of array type.
+	 */
+	FORCEINLINE static constexpr size_t GetTypeSize()
+	{
+		return sizeof(ElementType);
+	}
+
+	/**
+	 * Helper function returning the alignment of the inner type.
+	 */
+	FORCEINLINE static constexpr size_t GetTypeAlignment()
+	{
+		return alignof(ElementType);
+	}
+
+	/**
+	 * Checks array invariants: if array size is greater than zero and less
+	 * than maximum.
+	 */
+	FORCEINLINE void CheckInvariants() const
+	{
+		checkSlow(ArrayNum >= 0);
+	}
+
+	/**
+	 * Checks if index is in array range.
+	 *
+	 * @param Index Index to check.
+	 */
+	FORCEINLINE void RangeCheck(SizeType Index) const
+	{
+		CheckInvariants();
+
+		checkf((Index >= 0) & (Index < ArrayNum), TEXT("Array index out of bounds: %i from an array of size %i"), Index, ArrayNum); // & for one branch
+	}
+
+	/**
+	 * Tests if index is valid, i.e. more or equal to zero, and less than the number of elements in the array.
+	 *
+	 * @param Index Index to test.
+	 *
+	 * @returns True if index is valid. False otherwise.
+	 */
+	FORCEINLINE bool IsValidIndex(SizeType Index) const
+	{
+		return (Index >= 0) && (Index < ArrayNum);
+	}
+
+	/**
+	 * Returns true if the array view is empty and contains no elements.
+	 *
+	 * @returns True if the array view is empty.
+	 * @see Num
+	 */
+	bool IsEmpty() const
+	{
+		return ArrayNum == 0;
+	}
+
+	/**
+	 * Returns number of elements in array.
+	 *
+	 * @returns Number of elements in array.
+	 */
+	FORCEINLINE SizeType Num() const
+	{
+		return ArrayNum;
+	}
+
+	/**
+	 * Array bracket operator. Returns reference to element at give index.
+	 *
+	 * @returns Reference to indexed element.
+	 */
+	FORCEINLINE ElementType& operator[](SizeType Index) const
+	{
+		RangeCheck(Index);
+		return ElementArray[ArrayOffset + Index];
+	}
+
+	/**
+	  * DO NOT USE DIRECTLY
+      * STL-like iterators to enable range-based for loop support.
+    */
+	FORCEINLINE ElementType* begin() const { return GetData(); }
+	FORCEINLINE ElementType* end() const { return GetData() + Num(); }
+
+private:
+
+	TArray<InElementType>& ElementArray;
+	SizeType ArrayOffset;
+	SizeType ArrayNum;
+};
diff --git a/Engine/Source/Runtime/Experimental/Chaos/Public/Field/FieldSystem.h b/Engine/Source/Runtime/Experimental/Chaos/Public/Field/FieldSystem.h
--- a/Engine/Source/Runtime/Experimental/Chaos/Public/Field/FieldSystem.h
+++ b/Engine/Source/Runtime/Experimental/Chaos/Public/Field/FieldSystem.h
@@ -7,6 +7,7 @@
 #include "Containers/ArrayView.h"
 #include "Containers/Queue.h"
 #include "Field/FieldSystemTypes.h"
+#include "Field/FieldArrayView.h"
 #include "Chaos/ParticleHandle.h"
 #include "Math/Vector.h"
 
@@ -196,12 +197,12 @@ public:
 template<class T>
 class CHAOS_API FFieldSystemMetaDataResults : public FFieldSystemMetaData {
 public:
-	FFieldSystemMetaDataResults(const TArrayView<T>& ResultsIn) : Results(ResultsIn) {};
+	FFieldSystemMetaDataResults(const TFieldArrayView<T>& ResultsIn) : Results(ResultsIn) {};
 	virtual ~FFieldSystemMetaDataResults() {};
 	virtual EMetaType Type() const { return EMetaType::ECommandData_Results; }
 	virtual FFieldSystemMetaData* NewCopy() const { return new FFieldSystemMetaDataResults(Results); }
 
-	const TArrayView<T>& Results;
+	const TFieldArrayView<T>& Results;
 };
 
 class CHAOS_API FFieldSystemMetaDataIteration : public FFieldSystemMetaData {
@@ -249,7 +250,7 @@ struct CHAOS_API FFieldContext
 	FFieldContext& operator =(const FFieldContext&) = delete;
 	FFieldContext & operator =(FFieldContext&&) = delete;
 
-	FFieldContext(const TArrayView< FFieldContextIndex >& SampleIndicesIn, const TArrayView<FVector>& SamplePositionsIn,
+	FFieldContext(const TFieldArrayView< FFieldContextIndex >& SampleIndicesIn, const TFieldArrayView<FVector>& SamplePositionsIn,
 		const UniquePointerMap & MetaDataIn, const float TimeSecondsIn, TArray<FVector>& VectorResultsIn, TArray<float>& ScalarResultsIn, 
 		TArray<int32>& IntegerResultsIn, TArray<FFieldContextIndex>& IndexResultsIn, TArray<FFieldContextIndex>& CullingResultsIn)
 		: SampleIndices(SampleIndicesIn)
@@ -269,7 +270,7 @@ struct CHAOS_API FFieldContext
 		CullingData = MakeUnique<FFieldSystemMetaDataCulling>(CullingResultsIn);
 		MetaData.Add(FFieldSystemMetaData::EMetaType::ECommandData_Culling, CullingData.Get());
 	}
-	FFieldContext(const TArrayView< FFieldContextIndex >& SampleIndicesIn, const TArrayView<FVector>& SamplePositionsIn,
+	FFieldContext(const TFieldArrayView< FFieldContextIndex >& SampleIndicesIn, const TFieldArrayView<FVector>& SamplePositionsIn,
 		const PointerMap & MetaDataIn, const float TimeSecondsIn, TArray<FVector>& VectorResultsIn, TArray<float>& ScalarResultsIn, 
 				TArray<int32>& IntegerResultsIn, TArray<FFieldContextIndex>& IndexResultsIn, TArray<FFieldContextIndex>& CullingResultsIn)
 		: SampleIndices(SampleIndicesIn)
@@ -287,15 +288,14 @@ struct CHAOS_API FFieldContext
 
 	FFieldContext(FFieldExecutionDatas& ExecutionDatas,
 		const UniquePointerMap& MetaDataIn, const float TimeSecondsIn)
-		: TimeSeconds(TimeSecondsIn)
+		: SampleIndices(ExecutionDatas.SampleIndices, 0, ExecutionDatas.SampleIndices.Num())
+		, SamplePositions(ExecutionDatas.SamplePositions, 0, ExecutionDatas.SamplePositions.Num())
+		, TimeSeconds(TimeSecondsIn)
 		, VectorResults(ExecutionDatas.VectorResults[(uint8)EFieldCommandResultType::TransientResult])
 		, ScalarResults(ExecutionDatas.ScalarResults[(uint8)EFieldCommandResultType::TransientResult])
 		, IntegerResults(ExecutionDatas.IntegerResults[(uint8)EFieldCommandResultType::TransientResult])
 		, IndexResults(ExecutionDatas.IndexResults[(uint8)EFieldCommandResultType::TransientResult])
 	{
-		SamplePositions = TArrayView<FVector>(&(ExecutionDatas.SamplePositions[0]), ExecutionDatas.SamplePositions.Num());
-		SampleIndices = TArrayView<FFieldContextIndex>(&(ExecutionDatas.SampleIndices[0]), ExecutionDatas.SampleIndices.Num());
-
 		for (const TPair<FFieldSystemMetaData::EMetaType, TUniquePtr<FFieldSystemMetaData>>& Meta : MetaDataIn)
 		{
 			MetaData.Add(Meta.Key) = Meta.Value.Get();
@@ -305,7 +305,7 @@ struct CHAOS_API FFieldContext
 		MetaData.Add(FFieldSystemMetaData::EMetaType::ECommandData_Culling, CullingData.Get());
 	}
 
-	TArrayView<FFieldContextIndex> GetEvaluatedSamples()
+	TFieldArrayView<FFieldContextIndex> GetEvaluatedSamples()
 	{
 		if(!CullingData->bCullingActive)
 		{
@@ -314,7 +314,7 @@ struct CHAOS_API FFieldContext
 		}
 
 		// Culling fields created an evaluation set
-		return MakeArrayView(CullingData->CullingIndices);
+		return TFieldArrayView<FFieldContextIndex>(CullingData->CullingIndices, 0, CullingData->CullingIndices.Num());
 	}
 
 	//
@@ -324,8 +324,8 @@ struct CHAOS_API FFieldContext
 	// traversed also needs to change; possibly to some load balanced threaded iterator 
 	// or task based paradigm.
 
-	TArrayView<FFieldContextIndex> SampleIndices;
-	TArrayView<FVector> SamplePositions;
+	TFieldArrayView<FFieldContextIndex> SampleIndices;
+	TFieldArrayView<FVector> SamplePositions;
 
 	PointerMap MetaData;
 	TUniquePtr<FFieldSystemMetaDataCulling> CullingData;
@@ -456,7 +456,7 @@ private:
 *
 *  Typed field nodes are used for the evaluation of specific types of data arrays.
 *  For exampe, The FFieldNode<FVector>::Evaluate(...) will expect resutls 
-*  of type TArrayView<FVector>, and an example implementation is the UniformVectorField.
+*  of type TFieldArrayView<FVector>, and an example implementation is the UniformVectorField.
 *
 */
 template<class T>
@@ -466,7 +466,7 @@ public:
 	
 	virtual ~FFieldNode() {}
 
-	virtual void Evaluate(FFieldContext&, TArrayView<T>& Results) const = 0;
+	virtual void Evaluate(FFieldContext&, TFieldArrayView<T>& Results) const = 0;
 
 	static EFieldType StaticType();
 	virtual EFieldType Type() const { return StaticType(); }
diff --git a/Engine/Source/Runtime/Experimental/Chaos/Public/Field/FieldSystemNodes.h b/Engine/Source/Runtime/Experimental/Chaos/Public/Field/FieldSystemNodes.h
--- a/Engine/Source/Runtime/Experimental/Chaos/Public/Field/FieldSystemNodes.h
+++ b/Engine/Source/Runtime/Experimental/Chaos/Public/Field/FieldSystemNodes.h
@@ -6,7 +6,6 @@
 #include "Math/Vector.h"
 #include "Chaos/Box.h"
 
-
 /**
 * FUniformInteger
 **/
@@ -24,7 +23,7 @@ public:
 	virtual FFieldNodeBase * NewCopy() const override { return new FUniformInteger(Magnitude); }
 	virtual ~FUniformInteger() {}
 
-	void Evaluate(FFieldContext&, TArrayView<int32>& Results) const override;
+	void Evaluate(FFieldContext&, TFieldArrayView<int32>& Results) const override;
 	virtual bool operator==(const FFieldNodeBase& Node);
 
 	/** Serialization API */
@@ -61,7 +60,7 @@ public:
 	virtual FFieldNodeBase * NewCopy() const override { return new FRadialIntMask(Radius,Position,InteriorValue,ExteriorValue,SetMaskCondition); }
 	virtual ~FRadialIntMask() {}
 
-	void Evaluate(FFieldContext&, TArrayView<int32>& Results) const override;
+	void Evaluate(FFieldContext&, TFieldArrayView<int32>& Results) const override;
 	virtual bool operator==(const FFieldNodeBase& Node);
 
 	/** Serialization API */
@@ -94,7 +93,7 @@ public:
 
 	virtual ~FUniformScalar() {}
 
-	void Evaluate(FFieldContext&, TArrayView<float>& Results) const override;
+	void Evaluate(FFieldContext&, TFieldArrayView<float>& Results) const override;
 	virtual bool operator==(const FFieldNodeBase& Node);
 
 	/** Serialization API */
@@ -128,7 +127,7 @@ public:
 
 	virtual ~FWaveScalar() {}
 
-	void Evaluate(FFieldContext&, TArrayView<float>& Results) const override;
+	void Evaluate(FFieldContext&, TFieldArrayView<float>& Results) const override;
 	virtual bool operator==(const FFieldNodeBase& Node);
 
 	/** Serialization API */
@@ -174,7 +173,7 @@ public:
 
 	virtual ~FRadialFalloff() {}
 
-	void Evaluate(FFieldContext&, TArrayView<float>& Results) const override;
+	void Evaluate(FFieldContext&, TFieldArrayView<float>& Results) const override;
 	virtual bool operator==(const FFieldNodeBase& Node);
 
 	/** Serialization API */
@@ -190,7 +189,7 @@ public:
 	EFieldFalloffType Falloff;
 
 protected:
-	template<EFieldFalloffType> void Evaluator(const FFieldContext& Context, TArrayView<float>& Results) const;
+	template<EFieldFalloffType> void Evaluator(const FFieldContext& Context, TFieldArrayView<float>& Results) const;
 };
 
 
@@ -226,7 +225,7 @@ public:
 	virtual FFieldNodeBase * NewCopy() const override { return new FPlaneFalloff(Magnitude, MinRange, MaxRange, Default, Distance, Position, Normal, Falloff); }
 	virtual ~FPlaneFalloff() {}
 
-	void Evaluate(FFieldContext&, TArrayView<float>& Results) const override;
+	void Evaluate(FFieldContext&, TFieldArrayView<float>& Results) const override;
 	virtual bool operator==(const FFieldNodeBase& Node);
 
 	/** Serialization API */
@@ -243,7 +242,7 @@ public:
 	EFieldFalloffType Falloff;
 
 protected :
-	template<EFieldFalloffType> void Evaluator(const FFieldContext& Context, const FPlane& Plane, TArrayView<float>& Results) const;
+	template<EFieldFalloffType> void Evaluator(const FFieldContext& Context, const FPlane& Plane, TFieldArrayView<float>& Results) const;
 
 };
 
@@ -275,7 +274,7 @@ public:
 	virtual FFieldNodeBase * NewCopy() const override { return new FBoxFalloff(Magnitude, MinRange, MaxRange, Default, Transform, Falloff); }
 	virtual ~FBoxFalloff() {}
 
-	void Evaluate(FFieldContext&, TArrayView<float>& Results) const override;
+	void Evaluate(FFieldContext&, TFieldArrayView<float>& Results) const override;
 	virtual bool operator==(const FFieldNodeBase& Node);
 
 	/** Serialization API */
@@ -290,7 +289,7 @@ public:
 	EFieldFalloffType Falloff;
 
 protected:
-	template<EFieldFalloffType> void Evaluator(const FFieldContext& Context, TArrayView<float>& Results) const;
+	template<EFieldFalloffType> void Evaluator(const FFieldContext& Context, TFieldArrayView<float>& Results) const;
 
 };
 
@@ -313,7 +312,7 @@ public:
 	virtual FFieldNodeBase * NewCopy() const override { return new FNoiseField(MinRange,MaxRange,Transform); }
 	virtual ~FNoiseField() {}
 
-	void Evaluate(FFieldContext&, TArrayView<float>& Results) const override;	
+	void Evaluate(FFieldContext&, TFieldArrayView<float>& Results) const override;	
 	virtual bool operator==(const FFieldNodeBase& Node);
 
 	/** Serialization API */
@@ -347,7 +346,7 @@ public:
 	virtual FFieldNodeBase * NewCopy() const override { return new FUniformVector(Magnitude,Direction); }
 	virtual ~FUniformVector() {}
 
-	void Evaluate(FFieldContext&, TArrayView<FVector>& Results) const override;
+	void Evaluate(FFieldContext&, TFieldArrayView<FVector>& Results) const override;
 	virtual bool operator==(const FFieldNodeBase& Node);
 
 	/** Serialization API */
@@ -378,7 +377,7 @@ public:
 	virtual FFieldNodeBase * NewCopy() const override { return new FRadialVector(Magnitude,Position); }
 	virtual ~FRadialVector() {}
 
-	void Evaluate(FFieldContext&, TArrayView<FVector>& Results) const override;
+	void Evaluate(FFieldContext&, TFieldArrayView<FVector>& Results) const override;
 	virtual bool operator==(const FFieldNodeBase& Node);
 
 	/** Serialization API */
@@ -408,7 +407,7 @@ public:
 	virtual FFieldNodeBase * NewCopy() const override { return new FRandomVector(Magnitude); }
 	virtual ~FRandomVector() {}
 
-	void Evaluate(FFieldContext&, TArrayView<FVector>& Results) const override;
+	void Evaluate(FFieldContext&, TFieldArrayView<FVector>& Results) const override;
 	virtual bool operator==(const FFieldNodeBase& Node);
 
 	/** Serialization API */
@@ -448,7 +447,7 @@ public:
 	}
 	virtual ~FSumScalar() {}
 
-	void Evaluate(FFieldContext&, TArrayView<float>& Results) const override;
+	void Evaluate(FFieldContext&, TFieldArrayView<float>& Results) const override;
 	virtual bool operator==(const FFieldNodeBase& Node);
 
 	/** Serialization API */
@@ -493,7 +492,7 @@ public:
 	}
 	virtual ~FSumVector() {}
 
-	void Evaluate(FFieldContext&, TArrayView<FVector>& Results) const override;
+	void Evaluate(FFieldContext&, TFieldArrayView<FVector>& Results) const override;
 	virtual bool operator==(const FFieldNodeBase& Node);
 
 	/** Serialization API */
@@ -528,7 +527,7 @@ public:
 	}
 	virtual ~FConversionField() {}
 
-	void Evaluate(FFieldContext&, TArrayView<OutT>& Results) const override;
+	void Evaluate(FFieldContext&, TFieldArrayView<OutT>& Results) const override;
 	virtual bool operator==(const FFieldNodeBase& Node);
 
 	/** Serialization API */
@@ -568,7 +567,7 @@ public:
 	}
 	virtual ~FCullingField() {}
 
-	void Evaluate(FFieldContext&, TArrayView<T>& Results) const override;
+	void Evaluate(FFieldContext&, TFieldArrayView<T>& Results) const override;
 	virtual bool operator==(const FFieldNodeBase& Node);
 
 	/** Serialization API */
@@ -605,7 +604,7 @@ public:
 	}
 	virtual ~FReturnResultsTerminal() {}
 
-	void Evaluate(FFieldContext&, TArrayView<T>& Results) const override;
+	void Evaluate(FFieldContext&, TFieldArrayView<T>& Results) const override;
 	virtual bool operator==(const FFieldNodeBase& Node);
 
 	/** Serialization API */
diff --git a/Engine/Source/Runtime/Experimental/Chaos/Public/PhysicsProxy/FieldSystemProxyHelper.h b/Engine/Source/Runtime/Experimental/Chaos/Public/PhysicsProxy/FieldSystemProxyHelper.h
--- a/Engine/Source/Runtime/Experimental/Chaos/Public/PhysicsProxy/FieldSystemProxyHelper.h
+++ b/Engine/Source/Runtime/Experimental/Chaos/Public/PhysicsProxy/FieldSystemProxyHelper.h
@@ -348,7 +348,7 @@ namespace Chaos
 		Chaos::FPBDPositionConstraints& PositionTarget,
 		TMap<int32, int32>& TargetedParticles, TArray<int32>& FinalResults)
 	{
-		TArrayView<int32> ResultsView(&(FinalResults[0]), FinalResults.Num());
+		TFieldArrayView<int32> ResultsView(FinalResults, 0, FinalResults.Num());
 
 		if (FieldCommand.PhysicsType == EFieldPhysicsType::Field_DynamicState)
 		{
@@ -456,7 +456,7 @@ namespace Chaos
 		Chaos::FPBDPositionConstraints& PositionTarget,
 		TMap<int32, int32>& TargetedParticles, TArray<float>& FinalResults)
 	{
-		TArrayView<float> ResultsView(&(FinalResults[0]), FinalResults.Num());
+		TFieldArrayView<float> ResultsView(FinalResults, 0, FinalResults.Num());
 
 		if (FieldCommand.PhysicsType == EFieldPhysicsType::Field_ExternalClusterStrain)
 		{
@@ -553,7 +553,7 @@ namespace Chaos
 		Chaos::FPBDPositionConstraints& PositionTarget,
 		TMap<int32, int32>& TargetedParticles, TArray<FVector>& FinalResults)
 	{
-		TArrayView<FVector> ResultsView(&(FinalResults[0]), FinalResults.Num());
+		TFieldArrayView<FVector> ResultsView(FinalResults, 0, FinalResults.Num());
 
 		if (FieldCommand.PhysicsType == EFieldPhysicsType::Field_LinearVelocity)
 		{
@@ -624,7 +624,7 @@ namespace Chaos
 	static void FieldVectorForceUpdate(Chaos::FPBDRigidsSolver* RigidSolver, const FFieldSystemCommand& FieldCommand,
 		TArray<Chaos::FGeometryParticleHandle*>& ParticleHandles, FFieldContext& FieldContext, TArray<FVector>& FinalResults)
 	{
-		TArrayView<FVector> ResultsView(&(FinalResults[0]), FinalResults.Num());
+		TFieldArrayView<FVector> ResultsView(FinalResults, 0, FinalResults.Num());
 
 		static_cast<const FFieldNode<FVector>*>(FieldCommand.RootNode.Get())->Evaluate(FieldContext, ResultsView);
 
