ChaosVehicles: Tweak normal so vehicles no longer slide on hills when they are at rest.

#okforgitbug public, internal

diff --git a/Engine/Source/Runtime/Experimental/Chaos/Private/Chaos/PBDSuspensionConstraints.cpp b/Engine/Source/Runtime/Experimental/Chaos/Private/Chaos/PBDSuspensionConstraints.cpp
--- a/Engine/Source/Runtime/Experimental/Chaos/Private/Chaos/PBDSuspensionConstraints.cpp
+++ b/Engine/Source/Runtime/Experimental/Chaos/Private/Chaos/PBDSuspensionConstraints.cpp
@@ -95,6 +95,28 @@ namespace Chaos
 			const FVec3 WorldSpaceX = Particle->Q().RotateVector(SuspensionLocalOffset[ConstraintIndex]) + Particle->P();
 
 			FVec3 AxisWorld = Particle->Q() * Setting.Axis;
+
+			const float MPHToCmS = 100000.f / 2236.94185f;
+			const float SpeedThreshold = 10.0f * MPHToCmS;
+			const float FortyFiveDegreesThreshold = 0.707f;
+
+			if (AxisWorld.Z > FortyFiveDegreesThreshold)
+			{
+				if (Particle->V().SquaredLength() < 1.0f)
+				{
+					AxisWorld = FVec3(0.f, 0.f, 1.f);
+				}
+				else
+				{
+					const float Speed = FMath::Abs(Particle->V().Length());
+					if (Speed < SpeedThreshold)
+					{
+						AxisWorld = FMath::Lerp(FVec3(0.f, 0.f, 1.f), AxisWorld, Speed / SpeedThreshold);
+
+					}
+				}
+			}
+
 			FReal Distance = FVec3::DotProduct(WorldSpaceX - T, AxisWorld);
 			if (Distance >= Setting.MaxLength)
 			{
