Change plasticity to be based on percentage of center of mass distance.

diff --git a/Engine/Source/Runtime/Engine/Classes/PhysicsEngine/ConstraintInstance.h b/Engine/Source/Runtime/Engine/Classes/PhysicsEngine/ConstraintInstance.h
--- a/Engine/Source/Runtime/Engine/Classes/PhysicsEngine/ConstraintInstance.h
+++ b/Engine/Source/Runtime/Engine/Classes/PhysicsEngine/ConstraintInstance.h
@@ -46,7 +46,7 @@ struct ENGINE_API FConstraintProfileProperties
 	UPROPERTY(EditAnywhere, AdvancedDisplay, Category = Linear, meta = (editcondition = "bLinearBreakable", ClampMin = "0.0"))
 	float LinearBreakThreshold;
 
-	/** Percent threshold from target position needed to reset the spring rest length.*/
+	/** [Chaos Only] Percent threshold from center of mass distance needed to reset the linear drive position target. This value can be greater than 1. */
 	UPROPERTY(EditAnywhere, AdvancedDisplay, Category = Linear, meta = (editcondition = "bLinearPlasticity", ClampMin = "0.0"))
 	float LinearPlasticityThreshold;
 
@@ -54,7 +54,7 @@ struct ENGINE_API FConstraintProfileProperties
 	UPROPERTY(EditAnywhere, AdvancedDisplay, Category = Angular, meta = (editcondition = "bAngularBreakable", ClampMin = "0.0"))
 	float AngularBreakThreshold;
 
-	/** Degree threshold from target angle needed to reset the target angle.*/
+	/** [Chaos Only] Degree threshold from target angle needed to reset the target angle.*/
 	UPROPERTY(EditAnywhere, AdvancedDisplay, Category = Angular, meta = (editcondition = "bAngularPlasticity", ClampMin = "0.0"))
 	float AngularPlasticityThreshold;
 
diff --git a/Engine/Source/Runtime/Experimental/Chaos/Private/Chaos/PBDJointConstraints.cpp b/Engine/Source/Runtime/Experimental/Chaos/Private/Chaos/PBDJointConstraints.cpp
--- a/Engine/Source/Runtime/Experimental/Chaos/Private/Chaos/PBDJointConstraints.cpp
+++ b/Engine/Source/Runtime/Experimental/Chaos/Private/Chaos/PBDJointConstraints.cpp
@@ -167,6 +167,7 @@ namespace Chaos
 		, AngularDriveDamping(0)
 		, LinearBreakForce(FLT_MAX)
 		, LinearPlasticityLimit(FLT_MAX)
+		, InitCOMDistance(-1)
 		, AngularBreakTorque(FLT_MAX)
 		, AngularPlasticityLimit(FLT_MAX)
 		, UserData(nullptr)
@@ -789,6 +790,17 @@ namespace Chaos
 					Particle1->InvI().GetDiagonal(),
 					FParticleUtilities::ParticleLocalToCoMLocal(Particle0, JointFrames[Index0]),
 					FParticleUtilities::ParticleLocalToCoMLocal(Particle1, JointFrames[Index1]));
+
+				// Plasticity should not be turned on in the middle of simulation.
+				const bool bUseLinearPlasticity = JointSettings.LinearPlasticityLimit != FLT_MAX;
+				if(bUseLinearPlasticity)
+				{
+					const bool bIsCOMDistanceInitialized = JointSettings.InitCOMDistance >= 0;
+					if(!bIsCOMDistanceInitialized)
+					{
+						ConstraintSettings[JointIndex].InitCOMDistance = (FParticleUtilitiesXR::GetCoMWorldPosition(Particle0) - FParticleUtilitiesXR::GetCoMWorldPosition(Particle1)).Size();
+					}
+				} 
 			}
 		}
 	}
@@ -1511,7 +1523,8 @@ namespace Chaos
 				}
 			}
 			// Assuming that the dimensions which are locked or have no targets are 0. in LinearDrivePositionTarget
-			if((LinearDisplacement - JointSettings.LinearDrivePositionTarget).SizeSquared() > JointSettings.LinearPlasticityLimit * JointSettings.LinearPlasticityLimit)
+			FReal LinearPlasticityDistanceThreshold = JointSettings.LinearPlasticityLimit * JointSettings.InitCOMDistance;
+			if((LinearDisplacement - JointSettings.LinearDrivePositionTarget).SizeSquared() > LinearPlasticityDistanceThreshold * LinearPlasticityDistanceThreshold)
 			{
 				JointSettings.LinearDrivePositionTarget = LinearDisplacement;
 			}
diff --git a/Engine/Source/Runtime/Experimental/Chaos/Public/Chaos/PBDJointConstraintTypes.h b/Engine/Source/Runtime/Experimental/Chaos/Public/Chaos/PBDJointConstraintTypes.h
--- a/Engine/Source/Runtime/Experimental/Chaos/Public/Chaos/PBDJointConstraintTypes.h
+++ b/Engine/Source/Runtime/Experimental/Chaos/Public/Chaos/PBDJointConstraintTypes.h
@@ -163,6 +163,7 @@ namespace Chaos
 
 		FReal LinearBreakForce;
 		FReal LinearPlasticityLimit;
+		FReal InitCOMDistance;
 		FReal AngularBreakTorque;
 		FReal AngularPlasticityLimit;
 
