Gizmos: refactor Modeling Mode gizmo creation out of InteractiveGizmoManager. Editor will use other "default" transform gizmo implementations, and so the UTransformGizmo creation helper functions do not belong in GizmoManager. Instead a UTransformGizmoContextObject now provides this functionality. ModelingToolsEditorMode (and any other modes/systems that want to use these gizmo convenience functions) creates an instance of UTransformGizmoContextObject and registers it with the ContextObjectStore. Calling code can spawn a new UTransformGizmo by looking this object up in the ContextStore and calling it's helper functions. Static versions of the helper functions in the UE::TransformGizmoUtil:: namespace provide a single-line interface that replaces the previous GizmoManager call sites in the MeshModelingTools library.

IntervalGizmo is now just registered and unregistered as needed by the MeshSpaceDeformerTool, as this is the only place it is currently used.

Previous implementation in InteractiveGizmoManager is left intact as there are a few uses outside of MeshModelingTools that need to be cleaned up before it can be deleted.

UTransformGizmo now requires it's Builder to tell it which sub-gizmo identifier strings to pass to the GizmoManager to create axis/plane/rotation sub-gizmos (and the code that registers the Builder must then provide these strings). This cleans up previous explicit references to UInteractiveGizmoManager static strings from UTransformGizmo.

#rnx

diff --git a/Engine/Plugins/Experimental/MeshModelingToolset/Source/MeshModelingTools/Private/CSGMeshesTool.cpp b/Engine/Plugins/Experimental/MeshModelingToolset/Source/MeshModelingTools/Private/CSGMeshesTool.cpp
--- a/Engine/Plugins/Experimental/MeshModelingToolset/Source/MeshModelingTools/Private/CSGMeshesTool.cpp
+++ b/Engine/Plugins/Experimental/MeshModelingToolset/Source/MeshModelingTools/Private/CSGMeshesTool.cpp
@@ -4,6 +4,7 @@
 #include "CompositionOps/BooleanMeshesOp.h"
 
 #include "ToolSetupUtil.h"
+#include "BaseGizmos/TransformProxy.h"
 
 #include "DynamicMesh3.h"
 
diff --git a/Engine/Plugins/Experimental/MeshModelingToolset/Source/MeshModelingTools/Private/CutMeshWithMeshTool.cpp b/Engine/Plugins/Experimental/MeshModelingToolset/Source/MeshModelingTools/Private/CutMeshWithMeshTool.cpp
--- a/Engine/Plugins/Experimental/MeshModelingToolset/Source/MeshModelingTools/Private/CutMeshWithMeshTool.cpp
+++ b/Engine/Plugins/Experimental/MeshModelingToolset/Source/MeshModelingTools/Private/CutMeshWithMeshTool.cpp
@@ -3,6 +3,7 @@
 #include "CutMeshWithMeshTool.h"
 #include "CompositionOps/BooleanMeshesOp.h"
 #include "ToolSetupUtil.h"
+#include "BaseGizmos/TransformGizmoUtil.h"
 #include "Selection/ToolSelectionUtil.h"
 #include "AssetGenerationUtil.h"
 #include "DynamicMesh3.h"
diff --git a/Engine/Plugins/Experimental/MeshModelingToolset/Source/MeshModelingTools/Private/DrawPolygonTool.cpp b/Engine/Plugins/Experimental/MeshModelingToolset/Source/MeshModelingTools/Private/DrawPolygonTool.cpp
--- a/Engine/Plugins/Experimental/MeshModelingToolset/Source/MeshModelingTools/Private/DrawPolygonTool.cpp
+++ b/Engine/Plugins/Experimental/MeshModelingToolset/Source/MeshModelingTools/Private/DrawPolygonTool.cpp
@@ -27,7 +27,7 @@
 #include "DynamicMeshEditor.h"
 
 #include "BaseGizmos/GizmoComponents.h"
-#include "BaseGizmos/TransformGizmo.h"
+#include "BaseGizmos/TransformGizmoUtil.h"
 #include "Drawing/MeshDebugDrawing.h"
 
 #include "Selection/SelectClickedAction.h"
@@ -126,7 +126,7 @@ void UDrawPolygonTool::Setup()
 	// parent of any Components in this case, we just use it's transform and change delegate.
 	PlaneTransformProxy = NewObject<UTransformProxy>(this);
 	PlaneTransformProxy->SetTransform(FTransform((FQuat)DrawPlaneOrientation, (FVector)DrawPlaneOrigin));
-	PlaneTransformGizmo = GetToolManager()->GetPairedGizmoManager()->CreateCustomTransformGizmo(
+	PlaneTransformGizmo = UE::TransformGizmoUtil::CreateCustomTransformGizmo(GetToolManager(),
 		ETransformGizmoSubElements::StandardTranslateRotate, this);
 	PlaneTransformGizmo->SetActiveTarget(PlaneTransformProxy, GetToolManager());
 	// listen for changes to the proxy and update the plane when that happens
@@ -986,7 +986,7 @@ void UDrawPolygonTool::UpdateShowGizmoState(bool bNewVisibility)
 	else
 	{
 		if (!PlaneTransformGizmo) {
-			PlaneTransformGizmo = GetToolManager()->GetPairedGizmoManager()->CreateCustomTransformGizmo(
+			PlaneTransformGizmo = UE::TransformGizmoUtil::CreateCustomTransformGizmo(GetToolManager(),
 				ETransformGizmoSubElements::StandardTranslateRotate, this);
 		}
 		PlaneTransformGizmo->SetActiveTarget(PlaneTransformProxy, GetToolManager());
diff --git a/Engine/Plugins/Experimental/MeshModelingToolset/Source/MeshModelingTools/Private/DynamicMeshSculptTool.cpp b/Engine/Plugins/Experimental/MeshModelingToolset/Source/MeshModelingTools/Private/DynamicMeshSculptTool.cpp
--- a/Engine/Plugins/Experimental/MeshModelingToolset/Source/MeshModelingTools/Private/DynamicMeshSculptTool.cpp
+++ b/Engine/Plugins/Experimental/MeshModelingToolset/Source/MeshModelingTools/Private/DynamicMeshSculptTool.cpp
@@ -31,7 +31,7 @@
 
 #include "InteractiveGizmoManager.h"
 #include "BaseGizmos/GizmoComponents.h"
-#include "BaseGizmos/TransformGizmo.h"
+#include "BaseGizmos/TransformGizmoUtil.h"
 #include "UObject/ObjectMacros.h"
 #include "Materials/Material.h"
 #include "Materials/MaterialInstanceDynamic.h"
@@ -2671,7 +2671,7 @@ void UDynamicMeshSculptTool::UpdateFixedPlaneGizmoVisibility(bool bVisible)
 	{
 		if (PlaneTransformGizmo == nullptr)
 		{
-			PlaneTransformGizmo = GetToolManager()->GetPairedGizmoManager()->CreateCustomTransformGizmo(
+			PlaneTransformGizmo = UE::TransformGizmoUtil::CreateCustomTransformGizmo(GetToolManager(),
 				ETransformGizmoSubElements::StandardTranslateRotate, this);
 			PlaneTransformGizmo->bUseContextCoordinateSystem = false;
 			PlaneTransformGizmo->CurrentCoordinateSystem = EToolContextCoordinateSystem::Local;
diff --git a/Engine/Plugins/Experimental/MeshModelingToolset/Source/MeshModelingTools/Private/EditPivotTool.cpp b/Engine/Plugins/Experimental/MeshModelingToolset/Source/MeshModelingTools/Private/EditPivotTool.cpp
--- a/Engine/Plugins/Experimental/MeshModelingToolset/Source/MeshModelingTools/Private/EditPivotTool.cpp
+++ b/Engine/Plugins/Experimental/MeshModelingToolset/Source/MeshModelingTools/Private/EditPivotTool.cpp
@@ -17,7 +17,7 @@
 #include "Physics/ComponentCollisionUtil.h"
 
 #include "BaseGizmos/GizmoComponents.h"
-#include "BaseGizmos/TransformGizmo.h"
+#include "BaseGizmos/TransformGizmoUtil.h"
 
 #include "Components/PrimitiveComponent.h"
 #include "Engine/World.h"
@@ -316,7 +316,7 @@ void UEditPivotTool::SetActiveGizmos_Single(bool bLocalRotations)
 	{
 		Transformable.TransformProxy->AddComponent(TargetComponentInterface(ComponentIdx)->GetOwnerComponent());
 	}
-	Transformable.TransformGizmo = GizmoManager->CreateCustomTransformGizmo(
+	Transformable.TransformGizmo = UE::TransformGizmoUtil::CreateCustomTransformGizmo(GizmoManager,
 		ETransformGizmoSubElements::StandardTranslateRotate, this
 	);
 	Transformable.TransformGizmo->SetActiveTarget(Transformable.TransformProxy, GetToolManager());
diff --git a/Engine/Plugins/Experimental/MeshModelingToolset/Source/MeshModelingTools/Private/MeshSpaceDeformerTool.cpp b/Engine/Plugins/Experimental/MeshModelingToolset/Source/MeshModelingTools/Private/MeshSpaceDeformerTool.cpp
--- a/Engine/Plugins/Experimental/MeshModelingToolset/Source/MeshModelingTools/Private/MeshSpaceDeformerTool.cpp
+++ b/Engine/Plugins/Experimental/MeshModelingToolset/Source/MeshModelingTools/Private/MeshSpaceDeformerTool.cpp
@@ -19,7 +19,7 @@
 #include "Selection/SelectClickedAction.h"
 
 #include "BaseGizmos/GizmoComponents.h"
-#include "BaseGizmos/TransformGizmo.h"
+#include "BaseGizmos/TransformGizmoUtil.h"
 #include "BaseGizmos/IntervalGizmo.h"
 #include "PositionPlaneGizmo.h"
 #include "MeshDescriptionToDynamicMesh.h"
@@ -200,7 +200,7 @@ void UMeshSpaceDeformerTool::Setup()
 	// parent of any Components in this case, we just use it's transform and change delegate.
 	TransformProxy = NewObject<UTransformProxy>(this);
 	TransformProxy->SetTransform(GizmoFrame.ToFTransform());
-	TransformGizmo = GetToolManager()->GetPairedGizmoManager()->CreateCustomTransformGizmo(
+	TransformGizmo = UE::TransformGizmoUtil::CreateCustomTransformGizmo(GetToolManager(),
 		ETransformGizmoSubElements::StandardTranslateRotate, this);
 	TransformGizmo->bSnapToWorldGrid = Settings->bSnapToWorldGrid;
 
@@ -244,6 +244,7 @@ void UMeshSpaceDeformerTool::Setup()
 	});
 
 	// add the interval gizmo
+	GetToolManager()->GetPairedGizmoManager()->RegisterGizmoType(UIntervalGizmo::GizmoName, NewObject<UIntervalGizmoBuilder>());
 	IntervalGizmo = GetToolManager()->GetPairedGizmoManager()->CreateGizmo<UIntervalGizmo>(UIntervalGizmo::GizmoName, TEXT("MeshSpaceDefomerInterval"), this);
 
 	// wire in the transform and the interval sources.
@@ -338,6 +339,7 @@ void UMeshSpaceDeformerTool::Shutdown(EToolShutdownType ShutdownType)
 
 	UInteractiveGizmoManager* GizmoManager = GetToolManager()->GetPairedGizmoManager();
 	GizmoManager->DestroyAllGizmosByOwner(this);
+	GizmoManager->DeregisterGizmoType(UIntervalGizmo::GizmoName);
 }
 
 
diff --git a/Engine/Plugins/Experimental/MeshModelingToolset/Source/MeshModelingTools/Private/Sculpting/MeshSculptToolBase.cpp b/Engine/Plugins/Experimental/MeshModelingToolset/Source/MeshModelingTools/Private/Sculpting/MeshSculptToolBase.cpp
--- a/Engine/Plugins/Experimental/MeshModelingToolset/Source/MeshModelingTools/Private/Sculpting/MeshSculptToolBase.cpp
+++ b/Engine/Plugins/Experimental/MeshModelingToolset/Source/MeshModelingTools/Private/Sculpting/MeshSculptToolBase.cpp
@@ -8,7 +8,7 @@
 #include "ToolSceneQueriesUtil.h"
 #include "PreviewMesh.h"
 #include "BaseGizmos/GizmoComponents.h"
-#include "BaseGizmos/TransformGizmo.h"
+#include "BaseGizmos/TransformGizmoUtil.h"
 #include "Drawing/MeshDebugDrawing.h"
 
 #include "Sculpting/StampFalloffs.h"
@@ -1204,7 +1204,7 @@ void UMeshSculptToolBase::UpdateFixedPlaneGizmoVisibility(bool bVisible)
 	{
 		if (PlaneTransformGizmo == nullptr)
 		{
-			PlaneTransformGizmo = GetToolManager()->GetPairedGizmoManager()->CreateCustomTransformGizmo(
+			PlaneTransformGizmo = UE::TransformGizmoUtil::CreateCustomTransformGizmo(GetToolManager(),
 				ETransformGizmoSubElements::StandardTranslateRotate, this);
 			PlaneTransformGizmo->bUseContextCoordinateSystem = false;
 			PlaneTransformGizmo->CurrentCoordinateSystem = EToolContextCoordinateSystem::Local;
diff --git a/Engine/Plugins/Experimental/MeshModelingToolset/Source/MeshModelingTools/Private/SelfUnionMeshesTool.cpp b/Engine/Plugins/Experimental/MeshModelingToolset/Source/MeshModelingTools/Private/SelfUnionMeshesTool.cpp
--- a/Engine/Plugins/Experimental/MeshModelingToolset/Source/MeshModelingTools/Private/SelfUnionMeshesTool.cpp
+++ b/Engine/Plugins/Experimental/MeshModelingToolset/Source/MeshModelingTools/Private/SelfUnionMeshesTool.cpp
@@ -4,6 +4,7 @@
 #include "CompositionOps/SelfUnionMeshesOp.h"
 #include "InteractiveToolManager.h"
 #include "ToolSetupUtil.h"
+#include "BaseGizmos/TransformProxy.h"
 
 #include "DynamicMesh3.h"
 #include "DynamicMeshEditor.h"
diff --git a/Engine/Plugins/Experimental/MeshModelingToolset/Source/MeshModelingTools/Private/TransformMeshesTool.cpp b/Engine/Plugins/Experimental/MeshModelingToolset/Source/MeshModelingTools/Private/TransformMeshesTool.cpp
--- a/Engine/Plugins/Experimental/MeshModelingToolset/Source/MeshModelingTools/Private/TransformMeshesTool.cpp
+++ b/Engine/Plugins/Experimental/MeshModelingToolset/Source/MeshModelingTools/Private/TransformMeshesTool.cpp
@@ -11,7 +11,7 @@
 #include "ToolSceneQueriesUtil.h"
 
 #include "BaseGizmos/GizmoComponents.h"
-#include "BaseGizmos/TransformGizmo.h"
+#include "BaseGizmos/TransformGizmoUtil.h"
 
 #include "Components/PrimitiveComponent.h"
 #include "Engine/World.h"
@@ -239,7 +239,7 @@ void UTransformMeshesTool::SetActiveGizmos_Single(bool bLocalRotations)
 	bool bCanNonUniformScale = Targets.Num() == 1 || bLocalRotations;
 	ETransformGizmoSubElements GizmoElements = (bCanNonUniformScale) ?
 		ETransformGizmoSubElements::FullTranslateRotateScale : ETransformGizmoSubElements::TranslateRotateUniformScale;
-	Transformable.TransformGizmo = GizmoManager->CreateCustomRepositionableTransformGizmo(GizmoElements, this);
+	Transformable.TransformGizmo = UE::TransformGizmoUtil::CreateCustomRepositionableTransformGizmo(GetToolManager(), GizmoElements, this);
 	Transformable.TransformGizmo->SetActiveTarget(Transformable.TransformProxy);
 
 	DragAlignmentMechanic->AddToGizmo(Transformable.TransformGizmo, &ComponentsToIgnoreInAlignment);
@@ -264,7 +264,7 @@ void UTransformMeshesTool::SetActiveGizmos_PerObject()
 		ComponentsToIgnoreInAlignment.Add(TargetComponent->GetOwnerComponent());
 
 		ETransformGizmoSubElements GizmoElements = ETransformGizmoSubElements::FullTranslateRotateScale;
-		Transformable.TransformGizmo = GizmoManager->CreateCustomRepositionableTransformGizmo(GizmoElements, this);
+		Transformable.TransformGizmo = UE::TransformGizmoUtil::CreateCustomRepositionableTransformGizmo(GetToolManager(), GizmoElements, this);
 		Transformable.TransformGizmo->SetActiveTarget(Transformable.TransformProxy);
 
 		DragAlignmentMechanic->AddToGizmo(Transformable.TransformGizmo, &ComponentsToIgnoreInAlignment);
diff --git a/Engine/Plugins/Experimental/MeshModelingToolset/Source/MeshModelingTools/Private/UVProjectionTool.cpp b/Engine/Plugins/Experimental/MeshModelingToolset/Source/MeshModelingTools/Private/UVProjectionTool.cpp
--- a/Engine/Plugins/Experimental/MeshModelingToolset/Source/MeshModelingTools/Private/UVProjectionTool.cpp
+++ b/Engine/Plugins/Experimental/MeshModelingToolset/Source/MeshModelingTools/Private/UVProjectionTool.cpp
@@ -19,7 +19,7 @@
 #include "AssetGenerationUtil.h"
 
 #include "BaseGizmos/GizmoComponents.h"
-#include "BaseGizmos/TransformGizmo.h"
+#include "BaseGizmos/TransformGizmoUtil.h"
 
 #include "TargetInterfaces/MaterialProvider.h"
 #include "TargetInterfaces/MeshDescriptionCommitter.h"
@@ -193,7 +193,7 @@ void UUVProjectionTool::UpdateNumPreviews()
 			TransformProxy->SetTransform(LocalXF * TargetComponentInterface(PreviewIdx)->GetWorldTransform());
 			TransformProxy->OnTransformChanged.AddUObject(this, &UUVProjectionTool::TransformChanged);
 
-			UTransformGizmo* TransformGizmo = TransformGizmos.Add_GetRef(GizmoManager->Create3AxisTransformGizmo(this));
+			UTransformGizmo* TransformGizmo = TransformGizmos.Add_GetRef(UE::TransformGizmoUtil::Create3AxisTransformGizmo(GizmoManager, this));
 			TransformGizmo->SetActiveTarget(TransformProxy, GetToolManager());
 		}
 		check(TransformProxies.Num() == TargetNumPreview);
diff --git a/Engine/Plugins/Experimental/MeshModelingToolset/Source/ModelingComponents/Private/BaseTools/BaseCreateFromSelectedTool.cpp b/Engine/Plugins/Experimental/MeshModelingToolset/Source/ModelingComponents/Private/BaseTools/BaseCreateFromSelectedTool.cpp
--- a/Engine/Plugins/Experimental/MeshModelingToolset/Source/ModelingComponents/Private/BaseTools/BaseCreateFromSelectedTool.cpp
+++ b/Engine/Plugins/Experimental/MeshModelingToolset/Source/ModelingComponents/Private/BaseTools/BaseCreateFromSelectedTool.cpp
@@ -5,6 +5,7 @@
 #include "ToolBuilderUtil.h"
 #include "ToolSetupUtil.h"
 #include "SimpleDynamicMeshComponent.h"
+#include "BaseGizmos/TransformGizmoUtil.h"
 #include "Async/Async.h"
 
 #include "MeshNormals.h"
@@ -184,7 +185,7 @@ void UBaseCreateFromSelectedTool::SetTransformGizmos()
 	for (int ComponentIdx = 0; ComponentIdx < Targets.Num(); ComponentIdx++)
 	{
 		UTransformProxy* Proxy = TransformProxies.Add_GetRef(NewObject<UTransformProxy>(this));
-		UTransformGizmo* Gizmo = TransformGizmos.Add_GetRef(GizmoManager->Create3AxisTransformGizmo(this));
+		UTransformGizmo* Gizmo = TransformGizmos.Add_GetRef(UE::TransformGizmoUtil::Create3AxisTransformGizmo(GizmoManager, this));
 		Proxy->SetTransform(TargetComponentInterface(ComponentIdx)->GetWorldTransform());
 		Gizmo->SetActiveTarget(Proxy, GetToolManager());
 		Proxy->OnTransformChanged.AddUObject(this, &UBaseCreateFromSelectedTool::TransformChanged);
diff --git a/Engine/Plugins/Experimental/MeshModelingToolset/Source/ModelingComponents/Private/Mechanics/ConstructionPlaneMechanic.cpp b/Engine/Plugins/Experimental/MeshModelingToolset/Source/ModelingComponents/Private/Mechanics/ConstructionPlaneMechanic.cpp
--- a/Engine/Plugins/Experimental/MeshModelingToolset/Source/ModelingComponents/Private/Mechanics/ConstructionPlaneMechanic.cpp
+++ b/Engine/Plugins/Experimental/MeshModelingToolset/Source/ModelingComponents/Private/Mechanics/ConstructionPlaneMechanic.cpp
@@ -4,7 +4,7 @@
 #include "InteractiveToolManager.h"
 #include "InteractiveGizmoManager.h"
 #include "BaseBehaviors/SingleClickBehavior.h"
-#include "BaseGizmos/TransformGizmo.h"
+#include "BaseGizmos/TransformGizmoUtil.h"
 #include "Drawing/MeshDebugDrawing.h"
 #include "ToolSceneQueriesUtil.h"
 
@@ -31,7 +31,7 @@ void UConstructionPlaneMechanic::Initialize(UWorld* TargetWorld, const FFrame3d&
 	UInteractiveGizmoManager* GizmoManager = GetParentTool()->GetToolManager()->GetPairedGizmoManager();
 
 	PlaneTransformProxy = NewObject<UTransformProxy>(this);
-	PlaneTransformGizmo = GizmoManager->CreateCustomTransformGizmo(
+	PlaneTransformGizmo = UE::TransformGizmoUtil::CreateCustomTransformGizmo(GizmoManager,
 		ETransformGizmoSubElements::StandardTranslateRotate, this);
 	PlaneTransformProxy->OnTransformChanged.AddUObject(this, &UConstructionPlaneMechanic::TransformChanged);
 
diff --git a/Engine/Plugins/Experimental/MeshModelingToolset/Source/ModelingComponents/Private/Mechanics/CurveControlPointsMechanic.cpp b/Engine/Plugins/Experimental/MeshModelingToolset/Source/ModelingComponents/Private/Mechanics/CurveControlPointsMechanic.cpp
--- a/Engine/Plugins/Experimental/MeshModelingToolset/Source/ModelingComponents/Private/Mechanics/CurveControlPointsMechanic.cpp
+++ b/Engine/Plugins/Experimental/MeshModelingToolset/Source/ModelingComponents/Private/Mechanics/CurveControlPointsMechanic.cpp
@@ -4,7 +4,7 @@
 
 #include "BaseBehaviors/SingleClickBehavior.h"
 #include "BaseBehaviors/MouseHoverBehavior.h"
-#include "BaseGizmos/TransformGizmo.h"
+#include "BaseGizmos/TransformGizmoUtil.h"
 #include "BaseGizmos/TransformProxy.h"
 #include "Drawing/PreviewGeometryActor.h"
 #include "Drawing/LineSetComponent.h"
@@ -115,7 +115,7 @@ void UCurveControlPointsMechanic::Setup(UInteractiveTool* ParentToolIn)
 
 	UInteractiveGizmoManager* GizmoManager = GetParentTool()->GetToolManager()->GetPairedGizmoManager();
 	PointTransformProxy = NewObject<UTransformProxy>(this);
-	PointTransformGizmo = GizmoManager->CreateCustomTransformGizmo(
+	PointTransformGizmo = UE::TransformGizmoUtil::CreateCustomTransformGizmo(GizmoManager,
 		ETransformGizmoSubElements::TranslateAxisX | ETransformGizmoSubElements::TranslateAxisY | ETransformGizmoSubElements::TranslatePlaneXY,
 		this);
 	PointTransformProxy->OnTransformChanged.AddUObject(this, &UCurveControlPointsMechanic::GizmoTransformChanged);
diff --git a/Engine/Plugins/Experimental/MeshModelingToolset/Source/ModelingComponents/Private/Mechanics/LatticeControlPointsMechanic.cpp b/Engine/Plugins/Experimental/MeshModelingToolset/Source/ModelingComponents/Private/Mechanics/LatticeControlPointsMechanic.cpp
--- a/Engine/Plugins/Experimental/MeshModelingToolset/Source/ModelingComponents/Private/Mechanics/LatticeControlPointsMechanic.cpp
+++ b/Engine/Plugins/Experimental/MeshModelingToolset/Source/ModelingComponents/Private/Mechanics/LatticeControlPointsMechanic.cpp
@@ -5,7 +5,7 @@
 #include "BaseBehaviors/SingleClickBehavior.h"
 #include "BaseBehaviors/MouseHoverBehavior.h"
 #include "BaseBehaviors/ClickDragBehavior.h"
-#include "BaseGizmos/TransformGizmo.h"
+#include "BaseGizmos/TransformGizmoUtil.h"
 #include "BaseGizmos/TransformProxy.h"
 #include "Drawing/PreviewGeometryActor.h"
 #include "Drawing/LineSetComponent.h"
@@ -98,7 +98,7 @@ void ULatticeControlPointsMechanic::Setup(UInteractiveTool* ParentToolIn)
 	PointTransformProxy = NewObject<UTransformProxy>(this);
 	
 	// TODO: Maybe don't have the gizmo's axes flip around when it crosses the origin, if possible?
-	PointTransformGizmo = GizmoManager->CreateCustomTransformGizmo(
+	PointTransformGizmo = UE::TransformGizmoUtil::CreateCustomTransformGizmo(GizmoManager,
 		ETransformGizmoSubElements::FullTranslateRotateScale, this);
 
 	PointTransformProxy->OnTransformChanged.AddUObject(this, &ULatticeControlPointsMechanic::GizmoTransformChanged);
diff --git a/Engine/Plugins/Experimental/MeshModelingToolset/Source/ModelingComponents/Private/Mechanics/SpaceCurveDeformationMechanic.cpp b/Engine/Plugins/Experimental/MeshModelingToolset/Source/ModelingComponents/Private/Mechanics/SpaceCurveDeformationMechanic.cpp
--- a/Engine/Plugins/Experimental/MeshModelingToolset/Source/ModelingComponents/Private/Mechanics/SpaceCurveDeformationMechanic.cpp
+++ b/Engine/Plugins/Experimental/MeshModelingToolset/Source/ModelingComponents/Private/Mechanics/SpaceCurveDeformationMechanic.cpp
@@ -4,7 +4,7 @@
 
 #include "BaseBehaviors/SingleClickBehavior.h"
 #include "BaseBehaviors/MouseHoverBehavior.h"
-#include "BaseGizmos/TransformGizmo.h"
+#include "BaseGizmos/TransformGizmoUtil.h"
 #include "BaseGizmos/TransformProxy.h"
 #include "Drawing/PreviewGeometryActor.h"
 #include "Drawing/LineSetComponent.h"
@@ -72,7 +72,7 @@ void USpaceCurveDeformationMechanic::Setup(UInteractiveTool* ParentToolIn)
 
 	UInteractiveGizmoManager* GizmoManager = GetParentTool()->GetToolManager()->GetPairedGizmoManager();
 	PointTransformProxy = NewObject<UTransformProxy>(this);
-	PointTransformGizmo = GizmoManager->CreateCustomTransformGizmo( ETransformGizmoSubElements::StandardTranslateRotate, GetParentTool());
+	PointTransformGizmo = UE::TransformGizmoUtil::CreateCustomTransformGizmo(GizmoManager, ETransformGizmoSubElements::StandardTranslateRotate, GetParentTool());
 	PointTransformProxy->OnTransformChanged.AddUObject(this, &USpaceCurveDeformationMechanic::GizmoTransformChanged);
 	PointTransformProxy->OnBeginTransformEdit.AddUObject(this, &USpaceCurveDeformationMechanic::GizmoTransformStarted);
 	PointTransformProxy->OnEndTransformEdit.AddUObject(this, &USpaceCurveDeformationMechanic::GizmoTransformEnded);
diff --git a/Engine/Plugins/Experimental/MeshModelingToolset/Source/ModelingComponents/Private/Transforms/MultiTransformer.cpp b/Engine/Plugins/Experimental/MeshModelingToolset/Source/ModelingComponents/Private/Transforms/MultiTransformer.cpp
--- a/Engine/Plugins/Experimental/MeshModelingToolset/Source/ModelingComponents/Private/Transforms/MultiTransformer.cpp
+++ b/Engine/Plugins/Experimental/MeshModelingToolset/Source/ModelingComponents/Private/Transforms/MultiTransformer.cpp
@@ -1,6 +1,7 @@
 // Copyright Epic Games, Inc. All Rights Reserved.
 
 #include "Transforms/MultiTransformer.h"
+#include "BaseGizmos/TransformGizmoUtil.h"
 
 #include "ExplicitUseGeometryMathTypes.h"		// using UE::Geometry::(math types)
 using namespace UE::Geometry;
@@ -236,8 +237,8 @@ void UMultiTransformer::UpdateShowGizmoState(bool bNewVisibility)
 	else
 	{
 		check(TransformGizmo == nullptr);
-		TransformGizmo = bRepositionableGizmo ? GizmoManager->CreateCustomRepositionableTransformGizmo(ActiveGizmoSubElements, this)
-			: GizmoManager->CreateCustomTransformGizmo(ActiveGizmoSubElements, this);
+		TransformGizmo = bRepositionableGizmo ? UE::TransformGizmoUtil::CreateCustomRepositionableTransformGizmo(GizmoManager, ActiveGizmoSubElements, this)
+			: UE::TransformGizmoUtil::CreateCustomTransformGizmo(GizmoManager, ActiveGizmoSubElements, this);
 		if (bForceGizmoCoordSystem)
 		{
 			TransformGizmo->bUseContextCoordinateSystem = false;
diff --git a/Engine/Plugins/Experimental/MeshModelingToolset/Source/ModelingComponents/Public/BaseTools/BaseCreateFromSelectedTool.h b/Engine/Plugins/Experimental/MeshModelingToolset/Source/ModelingComponents/Public/BaseTools/BaseCreateFromSelectedTool.h
--- a/Engine/Plugins/Experimental/MeshModelingToolset/Source/ModelingComponents/Public/BaseTools/BaseCreateFromSelectedTool.h
+++ b/Engine/Plugins/Experimental/MeshModelingToolset/Source/ModelingComponents/Public/BaseTools/BaseCreateFromSelectedTool.h
@@ -5,12 +5,13 @@
 #include "CoreMinimal.h"
 #include "MultiSelectionTool.h"
 #include "InteractiveToolBuilder.h"
-#include "BaseGizmos/TransformGizmo.h"
 #include "MeshOpPreviewHelpers.h"
 #include "PropertySets/OnAcceptProperties.h"
 #include "BaseCreateFromSelectedTool.generated.h"
 
 class IAssetGenerationAPI;
+class UTransformGizmo;
+class UTransformProxy;
 
 /**
  * ToolBuilder for UBaseCreateFromSelectedTool
diff --git a/Engine/Plugins/Experimental/MeshModelingToolset/Source/ModelingComponents/Public/Transforms/MultiTransformer.h b/Engine/Plugins/Experimental/MeshModelingToolset/Source/ModelingComponents/Public/Transforms/MultiTransformer.h
--- a/Engine/Plugins/Experimental/MeshModelingToolset/Source/ModelingComponents/Public/Transforms/MultiTransformer.h
+++ b/Engine/Plugins/Experimental/MeshModelingToolset/Source/ModelingComponents/Public/Transforms/MultiTransformer.h
@@ -9,12 +9,14 @@
 #include "FrameTypes.h"
 #include "InteractiveGizmo.h"
 #include "Mechanics/DragAlignmentMechanic.h"
-#include "BaseGizmos/TransformGizmo.h"
 #include "InteractiveGizmoManager.h"
 
 #include "MultiTransformer.generated.h"
 
 
+class UTransformGizmo;
+class UTransformProxy;
+
 UENUM()
 enum class EMultiTransformerMode
 {
diff --git a/Engine/Plugins/Experimental/ModelingToolsEditorMode/Source/ModelingToolsEditorMode/Private/ModelingToolsEditorMode.cpp b/Engine/Plugins/Experimental/ModelingToolsEditorMode/Source/ModelingToolsEditorMode/Private/ModelingToolsEditorMode.cpp
--- a/Engine/Plugins/Experimental/ModelingToolsEditorMode/Source/ModelingToolsEditorMode/Private/ModelingToolsEditorMode.cpp
+++ b/Engine/Plugins/Experimental/ModelingToolsEditorMode/Source/ModelingToolsEditorMode/Private/ModelingToolsEditorMode.cpp
@@ -14,6 +14,7 @@
 #include "Framework/MultiBox/MultiBoxBuilder.h"
 #include "EditorViewportClient.h"
 #include "EngineAnalytics.h"
+#include "BaseGIzmos/TransformGizmoUtil.h"
 
 #include "Features/IModularFeatures.h"
 #include "ModelingModeToolExtensions.h"
@@ -303,6 +304,10 @@ void UModelingToolsEditorMode::Enter()
 	// register stylus event handler
 	StylusStateTracker = MakeUnique<FStylusStateTracker>();
 
+	// register gizmo helper
+	ModelingGizmoHelper = NewObject<UTransformGizmoContextObject>(ToolsContext->ToolManager);
+	ModelingGizmoHelper->RegisterGizmosWithManager(ToolsContext->ToolManager);
+
 	const FModelingToolsManagerCommands& ToolManagerCommands = FModelingToolsManagerCommands::Get();
 
 	// register tool set
@@ -671,6 +676,8 @@ void UModelingToolsEditorMode::Exit()
 
 	StylusStateTracker = nullptr;
 
+	ModelingGizmoHelper->DeregisterGizmosWithManager(ToolsContext->ToolManager);
+	
 	FModelingModeActionCommands::UnRegisterCommandBindings(Toolkit->GetToolkitCommands());
 
 	// clear realtime viewport override
diff --git a/Engine/Plugins/Experimental/ModelingToolsEditorMode/Source/ModelingToolsEditorMode/Public/ModelingToolsEditorMode.h b/Engine/Plugins/Experimental/ModelingToolsEditorMode/Source/ModelingToolsEditorMode/Public/ModelingToolsEditorMode.h
--- a/Engine/Plugins/Experimental/ModelingToolsEditorMode/Source/ModelingToolsEditorMode/Public/ModelingToolsEditorMode.h
+++ b/Engine/Plugins/Experimental/ModelingToolsEditorMode/Source/ModelingToolsEditorMode/Public/ModelingToolsEditorMode.h
@@ -17,6 +17,7 @@ class FEditorToolAssetAPI;
 class FUICommandList;
 class FStylusStateTracker;		// for stylus events
 class FModelingModeAssetAPI;
+class UTransformGizmoContextObject;
 
 UCLASS(Transient)
 class UModelingToolsEditorMode : public UBaseLegacyWidgetEdMode
@@ -68,6 +69,8 @@ protected:
 
 	TUniquePtr<FStylusStateTracker> StylusStateTracker;
 
+	TObjectPtr<UTransformGizmoContextObject> ModelingGizmoHelper;
+
 	void ModelingModeShortcutRequested(EModelingModeActionCommands Command);
 	void FocusCameraAtCursorHotkey();
 
diff --git a/Engine/Source/Runtime/Experimental/InteractiveToolsFramework/Private/BaseGizmos/RepositionableTransformGizmo.cpp b/Engine/Source/Runtime/Experimental/InteractiveToolsFramework/Private/BaseGizmos/RepositionableTransformGizmo.cpp
--- a/Engine/Source/Runtime/Experimental/InteractiveToolsFramework/Private/BaseGizmos/RepositionableTransformGizmo.cpp
+++ b/Engine/Source/Runtime/Experimental/InteractiveToolsFramework/Private/BaseGizmos/RepositionableTransformGizmo.cpp
@@ -23,6 +23,9 @@ UInteractiveGizmo* URepositionableTransformGizmoBuilder::BuildGizmo(const FToolB
 	// use default gizmo actor if client has not given us a new builder
 	NewGizmo->SetGizmoActorBuilder((GizmoActorBuilder) ? GizmoActorBuilder : MakeShared<FTransformGizmoActorFactory>());
 
+	NewGizmo->SetSubGizmoBuilderIdentifiers(AxisPositionBuilderIdentifier, PlanePositionBuilderIdentifier, AxisAngleBuilderIdentifier);
+
+
 	// override default hover function if proposed
 	if (UpdateHoverFunction)
 	{
diff --git a/Engine/Source/Runtime/Experimental/InteractiveToolsFramework/Private/BaseGizmos/TransformGizmo.cpp b/Engine/Source/Runtime/Experimental/InteractiveToolsFramework/Private/BaseGizmos/TransformGizmo.cpp
--- a/Engine/Source/Runtime/Experimental/InteractiveToolsFramework/Private/BaseGizmos/TransformGizmo.cpp
+++ b/Engine/Source/Runtime/Experimental/InteractiveToolsFramework/Private/BaseGizmos/TransformGizmo.cpp
@@ -224,6 +224,8 @@ UInteractiveGizmo* UTransformGizmoBuilder::BuildGizmo(const FToolBuilderState& S
 	// use default gizmo actor if client has not given us a new builder
 	NewGizmo->SetGizmoActorBuilder( (GizmoActorBuilder) ? GizmoActorBuilder : MakeShared<FTransformGizmoActorFactory>() );
 
+	NewGizmo->SetSubGizmoBuilderIdentifiers(AxisPositionBuilderIdentifier, PlanePositionBuilderIdentifier, AxisAngleBuilderIdentifier);
+
 	// override default hover function if proposed
 	if (UpdateHoverFunction)
 	{
@@ -251,6 +253,13 @@ void UTransformGizmo::SetGizmoActorBuilder(TSharedPtr<FTransformGizmoActorFactor
 	GizmoActorBuilder = Builder;
 }
 
+void UTransformGizmo::SetSubGizmoBuilderIdentifiers(FString AxisPositionBuilderIdentifierIn, FString PlanePositionBuilderIdentifierIn, FString AxisAngleBuilderIdentifierIn)
+{
+	AxisPositionBuilderIdentifier = AxisPositionBuilderIdentifierIn;
+	PlanePositionBuilderIdentifier = PlanePositionBuilderIdentifierIn;
+	AxisAngleBuilderIdentifier = AxisAngleBuilderIdentifierIn;
+}
+
 void UTransformGizmo::SetUpdateHoverFunction(TFunction<void(UPrimitiveComponent*, bool)> HoverFunction)
 {
 	UpdateHoverFunction = HoverFunction;
@@ -647,8 +656,7 @@ UInteractiveGizmo* UTransformGizmo::AddAxisTranslationGizmo(
 	IGizmoStateTarget* StateTargetIn)
 {
 	// create axis-position gizmo, axis-position parameter will drive translation
-	UAxisPositionGizmo* TranslateGizmo = Cast<UAxisPositionGizmo>(GetGizmoManager()->CreateGizmo(
-		UInteractiveGizmoManager::DefaultAxisPositionBuilderIdentifier));
+	UAxisPositionGizmo* TranslateGizmo = Cast<UAxisPositionGizmo>(GetGizmoManager()->CreateGizmo(AxisPositionBuilderIdentifier));
 	check(TranslateGizmo);
 
 	// axis source provides the translation axis
@@ -688,8 +696,7 @@ UInteractiveGizmo* UTransformGizmo::AddPlaneTranslationGizmo(
 	IGizmoStateTarget* StateTargetIn)
 {
 	// create axis-position gizmo, axis-position parameter will drive translation
-	UPlanePositionGizmo* TranslateGizmo = Cast<UPlanePositionGizmo>(GetGizmoManager()->CreateGizmo(
-		UInteractiveGizmoManager::DefaultPlanePositionBuilderIdentifier));
+	UPlanePositionGizmo* TranslateGizmo = Cast<UPlanePositionGizmo>(GetGizmoManager()->CreateGizmo(PlanePositionBuilderIdentifier));
 	check(TranslateGizmo);
 
 	// axis source provides the translation axis
@@ -731,8 +738,7 @@ UInteractiveGizmo* UTransformGizmo::AddAxisRotationGizmo(
 	IGizmoStateTarget* StateTargetIn)
 {
 	// create axis-angle gizmo, angle will drive axis-rotation
-	UAxisAngleGizmo* RotateGizmo = Cast<UAxisAngleGizmo>(GetGizmoManager()->CreateGizmo(
-		UInteractiveGizmoManager::DefaultAxisAngleBuilderIdentifier));
+	UAxisAngleGizmo* RotateGizmo = Cast<UAxisAngleGizmo>(GetGizmoManager()->CreateGizmo(AxisAngleBuilderIdentifier));
 	check(RotateGizmo);
 
 	// axis source provides the rotation axis
@@ -773,8 +779,7 @@ UInteractiveGizmo* UTransformGizmo::AddAxisScaleGizmo(
 	IGizmoStateTarget* StateTargetIn)
 {
 	// create axis-position gizmo, axis-position parameter will drive scale
-	UAxisPositionGizmo* ScaleGizmo = Cast<UAxisPositionGizmo>(GetGizmoManager()->CreateGizmo(
-		UInteractiveGizmoManager::DefaultAxisPositionBuilderIdentifier));
+	UAxisPositionGizmo* ScaleGizmo = Cast<UAxisPositionGizmo>(GetGizmoManager()->CreateGizmo(AxisPositionBuilderIdentifier));
 	ScaleGizmo->bEnableSignedAxis = true;
 	check(ScaleGizmo);
 
@@ -809,8 +814,7 @@ UInteractiveGizmo* UTransformGizmo::AddPlaneScaleGizmo(
 	IGizmoStateTarget* StateTargetIn)
 {
 	// create axis-position gizmo, axis-position parameter will drive scale
-	UPlanePositionGizmo* ScaleGizmo = Cast<UPlanePositionGizmo>(GetGizmoManager()->CreateGizmo(
-		UInteractiveGizmoManager::DefaultPlanePositionBuilderIdentifier));
+	UPlanePositionGizmo* ScaleGizmo = Cast<UPlanePositionGizmo>(GetGizmoManager()->CreateGizmo(PlanePositionBuilderIdentifier));
 	ScaleGizmo->bEnableSignedAxis = true;
 	check(ScaleGizmo);
 
@@ -848,8 +852,7 @@ UInteractiveGizmo* UTransformGizmo::AddUniformScaleGizmo(
 	IGizmoStateTarget* StateTargetIn)
 {
 	// create plane-position gizmo, plane-position parameter will drive scale
-	UPlanePositionGizmo* ScaleGizmo = Cast<UPlanePositionGizmo>(GetGizmoManager()->CreateGizmo(
-		UInteractiveGizmoManager::DefaultPlanePositionBuilderIdentifier));
+	UPlanePositionGizmo* ScaleGizmo = Cast<UPlanePositionGizmo>(GetGizmoManager()->CreateGizmo(PlanePositionBuilderIdentifier));
 	check(ScaleGizmo);
 
 	// axis source provides the translation plane
diff --git a/Engine/Source/Runtime/Experimental/InteractiveToolsFramework/Private/BaseGizmos/TransformGizmoUtil.cpp b/Engine/Source/Runtime/Experimental/InteractiveToolsFramework/Private/BaseGizmos/TransformGizmoUtil.cpp
new file mode 100644
--- /dev/null
+++ b/Engine/Source/Runtime/Experimental/InteractiveToolsFramework/Private/BaseGizmos/TransformGizmoUtil.cpp
@@ -0,0 +1,180 @@
+// Copyright Epic Games, Inc. All Rights Reserved.
+
+#include "BaseGizmos/TransformGizmoUtil.h"
+
+#include "ContextObjectStore.h"
+#include "InteractiveGizmoManager.h"
+#include "InteractiveToolManager.h"
+
+#include "BaseGizmos/AxisPositionGizmo.h"
+#include "BaseGizmos/PlanePositionGizmo.h"
+#include "BaseGizmos/AxisAngleGizmo.h"
+#include "BaseGizmos/TransformGizmo.h"
+#include "BaseGizmos/RepositionableTransformGizmo.h"
+
+const FString UTransformGizmoContextObject::DefaultAxisPositionBuilderIdentifier = TEXT("Util_StandardXFormAxisTranslationGizmo");
+const FString UTransformGizmoContextObject::DefaultPlanePositionBuilderIdentifier = TEXT("Util_StandardXFormPlaneTranslationGizmo");
+const FString UTransformGizmoContextObject::DefaultAxisAngleBuilderIdentifier = TEXT("Util_StandardXFormAxisRotationGizmo");
+const FString UTransformGizmoContextObject::DefaultThreeAxisTransformBuilderIdentifier = TEXT("Util_DefaultThreeAxisTransformBuilderIdentifier");
+const FString UTransformGizmoContextObject::CustomThreeAxisTransformBuilderIdentifier = TEXT("Util_CustomThreeAxisTransformBuilderIdentifier");
+const FString UTransformGizmoContextObject::CustomRepositionableThreeAxisTransformBuilderIdentifier = TEXT("Util_CustomRepositionableThreeAxisTransformBuilderIdentifier");
+
+
+void UTransformGizmoContextObject::RegisterGizmosWithManager(UInteractiveToolManager* ToolManager)
+{
+	if (ensure(!bDefaultGizmosRegistered) == false)
+	{
+		return;
+	}
+
+	UInteractiveGizmoManager* GizmoManager = ToolManager->GetPairedGizmoManager();
+	ToolManager->GetContextObjectStore()->AddContextObject(this);
+
+	UAxisPositionGizmoBuilder* AxisTranslationBuilder = NewObject<UAxisPositionGizmoBuilder>();
+	GizmoManager->RegisterGizmoType(DefaultAxisPositionBuilderIdentifier, AxisTranslationBuilder);
+
+	UPlanePositionGizmoBuilder* PlaneTranslationBuilder = NewObject<UPlanePositionGizmoBuilder>();
+	GizmoManager->RegisterGizmoType(DefaultPlanePositionBuilderIdentifier, PlaneTranslationBuilder);
+
+	UAxisAngleGizmoBuilder* AxisRotationBuilder = NewObject<UAxisAngleGizmoBuilder>();
+	GizmoManager->RegisterGizmoType(DefaultAxisAngleBuilderIdentifier, AxisRotationBuilder);
+
+	UTransformGizmoBuilder* TransformBuilder = NewObject<UTransformGizmoBuilder>();
+	TransformBuilder->AxisPositionBuilderIdentifier = DefaultAxisPositionBuilderIdentifier;
+	TransformBuilder->PlanePositionBuilderIdentifier = DefaultPlanePositionBuilderIdentifier;
+	TransformBuilder->AxisAngleBuilderIdentifier = DefaultAxisAngleBuilderIdentifier;
+	GizmoManager->RegisterGizmoType(DefaultThreeAxisTransformBuilderIdentifier, TransformBuilder);
+
+	GizmoActorBuilder = MakeShared<FTransformGizmoActorFactory>();
+
+	UTransformGizmoBuilder* CustomThreeAxisBuilder = NewObject<UTransformGizmoBuilder>();
+	CustomThreeAxisBuilder->AxisPositionBuilderIdentifier = DefaultAxisPositionBuilderIdentifier;
+	CustomThreeAxisBuilder->PlanePositionBuilderIdentifier = DefaultPlanePositionBuilderIdentifier;
+	CustomThreeAxisBuilder->AxisAngleBuilderIdentifier = DefaultAxisAngleBuilderIdentifier;
+	CustomThreeAxisBuilder->GizmoActorBuilder = GizmoActorBuilder;
+	GizmoManager->RegisterGizmoType(CustomThreeAxisTransformBuilderIdentifier, CustomThreeAxisBuilder);
+
+	URepositionableTransformGizmoBuilder* CustomRepositionableThreeAxisBuilder = NewObject<URepositionableTransformGizmoBuilder>();
+	CustomRepositionableThreeAxisBuilder->AxisPositionBuilderIdentifier = DefaultAxisPositionBuilderIdentifier;
+	CustomRepositionableThreeAxisBuilder->PlanePositionBuilderIdentifier = DefaultPlanePositionBuilderIdentifier;
+	CustomRepositionableThreeAxisBuilder->AxisAngleBuilderIdentifier = DefaultAxisAngleBuilderIdentifier;
+	CustomRepositionableThreeAxisBuilder->GizmoActorBuilder = GizmoActorBuilder;
+	GizmoManager->RegisterGizmoType(CustomRepositionableThreeAxisTransformBuilderIdentifier, CustomRepositionableThreeAxisBuilder);
+
+	bDefaultGizmosRegistered = true;
+
+
+}
+
+void UTransformGizmoContextObject::DeregisterGizmosWithManager(UInteractiveToolManager* ToolManager)
+{
+	UInteractiveGizmoManager* GizmoManager = ToolManager->GetPairedGizmoManager();
+	ToolManager->GetContextObjectStore()->RemoveContextObject(this);
+
+	ensure(bDefaultGizmosRegistered);
+	GizmoManager->DeregisterGizmoType(DefaultAxisPositionBuilderIdentifier);
+	GizmoManager->DeregisterGizmoType(DefaultPlanePositionBuilderIdentifier);
+	GizmoManager->DeregisterGizmoType(DefaultAxisAngleBuilderIdentifier);
+	GizmoManager->DeregisterGizmoType(DefaultThreeAxisTransformBuilderIdentifier);
+	GizmoManager->DeregisterGizmoType(CustomThreeAxisTransformBuilderIdentifier);
+	GizmoManager->DeregisterGizmoType(CustomRepositionableThreeAxisTransformBuilderIdentifier);
+	bDefaultGizmosRegistered = false;
+}
+
+
+
+UTransformGizmo* UTransformGizmoContextObject::Create3AxisTransformGizmo(UInteractiveGizmoManager* GizmoManager, void* Owner, const FString& InstanceIdentifier)
+{
+	if (ensure(bDefaultGizmosRegistered))
+	{
+		UInteractiveGizmo* NewGizmo = GizmoManager->CreateGizmo(DefaultThreeAxisTransformBuilderIdentifier, InstanceIdentifier, Owner);
+		ensure(NewGizmo);
+		return Cast<UTransformGizmo>(NewGizmo);
+	}
+	return nullptr;
+}
+UTransformGizmo* UE::TransformGizmoUtil::Create3AxisTransformGizmo(UInteractiveGizmoManager* GizmoManager, void* Owner, const FString& InstanceIdentifier)
+{
+	if (ensure(GizmoManager))
+	{
+		UTransformGizmoContextObject* UseThis = GizmoManager->GetContextObjectStore()->FindContext<UTransformGizmoContextObject>();
+		if (ensure(UseThis))
+		{
+			return UseThis->Create3AxisTransformGizmo(GizmoManager, Owner, InstanceIdentifier);
+		}
+	}
+	return nullptr;
+}
+UTransformGizmo* UE::TransformGizmoUtil::Create3AxisTransformGizmo(UInteractiveToolManager* ToolManager, void* Owner, const FString& InstanceIdentifier)
+{
+	if (ensure(ToolManager))
+	{
+		return Create3AxisTransformGizmo(ToolManager->GetPairedGizmoManager(), Owner, InstanceIdentifier);
+	}
+	return nullptr;
+}
+
+UTransformGizmo* UTransformGizmoContextObject::CreateCustomTransformGizmo(UInteractiveGizmoManager* GizmoManager, ETransformGizmoSubElements Elements, void* Owner, const FString& InstanceIdentifier)
+{
+	if (ensure(bDefaultGizmosRegistered))
+	{
+		GizmoActorBuilder->EnableElements = Elements;
+		UInteractiveGizmo* NewGizmo = GizmoManager->CreateGizmo(CustomThreeAxisTransformBuilderIdentifier, InstanceIdentifier, Owner);
+		ensure(NewGizmo);
+		return Cast<UTransformGizmo>(NewGizmo);
+	}
+	return nullptr;
+}
+UTransformGizmo* UE::TransformGizmoUtil::CreateCustomTransformGizmo(UInteractiveGizmoManager* GizmoManager, ETransformGizmoSubElements Elements, void* Owner, const FString& InstanceIdentifier)
+{
+	if (ensure(GizmoManager))
+	{
+		UTransformGizmoContextObject* UseThis = GizmoManager->GetContextObjectStore()->FindContext<UTransformGizmoContextObject>();
+		if (ensure(UseThis))
+		{
+			return UseThis->CreateCustomTransformGizmo(GizmoManager, Elements, Owner, InstanceIdentifier);
+		}
+	}
+	return nullptr;
+}
+UTransformGizmo* UE::TransformGizmoUtil::CreateCustomTransformGizmo(UInteractiveToolManager* ToolManager, ETransformGizmoSubElements Elements, void* Owner, const FString& InstanceIdentifier)
+{
+	if (ensure(ToolManager))
+	{
+		return CreateCustomTransformGizmo(ToolManager->GetPairedGizmoManager(), Elements, Owner, InstanceIdentifier);
+	}
+	return nullptr;
+}
+
+
+UTransformGizmo* UTransformGizmoContextObject::CreateCustomRepositionableTransformGizmo(UInteractiveGizmoManager* GizmoManager, ETransformGizmoSubElements Elements, void* Owner, const FString& InstanceIdentifier)
+{
+	if (ensure(bDefaultGizmosRegistered))
+	{
+		GizmoActorBuilder->EnableElements = Elements;
+		UInteractiveGizmo* NewGizmo = GizmoManager->CreateGizmo(CustomRepositionableThreeAxisTransformBuilderIdentifier, InstanceIdentifier, Owner);
+		ensure(NewGizmo);
+		return Cast<UTransformGizmo>(NewGizmo);
+	}
+	return nullptr;
+}
+UTransformGizmo* UE::TransformGizmoUtil::CreateCustomRepositionableTransformGizmo(UInteractiveGizmoManager* GizmoManager, ETransformGizmoSubElements Elements, void* Owner, const FString& InstanceIdentifier)
+{
+	if (ensure(GizmoManager))
+	{
+		UTransformGizmoContextObject* UseThis = GizmoManager->GetContextObjectStore()->FindContext<UTransformGizmoContextObject>();
+		if (ensure(UseThis))
+		{
+			return UseThis->CreateCustomRepositionableTransformGizmo(GizmoManager, Elements, Owner, InstanceIdentifier);
+		}
+	}
+	return nullptr;
+}
+UTransformGizmo* UE::TransformGizmoUtil::CreateCustomRepositionableTransformGizmo(UInteractiveToolManager* ToolManager, ETransformGizmoSubElements Elements, void* Owner, const FString& InstanceIdentifier)
+{
+	if (ensure(ToolManager))
+	{
+		return CreateCustomRepositionableTransformGizmo(ToolManager->GetPairedGizmoManager(), Elements, Owner, InstanceIdentifier);
+	}
+	return nullptr;
+}
diff --git a/Engine/Source/Runtime/Experimental/InteractiveToolsFramework/Private/InteractiveGizmoManager.cpp b/Engine/Source/Runtime/Experimental/InteractiveToolsFramework/Private/InteractiveGizmoManager.cpp
--- a/Engine/Source/Runtime/Experimental/InteractiveToolsFramework/Private/InteractiveGizmoManager.cpp
+++ b/Engine/Source/Runtime/Experimental/InteractiveToolsFramework/Private/InteractiveGizmoManager.cpp
@@ -8,7 +8,6 @@
 #include "BaseGizmos/PlanePositionGizmo.h"
 #include "BaseGizmos/AxisAngleGizmo.h"
 #include "BaseGizmos/TransformGizmo.h"
-#include "BaseGizmos/IntervalGizmo.h"
 #include "BaseGizmos/RepositionableTransformGizmo.h"
 #include "BaseGizmos/ScalableSphereGizmo.h"
 
@@ -262,6 +261,10 @@ void UInteractiveGizmoManager::EmitObjectChange(UObject* TargetObject, TUniquePt
 	TransactionsAPI->AppendChange(TargetObject, MoveTemp(Change), Description );
 }
 
+UContextObjectStore* UInteractiveGizmoManager::GetContextObjectStore() const
+{
+	return Cast<UInteractiveToolsContext>(GetOuter())->ContextObjectStore;
+}
 
 
 
@@ -287,21 +290,27 @@ void UInteractiveGizmoManager::RegisterDefaultGizmos()
 	RegisterGizmoType(DefaultAxisAngleBuilderIdentifier, AxisRotationBuilder);
 
 	UTransformGizmoBuilder* TransformBuilder = NewObject<UTransformGizmoBuilder>();
+	TransformBuilder->AxisPositionBuilderIdentifier = DefaultAxisPositionBuilderIdentifier;
+	TransformBuilder->PlanePositionBuilderIdentifier = DefaultPlanePositionBuilderIdentifier;
+	TransformBuilder->AxisAngleBuilderIdentifier = DefaultAxisAngleBuilderIdentifier;
 	RegisterGizmoType(DefaultThreeAxisTransformBuilderIdentifier, TransformBuilder);
 
 	GizmoActorBuilder = MakeShared<FTransformGizmoActorFactory>();
 
 	UTransformGizmoBuilder* CustomThreeAxisBuilder = NewObject<UTransformGizmoBuilder>();
+	CustomThreeAxisBuilder->AxisPositionBuilderIdentifier = DefaultAxisPositionBuilderIdentifier;
+	CustomThreeAxisBuilder->PlanePositionBuilderIdentifier = DefaultPlanePositionBuilderIdentifier;
+	CustomThreeAxisBuilder->AxisAngleBuilderIdentifier = DefaultAxisAngleBuilderIdentifier;
 	CustomThreeAxisBuilder->GizmoActorBuilder = GizmoActorBuilder;
 	RegisterGizmoType(CustomThreeAxisTransformBuilderIdentifier, CustomThreeAxisBuilder);
 
 	URepositionableTransformGizmoBuilder* CustomRepositionableThreeAxisBuilder = NewObject<URepositionableTransformGizmoBuilder>();
+	CustomRepositionableThreeAxisBuilder->AxisPositionBuilderIdentifier = DefaultAxisPositionBuilderIdentifier;
+	CustomRepositionableThreeAxisBuilder->PlanePositionBuilderIdentifier = DefaultPlanePositionBuilderIdentifier;
+	CustomRepositionableThreeAxisBuilder->AxisAngleBuilderIdentifier = DefaultAxisAngleBuilderIdentifier;
 	CustomRepositionableThreeAxisBuilder->GizmoActorBuilder = GizmoActorBuilder;
 	RegisterGizmoType(CustomRepositionableThreeAxisTransformBuilderIdentifier, CustomRepositionableThreeAxisBuilder);
 
-	UIntervalGizmoBuilder* IntervalGizmoBuilder = NewObject<UIntervalGizmoBuilder>();
-	RegisterGizmoType(UIntervalGizmo::GizmoName, IntervalGizmoBuilder);
-
 	UScalableSphereGizmoBuilder* ScalableSphereBuilder = NewObject<UScalableSphereGizmoBuilder>();
 	RegisterGizmoType(DefaultScalableSphereBuilderIdentifier, ScalableSphereBuilder);
 
diff --git a/Engine/Source/Runtime/Experimental/InteractiveToolsFramework/Public/BaseGizmos/TransformGizmo.h b/Engine/Source/Runtime/Experimental/InteractiveToolsFramework/Public/BaseGizmos/TransformGizmo.h
--- a/Engine/Source/Runtime/Experimental/InteractiveToolsFramework/Public/BaseGizmos/TransformGizmo.h
+++ b/Engine/Source/Runtime/Experimental/InteractiveToolsFramework/Public/BaseGizmos/TransformGizmo.h
@@ -189,6 +189,15 @@ class INTERACTIVETOOLSFRAMEWORK_API UTransformGizmoBuilder : public UInteractive
 	GENERATED_BODY()
 
 public:
+
+	/**
+	 * strings identifing GizmoBuilders already registered with GizmoManager. These builders will be used
+	 * to spawn the various sub-gizmos
+	 */
+	FString AxisPositionBuilderIdentifier;
+	FString PlanePositionBuilderIdentifier;
+	FString AxisAngleBuilderIdentifier;
+
 	/**
 	 * If set, this Actor Builder will be passed to UTransformGizmo instances.
 	 * Otherwise new instances of the base FTransformGizmoActorFactory are created internally.
@@ -244,6 +253,7 @@ public:
 
 	virtual void SetWorld(UWorld* World);
 	virtual void SetGizmoActorBuilder(TSharedPtr<FTransformGizmoActorFactory> Builder);
+	virtual void SetSubGizmoBuilderIdentifiers(FString AxisPositionBuilderIdentifier, FString PlanePositionBuilderIdentifier, FString AxisAngleBuilderIdentifier);
 	virtual void SetUpdateHoverFunction(TFunction<void(UPrimitiveComponent*, bool)> HoverFunction);
 	virtual void SetUpdateCoordSystemFunction(TFunction<void(UPrimitiveComponent*, EToolContextCoordinateSystem)> CoordSysFunction);
 	
@@ -376,6 +386,10 @@ public:
 protected:
 	TSharedPtr<FTransformGizmoActorFactory> GizmoActorBuilder;
 
+	FString AxisPositionBuilderIdentifier;
+	FString PlanePositionBuilderIdentifier;
+	FString AxisAngleBuilderIdentifier;
+
 	// This function is called on each active GizmoActor Component to update it's hover state.
 	// If the Component is not a UGizmoBaseCmponent, the client needs to provide a different implementation
 	// of this function via the ToolBuilder
diff --git a/Engine/Source/Runtime/Experimental/InteractiveToolsFramework/Public/BaseGizmos/TransformGizmoUtil.h b/Engine/Source/Runtime/Experimental/InteractiveToolsFramework/Public/BaseGizmos/TransformGizmoUtil.h
new file mode 100644
--- /dev/null
+++ b/Engine/Source/Runtime/Experimental/InteractiveToolsFramework/Public/BaseGizmos/TransformGizmoUtil.h
@@ -0,0 +1,116 @@
+// Copyright Epic Games, Inc. All Rights Reserved.
+
+#pragma once
+
+#include "CoreMinimal.h"
+#include "InteractiveGizmo.h"
+#include "BaseGizmos/TransformGizmo.h"
+
+#include "TransformGizmoUtil.generated.h"
+
+class UInteractiveToolManager;
+class UInteractiveGizmoManager;
+class FTransformGizmoActorFactory;
+
+
+namespace UE
+{
+namespace TransformGizmoUtil
+{
+	//
+	// The functions below are helper functions that simplify usage of a UTransformGizmoContextObject 
+	// that is registered as a ContextStoreObject in an InteractiveToolsContext
+	//
+
+	/**
+	 * Spawn a new standard 3-axis Transform gizmo (see UTransformGizmoContextObject::Create3AxisTransformGizmo for details)
+	 * GizmoManager's ToolsContext must have a UTransformGizmoContextObject registered (see UTransformGizmoContextObject for details)
+	 */
+	INTERACTIVETOOLSFRAMEWORK_API UTransformGizmo* Create3AxisTransformGizmo(UInteractiveGizmoManager* GizmoManager, void* Owner = nullptr, const FString& InstanceIdentifier = FString());
+	/**
+	 * Spawn a new standard 3-axis Transform gizmo (see UTransformGizmoContextObject::Create3AxisTransformGizmo for details)
+	 * ToolManager's ToolsContext must have a UTransformGizmoContextObject registered (see UTransformGizmoContextObject for details)
+	 */
+	INTERACTIVETOOLSFRAMEWORK_API UTransformGizmo* Create3AxisTransformGizmo(UInteractiveToolManager* ToolManager, void* Owner = nullptr, const FString& InstanceIdentifier = FString());
+
+
+	/**
+	 * Spawn a new custom Transform gizmo (see UTransformGizmoContextObject::CreateCustomTransformGizmo for details)
+	 * GizmoManager's ToolsContext must have a UTransformGizmoContextObject registered (see UTransformGizmoContextObject for details)
+	 */
+	INTERACTIVETOOLSFRAMEWORK_API UTransformGizmo* CreateCustomTransformGizmo(UInteractiveGizmoManager* GizmoManager, ETransformGizmoSubElements Elements, void* Owner = nullptr, const FString& InstanceIdentifier = FString());
+	/**
+	 * Spawn a new custom Transform gizmo (see UTransformGizmoContextObject::CreateCustomTransformGizmo for details)
+	 * ToolManager's ToolsContext must have a UTransformGizmoContextObject registered (see UTransformGizmoContextObject for details)
+	 */
+	INTERACTIVETOOLSFRAMEWORK_API UTransformGizmo* CreateCustomTransformGizmo(UInteractiveToolManager* ToolManager, ETransformGizmoSubElements Elements, void* Owner = nullptr, const FString& InstanceIdentifier = FString());
+
+
+	/**
+	 * Spawn a new custom Transform gizmo (see UTransformGizmoContextObject::CreateCustomTransformGizmo for details)
+	 * GizmoManager's ToolsContext must have a UTransformGizmoContextObject registered (see UTransformGizmoContextObject for details)
+	 */
+	INTERACTIVETOOLSFRAMEWORK_API UTransformGizmo* CreateCustomRepositionableTransformGizmo(UInteractiveGizmoManager* GizmoManager, ETransformGizmoSubElements Elements, void* Owner = nullptr, const FString& InstanceIdentifier = FString());
+	/**
+	 * Spawn a new custom Transform gizmo (see UTransformGizmoContextObject::CreateCustomRepositionableTransformGizmo for details)
+	 * ToolManager's ToolsContext must have a UTransformGizmoContextObject registered (see UTransformGizmoContextObject for details)
+	 */
+	INTERACTIVETOOLSFRAMEWORK_API UTransformGizmo* CreateCustomRepositionableTransformGizmo(UInteractiveToolManager* ToolManager, ETransformGizmoSubElements Elements, void* Owner = nullptr, const FString& InstanceIdentifier = FString());
+}
+}
+
+
+/**
+ * UTransformGizmoContextObject is a utility object that registers a set of Gizmo Builders
+ * for UTransformGizmo and variants. The intended usage is to call RegisterGizmosWithManager(),
+ * and then the UTransformGizmoContextObject will register itself as a ContextObject in the
+ * InteractiveToolsContext's ContextObjectStore. Then the Create3AxisTransformGizmo()/etc functions
+ * will spawn different variants of UTransformGizmo. The above UE::TransformGizmoUtil:: functions
+ * will look up the UTransformGizmoContextObject instance in the ContextObjectStore and then
+ * call the associated function below.
+ */
+UCLASS(Transient)
+class INTERACTIVETOOLSFRAMEWORK_API UTransformGizmoContextObject : public UObject
+{
+	GENERATED_BODY()
+public:
+
+public:
+	// builder identifiers for default gizmo types. Perhaps should have an API for this...
+	static const FString DefaultAxisPositionBuilderIdentifier;
+	static const FString DefaultPlanePositionBuilderIdentifier;
+	static const FString DefaultAxisAngleBuilderIdentifier;
+	static const FString DefaultThreeAxisTransformBuilderIdentifier;
+	static const FString CustomThreeAxisTransformBuilderIdentifier;
+	static const FString CustomRepositionableThreeAxisTransformBuilderIdentifier;
+
+	void RegisterGizmosWithManager(UInteractiveToolManager* ToolManager);
+	void DeregisterGizmosWithManager(UInteractiveToolManager* ToolManager);
+
+	/**
+	 * Activate a new instance of the default 3-axis transformation Gizmo. RegisterDefaultGizmos() must have been called first.
+	 * @param Owner optional void pointer to whatever "owns" this Gizmo. Allows Gizmo to later be deleted using DestroyAllGizmosByOwner()
+	 * @param InstanceIdentifier optional client-defined *unique* string that can be used to locate this instance
+	 * @return new Gizmo instance that has been created and initialized
+	 */
+	virtual UTransformGizmo* Create3AxisTransformGizmo(UInteractiveGizmoManager* GizmoManager, void* Owner = nullptr, const FString& InstanceIdentifier = FString());
+
+	/**
+	 * Activate a new customized instance of the default 3-axis transformation Gizmo, with only certain elements included. RegisterDefaultGizmos() must have been called first.
+	 * @param Elements flags that indicate which standard gizmo sub-elements should be included
+	 * @param Owner optional void pointer to whatever "owns" this Gizmo. Allows Gizmo to later be deleted using DestroyAllGizmosByOwner()
+	 * @param InstanceIdentifier optional client-defined *unique* string that can be used to locate this instance
+	 * @return new Gizmo instance that has been created and initialized
+	 */
+	virtual UTransformGizmo* CreateCustomTransformGizmo(UInteractiveGizmoManager* GizmoManager, ETransformGizmoSubElements Elements, void* Owner = nullptr, const FString& InstanceIdentifier = FString());
+
+	/**
+	 * Variant of CreateCustomTransformGizmo that creates a URepositionableTransformGizmo, which is an extension to UTransformGizmo that 
+	 * supports various snapping interactions
+	 */
+	virtual UTransformGizmo* CreateCustomRepositionableTransformGizmo(UInteractiveGizmoManager* GizmoManager, ETransformGizmoSubElements Elements, void* Owner = nullptr, const FString& InstanceIdentifier = FString());
+
+protected:
+	TSharedPtr<FTransformGizmoActorFactory> GizmoActorBuilder;
+	bool bDefaultGizmosRegistered = false;
+};
diff --git a/Engine/Source/Runtime/Experimental/InteractiveToolsFramework/Public/InteractiveGizmoManager.h b/Engine/Source/Runtime/Experimental/InteractiveToolsFramework/Public/InteractiveGizmoManager.h
--- a/Engine/Source/Runtime/Experimental/InteractiveToolsFramework/Public/InteractiveGizmoManager.h
+++ b/Engine/Source/Runtime/Experimental/InteractiveToolsFramework/Public/InteractiveGizmoManager.h
@@ -14,6 +14,7 @@
 class FTransformGizmoActorFactory;
 class UTransformGizmo;
 class UTransformGizmoBuilder;
+class UContextObjectStore;
 
 USTRUCT()
 struct FActiveGizmo
@@ -181,6 +182,10 @@ public:
 	/** @return current IToolsContextQueriesAPI */
 	virtual IToolsContextQueriesAPI* GetContextQueriesAPI() { return QueriesAPI; }
 
+	/**
+	 * @return the context object store from the owning tools context.
+	 */
+	UContextObjectStore* GetContextObjectStore() const;
 
 
 
