fix fracture crashes from casting TransformToConvexIndices TSet<int32>s to int32
+ crash from missing proximity data
+ fix autoembed and make it not delete hulls on transform nodes (as these may have been created by the convex tool)

#rnx

diff --git a/Engine/Plugins/Experimental/ChaosEditor/Source/FractureEditor/Private/FractureToolEmbed.cpp b/Engine/Plugins/Experimental/ChaosEditor/Source/FractureEditor/Private/FractureToolEmbed.cpp
--- a/Engine/Plugins/Experimental/ChaosEditor/Source/FractureEditor/Private/FractureToolEmbed.cpp
+++ b/Engine/Plugins/Experimental/ChaosEditor/Source/FractureEditor/Private/FractureToolEmbed.cpp
@@ -230,10 +230,19 @@ void UFractureToolAutoEmbedGeometry::Execute(TWeakPtr<FFractureEditorModeToolkit
 				// Which bone points to the closest convex?
 				const TManagedArray<FTransform>& Transform = ClosestGeometryCollection->Transform;
 				const TManagedArray<int32>& Parent = ClosestGeometryCollection->Parent;
-				TManagedArray<int32>& TransformToConvexIndices =
-					ClosestGeometryCollection->GetAttribute<int32>("TransformToConvexIndices", FTransformCollection::TransformGroup);
+				TManagedArray<TSet<int32>>& TransformToConvexIndices =
+					ClosestGeometryCollection->GetAttribute<TSet<int32>>("TransformToConvexIndices", FTransformCollection::TransformGroup);
 
-				const int32 BoneIndex = TransformToConvexIndices.Find(ClosestConvex);
+				int32 BoneIndex = INDEX_NONE;
+				for (int32 SearchIndex = 0; SearchIndex < TransformToConvexIndices.Num(); SearchIndex++)
+				{
+					const TSet<int32>& ConvexSet = TransformToConvexIndices[SearchIndex];
+					if (ConvexSet.Contains(ClosestConvex))
+					{
+						BoneIndex = SearchIndex;
+						break;
+					}
+				}
 
 				if (BoneIndex > INDEX_NONE)
 				{
diff --git a/Engine/Source/Runtime/Experimental/Chaos/Private/GeometryCollection/GeometryCollectionConvexUtility.cpp b/Engine/Source/Runtime/Experimental/Chaos/Private/GeometryCollection/GeometryCollectionConvexUtility.cpp
--- a/Engine/Source/Runtime/Experimental/Chaos/Private/GeometryCollection/GeometryCollectionConvexUtility.cpp
+++ b/Engine/Source/Runtime/Experimental/Chaos/Private/GeometryCollection/GeometryCollectionConvexUtility.cpp
@@ -19,7 +19,7 @@ FGeometryCollectionConvexUtility::FGeometryCollectionConvexData FGeometryCollect
 	if (!GeometryCollection->HasAttribute("TransformToConvexIndices", FTransformCollection::TransformGroup))
 	{
 		FManagedArrayCollection::FConstructionParameters ConvexDependency("Convex");
-		TManagedArray<TSet<int32>>& IndexArray = GeometryCollection->AddAttribute<TSet<int32>>("TransformToConvexIndices", FTransformCollection::TransformGroup, ConvexDependency);
+		GeometryCollection->AddAttribute<TSet<int32>>("TransformToConvexIndices", FTransformCollection::TransformGroup, ConvexDependency);
 	}
 
 	if (!GeometryCollection->HasAttribute("ConvexHull", "Convex"))
@@ -27,16 +27,14 @@ FGeometryCollectionConvexUtility::FGeometryCollectionConvexData FGeometryCollect
 		GeometryCollection->AddAttribute<TUniquePtr<Chaos::FConvex>>("ConvexHull", "Convex");
 	}
 
-	// Check for correct population. All (and only) rigid nodes should have a convex associated.
+	// Check for correct population. Make sure all rigid nodes should have a convex associated; leave convex hulls for transform nodes alone for now
 	const TManagedArray<int32>& SimulationType = GeometryCollection->GetAttribute<int32>("SimulationType", FTransformCollection::TransformGroup);
 	const TManagedArray<int32>& TransformToGeometryIndex = GeometryCollection->GetAttribute<int32>("TransformToGeometryIndex", FTransformCollection::TransformGroup);
 	TManagedArray<TSet<int32>>& TransformToConvexIndices = GeometryCollection->GetAttribute<TSet<int32>>("TransformToConvexIndices", FTransformCollection::TransformGroup);
 	TManagedArray<TUniquePtr<Chaos::FConvex>>& ConvexHull = GeometryCollection->GetAttribute<TUniquePtr<Chaos::FConvex>>("ConvexHull", "Convex");
-
+	
 	TArray<int32> ProduceConvexHulls;
 	ProduceConvexHulls.Reserve(SimulationType.Num());
-	TArray<int32> EliminateConvexHulls;
-	EliminateConvexHulls.Reserve(SimulationType.Num());
 
 	for (int32 Idx = 0; Idx < SimulationType.Num(); ++Idx)
 	{
@@ -44,19 +42,6 @@ FGeometryCollectionConvexUtility::FGeometryCollectionConvexData FGeometryCollect
 		{
 			ProduceConvexHulls.Add(Idx);
 		}
-		else if ((SimulationType[Idx] != FGeometryCollection::ESimulationTypes::FST_Rigid) && (TransformToConvexIndices[Idx].Num() > 0))
-		{
-			EliminateConvexHulls.Add(Idx);
-		}
-	}
-
-	if (EliminateConvexHulls.Num())
-	{
-		GeometryCollection->RemoveElements("Convex", EliminateConvexHulls);
-		for (int32 Idx : EliminateConvexHulls)
-		{
-			TransformToConvexIndices[Idx].Reset();
-		}
 	}
 
 	if (ProduceConvexHulls.Num())
@@ -191,7 +176,7 @@ void CreateNonoverlappingConvexHulls(
 	const TManagedArray<int32>& SimulationType,
 	int32 LeafType,
 	const TManagedArray<int32>& Parents,
-	const TManagedArray<TSet<int32>>& GeoProximity,
+	const TManagedArray<TSet<int32>>* GeoProximity,
 	const TManagedArray<int32>& GeometryToTransformIndex,
 	double FracAllowRemove
 )
@@ -211,14 +196,21 @@ void CreateNonoverlappingConvexHulls(
 
 	TArray<TSet<int32>> LeafProximity;
 	LeafProximity.SetNum(NumBones);
-	for (int32 GeomIdx = 0; GeomIdx < GeoProximity.Num(); GeomIdx++)
+	if (GeoProximity)
 	{
-		int32 TransformIdx = GeometryToTransformIndex[GeomIdx];
-		for (int32 NbrGeomIdx : GeoProximity[GeomIdx])
+		for (int32 GeomIdx = 0; GeomIdx < GeoProximity->Num(); GeomIdx++)
 		{
-			LeafProximity[TransformIdx].Add(GeometryToTransformIndex[NbrGeomIdx]);
+			int32 TransformIdx = GeometryToTransformIndex[GeomIdx];
+			for (int32 NbrGeomIdx : (*GeoProximity)[GeomIdx])
+			{
+				LeafProximity[TransformIdx].Add(GeometryToTransformIndex[NbrGeomIdx]);
+			}
 		}
 	}
+	else
+	{
+		// TODO: fill out LeafProximity by doing IsColliding calls (may have to construct Children array before this)
+	}
 
 	auto IsColliding2 = [&Convexes](int32 ConvexA, int32 ConvexB)
 	{
@@ -659,11 +651,8 @@ FGeometryCollectionConvexUtility::FGeometryCollectionConvexData FGeometryCollect
 {
 	check(GeometryCollection);
 
-	// TODO: add fallback in case we don't have proximity data
 	bool bHasProximity = GeometryCollection->HasAttribute("Proximity", FGeometryCollection::GeometryGroup);
-	ensure(bHasProximity);
-	
-	const TManagedArray<TSet<int32>>& GCProximity = GeometryCollection->GetAttribute<TSet<int32>>("Proximity", FGeometryCollection::GeometryGroup);
+	const TManagedArray<TSet<int32>>* GCProximity = GeometryCollection->FindAttribute<TSet<int32>>("Proximity", FGeometryCollection::GeometryGroup);
 	TArray<TUniquePtr<Chaos::FConvex>> Convexes;
 	TArray<TSet<int32>> TransformToConvexIndexArr;
 	TArray<FTransform> GlobalTransformArray;
@@ -682,7 +671,7 @@ FGeometryCollectionConvexUtility::FGeometryCollectionConvexData FGeometryCollect
 	if (!GeometryCollection->HasAttribute("TransformToConvexIndices", FTransformCollection::TransformGroup))
 	{
 		FManagedArrayCollection::FConstructionParameters ConvexDependency("Convex");
-		TManagedArray<TSet<int32>>& IndexArray = GeometryCollection->AddAttribute<TSet<int32>>("TransformToConvexIndices", FTransformCollection::TransformGroup, ConvexDependency);
+		GeometryCollection->AddAttribute<TSet<int32>>("TransformToConvexIndices", FTransformCollection::TransformGroup, ConvexDependency);
 	}
 
 	if (!GeometryCollection->HasAttribute("ConvexHull", "Convex"))
@@ -693,6 +682,7 @@ FGeometryCollectionConvexUtility::FGeometryCollectionConvexData FGeometryCollect
 	TManagedArray<TSet<int32>>& TransformToConvexIndices = GeometryCollection->GetAttribute<TSet<int32>>("TransformToConvexIndices", FTransformCollection::TransformGroup);
 	TManagedArray<TUniquePtr<Chaos::FConvex>>& ConvexHull = GeometryCollection->GetAttribute<TUniquePtr<Chaos::FConvex>>("ConvexHull", "Convex");
 	TransformToConvexIndices = MoveTemp(TransformToConvexIndexArr);
+	GeometryCollection->Resize(Convexes.Num(), "Convex");
 	ConvexHull = MoveTemp(Convexes);
 
 	return { TransformToConvexIndices, ConvexHull };
@@ -722,14 +712,17 @@ void FGeometryCollectionConvexUtility::RemoveConvexHulls(FGeometryCollection* Ge
 {
 	if (GeometryCollection->HasGroup("Convex") && GeometryCollection->HasAttribute("TransformToConvexIndices", FTransformCollection::TransformGroup))
 	{
-		TManagedArray<int32>& TransformToConvexIndices = GeometryCollection->GetAttribute<int32>("TransformToConvexIndices", FTransformCollection::TransformGroup);
+		TManagedArray<TSet<int32>>& TransformToConvexIndices = GeometryCollection->GetAttribute<TSet<int32>>("TransformToConvexIndices", FTransformCollection::TransformGroup);
 		TArray<int32> ConvexIndices;
 		for (int32 TransformIdx : SortedTransformDeletes)
 		{
-			if (TransformToConvexIndices[TransformIdx] > INDEX_NONE)
+			if (TransformToConvexIndices[TransformIdx].Num() > 0)
 			{
-				ConvexIndices.Add(TransformToConvexIndices[TransformIdx]);
-				TransformToConvexIndices[TransformIdx] = INDEX_NONE;
+				for (int32 ConvexIdx : TransformToConvexIndices[TransformIdx])
+				{
+					ConvexIndices.Add(ConvexIdx);
+				}
+				TransformToConvexIndices[TransformIdx].Empty();
 			}
 		}
 
@@ -744,16 +737,4 @@ void FGeometryCollectionConvexUtility::RemoveConvexHulls(FGeometryCollection* Ge
 
 void FGeometryCollectionConvexUtility::SetDefaults(FGeometryCollection* GeometryCollection, FName Group, uint32 StartSize, uint32 NumElements)
 {
-	if (Group == FTransformCollection::TransformGroup)
-	{
-		if (GeometryCollection->HasAttribute("TransformToConvexIndices", FTransformCollection::TransformGroup))
-		{
-			TManagedArray<int32>& TransformToConvexIndices = GeometryCollection->GetAttribute<int32>("TransformToConvexIndices", FTransformCollection::TransformGroup);
-
-			for (uint32 Idx = StartSize; Idx < StartSize + NumElements; ++Idx)
-			{
-				TransformToConvexIndices[Idx] = INDEX_NONE;
-			}
-		}
-	}
 }
\ No newline at end of file
