more direct / faster implementation of HasAttribute and RemoveAttribute in ManagedArrayCollection (directly use key in map, rather than linear searching for key)

#rnx

diff --git a/Engine/Source/Runtime/Experimental/Chaos/Private/GeometryCollection/ManagedArrayCollection.cpp b/Engine/Source/Runtime/Experimental/Chaos/Private/GeometryCollection/ManagedArrayCollection.cpp
--- a/Engine/Source/Runtime/Experimental/Chaos/Private/GeometryCollection/ManagedArrayCollection.cpp
+++ b/Engine/Source/Runtime/Experimental/Chaos/Private/GeometryCollection/ManagedArrayCollection.cpp
@@ -72,16 +72,8 @@ TArray<FName> FManagedArrayCollection::GroupNames() const
 
 bool FManagedArrayCollection::HasAttribute(FName Name, FName Group) const
 {
-	bool bReturnValue = false;
-	for (const TTuple<FKeyType, FValueType>& Entry : Map)
-	{
-		if (Entry.Key.Get<0>() == Name && Entry.Key.Get<1>() == Group)
-		{
-			bReturnValue = true;
-			break;
-		}
-	}
-	return bReturnValue;
+	FKeyType Key = FManagedArrayCollection::MakeMapKey(Name, Group);
+	return Map.Contains(Key);
 }
 
 TArray<FName> FManagedArrayCollection::AttributeNames(FName Group) const
@@ -132,14 +124,8 @@ int32 FManagedArrayCollection::AddElements(int32 NumberElements, FName Group)
 
 void FManagedArrayCollection::RemoveAttribute(FName Name, FName Group)
 {
-	for (const TTuple<FKeyType, FValueType>& Entry : Map)
-	{
-		if (Entry.Key.Get<0>() == Name && Entry.Key.Get<1>() == Group)
-		{
-			Map.Remove(Entry.Key);
-			return;
-		}
-	}
+	FKeyType Key = FManagedArrayCollection::MakeMapKey(Name, Group);
+	Map.Remove(Key);
 }
 
 void FManagedArrayCollection::RemoveGroup(FName Group)
