Move where PositionTolerance, AngularTolerance, MinParentMassRatio and MaxInertiaRatio are set

diff --git a/Engine/Source/Runtime/Experimental/Chaos/Private/PBDRigidsSolver.cpp b/Engine/Source/Runtime/Experimental/Chaos/Private/PBDRigidsSolver.cpp
--- a/Engine/Source/Runtime/Experimental/Chaos/Private/PBDRigidsSolver.cpp
+++ b/Engine/Source/Runtime/Experimental/Chaos/Private/PBDRigidsSolver.cpp
@@ -733,6 +733,10 @@ namespace Chaos
 		JointsSettings.MinSolverStiffness = ChaosSolverJointMinSolverStiffness;
 		JointsSettings.MaxSolverStiffness = ChaosSolverJointMaxSolverStiffness;
 		JointsSettings.NumIterationsAtMaxSolverStiffness = ChaosSolverJointNumIterationsAtMaxSolverStiffness;
+		JointsSettings.PositionTolerance = ChaosSolverJointPositionTolerance;
+		JointsSettings.AngleTolerance = ChaosSolverJointAngleTolerance;
+		JointsSettings.MinParentMassRatio = ChaosSolverJointMinParentMassRatio;
+		JointsSettings.MaxInertiaRatio = ChaosSolverJointMaxInertiaRatio;
 		JointConstraints.SetSettings(JointsSettings);
 
 		// Apply CVAR overrides if set
@@ -765,12 +769,6 @@ namespace Chaos
 			{
 				SetCollisionCullDistance(ChaosSolverCullDistance);
 			}
-			// Joint settings overrides
-			JointsSettings.PositionTolerance = ChaosSolverJointPositionTolerance;
-			JointsSettings.AngleTolerance = ChaosSolverJointAngleTolerance;
-			JointsSettings.MinParentMassRatio = ChaosSolverJointMinParentMassRatio;
-			JointsSettings.MaxInertiaRatio = ChaosSolverJointMaxInertiaRatio;
-			JointConstraints.SetSettings(JointsSettings);
 		}
 
 		UE_LOG(LogPBDRigidsSolver, Verbose, TEXT("PBDRigidsSolver::Tick(%3.5f)"), DeltaTime);
