Chaos : Geometry Collection
Collision Attribute Restructuring.
-	Updated the size specific data to be the primary container for collision settings within the geometry collection.
-	Removed non size specific default parameters from the asset.
-	Added support for Capsule collisions.
#okforgithup, public, internal

diff --git a/Engine/Plugins/Experimental/ChaosEditor/Source/FractureEditor/Private/FractureToolGenerators.cpp b/Engine/Plugins/Experimental/ChaosEditor/Source/FractureEditor/Private/FractureToolGenerators.cpp
--- a/Engine/Plugins/Experimental/ChaosEditor/Source/FractureEditor/Private/FractureToolGenerators.cpp
+++ b/Engine/Plugins/Experimental/ChaosEditor/Source/FractureEditor/Private/FractureToolGenerators.cpp
@@ -276,6 +276,7 @@ class AGeometryCollectionActor* UFractureToolGenerateAsset::CreateNewGeometryAct
 
 	UPackage* Package = CreatePackage(*UniquePackageName);
 	UGeometryCollection* InGeometryCollection = static_cast<UGeometryCollection*>(NewObject<UGeometryCollection>(Package, UGeometryCollection::StaticClass(), FName(*UniqueAssetName), RF_Transactional | RF_Public | RF_Standalone));
+	if(!InGeometryCollection->SizeSpecificData.Num()) InGeometryCollection->SizeSpecificData.Add(FGeometryCollectionSizeSpecificData());
 
 	// Create the new Geometry Collection actor
 	AGeometryCollectionActor* NewActor = Cast<AGeometryCollectionActor>(AddActor(GetSelectedLevel(), AGeometryCollectionActor::StaticClass()));
diff --git a/Engine/Plugins/Experimental/GeometryCollectionPlugin/Source/GeometryCollectionEditor/Private/GeometryCollection/GeometryCollectionFactory.cpp b/Engine/Plugins/Experimental/GeometryCollectionPlugin/Source/GeometryCollectionEditor/Private/GeometryCollection/GeometryCollectionFactory.cpp
--- a/Engine/Plugins/Experimental/GeometryCollectionPlugin/Source/GeometryCollectionEditor/Private/GeometryCollection/GeometryCollectionFactory.cpp
+++ b/Engine/Plugins/Experimental/GeometryCollectionPlugin/Source/GeometryCollectionEditor/Private/GeometryCollection/GeometryCollectionFactory.cpp
@@ -30,7 +30,9 @@ UGeometryCollectionFactory::UGeometryCollectionFactory(const FObjectInitializer&
 
 UGeometryCollection* UGeometryCollectionFactory::StaticFactoryCreateNew(UClass* Class, UObject* InParent, FName Name, EObjectFlags Flags, UObject* Context, FFeedbackContext* Warn)
 {
-	return static_cast<UGeometryCollection*>(NewObject<UGeometryCollection>(InParent, Class, Name, Flags | RF_Transactional | RF_Public | RF_Standalone));
+	UGeometryCollection* NewGeometryCollection = static_cast<UGeometryCollection*>(NewObject<UGeometryCollection>(InParent, Class, Name, Flags | RF_Transactional | RF_Public | RF_Standalone));
+	if (!NewGeometryCollection->SizeSpecificData.Num()) NewGeometryCollection->SizeSpecificData.Add(FGeometryCollectionSizeSpecificData());
+	return NewGeometryCollection;
 }
 
 void AppendActorComponentsRecursive(
@@ -150,14 +152,6 @@ UObject* UGeometryCollectionFactory::FactoryCreateNew(UClass* Class, UObject* In
 
 	UGeometryCollection* NewGeometryCollection = StaticFactoryCreateNew(Class, InParent, Name, Flags, Context, Warn);
 
-	// Set default collision type here in the factory instead of the geometry collection constructor to only
-	// affect new geometry collections and not update the default for previously created objects.
-	if(NewGeometryCollection)
-	{
-		NewGeometryCollection->CollisionType = ECollisionTypeEnum::Chaos_Surface_Volumetric;
-		NewGeometryCollection->ImplicitType = EImplicitTypeEnum::Chaos_Implicit_LevelSet;
-	}
-
 	for (GeometryCollectionStaticMeshConversionTuple & StaticMeshData : StaticMeshList)
 	{
 		FGeometryCollectionConversion::AppendStaticMesh(StaticMeshData.Get<0>(), StaticMeshData.Get<1>(), StaticMeshData.Get<2>(), NewGeometryCollection, false);
diff --git a/Engine/Source/Programs/HeadlessChaos/Private/GeometryCollection/GeometryCollectionTestFramework.cpp b/Engine/Source/Programs/HeadlessChaos/Private/GeometryCollection/GeometryCollectionTestFramework.cpp
--- a/Engine/Source/Programs/HeadlessChaos/Private/GeometryCollection/GeometryCollectionTestFramework.cpp
+++ b/Engine/Source/Programs/HeadlessChaos/Private/GeometryCollection/GeometryCollectionTestFramework.cpp
@@ -156,8 +156,8 @@ namespace GeometryCollectionTest
 			SimulationParams.PhysicalMaterialHandle = NewHandle;
 			SimulationParams.Shared.Mass = Params.Mass;
 			SimulationParams.Shared.bMassAsDensity = Params.bMassAsDensity;
-			SimulationParams.Shared.SizeSpecificData[0].CollisionType = Params.CollisionType;
-			SimulationParams.Shared.SizeSpecificData[0].ImplicitType = Params.ImplicitType;
+			SimulationParams.Shared.SizeSpecificData[0].CollisionShapesData[0].CollisionType = Params.CollisionType;
+			SimulationParams.Shared.SizeSpecificData[0].CollisionShapesData[0].ImplicitType = Params.ImplicitType;
 			SimulationParams.Simulating = Params.Simulating;
 			SimulationParams.EnableClustering = Params.EnableClustering;
 			SimulationParams.InitialLinearVelocity = Params.InitialLinearVelocity;
@@ -170,8 +170,8 @@ namespace GeometryCollectionTest
 			SimulationParams.ClusterGroupIndex = Params.ClusterGroupIndex;
 
 			FSharedSimulationSizeSpecificData Tmp;
-			Tmp.MinLevelSetResolution = Params.MinLevelSetResolution;
-			Tmp.MaxLevelSetResolution = Params.MaxLevelSetResolution;
+			Tmp.CollisionShapesData[0].LevelSetData.MinLevelSetResolution = Params.MinLevelSetResolution;
+			Tmp.CollisionShapesData[0].LevelSetData.MaxLevelSetResolution = Params.MaxLevelSetResolution;
 			SimulationParams.Shared.SizeSpecificData.Add(Tmp);
 
 			Chaos::FErrorReporter ErrorReporter;
diff --git a/Engine/Source/Runtime/Core/Public/UObject/UE5MainStreamObjectVersion.h b/Engine/Source/Runtime/Core/Public/UObject/UE5MainStreamObjectVersion.h
--- a/Engine/Source/Runtime/Core/Public/UObject/UE5MainStreamObjectVersion.h
+++ b/Engine/Source/Runtime/Core/Public/UObject/UE5MainStreamObjectVersion.h
@@ -145,6 +145,9 @@ struct CORE_API FUE5MainStreamObjectVersion
 		// Expanded separate translucency into multiple passes
 		MaterialTranslucencyPass,
 
+		// Chaos FGeometryCollectionObject user defined collision shapes support
+		GeometryCollectionUserDefinedCollisionShapes,
+
 		// -----<new versions can be added above this line>-------------------------------------------------
 		VersionPlusOne,
 		LatestVersion = VersionPlusOne - 1
diff --git a/Engine/Source/Runtime/Experimental/Chaos/Private/GeometryCollection/GeometryCollectionCollisionStructureManager.cpp b/Engine/Source/Runtime/Experimental/Chaos/Private/GeometryCollection/GeometryCollectionCollisionStructureManager.cpp
--- a/Engine/Source/Runtime/Experimental/Chaos/Private/GeometryCollection/GeometryCollectionCollisionStructureManager.cpp
+++ b/Engine/Source/Runtime/Experimental/Chaos/Private/GeometryCollection/GeometryCollectionCollisionStructureManager.cpp
@@ -7,6 +7,7 @@
 #include "Chaos/Levelset.h"
 #include "Chaos/Particles.h"
 #include "Chaos/Sphere.h"
+#include "Chaos/Capsule.h"
 #include "Chaos/Vector.h"
 #include "HAL/IConsoleManager.h"
 #include "Chaos/TriangleMesh.h"
@@ -244,6 +245,19 @@ FCollisionStructureManager::NewImplicitSphere(
 	return Implicit;
 }
 
+FCollisionStructureManager::FImplicit*
+FCollisionStructureManager::NewImplicitCapsule(
+	const float Radius,
+	const float Length,
+	const float CollisionObjectReduction,
+	const ECollisionTypeEnum CollisionType)
+{
+	const Chaos::FReal HalfLength = (Chaos::FReal)Length * (Chaos::FReal)0.5;
+	Chaos::FImplicitObject* Implicit = new Chaos::FCapsule(Chaos::FVec3(0, 0, -HalfLength), Chaos::FVec3(0, 0, +HalfLength), Radius * (1 - CollisionObjectReduction / 100.f));
+	UpdateImplicitFlags(Implicit, CollisionType);
+	return Implicit;
+}
+
 FCollisionStructureManager::FImplicit*
 FCollisionStructureManager::NewImplicitLevelset(
 	Chaos::FErrorReporter ErrorReporter,
diff --git a/Engine/Source/Runtime/Experimental/Chaos/Private/PhysicsProxy/GeometryCollectionPhysicsProxy.cpp b/Engine/Source/Runtime/Experimental/Chaos/Private/PhysicsProxy/GeometryCollectionPhysicsProxy.cpp
--- a/Engine/Source/Runtime/Experimental/Chaos/Private/PhysicsProxy/GeometryCollectionPhysicsProxy.cpp
+++ b/Engine/Source/Runtime/Experimental/Chaos/Private/PhysicsProxy/GeometryCollectionPhysicsProxy.cpp
@@ -203,10 +203,17 @@ void PopulateSimulatedParticle(
 
 	Handle->SetCollisionGroup(CollisionGroup);
 
+	// @todo(GCCollisionShapes) : add support for multiple shapes, currently just one. 
+	FCollectionCollisionTypeData SingleSuportedCollisionTypeData = FCollectionCollisionTypeData();
+	if (SharedParams.SizeSpecificData.Num() && SharedParams.SizeSpecificData[0].CollisionShapesData.Num())
+	{
+		SingleSuportedCollisionTypeData = SharedParams.SizeSpecificData[0].CollisionShapesData[0];
+	}
+
 	if (Implicit)	//todo(ocohen): this is only needed for cases where clusters have no proxy. Kind of gross though, should refactor
 	{
 		TSharedPtr<Chaos::FImplicitObject, ESPMode::ThreadSafe> SharedImplicitTS(Implicit->DeepCopy().Release());
-		FCollisionStructureManager::UpdateImplicitFlags(SharedImplicitTS.Get(), SharedParams.SizeSpecificData[0].CollisionType); // @todo(chaos) : Implicit constructor clobbers CollisionType
+		FCollisionStructureManager::UpdateImplicitFlags(SharedImplicitTS.Get(), SingleSuportedCollisionTypeData.CollisionType);
 		Handle->SetSharedGeometry(SharedImplicitTS);
 		Handle->SetHasBounds(true);
 		Handle->SetLocalBounds(SharedImplicitTS->BoundingBox());
@@ -216,7 +223,7 @@ void PopulateSimulatedParticle(
 		Handle->SetWorldSpaceInflatedBounds(TransformedBBox);
 	}
 
-	if (Simplicial && SharedParams.SizeSpecificData[0].CollisionType==ECollisionTypeEnum::Chaos_Surface_Volumetric)
+	if (Simplicial && SingleSuportedCollisionTypeData.CollisionType == ECollisionTypeEnum::Chaos_Surface_Volumetric)
 	{
 		Handle->CollisionParticlesInitIfNeeded();
 
@@ -348,8 +355,7 @@ FAutoConsoleVariableRef CVarReportHighParticleFraction(TEXT("p.gc.ReportHighPart
 void FGeometryCollectionPhysicsProxy::Initialize(Chaos::FPBDRigidsEvolutionBase *Evolution)
 {
 	check(IsInGameThread());
-
-	//
+										 //
 	// Game thread initilization. 
 	//
 	//  1) Create a input buffer to store all game thread side data. 
@@ -360,6 +366,12 @@ void FGeometryCollectionPhysicsProxy::Initialize(Chaos::FPBDRigidsEvolutionBase
 
 	InitializeDynamicCollection(DynamicCollection, *Parameters.RestCollection, Parameters);
 
+	// Attach the external particles to the gamethread collection
+	if (DynamicCollection.HasAttribute(FGeometryCollection::ParticlesAttribute, FTransformCollection::TransformGroup))
+		DynamicCollection.RemoveAttribute(FGeometryCollection::ParticlesAttribute, FTransformCollection::TransformGroup);
+	DynamicCollection.AddExternalAttribute<TUniquePtr<Chaos::FGeometryParticle>>(FGeometryCollection::ParticlesAttribute, FTransformCollection::TransformGroup, GTParticles);
+
+
 	NumParticles = DynamicCollection.NumElements(FGeometryCollection::TransformGroup);
 	BaseParticleIndex = 0; // Are we always zero indexed now?
 	SolverClusterID.Init(nullptr, NumParticles);
@@ -384,87 +396,21 @@ void FGeometryCollectionPhysicsProxy::Initialize(Chaos::FPBDRigidsEvolutionBase
 	{
 		for (const FSharedSimulationSizeSpecificData& Data : Parameters.Shared.SizeSpecificData)
 		{
-			if (Data.CollisionParticlesFraction >= ReportHighParticleFraction)
+			if (ensure(Data.CollisionShapesData.Num()))
 			{
-				ensureMsgf(false, TEXT("Collection with small particle fraction"));
-				UE_LOG(LogChaos, Warning, TEXT("Collection with small particle fraction(%f):%s"), Data.CollisionParticlesFraction, *Parameters.Name);
-			}
-		}
-	}
-
-
-#if TODO_REIMPLEMENT_RIGID_CACHING
-
-	// todo : Replace with normal funcs in this class
-	Callbacks->SetUpdateRecordedStateFunction([this](float SolverTime, const TManagedArray<int32>& RigidBodyID, const TManagedArray<FGeometryCollectionBoneNode>& Hierarchy, const FSolverCallbacks::FParticlesType& Particles, const FSolverCallbacks::FCollisionConstraintsType& CollisionRule)
-		{
-			UpdateRecordedState(SolverTime, RigidBodyID, Hierarchy, Particles, CollisionRule);
-		});
-
-	Callbacks->SetCommitRecordedStateFunction([this](FRecordedTransformTrack& InTrack)
-		{
-			InTrack = FRecordedTransformTrack::ProcessRawRecordedData(RecordedTracks);
-		});
-
-	// Old proxy init
-	RecordedTracks.Records.Reset();
-
-	if (Parameters.bClearCache)
-	{
-		if (ResetAnimationCacheCallback)
-		{
-			ResetAnimationCacheCallback();
-		}
-	}
-
-	ProxySimDuration = T(0);
-	LastSyncCountGT = 0;
-
-
-#endif
-
-	/*
-	// Initialize data for faster bound calculations
-	// precompute data used for bounds calculation
-	{
-		const TManagedArray<FBox>& BoundingBoxes = Parameters.RestCollection->BoundingBox;
-		const TManagedArray<int32>& TransformIndices = Parameters.RestCollection->TransformIndex;
-
-		const int32 NumBoxes = BoundingBoxes.Num();
-
-		ValidGeometryBoundingBoxes.Reset();
-		ValidGeometryTransformIndices.Reset();
-		for (int32 BoxIdx = 0; BoxIdx < NumBoxes; ++BoxIdx)
-		{
-			const int32 CurrTransformIndex = TransformIndices[BoxIdx];
-
-			if (Parameters.RestCollection->IsGeometry(CurrTransformIndex))
-			{
-				ValidGeometryBoundingBoxes.Add(BoundingBoxes[BoxIdx]);
-				ValidGeometryTransformIndices.Add(CurrTransformIndex);
+				if (Data.CollisionShapesData[0].CollisionParticleData.CollisionParticlesFraction >= ReportHighParticleFraction)
+				{
+					ensureMsgf(false, TEXT("Collection with small particle fraction"));
+					UE_LOG(LogChaos, Warning, TEXT("Collection with small particle fraction(%f):%s"), Data.CollisionShapesData[0].CollisionParticleData.CollisionParticlesFraction, *Parameters.Name);
+				}
 			}
 		}
-
-		FBox BoundingBox(ForceInit);
-		const FMatrix& ActorToWorld = Parameters.WorldTransform.ToMatrixWithScale();
-		for (int i = 0; i < ValidGeometryBoundingBoxes.Num(); ++i)
-		{
-			BoundingBox += ValidGeometryBoundingBoxes[i].TransformBy(
-				TmpGlobalTransforms[ValidGeometryTransformIndices[i]] * ActorToWorld);
-		}
-
-		//Results.Get(0).WorldBounds = FBoxSphereBounds(BoundingBox);
-		//Results.Get(1).WorldBounds = FBoxSphereBounds(BoundingBox);
 	}
-	*/
 
 	// Initialise GT/External particles
-	const int32 NumTransforms = GameThreadCollection.Transform.Num();
-
-	// Attach the external particles to the gamethread collection
-	GameThreadCollection.AddExternalAttribute<TUniquePtr<Chaos::FGeometryParticle>>(FGeometryCollection::ParticlesAttribute, FTransformCollection::TransformGroup, GTParticles);
+	const int32 NumTransforms = DynamicCollection.Transform.Num();
 
-	if(ensure(NumTransforms == GameThreadCollection.Implicits.Num() && NumTransforms == GTParticles.Num())) // Implicits are in the transform group so this invariant should always hold
+	if(ensure(NumTransforms == DynamicCollection.Implicits.Num() && NumTransforms == GTParticles.Num())) // Implicits are in the transform group so this invariant should always hold
 	{
 		for(int32 Index = 0; Index < NumTransforms; ++Index)
 		{
@@ -527,6 +473,9 @@ void FGeometryCollectionPhysicsProxy::Initialize(Chaos::FPBDRigidsEvolutionBase
 
 void FGeometryCollectionPhysicsProxy::InitializeDynamicCollection(FGeometryDynamicCollection& DynamicCollection, const FGeometryCollection& RestCollection, const FSimulationParameters& Params)
 {
+	// @todo(GCCollisionShapes) : add support for multiple shapes, currently just one. 
+	ensure(Params.Shared.SizeSpecificData[0].CollisionShapesData.Num());
+
 	// 
 	// This function will use the rest collection to populate the dynamic collection. 
 	//
@@ -556,7 +505,7 @@ void FGeometryCollectionPhysicsProxy::InitializeDynamicCollection(FGeometryDynam
 		static const bool bAnalyticsDisabled = (AnalyticDisableCVar && AnalyticDisableCVar->GetBool());
 
 		if (RestCollection.HasAttribute(DynamicCollection.SimplicialsAttribute, FTransformCollection::TransformGroup)
-			&& (Params.Shared.SizeSpecificData[0].CollisionType == ECollisionTypeEnum::Chaos_Surface_Volumetric || bAnalyticsDisabled))
+			&& (Params.Shared.SizeSpecificData[0].CollisionShapesData[0].CollisionType == ECollisionTypeEnum::Chaos_Surface_Volumetric || bAnalyticsDisabled))
 		{
 			const auto& RestSimplicials = RestCollection.GetAttribute<TUniquePtr<FSimplicial>>(
 				DynamicCollection.SimplicialsAttribute, FTransformCollection::TransformGroup);
@@ -734,7 +683,7 @@ void FGeometryCollectionPhysicsProxy::InitializeBodiesPT(Chaos::FPBDRigidsSolver
 					Shape->SetMaterial(Parameters.PhysicalMaterialHandle);
 				}
 			}
-		});
+		},true);
 
 		// After population, the states of each particle could have changed
 		Particles.UpdateGeometryCollectionViews();
@@ -1837,6 +1786,8 @@ void FGeometryCollectionPhysicsProxy::InitializeSharedCollisionStructures(
 	FGeometryCollection& RestCollection,
 	const FSharedSimulationParameters& SharedParams)
 {
+	check(SharedParams.SizeSpecificData.Num() && SharedParams.SizeSpecificData[0].CollisionShapesData.Num());
+
 	FString BaseErrorPrefix = ErrorReporter.GetPrefix();
 
 	// fracture tools can create an empty GC before appending new geometry
@@ -2170,9 +2121,9 @@ void FGeometryCollectionPhysicsProxy::InitializeSharedCollisionStructures(
 					FCollisionStructureManager::NewSimplicial(
 						MassSpaceParticles,
 						BoneMap,
-						SizeSpecificData.CollisionType,
+						SizeSpecificData.CollisionShapesData[0].CollisionType,
 						*TriMesh,
-						SizeSpecificData.CollisionParticlesFraction);
+						SizeSpecificData.CollisionShapesData[0].CollisionParticleData.CollisionParticlesFraction);
 				CollectionSimplicials[TransformGroupIndex] = TUniquePtr<FSimplicial>(Simplicial); // CollectionSimplicials is in the TransformGroup
 				//ensureMsgf(CollectionSimplicials[TransformGroupIndex], TEXT("No simplicial representation."));
 				if (!CollectionSimplicials[TransformGroupIndex]->Size())
@@ -2180,7 +2131,7 @@ void FGeometryCollectionPhysicsProxy::InitializeSharedCollisionStructures(
 					ensureMsgf(false, TEXT("Simplicial is empty."));
 				}
 
-				if (SizeSpecificData.ImplicitType == EImplicitTypeEnum::Chaos_Implicit_LevelSet)
+				if (SizeSpecificData.CollisionShapesData[0].ImplicitType == EImplicitTypeEnum::Chaos_Implicit_LevelSet)
 				{
 					LocalErrorReporter.SetPrefix(BaseErrorPrefix + " | Transform Index: " + FString::FromInt(TransformGroupIndex) + " of " + FString::FromInt(TransformIndex.Num()));
 					CollectionImplicits[TransformGroupIndex] = FGeometryDynamicCollection::FSharedImplicit(
@@ -2189,37 +2140,62 @@ void FGeometryCollectionPhysicsProxy::InitializeSharedCollisionStructures(
 							MassSpaceParticles,
 							*TriMesh,
 							InstanceBoundingBox,
-							SizeSpecificData.MinLevelSetResolution,
-							SizeSpecificData.MaxLevelSetResolution,
-							SizeSpecificData.CollisionObjectReductionPercentage,
-							SizeSpecificData.CollisionType));
+							SizeSpecificData.CollisionShapesData[0].LevelSetData.MinLevelSetResolution,
+							SizeSpecificData.CollisionShapesData[0].LevelSetData.MaxLevelSetResolution,
+							SizeSpecificData.CollisionShapesData[0].CollisionObjectReductionPercentage,
+							SizeSpecificData.CollisionShapesData[0].CollisionType));
 					// Fall back on sphere if level set rasterization failed.
 					if (!CollectionImplicits[TransformGroupIndex])
 					{
 						CollectionImplicits[TransformGroupIndex] = FGeometryDynamicCollection::FSharedImplicit(
 							FCollisionStructureManager::NewImplicitSphere(
 								InnerRadius[GeometryIndex],
-								SizeSpecificData.CollisionObjectReductionPercentage,
-								SizeSpecificData.CollisionType));
+								SizeSpecificData.CollisionShapesData[0].CollisionObjectReductionPercentage,
+								SizeSpecificData.CollisionShapesData[0].CollisionType));
 					}
 				}
-				else if (SizeSpecificData.ImplicitType == EImplicitTypeEnum::Chaos_Implicit_Box)
+				else if (SizeSpecificData.CollisionShapesData[0].ImplicitType == EImplicitTypeEnum::Chaos_Implicit_Box)
 				{
 					CollectionImplicits[TransformGroupIndex] = FGeometryDynamicCollection::FSharedImplicit(
 						FCollisionStructureManager::NewImplicitBox(
 							InstanceBoundingBox,
-							SizeSpecificData.CollisionObjectReductionPercentage,
-							SizeSpecificData.CollisionType));
+							SizeSpecificData.CollisionShapesData[0].CollisionObjectReductionPercentage,
+							SizeSpecificData.CollisionShapesData[0].CollisionType));
 				}
-				else if (SizeSpecificData.ImplicitType == EImplicitTypeEnum::Chaos_Implicit_Sphere)
+				else if (SizeSpecificData.CollisionShapesData[0].ImplicitType == EImplicitTypeEnum::Chaos_Implicit_Sphere)
 				{
 					CollectionImplicits[TransformGroupIndex] = FGeometryDynamicCollection::FSharedImplicit(
 						FCollisionStructureManager::NewImplicitSphere(
 							InnerRadius[GeometryIndex],
-							SizeSpecificData.CollisionObjectReductionPercentage,
-							SizeSpecificData.CollisionType));
+							SizeSpecificData.CollisionShapesData[0].CollisionObjectReductionPercentage,
+							SizeSpecificData.CollisionShapesData[0].CollisionType));
+				}
+				else if (SizeSpecificData.CollisionShapesData[0].ImplicitType == EImplicitTypeEnum::Chaos_Implicit_Capsule)
+				{
+					const FVector BBoxExtent = InstanceBoundingBox.GetExtent(); // FBox's extents are 1/2 (Max - Min)
+					const float Radius = BBoxExtent.GetAbsMin();
+					const float Length = (BBoxExtent.GetAbsMax() - Radius) * 2.f;
+
+					if (Length < SMALL_NUMBER)
+					{
+						// make a more optimized shape : sphere
+						CollectionImplicits[TransformGroupIndex] = FGeometryDynamicCollection::FSharedImplicit(
+							FCollisionStructureManager::NewImplicitSphere(
+								Radius,
+								SizeSpecificData.CollisionShapesData[0].CollisionObjectReductionPercentage,
+								SizeSpecificData.CollisionShapesData[0].CollisionType));
+					}
+					else
+					{
+						CollectionImplicits[TransformGroupIndex] = FGeometryDynamicCollection::FSharedImplicit(
+							FCollisionStructureManager::NewImplicitCapsule(
+								Radius,
+								Length,
+								SizeSpecificData.CollisionShapesData[0].CollisionObjectReductionPercentage,
+								SizeSpecificData.CollisionShapesData[0].CollisionType));
+					}
 				}
-				else if (SizeSpecificData.ImplicitType == EImplicitTypeEnum::Chaos_Implicit_None)
+				else if (SizeSpecificData.CollisionShapesData[0].ImplicitType == EImplicitTypeEnum::Chaos_Implicit_None)
 				{
 					CollectionImplicits[TransformGroupIndex] = nullptr;
 				}
@@ -2393,17 +2369,21 @@ void FGeometryCollectionPhysicsProxy::InitializeSharedCollisionStructures(
 				const int32 SizeSpecificIdx = FindSizeSpecificIdx(SharedParams.SizeSpecificData, InstanceBoundingBox);
 				const FSharedSimulationSizeSpecificData& SizeSpecificData = SharedParams.SizeSpecificData[SizeSpecificIdx];
 
-				if (SizeSpecificData.ImplicitType == EImplicitTypeEnum::Chaos_Implicit_LevelSet)
+				if (SizeSpecificData.CollisionShapesData[0].ImplicitType == EImplicitTypeEnum::Chaos_Implicit_LevelSet)
 				{
 					const FVec3 Scale = 2 * InstanceBoundingBox.GetExtent() / MaxChildBounds; // FBox's extents are 1/2 (Max - Min)
 					const float ScaleMax = Scale.GetAbsMax();
 					const float ScaleMin = Scale.GetAbsMin();
 
-					float MinResolution = ScaleMin * static_cast<float>(SizeSpecificData.MinLevelSetResolution);
-					MinResolution = FMath::Clamp<float>(MinResolution, static_cast<float>(SizeSpecificData.MinLevelSetResolution), static_cast<float>(SizeSpecificData.MinClusterLevelSetResolution));
+					float MinResolution = ScaleMin * (float)SizeSpecificData.CollisionShapesData[0].LevelSetData.MinLevelSetResolution;
+					MinResolution = FMath::Clamp<float>(MinResolution,
+						(float)SizeSpecificData.CollisionShapesData[0].LevelSetData.MinLevelSetResolution,
+						(float)SizeSpecificData.CollisionShapesData[0].LevelSetData.MinClusterLevelSetResolution);
 
-					float MaxResolution = ScaleMax * static_cast<float>(SizeSpecificData.MaxLevelSetResolution);
-					MaxResolution = FMath::Clamp<float>(MaxResolution, static_cast<float>(SizeSpecificData.MaxLevelSetResolution), static_cast<float>(SizeSpecificData.MaxClusterLevelSetResolution));
+					float MaxResolution = ScaleMax * (float)SizeSpecificData.CollisionShapesData[0].LevelSetData.MaxLevelSetResolution;
+					MaxResolution = FMath::Clamp<float>(MaxResolution,
+						(float)SizeSpecificData.CollisionShapesData[0].LevelSetData.MaxLevelSetResolution,
+						(float)SizeSpecificData.CollisionShapesData[0].LevelSetData.MaxClusterLevelSetResolution);
 
 					//don't support non level-set serialization
 					ErrorReporter.SetPrefix(BaseErrorPrefix + " | Cluster Transform Index: " + FString::FromInt(ClusterTransformIdx));
@@ -2415,51 +2395,73 @@ void FGeometryCollectionPhysicsProxy::InitializeSharedCollisionStructures(
 							InstanceBoundingBox,
 							FMath::FloorToInt(MinResolution),
 							FMath::FloorToInt(MaxResolution),
-							SizeSpecificData.CollisionObjectReductionPercentage,
-							SizeSpecificData.CollisionType));
+							SizeSpecificData.CollisionShapesData[0].CollisionObjectReductionPercentage,
+							SizeSpecificData.CollisionShapesData[0].CollisionType));
 					// Fall back on sphere if level set rasterization failed.
 					if (!CollectionImplicits[ClusterTransformIdx])
 					{
 						CollectionImplicits[ClusterTransformIdx] = FGeometryDynamicCollection::FSharedImplicit(
 							FCollisionStructureManager::NewImplicitSphere(
 								InstanceBoundingBox.GetExtent().GetAbsMin(), // FBox's extents are 1/2 (Max - Min)
-								SizeSpecificData.CollisionObjectReductionPercentage,
-								SizeSpecificData.CollisionType));
+								SizeSpecificData.CollisionShapesData[0].CollisionObjectReductionPercentage,
+								SizeSpecificData.CollisionShapesData[0].CollisionType));
 					}
 
 					CollectionSimplicials[ClusterTransformIdx] = TUniquePtr<FSimplicial>(
 						FCollisionStructureManager::NewSimplicial(MassSpaceParticles, *UnionMesh, CollectionImplicits[ClusterTransformIdx].Get(),
 						SharedParams.MaximumCollisionParticleCount));
 				}
-				else if (SizeSpecificData.ImplicitType == EImplicitTypeEnum::Chaos_Implicit_Box)
+				else if (SizeSpecificData.CollisionShapesData[0].ImplicitType == EImplicitTypeEnum::Chaos_Implicit_Box)
 				{
 					CollectionImplicits[ClusterTransformIdx] = FGeometryDynamicCollection::FSharedImplicit(
 						FCollisionStructureManager::NewImplicitBox(
 							InstanceBoundingBox,
-							SizeSpecificData.CollisionObjectReductionPercentage,
-							SizeSpecificData.CollisionType));
+							SizeSpecificData.CollisionShapesData[0].CollisionObjectReductionPercentage,
+							SizeSpecificData.CollisionShapesData[0].CollisionType));
 
 					CollectionSimplicials[ClusterTransformIdx] = TUniquePtr<FSimplicial>(
 						FCollisionStructureManager::NewSimplicial(MassSpaceParticles, *UnionMesh, CollectionImplicits[ClusterTransformIdx].Get(),
 						SharedParams.MaximumCollisionParticleCount));
 				}
-				else if (SizeSpecificData.ImplicitType == EImplicitTypeEnum::Chaos_Implicit_Sphere)
+				else if (SizeSpecificData.CollisionShapesData[0].ImplicitType == EImplicitTypeEnum::Chaos_Implicit_Sphere)
 				{
  					CollectionImplicits[ClusterTransformIdx] = FGeometryDynamicCollection::FSharedImplicit(
 						FCollisionStructureManager::NewImplicitSphere(
 							InstanceBoundingBox.GetExtent().GetAbsMin(), // FBox's extents are 1/2 (Max - Min)
-							SizeSpecificData.CollisionObjectReductionPercentage,
-							SizeSpecificData.CollisionType));
+							SizeSpecificData.CollisionShapesData[0].CollisionObjectReductionPercentage,
+							SizeSpecificData.CollisionShapesData[0].CollisionType));
 
 					CollectionSimplicials[ClusterTransformIdx] = TUniquePtr<FSimplicial>(
 						FCollisionStructureManager::NewSimplicial(MassSpaceParticles, *UnionMesh, CollectionImplicits[ClusterTransformIdx].Get(),
 						SharedParams.MaximumCollisionParticleCount));
 				}
-				else if(SizeSpecificData.ImplicitType == EImplicitTypeEnum::Chaos_Implicit_Capsule)
+				else if (SizeSpecificData.CollisionShapesData[0].ImplicitType == EImplicitTypeEnum::Chaos_Implicit_Capsule)
 				{
-					ensure(false); // unsupported implicit type
-					CollectionImplicits[ClusterTransformIdx].Reset();
-					CollectionSimplicials[ClusterTransformIdx].Reset();
+					const FVector BBoxExtent = InstanceBoundingBox.GetExtent(); // FBox's extents are 1/2 (Max - Min)
+					const float Radius = BBoxExtent.GetAbsMin(); 
+					const float Length = (BBoxExtent.GetAbsMax() - Radius) * 2.f;
+
+					if (Length < SMALL_NUMBER)
+					{
+						// make a more optimized shape : sphere
+						CollectionImplicits[ClusterTransformIdx] = FGeometryDynamicCollection::FSharedImplicit(
+							FCollisionStructureManager::NewImplicitSphere(
+								Radius, 
+								SizeSpecificData.CollisionShapesData[0].CollisionObjectReductionPercentage,
+								SizeSpecificData.CollisionShapesData[0].CollisionType));
+					}
+					else
+					{
+						CollectionImplicits[ClusterTransformIdx] = FGeometryDynamicCollection::FSharedImplicit(
+							FCollisionStructureManager::NewImplicitCapsule(
+								Radius,
+								Length,
+								SizeSpecificData.CollisionShapesData[0].CollisionObjectReductionPercentage,
+								SizeSpecificData.CollisionShapesData[0].CollisionType));
+					}
+					CollectionSimplicials[ClusterTransformIdx] = TUniquePtr<FSimplicial>(
+						FCollisionStructureManager::NewSimplicial(MassSpaceParticles, *UnionMesh, CollectionImplicits[ClusterTransformIdx].Get(),
+							SharedParams.MaximumCollisionParticleCount));
 				}
 				else // Assume it's a union???
 				{
diff --git a/Engine/Source/Runtime/Experimental/Chaos/Public/GeometryCollection/GeometryCollectionCollisionStructureManager.h b/Engine/Source/Runtime/Experimental/Chaos/Public/GeometryCollection/GeometryCollectionCollisionStructureManager.h
--- a/Engine/Source/Runtime/Experimental/Chaos/Public/GeometryCollection/GeometryCollectionCollisionStructureManager.h
+++ b/Engine/Source/Runtime/Experimental/Chaos/Public/GeometryCollection/GeometryCollectionCollisionStructureManager.h
@@ -98,6 +98,12 @@ public:
 		const float CollisionObjectReduction,
 		const ECollisionTypeEnum CollisionType);
 
+	static FImplicit* NewImplicitCapsule(
+		const float Radius,
+		const float Length,
+		const float CollisionObjectReduction,
+		const ECollisionTypeEnum CollisionType);
+
 	static FImplicit* NewImplicitLevelset(
 		Chaos::FErrorReporter ErrorReporter,
 		const Chaos::FParticles& MeshParticles,
diff --git a/Engine/Source/Runtime/Experimental/Chaos/Public/GeometryCollection/GeometryCollectionSimulationCoreTypes.h b/Engine/Source/Runtime/Experimental/Chaos/Public/GeometryCollection/GeometryCollectionSimulationCoreTypes.h
--- a/Engine/Source/Runtime/Experimental/Chaos/Public/GeometryCollection/GeometryCollectionSimulationCoreTypes.h
+++ b/Engine/Source/Runtime/Experimental/Chaos/Public/GeometryCollection/GeometryCollectionSimulationCoreTypes.h
@@ -13,33 +13,64 @@
 class FGeometryCollection;
 class FGeometryDynamicCollection;
 
-struct FSharedSimulationSizeSpecificData
+
+struct FCollectionLevelSetData
 {
-	FSharedSimulationSizeSpecificData()
-		: MaxSize(0.f)
-		, CollisionType(ECollisionTypeEnum::Chaos_Surface_Volumetric)
-		, ImplicitType(EImplicitTypeEnum::Chaos_Implicit_Sphere)
-		, MinLevelSetResolution(5)
+	FCollectionLevelSetData()
+		: MinLevelSetResolution(5)
 		, MaxLevelSetResolution(10)
 		, MinClusterLevelSetResolution(25)
 		, MaxClusterLevelSetResolution(50)
-		, CollisionObjectReductionPercentage(0.f)
-		, CollisionParticlesFraction(1.f)
-		, MaximumCollisionParticles(60)
-		, DamageThreshold(250.f)
-	{
-	}
+	{}
 
-	float MaxSize;
-	ECollisionTypeEnum CollisionType;
-	EImplicitTypeEnum ImplicitType;
 	int32 MinLevelSetResolution;
 	int32 MaxLevelSetResolution;
 	int32 MinClusterLevelSetResolution;
 	int32 MaxClusterLevelSetResolution;
-	float CollisionObjectReductionPercentage;
+};
+
+struct FCollectionCollisionParticleData
+{
+	FCollectionCollisionParticleData()
+		: CollisionParticlesFraction(1.f)
+		, MaximumCollisionParticles(60)
+	{}
+
 	float CollisionParticlesFraction;
 	int32 MaximumCollisionParticles;
+};
+
+
+
+struct FCollectionCollisionTypeData
+{
+	FCollectionCollisionTypeData()
+		: CollisionType(ECollisionTypeEnum::Chaos_Surface_Volumetric)
+		, ImplicitType(EImplicitTypeEnum::Chaos_Implicit_Sphere)
+		, LevelSetData()
+		, CollisionParticleData()
+		, CollisionObjectReductionPercentage(0.f)
+	{
+	}
+
+	ECollisionTypeEnum CollisionType;
+	EImplicitTypeEnum ImplicitType;
+	FCollectionLevelSetData LevelSetData;
+	FCollectionCollisionParticleData CollisionParticleData;
+	float CollisionObjectReductionPercentage;
+};
+
+struct FSharedSimulationSizeSpecificData
+{
+	FSharedSimulationSizeSpecificData()
+		: MaxSize(0.f)
+		, CollisionShapesData({ FCollectionCollisionTypeData() })
+		, DamageThreshold(250.f)
+	{
+	}
+
+	float MaxSize;
+	TArray<FCollectionCollisionTypeData> CollisionShapesData;
 	float DamageThreshold;
 
 	bool operator<(const FSharedSimulationSizeSpecificData& Rhs) const { return MaxSize < Rhs.MaxSize; }
@@ -97,14 +128,17 @@ struct FSharedSimulationParameters
 	, MaximumCollisionParticleCount(InMaximumCollisionParticleCount)
 	{
 		SizeSpecificData.AddDefaulted();
-		SizeSpecificData[0].CollisionType = InCollisionType;
-		SizeSpecificData[0].ImplicitType = InImplicitType;
-		SizeSpecificData[0].MinLevelSetResolution = InMinLevelSetResolution;
-		SizeSpecificData[0].MaxLevelSetResolution = InMaxLevelSetResolution;
-		SizeSpecificData[0].MinClusterLevelSetResolution = InMinClusterLevelSetResolution;
-		SizeSpecificData[0].MaxClusterLevelSetResolution = InMaxClusterLevelSetResolution;
-		SizeSpecificData[0].CollisionParticlesFraction = InCollisionParticlesFraction;
-		SizeSpecificData[0].MaximumCollisionParticles = InMaximumCollisionParticleCount;
+		if (ensure(SizeSpecificData.Num() && SizeSpecificData[0].CollisionShapesData.Num()))
+		{
+			SizeSpecificData[0].CollisionShapesData[0].CollisionType = InCollisionType;
+			SizeSpecificData[0].CollisionShapesData[0].ImplicitType = InImplicitType;
+			SizeSpecificData[0].CollisionShapesData[0].LevelSetData.MinLevelSetResolution = InMinLevelSetResolution;
+			SizeSpecificData[0].CollisionShapesData[0].LevelSetData.MaxLevelSetResolution = InMaxLevelSetResolution;
+			SizeSpecificData[0].CollisionShapesData[0].LevelSetData.MinClusterLevelSetResolution = InMinClusterLevelSetResolution;
+			SizeSpecificData[0].CollisionShapesData[0].LevelSetData.MaxClusterLevelSetResolution = InMaxClusterLevelSetResolution;
+			SizeSpecificData[0].CollisionShapesData[0].CollisionParticleData.CollisionParticlesFraction = InCollisionParticlesFraction;
+			SizeSpecificData[0].CollisionShapesData[0].CollisionParticleData.MaximumCollisionParticles = InMaximumCollisionParticleCount;
+		}
 	}
 
 	bool bMassAsDensity;
@@ -129,6 +163,7 @@ struct FCollisionDataSimulationParameters
 	FCollisionDataSimulationParameters()
 		: DoGenerateCollisionData(false)
 		, SaveCollisionData(false)
+
 		, CollisionDataSizeMax(512)
 		, DoCollisionDataSpatialHash(false)
 		, CollisionDataSpatialHashRadius(50.f)
diff --git a/Engine/Source/Runtime/Experimental/GeometryCollectionEngine/Private/GeometryCollection/GeometryCollectionObject.cpp b/Engine/Source/Runtime/Experimental/GeometryCollectionEngine/Private/GeometryCollection/GeometryCollectionObject.cpp
--- a/Engine/Source/Runtime/Experimental/GeometryCollectionEngine/Private/GeometryCollection/GeometryCollectionObject.cpp
+++ b/Engine/Source/Runtime/Experimental/GeometryCollectionEngine/Private/GeometryCollection/GeometryCollectionObject.cpp
@@ -64,18 +64,18 @@ UGeometryCollection::UGeometryCollection(const FObjectInitializer& ObjectInitial
 	, ClusterConnectionType(EClusterConnectionTypeEnum::Chaos_PointImplicit)
 	, bStripOnCook(false)
 	, EnableNanite(false)
-	, CollisionType(ECollisionTypeEnum::Chaos_Volumetric)
-	, ImplicitType(EImplicitTypeEnum::Chaos_Implicit_Box)
-	, MinLevelSetResolution(10)
-	, MaxLevelSetResolution(10)
-	, MinClusterLevelSetResolution(50)
-	, MaxClusterLevelSetResolution(50)
-	, CollisionObjectReductionPercentage(0.0f)
+#if WITH_EDITORONLY_DATA
+	, CollisionType_DEPRECATED(ECollisionTypeEnum::Chaos_Volumetric)
+	, ImplicitType_DEPRECATED(EImplicitTypeEnum::Chaos_Implicit_Box)
+	, MinLevelSetResolution_DEPRECATED(10)
+	, MaxLevelSetResolution_DEPRECATED(10)
+	, MinClusterLevelSetResolution_DEPRECATED(50)
+	, MaxClusterLevelSetResolution_DEPRECATED(50)
+	, CollisionObjectReductionPercentage_DEPRECATED(0.0f)
+#endif
 	, bMassAsDensity(false)
 	, Mass(1.0f)
 	, MinimumMassClamp(0.1f)
-	, CollisionParticlesFraction(1.0f)
-	, MaximumCollisionParticles(60)
 	, EnableRemovePiecesOnFracture(false)
 	, GeometryCollection(new FGeometryCollection())
 {
@@ -87,36 +87,121 @@ UGeometryCollection::UGeometryCollection(const FObjectInitializer& ObjectInitial
 #endif
 }
 
-FGeometryCollectionSizeSpecificData::FGeometryCollectionSizeSpecificData()
-	: MaxSize(0.0f)
-	, CollisionType(ECollisionTypeEnum::Chaos_Volumetric)
-	, ImplicitType(EImplicitTypeEnum::Chaos_Implicit_Box)
-	, MinLevelSetResolution(5)
+FGeometryCollectionLevelSetData::FGeometryCollectionLevelSetData()
+	: MinLevelSetResolution(5)
 	, MaxLevelSetResolution(10)
 	, MinClusterLevelSetResolution(25)
 	, MaxClusterLevelSetResolution(50)
-	, CollisionObjectReductionPercentage(0.0f)
-	, CollisionParticlesFraction(1.0f)
+{
+}
+
+FGeometryCollectionCollisionParticleData::FGeometryCollectionCollisionParticleData()
+	: CollisionParticlesFraction(1.0f)
 	, MaximumCollisionParticles(60)
+{
+}
+
+
+
+FGeometryCollectionCollisionTypeData::FGeometryCollectionCollisionTypeData()
+	: CollisionType(ECollisionTypeEnum::Chaos_Volumetric)
+	, ImplicitType(EImplicitTypeEnum::Chaos_Implicit_Box)
+	, LevelSet()
+	, CollisionParticles()
+	, CollisionObjectReductionPercentage(0.0f)
+{
+}
+
+FGeometryCollectionSizeSpecificData::FGeometryCollectionSizeSpecificData()
+	: MaxSize(FLT_MAX)
+	, CollisionShapes({ FGeometryCollectionCollisionTypeData()})
+#if WITH_EDITORONLY_DATA
+	, CollisionType_DEPRECATED(ECollisionTypeEnum::Chaos_Volumetric)
+	, ImplicitType_DEPRECATED(EImplicitTypeEnum::Chaos_Implicit_Box)
+	, MinLevelSetResolution_DEPRECATED(5)
+	, MaxLevelSetResolution_DEPRECATED(10)
+	, MinClusterLevelSetResolution_DEPRECATED(25)
+	, MaxClusterLevelSetResolution_DEPRECATED(50)
+	, CollisionObjectReductionPercentage_DEPRECATED(0)
+	, CollisionParticlesFraction_DEPRECATED(1.f)
+	, MaximumCollisionParticles_DEPRECATED(60)
+#endif
 	, DamageThreshold(250.0)
 {
 }
 
+
+bool FGeometryCollectionSizeSpecificData::Serialize(FArchive& Ar)
+{
+	Ar.UsingCustomVersion(FUE5MainStreamObjectVersion::GUID);
+	return false;	//We only have this function to mark custom GUID. Still want serialize tagged properties
+}
+
+#if WITH_EDITORONLY_DATA
+void FGeometryCollectionSizeSpecificData::PostSerialize(const FArchive& Ar)
+{
+	const int32 CustomVer = Ar.CustomVer(FUE5MainStreamObjectVersion::GUID);
+	// make sure to load back the deprecated values in the new structure if necessary
+	if (Ar.IsLoading() && CustomVer < FUE5MainStreamObjectVersion::GeometryCollectionUserDefinedCollisionShapes)
+	{
+		if (ensure(CollisionShapes.Num()))
+		{
+			CollisionShapes[0].CollisionType = CollisionType_DEPRECATED;
+			CollisionShapes[0].ImplicitType = ImplicitType_DEPRECATED;
+			CollisionShapes[0].CollisionObjectReductionPercentage = CollisionObjectReductionPercentage_DEPRECATED;
+			CollisionShapes[0].CollisionParticles.CollisionParticlesFraction = CollisionParticlesFraction_DEPRECATED;
+			CollisionShapes[0].CollisionParticles.MaximumCollisionParticles = MaximumCollisionParticles_DEPRECATED;
+			CollisionShapes[0].LevelSet.MinLevelSetResolution = MinLevelSetResolution_DEPRECATED;
+			CollisionShapes[0].LevelSet.MaxLevelSetResolution = MaxLevelSetResolution_DEPRECATED;
+			CollisionShapes[0].LevelSet.MinClusterLevelSetResolution = MinClusterLevelSetResolution_DEPRECATED;
+			CollisionShapes[0].LevelSet.MaxClusterLevelSetResolution = MaxClusterLevelSetResolution_DEPRECATED;
+		}
+	}
+}
+#endif
+
+
+
 void FillSharedSimulationSizeSpecificData(FSharedSimulationSizeSpecificData& ToData, const FGeometryCollectionSizeSpecificData& FromData)
 {
-	ToData.CollisionType = FromData.CollisionType;
-	ToData.ImplicitType = FromData.ImplicitType;
 	ToData.MaxSize = FromData.MaxSize;
-	ToData.MinLevelSetResolution = FromData.MinLevelSetResolution;
-	ToData.MaxLevelSetResolution = FromData.MaxLevelSetResolution;
-	ToData.MinClusterLevelSetResolution = FromData.MinClusterLevelSetResolution;
-	ToData.MaxClusterLevelSetResolution = FromData.MaxClusterLevelSetResolution;
-	ToData.CollisionObjectReductionPercentage = FromData.CollisionObjectReductionPercentage;
-	ToData.CollisionParticlesFraction = FromData.CollisionParticlesFraction;
-	ToData.MaximumCollisionParticles = FromData.MaximumCollisionParticles;
+	check(FromData.CollisionShapes.Num());
+
+	ToData.CollisionShapesData.SetNumUninitialized(FromData.CollisionShapes.Num());
+	for (int i=0;i<FromData.CollisionShapes.Num();i++)
+	{
+		ToData.CollisionShapesData[i].CollisionType = FromData.CollisionShapes[i].CollisionType;
+		ToData.CollisionShapesData[i].ImplicitType = FromData.CollisionShapes[i].ImplicitType;
+		ToData.CollisionShapesData[i].LevelSetData.MinLevelSetResolution = FromData.CollisionShapes[i].LevelSet.MinLevelSetResolution;
+		ToData.CollisionShapesData[i].LevelSetData.MaxLevelSetResolution = FromData.CollisionShapes[i].LevelSet.MaxLevelSetResolution;
+		ToData.CollisionShapesData[i].LevelSetData.MinClusterLevelSetResolution = FromData.CollisionShapes[i].LevelSet.MinClusterLevelSetResolution;
+		ToData.CollisionShapesData[i].LevelSetData.MaxClusterLevelSetResolution = FromData.CollisionShapes[i].LevelSet.MaxClusterLevelSetResolution;
+		ToData.CollisionShapesData[i].CollisionObjectReductionPercentage = FromData.CollisionShapes[i].CollisionObjectReductionPercentage;
+		ToData.CollisionShapesData[i].CollisionParticleData.CollisionParticlesFraction = FromData.CollisionShapes[i].CollisionParticles.CollisionParticlesFraction;
+		ToData.CollisionShapesData[i].CollisionParticleData.MaximumCollisionParticles = FromData.CollisionShapes[i].CollisionParticles.MaximumCollisionParticles;
+	}
 	ToData.DamageThreshold = FromData.DamageThreshold;
 }
 
+FGeometryCollectionSizeSpecificData UGeometryCollection::GeometryCollectionSizeSpecificDataDefaults() const
+{
+	FGeometryCollectionSizeSpecificData Data;
+	check(Data.CollisionShapes.Num());
+
+	Data.MaxSize = FLT_MAX;
+	Data.CollisionShapes[0].CollisionType = ECollisionTypeEnum::Chaos_Volumetric;
+	Data.CollisionShapes[0].ImplicitType = EImplicitTypeEnum::Chaos_Implicit_Sphere;
+	Data.CollisionShapes[0].LevelSet.MinLevelSetResolution = 5;
+	Data.CollisionShapes[0].LevelSet.MaxLevelSetResolution = 10;
+	Data.CollisionShapes[0].LevelSet.MinClusterLevelSetResolution = 25;
+	Data.CollisionShapes[0].LevelSet.MaxClusterLevelSetResolution = 50;
+	Data.CollisionShapes[0].CollisionObjectReductionPercentage = 1.0;
+	Data.CollisionShapes[0].CollisionParticles.CollisionParticlesFraction = 1.0;
+	Data.CollisionShapes[0].CollisionParticles.MaximumCollisionParticles = 60;
+	Data.DamageThreshold = 250.0f;
+	return Data;
+}
+
 float KgCm3ToKgM3(float Density)
 {
 	return Density * 1000000;
@@ -129,30 +214,33 @@ float KgM3ToKgCm3(float Density)
 
 void UGeometryCollection::GetSharedSimulationParams(FSharedSimulationParameters& OutParams) const
 {
+	const FGeometryCollectionSizeSpecificData& SizeSpecificDefault = GetDefaultSizeSpecificData();
+	check(SizeSpecificDefault.CollisionShapes.Num());
+
 	OutParams.bMassAsDensity = bMassAsDensity;
 	OutParams.Mass = bMassAsDensity ? KgM3ToKgCm3(Mass) : Mass;	//todo(ocohen): we still have the solver working in old units. This is mainly to fix ui issues. Long term need to normalize units for best precision
 	OutParams.MinimumMassClamp = MinimumMassClamp;
-	OutParams.MaximumCollisionParticleCount = MaximumCollisionParticles;
+	OutParams.MaximumCollisionParticleCount = SizeSpecificDefault.CollisionShapes[0].CollisionParticles.MaximumCollisionParticles;
 
-	ECollisionTypeEnum SelectedCollisionType = CollisionType;
+	ECollisionTypeEnum SelectedCollisionType = SizeSpecificDefault.CollisionShapes[0].CollisionType;
 
-	if(SelectedCollisionType == ECollisionTypeEnum::Chaos_Volumetric && ImplicitType == EImplicitTypeEnum::Chaos_Implicit_LevelSet)
+	if (SelectedCollisionType == ECollisionTypeEnum::Chaos_Volumetric && SizeSpecificDefault.CollisionShapes[0].ImplicitType == EImplicitTypeEnum::Chaos_Implicit_LevelSet)
 	{
 		UE_LOG(LogGeometryCollectionInternal, Verbose, TEXT("LevelSet geometry selected but non-particle collisions selected. Forcing particle-implicit collisions for %s"), *GetPathName());
 		SelectedCollisionType = ECollisionTypeEnum::Chaos_Surface_Volumetric;
 	}
 
 	FGeometryCollectionSizeSpecificData InfSize;
-	InfSize.MaxSize = FLT_MAX;
-	InfSize.CollisionType = SelectedCollisionType;
-	InfSize.ImplicitType = ImplicitType;
-	InfSize.MinLevelSetResolution = MinLevelSetResolution;
-	InfSize.MaxLevelSetResolution = MaxLevelSetResolution;
-	InfSize.MinClusterLevelSetResolution = MinClusterLevelSetResolution;
-	InfSize.MaxClusterLevelSetResolution = MaxClusterLevelSetResolution;
-	InfSize.CollisionObjectReductionPercentage = CollisionObjectReductionPercentage;
-	InfSize.CollisionParticlesFraction = CollisionParticlesFraction;
-	InfSize.MaximumCollisionParticles = MaximumCollisionParticles;
+	InfSize.MaxSize = TNumericLimits<float>::Max();
+	InfSize.CollisionShapes[0].CollisionType = SelectedCollisionType;
+	InfSize.CollisionShapes[0].ImplicitType = SizeSpecificDefault.CollisionShapes[0].ImplicitType;
+	InfSize.CollisionShapes[0].LevelSet.MinLevelSetResolution = SizeSpecificDefault.CollisionShapes[0].LevelSet.MinLevelSetResolution;
+	InfSize.CollisionShapes[0].LevelSet.MaxLevelSetResolution = SizeSpecificDefault.CollisionShapes[0].LevelSet.MaxLevelSetResolution;
+	InfSize.CollisionShapes[0].LevelSet.MinClusterLevelSetResolution = SizeSpecificDefault.CollisionShapes[0].LevelSet.MinClusterLevelSetResolution;
+	InfSize.CollisionShapes[0].LevelSet.MaxClusterLevelSetResolution = SizeSpecificDefault.CollisionShapes[0].LevelSet.MaxClusterLevelSetResolution;
+	InfSize.CollisionShapes[0].CollisionObjectReductionPercentage = SizeSpecificDefault.CollisionShapes[0].CollisionObjectReductionPercentage;
+	InfSize.CollisionShapes[0].CollisionParticles.CollisionParticlesFraction = SizeSpecificDefault.CollisionShapes[0].CollisionParticles.CollisionParticlesFraction;
+	InfSize.CollisionShapes[0].CollisionParticles.MaximumCollisionParticles = SizeSpecificDefault.CollisionShapes[0].CollisionParticles.MaximumCollisionParticles;
 
 	OutParams.SizeSpecificData.SetNum(SizeSpecificData.Num() + 1);
 	FillSharedSimulationSizeSpecificData(OutParams.SizeSpecificData[0], InfSize);
@@ -236,6 +324,41 @@ void UGeometryCollection::RemoveElements(const FName & Group, const TArray<int32
 	InvalidateCollection();
 }
 
+int UGeometryCollection::GetDefaultSizeSpecificDataIndex() const
+{
+	int LargestIndex = INDEX_NONE;
+	float MaxSize = TNumericLimits<float>::Lowest();
+	for (int i = 0; i < SizeSpecificData.Num(); i++)
+	{
+		const float SizeSpecificDataMaxSize = SizeSpecificData[i].MaxSize;
+		if (MaxSize < SizeSpecificDataMaxSize)
+		{
+			MaxSize = SizeSpecificDataMaxSize;
+			LargestIndex = i;
+		}
+	}
+	check(LargestIndex != INDEX_NONE && LargestIndex < SizeSpecificData.Num());
+	return LargestIndex;
+}
+
+/** Size Specific Data Access */
+FGeometryCollectionSizeSpecificData& UGeometryCollection::GetDefaultSizeSpecificData()
+{
+	if (!SizeSpecificData.Num())
+	{
+		SizeSpecificData.Add(GeometryCollectionSizeSpecificDataDefaults());
+	}
+	const int DefaultSizeIndex = GetDefaultSizeSpecificDataIndex();
+	return SizeSpecificData[DefaultSizeIndex];
+}
+
+const FGeometryCollectionSizeSpecificData& UGeometryCollection::GetDefaultSizeSpecificData() const
+{
+	ensure(SizeSpecificData.Num());
+	const int DefaultSizeIndex = GetDefaultSizeSpecificDataIndex();
+	return SizeSpecificData[DefaultSizeIndex];
+}
+
 /** ReindexMaterialSections */
 void UGeometryCollection::ReindexMaterialSections()
 {
@@ -430,6 +553,7 @@ static void SerializeOldNaniteData(FArchive& Ar, UGeometryCollection* Owner)
 /** Serialize */
 void UGeometryCollection::Serialize(FArchive& Ar)
 {
+	bool bCreateSimulationData = false;
 	Ar.UsingCustomVersion(FDestructionObjectVersion::GUID);
 	Ar.UsingCustomVersion(FUE5MainStreamObjectVersion::GUID);
 	Chaos::FChaosArchive ChaosAr(Ar);
@@ -543,9 +667,7 @@ void UGeometryCollection::Serialize(FArchive& Ar)
 	{
 		ArchiveGeometryCollection->UpdateOldAttributeNames();
 		InvalidateCollection();
-#if WITH_EDITOR
-		CreateSimulationData();
-#endif
+		bCreateSimulationData = true;
 	}
 
 	if (Ar.CustomVer(FUE5MainStreamObjectVersion::GUID) == FUE5MainStreamObjectVersion::GeometryCollectionNaniteData || 
@@ -572,10 +694,47 @@ void UGeometryCollection::Serialize(FArchive& Ar)
 		}
 	}
 
+#if WITH_EDITORONLY_DATA
+	if (Ar.IsLoading() && Ar.CustomVer(FUE5MainStreamObjectVersion::GUID) < FUE5MainStreamObjectVersion::GeometryCollectionUserDefinedCollisionShapes)
+	{
+		auto HasDefault = [](const TArray<FGeometryCollectionSizeSpecificData>& DatasIn)
+		{
+			for (const FGeometryCollectionSizeSpecificData& Data : DatasIn)
+				if (FMath::IsNearlyEqual(Data.MaxSize, FLT_MAX))
+					return true;
+			return false;
+		};
+
+		if (!SizeSpecificData.Num() || !HasDefault(SizeSpecificData))
+		{
+			FGeometryCollectionSizeSpecificData Data;
+			Data.CollisionShapes[0].CollisionType =	CollisionType_DEPRECATED;
+			Data.CollisionShapes[0].ImplicitType = ImplicitType_DEPRECATED;
+			Data.CollisionShapes[0].LevelSet.MinLevelSetResolution = MinLevelSetResolution_DEPRECATED;
+			Data.CollisionShapes[0].LevelSet.MaxLevelSetResolution = MaxLevelSetResolution_DEPRECATED;
+			Data.CollisionShapes[0].LevelSet.MinClusterLevelSetResolution = MinClusterLevelSetResolution_DEPRECATED;
+			Data.CollisionShapes[0].LevelSet.MaxClusterLevelSetResolution = MaxClusterLevelSetResolution_DEPRECATED;
+			Data.CollisionShapes[0].CollisionObjectReductionPercentage = CollisionObjectReductionPercentage_DEPRECATED;
+			if (Data.CollisionShapes[0].ImplicitType == EImplicitTypeEnum::Chaos_Implicit_LevelSet)
+				Data.CollisionShapes[0].CollisionType = ECollisionTypeEnum::Chaos_Surface_Volumetric;
+			SizeSpecificData.Add(Data);
+		}
+		check(SizeSpecificData.Num() && SizeSpecificData[0].CollisionShapes.Num());
+
+		if (Ar.IsLoading())
+		{
+			InvalidateCollection();
+			bCreateSimulationData = true;
+		}
+	}
+
+	if (bCreateSimulationData)
+	{
+		CreateSimulationData();
+	}
 
-#if WITH_EDITOR
 	//for all versions loaded, make sure sim data is up to date
-	if (Ar.IsLoading())
+ 	if (Ar.IsLoading())
 	{
 		EnsureDataIsCooked();	//make sure loaded content is built
 	}
diff --git a/Engine/Source/Runtime/Experimental/GeometryCollectionEngine/Public/GeometryCollection/GeometryCollectionObject.h b/Engine/Source/Runtime/Experimental/GeometryCollectionEngine/Public/GeometryCollection/GeometryCollectionObject.h
--- a/Engine/Source/Runtime/Experimental/GeometryCollectionEngine/Public/GeometryCollection/GeometryCollectionObject.h
+++ b/Engine/Source/Runtime/Experimental/GeometryCollectionEngine/Public/GeometryCollection/GeometryCollectionObject.h
@@ -66,17 +66,67 @@ struct GEOMETRYCOLLECTIONENGINE_API FGeometryCollectionEmbeddedExemplar
 	int32 InstanceCount;
 };
 
+USTRUCT()
+struct GEOMETRYCOLLECTIONENGINE_API FGeometryCollectionLevelSetData
+{
+	GENERATED_BODY()
+
+		FGeometryCollectionLevelSetData();
+
+	/*
+	*  Resolution on the smallest axes for the level set. (def: 5)
+	*/
+	UPROPERTY(EditAnywhere, Category = "LevelSet")
+		int32 MinLevelSetResolution;
+
+	/*
+	*  Resolution on the smallest axes for the level set. (def: 10)
+	*/
+	UPROPERTY(EditAnywhere, Category = "LevelSet")
+		int32 MaxLevelSetResolution;
+
+	/*
+	*  Resolution on the smallest axes for the level set. (def: 5)
+	*/
+	UPROPERTY(EditAnywhere, Category = "LevelSet")
+		int32 MinClusterLevelSetResolution;
+
+	/*
+	*  Resolution on the smallest axes for the level set. (def: 10)
+	*/
+	UPROPERTY(EditAnywhere, Category = "LevelSet")
+		int32 MaxClusterLevelSetResolution;
+};
+
 
 USTRUCT()
-struct GEOMETRYCOLLECTIONENGINE_API FGeometryCollectionSizeSpecificData
+struct GEOMETRYCOLLECTIONENGINE_API FGeometryCollectionCollisionParticleData
 {
 	GENERATED_BODY()
 
-	FGeometryCollectionSizeSpecificData();
+		FGeometryCollectionCollisionParticleData();
 
-	/** The max size these settings apply to*/
-	UPROPERTY(EditAnywhere, Category = "Collisions")
-	float MaxSize;
+	/**
+	 * Number of particles on the triangulated surface to use for collisions.
+	 */
+	UPROPERTY(EditAnywhere, Category = "Particle")
+	float CollisionParticlesFraction;
+
+	/**
+	 * Max number of particles.
+	 */
+	UPROPERTY(EditAnywhere, Category = "Particle")
+	int32 MaximumCollisionParticles;
+};
+
+
+
+USTRUCT()
+struct GEOMETRYCOLLECTIONENGINE_API FGeometryCollectionCollisionTypeData
+{
+	GENERATED_BODY()
+
+		FGeometryCollectionCollisionTypeData();
 
 	/*
 	*  CollisionType defines how to initialize the rigid collision structures.
@@ -91,53 +141,121 @@ struct GEOMETRYCOLLECTIONENGINE_API FGeometryCollectionSizeSpecificData
 	EImplicitTypeEnum ImplicitType;
 
 	/*
-	*  Resolution on the smallest axes for the level set. (def: 5)
+	*  LevelSet Resolution data for rasterization.
 	*/
-	UPROPERTY(EditAnywhere, Category = "Collisions")
-	int32 MinLevelSetResolution;
+	UPROPERTY(EditAnywhere, Category = "Collisions", meta = (EditCondition = "ImplicitType == EImplicitTypeEnum::Chaos_Implicit_LevelSet", EditConditionHides))
+	FGeometryCollectionLevelSetData LevelSet;
 
 	/*
-	*  Resolution on the smallest axes for the level set. (def: 10)
+	*  Collision Particle data for surface samples during Particle-LevelSet collisions.
 	*/
-	UPROPERTY(EditAnywhere, Category = "Collisions")
-	int32 MaxLevelSetResolution;
+	UPROPERTY(EditAnywhere, Category = "Collisions", meta = (EditCondition = "CollisionType == ECollisionTypeEnum::Chaos_Surface_Volumetric", EditConditionHides))
+	FGeometryCollectionCollisionParticleData CollisionParticles;
 
 	/*
-	*  Resolution on the smallest axes for the level set. (def: 5)
+	*  Uniform scale on the collision body. (def: 0)
 	*/
 	UPROPERTY(EditAnywhere, Category = "Collisions")
-	int32 MinClusterLevelSetResolution;
+	float CollisionObjectReductionPercentage;
 
-	/*
-	*  Resolution on the smallest axes for the level set. (def: 10)
-	*/
+};
+
+
+USTRUCT()
+struct GEOMETRYCOLLECTIONENGINE_API FGeometryCollectionSizeSpecificData
+{
+	GENERATED_BODY()
+
+	FGeometryCollectionSizeSpecificData();
+
+	/** The max size these settings apply to*/
 	UPROPERTY(EditAnywhere, Category = "Collisions")
-	int32 MaxClusterLevelSetResolution;
+	float MaxSize;
 
 	/*
-	*  Resolution on the smallest axes for the level set. (def: 10)
+	* Collision Shapes allow kfor multiple collision types per rigid body. 
 	*/
 	UPROPERTY(EditAnywhere, Category = "Collisions")
-	int32 CollisionObjectReductionPercentage;
+	TArray<FGeometryCollectionCollisionTypeData> CollisionShapes;
+
+#if WITH_EDITORONLY_DATA
+	/*
+	 *  CollisionType defines how to initialize the rigid collision structures.
+	 */
+	UPROPERTY(meta = (DeprecatedProperty, DeprecationMessage = "Use Collision.CollisionType instead."))
+	ECollisionTypeEnum CollisionType_DEPRECATED;
+
+	/*
+	 *  CollisionType defines how to initialize the rigid collision structures.
+	 */
+	UPROPERTY(meta = (DeprecatedProperty, DeprecationMessage = "Use Collision.ImplicitType instead."))
+	EImplicitTypeEnum ImplicitType_DEPRECATED;
+
+	/*
+	 *  Resolution on the smallest axes for the level set. (def: 5)
+	 */
+	UPROPERTY(meta = (DeprecatedProperty, DeprecationMessage = "Use Collision.LevelSet.MinLevelSetResolution instead."))
+	int32 MinLevelSetResolution_DEPRECATED;
+
+	/*
+	 *  Resolution on the smallest axes for the level set. (def: 10)
+	 */
+	UPROPERTY(meta = (DeprecatedProperty, DeprecationMessage = "Use Collision.LevelSet.MaxLevelSetResolution instead."))
+		int32 MaxLevelSetResolution_DEPRECATED;
+
+	/*
+	 *  Resolution on the smallest axes for the level set. (def: 5)
+	 */
+	UPROPERTY(meta = (DeprecatedProperty, DeprecationMessage = "Use Collision.LevelSet.MinClusterLevelSetResolution instead."))
+	int32 MinClusterLevelSetResolution_DEPRECATED;
+
+	/*
+	 *  Resolution on the smallest axes for the level set. (def: 10)
+	 */
+	UPROPERTY(meta = (DeprecatedProperty, DeprecationMessage = "Use Collision.LevelSet.MaxClusterLevelSetResolution instead."))
+	int32 MaxClusterLevelSetResolution_DEPRECATED;
+
+	/*
+	 *  Resolution on the smallest axes for the level set. (def: 10)
+	 */
+	UPROPERTY(meta = (DeprecatedProperty, DeprecationMessage = "Use Collision.CollisionObjectReductionPercentage instead."))
+	int32 CollisionObjectReductionPercentage_DEPRECATED;
 
 	/**
 	 * Number of particles on the triangulated surface to use for collisions.
 	 */
-	UPROPERTY(EditAnywhere, Category = "Collisions")
-	float CollisionParticlesFraction;
+	UPROPERTY(meta = (DeprecatedProperty, DeprecationMessage = "Use Collision.CollisionParticlesFraction instead."))
+	float CollisionParticlesFraction_DEPRECATED;
 
 	/**
 	 * Max number of particles.
 	 */
-	UPROPERTY(EditAnywhere, Category = "Collisions")
-	int32 MaximumCollisionParticles;
+	UPROPERTY(meta = (DeprecatedProperty, DeprecationMessage = "Use Collision.MaximumCollisionParticles instead."))
+	int32 MaximumCollisionParticles_DEPRECATED;
+#endif
 
 	/**
 	 * Max number of particles.
-	*/
+	 */
 	UPROPERTY(EditAnywhere, Category = "Collisions")
 	int32 DamageThreshold;
 
+	bool Serialize(FArchive& Ar);
+#if WITH_EDITORONLY_DATA
+	void PostSerialize(const FArchive& Ar);
+#endif
+};
+
+template<>
+struct TStructOpsTypeTraits<FGeometryCollectionSizeSpecificData> : public TStructOpsTypeTraitsBase2<FGeometryCollectionSizeSpecificData>
+{
+	enum
+	{
+		WithSerializer = true,
+#if WITH_EDITORONLY_DATA
+		WithPostSerialize = true
+#endif
+	};
 };
 
 class FGeometryCollectionNaniteData
@@ -306,48 +424,50 @@ public:
 	UPROPERTY(EditAnywhere, BlueprintReadOnly, Category = "Nanite")
 	bool EnableNanite;
 
+#if WITH_EDITORONLY_DATA
 	/*
 	*  CollisionType defines how to initialize the rigid collision structures.
 	*/
-	UPROPERTY(EditAnywhere, BlueprintReadOnly, Category = "Collisions")
-	ECollisionTypeEnum CollisionType;
+	UPROPERTY(meta = (DeprecatedProperty, DeprecationMessage = "This property is deprecated. Use the default SizeSpecificData instead."))
+	ECollisionTypeEnum CollisionType_DEPRECATED;
 
 	/*
 	*  CollisionType defines how to initialize the rigid collision structures.
 	*/
-	UPROPERTY(EditAnywhere, BlueprintReadOnly, Category = "Collisions")
-	EImplicitTypeEnum ImplicitType;
+	UPROPERTY(meta = (DeprecatedProperty, DeprecationMessage = "This property is deprecated. Use the default SizeSpecificData instead."))
+	EImplicitTypeEnum ImplicitType_DEPRECATED;
 
 	/*
 	*  Resolution on the smallest axes for the level set. (def: 5)
 	*/
-	UPROPERTY(EditAnywhere, BlueprintReadOnly, Category = "Collisions")
-	int32 MinLevelSetResolution;
+	UPROPERTY(meta = (DeprecatedProperty, DeprecationMessage = "This property is deprecated. Use the default SizeSpecificData instead."))
+	int32 MinLevelSetResolution_DEPRECATED;
 
 	/*
 	*  Resolution on the smallest axes for the level set. (def: 10)
 	*/
-	UPROPERTY(EditAnywhere, BlueprintReadOnly, Category = "Collisions")
-	int32 MaxLevelSetResolution;
+	UPROPERTY(meta = (DeprecatedProperty, DeprecationMessage = "This property is deprecated. Use the default SizeSpecificData instead."))
+	int32 MaxLevelSetResolution_DEPRECATED;
 
 	/*
 	*  Resolution on the smallest axes for the level set. (def: 5)
 	*/
-	UPROPERTY(EditAnywhere, BlueprintReadOnly, Category = "Collisions")
-	int32 MinClusterLevelSetResolution;
+	UPROPERTY(meta = (DeprecatedProperty, DeprecationMessage = "This property is deprecated. Use the default SizeSpecificData instead."))
+	int32 MinClusterLevelSetResolution_DEPRECATED;
 
 	/*
 	*  Resolution on the smallest axes for the level set. (def: 10)
 	*/
-	UPROPERTY(EditAnywhere, BlueprintReadOnly, Category = "Collisions")
-	int32 MaxClusterLevelSetResolution;
+	UPROPERTY(meta = (DeprecatedProperty, DeprecationMessage = "This property is deprecated. Use the default SizeSpecificData instead."))
+	int32 MaxClusterLevelSetResolution_DEPRECATED;
 
 	/*
 	*  Resolution on the smallest axes for the level set. (def: 10)
 	*/
-	UPROPERTY(EditAnywhere, BlueprintReadOnly, Category = "Collisions")
-	float CollisionObjectReductionPercentage;
-
+	UPROPERTY(meta = (DeprecatedProperty, DeprecationMessage = "This property is deprecated. Use the default SizeSpecificData instead."))
+	float CollisionObjectReductionPercentage_DEPRECATED;
+#endif
+	
 	/**
 	* Mass As Density, units are in kg/m^3
 	*/
@@ -366,21 +486,33 @@ public:
 	UPROPERTY(EditAnywhere, BlueprintReadOnly, Category = "Collisions")
 	float MinimumMassClamp;
 
+#if WITH_EDITORONLY_DATA
 	/**
 	 * Number of particles on the triangulated surface to use for collisions.
 	 */
-	UPROPERTY(EditAnywhere, BlueprintReadOnly, Category = "Collisions")
-	float CollisionParticlesFraction;
+	UPROPERTY(meta = (DeprecatedProperty, DeprecationMessage = "This property is deprecated. Use the default SizeSpecificData instead."))
+	float CollisionParticlesFraction_DEPRECATED;
 
 	/**
 	 * Max number of particles.
 	 */
-	UPROPERTY(EditAnywhere, BlueprintReadOnly, Category = "Collisions")
-	int32 MaximumCollisionParticles;
+	UPROPERTY(meta = (DeprecatedProperty, DeprecationMessage = "This property is deprecated. Use the default SizeSpecificData instead."))
+	int32 MaximumCollisionParticles_DEPRECATED;
+#endif
 
+	/*
+	* Size Specfic Data reflects the default geometry to bind to rigid bodies smaller
+	* than the max size volume. This can also be empty to reflect no collision geometry
+	* for the collection. 
+	*/
 	UPROPERTY(EditAnywhere, Category = "Collisions")
 	TArray<FGeometryCollectionSizeSpecificData> SizeSpecificData;
 
+	int GetDefaultSizeSpecificDataIndex() const;
+	FGeometryCollectionSizeSpecificData GeometryCollectionSizeSpecificDataDefaults() const;
+	FGeometryCollectionSizeSpecificData& GetDefaultSizeSpecificData();
+	const FGeometryCollectionSizeSpecificData& GetDefaultSizeSpecificData() const;
+
 	/**
 	* Enable remove pieces on fracture
 	*/
@@ -407,7 +539,6 @@ public:
 private:
 #if WITH_EDITOR
 	void CreateSimulationDataImp(bool bCopyFromDDC);
-
 #endif
 
 private:
