Chaos : Geometry Collection :  (Review Update)
Collision Attribute Restructuring.
-	Updated the size specific data to be the primary container for collision settings within the geometry collection.
-	Removed non size specific default parameters from the asset.
-	Added support for Capsule collisions.
#okforgithup, public, internal

diff --git a/Engine/Source/Runtime/Experimental/Chaos/Private/PhysicsProxy/GeometryCollectionPhysicsProxy.cpp b/Engine/Source/Runtime/Experimental/Chaos/Private/PhysicsProxy/GeometryCollectionPhysicsProxy.cpp
--- a/Engine/Source/Runtime/Experimental/Chaos/Private/PhysicsProxy/GeometryCollectionPhysicsProxy.cpp
+++ b/Engine/Source/Runtime/Experimental/Chaos/Private/PhysicsProxy/GeometryCollectionPhysicsProxy.cpp
@@ -204,16 +204,16 @@ void PopulateSimulatedParticle(
 	Handle->SetCollisionGroup(CollisionGroup);
 
 	// @todo(GCCollisionShapes) : add support for multiple shapes, currently just one. 
-	FCollectionCollisionTypeData SingleSuportedCollisionTypeData = FCollectionCollisionTypeData();
+	FCollectionCollisionTypeData SingleSupportedCollisionTypeData = FCollectionCollisionTypeData();
 	if (SharedParams.SizeSpecificData.Num() && SharedParams.SizeSpecificData[0].CollisionShapesData.Num())
 	{
-		SingleSuportedCollisionTypeData = SharedParams.SizeSpecificData[0].CollisionShapesData[0];
+		SingleSupportedCollisionTypeData = SharedParams.SizeSpecificData[0].CollisionShapesData[0];
 	}
 
 	if (Implicit)	//todo(ocohen): this is only needed for cases where clusters have no proxy. Kind of gross though, should refactor
 	{
 		TSharedPtr<Chaos::FImplicitObject, ESPMode::ThreadSafe> SharedImplicitTS(Implicit->DeepCopy().Release());
-		FCollisionStructureManager::UpdateImplicitFlags(SharedImplicitTS.Get(), SingleSuportedCollisionTypeData.CollisionType);
+		FCollisionStructureManager::UpdateImplicitFlags(SharedImplicitTS.Get(), SingleSupportedCollisionTypeData.CollisionType);
 		Handle->SetSharedGeometry(SharedImplicitTS);
 		Handle->SetHasBounds(true);
 		Handle->SetLocalBounds(SharedImplicitTS->BoundingBox());
@@ -223,7 +223,7 @@ void PopulateSimulatedParticle(
 		Handle->SetWorldSpaceInflatedBounds(TransformedBBox);
 	}
 
-	if (Simplicial && SingleSuportedCollisionTypeData.CollisionType == ECollisionTypeEnum::Chaos_Surface_Volumetric)
+	if (Simplicial && SingleSupportedCollisionTypeData.CollisionType == ECollisionTypeEnum::Chaos_Surface_Volumetric)
 	{
 		Handle->CollisionParticlesInitIfNeeded();
 
@@ -355,7 +355,7 @@ FAutoConsoleVariableRef CVarReportHighParticleFraction(TEXT("p.gc.ReportHighPart
 void FGeometryCollectionPhysicsProxy::Initialize(Chaos::FPBDRigidsEvolutionBase *Evolution)
 {
 	check(IsInGameThread());
-										 //
+	//
 	// Game thread initilization. 
 	//
 	//  1) Create a input buffer to store all game thread side data. 
diff --git a/Engine/Source/Runtime/Experimental/GeometryCollectionEngine/Private/GeometryCollection/GeometryCollectionObject.cpp b/Engine/Source/Runtime/Experimental/GeometryCollectionEngine/Private/GeometryCollection/GeometryCollectionObject.cpp
--- a/Engine/Source/Runtime/Experimental/GeometryCollectionEngine/Private/GeometryCollection/GeometryCollectionObject.cpp
+++ b/Engine/Source/Runtime/Experimental/GeometryCollectionEngine/Private/GeometryCollection/GeometryCollectionObject.cpp
@@ -700,8 +700,12 @@ void UGeometryCollection::Serialize(FArchive& Ar)
 		auto HasDefault = [](const TArray<FGeometryCollectionSizeSpecificData>& DatasIn)
 		{
 			for (const FGeometryCollectionSizeSpecificData& Data : DatasIn)
+			{
 				if (FMath::IsNearlyEqual(Data.MaxSize, FLT_MAX))
+				{
 					return true;
+				}
+			}
 			return false;
 		};
 
@@ -721,11 +725,8 @@ void UGeometryCollection::Serialize(FArchive& Ar)
 		}
 		check(SizeSpecificData.Num() && SizeSpecificData[0].CollisionShapes.Num());
 
-		if (Ar.IsLoading())
-		{
-			InvalidateCollection();
-			bCreateSimulationData = true;
-		}
+		InvalidateCollection();
+		bCreateSimulationData = true;
 	}
 
 	if (bCreateSimulationData)
