Fix Physics Project not running in LWC mode because of serialization issues

diff --git a/Engine/Source/Runtime/Experimental/Chaos/Public/Chaos/Box.h b/Engine/Source/Runtime/Experimental/Chaos/Public/Chaos/Box.h
--- a/Engine/Source/Runtime/Experimental/Chaos/Public/Chaos/Box.h
+++ b/Engine/Source/Runtime/Experimental/Chaos/Public/Chaos/Box.h
@@ -405,7 +405,10 @@ namespace Chaos
 			Ar.UsingCustomVersion(FReleaseObjectVersion::GUID);
 			if (Ar.CustomVer(FReleaseObjectVersion::GUID) >= FReleaseObjectVersion::MarginAddedToConvexAndBox)
 			{
-				Ar << FImplicitObject::Margin;
+				// LWC_TODO : potential precision loss, to be changed when we can serialize FReal as double
+				FRealSingle MarginFloat = (FRealSingle)FImplicitObject::Margin;
+				Ar << MarginFloat;
+				FImplicitObject::Margin = (T)MarginFloat;
 			}
 		}
 
diff --git a/Engine/Source/Runtime/Experimental/Chaos/Public/Chaos/Capsule.h b/Engine/Source/Runtime/Experimental/Chaos/Public/Chaos/Capsule.h
--- a/Engine/Source/Runtime/Experimental/Chaos/Public/Chaos/Capsule.h
+++ b/Engine/Source/Runtime/Experimental/Chaos/Public/Chaos/Capsule.h
@@ -298,7 +298,7 @@ namespace Chaos
 			MSegment.Serialize(Ar);
 
 			// Radius is now stored in the base class Margin
-			FReal ArRadius = GetRadius();
+			FRealSingle ArRadius = (FRealSingle)GetRadius(); // LWC_TODO : potential precision loss, to be changed when we can serialize FReal as double
 			Ar << ArRadius;
 			SetRadius(ArRadius);
 			
diff --git a/Engine/Source/Runtime/Experimental/Chaos/Public/Chaos/Convex.h b/Engine/Source/Runtime/Experimental/Chaos/Public/Chaos/Convex.h
--- a/Engine/Source/Runtime/Experimental/Chaos/Public/Chaos/Convex.h
+++ b/Engine/Source/Runtime/Experimental/Chaos/Public/Chaos/Convex.h
@@ -755,7 +755,10 @@ namespace Chaos
 
 			if (Ar.CustomVer(FExternalPhysicsCustomObjectVersion::GUID) >= FExternalPhysicsCustomObjectVersion::AddConvexCenterOfMassAndVolume)
 			{
-				Ar << Volume;
+				FRealSingle VolumeFloat = (FRealSingle)Volume; // LWC_TODO : potential precision loss, to be changed when we can serialize FReal as double
+				Ar << VolumeFloat;
+				Volume = (FReal)VolumeFloat;
+
 				Ar << CenterOfMass;
 			}
 			else if (Ar.IsLoading())
@@ -776,7 +779,9 @@ namespace Chaos
 			Ar.UsingCustomVersion(FReleaseObjectVersion::GUID);
 			if (Ar.CustomVer(FReleaseObjectVersion::GUID) >= FReleaseObjectVersion::MarginAddedToConvexAndBox)
 			{
-				Ar << FImplicitObject::Margin;
+				FRealSingle MarginFloat = (FRealSingle)FImplicitObject::Margin; // LWC_TODO : potential precision loss, to be changed when we can serialize FReal as double
+				Ar << MarginFloat;
+				FImplicitObject::Margin = (FReal)MarginFloat;
 			}
 
 			if (Ar.CustomVer(FReleaseObjectVersion::GUID) >= FReleaseObjectVersion::StructureDataAddedToConvex)
diff --git a/Engine/Source/Runtime/Experimental/Chaos/Public/Chaos/Segment.h b/Engine/Source/Runtime/Experimental/Chaos/Public/Chaos/Segment.h
--- a/Engine/Source/Runtime/Experimental/Chaos/Public/Chaos/Segment.h
+++ b/Engine/Source/Runtime/Experimental/Chaos/Public/Chaos/Segment.h
@@ -50,7 +50,11 @@ namespace Chaos
 
 		FORCEINLINE void Serialize(FArchive &Ar) 
 		{
-			Ar << MPoint << MAxis << MLength;
+			Ar << MPoint << MAxis;
+			
+			FRealSingle LengthFloat = (FRealSingle)MLength; // LWC_TODO : potential precision loss, to be changed when we can serialize FReal as double
+			Ar << LengthFloat;
+			MLength = (T)LengthFloat;
 		}
 
 		FORCEINLINE TAABB<T, 3> BoundingBox() const
diff --git a/Engine/Source/Runtime/Experimental/Chaos/Public/Chaos/Sphere.h b/Engine/Source/Runtime/Experimental/Chaos/Public/Chaos/Sphere.h
--- a/Engine/Source/Runtime/Experimental/Chaos/Public/Chaos/Sphere.h
+++ b/Engine/Source/Runtime/Experimental/Chaos/Public/Chaos/Sphere.h
@@ -258,7 +258,7 @@ namespace Chaos
 			Ar << Center;
 
 			// Radius is now stored in the base class Margin
-			FReal ArRadius = GetRadius();
+			FRealSingle ArRadius = (FRealSingle)GetRadius(); // LWC_TODO : potential precision loss, to be changed when we can serialize FReal as double
 			Ar << ArRadius;
 			SetRadius(ArRadius);
 		}
diff --git a/Engine/Source/Runtime/Experimental/ChaosCore/Public/Chaos/Core.h b/Engine/Source/Runtime/Experimental/ChaosCore/Public/Chaos/Core.h
--- a/Engine/Source/Runtime/Experimental/ChaosCore/Public/Chaos/Core.h
+++ b/Engine/Source/Runtime/Experimental/ChaosCore/Public/Chaos/Core.h
@@ -35,5 +35,5 @@ namespace Chaos
 	// NOTE: if you get a merge conflict on the GUID, you must replace it with a new GUID - do not accept the source or target
 	// or you will likely get DDC version conflicts resulting in crashes during load.
 	// Core version string for Chaos data. Any DDC builder dependent on Chaos for serialization should depend on this version
-	static const FString ChaosVersionString = TEXT("2010BDD8-2309-49A0-919D-6BE431CFEE3D");
+	static const FString ChaosVersionString = TEXT("314216DF-671E-40C5-B8D9-E508B42F18DC");
 }
