Chaos - ConstraintGraph crash fix
We now have non-dynamic particles in the constraint graph, but the edge coloring was not updated to take this into account. The coloring algorithm could modify edges belocing to a different island.

diff --git a/Engine/Source/Runtime/Experimental/Chaos/Private/Chaos/PBDConstraintColor.cpp b/Engine/Source/Runtime/Experimental/Chaos/Private/Chaos/PBDConstraintColor.cpp
--- a/Engine/Source/Runtime/Experimental/Chaos/Private/Chaos/PBDConstraintColor.cpp
+++ b/Engine/Source/Runtime/Experimental/Chaos/Private/Chaos/PBDConstraintColor.cpp
@@ -266,10 +266,6 @@ void FPBDConstraintColor::ComputeContactGraph(const int32 Island, const FPBDCons
 				continue;
 			}
 
-			// Assign the level and update max level for the island if required
-			ColorEdge.Level = Level;
-			IslandData[Island].MaxLevel = FGenericPlatformMath::Max(IslandData[Island].MaxLevel, ColorEdge.Level);
-
 			// Find adjacent node and recurse
 			int32 OtherNode = INDEX_NONE;
 			if(GraphEdge.FirstNode == NodeIndex)
@@ -281,10 +277,22 @@ void FPBDConstraintColor::ComputeContactGraph(const int32 Island, const FPBDCons
 				OtherNode = GraphEdge.FirstNode;
 			}
 
-			// If we have a node, append it to our queue on the next level
-			if(OtherNode != INDEX_NONE)
+			// Assign the level and update max level for the island if required
+			// NOTE: if we hit a non-dynamic particle (node), it will contain all of
+			// the contacts (edges) for dynamic particles interacting with it. They
+			// may not all be in the same island, which is ok (e.g., two separated
+			// boxes sat on a large plane). We need to ignore edges that are in other islands
+			// @todo(chaos): we should probably store the island index with each edge.
+			if (ConstraintGraph.Nodes[OtherNode].Island == Island)
 			{
-				NodeQueue.Emplace(ColorEdge.Level + 1, OtherNode);
+				ColorEdge.Level = Level;
+				IslandData[Island].MaxLevel = FGenericPlatformMath::Max(IslandData[Island].MaxLevel, ColorEdge.Level);
+
+				// If we have an other node, append it to our queue on the next level
+				if (OtherNode != INDEX_NONE)
+				{
+					NodeQueue.Emplace(ColorEdge.Level + 1, OtherNode);
+				}
 			}
 		}
 	}
@@ -376,3 +384,32 @@ int FPBDConstraintColor::GetIslandMaxLevel(int32 Island) const
 	}
 	return -1;
 }
+
+#if !UE_BUILD_SHIPPING
+bool FPBDConstraintColor::DebugCheckGraph(const FPBDConstraintGraph& ConstraintGraph) const
+{
+	if (IslandData.Num() != ConstraintGraph.NumIslands())
+	{
+		UE_LOG(LogChaos, Error, TEXT("Constraint Color island count mismatch %d != %d"), IslandData.Num(), ConstraintGraph.NumIslands());
+		return false;
+	}
+
+	bool bValid = true;
+
+	for (int32 Island = 0; Island < ConstraintGraph.NumIslands(); ++Island)
+	{
+		const TArray<int32>& ConstraintDataIndices = ConstraintGraph.GetIslandConstraintData(Island);
+		for (int32 EdgeIndex : ConstraintDataIndices)
+		{
+			if (Edges[EdgeIndex].Level > IslandData[Island].MaxLevel)
+			{
+				UE_LOG(LogChaos, Error, TEXT("Edge %d level out of range for island %d (%d > %d)"), EdgeIndex, Island, Edges[EdgeIndex].Level, IslandData[Island].MaxLevel);
+				bValid = false;
+			}
+		}
+	}
+
+	return bValid;
+}
+#endif
+
diff --git a/Engine/Source/Runtime/Experimental/Chaos/Private/Chaos/PBDConstraintGraph.cpp b/Engine/Source/Runtime/Experimental/Chaos/Private/Chaos/PBDConstraintGraph.cpp
--- a/Engine/Source/Runtime/Experimental/Chaos/Private/Chaos/PBDConstraintGraph.cpp
+++ b/Engine/Source/Runtime/Experimental/Chaos/Private/Chaos/PBDConstraintGraph.cpp
@@ -617,7 +617,7 @@ void FPBDConstraintGraph::ComputeIslands(const TParticleView<FPBDRigidParticles>
 	check(IslandToParticles.Num() == IslandToConstraints.Num());
 	check(IslandToParticles.Num() == IslandToData.Num());
 	// @todo(ccaulfield): make a more complex unit test to check island integrity
-	//checkSlow(CheckIslands(InParticles, ActiveIndices));
+	checkSlow(DebugCheckGraph());
 }
 
 
@@ -919,14 +919,13 @@ void FPBDConstraintGraph::DisableParticles(const TSet<FGeometryParticleHandle *>
 	}
 }
 
-
-bool FPBDConstraintGraph::CheckIslands(const TArray<FGeometryParticleHandle *>& Particles)
+#if !UE_BUILD_SHIPPING
+bool FPBDConstraintGraph::DebugCheckGraph() const
 {
 	bool bIsValid = true;
 
 	// Check that no particles are in multiple islands
 	TSet<FGeometryParticleHandle*> IslandParticlesUnionSet;
-	IslandParticlesUnionSet.Reserve(Particles.Num());
 	for (int32 Island = 0; Island < IslandToParticles.Num(); ++Island)
 	{
 		TSet<FGeometryParticleHandle*> IslandParticlesSet = TSet<FGeometryParticleHandle*>(IslandToParticles[Island]);
@@ -947,7 +946,7 @@ bool FPBDConstraintGraph::CheckIslands(const TArray<FGeometryParticleHandle *>&
 		IslandParticlesUnionSet = IslandParticlesUnionSet.Union(IslandParticlesSet);
 	}
 
-	// Check that no constraints refer in the same island
+	// Check that no constraints are in multiple islands
 	TSet<int32> IslandConstraintDataUnionSet;
 	IslandConstraintDataUnionSet.Reserve(Edges.Num());
 	for (int32 Island = 0; Island < IslandToConstraints.Num(); ++Island)
@@ -963,5 +962,36 @@ bool FPBDConstraintGraph::CheckIslands(const TArray<FGeometryParticleHandle *>&
 		IslandConstraintDataUnionSet = IslandConstraintDataUnionSet.Union(IslandConstraintDataSet);
 	}
 
+	// Check that no constraints connect dynamic particles in different islands
+	// (non-dynamic particles can be in multiple islands)
+	for (int32 Island = 0; Island < IslandToConstraints.Num(); ++Island)
+	{
+		const TArray<int32> IslandEdges = IslandToConstraints[Island];
+		for (int32 EdgeIndex : IslandEdges)
+		{
+			const FGraphEdge& Edge = Edges[EdgeIndex];
+			const FGraphNode& Node0 = Nodes[Edge.FirstNode];
+			const FGraphNode& Node1 = Nodes[Edge.SecondNode];
+			if ((Node0.Particle != nullptr) && (Node1.Particle != nullptr) && FConstGenericParticleHandle(Node0.Particle)->IsDynamic() && FConstGenericParticleHandle(Node1.Particle)->IsDynamic())
+			{
+				if (Node0.Island != Node1.Island)
+				{
+					UE_LOG(LogChaos, Error, TEXT("Island %d Constraint %d connects particles in different islands %d %d"), Island, EdgeIndex, Node0.Island, Node1.Island);
+					bIsValid = false;
+				}
+			}
+		}
+	}
+
+	for (const FGraphNode& Node : Nodes)
+	{
+		if ((Node.Particle != nullptr) && FConstGenericParticleHandle(Node.Particle)->IsDynamic() && (FConstGenericParticleHandle(Node.Particle)->Island() != Node.Island))
+		{
+			UE_LOG(LogChaos, Error, TEXT("Particle %d island mismatch %d %d"), Node.Particle->ParticleID().LocalID, FConstGenericParticleHandle(Node.Particle)->Island(), Node.Island);
+			bIsValid = false;
+		}
+	}
+
 	return bIsValid;
 }
+#endif
diff --git a/Engine/Source/Runtime/Experimental/Chaos/Public/Chaos/PBDConstraintColor.h b/Engine/Source/Runtime/Experimental/Chaos/Public/Chaos/PBDConstraintColor.h
--- a/Engine/Source/Runtime/Experimental/Chaos/Public/Chaos/PBDConstraintColor.h
+++ b/Engine/Source/Runtime/Experimental/Chaos/Public/Chaos/PBDConstraintColor.h
@@ -59,6 +59,11 @@ namespace Chaos
 			bUseContactGraph = bInUseContactGraph;
 		}
 
+#if !UE_BUILD_SHIPPING
+		// For debugging onlue: run some (slow) validation checks on the graph and log any errors. Return false if there were errors.
+		bool DebugCheckGraph(const FPBDConstraintGraph& ConstraintGraph) const;
+#endif
+
 	private:
 		void ComputeContactGraph(const int32 Island, const FPBDConstraintGraph& ConstraintGraph, uint32 ContainerId);
 		void ComputeIslandColoring(const int32 Island, const FPBDConstraintGraph& ConstraintGraph, uint32 ContainerId);
diff --git a/Engine/Source/Runtime/Experimental/Chaos/Public/Chaos/PBDConstraintGraph.h b/Engine/Source/Runtime/Experimental/Chaos/Public/Chaos/PBDConstraintGraph.h
--- a/Engine/Source/Runtime/Experimental/Chaos/Public/Chaos/PBDConstraintGraph.h
+++ b/Engine/Source/Runtime/Experimental/Chaos/Public/Chaos/PBDConstraintGraph.h
@@ -197,6 +197,11 @@ namespace Chaos
 		 */
 		void DisableParticles(const TSet<FGeometryParticleHandle*>& Particles);
 
+#if !UE_BUILD_SHIPPING
+		// For debugging onlue: run some (slow) validation checks on the graph and log any errors. Return false if there were errors.
+		bool DebugCheckGraph() const;
+#endif
+
 	private:
 		struct FGraphNode
 		{
@@ -238,7 +243,6 @@ namespace Chaos
 
 		void ComputeIslands(const TParticleView<FPBDRigidParticles>& PBDRigids, FPBDRigidsSOAs& Particles);
 		bool ComputeIsland(const int32 Node, const int32 Island, TSet<FGeometryParticleHandle*>& ParticlesInIsland);
-		bool CheckIslands(const TArray<FGeometryParticleHandle*>& Particles);
 		
 		void ParticleAdd(FGeometryParticleHandle* AddedParticle);
 		void ParticleRemove(FGeometryParticleHandle* RemovedParticle);
