Chaos : GeoemtryCollection
- Fix for serialization crash.

diff --git a/Engine/Source/Runtime/Experimental/GeometryCollectionEngine/Private/GeometryCollection/GeometryCollectionObject.cpp b/Engine/Source/Runtime/Experimental/GeometryCollectionEngine/Private/GeometryCollection/GeometryCollectionObject.cpp
--- a/Engine/Source/Runtime/Experimental/GeometryCollectionEngine/Private/GeometryCollection/GeometryCollectionObject.cpp
+++ b/Engine/Source/Runtime/Experimental/GeometryCollectionEngine/Private/GeometryCollection/GeometryCollectionObject.cpp
@@ -612,6 +612,12 @@ void UGeometryCollection::Serialize(FArchive& Ar)
 		Super::Serialize(Ar);
 	}
 
+
+	if (!SizeSpecificData.Num())
+	{
+		SizeSpecificData.Add(GeometryCollectionSizeSpecificDataDefaults());
+	}
+
 	if (Ar.CustomVer(FDestructionObjectVersion::GUID) < FDestructionObjectVersion::DensityUnitsChanged)
 	{
 		if (bMassAsDensity)
