Chaos : GeoemtryCollection (Integ from Main CL:16450371)
- Transfer defaults to size specific data

diff --git a/Engine/Source/Runtime/Experimental/GeometryCollectionEngine/Private/GeometryCollection/GeometryCollectionObject.cpp b/Engine/Source/Runtime/Experimental/GeometryCollectionEngine/Private/GeometryCollection/GeometryCollectionObject.cpp
--- a/Engine/Source/Runtime/Experimental/GeometryCollectionEngine/Private/GeometryCollection/GeometryCollectionObject.cpp
+++ b/Engine/Source/Runtime/Experimental/GeometryCollectionEngine/Private/GeometryCollection/GeometryCollectionObject.cpp
@@ -183,7 +183,7 @@ void FillSharedSimulationSizeSpecificData(FSharedSimulationSizeSpecificData& ToD
 	ToData.DamageThreshold = FromData.DamageThreshold;
 }
 
-FGeometryCollectionSizeSpecificData UGeometryCollection::GeometryCollectionSizeSpecificDataDefaults() const
+FGeometryCollectionSizeSpecificData UGeometryCollection::GeometryCollectionSizeSpecificDataDefaults() 
 {
 	FGeometryCollectionSizeSpecificData Data;
 	check(Data.CollisionShapes.Num());
@@ -202,6 +202,43 @@ FGeometryCollectionSizeSpecificData UGeometryCollection::GeometryCollectionSizeS
 	return Data;
 }
 
+
+void UGeometryCollection::UpdateSizeSpecificDataDefaults()
+{
+	auto HasDefault = [](const TArray<FGeometryCollectionSizeSpecificData>& DatasIn)
+	{
+		for (const FGeometryCollectionSizeSpecificData& Data : DatasIn)
+		{
+			if (FMath::IsNearlyEqual(Data.MaxSize, FLT_MAX))
+			{
+				return true;
+			}
+		}
+		return false;
+	};
+
+	if (!SizeSpecificData.Num() || !HasDefault(SizeSpecificData))
+	{
+		FGeometryCollectionSizeSpecificData Data = GeometryCollectionSizeSpecificDataDefaults();
+#if WITH_EDITORONLY_DATA
+		Data.CollisionShapes[0].CollisionType = CollisionType_DEPRECATED;
+		Data.CollisionShapes[0].ImplicitType = ImplicitType_DEPRECATED;
+		Data.CollisionShapes[0].LevelSet.MinLevelSetResolution = MinLevelSetResolution_DEPRECATED;
+		Data.CollisionShapes[0].LevelSet.MaxLevelSetResolution = MaxLevelSetResolution_DEPRECATED;
+		Data.CollisionShapes[0].LevelSet.MinClusterLevelSetResolution = MinClusterLevelSetResolution_DEPRECATED;
+		Data.CollisionShapes[0].LevelSet.MaxClusterLevelSetResolution = MaxClusterLevelSetResolution_DEPRECATED;
+		Data.CollisionShapes[0].CollisionObjectReductionPercentage = CollisionObjectReductionPercentage_DEPRECATED;
+#endif
+		if (Data.CollisionShapes[0].ImplicitType == EImplicitTypeEnum::Chaos_Implicit_LevelSet)
+		{
+			Data.CollisionShapes[0].CollisionType = ECollisionTypeEnum::Chaos_Surface_Volumetric;
+		}
+		SizeSpecificData.Add(Data);
+	}
+	check(SizeSpecificData.Num() && SizeSpecificData[0].CollisionShapes.Num());
+}
+
+
 float KgCm3ToKgM3(float Density)
 {
 	return Density * 1000000;
@@ -615,7 +652,7 @@ void UGeometryCollection::Serialize(FArchive& Ar)
 
 	if (!SizeSpecificData.Num())
 	{
-		SizeSpecificData.Add(GeometryCollectionSizeSpecificDataDefaults());
+		UpdateSizeSpecificDataDefaults();
 	}
 
 	if (Ar.CustomVer(FDestructionObjectVersion::GUID) < FDestructionObjectVersion::DensityUnitsChanged)
@@ -703,34 +740,6 @@ void UGeometryCollection::Serialize(FArchive& Ar)
 #if WITH_EDITORONLY_DATA
 	if (Ar.IsLoading() && Ar.CustomVer(FUE5MainStreamObjectVersion::GUID) < FUE5MainStreamObjectVersion::GeometryCollectionUserDefinedCollisionShapes)
 	{
-		auto HasDefault = [](const TArray<FGeometryCollectionSizeSpecificData>& DatasIn)
-		{
-			for (const FGeometryCollectionSizeSpecificData& Data : DatasIn)
-			{
-				if (FMath::IsNearlyEqual(Data.MaxSize, FLT_MAX))
-				{
-					return true;
-				}
-			}
-			return false;
-		};
-
-		if (!SizeSpecificData.Num() || !HasDefault(SizeSpecificData))
-		{
-			FGeometryCollectionSizeSpecificData Data;
-			Data.CollisionShapes[0].CollisionType =	CollisionType_DEPRECATED;
-			Data.CollisionShapes[0].ImplicitType = ImplicitType_DEPRECATED;
-			Data.CollisionShapes[0].LevelSet.MinLevelSetResolution = MinLevelSetResolution_DEPRECATED;
-			Data.CollisionShapes[0].LevelSet.MaxLevelSetResolution = MaxLevelSetResolution_DEPRECATED;
-			Data.CollisionShapes[0].LevelSet.MinClusterLevelSetResolution = MinClusterLevelSetResolution_DEPRECATED;
-			Data.CollisionShapes[0].LevelSet.MaxClusterLevelSetResolution = MaxClusterLevelSetResolution_DEPRECATED;
-			Data.CollisionShapes[0].CollisionObjectReductionPercentage = CollisionObjectReductionPercentage_DEPRECATED;
-			if (Data.CollisionShapes[0].ImplicitType == EImplicitTypeEnum::Chaos_Implicit_LevelSet)
-				Data.CollisionShapes[0].CollisionType = ECollisionTypeEnum::Chaos_Surface_Volumetric;
-			SizeSpecificData.Add(Data);
-		}
-		check(SizeSpecificData.Num() && SizeSpecificData[0].CollisionShapes.Num());
-
 		InvalidateCollection();
 		bCreateSimulationData = true;
 	}
diff --git a/Engine/Source/Runtime/Experimental/GeometryCollectionEngine/Public/GeometryCollection/GeometryCollectionObject.h b/Engine/Source/Runtime/Experimental/GeometryCollectionEngine/Public/GeometryCollection/GeometryCollectionObject.h
--- a/Engine/Source/Runtime/Experimental/GeometryCollectionEngine/Public/GeometryCollection/GeometryCollectionObject.h
+++ b/Engine/Source/Runtime/Experimental/GeometryCollectionEngine/Public/GeometryCollection/GeometryCollectionObject.h
@@ -509,9 +509,9 @@ public:
 	TArray<FGeometryCollectionSizeSpecificData> SizeSpecificData;
 
 	int GetDefaultSizeSpecificDataIndex() const;
-	FGeometryCollectionSizeSpecificData GeometryCollectionSizeSpecificDataDefaults() const;
 	FGeometryCollectionSizeSpecificData& GetDefaultSizeSpecificData();
 	const FGeometryCollectionSizeSpecificData& GetDefaultSizeSpecificData() const;
+	static FGeometryCollectionSizeSpecificData GeometryCollectionSizeSpecificDataDefaults();
 
 	/**
 	* Enable remove pieces on fracture
@@ -541,6 +541,11 @@ private:
 	void CreateSimulationDataImp(bool bCopyFromDDC);
 #endif
 
+	/*
+	* Used to transfer deprecated properties to the size specific structures during serialization.
+	*/
+	void UpdateSizeSpecificDataDefaults();
+
 private:
 	/** Guid created on construction of this collection. It should be used to uniquely identify this collection */
 	UPROPERTY()
