GeometryCollection
-- implemented deferred initialization using the existing SimulatePhysics property and SetSimulatePhysics blueprint function.
-- Cleaned out old caching code

diff --git a/Engine/Config/BaseEngine.ini b/Engine/Config/BaseEngine.ini
--- a/Engine/Config/BaseEngine.ini
+++ b/Engine/Config/BaseEngine.ini
@@ -929,6 +929,7 @@ NearClipPlane=10.0
 +PackageRedirects=(OldName="/Script/PropertyAccess", Removed=true)
 +StructRedirects=(OldName="AnimInstanceSubsystemData",NewName="/Script/Engine.AnimSubsystemInstance")
 +ClassRedirects=(OldName="/Script/AnimationModifiers.AnimationBlueprintLibrary",NewName="/Script/AnimationBlueprintLibrary.AnimationBlueprintLibrary")
++PropertyRedirects=(OldName="GeometryCollectionComponent.Simulating", NewName="SimulatePhysics")
 
 [CoreUObject.Metadata]
 ; Note: UnrealHeaderTool should be rerun after making changes to MetadataRedirects to catch any keys specified in class headers
diff --git a/Engine/Source/Programs/HeadlessChaos/Private/GeometryCollection/GeometryCollectionTestFramework.cpp b/Engine/Source/Programs/HeadlessChaos/Private/GeometryCollection/GeometryCollectionTestFramework.cpp
--- a/Engine/Source/Programs/HeadlessChaos/Private/GeometryCollection/GeometryCollectionTestFramework.cpp
+++ b/Engine/Source/Programs/HeadlessChaos/Private/GeometryCollection/GeometryCollectionTestFramework.cpp
@@ -193,10 +193,7 @@ namespace GeometryCollectionTest
 				*DynamicCollection, // Game thread collection
 				SimulationParams,
 				SimFilterData,
-				QueryFilterData,
-				nullptr,			// Init func
-				nullptr,			// Cache sync func
-				nullptr);			// Final sync func
+				QueryFilterData);
 		return new FGeometryCollectionWrapper(RestCollection, DynamicCollection, PhysObject);
 	}
 
diff --git a/Engine/Source/Runtime/Experimental/Chaos/Private/PhysicsProxy/GeometryCollectionPhysicsProxy.cpp b/Engine/Source/Runtime/Experimental/Chaos/Private/PhysicsProxy/GeometryCollectionPhysicsProxy.cpp
--- a/Engine/Source/Runtime/Experimental/Chaos/Private/PhysicsProxy/GeometryCollectionPhysicsProxy.cpp
+++ b/Engine/Source/Runtime/Experimental/Chaos/Private/PhysicsProxy/GeometryCollectionPhysicsProxy.cpp
@@ -315,9 +315,6 @@ FGeometryCollectionPhysicsProxy::FGeometryCollectionPhysicsProxy(
 	const FSimulationParameters& SimulationParameters,
 	FCollisionFilterData InSimFilter,
 	FCollisionFilterData InQueryFilter,
-	FInitFunc InInitFunc,
-	FCacheSyncFunc InCacheSyncFunc,
-	FFinalSyncFunc InFinalSyncFunc,
 	const Chaos::EMultiBufferMode BufferMode)
 	: Base(InOwner)
 	, Parameters(SimulationParameters)
@@ -332,10 +329,6 @@ FGeometryCollectionPhysicsProxy::FGeometryCollectionPhysicsProxy(
 	, ProxySimDuration(0.0f)
 	, LastSyncCountGT(MAX_uint32)
 #endif
-	, InitFunc(InInitFunc)
-	, CacheSyncFunc(InCacheSyncFunc)
-	, FinalSyncFunc(InFinalSyncFunc)
-
 	, CollisionParticlesPerObjectFraction(CollisionParticlesPerObjectFractionDefault)
 
 	, GameThreadCollection(GameThreadCollectionIn)
@@ -378,13 +371,6 @@ void FGeometryCollectionPhysicsProxy::Initialize(Chaos::FPBDRigidsEvolutionBase
 	SolverClusterHandles.Init(nullptr, NumParticles);
 	SolverParticleHandles.Init(nullptr, NumParticles);
 
-	//	
-	//  Give clients the opportunity to update the parameters before the simualtion is setup. 
-	//
-	if (InitFunc)
-	{
-		InitFunc(Parameters);
-	}
 	// compatibility requirement to make sure we at least initialize GameThreadPerFrameData properly
 	GameThreadPerFrameData.SetWorldTransform(Parameters.WorldTransform);
 
@@ -1392,12 +1378,7 @@ void FGeometryCollectionPhysicsProxy::OnRemoveFromScene()
 
 void FGeometryCollectionPhysicsProxy::SyncBeforeDestroy()
 {
-	if(FinalSyncFunc)
-	{
-#if TODO_REIMPLEMENT_RIGID_CACHING
-		FinalSyncFunc(RecordedTracks);
-#endif
-	}
+
 }
 
 void FGeometryCollectionPhysicsProxy::BufferGameState() 
@@ -1762,10 +1743,6 @@ bool FGeometryCollectionPhysicsProxy::PullFromPhysicsState(const Chaos::FDirtyGe
 		//question: why do we need this? Sleeping objects will always have to update GPU
 		DynamicCollection.MakeDirty();
 
-		if (CacheSyncFunc)
-		{
-			CacheSyncFunc(TargetResults);
-		}
 	}
 
 	return true;
diff --git a/Engine/Source/Runtime/Experimental/Chaos/Public/PhysicsProxy/GeometryCollectionPhysicsProxy.h b/Engine/Source/Runtime/Experimental/Chaos/Public/PhysicsProxy/GeometryCollectionPhysicsProxy.h
--- a/Engine/Source/Runtime/Experimental/Chaos/Public/PhysicsProxy/GeometryCollectionPhysicsProxy.h
+++ b/Engine/Source/Runtime/Experimental/Chaos/Public/PhysicsProxy/GeometryCollectionPhysicsProxy.h
@@ -63,9 +63,6 @@ public:
 	typedef Chaos::TPBDRigidClusteredParticleHandle<Chaos::FReal, 3> FClusterHandle;
 
 	/** Proxy publics */
-	using FInitFunc = TFunction<void(FSimulationParameters&)>;
-	using FCacheSyncFunc = TFunction<void(const FGeometryCollectionResults&)>;
-	using FFinalSyncFunc = TFunction<void(const FRecordedTransformTrack&)>;
 	using IPhysicsProxyBase::GetSolver;
 
 	FGeometryCollectionPhysicsProxy() = delete;
@@ -82,9 +79,6 @@ public:
 		const FSimulationParameters& SimulationParameters,
 		FCollisionFilterData InSimFilter,
 		FCollisionFilterData InQueryFilter,
-		FInitFunc InInitFunc, 
-		FCacheSyncFunc InCacheSyncFunc, 
-		FFinalSyncFunc InFinalSyncFunc  ,
 		const Chaos::EMultiBufferMode BufferMode=Chaos::EMultiBufferMode::TripleGuarded);
 	virtual ~FGeometryCollectionPhysicsProxy();
 
@@ -241,6 +235,8 @@ public:
 		}
 	}
 
+	bool GetIsObjectDynamic() const { return IsObjectDynamic; }
+
 protected:
 	/**
 	 * Build a physics thread cluster parent particle.
@@ -320,12 +316,6 @@ private:
 	FRecordedTransformTrack RecordedTracks;
 #endif
 
-	// Functions to handle engine-side events
-	FInitFunc InitFunc;
-	FCacheSyncFunc CacheSyncFunc;
-	FFinalSyncFunc FinalSyncFunc;
-
-
 	// Per object collision fraction.
 	float CollisionParticlesPerObjectFraction;
 
diff --git a/Engine/Source/Runtime/Experimental/GeometryCollectionEngine/Private/GeometryCollection/GeometryCollectionComponent.cpp b/Engine/Source/Runtime/Experimental/GeometryCollectionEngine/Private/GeometryCollection/GeometryCollectionComponent.cpp
--- a/Engine/Source/Runtime/Experimental/GeometryCollectionEngine/Private/GeometryCollection/GeometryCollectionComponent.cpp
+++ b/Engine/Source/Runtime/Experimental/GeometryCollectionEngine/Private/GeometryCollection/GeometryCollectionComponent.cpp
@@ -216,7 +216,6 @@ FGeomComponentCacheParameters::FGeomComponentCacheParameters()
 UGeometryCollectionComponent::UGeometryCollectionComponent(const FObjectInitializer& ObjectInitializer)
 	: Super(ObjectInitializer)
 	, ChaosSolverActor(nullptr)
-	, Simulating(true)
 	, InitializationState(ESimulationInitializationState::Unintialized)
 	, ObjectType(EObjectStateTypeEnum::Chaos_Object_Dynamic)
 	, EnableClustering(true)
@@ -273,6 +272,9 @@ UGeometryCollectionComponent::UGeometryCollectionComponent(const FObjectInitiali
 	// By default use the destructible object channel unless the user specifies otherwise
 	BodyInstance.SetObjectType(ECC_Destructible);
 
+	// By default, we initialize immediately. If this is set false, we defer initialization.
+	BodyInstance.bSimulatePhysics = true;
+
 	EventDispatcher = ObjectInitializer.CreateDefaultSubobject<UChaosGameplayEventDispatcher>(this, TEXT("GameplayEventDispatcher"));
 
 	DynamicCollection = nullptr;
@@ -623,6 +625,22 @@ void UGeometryCollectionComponent::SetNotifyRigidBodyCollision(bool bNewNotifyRi
 	UpdateRBCollisionEventRegistration();
 }
 
+bool UGeometryCollectionComponent::CanEditSimulatePhysics()
+{
+	return true;
+}
+
+void UGeometryCollectionComponent::SetSimulatePhysics(bool bEnabled)
+{
+	Super::SetSimulatePhysics(bEnabled);
+
+	if (bEnabled && !PhysicsProxy)
+	{
+		RegisterAndInitializePhysicsProxy();
+	}
+
+}
+
 void UGeometryCollectionComponent::DispatchBreakEvent(const FChaosBreakEvent& Event)
 {
 	// native
@@ -1736,7 +1754,7 @@ void UGeometryCollectionComponent::OnCreatePhysicsState()
 */
 	// Skip the chain - don't care about body instance setup
 	UActorComponent::OnCreatePhysicsState();
-	if (!Simulating) IsObjectLoading = false; // just mark as loaded if we are simulating.
+	if (!BodyInstance.bSimulatePhysics) IsObjectLoading = false; // just mark as loaded if we are simulating.
 
 /*#if WITH_PHYSX
 	DummyBodyInstance.SetCollisionEnabled(ECollisionEnabled::QueryOnly);
@@ -1767,210 +1785,6 @@ void UGeometryCollectionComponent::OnCreatePhysicsState()
 		{
 			FPhysxUserData::Set<UPrimitiveComponent>(&PhysicsUserData, this);
 
-			FSimulationParameters SimulationParameters;
-			{
-#if !(UE_BUILD_SHIPPING || UE_BUILD_TEST)
-				SimulationParameters.Name = GetPathName();
-#endif
-				EClusterConnectionTypeEnum ClusterCollectionType = ClusterConnectionType_DEPRECATED;
-				if (RestCollection)
-				{
-					RestCollection->GetSharedSimulationParams(SimulationParameters.Shared);
-					SimulationParameters.RestCollection = RestCollection->GetGeometryCollection().Get();
-					ClusterCollectionType = RestCollection->ClusterConnectionType;
-				}
-				SimulationParameters.Simulating = Simulating;
-				SimulationParameters.EnableClustering = EnableClustering;
-				SimulationParameters.ClusterGroupIndex = EnableClustering ? ClusterGroupIndex : 0;
-				SimulationParameters.MaxClusterLevel = MaxClusterLevel;
-				SimulationParameters.bUseSizeSpecificDamageThresholds = bUseSizeSpecificDamageThreshold;
-				SimulationParameters.DamageThreshold = DamageThreshold;
-				SimulationParameters.ClusterConnectionMethod = (Chaos::FClusterCreationParameters::EConnectionMethod)ClusterCollectionType;
-				SimulationParameters.CollisionGroup = CollisionGroup;
-				SimulationParameters.CollisionSampleFraction = CollisionSampleFraction;
-				SimulationParameters.InitialVelocityType = InitialVelocityType;
-				SimulationParameters.InitialLinearVelocity = InitialLinearVelocity;
-				SimulationParameters.InitialAngularVelocity = InitialAngularVelocity;
-				SimulationParameters.bClearCache = true;
-				SimulationParameters.ObjectType = ObjectType;
-				SimulationParameters.CacheType = CacheParameters.CacheMode;
-				SimulationParameters.ReverseCacheBeginTime = CacheParameters.ReverseCacheBeginTime;
-				SimulationParameters.CollisionData.SaveCollisionData = CacheParameters.SaveCollisionData;
-				SimulationParameters.CollisionData.DoGenerateCollisionData = CacheParameters.DoGenerateCollisionData;
-				SimulationParameters.CollisionData.CollisionDataSizeMax = CacheParameters.CollisionDataSizeMax;
-				SimulationParameters.CollisionData.DoCollisionDataSpatialHash = CacheParameters.DoCollisionDataSpatialHash;
-				SimulationParameters.CollisionData.CollisionDataSpatialHashRadius = CacheParameters.CollisionDataSpatialHashRadius;
-				SimulationParameters.CollisionData.MaxCollisionPerCell = CacheParameters.MaxCollisionPerCell;
-				SimulationParameters.BreakingData.SaveBreakingData = CacheParameters.SaveBreakingData;
-				SimulationParameters.BreakingData.DoGenerateBreakingData = CacheParameters.DoGenerateBreakingData;
-				SimulationParameters.BreakingData.BreakingDataSizeMax = CacheParameters.BreakingDataSizeMax;
-				SimulationParameters.BreakingData.DoBreakingDataSpatialHash = CacheParameters.DoBreakingDataSpatialHash;
-				SimulationParameters.BreakingData.BreakingDataSpatialHashRadius = CacheParameters.BreakingDataSpatialHashRadius;
-				SimulationParameters.BreakingData.MaxBreakingPerCell = CacheParameters.MaxBreakingPerCell;
-				SimulationParameters.TrailingData.SaveTrailingData = CacheParameters.SaveTrailingData;
-				SimulationParameters.TrailingData.DoGenerateTrailingData = CacheParameters.DoGenerateTrailingData;
-				SimulationParameters.TrailingData.TrailingDataSizeMax = CacheParameters.TrailingDataSizeMax;
-				SimulationParameters.TrailingData.TrailingMinSpeedThreshold = CacheParameters.TrailingMinSpeedThreshold;
-				SimulationParameters.TrailingData.TrailingMinVolumeThreshold = CacheParameters.TrailingMinVolumeThreshold;
-				SimulationParameters.RemoveOnFractureEnabled = SimulationParameters.Shared.RemoveOnFractureIndices.Num() > 0;
-				SimulationParameters.WorldTransform = GetComponentToWorld();
-				SimulationParameters.UserData = static_cast<void*>(&PhysicsUserData);
-
-				UPhysicalMaterial* EnginePhysicalMaterial = GetPhysicalMaterial();
-				if(ensure(EnginePhysicalMaterial))
-				{
-					SimulationParameters.PhysicalMaterialHandle = EnginePhysicalMaterial->GetPhysicsMaterial();
-				}
-			}
-
-
-			//
-			// Called from FGeometryCollectionPhysicsProxy::Initialize()
-			//
-			auto InitFunc = [this](FSimulationParameters& InParams)
-			{
-#if !(UE_BUILD_SHIPPING || UE_BUILD_TEST)
-				InParams.Name = GetPathName();
-#endif
-				GetInitializationCommands(InParams.InitializationCommands);
-				UGeometryCollectionCache* Cache = CacheParameters.TargetCache;
-				if(Cache && CacheParameters.CacheMode != EGeometryCollectionCacheType::None)
-				{
-					bool bCacheValid = false;
-
-					switch(CacheParameters.CacheMode)
-					{
-					case EGeometryCollectionCacheType::Record:
-					case EGeometryCollectionCacheType::RecordAndPlay:
-						bCacheValid = Cache->CompatibleWithForRecord(RestCollection);
-						break;
-					case EGeometryCollectionCacheType::Play:
-						bCacheValid = Cache->CompatibleWithForPlayback(RestCollection);
-						break;
-					default:
-						check(false);
-						break;
-					}
-
-					if(bCacheValid)
-					{
-						InParams.RecordedTrack = (InParams.IsCachePlaying() && CacheParameters.TargetCache) ? CacheParameters.TargetCache->GetData() : nullptr;
-					}
-					else
-					{
-						// We attempted to utilize a cache that was not compatible with this component. Report and do not use
-						UE_LOG(LogChaos, Error, TEXT("Geometry collection '%s' attempted to use cache '%s' but it is not compatible."), *GetName(), *(CacheParameters.TargetCache->GetName()));
-						InParams.RecordedTrack = nullptr;
-						InParams.CacheType = EGeometryCollectionCacheType::None;
-						CacheParameters.CacheMode = EGeometryCollectionCacheType::None;
-
-#if !(UE_BUILD_SHIPPING || UE_BUILD_TEST)
-						FMessageLog("PIE").Error()
-							->AddToken(FTextToken::Create(NSLOCTEXT("GeomCollectionComponent", "BadCache_01", "Geometry collection")))
-							->AddToken(FUObjectToken::Create(this))
-							->AddToken(FTextToken::Create(NSLOCTEXT("GeomCollectionComponent", "BadCache_02", "attempted to use cache")))
-							->AddToken(FUObjectToken::Create(CacheParameters.TargetCache))
-							->AddToken(FTextToken::Create(NSLOCTEXT("GeomCollectionComponent", "BadCache_03", "but it is not compatible")));
-#endif // !(UE_BUILD_SHIPPING || UE_BUILD_TEST)
-
-					}
-				}
-
-			};
-
-			auto CacheSyncFunc = [this](const FGeometryCollectionResults& Results)
-			{
-#if GEOMETRYCOLLECTION_DEBUG_DRAW
-				const bool bHasNumParticlesChanged = (NumParticlesAdded != Results.NumParticlesAdded);  // Needs to be evaluated before NumParticlesAdded gets updated
-#endif  // #if GEOMETRYCOLLECTION_DEBUG_DRAW
-				//RigidBodyIds.Init(Results.RigidBodyIds);
-				BaseRigidBodyIndex = Results.BaseIndex;
-				NumParticlesAdded = Results.NumParticlesAdded;
-				DisabledFlags = Results.DisabledStates;
-			
-				if (!IsObjectDynamic && Results.IsObjectDynamic)
-				{
-					IsObjectDynamic = Results.IsObjectDynamic;
-
-					NotifyGeometryCollectionPhysicsStateChange.Broadcast(this);
-					SwitchRenderModels(GetOwner());
-				}
-
-				if (IsObjectLoading && !Results.IsObjectLoading)
-				{
-					IsObjectLoading = Results.IsObjectLoading;
-
-					NotifyGeometryCollectionPhysicsLoadingStateChange.Broadcast(this);
-				}
-
-
-				WorldBounds = Results.WorldBounds;
-
-				// Update replication data for clients if necessary
-				UpdateRepData();
-
-#if GEOMETRYCOLLECTION_DEBUG_DRAW
-				// Notify debug draw componentUGeometryCollectionDebugDrawComponent of particle changes
-				if (bHasNumParticlesChanged)
-				{
-					if (const AGeometryCollectionActor* const Owner = Cast<AGeometryCollectionActor>(GetOwner()))
-					{
-						if (UGeometryCollectionDebugDrawComponent* const GeometryCollectionDebugDrawComponent = Owner->GetGeometryCollectionDebugDrawComponent())
-						{
-							GeometryCollectionDebugDrawComponent->OnClusterChanged();
-						}
-					}
-				}
-#endif  // #if GEOMETRYCOLLECTION_DEBUG_DRAW
-			};
-
-			auto FinalSyncFunc = [this](const FRecordedTransformTrack& InTrack)
-			{
-#if WITH_EDITOR && WITH_EDITORONLY_DATA
-				if (CacheParameters.CacheMode == EGeometryCollectionCacheType::Record && InTrack.Records.Num() > 0)
-				{
-					Modify();
-					if (!CacheParameters.TargetCache)
-					{
-						CacheParameters.TargetCache = UGeometryCollectionCache::CreateCacheForCollection(RestCollection);
-					}
-
-					if (CacheParameters.TargetCache)
-					{
-						// Queue this up to be dirtied after PIE ends
-						FPhysScene_Chaos* Scene = GetInnerChaosScene();
-
-						CacheParameters.TargetCache->PreEditChange(nullptr);
-						CacheParameters.TargetCache->Modify();
-						CacheParameters.TargetCache->SetFromRawTrack(InTrack);
-						CacheParameters.TargetCache->PostEditChange();
-
-						Scene->AddPieModifiedObject(CacheParameters.TargetCache);
-
-						if (EditorActor)
-						{
-							UGeometryCollectionComponent* EditorComponent = Cast<UGeometryCollectionComponent>(EditorUtilities::FindMatchingComponentInstance(this, EditorActor));
-
-							if (EditorComponent)
-							{
-								EditorComponent->PreEditChange(FindFProperty<FProperty>(EditorComponent->GetClass(), GET_MEMBER_NAME_CHECKED(UGeometryCollectionComponent, CacheParameters)));
-								EditorComponent->Modify();
-
-								EditorComponent->CacheParameters.TargetCache = CacheParameters.TargetCache;
-
-								EditorComponent->PostEditChange();
-
-								Scene->AddPieModifiedObject(EditorComponent);
-								Scene->AddPieModifiedObject(EditorActor);
-							}
-
-							EditorActor = nullptr;
-						}
-					}
-				}
-#endif
-			};
-
 			// If the Component is set to Dynamic, we look to the RestCollection for initial dynamic state override per transform.
 			TManagedArray<int32> & DynamicState = DynamicCollection->DynamicState;
 
@@ -1997,7 +1811,7 @@ void UGeometryCollectionComponent::OnCreatePhysicsState()
 			{
 				for (int i = 0; i < Active.Num(); i++)
 				{
-					Active[i] = Simulating;
+					Active[i] = true;
 				}
 			}
 			TManagedArray<int32> & CollisionGroupArray = DynamicCollection->CollisionGroup;
@@ -2030,11 +1844,11 @@ void UGeometryCollectionComponent::OnCreatePhysicsState()
 				InitialSimFilter.Word3 |= EPDF_ContactNotify;
 			}
 
- 			PhysicsProxy = new FGeometryCollectionPhysicsProxy(this, *DynamicCollection, SimulationParameters, InitialQueryFilter, InitialSimFilter, InitFunc, CacheSyncFunc, FinalSyncFunc);
-			FPhysScene_Chaos* Scene = GetInnerChaosScene();
-			Scene->AddObject(this, PhysicsProxy);
-
-			RegisterForEvents();
+			if (BodyInstance.bSimulatePhysics)
+			{
+				RegisterAndInitializePhysicsProxy();
+			}
+			
 		}
 	}
 
@@ -2047,6 +1861,71 @@ void UGeometryCollectionComponent::OnCreatePhysicsState()
 #endif // WITH_CHAOS
 }
 
+void UGeometryCollectionComponent::RegisterAndInitializePhysicsProxy()
+{
+	FSimulationParameters SimulationParameters;
+	{
+#if !(UE_BUILD_SHIPPING || UE_BUILD_TEST)
+		SimulationParameters.Name = GetPathName();
+#endif
+		EClusterConnectionTypeEnum ClusterCollectionType = ClusterConnectionType_DEPRECATED;
+		if (RestCollection)
+		{
+			RestCollection->GetSharedSimulationParams(SimulationParameters.Shared);
+			SimulationParameters.RestCollection = RestCollection->GetGeometryCollection().Get();
+			ClusterCollectionType = RestCollection->ClusterConnectionType;
+		}
+		SimulationParameters.Simulating = BodyInstance.bSimulatePhysics;
+		SimulationParameters.EnableClustering = EnableClustering;
+		SimulationParameters.ClusterGroupIndex = EnableClustering ? ClusterGroupIndex : 0;
+		SimulationParameters.MaxClusterLevel = MaxClusterLevel;
+		SimulationParameters.bUseSizeSpecificDamageThresholds = bUseSizeSpecificDamageThreshold;
+		SimulationParameters.DamageThreshold = DamageThreshold;
+		SimulationParameters.ClusterConnectionMethod = (Chaos::FClusterCreationParameters::EConnectionMethod)ClusterCollectionType;
+		SimulationParameters.CollisionGroup = CollisionGroup;
+		SimulationParameters.CollisionSampleFraction = CollisionSampleFraction;
+		SimulationParameters.InitialVelocityType = InitialVelocityType;
+		SimulationParameters.InitialLinearVelocity = InitialLinearVelocity;
+		SimulationParameters.InitialAngularVelocity = InitialAngularVelocity;
+		SimulationParameters.bClearCache = true;
+		SimulationParameters.ObjectType = ObjectType;
+		SimulationParameters.CacheType = CacheParameters.CacheMode;
+		SimulationParameters.ReverseCacheBeginTime = CacheParameters.ReverseCacheBeginTime;
+		SimulationParameters.CollisionData.SaveCollisionData = CacheParameters.SaveCollisionData;
+		SimulationParameters.CollisionData.DoGenerateCollisionData = CacheParameters.DoGenerateCollisionData;
+		SimulationParameters.CollisionData.CollisionDataSizeMax = CacheParameters.CollisionDataSizeMax;
+		SimulationParameters.CollisionData.DoCollisionDataSpatialHash = CacheParameters.DoCollisionDataSpatialHash;
+		SimulationParameters.CollisionData.CollisionDataSpatialHashRadius = CacheParameters.CollisionDataSpatialHashRadius;
+		SimulationParameters.CollisionData.MaxCollisionPerCell = CacheParameters.MaxCollisionPerCell;
+		SimulationParameters.BreakingData.SaveBreakingData = CacheParameters.SaveBreakingData;
+		SimulationParameters.BreakingData.DoGenerateBreakingData = CacheParameters.DoGenerateBreakingData;
+		SimulationParameters.BreakingData.BreakingDataSizeMax = CacheParameters.BreakingDataSizeMax;
+		SimulationParameters.BreakingData.DoBreakingDataSpatialHash = CacheParameters.DoBreakingDataSpatialHash;
+		SimulationParameters.BreakingData.BreakingDataSpatialHashRadius = CacheParameters.BreakingDataSpatialHashRadius;
+		SimulationParameters.BreakingData.MaxBreakingPerCell = CacheParameters.MaxBreakingPerCell;
+		SimulationParameters.TrailingData.SaveTrailingData = CacheParameters.SaveTrailingData;
+		SimulationParameters.TrailingData.DoGenerateTrailingData = CacheParameters.DoGenerateTrailingData;
+		SimulationParameters.TrailingData.TrailingDataSizeMax = CacheParameters.TrailingDataSizeMax;
+		SimulationParameters.TrailingData.TrailingMinSpeedThreshold = CacheParameters.TrailingMinSpeedThreshold;
+		SimulationParameters.TrailingData.TrailingMinVolumeThreshold = CacheParameters.TrailingMinVolumeThreshold;
+		SimulationParameters.RemoveOnFractureEnabled = SimulationParameters.Shared.RemoveOnFractureIndices.Num() > 0;
+		SimulationParameters.WorldTransform = GetComponentToWorld();
+		SimulationParameters.UserData = static_cast<void*>(&PhysicsUserData);
+
+		UPhysicalMaterial* EnginePhysicalMaterial = GetPhysicalMaterial();
+		if (ensure(EnginePhysicalMaterial))
+		{
+			SimulationParameters.PhysicalMaterialHandle = EnginePhysicalMaterial->GetPhysicsMaterial();
+		}
+	}
+
+	PhysicsProxy = new FGeometryCollectionPhysicsProxy(this, *DynamicCollection, SimulationParameters, InitialQueryFilter, InitialSimFilter);
+	FPhysScene_Chaos* Scene = GetInnerChaosScene();
+	Scene->AddObject(this, PhysicsProxy);
+
+	RegisterForEvents();
+}
+
 void UGeometryCollectionComponent::OnDestroyPhysicsState()
 {
 	UActorComponent::OnDestroyPhysicsState();
@@ -2591,6 +2470,10 @@ void UGeometryCollectionComponent::ApplyPhysicsField(bool Enabled, EGeometryColl
 	}
 }
 
+bool UGeometryCollectionComponent::GetIsObjectDynamic() const
+{ 
+	return PhysicsProxy ? PhysicsProxy->GetIsObjectDynamic() : IsObjectDynamic;
+}
 
 void UGeometryCollectionComponent::DispatchFieldCommand(const FFieldSystemCommand& InCommand)
 {
diff --git a/Engine/Source/Runtime/Experimental/GeometryCollectionEngine/Public/GeometryCollection/GeometryCollectionComponent.h b/Engine/Source/Runtime/Experimental/GeometryCollectionEngine/Public/GeometryCollection/GeometryCollectionComponent.h
--- a/Engine/Source/Runtime/Experimental/GeometryCollectionEngine/Public/GeometryCollection/GeometryCollectionComponent.h
+++ b/Engine/Source/Runtime/Experimental/GeometryCollectionEngine/Public/GeometryCollection/GeometryCollectionComponent.h
@@ -363,6 +363,8 @@ public:
 	virtual void OnRegister() override;
 	virtual FBodyInstance* GetBodyInstance(FName BoneName = NAME_None, bool bGetWelded = true, int32 Index = -1) const override;
 	virtual void SetNotifyRigidBodyCollision(bool bNewNotifyRigidBodyCollision) override;
+	virtual bool CanEditSimulatePhysics() override;
+	virtual void SetSimulatePhysics(bool bEnabled) override;
 	//~ End UPrimitiveComponent Interface.
 
 
@@ -436,11 +438,9 @@ public:
 	UPROPERTY(EditAnywhere, NoClear, BlueprintReadOnly, Category = "ChaosPhysics")
 	TArray<TObjectPtr<const AFieldSystemActor>> InitializationFields;
 
-	/**
-	* When Simulating is enabled the Component will initialize its rigid bodies within the solver.
-	*/
-	UPROPERTY(EditAnywhere, BlueprintReadWrite, Category = "ChaosPhysics|General")
-	bool Simulating;
+	UPROPERTY(meta = (DeprecatedProperty, DeprecationMessage = "GeometryCollection now abides the bSimulatePhysics flag from the base class."))
+	bool Simulating_DEPRECATED;
+
 	ESimulationInitializationState InitializationState;
 
 	/** ObjectType defines how to initialize the rigid objects state, Kinematic, Sleeping, Dynamic. */
@@ -549,7 +549,7 @@ public:
 	UPROPERTY(BlueprintAssignable, Category = "Game|Damage")
 	FNotifyGeometryCollectionPhysicsStateChange NotifyGeometryCollectionPhysicsStateChange;
 
-	bool GetIsObjectDynamic() { return IsObjectDynamic; }
+	bool GetIsObjectDynamic() const;
 
 	DECLARE_DYNAMIC_MULTICAST_DELEGATE_OneParam(FNotifyGeometryCollectionPhysicsLoadingStateChange, UGeometryCollectionComponent*, FracturedComponent);
 	UPROPERTY(BlueprintAssignable, Category = "Game|Loading")
@@ -596,6 +596,7 @@ public:
 	const TArray<bool>& GetDisabledFlags() const { return DisabledFlags; }
 
 	virtual void OnCreatePhysicsState() override;
+	void RegisterAndInitializePhysicsProxy();
 	virtual void OnDestroyPhysicsState() override;
 	virtual bool ShouldCreatePhysicsState() const override;
 	virtual bool HasValidPhysicsState() const override;
