Chaos : GeoemtryCollection (Integ from Main CL:16450371)
- Force Size Specific data to always have a single element.

diff --git a/Engine/Source/Runtime/Experimental/GeometryCollectionEngine/Private/GeometryCollection/GeometryCollectionObject.cpp b/Engine/Source/Runtime/Experimental/GeometryCollectionEngine/Private/GeometryCollection/GeometryCollectionObject.cpp
--- a/Engine/Source/Runtime/Experimental/GeometryCollectionEngine/Private/GeometryCollection/GeometryCollectionObject.cpp
+++ b/Engine/Source/Runtime/Experimental/GeometryCollectionEngine/Private/GeometryCollection/GeometryCollectionObject.cpp
@@ -203,7 +203,7 @@ FGeometryCollectionSizeSpecificData UGeometryCollection::GeometryCollectionSizeS
 }
 
 
-void UGeometryCollection::UpdateSizeSpecificDataDefaults()
+void UGeometryCollection::ValidateSizeSpecificDataDefaults()
 {
 	auto HasDefault = [](const TArray<FGeometryCollectionSizeSpecificData>& DatasIn)
 	{
@@ -652,7 +652,7 @@ void UGeometryCollection::Serialize(FArchive& Ar)
 
 	if (!SizeSpecificData.Num())
 	{
-		UpdateSizeSpecificDataDefaults();
+		ValidateSizeSpecificDataDefaults();
 	}
 
 	if (Ar.CustomVer(FDestructionObjectVersion::GUID) < FDestructionObjectVersion::DensityUnitsChanged)
@@ -1084,15 +1084,27 @@ void UGeometryCollection::PostEditChangeProperty(struct FPropertyChangedEvent& P
 {
 	if (PropertyChangedEvent.Property)
 	{
+		bool bRebuildSimulationData = false;
+
 		if (PropertyChangedEvent.Property->GetFName() == GET_MEMBER_NAME_CHECKED(UGeometryCollection, EnableNanite))
 		{
 			InvalidateCollection();
 			EnsureDataIsCooked();
 		}
+		else if (PropertyChangedEvent.Property->GetFName() == GET_MEMBER_NAME_CHECKED(UGeometryCollection, SizeSpecificData))
+		{
+			ValidateSizeSpecificDataDefaults();
+			InvalidateCollection();
+			bRebuildSimulationData = true;
+		}
 		else if (PropertyChangedEvent.Property->GetFName() != GET_MEMBER_NAME_CHECKED(UGeometryCollection, Materials))
 		{
 			InvalidateCollection();
-			
+			bRebuildSimulationData = true;
+		}
+
+		if (bRebuildSimulationData)
+		{
 			if (!bManualDataCreate)
 			{
 				CreateSimulationData();
@@ -1109,6 +1121,7 @@ bool UGeometryCollection::Modify(bool bAlwaysMarkDirty /*= true*/)
 	if (Package->IsDirty())
 	{
 		InvalidateCollection();
+		ValidateSizeSpecificDataDefaults();
 	}
 
 	return bSuperResult;
diff --git a/Engine/Source/Runtime/Experimental/GeometryCollectionEngine/Public/GeometryCollection/GeometryCollectionObject.h b/Engine/Source/Runtime/Experimental/GeometryCollectionEngine/Public/GeometryCollection/GeometryCollectionObject.h
--- a/Engine/Source/Runtime/Experimental/GeometryCollectionEngine/Public/GeometryCollection/GeometryCollectionObject.h
+++ b/Engine/Source/Runtime/Experimental/GeometryCollectionEngine/Public/GeometryCollection/GeometryCollectionObject.h
@@ -542,9 +542,11 @@ private:
 #endif
 
 	/*
-	* Used to transfer deprecated properties to the size specific structures during serialization.
+	* Used to transfer deprecated properties to the size specific structures during serialization
+	* and to add back the default size specific data when deleted. 
 	*/
-	void UpdateSizeSpecificDataDefaults();
+	void ValidateSizeSpecificDataDefaults();
+
 
 private:
 	/** Guid created on construction of this collection. It should be used to uniquely identify this collection */
