LWC: Fix chaos precision loss warnings

diff --git a/Engine/Source/Runtime/Core/Public/Math/GenericOctree.h b/Engine/Source/Runtime/Core/Public/Math/GenericOctree.h
--- a/Engine/Source/Runtime/Core/Public/Math/GenericOctree.h
+++ b/Engine/Source/Runtime/Core/Public/Math/GenericOctree.h
@@ -478,7 +478,7 @@ private:
 
 	TArray<FNodeIndex> FreeList;
 	/** The extent of a leaf at the maximum allowed depth of the tree. */
-	float MinLeafExtent;
+	FVector::FReal MinLeafExtent;
 
 	FNodeIndex AllocateEightNodes()
 	{
@@ -989,7 +989,7 @@ public:
 	}
 
 	/** Initialization constructor. */
-	TOctree2(const FVector& InOrigin,float InExtent)
+	TOctree2(const FVector& InOrigin,FVector::FReal InExtent)
 		: RootNodeContext(FBoxCenterAndExtent(InOrigin, FVector(InExtent, InExtent, InExtent)), 0, 0)
 		, MinLeafExtent(InExtent* FMath::Pow((1.0f + 1.0f / (float)FOctreeNodeContext::LoosenessDenominator) / 2.0f, OctreeSemantics::MaxNodeDepth))
 	{
diff --git a/Engine/Source/Runtime/Experimental/Chaos/Chaos.Build.cs b/Engine/Source/Runtime/Experimental/Chaos/Chaos.Build.cs
--- a/Engine/Source/Runtime/Experimental/Chaos/Chaos.Build.cs
+++ b/Engine/Source/Runtime/Experimental/Chaos/Chaos.Build.cs
@@ -43,10 +43,7 @@ namespace UnrealBuildTool.Rules
 				PublicDefinitions.Add("CHAOS_MEMORY_TRACKING=0");
 			}
 			
-			if (!Target.GlobalDefinitions.Contains("UE_LARGE_WORLD_COORDINATES_DISABLED=0")) // LWC_TODO: Temporarily disabled for LWC while these are properly fixed
-			{
-				UnsafeTypeCastWarningLevel = WarningLevel.Warning;
-			}
+			UnsafeTypeCastWarningLevel = WarningLevel.Error;
 		}
     }
 }
diff --git a/Engine/Source/Runtime/Experimental/Chaos/Private/Chaos/ChaosDebugDraw.cpp b/Engine/Source/Runtime/Experimental/Chaos/Private/Chaos/ChaosDebugDraw.cpp
--- a/Engine/Source/Runtime/Experimental/Chaos/Private/Chaos/ChaosDebugDraw.cpp
+++ b/Engine/Source/Runtime/Experimental/Chaos/Private/Chaos/ChaosDebugDraw.cpp
@@ -127,7 +127,7 @@ namespace Chaos
 						// Core shape and lines connecting core to outer
 						if (Margin > 0.0f)
 						{
-							const FReal LineThickness = 0.5f * Settings.ShapeThicknesScale * Settings.LineThickness;
+							const FRealSingle LineThickness = 0.5f * Settings.ShapeThicknesScale * Settings.LineThickness;
 							const FVec3 InnerP0 = ShapeTransform.TransformPositionNoScale(Shape->GetMarginAdjustedVertexScaled(VertexIndex0, Margin, ShapeTransform.GetScale3D()));
 							const FVec3 InnerP1 = ShapeTransform.TransformPositionNoScale(Shape->GetMarginAdjustedVertexScaled(VertexIndex1, Margin, ShapeTransform.GetScale3D()));
 							FDebugDrawQueue::GetInstance().DrawDebugLine(InnerP0, InnerP1, FColor::Blue, false, -1.f, Settings.DrawPriority, LineThickness);
@@ -550,7 +550,7 @@ namespace Chaos
 			FDebugDrawQueue::GetInstance().DrawDebugBox(P, 0.5f * Box.Extents(), Q, Color, false, KINDA_SMALL_NUMBER, Settings.DrawPriority, Settings.LineThickness);
 		}
 
-		void DrawParticleTransformImpl(const FRigidTransform3& SpaceTransform, const FGeometryParticleHandle* InParticle, int32 Index, FReal ColorScale, const FChaosDebugDrawSettings& Settings)
+		void DrawParticleTransformImpl(const FRigidTransform3& SpaceTransform, const FGeometryParticleHandle* InParticle, int32 Index, FRealSingle ColorScale, const FChaosDebugDrawSettings& Settings)
 		{
 			FColor Red = (ColorScale * FColor::Red).ToFColor(false);
 			FColor Green = (ColorScale * FColor::Green).ToFColor(false);
@@ -742,7 +742,7 @@ namespace Chaos
 			DrawCollisionImpl(SpaceTransform, ConstraintHandle->GetContact(), ColorScale, Settings);
 		}
 
-		void DrawCollidingShapesImpl(const FRigidTransform3& SpaceTransform, const FPBDCollisionConstraints& Collisions, float ColorScale, const FChaosDebugDrawSettings& Settings)
+		void DrawCollidingShapesImpl(const FRigidTransform3& SpaceTransform, const FPBDCollisionConstraints& Collisions, FRealSingle ColorScale, const FChaosDebugDrawSettings& Settings)
 		{
 			TArray<const FImplicitObject*> Shapes;
 			TArray<FConstGenericParticleHandle> ShapeParticles;
@@ -793,7 +793,7 @@ namespace Chaos
 			}
 		}
 
-		void DrawJointConstraintImpl(const FRigidTransform3& SpaceTransform, const FVec3& InPa, const FVec3& InCa, const FVec3& InXa, const FMatrix33& Ra, const FVec3& InPb, const FVec3& InCb, const FVec3& InXb, const FMatrix33& Rb, int32 IslandIndex, int32 LevelIndex, int32 ColorIndex, int32 BatchIndex, int32 Index, FReal ColorScale, const FChaosDebugDrawJointFeatures& FeatureMask, const FChaosDebugDrawSettings& Settings)
+		void DrawJointConstraintImpl(const FRigidTransform3& SpaceTransform, const FVec3& InPa, const FVec3& InCa, const FVec3& InXa, const FMatrix33& Ra, const FVec3& InPb, const FVec3& InCb, const FVec3& InXb, const FMatrix33& Rb, int32 IslandIndex, int32 LevelIndex, int32 ColorIndex, int32 BatchIndex, int32 Index, Chaos::FRealSingle ColorScale, const FChaosDebugDrawJointFeatures& FeatureMask, const FChaosDebugDrawSettings& Settings)
 		{
 			using namespace Chaos::DebugDraw;
 			FColor R = (ColorScale * FColor::Red).ToFColor(false);
@@ -811,7 +811,7 @@ namespace Chaos
 
 			if (FeatureMask.bActorConnector)
 			{
-				const FReal ConnectorThickness = 1.5f * Settings.LineThickness;
+				const FRealSingle ConnectorThickness = 1.5f * Settings.LineThickness;
 				const FReal CoMSize = Settings.DrawScale * Settings.JointComSize;
 				// Leave a gap around the actor position so we can see where the center is
 				FVec3 Sa = Pa;
@@ -833,7 +833,7 @@ namespace Chaos
 			}
 			if (FeatureMask.bCoMConnector)
 			{
-				const FReal ConnectorThickness = 1.5f * Settings.LineThickness;
+				const FRealSingle ConnectorThickness = 1.5f * Settings.LineThickness;
 				const FReal CoMSize = Settings.DrawScale * Settings.JointComSize;
 				// Leave a gap around the body position so we can see where the center is
 				FVec3 Sa = Ca;
@@ -855,7 +855,7 @@ namespace Chaos
 			}
 			if (FeatureMask.bStretch)
 			{
-				const FReal StretchThickness = 3.0f * Settings.LineThickness;
+				const FRealSingle StretchThickness = 3.0f * Settings.LineThickness;
 				FDebugDrawQueue::GetInstance().DrawDebugLine(Xa, Xb, M, false, KINDA_SMALL_NUMBER, Settings.DrawPriority, StretchThickness);
 			}
 			if (FeatureMask.bAxes)
@@ -895,7 +895,7 @@ namespace Chaos
 			}
 		}
 
-		void DrawJointConstraintImpl(const FRigidTransform3& SpaceTransform, const FPBDJointConstraintHandle* ConstraintHandle, FReal ColorScale, const FChaosDebugDrawJointFeatures& FeatureMask, const FChaosDebugDrawSettings& Settings)
+		void DrawJointConstraintImpl(const FRigidTransform3& SpaceTransform, const FPBDJointConstraintHandle* ConstraintHandle, Chaos::FRealSingle ColorScale, const FChaosDebugDrawJointFeatures& FeatureMask, const FChaosDebugDrawSettings& Settings)
 		{
 			TVec2<FGeometryParticleHandle*> ConstrainedParticles = ConstraintHandle->GetConstrainedParticles();
 			auto RigidParticle0 = ConstrainedParticles[0]->CastToRigidParticle();
@@ -1213,7 +1213,7 @@ namespace Chaos
 			}
 		}
 
-		void DrawCollidingShapes(const FRigidTransform3& SpaceTransform, const FPBDCollisionConstraints& Collisions, float ColorScale, const FChaosDebugDrawSettings* Settings)
+		void DrawCollidingShapes(const FRigidTransform3& SpaceTransform, const FPBDCollisionConstraints& Collisions, FRealSingle ColorScale, const FChaosDebugDrawSettings* Settings)
 		{
 			if (FDebugDrawQueue::IsDebugDrawingEnabled())
 			{
diff --git a/Engine/Source/Runtime/Experimental/Chaos/Private/Chaos/Evolution/PBDMinEvolution.cpp b/Engine/Source/Runtime/Experimental/Chaos/Private/Chaos/Evolution/PBDMinEvolution.cpp
--- a/Engine/Source/Runtime/Experimental/Chaos/Private/Chaos/Evolution/PBDMinEvolution.cpp
+++ b/Engine/Source/Runtime/Experimental/Chaos/Private/Chaos/Evolution/PBDMinEvolution.cpp
@@ -281,7 +281,7 @@ namespace Chaos
 				if (Particle.ObjectState() == EObjectStateType::Dynamic)
 				{
 					Particle.X() = FVec3::Lerp(Particle.Handle()->AuxilaryValue(ParticlePrevXs), Particle.X(), T);
-					Particle.R() = FRotation3::Slerp(Particle.Handle()->AuxilaryValue(ParticlePrevRs), Particle.R(), T);
+					Particle.R() = FRotation3::Slerp(Particle.Handle()->AuxilaryValue(ParticlePrevRs), Particle.R(), (decltype(FQuat::X))T);	// LWC_TODO: Remove decltype cast once FQuat supports variants
 				}
 			}
 		}
@@ -600,7 +600,7 @@ namespace Chaos
 				else
 				{
 					TargetPos = FVec3::Lerp(PrevX, KinematicTarget.GetTarget().GetLocation(), StepFraction);
-					TargetRot = FRotation3::Slerp(PrevR, KinematicTarget.GetTarget().GetRotation(), StepFraction);
+					TargetRot = FRotation3::Slerp(PrevR, KinematicTarget.GetTarget().GetRotation(), (decltype(FQuat::X))StepFraction);		// LWC_TODO: Remove decltype cast once FQuat supports variants
 				}
 				if (Dt > MinDt)
 				{
diff --git a/Engine/Source/Runtime/Experimental/Chaos/Private/Chaos/Framework/ChaosResultsManager.cpp b/Engine/Source/Runtime/Experimental/Chaos/Private/Chaos/Framework/ChaosResultsManager.cpp
--- a/Engine/Source/Runtime/Experimental/Chaos/Private/Chaos/Framework/ChaosResultsManager.cpp
+++ b/Engine/Source/Runtime/Experimental/Chaos/Private/Chaos/Framework/ChaosResultsManager.cpp
@@ -197,7 +197,7 @@ namespace Chaos
 			return FMath::Min((FReal)1, (ResultsTime - LatestDivergeTime) * InvResimInterpTime);
 		};
 
-		Results.Alpha = GlobalAlpha;
+		Results.Alpha = (float)GlobalAlpha;	 // LWC_TODO: Precision loss
 
 		if(ParticleToResimInfo.Num() == 0)	//no resim interpolation so just exit
 		{
@@ -223,7 +223,7 @@ namespace Chaos
 						if(LeashInfo->bDiverged)
 						{
 							//RigidData.LeashAlpha = LeashAlpha;
-							RigidData.LeashAlpha = ResimInterpStrength;
+							RigidData.LeashAlpha = (float)ResimInterpStrength;		 // LWC_TODO: Precision loss
 						}
 					}
 				}
@@ -250,7 +250,7 @@ namespace Chaos
 
 					RigidData.Next = Itr.Value.Next;	//not dirty from sim, so just use whatever last next was
 					RigidData.Prev = RigidData.Next;	//prev same as next since we're just using leash
-					RigidData.LeashAlpha = LeashAlpha < 1 ? ResimInterpStrength : 1;	//if last step with leash just snap to end
+					RigidData.LeashAlpha = LeashAlpha < 1 ? (float)ResimInterpStrength : 1;	//if last step with leash just snap to end		 // LWC_TODO: Precision loss
 				}
 				
 				if (LeashAlpha >= 1)	//leash will snap to end so not needed after this
diff --git a/Engine/Source/Runtime/Experimental/Chaos/Private/Chaos/PBDCollisionConstraintUtil.cpp b/Engine/Source/Runtime/Experimental/Chaos/Private/Chaos/PBDCollisionConstraintUtil.cpp
--- a/Engine/Source/Runtime/Experimental/Chaos/Private/Chaos/PBDCollisionConstraintUtil.cpp
+++ b/Engine/Source/Runtime/Experimental/Chaos/Private/Chaos/PBDCollisionConstraintUtil.cpp
@@ -15,14 +15,14 @@
 namespace Chaos
 {
 	void ComputeHashTable(const TArray<Chaos::FRigidBodyPointContactConstraint>& ConstraintsArray,
-						  const FBox& BoundingBox, TMultiMap<int32, int32>& HashTableMap, const FRealSingle SpatialHashRadius)
+						  const FBox& BoundingBox, TMultiMap<int32, int32>& HashTableMap, const FReal SpatialHashRadius)
 	{
-		FRealSingle CellSize = 2.f * SpatialHashRadius;
-		check(CellSize > 0.f);
+		FReal CellSize = 2 * SpatialHashRadius;
+		check(CellSize > 0);
 
 		// Compute number of cells along the principal axis
-		FVector Extent = 2.f * BoundingBox.GetExtent();
-		FRealSingle PrincipalAxisLength;
+		FVector Extent = 2 * BoundingBox.GetExtent();
+		FReal PrincipalAxisLength;
 		if (Extent.X > Extent.Y && Extent.X > Extent.Z)
 		{
 			PrincipalAxisLength = Extent.X;
@@ -38,8 +38,8 @@ namespace Chaos
 		int32 NumberOfCells = FMath::CeilToInt(PrincipalAxisLength / CellSize);
 		check(NumberOfCells > 0);
 
-		CellSize = PrincipalAxisLength / (FRealSingle)NumberOfCells;
-		FRealSingle CellSizeInv = 1.f / CellSize;
+		CellSize = PrincipalAxisLength / (FReal)NumberOfCells;
+		FReal CellSizeInv = 1 / CellSize;
 
 		int32 NumberOfCellsX = FMath::CeilToInt(Extent.X * CellSizeInv) + 1;
 		int32 NumberOfCellsY = FMath::CeilToInt(Extent.Y * CellSizeInv) + 1;
@@ -63,14 +63,14 @@ namespace Chaos
 	}
 
 	void ComputeHashTable(const TArray<FCollidingData>& CollisionsArray,
-						  const FBox& BoundingBox, TMultiMap<int32, int32>& HashTableMap, const FRealSingle SpatialHashRadius)
+						  const FBox& BoundingBox, TMultiMap<int32, int32>& HashTableMap, const FReal SpatialHashRadius)
 	{
-		FRealSingle CellSize = 2.f * SpatialHashRadius;
-		check(CellSize > 0.f);
+		FReal CellSize = 2 * SpatialHashRadius;
+		check(CellSize > 0);
 
 		// Compute number of cells along the principal axis
-		FVector Extent = 2.f * BoundingBox.GetExtent();
-		FRealSingle PrincipalAxisLength;
+		FVector Extent = 2 * BoundingBox.GetExtent();
+		FReal PrincipalAxisLength;
 		if (Extent.X > Extent.Y && Extent.X > Extent.Z)
 		{
 			PrincipalAxisLength = Extent.X;
@@ -86,8 +86,8 @@ namespace Chaos
 		int32 NumberOfCells = FMath::CeilToInt(PrincipalAxisLength / CellSize);
 		check(NumberOfCells > 0);
 
-		CellSize = PrincipalAxisLength / (FRealSingle)NumberOfCells;
-		FRealSingle CellSizeInv = 1.f / CellSize;
+		CellSize = PrincipalAxisLength / (FReal)NumberOfCells;
+		FReal CellSizeInv = 1 / CellSize;
 
 		int32 NumberOfCellsX = FMath::CeilToInt(Extent.X * CellSizeInv) + 1;
 		int32 NumberOfCellsY = FMath::CeilToInt(Extent.Y * CellSizeInv) + 1;
@@ -111,14 +111,14 @@ namespace Chaos
 	}
 
 	void ComputeHashTable(const TArray<FCollidingDataExt>& CollisionsArray,
-						  const FBox& BoundingBox, TMultiMap<int32, int32>& HashTableMap, const FRealSingle SpatialHashRadius)
+						  const FBox& BoundingBox, TMultiMap<int32, int32>& HashTableMap, const FReal SpatialHashRadius)
 	{
-		FRealSingle CellSize = 2.f * SpatialHashRadius;
-		check(CellSize > 0.f);
+		FReal CellSize = 2 * SpatialHashRadius;
+		check(CellSize > 0);
 
 		// Compute number of cells along the principal axis
-		FVector Extent = 2.f * BoundingBox.GetExtent();
-		FRealSingle PrincipalAxisLength;
+		FVector Extent = 2 * BoundingBox.GetExtent();
+		FReal PrincipalAxisLength;
 		if (Extent.X > Extent.Y && Extent.X > Extent.Z)
 		{
 			PrincipalAxisLength = Extent.X;
@@ -134,8 +134,8 @@ namespace Chaos
 		int32 NumberOfCells = FMath::CeilToInt(PrincipalAxisLength / CellSize);
 		check(NumberOfCells > 0);
 
-		CellSize = PrincipalAxisLength / (FRealSingle)NumberOfCells;
-		FRealSingle CellSizeInv = 1.f / CellSize;
+		CellSize = PrincipalAxisLength / (FReal)NumberOfCells;
+		FReal CellSizeInv = 1 / CellSize;
 
 		int32 NumberOfCellsX = FMath::CeilToInt(Extent.X * CellSizeInv) + 1;
 		int32 NumberOfCellsY = FMath::CeilToInt(Extent.Y * CellSizeInv) + 1;
@@ -149,7 +149,7 @@ namespace Chaos
 		if (NumberOfCellsXYZ < 0)
 		{
 			CellSize = PrincipalAxisLength / 1000;
-			CellSizeInv = 1.f / CellSize;
+			CellSizeInv = 1 / CellSize;
 			NumberOfCellsX = FMath::CeilToInt(Extent.X * CellSizeInv) + 1;
 			NumberOfCellsY = FMath::CeilToInt(Extent.Y * CellSizeInv) + 1;
 			NumberOfCellsZ = FMath::CeilToInt(Extent.Z * CellSizeInv) + 1;
@@ -169,14 +169,14 @@ namespace Chaos
 		}
 	}
 
-	void ComputeHashTable(const TArray<FVector>& ParticleArray, const FBox& BoundingBox, TMultiMap<int32, int32>& HashTableMap, const FRealSingle SpatialHashRadius)
+	void ComputeHashTable(const TArray<FVector>& ParticleArray, const FBox& BoundingBox, TMultiMap<int32, int32>& HashTableMap, const FReal SpatialHashRadius)
 	{
-		FRealSingle CellSize = 2.f * SpatialHashRadius;
-		check(CellSize > 0.f);
+		FReal CellSize = 2 * SpatialHashRadius;
+		check(CellSize > 0);
 
 		// Compute number of cells along the principal axis
-		FVector Extent = 2.f * BoundingBox.GetExtent();
-		FRealSingle PrincipalAxisLength;
+		FVector Extent = 2 * BoundingBox.GetExtent();
+		FReal PrincipalAxisLength;
 		if (Extent.X > Extent.Y && Extent.X > Extent.Z)
 		{
 			PrincipalAxisLength = Extent.X;
@@ -192,8 +192,8 @@ namespace Chaos
 		int32 NumberOfCells = FMath::CeilToInt(PrincipalAxisLength / CellSize);
 		check(NumberOfCells > 0);
 
-		CellSize = PrincipalAxisLength / (FRealSingle)NumberOfCells;
-		FRealSingle CellSizeInv = 1.f / CellSize;
+		CellSize = PrincipalAxisLength / (FReal)NumberOfCells;
+		FReal CellSizeInv = 1 / CellSize;
 
 		int32 NumberOfCellsX = FMath::CeilToInt(Extent.X * CellSizeInv) + 1;
 		int32 NumberOfCellsY = FMath::CeilToInt(Extent.Y * CellSizeInv) + 1;
@@ -217,14 +217,14 @@ namespace Chaos
 	}
 
 	void ComputeHashTable(const TArray<FBreakingData>& BreakingsArray,
-						  const FBox& BoundingBox, TMultiMap<int32, int32>& HashTableMap, const FRealSingle SpatialHashRadius)
+						  const FBox& BoundingBox, TMultiMap<int32, int32>& HashTableMap, const FReal SpatialHashRadius)
 	{
-		FRealSingle CellSize = 2.f * SpatialHashRadius;
-		check(CellSize > 0.f);
+		FReal CellSize = 2 * SpatialHashRadius;
+		check(CellSize > 0);
 
 		// Compute number of cells along the principal axis
-		FVector Extent = 2.f * BoundingBox.GetExtent();
-		FRealSingle PrincipalAxisLength;
+		FVector Extent = 2 * BoundingBox.GetExtent();
+		FReal PrincipalAxisLength;
 		if (Extent.X > Extent.Y && Extent.X > Extent.Z)
 		{
 			PrincipalAxisLength = Extent.X;
@@ -240,8 +240,8 @@ namespace Chaos
 		int32 NumberOfCells = FMath::CeilToInt(PrincipalAxisLength / CellSize);
 		check(NumberOfCells > 0);
 
-		CellSize = PrincipalAxisLength / (FRealSingle)NumberOfCells;
-		FRealSingle CellSizeInv = 1.f / CellSize;
+		CellSize = PrincipalAxisLength / (FReal)NumberOfCells;
+		FReal CellSizeInv = 1 / CellSize;
 
 		int32 NumberOfCellsX = FMath::CeilToInt(Extent.X * CellSizeInv) + 1;
 		int32 NumberOfCellsY = FMath::CeilToInt(Extent.Y * CellSizeInv) + 1;
@@ -265,14 +265,14 @@ namespace Chaos
 	}
 
 	void ComputeHashTable(const TArray<FBreakingDataExt>& BreakingsArray,
-						  const FBox& BoundingBox, TMultiMap<int32, int32>& HashTableMap, const FRealSingle SpatialHashRadius)
+						  const FBox& BoundingBox, TMultiMap<int32, int32>& HashTableMap, const FReal SpatialHashRadius)
 	{
-		FRealSingle CellSize = 2.f * SpatialHashRadius;
-		check(CellSize > 0.f);
+		FReal CellSize = 2 * SpatialHashRadius;
+		check(CellSize > 0);
 
 		// Compute number of cells along the principal axis
-		FVector Extent = 2.f * BoundingBox.GetExtent();
-		FRealSingle PrincipalAxisLength;
+		FVector Extent = 2 * BoundingBox.GetExtent();
+		FReal PrincipalAxisLength;
 		if (Extent.X > Extent.Y && Extent.X > Extent.Z)
 		{
 			PrincipalAxisLength = Extent.X;
@@ -288,8 +288,8 @@ namespace Chaos
 		int32 NumberOfCells = FMath::CeilToInt(PrincipalAxisLength / CellSize);
 		check(NumberOfCells > 0);
 
-		CellSize = PrincipalAxisLength / (FRealSingle)NumberOfCells;
-		FRealSingle CellSizeInv = 1.f / CellSize;
+		CellSize = PrincipalAxisLength / (FReal)NumberOfCells;
+		FReal CellSizeInv = 1 / CellSize;
 
 		int32 NumberOfCellsX = FMath::CeilToInt(Extent.X * CellSizeInv) + 1;
 		int32 NumberOfCellsY = FMath::CeilToInt(Extent.Y * CellSizeInv) + 1;
diff --git a/Engine/Source/Runtime/Experimental/Chaos/Private/Chaos/PBDCollisionSpringConstraintsBase.cpp b/Engine/Source/Runtime/Experimental/Chaos/Private/Chaos/PBDCollisionSpringConstraintsBase.cpp
--- a/Engine/Source/Runtime/Experimental/Chaos/Private/Chaos/PBDCollisionSpringConstraintsBase.cpp
+++ b/Engine/Source/Runtime/Experimental/Chaos/Private/Chaos/PBDCollisionSpringConstraintsBase.cpp
@@ -106,8 +106,9 @@ void FPBDCollisionSpringConstraintsBase::Init(const FPBDParticles& Particles)
 			const FVec3& Start = Particles.X(Index);
 			const FVec3 Direction = Particles.V(Index).GetSafeNormal();
 			const FVec3 End = Particles.P(Index) + Direction * Height;
-			const FVector4 Start4(Start.X, Start.Y, Start.Z, 0);
-			const FVector4 End4(End.X, End.Y, End.Z, 0);
+			using FVec4Real = decltype(FVector4::X);
+			const FVector4 Start4((FVec4Real)Start.X, (FVec4Real)Start.Y, (FVec4Real)Start.Z, 0);
+			const FVector4 End4((FVec4Real)End.X, (FVec4Real)End.Y, (FVec4Real)End.Z, 0);
 			TkDOPLineCollisionCheck<const FMeshBuildDataProvider, uint32> Ray(Start4, End4, true, DopDataProvider, &Result);
 			if (DopTree.LineCheck(Ray))
 			{
diff --git a/Engine/Source/Runtime/Experimental/Chaos/Private/Chaos/PBDRigidClustering.cpp b/Engine/Source/Runtime/Experimental/Chaos/Private/Chaos/PBDRigidClustering.cpp
--- a/Engine/Source/Runtime/Experimental/Chaos/Private/Chaos/PBDRigidClustering.cpp
+++ b/Engine/Source/Runtime/Experimental/Chaos/Private/Chaos/PBDRigidClustering.cpp
@@ -691,7 +691,7 @@ namespace Chaos
 	TSet<FPBDRigidParticleHandle*> 
 	TPBDRigidClustering<T_FPBDRigidsEvolution, T_FPBDCollisionConstraint>::ReleaseClusterParticles(
 		FPBDRigidClusteredParticleHandle* ClusteredParticle,
-		const TMap<FGeometryParticleHandle*, float>* ExternalStrainMap,
+		const TMap<FGeometryParticleHandle*, Chaos::FReal>* ExternalStrainMap,
 		bool bForceRelease)
 	{
 		SCOPE_CYCLE_COUNTER(STAT_ReleaseClusterParticles_STRAIN);
@@ -766,7 +766,7 @@ namespace Chaos
 
 			if(ExternalStrainMap)
 			{
-				const float* MapStrain = ExternalStrainMap->Find(Child);
+				const Chaos::FReal* MapStrain = ExternalStrainMap->Find(Child);
 				ChildStrain = MapStrain ? *MapStrain : Child->CollisionImpulses();
 			}
 			else
@@ -966,7 +966,7 @@ namespace Chaos
 	TSet<FPBDRigidParticleHandle*>
 		TPBDRigidClustering<T_FPBDRigidsEvolution, T_FPBDCollisionConstraint>::ReleaseClusterParticlesNoInternalCluster(
 			FPBDRigidClusteredParticleHandle* ClusteredParticle,
-			const TMap<FGeometryParticleHandle*, float>* ExternalStrainMap,
+			const TMap<FGeometryParticleHandle*, Chaos::FReal>* ExternalStrainMap,
 			bool bForceRelease)
 	{
 		/* This is a near duplicate of the ReleaseClusterParticles() method with the internal cluster creation removed.
@@ -1046,7 +1046,7 @@ namespace Chaos
 
 			if (ExternalStrainMap)
 			{
-				const float* MapStrain = ExternalStrainMap->Find(Child);
+				const Chaos::FReal* MapStrain = ExternalStrainMap->Find(Child);
 				ChildStrain = MapStrain ? *MapStrain : Child->CollisionImpulses();
 			}
 			else
@@ -1203,7 +1203,7 @@ namespace Chaos
 			FPBDRigidParticleHandle* ClusterHandle = nullptr;
 			//todo(ocohen): refactor incoming, for now just assume these all belong to same cluster and hack strain array
 			
-			TMap<FGeometryParticleHandle*, float> FakeStrain;
+			TMap<FGeometryParticleHandle*, FReal> FakeStrain;
 
 			bool bPreDoGenerateData = DoGenerateBreakingData;
 			DoGenerateBreakingData = false;
@@ -1374,7 +1374,7 @@ namespace Chaos
 	template<class T_FPBDRigidsEvolution, class T_FPBDCollisionConstraint>
 	TMap<FPBDRigidClusteredParticleHandle*, TSet<FPBDRigidParticleHandle*>> 
 	TPBDRigidClustering<T_FPBDRigidsEvolution, T_FPBDCollisionConstraint>::BreakingModel(
-		TMap<FGeometryParticleHandle*, float>* ExternalStrainMap)
+		TMap<FGeometryParticleHandle*, Chaos::FReal>* ExternalStrainMap)
 	{
 		SCOPE_CYCLE_COUNTER(STAT_BreakingModel);
 
diff --git a/Engine/Source/Runtime/Experimental/Chaos/Private/Chaos/PBDSuspensionConstraints.cpp b/Engine/Source/Runtime/Experimental/Chaos/Private/Chaos/PBDSuspensionConstraints.cpp
--- a/Engine/Source/Runtime/Experimental/Chaos/Private/Chaos/PBDSuspensionConstraints.cpp
+++ b/Engine/Source/Runtime/Experimental/Chaos/Private/Chaos/PBDSuspensionConstraints.cpp
@@ -108,7 +108,7 @@ namespace Chaos
 				}
 				else
 				{
-					const float Speed = FMath::Abs(Particle->V().Length());
+					const FReal Speed = FMath::Abs(Particle->V().Length());
 					if (Speed < SpeedThreshold)
 					{
 						AxisWorld = FMath::Lerp(FVec3(0.f, 0.f, 1.f), AxisWorld, Speed / SpeedThreshold);
diff --git a/Engine/Source/Runtime/Experimental/Chaos/Private/ChaosVisualDebugger/ChaosVisualDebuggerTrace.cpp b/Engine/Source/Runtime/Experimental/Chaos/Private/ChaosVisualDebugger/ChaosVisualDebuggerTrace.cpp
--- a/Engine/Source/Runtime/Experimental/Chaos/Private/ChaosVisualDebugger/ChaosVisualDebuggerTrace.cpp
+++ b/Engine/Source/Runtime/Experimental/Chaos/Private/ChaosVisualDebugger/ChaosVisualDebuggerTrace.cpp
@@ -5,9 +5,9 @@
 UE_TRACE_CHANNEL_DEFINE(PhysicsChannel);
 
 UE_TRACE_EVENT_BEGIN(Physics, ParticlePosition)
-UE_TRACE_EVENT_FIELD(float, PositionX)
-UE_TRACE_EVENT_FIELD(float, PositionY)
-UE_TRACE_EVENT_FIELD(float, PositionZ)
+UE_TRACE_EVENT_FIELD(Chaos::FReal, PositionX)
+UE_TRACE_EVENT_FIELD(Chaos::FReal, PositionY)
+UE_TRACE_EVENT_FIELD(Chaos::FReal, PositionZ)
 UE_TRACE_EVENT_END()
 
 using namespace Chaos;
diff --git a/Engine/Source/Runtime/Experimental/Chaos/Private/Field/FieldSystemNodes.cpp b/Engine/Source/Runtime/Experimental/Chaos/Private/Field/FieldSystemNodes.cpp
--- a/Engine/Source/Runtime/Experimental/Chaos/Private/Field/FieldSystemNodes.cpp
+++ b/Engine/Source/Runtime/Experimental/Chaos/Private/Field/FieldSystemNodes.cpp
@@ -133,7 +133,7 @@ void FRadialIntMask::Evaluate(FFieldContext& Context, TFieldArrayView<int32>& Re
 		const FFieldContextIndex& Index = Context.SampleIndices[SampleIndex];
 		{
 			int32 Result;
-			float Delta2 = (Position - Context.SamplePositions[Index.Sample]).SizeSquared();
+			FVector::FReal Delta2 = (Position - Context.SamplePositions[Index.Sample]).SizeSquared();
 
 			if(Delta2 < Radius2)
 			{
@@ -227,9 +227,9 @@ void FWaveScalar::Evaluate(FFieldContext& Context, TFieldArrayView<float>& Resul
 	const float Velocity = (Period != 0.0f ) ? Wavelength / Period : 0.0f;
 
 	const float Wavenumber = (Wavelength != 0.0f ) ? 2.0f * PI / Wavelength : 0.0f;
-	const float DeltaTime = FMath::Max(Context.TimeSeconds, 0.0f);
-	const float Radius = Wavelength * DeltaTime / Period;
-	const float Decay = DeltaTime / Period;
+	const Chaos::FReal DeltaTime = FMath::Max(Context.TimeSeconds, 0.0f);
+	const Chaos::FReal Radius = Wavelength * DeltaTime / Period;
+	const Chaos::FReal Decay = DeltaTime / Period;
 
 	int32 NumSamples = Context.SampleIndices.Num();
 	for (int32 SampleIndex = 0; SampleIndex < NumSamples; SampleIndex++)
@@ -238,21 +238,21 @@ void FWaveScalar::Evaluate(FFieldContext& Context, TFieldArrayView<float>& Resul
 
 		if (Function == EWaveFunctionType::Field_Wave_Decay)
 		{
-			Results[Index.Result] = Magnitude * FMath::Exp(-Decay * Decay);
+			Results[Index.Result] = Magnitude * (float)FMath::Exp(-Decay * Decay);							// LWC_TODO: Precision loss
 		}
 		else
 		{
-			const float Distance = (Context.SamplePositions[Index.Sample] - Position).Size();
-			const float Fraction = (1.0f - Distance / Radius);
-			const float Phase = -Wavenumber * Radius * Fraction;
+			const Chaos::FReal Distance = (float)(Context.SamplePositions[Index.Sample] - Position).Size();
+			const Chaos::FReal Fraction = (1.0f - Distance / Radius);
+			const Chaos::FReal Phase = -Wavenumber * Radius * Fraction;
 
 			if (Function == EWaveFunctionType::Field_Wave_Cosine)
 			{
-				Results[Index.Result] = Magnitude * FMath::Cos(Phase);
+				Results[Index.Result] = Magnitude * (float)FMath::Cos(Phase);								// LWC_TODO: Precision loss
 			}
 			else if (Function == EWaveFunctionType::Field_Wave_Gaussian)
 			{
-				Results[Index.Result] = Magnitude * FMath::Exp(-Phase * Phase);
+				Results[Index.Result] = Magnitude * (float)FMath::Exp(-Phase * Phase);						// LWC_TODO: Precision loss
 			}
 			else if (Function == EWaveFunctionType::Field_Wave_Falloff)
 			{
@@ -264,19 +264,19 @@ void FWaveScalar::Evaluate(FFieldContext& Context, TFieldArrayView<float>& Resul
 					}
 					else if (Falloff == EFieldFalloffType::Field_Falloff_Linear)
 					{
-						Results[Index.Result] = Magnitude * Fraction;
+						Results[Index.Result] = Magnitude * float(Fraction);								// LWC_TODO: Precision loss
 					}
 					else if (Falloff == EFieldFalloffType::Field_Falloff_Squared)
 					{
-						Results[Index.Result] = Magnitude * Fraction * Fraction;
+						Results[Index.Result] = Magnitude * float(Fraction * Fraction);						// LWC_TODO: Precision loss
 					}
 					else if (Falloff == EFieldFalloffType::Field_Falloff_Inverse && Fraction > 0.0f)
 					{
-						Results[Index.Result] = Magnitude * 2.0f * (1.0f - 1.0f / (Fraction + 1.0f));
+						Results[Index.Result] = Magnitude * 2.0f * (1.0f - 1.0f / float(Fraction + 1.0f));	// LWC_TODO: Precision loss
 					}
 					else if (Falloff == EFieldFalloffType::Field_Falloff_Logarithmic)
 					{
-						Results[Index.Result] = Magnitude * FMath::LogX(2.0f, Fraction + 1.0f);
+						Results[Index.Result] = Magnitude * FMath::LogX(2.0f, (float)Fraction + 1.0f);		// LWC_TODO: Precision loss
 					}
 				}
 				else
@@ -376,11 +376,11 @@ void FRadialFalloff::Evaluator(const FFieldContext& Context, TFieldArrayView<flo
 			const FFieldContextIndex& Index = Context.SampleIndices[SampleIndex];
 			{
 				Results[Index.Result] = Default;
-				const float Delta = (Context.SamplePositions[Index.Sample] - Position).Size();
+				const FVector::FReal Delta = (Context.SamplePositions[Index.Sample] - Position).Size();
 
 				if (Delta < Radius)
 				{
-					const float Function = 1.0f - Delta / Radius;
+					const float Function = float(1.0f - Delta / Radius);		// LWC_TODO: Precision loss
 					Results[Index.Result] = EvalFalloffFunction<FalloffType>(MinRange, DeltaRange, Magnitude, Function);
 				}
 			}
@@ -454,11 +454,11 @@ void FPlaneFalloff::Evaluator(const FFieldContext& Context, const FPlane& Plane,
 			const FFieldContextIndex& Index = Context.SampleIndices[SampleIndex];
 			{
 				Results[Index.Result] = Default;
-				const float Delta = Plane.PlaneDot(Context.SamplePositions[Index.Sample]);
+				const FPlane::FReal Delta = Plane.PlaneDot(Context.SamplePositions[Index.Sample]);
 
 				if (Delta < -SMALL_NUMBER && Delta > -Distance)
 				{
-					const float Function = 1.0f + Delta / Distance;
+					const float Function = float(1.0f + Delta / Distance);		// LWC_TODO: Precision loss
 					Results[Index.Result] = EvalFalloffFunction<FalloffType>(MinRange, DeltaRange, Magnitude, Function);
 				}
 			}
@@ -542,7 +542,7 @@ void FBoxFalloff::Evaluator(const FFieldContext& Context, TFieldArrayView<float>
 			if (UnitBox.IsInside(LocalPoint))
 			{
 				const FVector Distance(FMath::Abs(LocalPoint.X)- HalfBox, FMath::Abs(LocalPoint.Y) - HalfBox, FMath::Abs(LocalPoint.Z) - HalfBox);
-				const float Delta = FMath::Min(FMath::Max(Distance.X, FMath::Max(Distance.Y, Distance.Z)), 0.0f);
+				const float Delta = (float)FMath::Min(FMath::Max(Distance.X, FMath::Max(Distance.Y, Distance.Z)), 0.0f);	// LWC_TODO: Precision loss
 				const float Function = - Delta / HalfBox;
 
 				Results[Index.Result] = EvalFalloffFunction<FalloffType>(MinRange, DeltaRange, Magnitude, Function);
@@ -614,7 +614,7 @@ FNoiseField::Evaluate(FFieldContext& Context, TFieldArrayView<float>& Results) c
 	{
 		const FFieldContextIndex& Index = Context.SampleIndices[SampleIndex];
 
-		float Dummy = 0.0f;
+		FVector::FReal Dummy = 0.0f;
 		FVector LocalPoint = Transform.InverseTransformPosition(Context.SamplePositions[Index.Sample]);
 		LocalPoint = FVector(FMath::Modf(LocalPoint.X, &Dummy) * 0.5f + 0.5f, 
 							 FMath::Modf(LocalPoint.Y, &Dummy) * 0.5f + 0.5f,
@@ -622,7 +622,7 @@ FNoiseField::Evaluate(FFieldContext& Context, TFieldArrayView<float>& Results) c
 
 		// Samples for the Perlin noise must be btw 0->255
 		float PerlinValue = 0.0f;
-		Field::PerlinNoise::Sample(&PerlinValue, LocalPoint.X, LocalPoint.Y, LocalPoint.Z);
+		Field::PerlinNoise::Sample(&PerlinValue, (float)LocalPoint.X, (float)LocalPoint.Y, (float)LocalPoint.Z);
 
 		// Perlin noise result is btw -1 -> 1
 		PerlinValue = 0.5f * ( PerlinValue + 1.0f );
diff --git a/Engine/Source/Runtime/Experimental/Chaos/Private/GeometryCollection/GeometryCollectionCollisionStructureManager.cpp b/Engine/Source/Runtime/Experimental/Chaos/Private/GeometryCollection/GeometryCollectionCollisionStructureManager.cpp
--- a/Engine/Source/Runtime/Experimental/Chaos/Private/GeometryCollection/GeometryCollectionCollisionStructureManager.cpp
+++ b/Engine/Source/Runtime/Experimental/Chaos/Private/GeometryCollection/GeometryCollectionCollisionStructureManager.cpp
@@ -191,7 +191,7 @@ FCollisionStructureManager::NewImplicit(
 	const Chaos::FParticles& MeshParticles,
 	const Chaos::FTriangleMesh& TriMesh,
 	const FBox& CollisionBounds,
-	const float Radius,
+	const Chaos::FReal Radius,
 	const int32 MinRes,
 	const int32 MaxRes,
 	const float CollisionObjectReduction,
@@ -227,7 +227,7 @@ FCollisionStructureManager::NewImplicitBox(
 	// @todo(chaos): pass margin settings into the collision manager?
 	float CollisionMarginFraction = 0.1f;// FMath::Max(0.0f, UPhysicsSettingsCore::Get()->SolverOptions.CollisionMarginFraction);
 	float CollisionMarginMax = 10.0f;// FMath::Max(0.0f, UPhysicsSettingsCore::Get()->SolverOptions.CollisionMarginMax);
-	const float Margin = FMath::Min(CollisionMarginFraction * 0.5f * HalfExtents.GetMin(), CollisionMarginMax);
+	const Chaos::FReal Margin = FMath::Min(CollisionMarginFraction * 0.5f * HalfExtents.GetMin(), CollisionMarginMax);
 
 	Chaos::FImplicitObject* Implicit = new Chaos::TBox<Chaos::FReal, 3>(Center - HalfExtents, Center + HalfExtents, Margin);
 	UpdateImplicitFlags(Implicit, CollisionType);
@@ -236,7 +236,7 @@ FCollisionStructureManager::NewImplicitBox(
 
 FCollisionStructureManager::FImplicit*
 FCollisionStructureManager::NewImplicitSphere(
-	const float Radius,
+	const Chaos::FReal Radius,
 	const float CollisionObjectReduction,
 	const ECollisionTypeEnum CollisionType)
 {
@@ -247,8 +247,8 @@ FCollisionStructureManager::NewImplicitSphere(
 
 FCollisionStructureManager::FImplicit*
 FCollisionStructureManager::NewImplicitCapsule(
-	const float Radius,
-	const float Length,
+	const Chaos::FReal Radius,
+	const Chaos::FReal Length,
 	const float CollisionObjectReduction,
 	const ECollisionTypeEnum CollisionType)
 {
@@ -366,7 +366,7 @@ Chaos::FLevelSet* FCollisionStructureManager::NewLevelset(
 FVector 
 FCollisionStructureManager::CalculateUnitMassInertiaTensor(
 	const FBox& Bounds,
-	const float Radius,
+	const Chaos::FReal Radius,
 	const EImplicitTypeEnum ImplicitType
 )
 {	
@@ -385,10 +385,10 @@ FCollisionStructureManager::CalculateUnitMassInertiaTensor(
 	return Tensor;
 }
 
-float
+Chaos::FReal
 FCollisionStructureManager::CalculateVolume(
 	const FBox& Bounds,
-	const float Radius,
+	const Chaos::FReal Radius,
 	const EImplicitTypeEnum ImplicitType
 )
 {
diff --git a/Engine/Source/Runtime/Experimental/Chaos/Private/GeometryCollection/GeometryCollectionProximityUtility.cpp b/Engine/Source/Runtime/Experimental/Chaos/Private/GeometryCollection/GeometryCollectionProximityUtility.cpp
--- a/Engine/Source/Runtime/Experimental/Chaos/Private/GeometryCollection/GeometryCollectionProximityUtility.cpp
+++ b/Engine/Source/Runtime/Experimental/Chaos/Private/GeometryCollection/GeometryCollectionProximityUtility.cpp
@@ -75,10 +75,10 @@ typedef TOctree2<FProximityTriangle, FMeshProximityTriangleOctreeSemantics> FPro
 
 bool FGeometryCollectionProximityUtility::IsPointInsideOfTriangle(const FVector& P, const FVector& Vertex0, const FVector& Vertex1, const FVector& Vertex2, float Threshold)
 {
-	float FaceArea  = 0.5f * FVector::CrossProduct((Vertex1 - Vertex0), (Vertex2 - Vertex0)).SizeSquared();
-	float Face1Area = 0.5f * FVector::CrossProduct((Vertex0 - P),(Vertex2 - P)).SizeSquared();
-	float Face2Area = 0.5f * FVector::CrossProduct((Vertex0 - P),(Vertex1 - P)).SizeSquared();
-	float Face3Area = 0.5f * FVector::CrossProduct((Vertex2 - P),(Vertex1 - P)).SizeSquared();
+	FVector::FReal FaceArea  = 0.5f * FVector::CrossProduct((Vertex1 - Vertex0), (Vertex2 - Vertex0)).SizeSquared();
+	FVector::FReal Face1Area = 0.5f * FVector::CrossProduct((Vertex0 - P),(Vertex2 - P)).SizeSquared();
+	FVector::FReal Face2Area = 0.5f * FVector::CrossProduct((Vertex0 - P),(Vertex1 - P)).SizeSquared();
+	FVector::FReal Face3Area = 0.5f * FVector::CrossProduct((Vertex2 - P),(Vertex1 - P)).SizeSquared();
 
 	return (FMath::Abs(Face1Area + Face2Area + Face3Area - FaceArea) < Threshold * Threshold);
 }
@@ -515,15 +515,15 @@ void FGeometryCollectionProximityUtility::UpdateProximity(FGeometryCollection* G
 		FaceByConnectedTransformsMap.MultiFind(Elem, FaceIndexArray);
 
 		// Find the centroid of the region and save it into BreakingRegionCentroidArray
-		FVector Centroid = FVector(ForceInitToZero);
+		FVector3f Centroid = FVector3f(ForceInitToZero);
 		float TotalArea = 0.f;
 		for (int32 LocalIdxFace = 0; LocalIdxFace < FaceIndexArray.Num(); ++LocalIdxFace)
 		{
-			const FVector& Vertex0 = VertexArray[IndicesArray[FaceIndexArray[LocalIdxFace]][0]];
-			const FVector& Vertex1 = VertexArray[IndicesArray[FaceIndexArray[LocalIdxFace]][1]];
-			const FVector& Vertex2 = VertexArray[IndicesArray[FaceIndexArray[LocalIdxFace]][2]];
+			const FVector3f& Vertex0 = VertexArray[IndicesArray[FaceIndexArray[LocalIdxFace]][0]];
+			const FVector3f& Vertex1 = VertexArray[IndicesArray[FaceIndexArray[LocalIdxFace]][1]];
+			const FVector3f& Vertex2 = VertexArray[IndicesArray[FaceIndexArray[LocalIdxFace]][2]];
 
-			FVector FaceCentroid((Vertex0 + Vertex1 + Vertex2) / 3.f);
+			FVector3f FaceCentroid((Vertex0 + Vertex1 + Vertex2) / 3.f);
 			float FaceArea = 0.5f * ((Vertex1 - Vertex0) ^ (Vertex2 - Vertex0)).Size();
 
 			Centroid = (TotalArea * Centroid + FaceArea * FaceCentroid) / (TotalArea + FaceArea);
@@ -536,7 +536,7 @@ void FGeometryCollectionProximityUtility::UpdateProximity(FGeometryCollection* G
 		float RadiusMin = FLT_MAX;
 		float RadiusMax = FLT_MIN;
 
-		TArray<FVector> TestPoints;
+		TArray<FVector3f> TestPoints;
 		for (int32 LocalIdxFace = 0; LocalIdxFace < FaceIndexArray.Num(); ++LocalIdxFace)
 		{
 			for (int32 Idx = 0; Idx < 3; Idx++)
@@ -560,9 +560,9 @@ void FGeometryCollectionProximityUtility::UpdateProximity(FGeometryCollection* G
 		BreakingRegionRadiusArray[IdxBreak] = RadiusMin;
 
 		// Normal
-		const FVector& VertexA = VertexArray[IndicesArray[FaceIndexArray[0]][0]];
-		const FVector& VertexB = VertexArray[IndicesArray[FaceIndexArray[0]][1]];
-		const FVector& VertexC = VertexArray[IndicesArray[FaceIndexArray[0]][2]];
+		const FVector3f& VertexA = VertexArray[IndicesArray[FaceIndexArray[0]][0]];
+		const FVector3f& VertexB = VertexArray[IndicesArray[FaceIndexArray[0]][1]];
+		const FVector3f& VertexC = VertexArray[IndicesArray[FaceIndexArray[0]][2]];
 		BreakingRegionNormalArray[IdxBreak] = ((VertexA - VertexB) ^ (VertexC - VertexB)).GetSafeNormal();
 
 		// grab the first face from the region and save it into BreakingFaceIndexArray
diff --git a/Engine/Source/Runtime/Experimental/Chaos/Private/GeometryCollection/GeometryCollectionUtility.cpp b/Engine/Source/Runtime/Experimental/Chaos/Private/GeometryCollection/GeometryCollectionUtility.cpp
--- a/Engine/Source/Runtime/Experimental/Chaos/Private/GeometryCollection/GeometryCollectionUtility.cpp
+++ b/Engine/Source/Runtime/Experimental/Chaos/Private/GeometryCollection/GeometryCollectionUtility.cpp
@@ -48,14 +48,15 @@ namespace GeometryCollection
 
 		// set the vertex information
 		int32 Index = 0;
-		Vertices[0] = FVector3f(-Scale.X / 2.f, -Scale.Y / 2.f, -Scale.Z / 2.f);
-		Vertices[1] = FVector3f(+Scale.X / 2.f, -Scale.Y / 2.f, -Scale.Z / 2.f);
-		Vertices[2] = FVector3f(-Scale.X / 2.f, +Scale.Y / 2.f, -Scale.Z / 2.f);
-		Vertices[3] = FVector3f(+Scale.X / 2.f, +Scale.Y / 2.f, -Scale.Z / 2.f);
-		Vertices[4] = FVector3f(-Scale.X / 2.f, -Scale.Y / 2.f, +Scale.Z / 2.f);
-		Vertices[5] = FVector3f(+Scale.X / 2.f, -Scale.Y / 2.f, +Scale.Z / 2.f);
-		Vertices[6] = FVector3f(-Scale.X / 2.f, +Scale.Y / 2.f, +Scale.Z / 2.f);
-		Vertices[7] = FVector3f(+Scale.X / 2.f, +Scale.Y / 2.f, +Scale.Z / 2.f);
+		FVector3f Scale3f = FVector3f(Scale);	// LWC_TODO: Precision loss - Not significant as everything it applies to is already float.
+		Vertices[0] = FVector3f(-Scale3f.X / 2.f, -Scale3f.Y / 2.f, -Scale3f.Z / 2.f);
+		Vertices[1] = FVector3f(+Scale3f.X / 2.f, -Scale3f.Y / 2.f, -Scale3f.Z / 2.f);
+		Vertices[2] = FVector3f(-Scale3f.X / 2.f, +Scale3f.Y / 2.f, -Scale3f.Z / 2.f);
+		Vertices[3] = FVector3f(+Scale3f.X / 2.f, +Scale3f.Y / 2.f, -Scale3f.Z / 2.f);
+		Vertices[4] = FVector3f(-Scale3f.X / 2.f, -Scale3f.Y / 2.f, +Scale3f.Z / 2.f);
+		Vertices[5] = FVector3f(+Scale3f.X / 2.f, -Scale3f.Y / 2.f, +Scale3f.Z / 2.f);
+		Vertices[6] = FVector3f(-Scale3f.X / 2.f, +Scale3f.Y / 2.f, +Scale3f.Z / 2.f);
+		Vertices[7] = FVector3f(+Scale3f.X / 2.f, +Scale3f.Y / 2.f, +Scale3f.Z / 2.f);
 
 		Normals[0] = FVector3f(-1.f, -1.f, -1.f).GetSafeNormal();
 		Normals[1] = FVector3f(1.f, -1.f, -1.f).GetSafeNormal();
@@ -303,8 +304,8 @@ namespace GeometryCollection
 					}
 
 					// build vertex summary information
-					TArray<FVector> Center;
-					Center.Init(FVector(0), GeometryIndices.Num());
+					TArray<FVector3f> Center;
+					Center.Init(FVector3f(0), GeometryIndices.Num());
 					int CurrentParticleIndex = FGeometryCollection::Invalid;
 					for (int vdx = 0; vdx < Vertex.Num(); vdx++)
 					{
@@ -337,7 +338,7 @@ namespace GeometryCollection
 					{
 						if (VertexCount[GeometryIndex])
 						{
-							Center[GeometryIndex] /= static_cast<Chaos::FReal>(VertexCount[GeometryIndex]);
+							Center[GeometryIndex] /= static_cast<Chaos::FRealSingle>(VertexCount[GeometryIndex]);
 						}
 					}
 
@@ -372,7 +373,7 @@ namespace GeometryCollection
 					for (int vdx = 0; vdx < Vertex.Num(); vdx++)
 					{
 						int32 GeometryIndex = ReverseMap[BoneMap[vdx]]; // double indexing safe due to check in previous loop.
-						float Delta = (Center[GeometryIndex] - Vertex[vdx]).Size();
+						Chaos::FRealSingle Delta = (Center[GeometryIndex] - Vertex[vdx]).Size();
 						InnerRadius[GeometryIndex] = FMath::Min(InnerRadius[GeometryIndex], Delta);
 						OuterRadius[GeometryIndex] = FMath::Max(OuterRadius[GeometryIndex], Delta);
 					}
diff --git a/Engine/Source/Runtime/Experimental/Chaos/Private/PhysicsProxy/GeometryCollectionPhysicsProxy.cpp b/Engine/Source/Runtime/Experimental/Chaos/Private/PhysicsProxy/GeometryCollectionPhysicsProxy.cpp
--- a/Engine/Source/Runtime/Experimental/Chaos/Private/PhysicsProxy/GeometryCollectionPhysicsProxy.cpp
+++ b/Engine/Source/Runtime/Experimental/Chaos/Private/PhysicsProxy/GeometryCollectionPhysicsProxy.cpp
@@ -149,7 +149,7 @@ void PopulateSimulatedParticle(
 	FGeometryDynamicCollection::FSharedImplicit Implicit,
 	const FCollisionFilterData SimFilterIn,
 	const FCollisionFilterData QueryFilterIn,
-	float MassIn,
+	Chaos::FReal MassIn,
 	FVector3f InertiaTensorVec,
 	const FTransform& WorldTransform, 
 	const uint8 DynamicState, 
@@ -174,10 +174,10 @@ void PopulateSimulatedParticle(
 	{
 		Handle->SetObjectStateLowLevel(Chaos::EObjectStateType::Uninitialized);
 
-		if (!CHAOS_ENSURE_MSG(FMath::IsWithinInclusive(MassIn, SharedParams.MinimumMassClamp, SharedParams.MaximumMassClamp),
+		if (!CHAOS_ENSURE_MSG(FMath::IsWithinInclusive<Chaos::FReal>(MassIn, SharedParams.MinimumMassClamp, SharedParams.MaximumMassClamp),
 			TEXT("Clamped mass[%3.5f] to range [%3.5f,%3.5f]"), MassIn, SharedParams.MinimumMassClamp, SharedParams.MaximumMassClamp))
 		{
-			MassIn = FMath::Clamp(MassIn, SharedParams.MinimumMassClamp, SharedParams.MaximumMassClamp);
+			MassIn = FMath::Clamp<Chaos::FReal>(MassIn, SharedParams.MinimumMassClamp, SharedParams.MaximumMassClamp);
 		}
 
 		if (!CHAOS_ENSURE_MSG(!FMath::IsNaN(InertiaTensorVec[0]) && !FMath::IsNaN(InertiaTensorVec[1]) && !FMath::IsNaN(InertiaTensorVec[2]),
@@ -651,7 +651,7 @@ void FGeometryCollectionPhysicsProxy::InitializeBodiesPT(Chaos::FPBDRigidsSolver
 				if (CollisionParticles)
 				{
 					Chaos::FReal NumCollisionParticles = static_cast<Chaos::FReal>(CollisionParticles->Size());
-					int32 CollisionParticlesSize = FMath::TruncToInt(Chaos::FReal(FMath::Max(0, FMath::Min(int(NumCollisionParticles * CollisionParticlesPerObjectFraction), NumCollisionParticles))));
+					int32 CollisionParticlesSize = FMath::TruncToInt(Chaos::FReal(FMath::Max((Chaos::FReal)0, FMath::Min(Chaos::FReal(NumCollisionParticles * CollisionParticlesPerObjectFraction), NumCollisionParticles))));
 					CollisionParticles->Resize(CollisionParticlesSize); // Truncates!
 				}
 
@@ -882,7 +882,7 @@ FAutoConsoleVariableRef CVarReportNoLevelsetCluster(TEXT("p.gc.ReportNoLevelsetC
 DECLARE_CYCLE_STAT(TEXT("FGeometryCollectionPhysicsProxy::BuildClusters"), STAT_BuildClusters, STATGROUP_Chaos);
 DECLARE_CYCLE_STAT(TEXT("FGeometryCollectionPhysicsProxy::BuildClusters:GlobalMatrices"), STAT_BuildClustersGlobalMatrices, STATGROUP_Chaos);
 
-int32 FindSizeSpecificIdx(const TArray<FSharedSimulationSizeSpecificData>& SizeSpecificData, float Volume)
+int32 FindSizeSpecificIdx(const TArray<FSharedSimulationSizeSpecificData>& SizeSpecificData, Chaos::FReal Volume)
 {
 	check(SizeSpecificData.Num());
 	int32 UseIdx = 0;
@@ -944,7 +944,7 @@ FGeometryCollectionPhysicsProxy::BuildClusters(
 	{
 		const Chaos::FReal NumCollisionParticles = static_cast<Chaos::FReal>(ClusterCreationParameters.CollisionParticles->Size());
 		const int32 ClampedCollisionParticlesSize = 
-			FMath::TruncToInt(Chaos::FReal(FMath::Max(0, FMath::Min(int(NumCollisionParticles * CollisionParticlesPerObjectFraction), NumCollisionParticles))));
+			FMath::TruncToInt(Chaos::FReal(FMath::Max(0, FMath::Min(NumCollisionParticles * CollisionParticlesPerObjectFraction, NumCollisionParticles))));
 		ClusterCreationParameters.CollisionParticles->Resize(ClampedCollisionParticlesSize);
 	}
 	TArray<Chaos::TPBDRigidParticleHandle<Chaos::FReal, 3>*> ChildHandlesCopy(ChildHandles);
@@ -1800,8 +1800,8 @@ void FGeometryCollectionPhysicsProxy::InitializeSharedCollisionStructures(
 	TManagedArray<FVector3f>& CollectionInertiaTensor =
 		RestCollection.AddAttribute<FVector3f>(
 			TEXT("InertiaTensor"), FTransformCollection::TransformGroup);
-	TManagedArray<float>& CollectionMass =
-		RestCollection.AddAttribute<float>(
+	TManagedArray<FReal>& CollectionMass =
+		RestCollection.AddAttribute<FReal>(
 			TEXT("Mass"), FTransformCollection::TransformGroup);
 	TManagedArray<TUniquePtr<FSimplicial>>& CollectionSimplicials =
 		RestCollection.AddAttribute<TUniquePtr<FSimplicial>>(
@@ -1868,7 +1868,7 @@ void FGeometryCollectionPhysicsProxy::InitializeSharedCollisionStructures(
 	// until we complete the loop before reporting the skips.
 	bool bSkippedSmallGeometry = false;
 	
-	float TotalVolume = 0.f;
+	FReal TotalVolume = 0.f;
 	// The geometry group has a set of transform indices that maps a geometry index
 	// to a transform index, but only in the case where there is a 1-to-1 mapping 
 	// between the two.  In the event where a geometry is instanced for multiple
@@ -1958,7 +1958,7 @@ void FGeometryCollectionPhysicsProxy::InitializeSharedCollisionStructures(
 				const int32 VCount = VertexCount[GeometryIndex];
 				if (VCount != 0)
 				{
-					const FVector Center = BoundingBox[GeometryIndex].GetCenter();
+					const FVector3f Center = (FVector3f)BoundingBox[GeometryIndex].GetCenter();
 					const int32 VStart = VertexStart[GeometryIndex];
 
 					InnerRadius[GeometryIndex] = VCount ? TNumericLimits<FRealSingle>::Max() : 0.0f;
@@ -1967,7 +1967,7 @@ void FGeometryCollectionPhysicsProxy::InitializeSharedCollisionStructures(
 					{
 						const int32 PtIdx = VStart + VIdx;
 						const FVector3f& Pt = Vertex[PtIdx];
-						const float DistSq = FVector::DistSquared(Pt, Center);
+						const float DistSq = FVector3f::DistSquared(Pt, Center);
 						if (InnerRadius[GeometryIndex] > DistSq)
 						{
 							InnerRadius[GeometryIndex] = DistSq;
@@ -2000,9 +2000,9 @@ void FGeometryCollectionPhysicsProxy::InitializeSharedCollisionStructures(
 	//Density must be the same for individual parts and the total. Density_i = Density = Mass_i / Volume_i
 	//Total mass must equal sum of individual parts. Mass_i = TotalMass * Volume_i / TotalVolume => Density_i = TotalMass / TotalVolume
 	TotalVolume = FMath::Max(TotalVolume, MinBoundsExtents * MinBoundsExtents * MinBoundsExtents);
-	const float DesiredTotalMass = SharedParams.bMassAsDensity ? SharedParams.Mass * TotalVolume : SharedParams.Mass;
-	const float ClampedTotalMass = FMath::Clamp(DesiredTotalMass, SharedParams.MinimumMassClamp, SharedParams.MaximumMassClamp);
-	const float DesiredDensity = ClampedTotalMass / TotalVolume;
+	const FReal DesiredTotalMass = SharedParams.bMassAsDensity ? SharedParams.Mass * TotalVolume : SharedParams.Mass;
+	const FReal ClampedTotalMass = FMath::Clamp(DesiredTotalMass, SharedParams.MinimumMassClamp, SharedParams.MaximumMassClamp);
+	const FReal DesiredDensity = ClampedTotalMass / TotalVolume;
 
 	FVec3 MaxChildBounds(1);
 	ParallelFor(NumGeometries, [&](int32 GeometryIndex)
@@ -2012,7 +2012,7 @@ void FGeometryCollectionPhysicsProxy::InitializeSharedCollisionStructures(
 		Chaos::FErrorReporter LocalErrorReporter;
 		const int32 TransformGroupIndex = TransformIndex[GeometryIndex];
 
-		const float Volume_i = MassPropertiesArray[GeometryIndex].Volume;
+		const FReal Volume_i = MassPropertiesArray[GeometryIndex].Volume;
 		if (CollectionSimulatableParticles[TransformGroupIndex])
 		{
 			//Must clamp each individual mass regardless of desired density
@@ -2031,8 +2031,8 @@ void FGeometryCollectionPhysicsProxy::InitializeSharedCollisionStructures(
 			TUniquePtr<FTriangleMesh>& TriMesh = TriangleMeshesArray[TransformGroupIndex];
 			FMassProperties& MassProperties = MassPropertiesArray[GeometryIndex];
 
-			const float Mass_i = FMath::Max(DesiredDensity * Volume_i, SharedParams.MinimumMassClamp);
-			const float Density_i = Mass_i / Volume_i;
+			const FReal Mass_i = FMath::Max(DesiredDensity * Volume_i, SharedParams.MinimumMassClamp);
+			const FReal Density_i = Mass_i / Volume_i;
 			CollectionMass[TransformGroupIndex] = Mass_i;
 
 			if (InertiaComputationNeeded[GeometryIndex])
@@ -2150,8 +2150,8 @@ void FGeometryCollectionPhysicsProxy::InitializeSharedCollisionStructures(
 				else if (SizeSpecificData.CollisionShapesData[0].ImplicitType == EImplicitTypeEnum::Chaos_Implicit_Capsule)
 				{
 					const FVector BBoxExtent = InstanceBoundingBox.GetExtent(); // FBox's extents are 1/2 (Max - Min)
-					const float Radius = BBoxExtent.GetAbsMin();
-					const float Length = (BBoxExtent.GetAbsMax() - Radius) * 2.f;
+					const Chaos::FReal Radius = BBoxExtent.GetAbsMin();
+					const Chaos::FReal Length = (BBoxExtent.GetAbsMax() - Radius) * 2.f;
 
 					if (Length < SMALL_NUMBER)
 					{
@@ -2339,8 +2339,8 @@ void FGeometryCollectionPhysicsProxy::InitializeSharedCollisionStructures(
 
 				TUniquePtr<FTriangleMesh> UnionMesh(new FTriangleMesh(MoveTemp(UnionMeshIndices)));
 				const FMatrix& InertiaMatrix = CollectionSpaceParticles->I(ClusterTransformIdx);
-				const FVector3f InertiaDiagonal(InertiaMatrix.M[0][0], InertiaMatrix.M[1][1], InertiaMatrix.M[2][2]);
-				CollectionInertiaTensor[ClusterTransformIdx] = InertiaDiagonal;
+				const FVector InertiaDiagonal(InertiaMatrix.M[0][0], InertiaMatrix.M[1][1], InertiaMatrix.M[2][2]);
+				CollectionInertiaTensor[ClusterTransformIdx] = (FVector3f)InertiaDiagonal;	// LWC_TODO: Precision loss
 				CollectionMass[ClusterTransformIdx] = CollectionSpaceParticles->M(ClusterTransformIdx);
 
 				const int32 SizeSpecificIdx = FindSizeSpecificIdx(SharedParams.SizeSpecificData, InstanceBoundingBox);
@@ -2349,18 +2349,18 @@ void FGeometryCollectionPhysicsProxy::InitializeSharedCollisionStructures(
 				if (SizeSpecificData.CollisionShapesData[0].ImplicitType == EImplicitTypeEnum::Chaos_Implicit_LevelSet)
 				{
 					const FVec3 Scale = 2 * InstanceBoundingBox.GetExtent() / MaxChildBounds; // FBox's extents are 1/2 (Max - Min)
-					const float ScaleMax = Scale.GetAbsMax();
-					const float ScaleMin = Scale.GetAbsMin();
+					const FReal ScaleMax = Scale.GetAbsMax();
+					const FReal ScaleMin = Scale.GetAbsMin();
 
-					float MinResolution = ScaleMin * (float)SizeSpecificData.CollisionShapesData[0].LevelSetData.MinLevelSetResolution;
-					MinResolution = FMath::Clamp<float>(MinResolution,
-						(float)SizeSpecificData.CollisionShapesData[0].LevelSetData.MinLevelSetResolution,
-						(float)SizeSpecificData.CollisionShapesData[0].LevelSetData.MinClusterLevelSetResolution);
+					FReal MinResolution = ScaleMin * (FReal)SizeSpecificData.CollisionShapesData[0].LevelSetData.MinLevelSetResolution;
+					MinResolution = FMath::Clamp(MinResolution,
+						(FReal)SizeSpecificData.CollisionShapesData[0].LevelSetData.MinLevelSetResolution,
+						(FReal)SizeSpecificData.CollisionShapesData[0].LevelSetData.MinClusterLevelSetResolution);
 
-					float MaxResolution = ScaleMax * (float)SizeSpecificData.CollisionShapesData[0].LevelSetData.MaxLevelSetResolution;
-					MaxResolution = FMath::Clamp<float>(MaxResolution,
-						(float)SizeSpecificData.CollisionShapesData[0].LevelSetData.MaxLevelSetResolution,
-						(float)SizeSpecificData.CollisionShapesData[0].LevelSetData.MaxClusterLevelSetResolution);
+					FReal MaxResolution = ScaleMax * (FReal)SizeSpecificData.CollisionShapesData[0].LevelSetData.MaxLevelSetResolution;
+					MaxResolution = FMath::Clamp(MaxResolution,
+						(FReal)SizeSpecificData.CollisionShapesData[0].LevelSetData.MaxLevelSetResolution,
+						(FReal)SizeSpecificData.CollisionShapesData[0].LevelSetData.MaxClusterLevelSetResolution);
 
 					//don't support non level-set serialization
 					ErrorReporter.SetPrefix(BaseErrorPrefix + " | Cluster Transform Index: " + FString::FromInt(ClusterTransformIdx));
@@ -2415,8 +2415,8 @@ void FGeometryCollectionPhysicsProxy::InitializeSharedCollisionStructures(
 				else if (SizeSpecificData.CollisionShapesData[0].ImplicitType == EImplicitTypeEnum::Chaos_Implicit_Capsule)
 				{
 					const FVector BBoxExtent = InstanceBoundingBox.GetExtent(); // FBox's extents are 1/2 (Max - Min)
-					const float Radius = BBoxExtent.GetAbsMin(); 
-					const float Length = (BBoxExtent.GetAbsMax() - Radius) * 2.f;
+					const FReal Radius = BBoxExtent.GetAbsMin();
+					const FReal Length = (BBoxExtent.GetAbsMax() - Radius) * 2.f;
 
 					if (Length < SMALL_NUMBER)
 					{
@@ -2601,7 +2601,7 @@ void FGeometryCollectionPhysicsProxy::FieldParameterUpdateCallback(Chaos::FPBDRi
 			{
 				if (Chaos::BuildFieldSamplePoints(this, RigidSolver, FieldCommand, ExecutionDatas, PrevResolutionType, PrevFilterType, PrevObjectType, PrevPositionType))
 				{
-					const float TimeSeconds = RigidSolver->GetSolverTime() - FieldCommand.TimeCreation;
+					const Chaos::FReal TimeSeconds = RigidSolver->GetSolverTime() - FieldCommand.TimeCreation;
 
 					FFieldContext FieldContext(
 						ExecutionDatas,
@@ -2756,7 +2756,7 @@ void FGeometryCollectionPhysicsProxy::FieldForcesUpdateCallback(Chaos::FPBDRigid
 			{
 				if (Chaos::BuildFieldSamplePoints(this, RigidSolver, FieldCommand, ExecutionDatas, PrevResolutionType, PrevFilterType, PrevObjectType, PrevPositionType))
 				{
-					const float TimeSeconds = RigidSolver->GetSolverTime() - FieldCommand.TimeCreation;
+					const Chaos::FReal TimeSeconds = RigidSolver->GetSolverTime() - FieldCommand.TimeCreation;
 
 					FFieldContext FieldContext(
 						ExecutionDatas,
diff --git a/Engine/Source/Runtime/Experimental/Chaos/Private/PhysicsProxy/PerSolverFieldSystem.cpp b/Engine/Source/Runtime/Experimental/Chaos/Private/PhysicsProxy/PerSolverFieldSystem.cpp
--- a/Engine/Source/Runtime/Experimental/Chaos/Private/PhysicsProxy/PerSolverFieldSystem.cpp
+++ b/Engine/Source/Runtime/Experimental/Chaos/Private/PhysicsProxy/PerSolverFieldSystem.cpp
@@ -51,7 +51,7 @@ void FPerSolverFieldSystem::FieldParameterUpdateInternal(
 			{
 				if (Chaos::BuildFieldSamplePoints(this, RigidSolver, FieldCommand, ExecutionDatas, PrevResolutionType, PrevFilterType, PrevObjectType, PrevPositionType))
 				{
-					const float TimeSeconds = RigidSolver->GetSolverTime() - FieldCommand.TimeCreation;
+					const Chaos::FReal TimeSeconds = RigidSolver->GetSolverTime() - FieldCommand.TimeCreation;
 
 					FFieldContext FieldContext(
 						ExecutionDatas,
@@ -134,7 +134,7 @@ void FPerSolverFieldSystem::FieldForcesUpdateInternal(
 			{
 				if (Chaos::BuildFieldSamplePoints(this, RigidSolver, FieldCommand, ExecutionDatas, PrevResolutionType, PrevFilterType, PrevObjectType, PrevPositionType))
 				{
-					const float TimeSeconds = RigidSolver->GetSolverTime() - FieldCommand.TimeCreation;
+					const Chaos::FReal TimeSeconds = RigidSolver->GetSolverTime() - FieldCommand.TimeCreation;
 
 					FFieldContext FieldContext(
 						ExecutionDatas,
@@ -204,7 +204,7 @@ FORCEINLINE void EvaluateImpulseField(
 
 void ComputeFieldRigidImpulseInternal(
 	FFieldExecutionDatas& ExecutionDatas,
-	const float SolverTime,
+	const Chaos::FReal SolverTime,
 	TArray<FFieldSystemCommand>& Commands, const bool IsTransient)
 {
 	const int32 NumCommands = Commands.Num();
@@ -225,7 +225,7 @@ void ComputeFieldRigidImpulseInternal(
 
 			if ((ObjectType == EFieldObjectType::Field_Object_Character) || (ObjectType == EFieldObjectType::Field_Object_All) || (ObjectType == EFieldObjectType::Field_Object_Max))
 			{
-				const float TimeSeconds = SolverTime - FieldCommand.TimeCreation;
+				const Chaos::FReal TimeSeconds = SolverTime - FieldCommand.TimeCreation;
 
 				FFieldContext FieldContext(
 					ExecutionDatas,
@@ -278,7 +278,7 @@ void ComputeFieldRigidImpulseInternal(
 }
 
 void FPerSolverFieldSystem::ComputeFieldRigidImpulse(
-	const float SolverTime)
+	const Chaos::FReal SolverTime)
 {
 	static const TArray<EFieldCommandOutputType> EmptyTargets = { EFieldCommandOutputType::LinearVelocity,
 																  EFieldCommandOutputType::LinearForce,
@@ -294,7 +294,7 @@ void FPerSolverFieldSystem::ComputeFieldRigidImpulse(
 
 void ComputeFieldLinearImpulseInternal(
 	FFieldExecutionDatas& ExecutionDatas,
-	const float SolverTime,
+	const Chaos::FReal SolverTime,
 	TArray<FFieldSystemCommand>& Commands, const bool IsTransient)
 {
 	const int32 NumCommands = Commands.Num();
@@ -315,7 +315,7 @@ void ComputeFieldLinearImpulseInternal(
 
 			if ((ObjectType == EFieldObjectType::Field_Object_Cloth) || (ObjectType == EFieldObjectType::Field_Object_All) || (ObjectType == EFieldObjectType::Field_Object_Max))
 			{
-				const float TimeSeconds = SolverTime - FieldCommand.TimeCreation;
+				const Chaos::FReal TimeSeconds = SolverTime - FieldCommand.TimeCreation;
 
 				FFieldContext FieldContext(
 					ExecutionDatas,
@@ -355,7 +355,7 @@ void ComputeFieldLinearImpulseInternal(
 	}
 }
 
-void FPerSolverFieldSystem::ComputeFieldLinearImpulse(const float SolverTime)
+void FPerSolverFieldSystem::ComputeFieldLinearImpulse(const Chaos::FReal SolverTime)
 {
 	static const TArray<EFieldCommandOutputType> EmptyTargets = { EFieldCommandOutputType::LinearVelocity,
 												EFieldCommandOutputType::LinearForce };
diff --git a/Engine/Source/Runtime/Experimental/Chaos/Private/kDOP.h b/Engine/Source/Runtime/Experimental/Chaos/Private/kDOP.h
--- a/Engine/Source/Runtime/Experimental/Chaos/Private/kDOP.h
+++ b/Engine/Source/Runtime/Experimental/Chaos/Private/kDOP.h
@@ -122,8 +122,8 @@ struct TTraversalHistory
  */
 FORCEINLINE bool appLineCheckTriangle(const FVector4& Start, const FVector4& End, const FVector4& Dir, const FVector4& V0, const FVector4& V1, const FVector4& V2, const FVector4& Normal, float& IntersectionTime)
 {
-	const float StartDist = FPlane(Normal).PlaneDot(Start);
-	const float EndDist = FPlane(Normal).PlaneDot(End);
+	const FPlane::FReal StartDist = FPlane(Normal).PlaneDot(Start);
+	const FPlane::FReal EndDist = FPlane(Normal).PlaneDot(End);
 
 	// Check if the line is completely on one side of the triangle, or if it's co-planar.
 #if 1
@@ -136,7 +136,7 @@ FORCEINLINE bool appLineCheckTriangle(const FVector4& Start, const FVector4& End
 	}
 
 	// Figure out when it will hit the triangle
-	float Time = -StartDist / (EndDist - StartDist);
+	float Time = float(-StartDist / (EndDist - StartDist));	// LWC_TODO: Precision loss
 
 	// If this triangle is not closer than the previous hit, reject it
 	if (Time < 0.f || Time >= IntersectionTime)
@@ -341,12 +341,13 @@ struct FFourBox
 	 */
 	void SetBox( int32 BoundingVolumeIndex, const FBox& Box )
 	{
-		Min[0].Component(BoundingVolumeIndex) = Box.Min.X;
-		Min[1].Component(BoundingVolumeIndex) = Box.Min.Y;
-		Min[2].Component(BoundingVolumeIndex) = Box.Min.Z;
-		Max[0].Component(BoundingVolumeIndex) = Box.Max.X;
-		Max[1].Component(BoundingVolumeIndex) = Box.Max.Y;
-		Max[2].Component(BoundingVolumeIndex) = Box.Max.Z;
+		using FVec4Real = decltype(FVector4::X);
+		Min[0].Component(BoundingVolumeIndex) = (FVec4Real)Box.Min.X;
+		Min[1].Component(BoundingVolumeIndex) = (FVec4Real)Box.Min.Y;
+		Min[2].Component(BoundingVolumeIndex) = (FVec4Real)Box.Min.Z;
+		Max[0].Component(BoundingVolumeIndex) = (FVec4Real)Box.Max.X;
+		Max[1].Component(BoundingVolumeIndex) = (FVec4Real)Box.Max.Y;
+		Max[2].Component(BoundingVolumeIndex) = (FVec4Real)Box.Max.Z;
 	}
 
 	/**
@@ -900,9 +901,9 @@ struct TkDOPNode
 			if ( SubIndex >= 0 )
 			{
 				bHit = true;
-				Check.LocalHitNormal.X = VectorGetComponentDynamic(TriangleSOA.Normals.X, SubIndex);
-				Check.LocalHitNormal.Y = VectorGetComponentDynamic(TriangleSOA.Normals.Y, SubIndex);
-				Check.LocalHitNormal.Z = VectorGetComponentDynamic(TriangleSOA.Normals.Z, SubIndex);
+				Check.LocalHitNormal.X = (decltype(FVector4::X))VectorGetComponentDynamic(TriangleSOA.Normals.X, SubIndex);
+				Check.LocalHitNormal.Y = (decltype(FVector4::Y))VectorGetComponentDynamic(TriangleSOA.Normals.Y, SubIndex);
+				Check.LocalHitNormal.Z = (decltype(FVector4::X))VectorGetComponentDynamic(TriangleSOA.Normals.Z, SubIndex);
 				Check.Result->Item = TriangleSOA.Payload[SubIndex];
 				Check.HitNodeIndex = History.GetOldestNode();
 
diff --git a/Engine/Source/Runtime/Experimental/Chaos/Public/Chaos/ChaosDebugDraw.h b/Engine/Source/Runtime/Experimental/Chaos/Public/Chaos/ChaosDebugDraw.h
--- a/Engine/Source/Runtime/Experimental/Chaos/Public/Chaos/ChaosDebugDraw.h
+++ b/Engine/Source/Runtime/Experimental/Chaos/Public/Chaos/ChaosDebugDraw.h
@@ -149,7 +149,7 @@ namespace Chaos
 		CHAOS_API void DrawSimulationSpace(const FSimulationSpace& SimSpace, const FChaosDebugDrawSettings* Settings = nullptr);
 		CHAOS_API void DrawShape(const FRigidTransform3& ShapeTransform, const FImplicitObject* Shape, const FColor& Color, const FChaosDebugDrawSettings* Settings = nullptr);
 		CHAOS_API void DrawConstraintGraph(const FRigidTransform3& ShapeTransform, const FPBDConstraintColor& Graph, const FChaosDebugDrawSettings* Settings = nullptr);
-		CHAOS_API void DrawCollidingShapes(const FRigidTransform3& SpaceTransform, const FPBDCollisionConstraints& Collisions, float ColorScale, const FChaosDebugDrawSettings* Settings = nullptr);
+		CHAOS_API void DrawCollidingShapes(const FRigidTransform3& SpaceTransform, const FPBDCollisionConstraints& Collisions, FRealSingle ColorScale, const FChaosDebugDrawSettings* Settings = nullptr);
 #endif
 	}
 }
diff --git a/Engine/Source/Runtime/Experimental/Chaos/Public/Chaos/DebugDrawQueue.h b/Engine/Source/Runtime/Experimental/Chaos/Public/Chaos/DebugDrawQueue.h
--- a/Engine/Source/Runtime/Experimental/Chaos/Public/Chaos/DebugDrawQueue.h
+++ b/Engine/Source/Runtime/Experimental/Chaos/Public/Chaos/DebugDrawQueue.h
@@ -93,8 +93,8 @@ struct CHAOS_API FLatentDrawCommand
 	float LifeTime;
 	uint8 DepthPriority;
 	float Thickness;
-	float Radius;
-	float HalfHeight;
+	FReal Radius;
+	FReal HalfHeight;
 	FVector Center;
 	FVector Extent;
 	FQuat Rotation;
@@ -169,7 +169,7 @@ struct CHAOS_API FLatentDrawCommand
 		return Command;
 	}
 
-	static FLatentDrawCommand DrawDebugSphere(const FVector& Center, float Radius, int32 Segments, const FColor& Color, bool bPersistentLines, float LifeTime, uint8 DepthPriority, float Thickness)
+	static FLatentDrawCommand DrawDebugSphere(const FVector& Center, FVector::FReal Radius, int32 Segments, const FColor& Color, bool bPersistentLines, float LifeTime, uint8 DepthPriority, float Thickness)
 	{
 		FLatentDrawCommand Command;
 		Command.LineStart = Center;
@@ -214,7 +214,7 @@ struct CHAOS_API FLatentDrawCommand
 		return Command;
 	}
 
-	static FLatentDrawCommand DrawDebugCircle(const FVector& Center, float Radius, int32 Segments, const FColor& Color, bool bPersistentLines, float LifeTime, uint8 DepthPriority, float Thickness, const FVector& YAxis, const FVector& ZAxis, bool bDrawAxis)
+	static FLatentDrawCommand DrawDebugCircle(const FVector& Center, FReal Radius, int32 Segments, const FColor& Color, bool bPersistentLines, float LifeTime, uint8 DepthPriority, float Thickness, const FVector& YAxis, const FVector& ZAxis, bool bDrawAxis)
 	{
 		FLatentDrawCommand Command;
 		Command.Center = Center;
@@ -232,7 +232,7 @@ struct CHAOS_API FLatentDrawCommand
 		return Command;
 	}
 
-	static FLatentDrawCommand DrawDebugCapsule(const FVector& Center, float HalfHeight, float Radius, const FQuat& Rotation, const FColor& Color, bool bPersistentLines, float LifeTime, uint8 DepthPriority, float Thickness)
+	static FLatentDrawCommand DrawDebugCapsule(const FVector& Center, FReal HalfHeight, FReal Radius, const FQuat& Rotation, const FColor& Color, bool bPersistentLines, float LifeTime, uint8 DepthPriority, float Thickness)
 	{
 		FLatentDrawCommand Command;
 		Command.Center = Center;
@@ -307,7 +307,7 @@ public:
 		}
 	}
 
-	void DrawDebugCoordinateSystem(const FVector& Position, const FRotator& AxisRot, float Scale, bool bPersistentLines = false, float LifeTime = -1.f, uint8 DepthPriority = 0, float Thickness = 0.f)
+	void DrawDebugCoordinateSystem(const FVector& Position, const FRotator& AxisRot, FReal Scale, bool bPersistentLines = false, float LifeTime = -1.f, uint8 DepthPriority = 0, float Thickness = 0.f)
 	{
 		if (IsDebugDrawingEnabled())
 		{
@@ -328,7 +328,7 @@ public:
 	}
 
 
-	void DrawDebugSphere(FVector const& Center, float Radius, int32 Segments, const FColor& Color, bool bPersistentLines = false, float LifeTime = -1.f, uint8 DepthPriority = 0, float Thickness = 0.f)
+	void DrawDebugSphere(FVector const& Center, FReal Radius, int32 Segments, const FColor& Color, bool bPersistentLines = false, float LifeTime = -1.f, uint8 DepthPriority = 0, float Thickness = 0.f)
 	{
 		if (IsDebugDrawingEnabled())
 		{
@@ -366,7 +366,7 @@ public:
 		}
 	}
 
-	void DrawDebugCircle(const FVector& Center, float Radius, int32 Segments, const FColor& Color, bool bPersistentLines, float LifeTime, uint8 DepthPriority, float Thickness, const FVector& YAxis, const FVector& ZAxis, bool bDrawAxis)
+	void DrawDebugCircle(const FVector& Center, FReal Radius, int32 Segments, const FColor& Color, bool bPersistentLines, float LifeTime, uint8 DepthPriority, float Thickness, const FVector& YAxis, const FVector& ZAxis, bool bDrawAxis)
 	{
 		if (IsDebugDrawingEnabled())
 		{
@@ -378,7 +378,7 @@ public:
 		}
 	}
 
-	void DrawDebugCapsule(const FVector& Center, float HalfHeight, float Radius, const FQuat& Rotation, const FColor& Color, bool bPersistentLines, float LifeTime, uint8 DepthPriority, float Thickness)
+	void DrawDebugCapsule(const FVector& Center, FReal HalfHeight, FReal Radius, const FQuat& Rotation, const FColor& Color, bool bPersistentLines, float LifeTime, uint8 DepthPriority, float Thickness)
 	{
 		if (IsDebugDrawingEnabled())
 		{
@@ -483,7 +483,7 @@ private:
 	FCriticalSection ConsumersCS;
 
 	FVector CenterOfInterest;
-	float RadiusOfInterest;
+	FReal RadiusOfInterest;
 	bool bEnableDebugDrawing;
 };
 }
diff --git a/Engine/Source/Runtime/Experimental/Chaos/Public/Chaos/Framework/PhysicsProxy.h b/Engine/Source/Runtime/Experimental/Chaos/Public/Chaos/Framework/PhysicsProxy.h
--- a/Engine/Source/Runtime/Experimental/Chaos/Public/Chaos/Framework/PhysicsProxy.h
+++ b/Engine/Source/Runtime/Experimental/Chaos/Public/Chaos/Framework/PhysicsProxy.h
@@ -65,12 +65,12 @@ public:
 
 	// Previously callback related functions, all called in the context of the physics thread if enabled.
 	bool IsSimulating() const { return static_cast<const Concrete*>(this)->IsSimulating(); }
-	void UpdateKinematicBodiesCallback(const FParticlesType& InParticles, const Chaos::FReal InDt, const Chaos::FReal InTime, FKinematicProxy& InKinematicProxy) { static_cast<Concrete*>(this)->UpdateKinematicBodiesCallback(InParticles, InDt, InTime, InKinematicProxy); }
-	void StartFrameCallback(const Chaos::FReal InDt, const Chaos::FReal InTime) { static_cast<Concrete*>(this)->StartFrameCallback(InDt, InTime); }
-	void EndFrameCallback(const Chaos::FReal InDt) { static_cast<Concrete*>(this)->EndFrameCallback(InDt); }
+	void UpdateKinematicBodiesCallback(const FParticlesType& InParticles, const float InDt, const float InTime, FKinematicProxy& InKinematicProxy) { static_cast<Concrete*>(this)->UpdateKinematicBodiesCallback(InParticles, InDt, InTime, InKinematicProxy); }
+	void StartFrameCallback(const float InDt, const float InTime) { static_cast<Concrete*>(this)->StartFrameCallback(InDt, InTime); }
+	void EndFrameCallback(const float InDt) { static_cast<Concrete*>(this)->EndFrameCallback(InDt); }
 	void CreateRigidBodyCallback(FParticlesType& InOutParticles) { static_cast<Concrete*>(this)->CreateRigidBodyCallback(InOutParticles); }
 	void DisableCollisionsCallback(TSet<TTuple<int32, int32>>& InPairs) { static_cast<Concrete*>(this)->DisableCollisionsCallback(InPairs); }
-	void AddForceCallback(FParticlesType& InParticles, const Chaos::FReal InDt, const int32 InIndex) { static_cast<Concrete*>(this)->AddForceCallback(InParticles, InDt, InIndex); }
+	void AddForceCallback(FParticlesType& InParticles, const float InDt, const int32 InIndex) { static_cast<Concrete*>(this)->AddForceCallback(InParticles, InDt, InIndex); }
 
 	/** The Particle Binding creates a connection between the particles in the simulation and the solver objects dataset. */
 	void BindParticleCallbackMapping(Chaos::TArrayCollectionArray<PhysicsProxyWrapper> & PhysicsProxyReverseMap, Chaos::TArrayCollectionArray<int32> & ParticleIDReverseMap) {static_cast<Concrete*>(this)->BindParticleCallbackMapping(PhysicsProxyReverseMap, ParticleIDReverseMap);}
diff --git a/Engine/Source/Runtime/Experimental/Chaos/Public/Chaos/ImplicitObject.h b/Engine/Source/Runtime/Experimental/Chaos/Public/Chaos/ImplicitObject.h
--- a/Engine/Source/Runtime/Experimental/Chaos/Public/Chaos/ImplicitObject.h
+++ b/Engine/Source/Runtime/Experimental/Chaos/Public/Chaos/ImplicitObject.h
@@ -364,7 +364,7 @@ protected:
 	FVec3 ScaledNormalHelper(const FVec3& Normal, const FVec3& Scale) const
 	{
 		const FVec3 ScaledNormal = Scale * Normal;
-		const float ScaledNormalLen = ScaledNormal.Size();
+		const FReal ScaledNormalLen = ScaledNormal.Size();
 		return ensure(ScaledNormalLen > TNumericLimits<FReal>::Min())
 			? ScaledNormal / ScaledNormalLen
 			: FVec3(0.f, 0.f, 1.f);
diff --git a/Engine/Source/Runtime/Experimental/Chaos/Public/Chaos/PBDCollisionConstraintsUtil.h b/Engine/Source/Runtime/Experimental/Chaos/Public/Chaos/PBDCollisionConstraintsUtil.h
--- a/Engine/Source/Runtime/Experimental/Chaos/Public/Chaos/PBDCollisionConstraintsUtil.h
+++ b/Engine/Source/Runtime/Experimental/Chaos/Public/Chaos/PBDCollisionConstraintsUtil.h
@@ -8,20 +8,20 @@
 namespace Chaos
 {
 	void CHAOS_API ComputeHashTable(const TArray<Chaos::FRigidBodyPointContactConstraint>& ConstraintsArray,
-									const FBox& BoundingBox, TMultiMap<int32, int32>& HashTableMap, const Chaos::FRealSingle SpatialHashRadius);
+									const FBox& BoundingBox, TMultiMap<int32, int32>& HashTableMap, const Chaos::FReal SpatialHashRadius);
 
 	void CHAOS_API ComputeHashTable(const TArray<FCollidingData>& CollisionsArray,
-									const FBox& BoundingBox, TMultiMap<int32, int32>& HashTableMap, const Chaos::FRealSingle SpatialHashRadius);
+									const FBox& BoundingBox, TMultiMap<int32, int32>& HashTableMap, const Chaos::FReal SpatialHashRadius);
 
 	void CHAOS_API ComputeHashTable(const TArray<FCollidingDataExt>& CollisionsArray,
-									const FBox& BoundingBox, TMultiMap<int32, int32>& HashTableMap, const Chaos::FRealSingle SpatialHashRadius);
+									const FBox& BoundingBox, TMultiMap<int32, int32>& HashTableMap, const Chaos::FReal SpatialHashRadius);
 
 	void CHAOS_API ComputeHashTable(const TArray<FVector>& ParticleArray, const FBox& BoundingBox, 
 									TMultiMap<int32, int32>& HashTableMap, const float SpatialHashRadius);
 
 	void CHAOS_API ComputeHashTable(const TArray<FBreakingData>& BreakingsArray,
-									const FBox& BoundingBox, TMultiMap<int32, int32>& HashTableMap, const Chaos::FRealSingle SpatialHashRadius);
+									const FBox& BoundingBox, TMultiMap<int32, int32>& HashTableMap, const Chaos::FReal SpatialHashRadius);
 
 	void CHAOS_API ComputeHashTable(const TArray<FBreakingDataExt>& BreakingsArray,
-									const FBox& BoundingBox, TMultiMap<int32, int32>& HashTableMap, const Chaos::FRealSingle SpatialHashRadius);
+									const FBox& BoundingBox, TMultiMap<int32, int32>& HashTableMap, const Chaos::FReal SpatialHashRadius);
 }
diff --git a/Engine/Source/Runtime/Experimental/Chaos/Public/Chaos/PBDRigidClustering.h b/Engine/Source/Runtime/Experimental/Chaos/Public/Chaos/PBDRigidClustering.h
--- a/Engine/Source/Runtime/Experimental/Chaos/Public/Chaos/PBDRigidClustering.h
+++ b/Engine/Source/Runtime/Experimental/Chaos/Public/Chaos/PBDRigidClustering.h
@@ -121,12 +121,12 @@ public:
 	*/
 	TSet<FPBDRigidParticleHandle*> ReleaseClusterParticles(
 		FPBDRigidClusteredParticleHandle* ClusteredParticle, 
-		const TMap<FGeometryParticleHandle*, float>* ExternalStrainMap = nullptr,
+		const TMap<FGeometryParticleHandle*, Chaos::FReal>* ExternalStrainMap = nullptr,
 		bool bForceRelease = false);
 
 	TSet<FPBDRigidParticleHandle*> ReleaseClusterParticlesNoInternalCluster(
 		FPBDRigidClusteredParticleHandle* ClusteredParticle,
-		const TMap<FGeometryParticleHandle*, float>* ExternalStrainMap = nullptr,
+		const TMap<FGeometryParticleHandle*, Chaos::FReal>* ExternalStrainMap = nullptr,
 		bool bForceRelease = false);
 
 	/*
@@ -169,7 +169,7 @@ public:
 	*    the children clusters.
 	*/
 	TMap<FPBDRigidClusteredParticleHandle*, TSet<FPBDRigidParticleHandle*>> BreakingModel(
-		TMap<FGeometryParticleHandle*, float>* ExternalStrainMap = nullptr);
+		TMap<FGeometryParticleHandle*, Chaos::FReal>* ExternalStrainMap = nullptr);
 
 	/**
 	*  PromoteStrains
diff --git a/Engine/Source/Runtime/Experimental/Chaos/Public/Chaos/PBDRigidsEvolution.h b/Engine/Source/Runtime/Experimental/Chaos/Public/Chaos/PBDRigidsEvolution.h
--- a/Engine/Source/Runtime/Experimental/Chaos/Public/Chaos/PBDRigidsEvolution.h
+++ b/Engine/Source/Runtime/Experimental/Chaos/Public/Chaos/PBDRigidsEvolution.h
@@ -673,7 +673,7 @@ public:
 				else
 				{
 					TargetPos = FVec3::Lerp(PrevX, KinematicTarget.GetTarget().GetLocation(), StepFraction);
-					TargetRot = FRotation3::Slerp(PrevR, KinematicTarget.GetTarget().GetRotation(), StepFraction);
+					TargetRot = FRotation3::Slerp(PrevR, KinematicTarget.GetTarget().GetRotation(), decltype(FQuat::X)(StepFraction));
 				}
 				if (Dt > MinDt)
 				{
diff --git a/Engine/Source/Runtime/Experimental/Chaos/Public/Chaos/Plane.h b/Engine/Source/Runtime/Experimental/Chaos/Public/Chaos/Plane.h
--- a/Engine/Source/Runtime/Experimental/Chaos/Public/Chaos/Plane.h
+++ b/Engine/Source/Runtime/Experimental/Chaos/Public/Chaos/Plane.h
@@ -321,8 +321,9 @@ TVector<T, 2> ComputeBarycentricInPlane(const TVector<T, d>& P0, const TVector<T
 	T ProjP1 = TVector<T, d>::DotProduct(PP0, P10);
 	T ProjP2 = TVector<T, d>::DotProduct(PP0, P20);
 	T Denom = Size10 * Size20 - ProjSides * ProjSides;
-	Bary.X = (Size20 * ProjP1 - ProjSides * ProjP2) / Denom;
-	Bary.Y = (Size10 * ProjP2 - ProjSides * ProjP1) / Denom;
+	using FVec2Real = decltype(Bary.X);
+	Bary.X = FVec2Real((Size20 * ProjP1 - ProjSides * ProjP2) / Denom);
+	Bary.Y = FVec2Real((Size10 * ProjP2 - ProjSides * ProjP1) / Denom);
 	return Bary;
 }
 
diff --git a/Engine/Source/Runtime/Experimental/Chaos/Public/Chaos/Rotation.h b/Engine/Source/Runtime/Experimental/Chaos/Public/Chaos/Rotation.h
--- a/Engine/Source/Runtime/Experimental/Chaos/Public/Chaos/Rotation.h
+++ b/Engine/Source/Runtime/Experimental/Chaos/Public/Chaos/Rotation.h
@@ -152,7 +152,8 @@ namespace Chaos
 		 */
 		static TRotation<FReal, 3> FromElements(const FReal X, const FReal Y, const FReal Z, const FReal W)
 		{
-			return FQuat(X, Y, Z, W);
+			using FQuatReal = decltype(FQuat::X);	// LWC_TODO: Remove once FQuat supports variants
+			return FQuat((FQuatReal)X, (FQuatReal)Y, (FQuatReal)Z, (FQuatReal)W);
 		}
 
 		/**
@@ -168,7 +169,7 @@ namespace Chaos
 		 */
 		static TRotation<FReal, 3> FromAxisAngle(const ::Chaos::TVector<FReal, 3>& Axis, const FReal AngleRad)
 		{
-			return FQuat(FVector(Axis.X, Axis.Y, Axis.Z), AngleRad);
+			return FQuat(FVector(Axis.X, Axis.Y, Axis.Z), (decltype(FQuat::X))AngleRad);
 		}
 
 		/**
@@ -176,14 +177,15 @@ namespace Chaos
 		 */
 		static TRotation<FReal, 3> FromVector(const ::Chaos::TVector<FReal, 3>& V)
 		{
+			using FQuatReal = decltype(FQuat::X);	// LWC_TODO: Remove once FQuat supports variants
 			TRotation<FReal, 3> Rot;
 			FReal HalfSize = 0.5f * V.Size();
 			FReal sinc = (FMath::Abs(HalfSize) > 1e-8) ? FMath::Sin(HalfSize) / HalfSize : 1;
 			auto RotV = 0.5f * sinc * V;
-			Rot.X = RotV.X;
-			Rot.Y = RotV.Y;
-			Rot.Z = RotV.Z;
-			Rot.W = FMath::Cos(HalfSize);
+			Rot.X = (FQuatReal)RotV.X;
+			Rot.Y = (FQuatReal)RotV.Y;
+			Rot.Z = (FQuatReal)RotV.Z;
+			Rot.W = (FQuatReal)FMath::Cos(HalfSize);
 			return Rot;
 		}
 
@@ -231,7 +233,7 @@ namespace Chaos
 			R1.EnforceShortestArcWith(R0);
 
 			// W = 2 * dQ/dT * Qinv
-			const TRotation<FReal, 3> DRDt = (R1 - R0) / InDt;
+			const TRotation<FReal, 3> DRDt = (R1 - R0) / (decltype(R0.X))InDt;
 			const TRotation<FReal, 3> RInv = Conjugate(R0);
 			const TRotation<FReal, 3> W = (DRDt * RInv) * 2.0f;
 
@@ -289,7 +291,7 @@ namespace Chaos
 		 */
 		static TRotation<FReal, 3> IntegrateRotationWithAngularVelocity(const TRotation<FReal, 3>& InR0, const TVector<FReal, 3>& InW, const FReal InDt)
 		{
-			TRotation<FReal, 3> R1 = InR0 + (TRotation<FReal, 3>::FromElements(InW.X, InW.Y, InW.Z, 0.f) * InR0) * (InDt * 0.5f);
+			TRotation<FReal, 3> R1 = InR0 + (TRotation<FReal, 3>::FromElements(InW.X, InW.Y, InW.Z, 0.f) * InR0) * decltype(InR0.X)(InDt * 0.5f);
 			return R1.GetNormalized();
 		}
 
diff --git a/Engine/Source/Runtime/Experimental/Chaos/Public/EventsData.h b/Engine/Source/Runtime/Experimental/Chaos/Public/EventsData.h
--- a/Engine/Source/Runtime/Experimental/Chaos/Public/EventsData.h
+++ b/Engine/Source/Runtime/Experimental/Chaos/Public/EventsData.h
@@ -13,8 +13,8 @@ namespace Chaos
 	// base class for data that requires time of creation to be recorded
 	struct FTimeResource
 	{
-		FTimeResource() : TimeCreated(-FLT_MAX) {}
-		float TimeCreated;
+		FTimeResource() : TimeCreated(-TNumericLimits<FReal>::Max()) {}
+		FReal TimeCreated;
 	};
 
 	typedef TArray<FCollidingData> FCollisionDataArray;
diff --git a/Engine/Source/Runtime/Experimental/Chaos/Public/Field/FieldSystem.h b/Engine/Source/Runtime/Experimental/Chaos/Public/Field/FieldSystem.h
--- a/Engine/Source/Runtime/Experimental/Chaos/Public/Field/FieldSystem.h
+++ b/Engine/Source/Runtime/Experimental/Chaos/Public/Field/FieldSystem.h
@@ -251,7 +251,7 @@ struct CHAOS_API FFieldContext
 	FFieldContext & operator =(FFieldContext&&) = delete;
 
 	FFieldContext(const TFieldArrayView< FFieldContextIndex >& SampleIndicesIn, const TFieldArrayView<FVector>& SamplePositionsIn,
-		const UniquePointerMap & MetaDataIn, const float TimeSecondsIn, TArray<FVector>& VectorResultsIn, TArray<float>& ScalarResultsIn, 
+		const UniquePointerMap & MetaDataIn, const Chaos::FReal TimeSecondsIn, TArray<FVector>& VectorResultsIn, TArray<float>& ScalarResultsIn,
 		TArray<int32>& IntegerResultsIn, TArray<FFieldContextIndex>& IndexResultsIn, TArray<FFieldContextIndex>& CullingResultsIn)
 		: SampleIndices(SampleIndicesIn)
 		, SamplePositions(SamplePositionsIn)
@@ -271,7 +271,7 @@ struct CHAOS_API FFieldContext
 		MetaData.Add(FFieldSystemMetaData::EMetaType::ECommandData_Culling, CullingData.Get());
 	}
 	FFieldContext(const TFieldArrayView< FFieldContextIndex >& SampleIndicesIn, const TFieldArrayView<FVector>& SamplePositionsIn,
-		const PointerMap & MetaDataIn, const float TimeSecondsIn, TArray<FVector>& VectorResultsIn, TArray<float>& ScalarResultsIn, 
+		const PointerMap & MetaDataIn, const Chaos::FReal TimeSecondsIn, TArray<FVector>& VectorResultsIn, TArray<float>& ScalarResultsIn,
 				TArray<int32>& IntegerResultsIn, TArray<FFieldContextIndex>& IndexResultsIn, TArray<FFieldContextIndex>& CullingResultsIn)
 		: SampleIndices(SampleIndicesIn)
 		, SamplePositions(SamplePositionsIn)
@@ -287,7 +287,7 @@ struct CHAOS_API FFieldContext
 	}
 
 	FFieldContext(FFieldExecutionDatas& ExecutionDatas,
-		const UniquePointerMap& MetaDataIn, const float TimeSecondsIn)
+		const UniquePointerMap& MetaDataIn, const Chaos::FReal TimeSecondsIn)
 		: SampleIndices(ExecutionDatas.SampleIndices, 0, ExecutionDatas.SampleIndices.Num())
 		, SamplePositions(ExecutionDatas.SamplePositions, 0, ExecutionDatas.SamplePositions.Num())
 		, TimeSeconds(TimeSecondsIn)
@@ -329,7 +329,7 @@ struct CHAOS_API FFieldContext
 
 	PointerMap MetaData;
 	TUniquePtr<FFieldSystemMetaDataCulling> CullingData;
-	float TimeSeconds;
+	Chaos::FReal TimeSeconds;
 
 	TArray<FVector>& VectorResults;
 	TArray<float>& ScalarResults;
diff --git a/Engine/Source/Runtime/Experimental/Chaos/Public/GeometryCollection/GeometryCollectionCollisionStructureManager.h b/Engine/Source/Runtime/Experimental/Chaos/Public/GeometryCollection/GeometryCollectionCollisionStructureManager.h
--- a/Engine/Source/Runtime/Experimental/Chaos/Public/GeometryCollection/GeometryCollectionCollisionStructureManager.h
+++ b/Engine/Source/Runtime/Experimental/Chaos/Public/GeometryCollection/GeometryCollectionCollisionStructureManager.h
@@ -81,7 +81,7 @@ public:
 		const Chaos::FParticles& MeshParticles,
 		const Chaos::FTriangleMesh& TriMesh,
 		const FBox& CollisionBounds,
-		const float Radius,
+		const Chaos::FReal Radius,
 		const int32 MinRes,
 		const int32 MaxRes,
 		const float CollisionObjectReduction,
@@ -94,13 +94,13 @@ public:
 		const ECollisionTypeEnum CollisionType);
 
 	static FImplicit* NewImplicitSphere(
-		const float Radius,
+		const Chaos::FReal Radius,
 		const float CollisionObjectReduction,
 		const ECollisionTypeEnum CollisionType);
 
 	static FImplicit* NewImplicitCapsule(
-		const float Radius,
-		const float Length,
+		const Chaos::FReal Radius,
+		const Chaos::FReal Length,
 		const float CollisionObjectReduction,
 		const ECollisionTypeEnum CollisionType);
 
@@ -125,11 +125,11 @@ public:
 
 	static FVector CalculateUnitMassInertiaTensor(
 		const FBox& BoundingBox,
-		const float Radius,
+		const Chaos::FReal Radius,
 		const EImplicitTypeEnum ImplicitType);
 
-	static float CalculateVolume(
+	static Chaos::FReal CalculateVolume(
 		const FBox& BoundingBox,
-		const float Radius,
+		const Chaos::FReal Radius,
 		const EImplicitTypeEnum ImplicitType);
 };
diff --git a/Engine/Source/Runtime/Experimental/Chaos/Public/GeometryCollection/GeometryCollectionProximityUtility.h b/Engine/Source/Runtime/Experimental/Chaos/Public/GeometryCollection/GeometryCollectionProximityUtility.h
--- a/Engine/Source/Runtime/Experimental/Chaos/Public/GeometryCollection/GeometryCollectionProximityUtility.h
+++ b/Engine/Source/Runtime/Experimental/Chaos/Public/GeometryCollection/GeometryCollectionProximityUtility.h
@@ -21,9 +21,9 @@ public:
 
 	struct FVertexPair
 	{
-		FVector Vertex1, Vertex2;
-		float Distance() { return (Vertex1 - Vertex2).Size(); }
-		float DistanceSquared() { return (Vertex1 - Vertex2).SizeSquared(); }
+		FVector Vertex1, Vertex2;	// World space vertex positions
+		FVector::FReal Distance() { return (Vertex1 - Vertex2).Size(); }
+		FVector::FReal DistanceSquared() { return (Vertex1 - Vertex2).SizeSquared(); }
 	};
 
 	struct FOverlappingFacePair
diff --git a/Engine/Source/Runtime/Experimental/Chaos/Public/GeometryCollection/ManagedArrayTypeValues.inl b/Engine/Source/Runtime/Experimental/Chaos/Public/GeometryCollection/ManagedArrayTypeValues.inl
--- a/Engine/Source/Runtime/Experimental/Chaos/Public/GeometryCollection/ManagedArrayTypeValues.inl
+++ b/Engine/Source/Runtime/Experimental/Chaos/Public/GeometryCollection/ManagedArrayTypeValues.inl
@@ -41,6 +41,7 @@ MANAGED_ARRAY_TYPE(Chaos::NotThreadSafeSharedPtr_FImplicitObject, FImplicitObjec
 MANAGED_ARRAY_TYPE(Chaos::FPBDRigidClusteredParticleHandle*, TPBDRigidClusteredParticleHandle3fPtr)
 MANAGED_ARRAY_TYPE(TUniquePtr<Chaos::FConvex>, FConvexUniquePtr)
 MANAGED_ARRAY_TYPE(TArray<FVector2D>, Vector2DArray)
+MANAGED_ARRAY_TYPE(double, Double)
 
 // NOTE: new types must be added at the bottom to keep serialization from breaking
 
diff --git a/Engine/Source/Runtime/Experimental/Chaos/Public/GeometryCollectionProxyData.h b/Engine/Source/Runtime/Experimental/Chaos/Public/GeometryCollectionProxyData.h
--- a/Engine/Source/Runtime/Experimental/Chaos/Public/GeometryCollectionProxyData.h
+++ b/Engine/Source/Runtime/Experimental/Chaos/Public/GeometryCollectionProxyData.h
@@ -143,7 +143,7 @@ public:
 		ParticleToWorldTransforms.SetNumUninitialized(NumTransforms);
 	}
 
-	float SolverDt;
+	Chaos::FReal SolverDt;
 	int32 BaseIndex;
 	int32 NumParticlesAdded;
 	TArray<bool> DisabledStates;
diff --git a/Engine/Source/Runtime/Experimental/Chaos/Public/PhysicsProxy/FieldSystemProxyHelper.h b/Engine/Source/Runtime/Experimental/Chaos/Public/PhysicsProxy/FieldSystemProxyHelper.h
--- a/Engine/Source/Runtime/Experimental/Chaos/Public/PhysicsProxy/FieldSystemProxyHelper.h
+++ b/Engine/Source/Runtime/Experimental/Chaos/Public/PhysicsProxy/FieldSystemProxyHelper.h
@@ -253,7 +253,7 @@ namespace Chaos
 	 * @param    Rigidsolver Rigid solver owning the breaking model
 	 * @param    ExternalStrain Strain to be used to update the breaking model
 	 */
-	static void UpdateSolverBreakingModel(Chaos::FPBDRigidsSolver* RigidSolver, TMap<Chaos::FGeometryParticleHandle*, float>& ExternalStrain)
+	static void UpdateSolverBreakingModel(Chaos::FPBDRigidsSolver* RigidSolver, TMap<Chaos::FGeometryParticleHandle*, Chaos::FReal>& ExternalStrain)
 	{
 		// Capture the results from the breaking model to post-process
 		TMap<Chaos::FPBDRigidClusteredParticleHandle*, TSet<Chaos::FPBDRigidParticleHandle*>> BreakResults =
@@ -279,7 +279,7 @@ namespace Chaos
 	 * @param    RigidHandle Particle handle on which the threshold will be updated
 	 * @param    ResultThreshold Threshoild to be set onto the handle
 	 */
-	static void UpdateMaterialSleepingThreshold(Chaos::FPBDRigidsSolver* RigidSolver, Chaos::FPBDRigidParticleHandle* RigidHandle, const float ResultThreshold)
+	static void UpdateMaterialSleepingThreshold(Chaos::FPBDRigidsSolver* RigidSolver, Chaos::FPBDRigidParticleHandle* RigidHandle, const Chaos::FReal ResultThreshold)
 	{
 		// if no per particle physics material is set, make one
 		if (!RigidSolver->GetEvolution()->GetPerParticlePhysicsMaterial(RigidHandle).IsValid())
@@ -309,7 +309,7 @@ namespace Chaos
 	 * @param    RigidHandle Particle handle on which the threshold will be updated
 	 * @param    ResultThreshold Threshoild to be set onto the handle
 	 */
-	static void UpdateMaterialDisableThreshold(Chaos::FPBDRigidsSolver* RigidSolver, Chaos::FPBDRigidParticleHandle* RigidHandle, const float ResultThreshold)
+	static void UpdateMaterialDisableThreshold(Chaos::FPBDRigidsSolver* RigidSolver, Chaos::FPBDRigidParticleHandle* RigidHandle, const Chaos::FReal ResultThreshold)
 	{
 		// if no per particle physics material is set, make one
 		if (!RigidSolver->GetEvolution()->GetPerParticlePhysicsMaterial(RigidHandle).IsValid())
@@ -462,7 +462,7 @@ namespace Chaos
 		{
 			SCOPE_CYCLE_COUNTER(STAT_ParamUpdateField_ExternalClusterStrain);
 			{
-				TMap<Chaos::FGeometryParticleHandle*, float> ExternalStrain;
+				TMap<Chaos::FGeometryParticleHandle*, Chaos::FReal> ExternalStrain;
 
 				static_cast<const FFieldNode<float>*>(FieldCommand.RootNode.Get())->Evaluate(FieldContext, ResultsView);
 				for (const FFieldContextIndex& Index : FieldContext.GetEvaluatedSamples())
diff --git a/Engine/Source/Runtime/Experimental/Chaos/Public/PhysicsProxy/PerSolverFieldSystem.h b/Engine/Source/Runtime/Experimental/Chaos/Public/PhysicsProxy/PerSolverFieldSystem.h
--- a/Engine/Source/Runtime/Experimental/Chaos/Public/PhysicsProxy/PerSolverFieldSystem.h
+++ b/Engine/Source/Runtime/Experimental/Chaos/Public/PhysicsProxy/PerSolverFieldSystem.h
@@ -54,7 +54,7 @@ public:
 	 * EFieldPhysicsType::Field_AngularVelocity
 	 * EFieldPhysicsType::Field_AngularrTorque
 	 */
-	void ComputeFieldRigidImpulse(const float SolverTime);
+	void ComputeFieldRigidImpulse(const Chaos::FReal SolverTime);
 
 	/**
 	 * Compute field linear velocity/force given a list of samples (positions + indices)
@@ -63,7 +63,7 @@ public:
 	 * EFieldPhysicsType::Field_LinearVelocity
 	 * EFieldPhysicsType::Field_LinearForce
 	 */
-	void ComputeFieldLinearImpulse(const float SolverTime);
+	void ComputeFieldLinearImpulse(const Chaos::FReal SolverTime);
 
 	/** Add the transient field command */
 	void AddTransientCommand(const FFieldSystemCommand& FieldCommand);
diff --git a/Engine/Source/Runtime/Experimental/ChaosCore/Public/Chaos/Vector.h b/Engine/Source/Runtime/Experimental/ChaosCore/Public/Chaos/Vector.h
--- a/Engine/Source/Runtime/Experimental/ChaosCore/Public/Chaos/Vector.h
+++ b/Engine/Source/Runtime/Experimental/ChaosCore/Public/Chaos/Vector.h
@@ -369,6 +369,7 @@ namespace Chaos
 	class TVector<FReal, 4> : public FVector4
 	{
 	public:
+		using FElement = decltype(FVector4::X);
 		using FVector4::W;
 		using FVector4::X;
 		using FVector4::Y;
@@ -388,6 +389,7 @@ namespace Chaos
 	class TVector<FReal, 3> : public UE::Math::TVector<FReal>
 	{
 	public:
+		using FElement = FReal;
 		using UE::Math::TVector<FReal>::X;
 		using UE::Math::TVector<FReal>::Y;
 		using UE::Math::TVector<FReal>::Z;
@@ -587,7 +589,7 @@ namespace Chaos
 	class TVector<FReal, 2> : public FVector2D
 	{
 	public:
-
+		using FElement = decltype(FVector2D::X);
 		using FVector2D::X;
 		using FVector2D::Y;
 
@@ -607,8 +609,8 @@ namespace Chaos
 		template <typename OtherT>
 		TVector(const TVector<OtherT, 2>& InVector)
 		{
-			X = ((FReal)InVector[0]);
-			Y = ((FReal)InVector[1]);
+			X = ((decltype(X))InVector[0]);	// LWC_TODO: Remove casts once FVector2D supports variants
+			Y = ((decltype(Y))InVector[1]);
 		}
 		void Write(std::ostream& Stream) const
 		{
@@ -673,6 +675,8 @@ namespace Chaos
 	class TVector<T, 3>
 	{
 	public:
+		using FElement = T;
+
 		FORCEINLINE TVector() {}
 		FORCEINLINE explicit TVector(T InX)
 		    : X(InX), Y(InX), Z(InX) {}
@@ -967,12 +971,12 @@ namespace Chaos
 	template<typename T, int d>
 	inline FArchive& SerializeReal(FArchive& Ar, TVector<T, d>& ValueIn)
 	{
-		static_assert(TIsSame<T, float>::Value || TIsSame<T, double>::Value, "only float or double are supoprted by this function");
+		static_assert(TIsSame<T, float>::Value || TIsSame<T, double>::Value, "only float or double are supported by this function");
 		for (int32 Idx = 0; Idx < d; ++Idx)
 		{
 			FRealSingle RealSingle = (FRealSingle)ValueIn[Idx];
 			Ar << RealSingle;
-			ValueIn[Idx] = (T)RealSingle;
+			ValueIn[Idx] = (typename TVector<T, d>::FElement)RealSingle;
 		}
 		return Ar;
 	}
