Reactivate the geocollection initialization fields

diff --git a/Engine/Source/Runtime/Experimental/GeometryCollectionEngine/Private/GeometryCollection/GeometryCollectionComponent.cpp b/Engine/Source/Runtime/Experimental/GeometryCollectionEngine/Private/GeometryCollection/GeometryCollectionComponent.cpp
--- a/Engine/Source/Runtime/Experimental/GeometryCollectionEngine/Private/GeometryCollection/GeometryCollectionComponent.cpp
+++ b/Engine/Source/Runtime/Experimental/GeometryCollectionEngine/Private/GeometryCollection/GeometryCollectionComponent.cpp
@@ -1918,6 +1918,7 @@ void UGeometryCollectionComponent::RegisterAndInitializePhysicsProxy()
 		{
 			SimulationParameters.PhysicalMaterialHandle = EnginePhysicalMaterial->GetPhysicsMaterial();
 		}
+		GetInitializationCommands(SimulationParameters.InitializationCommands);
 	}
 
 	PhysicsProxy = new FGeometryCollectionPhysicsProxy(this, *DynamicCollection, SimulationParameters, InitialQueryFilter, InitialSimFilter);
