ToolsFramework: add register and unregister helper functions to TransformGizmoUtil, to facilitate EdModes sharing the context object
#rnx

diff --git a/Engine/Plugins/Experimental/ModelingToolsEditorMode/Source/ModelingToolsEditorMode/Private/ModelingToolsEditorMode.cpp b/Engine/Plugins/Experimental/ModelingToolsEditorMode/Source/ModelingToolsEditorMode/Private/ModelingToolsEditorMode.cpp
--- a/Engine/Plugins/Experimental/ModelingToolsEditorMode/Source/ModelingToolsEditorMode/Private/ModelingToolsEditorMode.cpp
+++ b/Engine/Plugins/Experimental/ModelingToolsEditorMode/Source/ModelingToolsEditorMode/Private/ModelingToolsEditorMode.cpp
@@ -305,8 +305,7 @@ void UModelingToolsEditorMode::Enter()
 	StylusStateTracker = MakeUnique<FStylusStateTracker>();
 
 	// register gizmo helper
-	ModelingGizmoHelper = NewObject<UTransformGizmoContextObject>(ToolsContext->ToolManager);
-	ModelingGizmoHelper->RegisterGizmosWithManager(ToolsContext->ToolManager);
+	UE::TransformGizmoUtil::RegisterTransformGizmoContextObject(ToolsContext.Get());
 
 	const FModelingToolsManagerCommands& ToolManagerCommands = FModelingToolsManagerCommands::Get();
 
@@ -676,7 +675,9 @@ void UModelingToolsEditorMode::Exit()
 
 	StylusStateTracker = nullptr;
 
-	ModelingGizmoHelper->DeregisterGizmosWithManager(ToolsContext->ToolManager);
+	// TODO: cannot deregister currently because if another mode is also registering, it's Enter()
+	// will be called before our Exit()
+	//UE::TransformGizmoUtil::DeregisterTransformGizmoContextObject(ToolsContext.Get());
 	
 	FModelingModeActionCommands::UnRegisterCommandBindings(Toolkit->GetToolkitCommands());
 
diff --git a/Engine/Plugins/Experimental/ModelingToolsEditorMode/Source/ModelingToolsEditorMode/Public/ModelingToolsEditorMode.h b/Engine/Plugins/Experimental/ModelingToolsEditorMode/Source/ModelingToolsEditorMode/Public/ModelingToolsEditorMode.h
--- a/Engine/Plugins/Experimental/ModelingToolsEditorMode/Source/ModelingToolsEditorMode/Public/ModelingToolsEditorMode.h
+++ b/Engine/Plugins/Experimental/ModelingToolsEditorMode/Source/ModelingToolsEditorMode/Public/ModelingToolsEditorMode.h
@@ -17,7 +17,6 @@ class FEditorToolAssetAPI;
 class FUICommandList;
 class FStylusStateTracker;		// for stylus events
 class FModelingModeAssetAPI;
-class UTransformGizmoContextObject;
 
 UCLASS(Transient)
 class UModelingToolsEditorMode : public UBaseLegacyWidgetEdMode
@@ -69,8 +68,6 @@ protected:
 
 	TUniquePtr<FStylusStateTracker> StylusStateTracker;
 
-	TObjectPtr<UTransformGizmoContextObject> ModelingGizmoHelper;
-
 	void ModelingModeShortcutRequested(EModelingModeActionCommands Command);
 	void FocusCameraAtCursorHotkey();
 
diff --git a/Engine/Source/Runtime/Experimental/InteractiveToolsFramework/Private/BaseGizmos/TransformGizmoUtil.cpp b/Engine/Source/Runtime/Experimental/InteractiveToolsFramework/Private/BaseGizmos/TransformGizmoUtil.cpp
--- a/Engine/Source/Runtime/Experimental/InteractiveToolsFramework/Private/BaseGizmos/TransformGizmoUtil.cpp
+++ b/Engine/Source/Runtime/Experimental/InteractiveToolsFramework/Private/BaseGizmos/TransformGizmoUtil.cpp
@@ -3,6 +3,7 @@
 #include "BaseGizmos/TransformGizmoUtil.h"
 
 #include "ContextObjectStore.h"
+#include "InteractiveToolsContext.h"
 #include "InteractiveGizmoManager.h"
 #include "InteractiveToolManager.h"
 
@@ -83,6 +84,48 @@ void UTransformGizmoContextObject::DeregisterGizmosWithManager(UInteractiveToolM
 
 
 
+
+bool UE::TransformGizmoUtil::RegisterTransformGizmoContextObject(UInteractiveToolsContext* ToolsContext)
+{
+	if (ensure(ToolsContext))
+	{
+		UTransformGizmoContextObject* Found = ToolsContext->ContextObjectStore->FindContext<UTransformGizmoContextObject>();
+		if (Found == nullptr)
+		{
+			UTransformGizmoContextObject* GizmoHelper = NewObject<UTransformGizmoContextObject>(ToolsContext->ToolManager);
+			if (ensure(GizmoHelper))
+			{
+				GizmoHelper->RegisterGizmosWithManager(ToolsContext->ToolManager);
+				return true;
+			}
+			else
+			{
+				return false;
+			}
+		}
+		return true;
+	}
+	return false;
+}
+
+
+bool UE::TransformGizmoUtil::DeregisterTransformGizmoContextObject(UInteractiveToolsContext* ToolsContext)
+{
+	if (ensure(ToolsContext))
+	{
+		UTransformGizmoContextObject* Found = ToolsContext->ContextObjectStore->FindContext<UTransformGizmoContextObject>();
+		if (Found != nullptr)
+		{
+			Found->DeregisterGizmosWithManager(ToolsContext->ToolManager);
+			ToolsContext->ContextObjectStore->RemoveContextObject(Found);
+		}
+		return true;
+	}
+	return false;
+}
+
+
+
 UTransformGizmo* UTransformGizmoContextObject::Create3AxisTransformGizmo(UInteractiveGizmoManager* GizmoManager, void* Owner, const FString& InstanceIdentifier)
 {
 	if (ensure(bDefaultGizmosRegistered))
diff --git a/Engine/Source/Runtime/Experimental/InteractiveToolsFramework/Public/BaseGizmos/TransformGizmoUtil.h b/Engine/Source/Runtime/Experimental/InteractiveToolsFramework/Public/BaseGizmos/TransformGizmoUtil.h
--- a/Engine/Source/Runtime/Experimental/InteractiveToolsFramework/Public/BaseGizmos/TransformGizmoUtil.h
+++ b/Engine/Source/Runtime/Experimental/InteractiveToolsFramework/Public/BaseGizmos/TransformGizmoUtil.h
@@ -8,6 +8,7 @@
 
 #include "TransformGizmoUtil.generated.h"
 
+class UInteractiveToolsContext;
 class UInteractiveToolManager;
 class UInteractiveGizmoManager;
 class FTransformGizmoActorFactory;
@@ -22,6 +23,20 @@ namespace TransformGizmoUtil
 	// that is registered as a ContextStoreObject in an InteractiveToolsContext
 	//
 
+	/**
+	 * If one does not already exist, create a new instance of UTransformGizmoContextObject and add it to the
+	 * ToolsContext's ContextObjectStore
+	 * @return true if the ContextObjectStore now has a UTransformGizmoContextObject (whether it already existed, or was created)
+	 */
+	INTERACTIVETOOLSFRAMEWORK_API bool RegisterTransformGizmoContextObject(UInteractiveToolsContext* ToolsContext);
+
+	/**
+	 * Remove any existing UTransformGizmoContextObject from the ToolsContext's ContextObjectStore
+	 * @return true if the ContextObjectStore no longer has a UTransformGizmoContextObject (whether it was removed, or did not exist)
+	 */
+	INTERACTIVETOOLSFRAMEWORK_API bool DeregisterTransformGizmoContextObject(UInteractiveToolsContext* ToolsContext);
+
+
 	/**
 	 * Spawn a new standard 3-axis Transform gizmo (see UTransformGizmoContextObject::Create3AxisTransformGizmo for details)
 	 * GizmoManager's ToolsContext must have a UTransformGizmoContextObject registered (see UTransformGizmoContextObject for details)
