ToolsFramework: add a mechanism for clients of UInteractiveToolsContext to provide alternative implementations of the core ITF objects (eg InputRouter, ToolManager, etc) without having to subclass and duplicate all the Initialize()/Shutdown() logic.
#rnx

diff --git a/Engine/Source/Runtime/Experimental/InteractiveToolsFramework/Private/InteractiveToolsContext.cpp b/Engine/Source/Runtime/Experimental/InteractiveToolsFramework/Private/InteractiveToolsContext.cpp
--- a/Engine/Source/Runtime/Experimental/InteractiveToolsFramework/Private/InteractiveToolsContext.cpp
+++ b/Engine/Source/Runtime/Experimental/InteractiveToolsFramework/Private/InteractiveToolsContext.cpp
@@ -9,47 +9,148 @@ UInteractiveToolsContext::UInteractiveToolsContext()
 {
 	InputRouter = nullptr;
 	ToolManager = nullptr;
+	GizmoManager = nullptr;
 	TargetManager = nullptr;
+	ContextObjectStore = nullptr;
 	ToolManagerClass = UInteractiveToolManager::StaticClass();
+
+	SetCreateInputRouterFunc([](const FContextInitInfo& ContextInfo)
+	{
+		UInputRouter* NewInputRouter = NewObject<UInputRouter>(ContextInfo.ToolsContext);
+		NewInputRouter->Initialize(ContextInfo.TransactionsAPI);
+		return NewInputRouter;
+	});
+	SetShutdownInputRouterFunc([](UInputRouter* InputRouterIn)
+	{
+		InputRouterIn->ForceTerminateAll();
+		InputRouterIn->Shutdown();
+	});
+
+	SetCreateToolManagerFunc([this](const FContextInitInfo& ContextInfo)
+	{
+		UInteractiveToolManager* NewToolManager = NewObject<UInteractiveToolManager>(ContextInfo.ToolsContext, this->ToolManagerClass.Get());
+		NewToolManager->Initialize(ContextInfo.QueriesAPI, ContextInfo.TransactionsAPI, ContextInfo.InputRouter);
+		return NewToolManager;
+	});
+	SetShutdownToolManagerFunc([](UInteractiveToolManager* ToolManagerIn)
+	{
+		ToolManagerIn->Shutdown();
+	});
+
+	SetCreateToolTargetManagerFunc([](const FContextInitInfo& ContextInfo)
+	{
+		UToolTargetManager* NewTargetManager = NewObject<UToolTargetManager>(ContextInfo.ToolsContext);
+		NewTargetManager->Initialize();
+		return NewTargetManager;
+	});
+	SetShutdownToolTargetManagerFunc([](UToolTargetManager* TargetManagerIn)
+	{
+	});
+
+	SetCreateGizmoManagerFunc([](const FContextInitInfo& ContextInfo)
+	{
+		UInteractiveGizmoManager* NewGizmoManager = NewObject<UInteractiveGizmoManager>(ContextInfo.ToolsContext);
+		NewGizmoManager->Initialize(ContextInfo.QueriesAPI, ContextInfo.TransactionsAPI, ContextInfo.InputRouter);
+		NewGizmoManager->RegisterDefaultGizmos();
+		return NewGizmoManager;
+	});
+	SetShutdownGizmoManagerFunc([](UInteractiveGizmoManager* GizmoManagerIn)
+	{
+		GizmoManagerIn->Shutdown();
+	});
+
+	SetCreateContextStoreFunc([](const FContextInitInfo& ContextInfo)
+	{
+		UContextObjectStore* NewContextStore = NewObject<UContextObjectStore>(ContextInfo.ToolsContext);
+		return NewContextStore;
+	});
+	SetShutdownContextStoreFunc([](UContextObjectStore* ContextStoreIn)
+	{
+		ContextStoreIn->Shutdown();
+	});
 }
 
-void UInteractiveToolsContext::Initialize(IToolsContextQueriesAPI* QueriesAPI, IToolsContextTransactionsAPI* TransactionsAPI)
+
+void UInteractiveToolsContext::SetCreateInputRouterFunc(TUniqueFunction<UInputRouter* (const FContextInitInfo&)> Func)
+{
+	CreateInputRouterFunc = MoveTemp(Func);
+}
+void UInteractiveToolsContext::SetCreateToolManagerFunc(TUniqueFunction<UInteractiveToolManager* (const FContextInitInfo&)> Func)
+{
+	CreateToolManagerFunc = MoveTemp(Func);
+}
+void UInteractiveToolsContext::SetCreateToolTargetManagerFunc(TUniqueFunction<UToolTargetManager* (const FContextInitInfo&)> Func)
+{
+	CreateToolTargetManagerFunc = MoveTemp(Func);
+}
+void UInteractiveToolsContext::SetCreateGizmoManagerFunc(TUniqueFunction<UInteractiveGizmoManager* (const FContextInitInfo&)> Func)
 {
-	InputRouter = NewObject<UInputRouter>(this);
-	InputRouter->Initialize(TransactionsAPI);
+	CreateGizmoManagerFunc = MoveTemp(Func);
+}
+void UInteractiveToolsContext::SetCreateContextStoreFunc(TUniqueFunction<UContextObjectStore* (const FContextInitInfo&)> Func)
+{
+	CreateContextStoreFunc = MoveTemp(Func);
+}
 
-	ToolManager = NewObject<UInteractiveToolManager>(this, ToolManagerClass.Get());
-	ToolManager->Initialize(QueriesAPI, TransactionsAPI, InputRouter);
 
-	TargetManager = NewObject<UToolTargetManager>(this);
-	TargetManager->Initialize();
+void UInteractiveToolsContext::SetShutdownInputRouterFunc(TUniqueFunction<void(UInputRouter*)> Func)
+{
+	ShutdownInputRouterFunc = MoveTemp(Func);
+}
+void UInteractiveToolsContext::SetShutdownToolManagerFunc(TUniqueFunction<void(UInteractiveToolManager*)> Func)
+{
+	ShutdownToolManagerFunc = MoveTemp(Func);
+}
+void UInteractiveToolsContext::SetShutdownToolTargetManagerFunc(TUniqueFunction<void(UToolTargetManager*)> Func)
+{
+	ShutdownToolTargetManagerFunc = MoveTemp(Func);
+}
+void UInteractiveToolsContext::SetShutdownGizmoManagerFunc(TUniqueFunction<void(UInteractiveGizmoManager*)> Func)
+{
+	ShutdownGizmoManagerFunc = MoveTemp(Func);
+}
+void UInteractiveToolsContext::SetShutdownContextStoreFunc(TUniqueFunction<void(UContextObjectStore*)> Func)
+{
+	ShutdownContextStoreFunc = MoveTemp(Func);
+}
+
 
-	GizmoManager = NewObject<UInteractiveGizmoManager>(this);
-	GizmoManager->Initialize(QueriesAPI, TransactionsAPI, InputRouter);
+void UInteractiveToolsContext::Initialize(IToolsContextQueriesAPI* QueriesAPI, IToolsContextTransactionsAPI* TransactionsAPI)
+{
+	FContextInitInfo InitInfo;
+	InitInfo.ToolsContext = this;
+	InitInfo.QueriesAPI = QueriesAPI;
+	InitInfo.TransactionsAPI = TransactionsAPI;
 
-	GizmoManager->RegisterDefaultGizmos();
+	InputRouter = CreateInputRouterFunc(InitInfo);
+	InitInfo.InputRouter = InputRouter;
 
-	ContextObjectStore = NewObject<UContextObjectStore>(this);
+	ToolManager = CreateToolManagerFunc(InitInfo);
+	TargetManager = CreateToolTargetManagerFunc(InitInfo);
+	GizmoManager = CreateGizmoManagerFunc(InitInfo);
+	ContextObjectStore = CreateContextStoreFunc(InitInfo);
 }
 
 
 void UInteractiveToolsContext::Shutdown()
 {
-	// force-terminate any remaining captures/hovers/etc
-	InputRouter->ForceTerminateAll();
-	InputRouter->Shutdown();
+	ShutdownInputRouterFunc(InputRouter);
 	InputRouter = nullptr;
 
-	GizmoManager->Shutdown();
+	ShutdownGizmoManagerFunc(GizmoManager);
 	GizmoManager = nullptr;
 
-	ToolManager->Shutdown();
+	ShutdownToolManagerFunc(ToolManager);
 	ToolManager = nullptr;
 
-	ContextObjectStore->Shutdown();
+	ShutdownToolTargetManagerFunc(TargetManager);
+	TargetManager = nullptr;
+
+	ShutdownContextStoreFunc(ContextObjectStore);
 	ContextObjectStore = nullptr;
 }
 
+
 void UInteractiveToolsContext::DeactivateActiveTool(EToolSide WhichSide, EToolShutdownType ShutdownType)
 {
 	ToolManager->DeactivateTool(WhichSide, ShutdownType);
diff --git a/Engine/Source/Runtime/Experimental/InteractiveToolsFramework/Public/InteractiveToolsContext.h b/Engine/Source/Runtime/Experimental/InteractiveToolsFramework/Public/InteractiveToolsContext.h
--- a/Engine/Source/Runtime/Experimental/InteractiveToolsFramework/Public/InteractiveToolsContext.h
+++ b/Engine/Source/Runtime/Experimental/InteractiveToolsFramework/Public/InteractiveToolsContext.h
@@ -10,10 +10,14 @@
 class UToolTargetManager;
 class UContextObjectStore;
 
+
 /**
- * InteractiveToolsContext owns a ToolManager and an InputRouter. This is just a top-level 
- * UObject container, however implementations like UEdModeInteractiveToolsContext extend
- * this class to make it easier to connect external systems (like an FEdMode) to the ToolsFramework.
+ * InteractiveToolsContext owns the core parts of an Interactive Tools Framework implementation - the
+ * InputRouter, ToolManager, GizmoManager, TargetManager, and ContextStore. In the simplest
+ * use case, UInteractiveToolsContext is just a top-level container that will keep the various UObjects
+ * alive, and provide an easy way to access them. However in a more complex situation it may be
+ * desirable to subclass and extend the ToolsContext. For example, UEdModeInteractiveToolsContext allows
+ * a ToolsContext to live within a UEdMode.
  */
 UCLASS(Transient)
 class INTERACTIVETOOLSFRAMEWORK_API UInteractiveToolsContext : public UObject
@@ -23,8 +27,36 @@ class INTERACTIVETOOLSFRAMEWORK_API UInteractiveToolsContext : public UObject
 public:
 	UInteractiveToolsContext();
 
+	// FContextInitInfo is used by Initialize() to pass information to the various creator functions below.
+	// See ::UInteractiveToolsContext() and ::Initialize() for how these fields are used/initialized.
+	// In particular, InputRouter will not be set until after the UInputRouter has been created
+	struct FContextInitInfo
+	{
+		UInteractiveToolsContext* ToolsContext = nullptr;
+		IToolsContextQueriesAPI* QueriesAPI = nullptr;
+		IToolsContextTransactionsAPI* TransactionsAPI = nullptr;
+		UInputRouter* InputRouter = nullptr;
+	};
+
+	//
+	// Replace the internal functions that are called to create/destroy the sub-objects owned by the ITC.
+	// By doing this, clients of the ITF can provide their own subclass implementations, or do other
+	// custom setup/teardown as necessary. See comments above CreateInputRouterFunc below
+	//
+	virtual void SetCreateInputRouterFunc(TUniqueFunction<UInputRouter* (const FContextInitInfo&)> Func);
+	virtual void SetCreateToolManagerFunc(TUniqueFunction<UInteractiveToolManager* (const FContextInitInfo&)> Func);
+	virtual void SetCreateToolTargetManagerFunc(TUniqueFunction<UToolTargetManager* (const FContextInitInfo&)> Func);
+	virtual void SetCreateGizmoManagerFunc(TUniqueFunction<UInteractiveGizmoManager* (const FContextInitInfo&)> Func);
+	virtual void SetCreateContextStoreFunc(TUniqueFunction<UContextObjectStore* (const FContextInitInfo&)> Func);
+
+	virtual void SetShutdownInputRouterFunc(TUniqueFunction<void(UInputRouter*)> Func);
+	virtual void SetShutdownToolManagerFunc(TUniqueFunction<void(UInteractiveToolManager*)> Func);
+	virtual void SetShutdownToolTargetManagerFunc(TUniqueFunction<void(UToolTargetManager*)> Func);
+	virtual void SetShutdownGizmoManagerFunc(TUniqueFunction<void(UInteractiveGizmoManager*)> Func);
+	virtual void SetShutdownContextStoreFunc(TUniqueFunction<void(UContextObjectStore*)> Func);
+
 	/** 
-	 * Initialize the Context. This creates the InputRouter and ToolManager 
+	 * Initialize the Context. This creates the InputRouter, ToolManager, GizmoManager, TargetManager, and ContextStore
 	 * @param QueriesAPI client-provided implementation of the API for querying the higher-evel scene state
 	 * @param TransactionsAPI client-provided implementation of the API for publishing events and transactions
 	 */
@@ -81,6 +113,30 @@ public:
 	TObjectPtr<UContextObjectStore> ContextObjectStore;
 
 protected:
+
+	// Initialize() calls these functions to create the main child objects needed to operate
+	// the Tools Framework - InputRouter, ToolManager, GizmoManager, TargetManager, ContextStore.
+	// Default implementations are set up in the UInteractiveToolsContext() constructor, however
+	// users of the framework are free to replace these with subclasses, or do more complex
+	// initialization, by replacing these functions with their own versions before calling Initialize().
+	// (the same could be accomplished by subclassing and overriding Initialize(), however using these
+	//  lambda functions will be a safer alternative in many cases)
+	TUniqueFunction<UInputRouter* (const FContextInitInfo&)> CreateInputRouterFunc;
+	TUniqueFunction<UInteractiveToolManager* (const FContextInitInfo&)> CreateToolManagerFunc;
+	TUniqueFunction<UToolTargetManager* (const FContextInitInfo&)> CreateToolTargetManagerFunc;
+	TUniqueFunction<UInteractiveGizmoManager* (const FContextInitInfo&)> CreateGizmoManagerFunc;
+	TUniqueFunction<UContextObjectStore* (const FContextInitInfo&)> CreateContextStoreFunc;
+
+	// Analogous to the CreateX() functions above, these function are called by Shutdown()
+	// to terminate and clean up after the various elements. 
+	TUniqueFunction<void(UInputRouter*)> ShutdownInputRouterFunc;
+	TUniqueFunction<void(UInteractiveToolManager*)> ShutdownToolManagerFunc;
+	TUniqueFunction<void(UToolTargetManager*)> ShutdownToolTargetManagerFunc;
+	TUniqueFunction<void(UInteractiveGizmoManager*)> ShutdownGizmoManagerFunc;
+	TUniqueFunction<void(UContextObjectStore*)> ShutdownContextStoreFunc;
+
+
+	// todo: deprecate and remove this, can now be accomplished via CreateToolManagerFunc()
 	UPROPERTY()
 	TSoftClassPtr<UInteractiveToolManager> ToolManagerClass;
 };
\ No newline at end of file
