HeadlessChaos: fix SOA active view related errors
Fix bad merge  from CL 16281007  where in particle SOA kinematics end up in the active particle array

diff --git a/Engine/Source/Runtime/Experimental/Chaos/Public/Chaos/PBDRigidsSOAs.h b/Engine/Source/Runtime/Experimental/Chaos/Public/Chaos/PBDRigidsSOAs.h
--- a/Engine/Source/Runtime/Experimental/Chaos/Public/Chaos/PBDRigidsSOAs.h
+++ b/Engine/Source/Runtime/Experimental/Chaos/Public/Chaos/PBDRigidsSOAs.h
@@ -569,7 +569,7 @@ public:
 		if (Particle->Disabled())
 		{
 			Particle->MoveToSOA(*DynamicDisabledParticles);
-			ActiveParticlesMapArray.Remove(Particle->CastToRigidParticle());
+			RemoveFromActiveArray(Particle->CastToRigidParticle(), /*bStillDirty=*/ false);
 		}
 		else
 		{
@@ -894,7 +894,7 @@ private:
 		if(bStillDirty)
 		{
 			//no longer active, but still dirty
-			ActiveParticlesMapArray.Insert(Particle);
+			TransientDirtyMapArray.Insert(Particle);
 		}
 		else
 		{
