InteractiveToolsFramework: change checks() to ensures() and switch to non-fatal error handling for most situations in GizmoManager and ToolManager. Add TransformGizmo::GetGizmoTransform()
#rnx

diff --git a/Engine/Source/Runtime/Experimental/InteractiveToolsFramework/Private/BaseGizmos/TransformGizmo.cpp b/Engine/Source/Runtime/Experimental/InteractiveToolsFramework/Private/BaseGizmos/TransformGizmo.cpp
--- a/Engine/Source/Runtime/Experimental/InteractiveToolsFramework/Private/BaseGizmos/TransformGizmo.cpp
+++ b/Engine/Source/Runtime/Experimental/InteractiveToolsFramework/Private/BaseGizmos/TransformGizmo.cpp
@@ -586,6 +586,12 @@ void UTransformGizmo::SetActiveTarget(UTransformProxy* Target, IToolContextTrans
 	}
 }
 
+FTransform UTransformGizmo::GetGizmoTransform() const
+{
+	USceneComponent* GizmoComponent = GizmoActor->GetRootComponent();
+	return GizmoComponent->GetComponentTransform();
+}
+
 void UTransformGizmo::ReinitializeGizmoTransform(const FTransform& NewTransform, bool bKeepGizmoUnscaled)
 {
 	// To update the gizmo location without triggering any callbacks, we temporarily
diff --git a/Engine/Source/Runtime/Experimental/InteractiveToolsFramework/Private/InteractiveGizmoManager.cpp b/Engine/Source/Runtime/Experimental/InteractiveToolsFramework/Private/InteractiveGizmoManager.cpp
--- a/Engine/Source/Runtime/Experimental/InteractiveToolsFramework/Private/InteractiveGizmoManager.cpp
+++ b/Engine/Source/Runtime/Experimental/InteractiveToolsFramework/Private/InteractiveGizmoManager.cpp
@@ -57,8 +57,10 @@ void UInteractiveGizmoManager::Shutdown()
 
 void UInteractiveGizmoManager::RegisterGizmoType(const FString& Identifier, UInteractiveGizmoBuilder* Builder)
 {
-	check(GizmoBuilders.Contains(Identifier) == false);
-	GizmoBuilders.Add(Identifier, Builder );
+	if (ensure(GizmoBuilders.Contains(Identifier) == false))
+	{
+		GizmoBuilders.Add(Identifier, Builder);
+	}
 }
 
 
@@ -132,7 +134,10 @@ UInteractiveGizmo* UInteractiveGizmoManager::CreateGizmo(const FString& BuilderI
 bool UInteractiveGizmoManager::DestroyGizmo(UInteractiveGizmo* Gizmo)
 {
 	auto Pred = [Gizmo](const FActiveGizmo& ActiveGizmo) {return ActiveGizmo.Gizmo == Gizmo; };
-	check(ActiveGizmos.FindByPredicate(Pred));
+	if (!ensure(ActiveGizmos.FindByPredicate(Pred)))
+	{
+		return false;
+	}
 
 	InputRouter->ForceTerminateSource(Gizmo);
 
diff --git a/Engine/Source/Runtime/Experimental/InteractiveToolsFramework/Private/InteractiveToolManager.cpp b/Engine/Source/Runtime/Experimental/InteractiveToolsFramework/Private/InteractiveToolManager.cpp
--- a/Engine/Source/Runtime/Experimental/InteractiveToolsFramework/Private/InteractiveToolManager.cpp
+++ b/Engine/Source/Runtime/Experimental/InteractiveToolsFramework/Private/InteractiveToolManager.cpp
@@ -67,8 +67,10 @@ void UInteractiveToolManager::DoPostSetup(EToolSide Side, UInteractiveTool* InIn
 
 void UInteractiveToolManager::RegisterToolType(const FString& Identifier, UInteractiveToolBuilder* Builder)
 {
-	check(ToolBuilders.Contains(Identifier) == false);
-	ToolBuilders.Add(Identifier, Builder );
+	if (ensure(ToolBuilders.Contains(Identifier) == false))
+	{
+		ToolBuilders.Add(Identifier, Builder);
+	}
 }
 
 void UInteractiveToolManager::UnregisterToolType(const FString& Identifier)
@@ -109,7 +111,10 @@ bool UInteractiveToolManager::SelectActiveToolType(EToolSide Side, const FString
 
 bool UInteractiveToolManager::CanActivateTool(EToolSide Side, FString Identifier)
 {
-	check(Side == EToolSide::Left);   // TODO: support right-side tool
+	if (ensure(Side == EToolSide::Left) == false)  // TODO: support right-side tool
+	{
+		return false;
+	}
 
 	if (ToolBuilders.Contains(Identifier))
 	{
@@ -126,7 +131,10 @@ bool UInteractiveToolManager::CanActivateTool(EToolSide Side, FString Identifier
 
 bool UInteractiveToolManager::ActivateTool(EToolSide Side)
 {
-	check(Side == EToolSide::Left);   // TODO: support right-side tool
+	if (ensure(Side == EToolSide::Left) == false)  // TODO: support right-side tool
+	{
+		return false;
+	}
 
 	// wrap tool change in a transaction so that deactivate and activate are grouped
 	bool bInTransaction = false;
@@ -152,8 +160,10 @@ bool UInteractiveToolManager::ActivateTool(EToolSide Side)
 
 	if (ActiveToolChangeTrackingMode == EToolChangeTrackingMode::FullUndoRedo)
 	{
-		check(TransactionsAPI);
-		TransactionsAPI->AppendChange(this, MakeUnique<FActivateToolChange>(Side, ActiveLeftToolName), LOCTEXT("ActivateToolChange", "Activate Tool"));
+		if (ensure(TransactionsAPI))
+		{
+			TransactionsAPI->AppendChange(this, MakeUnique<FActivateToolChange>(Side, ActiveLeftToolName), LOCTEXT("ActivateToolChange", "Activate Tool"));
+		}
 	} 
 	else if (ActiveToolChangeTrackingMode == EToolChangeTrackingMode::UndoToExit)
 	{
@@ -211,13 +221,19 @@ bool UInteractiveToolManager::ActivateToolInternal(EToolSide Side)
 
 void UInteractiveToolManager::DeactivateTool(EToolSide Side, EToolShutdownType ShutdownType)
 {
-	check(Side == EToolSide::Left);   // TODO: support right-side tool
+	if (ensure(Side == EToolSide::Left) == false)  // TODO: support right-side tool
+	{
+		return;
+	}
+
 	if (ActiveLeftTool != nullptr)
 	{
 		if (ActiveToolChangeTrackingMode == EToolChangeTrackingMode::FullUndoRedo)
 		{
-			check(TransactionsAPI);
-			TransactionsAPI->AppendChange(this, MakeUnique<FActivateToolChange>(Side, ActiveLeftToolName, ShutdownType), LOCTEXT("DeactivateToolChange", "Deactivate Tool"));
+			if (ensure(TransactionsAPI))
+			{
+				TransactionsAPI->AppendChange(this, MakeUnique<FActivateToolChange>(Side, ActiveLeftToolName, ShutdownType), LOCTEXT("DeactivateToolChange", "Deactivate Tool"));
+			}
 		}
 
 		DeactivateToolInternal(Side, ShutdownType);
@@ -229,7 +245,11 @@ void UInteractiveToolManager::DeactivateToolInternal(EToolSide Side, EToolShutdo
 {
 	if (Side == EToolSide::Left)
 	{
-		check(ActiveLeftTool);
+		if (!ensure(ActiveLeftTool))
+		{
+			return;
+		}
+
 		InputRouter->ForceTerminateSource(ActiveLeftTool);
 
 		ActiveLeftTool->Shutdown(ShutdownType);
@@ -398,14 +418,28 @@ void UInteractiveToolManager::EndUndoTransaction()
 
 void UInteractiveToolManager::EmitObjectChange(UObject* TargetObject, TUniquePtr<FToolCommandChange> Change, const FText& Description)
 {
-	// wrap change 
-	check(HasActiveTool(EToolSide::Left));
-	TUniquePtr<FToolChangeWrapperChange> Wrapper = MakeUnique<FToolChangeWrapperChange>();
-	Wrapper->ToolManager = this;
-	Wrapper->ActiveTool = GetActiveTool(EToolSide::Left);
-	Wrapper->ToolChange = MoveTemp(Change);
+	// Currently we do not support undo/redo of changes emitted /during/ a Tool after the Tool exits, 
+	// because we do not support "undo back into a Tool". So once a Tool exits, any Changes it emitted are considered invalid.
+	// We do not assume we can modify the external history/transaction system (because in the Editor we cannot)
+	// so all we can do is mark those changes as "expired" and hope that the external transaction system respects that.
+	// The Tools themselves will not be around to know whether they have exited, so we wrap
+	// the changes they emit in a FToolChangeWrapperChange which will allow the ToolManager
+	// to decide whether the change is still valid/applicable.
+
+	// Note: if you do not want this wrapping behavior for particular changes, you can use the TransactionsAPI directly 
+	// via GetContextTransactionsAPI()
 
-	TransactionsAPI->AppendChange(TargetObject, MoveTemp(Wrapper), Description );
+	// If you hit this ensure from a Gizmo, you may have accidentally registered the ToolManager as the Gizmo's TransactionProvider,
+	// this should only be done if the Gizmo is embedded inside a Tool
+	if (ensure(HasActiveTool(EToolSide::Left)))
+	{
+		TUniquePtr<FToolChangeWrapperChange> Wrapper = MakeUnique<FToolChangeWrapperChange>();
+		Wrapper->ToolManager = this;
+		Wrapper->ActiveTool = GetActiveTool(EToolSide::Left);
+		Wrapper->ToolChange = MoveTemp(Change);
+
+		TransactionsAPI->AppendChange(TargetObject, MoveTemp(Wrapper), Description);
+	}
 }
 
 bool UInteractiveToolManager::RequestSelectionChange(const FSelectedOjectsChangeList& SelectionChange)
diff --git a/Engine/Source/Runtime/Experimental/InteractiveToolsFramework/Public/BaseGizmos/TransformGizmo.h b/Engine/Source/Runtime/Experimental/InteractiveToolsFramework/Public/BaseGizmos/TransformGizmo.h
--- a/Engine/Source/Runtime/Experimental/InteractiveToolsFramework/Public/BaseGizmos/TransformGizmo.h
+++ b/Engine/Source/Runtime/Experimental/InteractiveToolsFramework/Public/BaseGizmos/TransformGizmo.h
@@ -315,6 +315,11 @@ public:
 	 */
 	ATransformGizmoActor* GetGizmoActor() const { return GizmoActor; }
 
+	/**
+	 * @return current transform of Gizmo
+	 */
+	FTransform GetGizmoTransform() const;
+
 	/**
 	 * Repositions the gizmo without issuing undo/redo changes, triggering callbacks, 
 	 * or moving any components. Useful for resetting the gizmo to a new location without
