LWC: Build fixes

diff --git a/Engine/Plugins/Enterprise/DatasmithCADImporter/Source/CADLibrary/Private/MeshDescriptionHelper.cpp b/Engine/Plugins/Enterprise/DatasmithCADImporter/Source/CADLibrary/Private/MeshDescriptionHelper.cpp
--- a/Engine/Plugins/Enterprise/DatasmithCADImporter/Source/CADLibrary/Private/MeshDescriptionHelper.cpp
+++ b/Engine/Plugins/Enterprise/DatasmithCADImporter/Source/CADLibrary/Private/MeshDescriptionHelper.cpp
@@ -217,9 +217,7 @@ namespace CADLibrary
 		BBox.Min *= ImportParams.ScaleFactor;
 		BBox.Max *= ImportParams.ScaleFactor;
 
-		TVertexAttributesRef<FVector> VertexPositions = MeshDescription.GetVertexPositions();
-
-		FVector Position;
+		TVertexAttributesRef<FVector3f> VertexPositions = MeshDescription.GetVertexPositions();
 
 		TArray<FVector>& VertexArray = Body.VertexArray;
 
diff --git a/Engine/Source/Runtime/Experimental/Chaos/Private/Chaos/PBDLongRangeConstraintsBase.cpp b/Engine/Source/Runtime/Experimental/Chaos/Private/Chaos/PBDLongRangeConstraintsBase.cpp
--- a/Engine/Source/Runtime/Experimental/Chaos/Private/Chaos/PBDLongRangeConstraintsBase.cpp
+++ b/Engine/Source/Runtime/Experimental/Chaos/Private/Chaos/PBDLongRangeConstraintsBase.cpp
@@ -41,15 +41,15 @@ FPBDLongRangeConstraintsBase::FPBDLongRangeConstraintsBase(
 	}
 
 	// Scale the tether's reference lengths
-	const FReal ScaleOffset = FMath::Clamp(Scale[0], (FReal)0.01, (FReal)10.);
+	const FReal ScaleOffset = FMath::Clamp(Scale[0], 0.01f, 10.f);
 
 	if (ScaleMultipliers.Num() == InParticleCount)
 	{
-		const FReal ScaleRange = FMath::Clamp(Scale[1], (FReal)0.01, (FReal)10.) - Scale[0];
+		const FReal ScaleRange = FMath::Clamp(Scale[1], 0.01f, 10.f) - Scale[0];
 
 		for (FTether& Tether : Tethers)
 		{
-			Tether.RefLength *= ScaleOffset + ScaleRange * FMath::Clamp((FReal)ScaleMultipliers[Tether.End], (FReal)0., (FReal)1.);
+			Tether.RefLength *= ScaleOffset + ScaleRange * FMath::Clamp(ScaleMultipliers[Tether.End], 0.f, 1.f);
 		}
 	}
 	else
