Chaos - wire SetConstraintFrames from game thread to physics thread
(Note: works in BeginPlay or Tick, but not from construction script yet)

diff --git a/Engine/Source/Runtime/Engine/Private/PhysicsEngine/ImmediatePhysics/ImmediatePhysicsChaos/ImmediatePhysicsJointHandle_Chaos.cpp b/Engine/Source/Runtime/Engine/Private/PhysicsEngine/ImmediatePhysics/ImmediatePhysicsChaos/ImmediatePhysicsJointHandle_Chaos.cpp
--- a/Engine/Source/Runtime/Engine/Private/PhysicsEngine/ImmediatePhysics/ImmediatePhysicsChaos/ImmediatePhysicsJointHandle_Chaos.cpp
+++ b/Engine/Source/Runtime/Engine/Private/PhysicsEngine/ImmediatePhysics/ImmediatePhysicsChaos/ImmediatePhysicsJointHandle_Chaos.cpp
@@ -152,23 +152,22 @@ namespace ImmediatePhysics_Chaos
 		using namespace Chaos;
 
 		FPBDJointSettings ConstraintSettings;
-		TVec2<FRigidTransform3> ConstraintFrames;
 
 		if (ConstraintInstance != nullptr)
 		{
 			// BodyInstance/PhysX has the constraint locations in actor-space, but we need them in Center-of-Mass space
-			ConstraintFrames[0] = FParticleUtilities::ActorLocalToParticleLocal(FGenericParticleHandle(Actor1->GetParticle()), ConstraintInstance->GetRefFrame(EConstraintFrame::Frame1));
-			ConstraintFrames[1] = FParticleUtilities::ActorLocalToParticleLocal(FGenericParticleHandle(Actor2->GetParticle()), ConstraintInstance->GetRefFrame(EConstraintFrame::Frame2));
-			FReal JointScale = ConstraintInstance->GetLastKnownScale();
-			ConstraintFrames[0].ScaleTranslation(JointScale);
-			ConstraintFrames[1].ScaleTranslation(JointScale);
 			TransferJointSettings(ConstraintInstance, ConstraintSettings);
+			FReal JointScale = ConstraintInstance->GetLastKnownScale();
+			ConstraintSettings.ConnectorTransforms[0] = FParticleUtilities::ActorLocalToParticleLocal(FGenericParticleHandle(Actor1->GetParticle()), ConstraintInstance->GetRefFrame(EConstraintFrame::Frame1));
+			ConstraintSettings.ConnectorTransforms[1] = FParticleUtilities::ActorLocalToParticleLocal(FGenericParticleHandle(Actor2->GetParticle()), ConstraintInstance->GetRefFrame(EConstraintFrame::Frame2));
+			ConstraintSettings.ConnectorTransforms[0].ScaleTranslation(JointScale);
+			ConstraintSettings.ConnectorTransforms[1].ScaleTranslation(JointScale);
 		}
 		else
 		{
 			// TEMP: all creation with null ConstraintIndex for PhAt handles
-			ConstraintFrames[0] = Actor2->GetWorldTransform().GetRelativeTransform(Actor1->GetWorldTransform());
-			ConstraintFrames[1] = FRigidTransform3();
+			ConstraintSettings.ConnectorTransforms[0] = Actor2->GetWorldTransform().GetRelativeTransform(Actor1->GetWorldTransform());
+			ConstraintSettings.ConnectorTransforms[1] = FRigidTransform3();
 			ConstraintSettings.LinearMotionTypes = { EJointMotionType::Limited, EJointMotionType::Limited, EJointMotionType::Limited };
 			ConstraintSettings.LinearLimit = 0.1f;
 			ConstraintSettings.SoftLinearStiffness = 500.0f;
@@ -181,7 +180,7 @@ namespace ImmediatePhysics_Chaos
 
 		ConstraintSettings.Sanitize();
 
-		ConstraintHandle = Constraints->AddConstraint({ Actor1->ParticleHandle, Actor2->ParticleHandle }, ConstraintFrames, ConstraintSettings);
+		ConstraintHandle = Constraints->AddConstraint({ Actor1->ParticleHandle, Actor2->ParticleHandle }, ConstraintSettings);
 	}
 
 	FJointHandle::~FJointHandle()
diff --git a/Engine/Source/Runtime/Experimental/Chaos/Private/Chaos/PBDJointConstraintData.cpp b/Engine/Source/Runtime/Experimental/Chaos/Private/Chaos/PBDJointConstraintData.cpp
--- a/Engine/Source/Runtime/Experimental/Chaos/Private/Chaos/PBDJointConstraintData.cpp
+++ b/Engine/Source/Runtime/Experimental/Chaos/Private/Chaos/PBDJointConstraintData.cpp
@@ -7,14 +7,11 @@ namespace Chaos
 {
 	FJointConstraint::FJointConstraint()
 		: FConstraintBase(EConstraintType::JointConstraintType)
-		, JointTransforms({ FTransform::Identity, FTransform::Identity })
 		, UserData(nullptr)
 		, KinematicEndPoint(nullptr)
 	{
 	}
 
-	FJointConstraint::FTransformPair FJointConstraint::GetJointTransforms() { return JointTransforms; }
-
 	void FJointConstraint::SetKinematicEndPoint(FSingleParticlePhysicsProxy* InDummyParticle, FPBDRigidsSolver* Solver)
 	{
 		ensure(KinematicEndPoint == nullptr);
@@ -22,11 +19,12 @@ namespace Chaos
 		Solver->RegisterObject(KinematicEndPoint);
 	}
 
-	const FJointConstraint::FTransformPair FJointConstraint::GetJointTransforms() const { return JointTransforms; }
+	const FJointConstraint::FTransformPair FJointConstraint::GetJointTransforms() const { return JointSettings.ConnectorTransforms; }
 	void FJointConstraint::SetJointTransforms(const Chaos::FJointConstraint::FTransformPair& InJointTransforms)
 	{
-		JointTransforms[0] = InJointTransforms[0];
-		JointTransforms[1] = InJointTransforms[1];
+		JointSettings.ConnectorTransforms = InJointTransforms;
+		MDirtyFlags.MarkDirty(EJointConstraintFlags::JointTransforms);
+		SetProxy(Proxy);
 	}
 
 	void FJointConstraint::SetLinearPositionDriveEnabled(TVector<bool,3> Enabled)
diff --git a/Engine/Source/Runtime/Experimental/Chaos/Private/Chaos/PBDJointConstraints.cpp b/Engine/Source/Runtime/Experimental/Chaos/Private/Chaos/PBDJointConstraints.cpp
--- a/Engine/Source/Runtime/Experimental/Chaos/Private/Chaos/PBDJointConstraints.cpp
+++ b/Engine/Source/Runtime/Experimental/Chaos/Private/Chaos/PBDJointConstraints.cpp
@@ -369,33 +369,34 @@ namespace Chaos
 
 	typename FPBDJointConstraints::FConstraintContainerHandle* FPBDJointConstraints::AddConstraint(const FParticlePair& InConstrainedParticles, const FRigidTransform3& WorldConstraintFrame)
 	{
-		FTransformPair JointFrames;
-		JointFrames[0] = FRigidTransform3(
+		FPBDJointSettings JointSettings;
+		JointSettings.ConnectorTransforms[0] = FRigidTransform3(
 			WorldConstraintFrame.GetTranslation() - InConstrainedParticles[0]->X(),
 			WorldConstraintFrame.GetRotation() * InConstrainedParticles[0]->R().Inverse()
 			);
-		JointFrames[1] = FRigidTransform3(
+		JointSettings.ConnectorTransforms[1] = FRigidTransform3(
 			WorldConstraintFrame.GetTranslation() - InConstrainedParticles[1]->X(),
 			WorldConstraintFrame.GetRotation() * InConstrainedParticles[1]->R().Inverse()
 			);
-		return AddConstraint(InConstrainedParticles, JointFrames, FPBDJointSettings());
+		return AddConstraint(InConstrainedParticles, JointSettings);
 	}
 
 	
-	typename FPBDJointConstraints::FConstraintContainerHandle* FPBDJointConstraints::AddConstraint(const FParticlePair& InConstrainedParticles, const FTransformPair& InConstraintFrames)
+	typename FPBDJointConstraints::FConstraintContainerHandle* FPBDJointConstraints::AddConstraint(const FParticlePair& InConstrainedParticles, const FTransformPair& InConnectorTransforms)
 	{
-		return AddConstraint(InConstrainedParticles, InConstraintFrames, FPBDJointSettings());
+		FPBDJointSettings JointSettings;
+		JointSettings.ConnectorTransforms = InConnectorTransforms;
+		return AddConstraint(InConstrainedParticles, JointSettings);
 	}
 
 	
-	typename FPBDJointConstraints::FConstraintContainerHandle* FPBDJointConstraints::AddConstraint(const FParticlePair& InConstrainedParticles, const FTransformPair& InConstraintFrames, const FPBDJointSettings& InConstraintSettings)
+	typename FPBDJointConstraints::FConstraintContainerHandle* FPBDJointConstraints::AddConstraint(const FParticlePair& InConstrainedParticles, const FPBDJointSettings& InConstraintSettings)
 	{
 		bJointsDirty = true;
 
 		int ConstraintIndex = Handles.Num();
 		Handles.Add(HandleAllocator.AllocHandle(this, ConstraintIndex));
 		ConstraintParticles.Add(InConstrainedParticles);
-		ConstraintFrames.Add(InConstraintFrames);
 		ConstraintStates.Add(FPBDJointState());
 
 		ConstraintSettings.AddDefaulted();
@@ -425,7 +426,6 @@ namespace Chaos
 		// Swap the last constraint into the gap to keep the array packed
 		ConstraintParticles.RemoveAtSwap(ConstraintIndex);
 		ConstraintSettings.RemoveAtSwap(ConstraintIndex);
-		ConstraintFrames.RemoveAtSwap(ConstraintIndex);
 		ConstraintStates.RemoveAtSwap(ConstraintIndex);
 		Handles.RemoveAtSwap(ConstraintIndex);
 
@@ -491,11 +491,9 @@ namespace Chaos
 			});
 
 		TArray<FPBDJointSettings> SortedConstraintSettings;
-		TArray<FTransformPair> SortedConstraintFrames;
 		TArray<FParticlePair> SortedConstraintParticles;
 		TArray<FPBDJointState> SortedConstraintStates;
 		SortedConstraintSettings.Reserve(SortedHandles.Num());
-		SortedConstraintFrames.Reserve(SortedHandles.Num());
 		SortedConstraintParticles.Reserve(SortedHandles.Num());
 		SortedConstraintStates.Reserve(SortedHandles.Num());
 
@@ -505,14 +503,12 @@ namespace Chaos
 			int32 UnsortedConstraintIndex = Handle->GetConstraintIndex();
 
 			SortedConstraintSettings.Add(ConstraintSettings[UnsortedConstraintIndex]);
-			SortedConstraintFrames.Add(ConstraintFrames[UnsortedConstraintIndex]);
 			SortedConstraintParticles.Add(ConstraintParticles[UnsortedConstraintIndex]);
 			SortedConstraintStates.Add(ConstraintStates[UnsortedConstraintIndex]);
 			SetConstraintIndex(Handle, SortedConstraintIndex);
 		}
 
 		Swap(ConstraintSettings, SortedConstraintSettings);
-		Swap(ConstraintFrames, SortedConstraintFrames);
 		Swap(ConstraintParticles, SortedConstraintParticles);
 		Swap(ConstraintStates, SortedConstraintStates);
 		Swap(Handles, SortedHandles);
@@ -739,7 +735,7 @@ namespace Chaos
 				FJointSolverJointState& JointState = SolverConstraintStates[JointIndex];
 				const FPBDJointSettings& JointSettings = ConstraintSettings[JointIndex];
 
-				const FTransformPair& JointFrames = ConstraintFrames[JointIndex];
+				const FTransformPair& JointFrames = JointSettings.ConnectorTransforms;
 				int32 Index0, Index1;
 				GetConstrainedParticleIndices(JointIndex, Index0, Index1);
 				FGenericParticleHandle Particle0 = FGenericParticleHandle(ConstraintParticles[JointIndex][Index0]);
@@ -768,7 +764,7 @@ namespace Chaos
 
 				const FPBDJointSettings& JointSettings = ConstraintSettings[JointIndex];
 
-				const FTransformPair& JointFrames = ConstraintFrames[JointIndex];
+				const FTransformPair& JointFrames = JointSettings.ConnectorTransforms;
 				FJointSolverGaussSeidel& Solver = ConstraintSolvers[JointIndex];
 
 				int32 Index0, Index1;
@@ -857,8 +853,8 @@ namespace Chaos
 		const FRotation3 Q0 = FParticleUtilities::GetCoMWorldRotation(Particle0);
 		const FVec3 P1 = FParticleUtilities::GetCoMWorldPosition(Particle1);
 		const FRotation3 Q1 = FParticleUtilities::GetCoMWorldRotation(Particle1);
-		const FRigidTransform3& XL0 = FParticleUtilities::ParticleLocalToCoMLocal(Particle0, ConstraintFrames[ConstraintIndex][Index0]);
-		const FRigidTransform3& XL1 = FParticleUtilities::ParticleLocalToCoMLocal(Particle1, ConstraintFrames[ConstraintIndex][Index1]);
+		const FRigidTransform3& XL0 = FParticleUtilities::ParticleLocalToCoMLocal(Particle0, ConstraintSettings[ConstraintIndex].ConnectorTransforms[Index0]);
+		const FRigidTransform3& XL1 = FParticleUtilities::ParticleLocalToCoMLocal(Particle1, ConstraintSettings[ConstraintIndex].ConnectorTransforms[Index1]);
 
 		OutX0 = P0 + Q0 * XL0.GetTranslation();
 		OutX1 = P1 + Q1 * XL1.GetTranslation();
@@ -1504,7 +1500,7 @@ namespace Chaos
 		FGenericParticleHandle Particle0 = FGenericParticleHandle(ConstraintParticles[ConstraintIndex][Index0]);
 		FGenericParticleHandle Particle1 = FGenericParticleHandle(ConstraintParticles[ConstraintIndex][Index1]);
 
-		const FTransformPair &ConstraintFramesLocal = ConstraintFrames[ConstraintIndex];
+		const FTransformPair &ConstraintFramesLocal = JointSettings.ConnectorTransforms;
 		FTransformPair ConstraintFramesGlobal(ConstraintFramesLocal[Index0] * FRigidTransform3(Particle0->P(), Particle0->Q()), ConstraintFramesLocal[Index1] * FRigidTransform3(Particle1->P(), Particle1->Q()));
 		FQuat Q1 = ConstraintFramesGlobal[1].GetRotation();
 		Q1.EnforceShortestArcWith(ConstraintFramesGlobal[0].GetRotation());
diff --git a/Engine/Source/Runtime/Experimental/Chaos/Private/PhysicsProxy/JointConstraintProxy.cpp b/Engine/Source/Runtime/Experimental/Chaos/Private/PhysicsProxy/JointConstraintProxy.cpp
--- a/Engine/Source/Runtime/Experimental/Chaos/Private/PhysicsProxy/JointConstraintProxy.cpp
+++ b/Engine/Source/Runtime/Experimental/Chaos/Private/PhysicsProxy/JointConstraintProxy.cpp
@@ -116,6 +116,12 @@ void FJointConstraintPhysicsProxy::PushStateOnGameThread(Chaos::FPBDRigidsSolver
 	{
 		if (Constraint->IsDirty())
 		{
+			if (Constraint->IsDirty(Chaos::EJointConstraintFlags::JointTransforms))
+			{
+				JointSettingsBuffer.ConnectorTransforms = Constraint->GetJointTransforms();
+				DirtyFlagsBuffer.MarkDirty(Chaos::EJointConstraintFlags::JointTransforms);
+			}
+
 			if (Constraint->IsDirty(Chaos::EJointConstraintFlags::CollisionEnabled))
 			{
 				JointSettingsBuffer.bCollisionEnabled = Constraint->GetCollisionEnabled();
@@ -241,6 +247,11 @@ void FJointConstraintPhysicsProxy::PushStateOnPhysicsThread(Chaos::FPBDRigidsSol
 		{
 			FConstraintData ConstraintSettings = Handle->GetSettings();
 
+			if (DirtyFlagsBuffer.IsDirty(Chaos::EJointConstraintFlags::JointTransforms))
+			{
+				ConstraintSettings.ConnectorTransforms = JointSettingsBuffer.ConnectorTransforms;
+			}
+
 			if (DirtyFlagsBuffer.IsDirty(Chaos::EJointConstraintFlags::CollisionEnabled))
 			{
 				if (!JointSettingsBuffer.bCollisionEnabled)
diff --git a/Engine/Source/Runtime/Experimental/Chaos/Public/Chaos/PBDJointConstraintData.h b/Engine/Source/Runtime/Experimental/Chaos/Public/Chaos/PBDJointConstraintData.h
--- a/Engine/Source/Runtime/Experimental/Chaos/Public/Chaos/PBDJointConstraintData.h
+++ b/Engine/Source/Runtime/Experimental/Chaos/Public/Chaos/PBDJointConstraintData.h
@@ -14,7 +14,7 @@ namespace Chaos
 
 	enum class EJointConstraintFlags : uint64_t
 	{
-		Position                    = 0,
+		JointTransforms             = static_cast<uint64_t>(1) << 0,
 		CollisionEnabled            = static_cast<uint64_t>(1) << 1,
 		Projection                  = static_cast<uint64_t>(1) << 2,
 		ParentInvMassScale          = static_cast<uint64_t>(1) << 3,
@@ -38,7 +38,6 @@ namespace Chaos
 		typedef FPBDJointSettings FData;
 		typedef FPBDJointConstraintHandle FHandle;
 		typedef TVector<FTransform, 2> FTransformPair;
-		friend FData;
 
 		friend class FPBDRigidsSolver; // friend so we can call ReleaseKinematicEndPoint when unregistering joint.
 
@@ -49,15 +48,15 @@ namespace Chaos
 		bool IsDirty(const EJointConstraintFlags CheckBits) const { return MDirtyFlags.IsDirty(CheckBits); }
 		void ClearDirtyFlags() { MDirtyFlags.Clear(); }
 
-		void SetJointTransforms(const Chaos::FJointConstraint::FTransformPair& InJoinTransforms);
-		const FTransformPair GetJointTransforms() const;
-		FTransformPair GetJointTransforms();
-
 		const FData& GetJointSettings()const { return JointSettings; }
 
 		// If we created particle to serve as kinematic endpoint, track so we can release later. This will add particle to solver.
 		void SetKinematicEndPoint(FSingleParticlePhysicsProxy* InParticle, FPBDRigidsSolver* Solver);
 
+		//CONSTRAINT_JOINT_PROPERPETY_IMPL(bool, JointTransforms, EJointConstraintFlags::Position, JointSettings.ConnectorTransforms);
+		void SetJointTransforms(const Chaos::FJointConstraint::FTransformPair& InJointTransforms);
+		const FTransformPair GetJointTransforms() const;
+
 		CONSTRAINT_JOINT_PROPERPETY_IMPL(bool, CollisionEnabled, EJointConstraintFlags::CollisionEnabled, JointSettings.bCollisionEnabled);
 		//void SetCollisionEnabled(bool InValue);
 		//bool GetCollisionEnabled() const;
@@ -312,7 +311,6 @@ namespace Chaos
 		FJointConstraintDirtyFlags MDirtyFlags;
 		FData JointSettings;
 
-		FTransformPair JointTransforms;
 		void* UserData;
 		FOutputData Output;
 
diff --git a/Engine/Source/Runtime/Experimental/Chaos/Public/Chaos/PBDJointConstraintTypes.h b/Engine/Source/Runtime/Experimental/Chaos/Public/Chaos/PBDJointConstraintTypes.h
--- a/Engine/Source/Runtime/Experimental/Chaos/Public/Chaos/PBDJointConstraintTypes.h
+++ b/Engine/Source/Runtime/Experimental/Chaos/Public/Chaos/PBDJointConstraintTypes.h
@@ -99,12 +99,16 @@ namespace Chaos
 	class CHAOS_API FPBDJointSettings
 	{
 	public:
+		using FTransformPair = TVector<FRigidTransform3, 2>;
+
 		FPBDJointSettings();
 
 		// Ensure that settings are consistent and within valid ranges. Should be called
 		// whenever settings change.
 		void Sanitize();
 
+		FTransformPair ConnectorTransforms;
+
 		FReal Stiffness;
 		FReal LinearProjection;
 		FReal AngularProjection;
diff --git a/Engine/Source/Runtime/Experimental/Chaos/Public/Chaos/PBDJointConstraints.h b/Engine/Source/Runtime/Experimental/Chaos/Public/Chaos/PBDJointConstraints.h
--- a/Engine/Source/Runtime/Experimental/Chaos/Public/Chaos/PBDJointConstraints.h
+++ b/Engine/Source/Runtime/Experimental/Chaos/Public/Chaos/PBDJointConstraints.h
@@ -106,8 +106,8 @@ namespace Chaos
 		 * Add a constraint with particle-space constraint offsets.
 		 */
 		FConstraintContainerHandle* AddConstraint(const FParticlePair& InConstrainedParticles, const FRigidTransform3& WorldConstraintFrame);
-		FConstraintContainerHandle* AddConstraint(const FParticlePair& InConstrainedParticles, const FTransformPair& ConstraintFrames);
-		FConstraintContainerHandle* AddConstraint(const FParticlePair& InConstrainedParticles, const FTransformPair& ConstraintFrames, const FPBDJointSettings& InConstraintSettings);
+		FConstraintContainerHandle* AddConstraint(const FParticlePair& InConstrainedParticles, const FTransformPair& ConnectorTransforms);
+		FConstraintContainerHandle* AddConstraint(const FParticlePair& InConstrainedParticles, const FPBDJointSettings& InConstraintSettings);
 
 		/**
 		 * Remove the specified constraint.
@@ -257,7 +257,6 @@ namespace Chaos
 		FPBDJointSolverSettings Settings;
 
 		TArray<FPBDJointSettings> ConstraintSettings;
-		TArray<FTransformPair> ConstraintFrames;
 		TArray<FParticlePair> ConstraintParticles;
 		TArray<FPBDJointState> ConstraintStates;
 
diff --git a/Engine/Source/Runtime/PhysicsCore/Private/ChaosEngineInterface.cpp b/Engine/Source/Runtime/PhysicsCore/Private/ChaosEngineInterface.cpp
--- a/Engine/Source/Runtime/PhysicsCore/Private/ChaosEngineInterface.cpp
+++ b/Engine/Source/Runtime/PhysicsCore/Private/ChaosEngineInterface.cpp
@@ -1296,7 +1296,22 @@ void FChaosEngineInterface::SetPlasticityLimits_AssumesLocked(const FPhysicsCons
 
 void FChaosEngineInterface::SetLocalPose(const FPhysicsConstraintHandle& InConstraintRef,const FTransform& InPose,EConstraintFrame::Type InFrame)
 {
-	// @todo(chaos) :  Joint Constraints : Motors
+	if (InConstraintRef.IsValid() && InConstraintRef.Constraint->IsType(Chaos::EConstraintType::JointConstraintType))
+	{
+		if (Chaos::FJointConstraint* Constraint = static_cast<Chaos::FJointConstraint*>(InConstraintRef.Constraint))
+		{
+			Chaos::FJointConstraint::FTransformPair JointTransforms = Constraint->GetJointTransforms();
+			if (InFrame == EConstraintFrame::Frame1)
+			{
+				JointTransforms[0] = InPose;
+			}
+			else
+			{
+				JointTransforms[1] = InPose;
+			}
+			Constraint->SetJointTransforms(JointTransforms);
+		}
+	}
 }
 
 void FChaosEngineInterface::SetDrivePosition(const FPhysicsConstraintHandle& InConstraintRef,const FVector& InPosition)
