FGeometryCollectionSceneProxy
-- dtor releases DynamicData instead of deleting it.

diff --git a/Engine/Source/Runtime/Experimental/GeometryCollectionEngine/Private/GeometryCollection/GeometryCollectionSceneProxy.cpp b/Engine/Source/Runtime/Experimental/GeometryCollectionEngine/Private/GeometryCollection/GeometryCollectionSceneProxy.cpp
--- a/Engine/Source/Runtime/Experimental/GeometryCollectionEngine/Private/GeometryCollection/GeometryCollectionSceneProxy.cpp
+++ b/Engine/Source/Runtime/Experimental/GeometryCollectionEngine/Private/GeometryCollection/GeometryCollectionSceneProxy.cpp
@@ -190,7 +190,8 @@ FGeometryCollectionSceneProxy::~FGeometryCollectionSceneProxy()
 
 	if (DynamicData != nullptr)
 	{
-		delete DynamicData;
+		GDynamicDataPool.Release(DynamicData);
+		DynamicData = nullptr;
 	}
 
 	if (ConstantData != nullptr)
