Chaos - fix joint unit test

diff --git a/Engine/Source/Programs/HeadlessChaos/Private/HeadlessChaosTestJoint.h b/Engine/Source/Programs/HeadlessChaos/Private/HeadlessChaosTestJoint.h
--- a/Engine/Source/Programs/HeadlessChaos/Private/HeadlessChaosTestJoint.h
+++ b/Engine/Source/Programs/HeadlessChaos/Private/HeadlessChaosTestJoint.h
@@ -38,6 +38,10 @@ namespace ChaosTest
 		{
 			FPBDJointConstraintHandle* Joint = Joints.AddConstraint(InConstrainedParticleIndices, FRigidTransform3(JointPositions[JointIndex], FRotation3::FromIdentity()));
 
+			// @todo(chaos): this indicates we need to change the AddConstraint API (since ConnectorTransforms were added to Settings). 
+			// Calling AddConstraint followed by SetSettings will overwrite the ConnectorTransforms
+			JointSettings[JointIndex].ConnectorTransforms = Joint->GetSettings().ConnectorTransforms;
+
 			if (JointIndex < JointSettings.Num())
 			{
 				Joint->SetSettings(JointSettings[JointIndex]);
diff --git a/Engine/Source/Runtime/Experimental/Chaos/Private/Chaos/PBDConstraintGraph.cpp b/Engine/Source/Runtime/Experimental/Chaos/Private/Chaos/PBDConstraintGraph.cpp
--- a/Engine/Source/Runtime/Experimental/Chaos/Private/Chaos/PBDConstraintGraph.cpp
+++ b/Engine/Source/Runtime/Experimental/Chaos/Private/Chaos/PBDConstraintGraph.cpp
@@ -1002,6 +1002,13 @@ bool FPBDConstraintGraph::DebugCheckGraph() const
 		for (int32 EdgeIndex : IslandEdges)
 		{
 			const FGraphEdge& Edge = Edges[EdgeIndex];
+
+			if (Edge.SecondNode == INDEX_NONE)
+			{
+				// We add "edges" for single-particle constraints as well. In this case, they have no second node
+				continue;
+			}
+
 			const FGraphNode& Node0 = Nodes[Edge.FirstNode];
 			const FGraphNode& Node1 = Nodes[Edge.SecondNode];
 			if ((Node0.Particle != nullptr) && (Node1.Particle != nullptr) && FConstGenericParticleHandle(Node0.Particle)->IsDynamic() && FConstGenericParticleHandle(Node1.Particle)->IsDynamic())
diff --git a/Engine/Source/Runtime/Experimental/Chaos/Public/Chaos/PBDJointConstraints.h b/Engine/Source/Runtime/Experimental/Chaos/Public/Chaos/PBDJointConstraints.h
--- a/Engine/Source/Runtime/Experimental/Chaos/Public/Chaos/PBDJointConstraints.h
+++ b/Engine/Source/Runtime/Experimental/Chaos/Public/Chaos/PBDJointConstraints.h
@@ -104,6 +104,8 @@ namespace Chaos
 
 		/**
 		 * Add a constraint with particle-space constraint offsets.
+		 * @todo(chaos): clean up this set of functions (now that ConnectorTransforms is in the settings, calling AddConstraint then SetSettings leads 
+		 * to unexpected behaviour - overwriting the ConnectorTransforms with Identity)
 		 */
 		FConstraintContainerHandle* AddConstraint(const FParticlePair& InConstrainedParticles, const FRigidTransform3& WorldConstraintFrame);
 		FConstraintContainerHandle* AddConstraint(const FParticlePair& InConstrainedParticles, const FTransformPair& ConnectorTransforms);
