Chaos : GeoemtryCollection
- Convex collisions

diff --git a/Engine/Source/Runtime/Experimental/Chaos/Private/GeometryCollection/GeometryCollectionCollisionStructureManager.cpp b/Engine/Source/Runtime/Experimental/Chaos/Private/GeometryCollection/GeometryCollectionCollisionStructureManager.cpp
--- a/Engine/Source/Runtime/Experimental/Chaos/Private/GeometryCollection/GeometryCollectionCollisionStructureManager.cpp
+++ b/Engine/Source/Runtime/Experimental/Chaos/Private/GeometryCollection/GeometryCollectionCollisionStructureManager.cpp
@@ -4,6 +4,7 @@
 #include "ChaosLog.h"
 #include "Chaos/Box.h"
 #include "Chaos/ErrorReporter.h"
+#include "Chaos/ImplicitObjectUnion.h"
 #include "Chaos/Levelset.h"
 #include "Chaos/Particles.h"
 #include "Chaos/Sphere.h"
@@ -245,6 +246,35 @@ FCollisionStructureManager::NewImplicitSphere(
 	return Implicit;
 }
 
+FCollisionStructureManager::FImplicit*
+FCollisionStructureManager::NewImplicitConvex(
+	const TArray<int32>& ConvexIndices,
+	const TManagedArray<TUniquePtr<Chaos::FConvex>>* ConvexGeometry,
+	const ECollisionTypeEnum CollisionType)
+{
+	if (ConvexIndices.Num())
+	{
+		TArray<TUniquePtr<Chaos::FImplicitObject>> Implicits;
+		for (auto& Index : ConvexIndices)
+		{
+			Chaos::FImplicitObject* Implicit = new Chaos::FConvex((*ConvexGeometry)[Index]->GetVertices(), (*ConvexGeometry)[Index]->GetMargin());
+			UpdateImplicitFlags(Implicit, CollisionType);
+			Implicits.Add(TUniquePtr<Chaos::FImplicitObject>(Implicit));
+		}
+
+		if (Implicits.Num() == 1)
+		{
+			return Implicits[0].Release();
+		}
+		else
+		{
+			Chaos::FImplicitObjectUnion* ImplicitUnion = new Chaos::FImplicitObjectUnion(MoveTemp(Implicits));
+			return ImplicitUnion;
+		}
+	}
+	return nullptr;
+}
+
 FCollisionStructureManager::FImplicit*
 FCollisionStructureManager::NewImplicitCapsule(
 	const Chaos::FReal Radius,
diff --git a/Engine/Source/Runtime/Experimental/Chaos/Private/PhysicsProxy/GeometryCollectionPhysicsProxy.cpp b/Engine/Source/Runtime/Experimental/Chaos/Private/PhysicsProxy/GeometryCollectionPhysicsProxy.cpp
--- a/Engine/Source/Runtime/Experimental/Chaos/Private/PhysicsProxy/GeometryCollectionPhysicsProxy.cpp
+++ b/Engine/Source/Runtime/Experimental/Chaos/Private/PhysicsProxy/GeometryCollectionPhysicsProxy.cpp
@@ -1813,6 +1813,10 @@ void FGeometryCollectionPhysicsProxy::InitializeSharedCollisionStructures(
 		RestCollection.AddAttribute<FGeometryDynamicCollection::FSharedImplicit>(
 			FGeometryDynamicCollection::ImplicitsAttribute, FTransformCollection::TransformGroup);
 
+	TManagedArray<TSet<int32>>* TransformToConvexIndices = RestCollection.FindAttribute<TSet<int32>>("TransformToConvexIndices", FTransformCollection::TransformGroup);
+	TManagedArray<TUniquePtr<Chaos::FConvex>>* ConvexGeometry = RestCollection.FindAttribute<TUniquePtr<Chaos::FConvex>>("ConvexHull", "Convex");
+
+
 	// @todo(chaos_transforms) : do we still use this?
 	TManagedArray<FTransform>& CollectionMassToLocal =
 		RestCollection.AddAttribute<FTransform>(
@@ -2147,6 +2151,17 @@ void FGeometryCollectionPhysicsProxy::InitializeSharedCollisionStructures(
 							SizeSpecificData.CollisionShapesData[0].CollisionObjectReductionPercentage,
 							SizeSpecificData.CollisionShapesData[0].CollisionType));
 				}
+				else if (SizeSpecificData.CollisionShapesData[0].ImplicitType == EImplicitTypeEnum::Chaos_Implicit_Convex)
+				{
+					if (ConvexGeometry && TransformToConvexIndices)
+					{
+						CollectionImplicits[TransformGroupIndex] = FGeometryDynamicCollection::FSharedImplicit(
+								FCollisionStructureManager::NewImplicitConvex(
+									(*TransformToConvexIndices)[TransformGroupIndex].Array(),
+									ConvexGeometry,
+									SizeSpecificData.CollisionShapesData[0].CollisionType));
+					}
+				}
 				else if (SizeSpecificData.CollisionShapesData[0].ImplicitType == EImplicitTypeEnum::Chaos_Implicit_Capsule)
 				{
 					const FVector BBoxExtent = InstanceBoundingBox.GetExtent(); // FBox's extents are 1/2 (Max - Min)
@@ -2412,6 +2427,17 @@ void FGeometryCollectionPhysicsProxy::InitializeSharedCollisionStructures(
 						FCollisionStructureManager::NewSimplicial(MassSpaceParticles, *UnionMesh, CollectionImplicits[ClusterTransformIdx].Get(),
 						SharedParams.MaximumCollisionParticleCount));
 				}
+				else if (SizeSpecificData.CollisionShapesData[0].ImplicitType == EImplicitTypeEnum::Chaos_Implicit_Convex)
+				{
+					if (ConvexGeometry && TransformToConvexIndices)
+					{
+						CollectionImplicits[TransformGroupIndex] = FGeometryDynamicCollection::FSharedImplicit(
+							FCollisionStructureManager::NewImplicitConvex(
+								(*TransformToConvexIndices)[TransformGroupIndex].Array(),
+								ConvexGeometry,
+								SizeSpecificData.CollisionShapesData[0].CollisionType));
+					}
+				}
 				else if (SizeSpecificData.CollisionShapesData[0].ImplicitType == EImplicitTypeEnum::Chaos_Implicit_Capsule)
 				{
 					const FVector BBoxExtent = InstanceBoundingBox.GetExtent(); // FBox's extents are 1/2 (Max - Min)
diff --git a/Engine/Source/Runtime/Experimental/Chaos/Public/Chaos/Convex.h b/Engine/Source/Runtime/Experimental/Chaos/Public/Chaos/Convex.h
--- a/Engine/Source/Runtime/Experimental/Chaos/Public/Chaos/Convex.h
+++ b/Engine/Source/Runtime/Experimental/Chaos/Public/Chaos/Convex.h
@@ -139,6 +139,12 @@ namespace Chaos
 			return *this;
 		}
 
+		virtual TUniquePtr<FImplicitObject> Copy() const
+		{
+			return TUniquePtr<FImplicitObject>(new FConvex(GetVertices(), GetMargin()));
+		}
+
+
 		void MovePlanesAndRebuild(FReal InDelta);
 
 	private:
diff --git a/Engine/Source/Runtime/Experimental/Chaos/Public/GeometryCollection/GeometryCollectionCollisionStructureManager.h b/Engine/Source/Runtime/Experimental/Chaos/Public/GeometryCollection/GeometryCollectionCollisionStructureManager.h
--- a/Engine/Source/Runtime/Experimental/Chaos/Public/GeometryCollection/GeometryCollectionCollisionStructureManager.h
+++ b/Engine/Source/Runtime/Experimental/Chaos/Public/GeometryCollection/GeometryCollectionCollisionStructureManager.h
@@ -104,6 +104,11 @@ public:
 		const float CollisionObjectReduction,
 		const ECollisionTypeEnum CollisionType);
 
+	static FImplicit* NewImplicitConvex(
+		const TArray<int32>& ConvexIndices,
+		const TManagedArray<TUniquePtr<Chaos::FConvex>>* ConvexGeometry,
+		const ECollisionTypeEnum CollisionType);
+
 	static FImplicit* NewImplicitLevelset(
 		Chaos::FErrorReporter ErrorReporter,
 		const Chaos::FParticles& MeshParticles,
diff --git a/Engine/Source/Runtime/Experimental/Chaos/Public/GeometryCollection/GeometryCollectionSimulationTypes.h b/Engine/Source/Runtime/Experimental/Chaos/Public/GeometryCollection/GeometryCollectionSimulationTypes.h
--- a/Engine/Source/Runtime/Experimental/Chaos/Public/GeometryCollection/GeometryCollectionSimulationTypes.h
+++ b/Engine/Source/Runtime/Experimental/Chaos/Public/GeometryCollection/GeometryCollectionSimulationTypes.h
@@ -23,6 +23,7 @@ enum class EImplicitTypeEnum : uint8
 	Chaos_Implicit_Capsule UMETA(DisplayName = "Capsule"),
 	Chaos_Implicit_LevelSet UMETA(DisplayName = "Level Set"),
 	Chaos_Implicit_None UMETA(DisplayName = "None"),
+	Chaos_Implicit_Convex UMETA(DisplayName = "Convex"),
 	//
 	Chaos_Max                UMETA(Hidden)
 };
