Add option to use full precision UVs on geometry collection scene proxy

diff --git a/Engine/Source/Runtime/Experimental/GeometryCollectionEngine/Private/GeometryCollection/GeometryCollectionObject.cpp b/Engine/Source/Runtime/Experimental/GeometryCollectionEngine/Private/GeometryCollection/GeometryCollectionObject.cpp
--- a/Engine/Source/Runtime/Experimental/GeometryCollectionEngine/Private/GeometryCollection/GeometryCollectionObject.cpp
+++ b/Engine/Source/Runtime/Experimental/GeometryCollectionEngine/Private/GeometryCollection/GeometryCollectionObject.cpp
@@ -62,6 +62,7 @@ UGeometryCollection::UGeometryCollection(const FObjectInitializer& ObjectInitial
 	, MaxClusterLevel(100)
 	, DamageThreshold({ 250.0 })
 	, ClusterConnectionType(EClusterConnectionTypeEnum::Chaos_PointImplicit)
+	, bUseFullPrecisionUVs(false)
 	, bStripOnCook(false)
 	, EnableNanite(false)
 #if WITH_EDITORONLY_DATA
@@ -1091,6 +1092,10 @@ void UGeometryCollection::PostEditChangeProperty(struct FPropertyChangedEvent& P
 			InvalidateCollection();
 			EnsureDataIsCooked();
 		}
+		else if (PropertyChangedEvent.Property->GetFName() == GET_MEMBER_NAME_CHECKED(UGeometryCollection, bUseFullPrecisionUVs))
+		{
+			InvalidateCollection();
+		}
 		else if (PropertyChangedEvent.Property->GetFName() == GET_MEMBER_NAME_CHECKED(UGeometryCollection, SizeSpecificData))
 		{
 			ValidateSizeSpecificDataDefaults();
diff --git a/Engine/Source/Runtime/Experimental/GeometryCollectionEngine/Private/GeometryCollection/GeometryCollectionSceneProxy.cpp b/Engine/Source/Runtime/Experimental/GeometryCollectionEngine/Private/GeometryCollection/GeometryCollectionSceneProxy.cpp
--- a/Engine/Source/Runtime/Experimental/GeometryCollectionEngine/Private/GeometryCollection/GeometryCollectionSceneProxy.cpp
+++ b/Engine/Source/Runtime/Experimental/GeometryCollectionEngine/Private/GeometryCollection/GeometryCollectionSceneProxy.cpp
@@ -92,6 +92,7 @@ FGeometryCollectionSceneProxy::FGeometryCollectionSceneProxy(UGeometryCollection
 	, bShowBoneColors(Component->GetShowBoneColors())
 	, bEnableBoneSelection(Component->GetEnableBoneSelection())
 	, BoneSelectionMaterialID(Component->GetBoneSelectedMaterialID())
+	, bUseFullPrecisionUVs(Component->GetRestCollection()->bUseFullPrecisionUVs)
 	, TransformVertexBuffersContainsOriginalMesh(false)
 {
 	Materials.Empty();
@@ -214,6 +215,8 @@ void FGeometryCollectionSceneProxy::InitResources()
 
 	// get vertex factory data
 	FGeometryCollectionVertexFactory::FDataType Data;
+
+	VertexBuffers.StaticMeshVertexBuffer.SetUseFullPrecisionUVs(bUseFullPrecisionUVs);
 	
 	// Init buffers
 	VertexBuffers.PositionVertexBuffer.Init(NumVertices);
diff --git a/Engine/Source/Runtime/Experimental/GeometryCollectionEngine/Private/GeometryCollection/GeometryCollectionSceneProxy.h b/Engine/Source/Runtime/Experimental/GeometryCollectionEngine/Private/GeometryCollection/GeometryCollectionSceneProxy.h
--- a/Engine/Source/Runtime/Experimental/GeometryCollectionEngine/Private/GeometryCollection/GeometryCollectionSceneProxy.h
+++ b/Engine/Source/Runtime/Experimental/GeometryCollectionEngine/Private/GeometryCollection/GeometryCollectionSceneProxy.h
@@ -263,6 +263,8 @@ class FGeometryCollectionSceneProxy final : public FPrimitiveSceneProxy
 	bool bEnableBoneSelection;
 	int BoneSelectionMaterialID;
 
+	bool bUseFullPrecisionUVs = false;
+
 	bool TransformVertexBuffersContainsOriginalMesh;
 
 public:
diff --git a/Engine/Source/Runtime/Experimental/GeometryCollectionEngine/Public/GeometryCollection/GeometryCollectionObject.h b/Engine/Source/Runtime/Experimental/GeometryCollectionEngine/Public/GeometryCollection/GeometryCollectionObject.h
--- a/Engine/Source/Runtime/Experimental/GeometryCollectionEngine/Public/GeometryCollection/GeometryCollectionObject.h
+++ b/Engine/Source/Runtime/Experimental/GeometryCollectionEngine/Public/GeometryCollection/GeometryCollectionObject.h
@@ -415,12 +415,19 @@ public:
 	UPROPERTY(EditAnywhere, Category = "EmbeddedGeometry")
 	TArray<FGeometryCollectionEmbeddedExemplar> EmbeddedGeometryExemplar;
 
+	/** Whether to use full precision UVs when rendering this geometry. (Does not apply to Nanite rendering) */
+	UPROPERTY(EditAnywhere, Category = "Rendering")
+	bool bUseFullPrecisionUVs = false;
+
 	/**
-	* Enable support for Nanite.
-	*/
+	 * Strip unnecessary data from the Geometry Collection to keep the memory footprint as small as possible.
+	 */
 	UPROPERTY(EditAnywhere, Category = "Nanite")
 	bool bStripOnCook;
 
+	/**
+	 * Enable support for Nanite.
+	 */
 	UPROPERTY(EditAnywhere, BlueprintReadOnly, Category = "Nanite")
 	bool EnableNanite;
 
