ToolsFramework: remove the AssetAPI concept from the Interactive Tools Framework. Replaced by ModelingObjectCreationAPI in ModelingComponents.
#rnx

diff --git a/Engine/Plugins/Experimental/MeshModelingToolset/Source/MeshModelingTools/Public/BakeTransformTool.h b/Engine/Plugins/Experimental/MeshModelingToolset/Source/MeshModelingTools/Public/BakeTransformTool.h
--- a/Engine/Plugins/Experimental/MeshModelingToolset/Source/MeshModelingTools/Public/BakeTransformTool.h
+++ b/Engine/Plugins/Experimental/MeshModelingToolset/Source/MeshModelingTools/Public/BakeTransformTool.h
@@ -26,8 +26,6 @@ class MESHMODELINGTOOLS_API UBakeTransformToolBuilder : public UInteractiveToolB
 	GENERATED_BODY()
 
 public:
-	IToolsContextAssetAPI* AssetAPI = nullptr;
-
 	virtual bool CanBuildTool(const FToolBuilderState& SceneState) const override;
 	virtual UInteractiveTool* BuildTool(const FToolBuilderState& SceneState) const override;
 
diff --git a/Engine/Plugins/Experimental/ModelingToolsEditorMode/Source/ModelingToolsEditorMode/Public/ModelingToolsEditorMode.h b/Engine/Plugins/Experimental/ModelingToolsEditorMode/Source/ModelingToolsEditorMode/Public/ModelingToolsEditorMode.h
--- a/Engine/Plugins/Experimental/ModelingToolsEditorMode/Source/ModelingToolsEditorMode/Public/ModelingToolsEditorMode.h
+++ b/Engine/Plugins/Experimental/ModelingToolsEditorMode/Source/ModelingToolsEditorMode/Public/ModelingToolsEditorMode.h
@@ -13,7 +13,6 @@
 #include "ModelingToolsEditorMode.generated.h"
 
 class FEditorComponentSourceFactory;
-class FEditorToolAssetAPI;
 class FUICommandList;
 class FStylusStateTracker;		// for stylus events
 
diff --git a/Engine/Source/Editor/Experimental/EditorInteractiveToolsFramework/Private/EdModeInteractiveToolsContext.cpp b/Engine/Source/Editor/Experimental/EditorInteractiveToolsFramework/Private/EdModeInteractiveToolsContext.cpp
--- a/Engine/Source/Editor/Experimental/EditorInteractiveToolsFramework/Private/EdModeInteractiveToolsContext.cpp
+++ b/Engine/Source/Editor/Experimental/EditorInteractiveToolsFramework/Private/EdModeInteractiveToolsContext.cpp
@@ -24,7 +24,6 @@
 #include "Components/StaticMeshComponent.h"
 
 #include "ToolContextInterfaces.h"
-#include "Tools/EditorToolAssetAPI.h"
 #include "InteractiveToolObjects.h"
 #include "InteractiveToolsSelectionStoreSubsystem.h"
 #include "BaseBehaviors/ClickDragBehavior.h"
@@ -535,7 +534,6 @@ UEdModeInteractiveToolsContext::UEdModeInteractiveToolsContext()
 {
 	QueriesAPI = nullptr;
 	TransactionAPI = nullptr;
-	AssetAPI = nullptr;
 }
 
 
@@ -596,21 +594,19 @@ void UEdModeInteractiveToolsContext::Shutdown()
 	}
 }
 
-void UEdModeInteractiveToolsContext::InitializeContextFromEdMode(FEdMode* EditorModeIn, IToolsContextAssetAPI* UseAssetAPI)
+void UEdModeInteractiveToolsContext::InitializeContextFromEdMode(FEdMode* EditorModeIn)
 {
 	check(EditorModeIn);
-	InitializeContextWithEditorModeManager(EditorModeIn->GetModeManager(), UseAssetAPI);
+	InitializeContextWithEditorModeManager(EditorModeIn->GetModeManager());
 }
 
-void UEdModeInteractiveToolsContext::InitializeContextWithEditorModeManager(FEditorModeTools* InEditorModeManager,
-	IToolsContextAssetAPI* UseAssetAPI)
+void UEdModeInteractiveToolsContext::InitializeContextWithEditorModeManager(FEditorModeTools* InEditorModeManager)
 {
 	check(InEditorModeManager);
 	EditorModeManager = InEditorModeManager;
 
 	this->TransactionAPI = new FEdModeToolsContextTransactionImpl(this, InEditorModeManager);
 	this->QueriesAPI = new FEdModeToolsContextQueriesImpl(this, InEditorModeManager);
-	this->AssetAPI = (UseAssetAPI != nullptr) ? UseAssetAPI: new FEditorToolAssetAPI();
 
 	Initialize(QueriesAPI, TransactionAPI);
 
@@ -641,12 +637,6 @@ void UEdModeInteractiveToolsContext::ShutdownContext()
 		delete TransactionAPI;
 		TransactionAPI = nullptr;
 	}
-
-	if (AssetAPI != nullptr)
-	{
-		delete AssetAPI;
-		AssetAPI = nullptr;
-	}
 }
 
 
diff --git a/Engine/Source/Editor/Experimental/EditorInteractiveToolsFramework/Public/EdModeInteractiveToolsContext.h b/Engine/Source/Editor/Experimental/EditorInteractiveToolsFramework/Public/EdModeInteractiveToolsContext.h
--- a/Engine/Source/Editor/Experimental/EditorInteractiveToolsFramework/Public/EdModeInteractiveToolsContext.h
+++ b/Engine/Source/Editor/Experimental/EditorInteractiveToolsFramework/Public/EdModeInteractiveToolsContext.h
@@ -34,10 +34,8 @@ public:
 	UEdModeInteractiveToolsContext();
 
 	UE_DEPRECATED(5.0, "You should no longer create your own EdModeInteractiveToolsContext; use the one in the FEditorModeTools::GetInteractiveToolsContext instead.")
-	void InitializeContextFromEdMode(FEdMode* EditorModeIn,
-		IToolsContextAssetAPI* UseAssetAPI = nullptr);
-	void InitializeContextWithEditorModeManager(FEditorModeTools* InEditorModeManager,
-		IToolsContextAssetAPI* UseAssetAPI = nullptr);
+	void InitializeContextFromEdMode(FEdMode* EditorModeIn);
+	void InitializeContextWithEditorModeManager(FEditorModeTools* InEditorModeManager);
 	void ShutdownContext();
 
 	// default behavior is to accept active tool
@@ -51,7 +49,6 @@ public:
 
 	IToolsContextQueriesAPI* GetQueriesAPI() const { return QueriesAPI; }
 	IToolsContextTransactionsAPI* GetTransactionAPI() const { return TransactionAPI; }
-	IToolsContextAssetAPI* GetAssetAPI() const { return AssetAPI; }
 
 	void PostInvalidation();
 
@@ -116,7 +113,6 @@ protected:
 	// EdMode implementation of InteractiveToolFramework APIs - see ToolContextInterfaces.h
 	IToolsContextQueriesAPI* QueriesAPI;
 	IToolsContextTransactionsAPI* TransactionAPI;
-	IToolsContextAssetAPI* AssetAPI;
 
 	// Tools need to be able to Invalidate the view, in case it is not Realtime.
 	// Currently we do this very aggressively, and also force Realtime to be on, but in general we should be able to rely on Invalidation.
diff --git a/Engine/Source/Editor/UnrealEd/Private/Tools/EditorToolAssetAPI.cpp b/Engine/Source/Editor/UnrealEd/Private/Tools/EditorToolAssetAPI.cpp
deleted file mode 100644
diff --git a/Engine/Source/Editor/UnrealEd/Public/Tools/EditorToolAssetAPI.h b/Engine/Source/Editor/UnrealEd/Public/Tools/EditorToolAssetAPI.h
deleted file mode 100644
diff --git a/Engine/Source/Runtime/Experimental/InteractiveToolsFramework/Private/ToolContextInterfaces.cpp b/Engine/Source/Runtime/Experimental/InteractiveToolsFramework/Private/ToolContextInterfaces.cpp
--- a/Engine/Source/Runtime/Experimental/InteractiveToolsFramework/Private/ToolContextInterfaces.cpp
+++ b/Engine/Source/Runtime/Experimental/InteractiveToolsFramework/Private/ToolContextInterfaces.cpp
@@ -4,6 +4,3 @@
 #include "MeshDescription.h"
 
 
-FGeneratedStaticMeshAssetConfig::~FGeneratedStaticMeshAssetConfig()
-{
-}
\ No newline at end of file
diff --git a/Engine/Source/Runtime/Experimental/InteractiveToolsFramework/Public/ToolContextInterfaces.h b/Engine/Source/Runtime/Experimental/InteractiveToolsFramework/Public/ToolContextInterfaces.h
--- a/Engine/Source/Runtime/Experimental/InteractiveToolsFramework/Public/ToolContextInterfaces.h
+++ b/Engine/Source/Runtime/Experimental/InteractiveToolsFramework/Public/ToolContextInterfaces.h
@@ -13,7 +13,6 @@ class UWorld;
 class AActor;
 class UActorComponent;
 class FToolCommandChange;
-class UPackage;
 class FPrimitiveDrawInterface;
 class FSceneView;
 class FViewport;
@@ -21,7 +20,6 @@ class UInteractiveToolManager;
 class UInteractiveGizmoManager;
 class UToolTargetManager;
 struct FMeshDescription;
-class UTexture2D;
 class UInteractiveToolStorableSelection;
 struct FTypedElementHandle;
 class UTypedElementSelectionSet;
@@ -415,68 +413,3 @@ public:
 	virtual EViewInteractionState GetViewInteractionState() = 0;
 };
 
-
-/**
- * FGeneratedStaticMeshAssetConfig is passed to IToolsContextAssetAPI::GenerateStaticMeshActor() to
- * provide the underlying mesh, materials, and configuration settings. Note that all implementations
- * may not use/respect all settings.
- */
-struct FGeneratedStaticMeshAssetConfig
-{
-	INTERACTIVETOOLSFRAMEWORK_API ~FGeneratedStaticMeshAssetConfig();
-
-	TUniquePtr<FMeshDescription> MeshDescription;
-	TArray<UMaterialInterface*> Materials; // Materials for the component (and asset if no AssetMaterials)
-	TArray<UMaterialInterface*> AssetMaterials; // Optional, for the asset only
-
-	bool bEnableRecomputeNormals = false;
-	bool bEnableRecomputeTangents = false;
-
-	bool bEnablePhysics = true;
-	bool bEnableComplexAsSimpleCollision = true;
-};
-
-
-/**
- * Users of the Tools Framework need to provide an IToolsContextAssetAPI implementation
- * that allows Packages and Assets to be created/saved. Note that this is not strictly
- * necessary, for example a trivial implementation could just store things in the Transient
- * package and not do any saving.
- */
-class IToolsContextAssetAPI
-{
-public:
-	virtual ~IToolsContextAssetAPI() {}
-
-	/** 
-	 * Get a path to save assets in that is relative to the given UWorld. 
-	 */
-	virtual FString GetWorldRelativeAssetRootPath(const UWorld* World) = 0;
-
-	/** 
-	 * Get a "currently-visible/selected" location to save assets in. For example the currently-visible path in the Editor Content Browser.
-	 */
-	virtual FString GetActiveAssetFolderPath() = 0;
-
-	/** Allow the user to select a path and filename for an asset using a modal dialog */
-	virtual FString InteractiveSelectAssetPath(const FString& DefaultAssetName, const FText& DialogTitleMessage) = 0;
-
-	/**
-	 * Creates a new package for an asset
-	 * @param FolderPath path for new package
-	 * @param AssetBaseName base name for asset
-	 * @param UniqueAssetNameOut unique name in form of AssetBaseName##, where ## is a unqiue index
-	 * @return new package
-	 */
-	virtual UPackage* MakeNewAssetPackage(const FString& FolderPath, const FString& AssetBaseName, FString& UniqueAssetNameOut) = 0;
-
-	/** Request saving of asset to persistent storage via something like an interactive popup dialog */
-	virtual void InteractiveSaveGeneratedAsset(UObject* Asset, UPackage* AssetPackage) = 0;
-
-	/** Autosave asset to persistent storage */
-	virtual void AutoSaveGeneratedAsset(UObject* Asset, UPackage* AssetPackage) = 0;
-
-	/** Notify that asset has been created and is dirty */
-	virtual void NotifyGeneratedAssetModified(UObject* Asset, UPackage* AssetPackage) = 0;
-};
-
