Chaos: Geometry Collection Asset
- Allow empty collisions array in the size specific data.

diff --git a/Engine/Source/Runtime/Experimental/Chaos/Private/PhysicsProxy/GeometryCollectionPhysicsProxy.cpp b/Engine/Source/Runtime/Experimental/Chaos/Private/PhysicsProxy/GeometryCollectionPhysicsProxy.cpp
--- a/Engine/Source/Runtime/Experimental/Chaos/Private/PhysicsProxy/GeometryCollectionPhysicsProxy.cpp
+++ b/Engine/Source/Runtime/Experimental/Chaos/Private/PhysicsProxy/GeometryCollectionPhysicsProxy.cpp
@@ -460,7 +460,6 @@ void FGeometryCollectionPhysicsProxy::Initialize(Chaos::FPBDRigidsEvolutionBase
 void FGeometryCollectionPhysicsProxy::InitializeDynamicCollection(FGeometryDynamicCollection& DynamicCollection, const FGeometryCollection& RestCollection, const FSimulationParameters& Params)
 {
 	// @todo(GCCollisionShapes) : add support for multiple shapes, currently just one. 
-	ensure(Params.Shared.SizeSpecificData[0].CollisionShapesData.Num());
 
 	// 
 	// This function will use the rest collection to populate the dynamic collection. 
@@ -491,6 +490,7 @@ void FGeometryCollectionPhysicsProxy::InitializeDynamicCollection(FGeometryDynam
 		static const bool bAnalyticsDisabled = (AnalyticDisableCVar && AnalyticDisableCVar->GetBool());
 
 		if (RestCollection.HasAttribute(DynamicCollection.SimplicialsAttribute, FTransformCollection::TransformGroup)
+			&& Params.Shared.SizeSpecificData[0].CollisionShapesData.Num()
 			&& (Params.Shared.SizeSpecificData[0].CollisionShapesData[0].CollisionType == ECollisionTypeEnum::Chaos_Surface_Volumetric || bAnalyticsDisabled))
 		{
 			const auto& RestSimplicials = RestCollection.GetAttribute<TUniquePtr<FSimplicial>>(
@@ -1763,7 +1763,7 @@ void FGeometryCollectionPhysicsProxy::InitializeSharedCollisionStructures(
 	FGeometryCollection& RestCollection,
 	const FSharedSimulationParameters& SharedParams)
 {
-	check(SharedParams.SizeSpecificData.Num() && SharedParams.SizeSpecificData[0].CollisionShapesData.Num());
+	check(SharedParams.SizeSpecificData.Num());
 
 	FString BaseErrorPrefix = ErrorReporter.GetPrefix();
 
@@ -2092,116 +2092,119 @@ void FGeometryCollectionPhysicsProxy::InitializeSharedCollisionStructures(
 			const int32 SizeSpecificIdx = FindSizeSpecificIdx(SharedParams.SizeSpecificData, InstanceBoundingBox);
 			const FSharedSimulationSizeSpecificData& SizeSpecificData = SharedParams.SizeSpecificData[SizeSpecificIdx];
 
-			//
-			//  Build the simplicial for the rest collection. This will be used later in the DynamicCollection to 
-			//  populate the collision structures of the simulation. 
-			//
-			if (ensureMsgf(TriMesh, TEXT("No Triangle representation")))
+			if (SizeSpecificData.CollisionShapesData.Num())
 			{
-				Chaos::FBVHParticles* Simplicial =
-					FCollisionStructureManager::NewSimplicial(
-						MassSpaceParticles,
-						BoneMap,
-						SizeSpecificData.CollisionShapesData[0].CollisionType,
-						*TriMesh,
-						SizeSpecificData.CollisionShapesData[0].CollisionParticleData.CollisionParticlesFraction);
-				CollectionSimplicials[TransformGroupIndex] = TUniquePtr<FSimplicial>(Simplicial); // CollectionSimplicials is in the TransformGroup
-				//ensureMsgf(CollectionSimplicials[TransformGroupIndex], TEXT("No simplicial representation."));
-				if (!CollectionSimplicials[TransformGroupIndex]->Size())
+				//
+				//  Build the simplicial for the rest collection. This will be used later in the DynamicCollection to 
+				//  populate the collision structures of the simulation. 
+				//
+				if (ensureMsgf(TriMesh, TEXT("No Triangle representation")))
 				{
-					ensureMsgf(false, TEXT("Simplicial is empty."));
-				}
-
-				if (SizeSpecificData.CollisionShapesData[0].ImplicitType == EImplicitTypeEnum::Chaos_Implicit_LevelSet)
-				{
-					LocalErrorReporter.SetPrefix(BaseErrorPrefix + " | Transform Index: " + FString::FromInt(TransformGroupIndex) + " of " + FString::FromInt(TransformIndex.Num()));
-					CollectionImplicits[TransformGroupIndex] = FGeometryDynamicCollection::FSharedImplicit(
-						FCollisionStructureManager::NewImplicitLevelset(
-							LocalErrorReporter,
+					Chaos::FBVHParticles* Simplicial =
+						FCollisionStructureManager::NewSimplicial(
 							MassSpaceParticles,
+							BoneMap,
+							SizeSpecificData.CollisionShapesData[0].CollisionType,
 							*TriMesh,
-							InstanceBoundingBox,
-							SizeSpecificData.CollisionShapesData[0].LevelSetData.MinLevelSetResolution,
-							SizeSpecificData.CollisionShapesData[0].LevelSetData.MaxLevelSetResolution,
-							SizeSpecificData.CollisionShapesData[0].CollisionObjectReductionPercentage,
-							SizeSpecificData.CollisionShapesData[0].CollisionType));
-					// Fall back on sphere if level set rasterization failed.
-					if (!CollectionImplicits[TransformGroupIndex])
+							SizeSpecificData.CollisionShapesData[0].CollisionParticleData.CollisionParticlesFraction);
+					CollectionSimplicials[TransformGroupIndex] = TUniquePtr<FSimplicial>(Simplicial); // CollectionSimplicials is in the TransformGroup
+					//ensureMsgf(CollectionSimplicials[TransformGroupIndex], TEXT("No simplicial representation."));
+					if (!CollectionSimplicials[TransformGroupIndex]->Size())
 					{
+						ensureMsgf(false, TEXT("Simplicial is empty."));
+					}
+
+					if (SizeSpecificData.CollisionShapesData[0].ImplicitType == EImplicitTypeEnum::Chaos_Implicit_LevelSet)
+					{
+						LocalErrorReporter.SetPrefix(BaseErrorPrefix + " | Transform Index: " + FString::FromInt(TransformGroupIndex) + " of " + FString::FromInt(TransformIndex.Num()));
 						CollectionImplicits[TransformGroupIndex] = FGeometryDynamicCollection::FSharedImplicit(
-							FCollisionStructureManager::NewImplicitSphere(
-								InnerRadius[GeometryIndex],
+							FCollisionStructureManager::NewImplicitLevelset(
+								LocalErrorReporter,
+								MassSpaceParticles,
+								*TriMesh,
+								InstanceBoundingBox,
+								SizeSpecificData.CollisionShapesData[0].LevelSetData.MinLevelSetResolution,
+								SizeSpecificData.CollisionShapesData[0].LevelSetData.MaxLevelSetResolution,
 								SizeSpecificData.CollisionShapesData[0].CollisionObjectReductionPercentage,
 								SizeSpecificData.CollisionShapesData[0].CollisionType));
+						// Fall back on sphere if level set rasterization failed.
+						if (!CollectionImplicits[TransformGroupIndex])
+						{
+							CollectionImplicits[TransformGroupIndex] = FGeometryDynamicCollection::FSharedImplicit(
+								FCollisionStructureManager::NewImplicitSphere(
+									InnerRadius[GeometryIndex],
+									SizeSpecificData.CollisionShapesData[0].CollisionObjectReductionPercentage,
+									SizeSpecificData.CollisionShapesData[0].CollisionType));
+						}
 					}
-				}
-				else if (SizeSpecificData.CollisionShapesData[0].ImplicitType == EImplicitTypeEnum::Chaos_Implicit_Box)
-				{
-					CollectionImplicits[TransformGroupIndex] = FGeometryDynamicCollection::FSharedImplicit(
-						FCollisionStructureManager::NewImplicitBox(
-							InstanceBoundingBox,
-							SizeSpecificData.CollisionShapesData[0].CollisionObjectReductionPercentage,
-							SizeSpecificData.CollisionShapesData[0].CollisionType));
-				}
-				else if (SizeSpecificData.CollisionShapesData[0].ImplicitType == EImplicitTypeEnum::Chaos_Implicit_Sphere)
-				{
-					CollectionImplicits[TransformGroupIndex] = FGeometryDynamicCollection::FSharedImplicit(
-						FCollisionStructureManager::NewImplicitSphere(
-							InnerRadius[GeometryIndex],
-							SizeSpecificData.CollisionShapesData[0].CollisionObjectReductionPercentage,
-							SizeSpecificData.CollisionShapesData[0].CollisionType));
-				}
-				else if (SizeSpecificData.CollisionShapesData[0].ImplicitType == EImplicitTypeEnum::Chaos_Implicit_Convex)
-				{
-					if (ConvexGeometry && TransformToConvexIndices)
+					else if (SizeSpecificData.CollisionShapesData[0].ImplicitType == EImplicitTypeEnum::Chaos_Implicit_Box)
 					{
 						CollectionImplicits[TransformGroupIndex] = FGeometryDynamicCollection::FSharedImplicit(
-								FCollisionStructureManager::NewImplicitConvex(
-									(*TransformToConvexIndices)[TransformGroupIndex].Array(),
-									ConvexGeometry,
-									SizeSpecificData.CollisionShapesData[0].CollisionType));
+							FCollisionStructureManager::NewImplicitBox(
+								InstanceBoundingBox,
+								SizeSpecificData.CollisionShapesData[0].CollisionObjectReductionPercentage,
+								SizeSpecificData.CollisionShapesData[0].CollisionType));
 					}
-				}
-				else if (SizeSpecificData.CollisionShapesData[0].ImplicitType == EImplicitTypeEnum::Chaos_Implicit_Capsule)
-				{
-					const FVector BBoxExtent = InstanceBoundingBox.GetExtent(); // FBox's extents are 1/2 (Max - Min)
-					const Chaos::FReal Radius = BBoxExtent.GetAbsMin();
-					const Chaos::FReal Length = (BBoxExtent.GetAbsMax() - Radius) * 2.f;
-
-					if (Length < SMALL_NUMBER)
+					else if (SizeSpecificData.CollisionShapesData[0].ImplicitType == EImplicitTypeEnum::Chaos_Implicit_Sphere)
 					{
-						// make a more optimized shape : sphere
 						CollectionImplicits[TransformGroupIndex] = FGeometryDynamicCollection::FSharedImplicit(
 							FCollisionStructureManager::NewImplicitSphere(
-								Radius,
+								InnerRadius[GeometryIndex],
 								SizeSpecificData.CollisionShapesData[0].CollisionObjectReductionPercentage,
 								SizeSpecificData.CollisionShapesData[0].CollisionType));
 					}
+					else if (SizeSpecificData.CollisionShapesData[0].ImplicitType == EImplicitTypeEnum::Chaos_Implicit_Convex)
+					{
+						if (ConvexGeometry && TransformToConvexIndices)
+						{
+							CollectionImplicits[TransformGroupIndex] = FGeometryDynamicCollection::FSharedImplicit(
+									FCollisionStructureManager::NewImplicitConvex(
+										(*TransformToConvexIndices)[TransformGroupIndex].Array(),
+										ConvexGeometry,
+										SizeSpecificData.CollisionShapesData[0].CollisionType));
+						}
+					}
+					else if (SizeSpecificData.CollisionShapesData[0].ImplicitType == EImplicitTypeEnum::Chaos_Implicit_Capsule)
+					{
+						const FVector BBoxExtent = InstanceBoundingBox.GetExtent(); // FBox's extents are 1/2 (Max - Min)
+						const Chaos::FReal Radius = BBoxExtent.GetAbsMin();
+						const Chaos::FReal Length = (BBoxExtent.GetAbsMax() - Radius) * 2.f;
+
+						if (Length < SMALL_NUMBER)
+						{
+							// make a more optimized shape : sphere
+							CollectionImplicits[TransformGroupIndex] = FGeometryDynamicCollection::FSharedImplicit(
+								FCollisionStructureManager::NewImplicitSphere(
+									Radius,
+									SizeSpecificData.CollisionShapesData[0].CollisionObjectReductionPercentage,
+									SizeSpecificData.CollisionShapesData[0].CollisionType));
+						}
+						else
+						{
+							CollectionImplicits[TransformGroupIndex] = FGeometryDynamicCollection::FSharedImplicit(
+								FCollisionStructureManager::NewImplicitCapsule(
+									Radius,
+									Length,
+									SizeSpecificData.CollisionShapesData[0].CollisionObjectReductionPercentage,
+									SizeSpecificData.CollisionShapesData[0].CollisionType));
+						}
+					}
+					else if (SizeSpecificData.CollisionShapesData[0].ImplicitType == EImplicitTypeEnum::Chaos_Implicit_None)
+					{
+						CollectionImplicits[TransformGroupIndex] = nullptr;
+					}
 					else
 					{
-						CollectionImplicits[TransformGroupIndex] = FGeometryDynamicCollection::FSharedImplicit(
-							FCollisionStructureManager::NewImplicitCapsule(
-								Radius,
-								Length,
-								SizeSpecificData.CollisionShapesData[0].CollisionObjectReductionPercentage,
-								SizeSpecificData.CollisionShapesData[0].CollisionType));
+						ensure(false); // unsupported implicit type!
 					}
-				}
-				else if (SizeSpecificData.CollisionShapesData[0].ImplicitType == EImplicitTypeEnum::Chaos_Implicit_None)
-				{
-					CollectionImplicits[TransformGroupIndex] = nullptr;
-				}
-				else
-				{
-					ensure(false); // unsupported implicit type!
-				}
 
-				if (CollectionImplicits[TransformGroupIndex] && CollectionImplicits[TransformGroupIndex]->HasBoundingBox())
-				{
-					const auto Implicit = CollectionImplicits[TransformGroupIndex];
-					const auto BBox = Implicit->BoundingBox();
-					const FVec3 Extents = BBox.Extents(); // Chaos::FAABB3::Extents() is Max - Min
-					MaxChildBounds = MaxChildBounds.ComponentwiseMax(Extents);
+					if (CollectionImplicits[TransformGroupIndex] && CollectionImplicits[TransformGroupIndex]->HasBoundingBox())
+					{
+						const auto Implicit = CollectionImplicits[TransformGroupIndex];
+						const auto BBox = Implicit->BoundingBox();
+						const FVec3 Extents = BBox.Extents(); // Chaos::FAABB3::Extents() is Max - Min
+						MaxChildBounds = MaxChildBounds.ComponentwiseMax(Extents);
+					}
 				}
 			}
 		}
@@ -2361,115 +2364,118 @@ void FGeometryCollectionPhysicsProxy::InitializeSharedCollisionStructures(
 				const int32 SizeSpecificIdx = FindSizeSpecificIdx(SharedParams.SizeSpecificData, InstanceBoundingBox);
 				const FSharedSimulationSizeSpecificData& SizeSpecificData = SharedParams.SizeSpecificData[SizeSpecificIdx];
 
-				if (SizeSpecificData.CollisionShapesData[0].ImplicitType == EImplicitTypeEnum::Chaos_Implicit_LevelSet)
+				if (SizeSpecificData.CollisionShapesData.Num())
 				{
-					const FVec3 Scale = 2 * InstanceBoundingBox.GetExtent() / MaxChildBounds; // FBox's extents are 1/2 (Max - Min)
-					const FReal ScaleMax = Scale.GetAbsMax();
-					const FReal ScaleMin = Scale.GetAbsMin();
-
-					FReal MinResolution = ScaleMin * (FReal)SizeSpecificData.CollisionShapesData[0].LevelSetData.MinLevelSetResolution;
-					MinResolution = FMath::Clamp(MinResolution,
-						(FReal)SizeSpecificData.CollisionShapesData[0].LevelSetData.MinLevelSetResolution,
-						(FReal)SizeSpecificData.CollisionShapesData[0].LevelSetData.MinClusterLevelSetResolution);
-
-					FReal MaxResolution = ScaleMax * (FReal)SizeSpecificData.CollisionShapesData[0].LevelSetData.MaxLevelSetResolution;
-					MaxResolution = FMath::Clamp(MaxResolution,
-						(FReal)SizeSpecificData.CollisionShapesData[0].LevelSetData.MaxLevelSetResolution,
-						(FReal)SizeSpecificData.CollisionShapesData[0].LevelSetData.MaxClusterLevelSetResolution);
-
-					//don't support non level-set serialization
-					ErrorReporter.SetPrefix(BaseErrorPrefix + " | Cluster Transform Index: " + FString::FromInt(ClusterTransformIdx));
-					CollectionImplicits[ClusterTransformIdx] = FGeometryDynamicCollection::FSharedImplicit(
-						FCollisionStructureManager::NewImplicitLevelset(
-							ErrorReporter,
-							MassSpaceParticles,
-							*UnionMesh,
-							InstanceBoundingBox,
-							FMath::FloorToInt(MinResolution),
-							FMath::FloorToInt(MaxResolution),
-							SizeSpecificData.CollisionShapesData[0].CollisionObjectReductionPercentage,
-							SizeSpecificData.CollisionShapesData[0].CollisionType));
-					// Fall back on sphere if level set rasterization failed.
-					if (!CollectionImplicits[ClusterTransformIdx])
+					if (SizeSpecificData.CollisionShapesData[0].ImplicitType == EImplicitTypeEnum::Chaos_Implicit_LevelSet)
 					{
+						const FVec3 Scale = 2 * InstanceBoundingBox.GetExtent() / MaxChildBounds; // FBox's extents are 1/2 (Max - Min)
+						const FReal ScaleMax = Scale.GetAbsMax();
+						const FReal ScaleMin = Scale.GetAbsMin();
+
+						FReal MinResolution = ScaleMin * (FReal)SizeSpecificData.CollisionShapesData[0].LevelSetData.MinLevelSetResolution;
+						MinResolution = FMath::Clamp(MinResolution,
+							(FReal)SizeSpecificData.CollisionShapesData[0].LevelSetData.MinLevelSetResolution,
+							(FReal)SizeSpecificData.CollisionShapesData[0].LevelSetData.MinClusterLevelSetResolution);
+
+						FReal MaxResolution = ScaleMax * (FReal)SizeSpecificData.CollisionShapesData[0].LevelSetData.MaxLevelSetResolution;
+						MaxResolution = FMath::Clamp(MaxResolution,
+							(FReal)SizeSpecificData.CollisionShapesData[0].LevelSetData.MaxLevelSetResolution,
+							(FReal)SizeSpecificData.CollisionShapesData[0].LevelSetData.MaxClusterLevelSetResolution);
+
+						//don't support non level-set serialization
+						ErrorReporter.SetPrefix(BaseErrorPrefix + " | Cluster Transform Index: " + FString::FromInt(ClusterTransformIdx));
 						CollectionImplicits[ClusterTransformIdx] = FGeometryDynamicCollection::FSharedImplicit(
-							FCollisionStructureManager::NewImplicitSphere(
-								InstanceBoundingBox.GetExtent().GetAbsMin(), // FBox's extents are 1/2 (Max - Min)
+							FCollisionStructureManager::NewImplicitLevelset(
+								ErrorReporter,
+								MassSpaceParticles,
+								*UnionMesh,
+								InstanceBoundingBox,
+								FMath::FloorToInt(MinResolution),
+								FMath::FloorToInt(MaxResolution),
 								SizeSpecificData.CollisionShapesData[0].CollisionObjectReductionPercentage,
 								SizeSpecificData.CollisionShapesData[0].CollisionType));
-					}
+						// Fall back on sphere if level set rasterization failed.
+						if (!CollectionImplicits[ClusterTransformIdx])
+						{
+							CollectionImplicits[ClusterTransformIdx] = FGeometryDynamicCollection::FSharedImplicit(
+								FCollisionStructureManager::NewImplicitSphere(
+									InstanceBoundingBox.GetExtent().GetAbsMin(), // FBox's extents are 1/2 (Max - Min)
+									SizeSpecificData.CollisionShapesData[0].CollisionObjectReductionPercentage,
+									SizeSpecificData.CollisionShapesData[0].CollisionType));
+						}
 
-					CollectionSimplicials[ClusterTransformIdx] = TUniquePtr<FSimplicial>(
-						FCollisionStructureManager::NewSimplicial(MassSpaceParticles, *UnionMesh, CollectionImplicits[ClusterTransformIdx].Get(),
-						SharedParams.MaximumCollisionParticleCount));
-				}
-				else if (SizeSpecificData.CollisionShapesData[0].ImplicitType == EImplicitTypeEnum::Chaos_Implicit_Box)
-				{
-					CollectionImplicits[ClusterTransformIdx] = FGeometryDynamicCollection::FSharedImplicit(
-						FCollisionStructureManager::NewImplicitBox(
-							InstanceBoundingBox,
-							SizeSpecificData.CollisionShapesData[0].CollisionObjectReductionPercentage,
-							SizeSpecificData.CollisionShapesData[0].CollisionType));
-
-					CollectionSimplicials[ClusterTransformIdx] = TUniquePtr<FSimplicial>(
-						FCollisionStructureManager::NewSimplicial(MassSpaceParticles, *UnionMesh, CollectionImplicits[ClusterTransformIdx].Get(),
-						SharedParams.MaximumCollisionParticleCount));
-				}
-				else if (SizeSpecificData.CollisionShapesData[0].ImplicitType == EImplicitTypeEnum::Chaos_Implicit_Sphere)
-				{
- 					CollectionImplicits[ClusterTransformIdx] = FGeometryDynamicCollection::FSharedImplicit(
-						FCollisionStructureManager::NewImplicitSphere(
-							InstanceBoundingBox.GetExtent().GetAbsMin(), // FBox's extents are 1/2 (Max - Min)
-							SizeSpecificData.CollisionShapesData[0].CollisionObjectReductionPercentage,
-							SizeSpecificData.CollisionShapesData[0].CollisionType));
-
-					CollectionSimplicials[ClusterTransformIdx] = TUniquePtr<FSimplicial>(
-						FCollisionStructureManager::NewSimplicial(MassSpaceParticles, *UnionMesh, CollectionImplicits[ClusterTransformIdx].Get(),
-						SharedParams.MaximumCollisionParticleCount));
-				}
-				else if (SizeSpecificData.CollisionShapesData[0].ImplicitType == EImplicitTypeEnum::Chaos_Implicit_Convex)
-				{
-					if (ConvexGeometry && TransformToConvexIndices)
-					{
-						CollectionImplicits[TransformGroupIndex] = FGeometryDynamicCollection::FSharedImplicit(
-							FCollisionStructureManager::NewImplicitConvex(
-								(*TransformToConvexIndices)[TransformGroupIndex].Array(),
-								ConvexGeometry,
-								SizeSpecificData.CollisionShapesData[0].CollisionType));
+						CollectionSimplicials[ClusterTransformIdx] = TUniquePtr<FSimplicial>(
+							FCollisionStructureManager::NewSimplicial(MassSpaceParticles, *UnionMesh, CollectionImplicits[ClusterTransformIdx].Get(),
+							SharedParams.MaximumCollisionParticleCount));
 					}
-				}
-				else if (SizeSpecificData.CollisionShapesData[0].ImplicitType == EImplicitTypeEnum::Chaos_Implicit_Capsule)
-				{
-					const FVector BBoxExtent = InstanceBoundingBox.GetExtent(); // FBox's extents are 1/2 (Max - Min)
-					const FReal Radius = BBoxExtent.GetAbsMin();
-					const FReal Length = (BBoxExtent.GetAbsMax() - Radius) * 2.f;
-
-					if (Length < SMALL_NUMBER)
+					else if (SizeSpecificData.CollisionShapesData[0].ImplicitType == EImplicitTypeEnum::Chaos_Implicit_Box)
 					{
-						// make a more optimized shape : sphere
 						CollectionImplicits[ClusterTransformIdx] = FGeometryDynamicCollection::FSharedImplicit(
-							FCollisionStructureManager::NewImplicitSphere(
-								Radius, 
+							FCollisionStructureManager::NewImplicitBox(
+								InstanceBoundingBox,
 								SizeSpecificData.CollisionShapesData[0].CollisionObjectReductionPercentage,
 								SizeSpecificData.CollisionShapesData[0].CollisionType));
+
+						CollectionSimplicials[ClusterTransformIdx] = TUniquePtr<FSimplicial>(
+							FCollisionStructureManager::NewSimplicial(MassSpaceParticles, *UnionMesh, CollectionImplicits[ClusterTransformIdx].Get(),
+							SharedParams.MaximumCollisionParticleCount));
 					}
-					else
+					else if (SizeSpecificData.CollisionShapesData[0].ImplicitType == EImplicitTypeEnum::Chaos_Implicit_Sphere)
 					{
 						CollectionImplicits[ClusterTransformIdx] = FGeometryDynamicCollection::FSharedImplicit(
-							FCollisionStructureManager::NewImplicitCapsule(
-								Radius,
-								Length,
+							FCollisionStructureManager::NewImplicitSphere(
+								InstanceBoundingBox.GetExtent().GetAbsMin(), // FBox's extents are 1/2 (Max - Min)
 								SizeSpecificData.CollisionShapesData[0].CollisionObjectReductionPercentage,
 								SizeSpecificData.CollisionShapesData[0].CollisionType));
-					}
-					CollectionSimplicials[ClusterTransformIdx] = TUniquePtr<FSimplicial>(
-						FCollisionStructureManager::NewSimplicial(MassSpaceParticles, *UnionMesh, CollectionImplicits[ClusterTransformIdx].Get(),
+
+						CollectionSimplicials[ClusterTransformIdx] = TUniquePtr<FSimplicial>(
+							FCollisionStructureManager::NewSimplicial(MassSpaceParticles, *UnionMesh, CollectionImplicits[ClusterTransformIdx].Get(),
 							SharedParams.MaximumCollisionParticleCount));
-				}
-				else // Assume it's a union???
-				{
-					CollectionImplicits[ClusterTransformIdx].Reset();	//union so just set as null
-					CollectionSimplicials[ClusterTransformIdx].Reset();
+					}
+					else if (SizeSpecificData.CollisionShapesData[0].ImplicitType == EImplicitTypeEnum::Chaos_Implicit_Convex)
+					{
+						if (ConvexGeometry && TransformToConvexIndices)
+						{
+							CollectionImplicits[TransformGroupIndex] = FGeometryDynamicCollection::FSharedImplicit(
+								FCollisionStructureManager::NewImplicitConvex(
+									(*TransformToConvexIndices)[TransformGroupIndex].Array(),
+									ConvexGeometry,
+									SizeSpecificData.CollisionShapesData[0].CollisionType));
+						}
+					}
+					else if (SizeSpecificData.CollisionShapesData[0].ImplicitType == EImplicitTypeEnum::Chaos_Implicit_Capsule)
+					{
+						const FVector BBoxExtent = InstanceBoundingBox.GetExtent(); // FBox's extents are 1/2 (Max - Min)
+						const FReal Radius = BBoxExtent.GetAbsMin();
+						const FReal Length = (BBoxExtent.GetAbsMax() - Radius) * 2.f;
+
+						if (Length < SMALL_NUMBER)
+						{
+							// make a more optimized shape : sphere
+							CollectionImplicits[ClusterTransformIdx] = FGeometryDynamicCollection::FSharedImplicit(
+								FCollisionStructureManager::NewImplicitSphere(
+									Radius, 
+									SizeSpecificData.CollisionShapesData[0].CollisionObjectReductionPercentage,
+									SizeSpecificData.CollisionShapesData[0].CollisionType));
+						}
+						else
+						{
+							CollectionImplicits[ClusterTransformIdx] = FGeometryDynamicCollection::FSharedImplicit(
+								FCollisionStructureManager::NewImplicitCapsule(
+									Radius,
+									Length,
+									SizeSpecificData.CollisionShapesData[0].CollisionObjectReductionPercentage,
+									SizeSpecificData.CollisionShapesData[0].CollisionType));
+						}
+						CollectionSimplicials[ClusterTransformIdx] = TUniquePtr<FSimplicial>(
+							FCollisionStructureManager::NewSimplicial(MassSpaceParticles, *UnionMesh, CollectionImplicits[ClusterTransformIdx].Get(),
+								SharedParams.MaximumCollisionParticleCount));
+					}
+					else // Assume it's a union???
+					{
+						CollectionImplicits[ClusterTransformIdx].Reset();	//union so just set as null
+						CollectionSimplicials[ClusterTransformIdx].Reset();
+					}
 				}
 
 				TriangleMeshesArray[ClusterTransformIdx] = MoveTemp(UnionMesh);
diff --git a/Engine/Source/Runtime/Experimental/GeometryCollectionEngine/Private/GeometryCollection/GeometryCollectionObject.cpp b/Engine/Source/Runtime/Experimental/GeometryCollectionEngine/Private/GeometryCollection/GeometryCollectionObject.cpp
--- a/Engine/Source/Runtime/Experimental/GeometryCollectionEngine/Private/GeometryCollection/GeometryCollectionObject.cpp
+++ b/Engine/Source/Runtime/Experimental/GeometryCollectionEngine/Private/GeometryCollection/GeometryCollectionObject.cpp
@@ -145,8 +145,9 @@ void FGeometryCollectionSizeSpecificData::PostSerialize(const FArchive& Ar)
 	// make sure to load back the deprecated values in the new structure if necessary
 	if (Ar.IsLoading() && CustomVer < FUE5MainStreamObjectVersion::GeometryCollectionUserDefinedCollisionShapes)
 	{
-		if (ensure(CollisionShapes.Num()))
+		if (CollisionShapes.Num())
 		{
+			// @todo(chaos destruction collisions) : Add support for many
 			CollisionShapes[0].CollisionType = CollisionType_DEPRECATED;
 			CollisionShapes[0].ImplicitType = ImplicitType_DEPRECATED;
 			CollisionShapes[0].CollisionObjectReductionPercentage = CollisionObjectReductionPercentage_DEPRECATED;
@@ -166,20 +167,23 @@ void FGeometryCollectionSizeSpecificData::PostSerialize(const FArchive& Ar)
 void FillSharedSimulationSizeSpecificData(FSharedSimulationSizeSpecificData& ToData, const FGeometryCollectionSizeSpecificData& FromData)
 {
 	ToData.MaxSize = FromData.MaxSize;
-	check(FromData.CollisionShapes.Num());
 
 	ToData.CollisionShapesData.SetNumUninitialized(FromData.CollisionShapes.Num());
-	for (int i=0;i<FromData.CollisionShapes.Num();i++)
+
+	if (FromData.CollisionShapes.Num())
 	{
-		ToData.CollisionShapesData[i].CollisionType = FromData.CollisionShapes[i].CollisionType;
-		ToData.CollisionShapesData[i].ImplicitType = FromData.CollisionShapes[i].ImplicitType;
-		ToData.CollisionShapesData[i].LevelSetData.MinLevelSetResolution = FromData.CollisionShapes[i].LevelSet.MinLevelSetResolution;
-		ToData.CollisionShapesData[i].LevelSetData.MaxLevelSetResolution = FromData.CollisionShapes[i].LevelSet.MaxLevelSetResolution;
-		ToData.CollisionShapesData[i].LevelSetData.MinClusterLevelSetResolution = FromData.CollisionShapes[i].LevelSet.MinClusterLevelSetResolution;
-		ToData.CollisionShapesData[i].LevelSetData.MaxClusterLevelSetResolution = FromData.CollisionShapes[i].LevelSet.MaxClusterLevelSetResolution;
-		ToData.CollisionShapesData[i].CollisionObjectReductionPercentage = FromData.CollisionShapes[i].CollisionObjectReductionPercentage;
-		ToData.CollisionShapesData[i].CollisionParticleData.CollisionParticlesFraction = FromData.CollisionShapes[i].CollisionParticles.CollisionParticlesFraction;
-		ToData.CollisionShapesData[i].CollisionParticleData.MaximumCollisionParticles = FromData.CollisionShapes[i].CollisionParticles.MaximumCollisionParticles;
+		for (int i = 0; i < FromData.CollisionShapes.Num(); i++)
+		{
+			ToData.CollisionShapesData[i].CollisionType = FromData.CollisionShapes[i].CollisionType;
+			ToData.CollisionShapesData[i].ImplicitType = FromData.CollisionShapes[i].ImplicitType;
+			ToData.CollisionShapesData[i].LevelSetData.MinLevelSetResolution = FromData.CollisionShapes[i].LevelSet.MinLevelSetResolution;
+			ToData.CollisionShapesData[i].LevelSetData.MaxLevelSetResolution = FromData.CollisionShapes[i].LevelSet.MaxLevelSetResolution;
+			ToData.CollisionShapesData[i].LevelSetData.MinClusterLevelSetResolution = FromData.CollisionShapes[i].LevelSet.MinClusterLevelSetResolution;
+			ToData.CollisionShapesData[i].LevelSetData.MaxClusterLevelSetResolution = FromData.CollisionShapes[i].LevelSet.MaxClusterLevelSetResolution;
+			ToData.CollisionShapesData[i].CollisionObjectReductionPercentage = FromData.CollisionShapes[i].CollisionObjectReductionPercentage;
+			ToData.CollisionShapesData[i].CollisionParticleData.CollisionParticlesFraction = FromData.CollisionShapes[i].CollisionParticles.CollisionParticlesFraction;
+			ToData.CollisionShapesData[i].CollisionParticleData.MaximumCollisionParticles = FromData.CollisionShapes[i].CollisionParticles.MaximumCollisionParticles;
+		}
 	}
 	ToData.DamageThreshold = FromData.DamageThreshold;
 }
@@ -187,18 +191,20 @@ void FillSharedSimulationSizeSpecificData(FSharedSimulationSizeSpecificData& ToD
 FGeometryCollectionSizeSpecificData UGeometryCollection::GeometryCollectionSizeSpecificDataDefaults() 
 {
 	FGeometryCollectionSizeSpecificData Data;
-	check(Data.CollisionShapes.Num());
 
 	Data.MaxSize = FLT_MAX;
-	Data.CollisionShapes[0].CollisionType = ECollisionTypeEnum::Chaos_Volumetric;
-	Data.CollisionShapes[0].ImplicitType = EImplicitTypeEnum::Chaos_Implicit_Sphere;
-	Data.CollisionShapes[0].LevelSet.MinLevelSetResolution = 5;
-	Data.CollisionShapes[0].LevelSet.MaxLevelSetResolution = 10;
-	Data.CollisionShapes[0].LevelSet.MinClusterLevelSetResolution = 25;
-	Data.CollisionShapes[0].LevelSet.MaxClusterLevelSetResolution = 50;
-	Data.CollisionShapes[0].CollisionObjectReductionPercentage = 1.0;
-	Data.CollisionShapes[0].CollisionParticles.CollisionParticlesFraction = 1.0;
-	Data.CollisionShapes[0].CollisionParticles.MaximumCollisionParticles = 60;
+	if (Data.CollisionShapes.Num())
+	{
+		Data.CollisionShapes[0].CollisionType = ECollisionTypeEnum::Chaos_Volumetric;
+		Data.CollisionShapes[0].ImplicitType = EImplicitTypeEnum::Chaos_Implicit_Sphere;
+		Data.CollisionShapes[0].LevelSet.MinLevelSetResolution = 5;
+		Data.CollisionShapes[0].LevelSet.MaxLevelSetResolution = 10;
+		Data.CollisionShapes[0].LevelSet.MinClusterLevelSetResolution = 25;
+		Data.CollisionShapes[0].LevelSet.MaxClusterLevelSetResolution = 50;
+		Data.CollisionShapes[0].CollisionObjectReductionPercentage = 1.0;
+		Data.CollisionShapes[0].CollisionParticles.CollisionParticlesFraction = 1.0;
+		Data.CollisionShapes[0].CollisionParticles.MaximumCollisionParticles = 60;
+	}
 	Data.DamageThreshold = 250.0f;
 	return Data;
 }
@@ -221,22 +227,25 @@ void UGeometryCollection::ValidateSizeSpecificDataDefaults()
 	if (!SizeSpecificData.Num() || !HasDefault(SizeSpecificData))
 	{
 		FGeometryCollectionSizeSpecificData Data = GeometryCollectionSizeSpecificDataDefaults();
+		if (Data.CollisionShapes.Num())
+		{
 #if WITH_EDITORONLY_DATA
-		Data.CollisionShapes[0].CollisionType = CollisionType_DEPRECATED;
-		Data.CollisionShapes[0].ImplicitType = ImplicitType_DEPRECATED;
-		Data.CollisionShapes[0].LevelSet.MinLevelSetResolution = MinLevelSetResolution_DEPRECATED;
-		Data.CollisionShapes[0].LevelSet.MaxLevelSetResolution = MaxLevelSetResolution_DEPRECATED;
-		Data.CollisionShapes[0].LevelSet.MinClusterLevelSetResolution = MinClusterLevelSetResolution_DEPRECATED;
-		Data.CollisionShapes[0].LevelSet.MaxClusterLevelSetResolution = MaxClusterLevelSetResolution_DEPRECATED;
-		Data.CollisionShapes[0].CollisionObjectReductionPercentage = CollisionObjectReductionPercentage_DEPRECATED;
+			Data.CollisionShapes[0].CollisionType = CollisionType_DEPRECATED;
+			Data.CollisionShapes[0].ImplicitType = ImplicitType_DEPRECATED;
+			Data.CollisionShapes[0].LevelSet.MinLevelSetResolution = MinLevelSetResolution_DEPRECATED;
+			Data.CollisionShapes[0].LevelSet.MaxLevelSetResolution = MaxLevelSetResolution_DEPRECATED;
+			Data.CollisionShapes[0].LevelSet.MinClusterLevelSetResolution = MinClusterLevelSetResolution_DEPRECATED;
+			Data.CollisionShapes[0].LevelSet.MaxClusterLevelSetResolution = MaxClusterLevelSetResolution_DEPRECATED;
+			Data.CollisionShapes[0].CollisionObjectReductionPercentage = CollisionObjectReductionPercentage_DEPRECATED;
 #endif
-		if (Data.CollisionShapes[0].ImplicitType == EImplicitTypeEnum::Chaos_Implicit_LevelSet)
-		{
-			Data.CollisionShapes[0].CollisionType = ECollisionTypeEnum::Chaos_Surface_Volumetric;
+			if (Data.CollisionShapes[0].ImplicitType == EImplicitTypeEnum::Chaos_Implicit_LevelSet)
+			{
+				Data.CollisionShapes[0].CollisionType = ECollisionTypeEnum::Chaos_Surface_Volumetric;
+			}
 		}
 		SizeSpecificData.Add(Data);
 	}
-	check(SizeSpecificData.Num() && SizeSpecificData[0].CollisionShapes.Num());
+	check(SizeSpecificData.Num());
 }
 
 
@@ -253,33 +262,36 @@ float KgM3ToKgCm3(float Density)
 void UGeometryCollection::GetSharedSimulationParams(FSharedSimulationParameters& OutParams) const
 {
 	const FGeometryCollectionSizeSpecificData& SizeSpecificDefault = GetDefaultSizeSpecificData();
-	check(SizeSpecificDefault.CollisionShapes.Num());
 
 	OutParams.bMassAsDensity = bMassAsDensity;
 	OutParams.Mass = bMassAsDensity ? KgM3ToKgCm3(Mass) : Mass;	//todo(ocohen): we still have the solver working in old units. This is mainly to fix ui issues. Long term need to normalize units for best precision
 	OutParams.MinimumMassClamp = MinimumMassClamp;
-	OutParams.MaximumCollisionParticleCount = SizeSpecificDefault.CollisionShapes[0].CollisionParticles.MaximumCollisionParticles;
-
-	ECollisionTypeEnum SelectedCollisionType = SizeSpecificDefault.CollisionShapes[0].CollisionType;
 
-	if (SelectedCollisionType == ECollisionTypeEnum::Chaos_Volumetric && SizeSpecificDefault.CollisionShapes[0].ImplicitType == EImplicitTypeEnum::Chaos_Implicit_LevelSet)
+	FGeometryCollectionSizeSpecificData InfSize;
+	if (SizeSpecificDefault.CollisionShapes.Num())
 	{
-		UE_LOG(LogGeometryCollectionInternal, Verbose, TEXT("LevelSet geometry selected but non-particle collisions selected. Forcing particle-implicit collisions for %s"), *GetPathName());
-		SelectedCollisionType = ECollisionTypeEnum::Chaos_Surface_Volumetric;
-	}
+		InfSize.CollisionShapes.SetNum(1); // @todo(chaos destruction collisions) : Add support for multiple shapes. 
+		OutParams.MaximumCollisionParticleCount = SizeSpecificDefault.CollisionShapes[0].CollisionParticles.MaximumCollisionParticles;
 
-	FGeometryCollectionSizeSpecificData InfSize;
-	InfSize.MaxSize = TNumericLimits<float>::Max();
-	InfSize.CollisionShapes[0].CollisionType = SelectedCollisionType;
-	InfSize.CollisionShapes[0].ImplicitType = SizeSpecificDefault.CollisionShapes[0].ImplicitType;
-	InfSize.CollisionShapes[0].LevelSet.MinLevelSetResolution = SizeSpecificDefault.CollisionShapes[0].LevelSet.MinLevelSetResolution;
-	InfSize.CollisionShapes[0].LevelSet.MaxLevelSetResolution = SizeSpecificDefault.CollisionShapes[0].LevelSet.MaxLevelSetResolution;
-	InfSize.CollisionShapes[0].LevelSet.MinClusterLevelSetResolution = SizeSpecificDefault.CollisionShapes[0].LevelSet.MinClusterLevelSetResolution;
-	InfSize.CollisionShapes[0].LevelSet.MaxClusterLevelSetResolution = SizeSpecificDefault.CollisionShapes[0].LevelSet.MaxClusterLevelSetResolution;
-	InfSize.CollisionShapes[0].CollisionObjectReductionPercentage = SizeSpecificDefault.CollisionShapes[0].CollisionObjectReductionPercentage;
-	InfSize.CollisionShapes[0].CollisionParticles.CollisionParticlesFraction = SizeSpecificDefault.CollisionShapes[0].CollisionParticles.CollisionParticlesFraction;
-	InfSize.CollisionShapes[0].CollisionParticles.MaximumCollisionParticles = SizeSpecificDefault.CollisionShapes[0].CollisionParticles.MaximumCollisionParticles;
+		ECollisionTypeEnum SelectedCollisionType = SizeSpecificDefault.CollisionShapes[0].CollisionType;
+
+		if (SelectedCollisionType == ECollisionTypeEnum::Chaos_Volumetric && SizeSpecificDefault.CollisionShapes[0].ImplicitType == EImplicitTypeEnum::Chaos_Implicit_LevelSet)
+		{
+			UE_LOG(LogGeometryCollectionInternal, Verbose, TEXT("LevelSet geometry selected but non-particle collisions selected. Forcing particle-implicit collisions for %s"), *GetPathName());
+			SelectedCollisionType = ECollisionTypeEnum::Chaos_Surface_Volumetric;
+		}
 
+		InfSize.CollisionShapes[0].CollisionType = SelectedCollisionType;
+		InfSize.CollisionShapes[0].ImplicitType = SizeSpecificDefault.CollisionShapes[0].ImplicitType;
+		InfSize.CollisionShapes[0].LevelSet.MinLevelSetResolution = SizeSpecificDefault.CollisionShapes[0].LevelSet.MinLevelSetResolution;
+		InfSize.CollisionShapes[0].LevelSet.MaxLevelSetResolution = SizeSpecificDefault.CollisionShapes[0].LevelSet.MaxLevelSetResolution;
+		InfSize.CollisionShapes[0].LevelSet.MinClusterLevelSetResolution = SizeSpecificDefault.CollisionShapes[0].LevelSet.MinClusterLevelSetResolution;
+		InfSize.CollisionShapes[0].LevelSet.MaxClusterLevelSetResolution = SizeSpecificDefault.CollisionShapes[0].LevelSet.MaxClusterLevelSetResolution;
+		InfSize.CollisionShapes[0].CollisionObjectReductionPercentage = SizeSpecificDefault.CollisionShapes[0].CollisionObjectReductionPercentage;
+		InfSize.CollisionShapes[0].CollisionParticles.CollisionParticlesFraction = SizeSpecificDefault.CollisionShapes[0].CollisionParticles.CollisionParticlesFraction;
+		InfSize.CollisionShapes[0].CollisionParticles.MaximumCollisionParticles = SizeSpecificDefault.CollisionShapes[0].CollisionParticles.MaximumCollisionParticles;
+	}
+	InfSize.MaxSize = TNumericLimits<float>::Max();
 	OutParams.SizeSpecificData.SetNum(SizeSpecificData.Num() + 1);
 	FillSharedSimulationSizeSpecificData(OutParams.SizeSpecificData[0], InfSize);
 
@@ -1108,6 +1120,7 @@ void UGeometryCollection::PostEditChangeProperty(struct FPropertyChangedEvent& P
 			bRebuildSimulationData = true;
 		}
 
+
 		if (bRebuildSimulationData)
 		{
 			if (!bManualDataCreate)
