Add Chaos specific CSV profiling stats (Constraints, Collision detection )

diff --git a/Engine/Source/Runtime/Experimental/Chaos/Private/Chaos/PBDCollisionConstraints.cpp b/Engine/Source/Runtime/Experimental/Chaos/Private/Chaos/PBDCollisionConstraints.cpp
--- a/Engine/Source/Runtime/Experimental/Chaos/Private/Chaos/PBDCollisionConstraints.cpp
+++ b/Engine/Source/Runtime/Experimental/Chaos/Private/Chaos/PBDCollisionConstraints.cpp
@@ -25,6 +25,7 @@
 #include "ChaosStats.h"
 #include "Containers/Queue.h"
 #include "ProfilingDebugging/ScopedTimers.h"
+#include "ProfilingDebugging/CsvProfiler.h"
 #include "Algo/Sort.h"
 
 #if INTEL_ISPC
@@ -526,6 +527,7 @@ namespace Chaos
 	bool FPBDCollisionConstraints::Apply(const FReal Dt, const int32 Iterations, const int32 NumIterations)
 	{
 		SCOPE_CYCLE_COUNTER(STAT_Collisions_Apply);
+		CSV_SCOPED_TIMING_STAT(Chaos, ApplyCollisionConstraints);
 
 		bool bNeedsAnotherIteration = false;
 		if (MApplyPairIterations > 0)
@@ -620,6 +622,7 @@ namespace Chaos
 	bool FPBDCollisionConstraints::Apply(const FReal Dt, const TArray<FPBDCollisionConstraintHandle*>& InConstraintHandles, const int32 Iterations, const int32 NumIterations)
 	{
 		SCOPE_CYCLE_COUNTER(STAT_Collisions_Apply);
+		CSV_SCOPED_TIMING_STAT(Chaos, ApplyCollisionConstraints);
 
 		TAtomic<bool> bNeedsAnotherIterationAtomic;
 		bNeedsAnotherIterationAtomic.Store(false);
diff --git a/Engine/Source/Runtime/Experimental/Chaos/Private/Chaos/PBDJointConstraints.cpp b/Engine/Source/Runtime/Experimental/Chaos/Private/Chaos/PBDJointConstraints.cpp
--- a/Engine/Source/Runtime/Experimental/Chaos/Private/Chaos/PBDJointConstraints.cpp
+++ b/Engine/Source/Runtime/Experimental/Chaos/Private/Chaos/PBDJointConstraints.cpp
@@ -11,6 +11,7 @@
 #include "Chaos/Utilities.h"
 #include "ChaosLog.h"
 #include "ChaosStats.h"
+#include "ProfilingDebugging/CsvProfiler.h"
 
 #include "HAL/IConsoleManager.h"
 
@@ -870,6 +871,8 @@ namespace Chaos
 
 	bool FPBDJointConstraints::Apply(const FReal Dt, const int32 It, const int32 NumIts)
 	{
+		CSV_SCOPED_TIMING_STAT(Chaos, ApplyJointConstraints);
+
 		if (PreApplyCallback != nullptr)
 		{
 			PreApplyCallback(Dt, Handles);
@@ -956,6 +959,7 @@ namespace Chaos
 	bool FPBDJointConstraints::Apply(const FReal Dt, const TArray<FConstraintContainerHandle*>& InConstraintHandles, const int32 It, const int32 NumIts)
 	{
 		SCOPE_CYCLE_COUNTER(STAT_Joints_Apply);
+		CSV_SCOPED_TIMING_STAT(Chaos, ApplyJointConstraints);
 
 		// @todo(ccaulfield): make sorting optional
 		// @todo(ccaulfield): handles should be sorted by level by the constraint rule/graph
diff --git a/Engine/Source/Runtime/Experimental/Chaos/Private/ChaosLog.cpp b/Engine/Source/Runtime/Experimental/Chaos/Private/ChaosLog.cpp
--- a/Engine/Source/Runtime/Experimental/Chaos/Private/ChaosLog.cpp
+++ b/Engine/Source/Runtime/Experimental/Chaos/Private/ChaosLog.cpp
@@ -4,9 +4,12 @@
 
 #include "ChaosModule.h"
 #include "CoreMinimal.h"
+#include "ProfilingDebugging/CsvProfiler.h"
 
 DEFINE_LOG_CATEGORY(LogChaosGeneral);
 DEFINE_LOG_CATEGORY(LogChaosThread);
 DEFINE_LOG_CATEGORY(LogChaosSimulation);
 DEFINE_LOG_CATEGORY(LogChaosDebug);
 DEFINE_LOG_CATEGORY(LogChaos);
+
+CSV_DEFINE_CATEGORY_MODULE(CHAOS_API, Chaos, true);
\ No newline at end of file
diff --git a/Engine/Source/Runtime/Experimental/Chaos/Public/Chaos/Collision/ParticlePairCollisionDetector.h b/Engine/Source/Runtime/Experimental/Chaos/Public/Chaos/Collision/ParticlePairCollisionDetector.h
--- a/Engine/Source/Runtime/Experimental/Chaos/Public/Chaos/Collision/ParticlePairCollisionDetector.h
+++ b/Engine/Source/Runtime/Experimental/Chaos/Public/Chaos/Collision/ParticlePairCollisionDetector.h
@@ -4,6 +4,7 @@
 #include "Chaos/Collision/CollisionDetector.h"
 #include "Chaos/Collision/NarrowPhase.h"
 #include "Chaos/Collision/ParticlePairBroadPhase.h"
+#include "ProfilingDebugging/CsvProfiler.h"
 
 namespace Chaos
 {
@@ -22,6 +23,7 @@ namespace Chaos
 		{
 			SCOPE_CYCLE_COUNTER(STAT_Collisions_Detect);
 			CHAOS_SCOPED_TIMER(DetectCollisions);
+			CSV_SCOPED_TIMING_STAT(Chaos, DetectCollisions);
 
 			if (!GetCollisionContainer().GetCollisionsEnabled())
 			{
diff --git a/Engine/Source/Runtime/Experimental/Chaos/Public/Chaos/Collision/SpatialAccelerationCollisionDetector.h b/Engine/Source/Runtime/Experimental/Chaos/Public/Chaos/Collision/SpatialAccelerationCollisionDetector.h
--- a/Engine/Source/Runtime/Experimental/Chaos/Public/Chaos/Collision/SpatialAccelerationCollisionDetector.h
+++ b/Engine/Source/Runtime/Experimental/Chaos/Public/Chaos/Collision/SpatialAccelerationCollisionDetector.h
@@ -7,6 +7,7 @@
 #include "Chaos/Collision/NarrowPhase.h"
 #include "Chaos/Collision/SpatialAccelerationBroadPhase.h"
 #include "Chaos/EvolutionResimCache.h"
+#include "ProfilingDebugging/CsvProfiler.h"
 
 namespace Chaos
 {
@@ -25,6 +26,7 @@ namespace Chaos
 		{
 			SCOPE_CYCLE_COUNTER(STAT_Collisions_Detect);
 			CHAOS_SCOPED_TIMER(DetectCollisions);
+			CSV_SCOPED_TIMING_STAT(Chaos, DetectCollisions);
 
 			if (!GetCollisionContainer().GetCollisionsEnabled())
 			{
diff --git a/Engine/Source/Runtime/Experimental/Chaos/Public/Chaos/PBDRigidsEvolution.h b/Engine/Source/Runtime/Experimental/Chaos/Public/Chaos/PBDRigidsEvolution.h
--- a/Engine/Source/Runtime/Experimental/Chaos/Public/Chaos/PBDRigidsEvolution.h
+++ b/Engine/Source/Runtime/Experimental/Chaos/Public/Chaos/PBDRigidsEvolution.h
@@ -17,6 +17,7 @@
 #include "Chaos/PBDRigidsEvolutionFwd.h"
 #include "Chaos/Defines.h"
 #include "Chaos/PendingSpatialData.h"
+#include "ProfilingDebugging/CsvProfiler.h"
 
 
 extern int32 ChaosRigidsEvolutionApplyAllowEarlyOutCVar;
@@ -614,6 +615,8 @@ public:
 
 	void UpdateAccelerationStructures(int32 Island)
 	{
+		CSV_SCOPED_TIMING_STAT(Chaos, UpdateAccelerationStructures);
+
 		for (FPBDConstraintGraphRule* ConstraintRule : ConstraintRules)
 		{
 			ConstraintRule->UpdateAccelerationStructures(Island);
@@ -624,19 +627,23 @@ public:
 	{
 		UpdateAccelerationStructures(Island);
 
-		int32 LocalNumIterations = ChaosNumContactIterationsOverride >= 0 ? ChaosNumContactIterationsOverride : NumIterations;
-		// @todo(ccaulfield): track whether we are sufficiently solved and can early-out
-		for (int i = 0; i < LocalNumIterations; ++i)
 		{
-			bool bNeedsAnotherIteration = false;
-			for (FPBDConstraintGraphRule* ConstraintRule : PrioritizedConstraintRules)
-			{
-				bNeedsAnotherIteration |= ConstraintRule->ApplyConstraints(Dt, Island, i, LocalNumIterations);
-			}
+			CSV_SCOPED_TIMING_STAT(Chaos, ApplyConstraints);
 
-			if (ChaosRigidsEvolutionApplyAllowEarlyOutCVar && !bNeedsAnotherIteration)
+			int32 LocalNumIterations = ChaosNumContactIterationsOverride >= 0 ? ChaosNumContactIterationsOverride : NumIterations;
+			// @todo(ccaulfield): track whether we are sufficiently solved and can early-out
+			for (int i = 0; i < LocalNumIterations; ++i)
 			{
-				break;
+				bool bNeedsAnotherIteration = false;
+				for (FPBDConstraintGraphRule* ConstraintRule : PrioritizedConstraintRules)
+				{
+					bNeedsAnotherIteration |= ConstraintRule->ApplyConstraints(Dt, Island, i, LocalNumIterations);
+				}
+
+				if (ChaosRigidsEvolutionApplyAllowEarlyOutCVar && !bNeedsAnotherIteration)
+				{
+					break;
+				}
 			}
 		}
 	}
diff --git a/Engine/Source/Runtime/Experimental/Chaos/Public/ChaosLog.h b/Engine/Source/Runtime/Experimental/Chaos/Public/ChaosLog.h
--- a/Engine/Source/Runtime/Experimental/Chaos/Public/ChaosLog.h
+++ b/Engine/Source/Runtime/Experimental/Chaos/Public/ChaosLog.h
@@ -3,10 +3,13 @@
 #pragma once
 
 #include "CoreMinimal.h"
+#include "ProfilingDebugging/CsvProfiler.h"
 
 CHAOS_API DECLARE_LOG_CATEGORY_EXTERN(LogChaosGeneral, Log, All);
 CHAOS_API DECLARE_LOG_CATEGORY_EXTERN(LogChaosThread, Log, All);
 CHAOS_API DECLARE_LOG_CATEGORY_EXTERN(LogChaosSimulation, Log, All);
 CHAOS_API DECLARE_LOG_CATEGORY_EXTERN(LogChaosDebug, Log, All);
 
-CHAOS_API DECLARE_LOG_CATEGORY_EXTERN(LogChaos, Log, All);
\ No newline at end of file
+CHAOS_API DECLARE_LOG_CATEGORY_EXTERN(LogChaos, Log, All);
+
+CSV_DECLARE_CATEGORY_MODULE_EXTERN(CHAOS_API, Chaos);
\ No newline at end of file
