Chaos : Geometry Collection Object
- Dont force a rebuild on load for the size specific default changes.

diff --git a/Engine/Source/Runtime/Experimental/GeometryCollectionEngine/Private/GeometryCollection/GeometryCollectionObject.cpp b/Engine/Source/Runtime/Experimental/GeometryCollectionEngine/Private/GeometryCollection/GeometryCollectionObject.cpp
--- a/Engine/Source/Runtime/Experimental/GeometryCollectionEngine/Private/GeometryCollection/GeometryCollectionObject.cpp
+++ b/Engine/Source/Runtime/Experimental/GeometryCollectionEngine/Private/GeometryCollection/GeometryCollectionObject.cpp
@@ -751,12 +751,6 @@ void UGeometryCollection::Serialize(FArchive& Ar)
 	}
 
 #if WITH_EDITORONLY_DATA
-	if (Ar.IsLoading() && Ar.CustomVer(FUE5MainStreamObjectVersion::GUID) < FUE5MainStreamObjectVersion::GeometryCollectionUserDefinedCollisionShapes)
-	{
-		InvalidateCollection();
-		bCreateSimulationData = true;
-	}
-
 	if (bCreateSimulationData)
 	{
 		CreateSimulationData();
