MeshModelingTools: Tweaked tool target requirements object storage default and added a convenience function to add requirements. Also added several comments in places.

#rnx

diff --git a/Engine/Plugins/Experimental/MeshModelingToolset/Source/ModelingComponents/Public/Selection/PolygonSelectionMechanic.h b/Engine/Plugins/Experimental/MeshModelingToolset/Source/ModelingComponents/Public/Selection/PolygonSelectionMechanic.h
--- a/Engine/Plugins/Experimental/MeshModelingToolset/Source/ModelingComponents/Public/Selection/PolygonSelectionMechanic.h
+++ b/Engine/Plugins/Experimental/MeshModelingToolset/Source/ModelingComponents/Public/Selection/PolygonSelectionMechanic.h
@@ -122,7 +122,7 @@ public:
 	 * will need to call UpdateSelection(), UpdateHighlight(), ClearHighlight(), and 
 	 * ClearSelection() from its own hover/click handlers.
 	 *
-	 * Must be called after Setup() to have an effect.
+	 * Must be called during tool Setup() after calling Setup() on the mechanic to have an effect.
 	 *
 	 * @param ParentToolIn The parent tool, needed to be able to remove the behaviors.
 	 */
diff --git a/Engine/Source/Runtime/Experimental/InteractiveToolsFramework/Public/InteractiveTool.h b/Engine/Source/Runtime/Experimental/InteractiveToolsFramework/Public/InteractiveTool.h
--- a/Engine/Source/Runtime/Experimental/InteractiveToolsFramework/Public/InteractiveTool.h
+++ b/Engine/Source/Runtime/Experimental/InteractiveToolsFramework/Public/InteractiveTool.h
@@ -392,7 +392,8 @@ public:
 	//
 
 	/**
-	 * Add an input behavior for this Tool
+	 * Add an input behavior for this Tool. Typically only effective when called during tool Setup(), as
+	 * the behavior set gets submitted to the input router after that call.
 	 * @param Behavior behavior to add
 	 * @param Source Optional pointer that could be used to identify the behavior for removal later.
 	 */
@@ -400,6 +401,8 @@ public:
 
 	/**
 	 * Remove all input behaviors that had the given source pointer set during their addition.
+	 * Typically only effective when called during tool Setup(), as the behavior set gets submitted
+	 * to the input router after that call.
 	 * @param Source Identifying pointer
 	 */
 	virtual void RemoveInputBehaviorsBySource(void* Source);
diff --git a/Engine/Source/Runtime/Experimental/InteractiveToolsFramework/Public/ToolTargetManager.h b/Engine/Source/Runtime/Experimental/InteractiveToolsFramework/Public/ToolTargetManager.h
--- a/Engine/Source/Runtime/Experimental/InteractiveToolsFramework/Public/ToolTargetManager.h
+++ b/Engine/Source/Runtime/Experimental/InteractiveToolsFramework/Public/ToolTargetManager.h
@@ -13,12 +13,24 @@
 // target or another could be more efficient for certain tasks. This is probably only worth
 // thinking about once we encounter it.
 
+// NOTE FOR IMPLEMENTING CACHING: Someday we will probably write a tool target manager that
+// caches tool targets and gives prebuilt ones when possible to avoid any extra conversions.
+// When implementing this, note that currently, some targets store converted versions of their
+// source objects without a way to determine whether their stored conversion is out of date.
+// The cached target will then give an incorrect result if the source object has been
+// modified after caching.
+// One way to deal with this is to have subclasses of UToolTarget that give the degree of
+// cacheability- fully cacheable (i.e., detects changes and rebuilds if necessary), cacheable
+// only if changes are made through that tool target (i.e. target manager has to clear it if
+// object may have changed; may choose not to have this option), or not cacheable. 
+
 /**
  * The tool target manager converts input objects into tool targets- objects that
- * can expose various interfaces that tools might expect but which the original
- * objects may not know about.
+ * can expose various interfaces that tools need to operate on them.
  *
  * Someday, the tool target manager may implement caching of targets.
+ * 
+ * See the class comment for UToolTarget for more info.
  */
 UCLASS(Transient)
 class INTERACTIVETOOLSFRAMEWORK_API UToolTargetManager : public UObject
diff --git a/Engine/Source/Runtime/Experimental/InteractiveToolsFramework/Public/ToolTargets/ToolTarget.h b/Engine/Source/Runtime/Experimental/InteractiveToolsFramework/Public/ToolTargets/ToolTarget.h
--- a/Engine/Source/Runtime/Experimental/InteractiveToolsFramework/Public/ToolTargets/ToolTarget.h
+++ b/Engine/Source/Runtime/Experimental/InteractiveToolsFramework/Public/ToolTargets/ToolTarget.h
@@ -9,10 +9,22 @@
 class FToolTargetTypeRequirements;
 
 /**
- * A tool target is a stand-in object that a tool can operate on. It can implement any
- * interface(s) that a tool requires without having to implement those interfaces in
- * specific actor classes, and it allows the tools to work on anything that can provide
- * a qualifying tool target.
+ * A tool target is a stand-in object that a tool can operate on. It exposes the necessary
+ * interfaces to the tool.
+ *
+ * There are two intended purposes of the tool target system:
+ * 1. Allow tools to operate on arbitrary objects as long as they can be made to provide the 
+ *   tool with the necessary inputs. For instance, a mesh editing tool should be able to operate
+ *   on skeletal, static, volume, and other mesh as long as the target manager has a registered
+ *   factory that can use that type of mesh to create a suitable target.
+ * 2. (not yet used) Help cache tool inputs. I.e., if a tool requires an expensive
+ *   conversion before it can work on an item, the converted result can be stored in the
+ *   tool target which can be cached by the target manager and provided the next time the
+ *   same type of target is requested for that item.
+ * 
+ * Given an object, tool builders usually ask the target manager to turn it into a target that
+ * has the interfaces the tools needs. The tools cast the target to those interfaces to use
+ * them.
  */
 UCLASS(Transient, Abstract)
 class INTERACTIVETOOLSFRAMEWORK_API UToolTarget : public UObject
@@ -27,13 +39,13 @@ public:
 
 /**
  * A structure used to specify the requirements of a tool for its target. E.g., a tool
- * may need a target that has base type x and interfaces w,y,z.
+ * may need a target that has interfaces x,y,z.
  */
 class INTERACTIVETOOLSFRAMEWORK_API FToolTargetTypeRequirements
 {
 public:
 
-	TArray<const UClass*, TInlineAllocator<2>> Interfaces;
+	TArray<const UClass*, TInlineAllocator<5>> Interfaces;
 
 	FToolTargetTypeRequirements()
 	{
@@ -49,6 +61,12 @@ public:
 		Interfaces = InterfacesIn;
 	}
 
+	FToolTargetTypeRequirements& Add(UClass* Interface)
+	{
+		Interfaces.Add(Interface);
+		return *this;
+	}
+
 	bool AreSatisfiedBy(UClass* Class) const;
 
 	bool AreSatisfiedBy(UToolTarget* ToolTarget) const;
@@ -57,7 +75,8 @@ public:
 
 /**
  * Base class for factories of tool targets, which let a tool manager build targets
- * out of inputs without knowing anything itself.
+ * out of inputs without knowing anything about the inputs itself, as long as it
+ * has a factory registered that is able to process the input.
  */
 UCLASS(Transient, Abstract)
 class INTERACTIVETOOLSFRAMEWORK_API UToolTargetFactory : public UObject
diff --git a/Engine/Source/Runtime/GeometryCore/Private/DynamicMesh/DynamicMeshOverlay.cpp b/Engine/Source/Runtime/GeometryCore/Private/DynamicMesh/DynamicMeshOverlay.cpp
--- a/Engine/Source/Runtime/GeometryCore/Private/DynamicMesh/DynamicMeshOverlay.cpp
+++ b/Engine/Source/Runtime/GeometryCore/Private/DynamicMesh/DynamicMeshOverlay.cpp
@@ -431,7 +431,12 @@ void TDynamicMeshOverlay<RealType, ElementSize>::InternalSetTriangle(int tid, co
 
 		for (int VInd = 0; VInd < 3; ++VInd)
 		{
+			// Checks that the parent vertices of the elements that we're referencing in the overlay
+			// triangle are either not yet set or already point to the vertices of the corresponding
+			// mesh triangle (and so will remain unchanged). Remember that the same element is not
+			// allowed to be used for multiple vertices.
 			checkSlow(ParentVertices[tv[VInd]] == ParentTriangle[VInd] || ParentVertices[tv[VInd]] == FDynamicMesh3::InvalidID);
+
 			ParentVertices.InsertAt(ParentTriangle[VInd], tv[VInd]);
 		}
 	}
diff --git a/Engine/Source/Runtime/GeometryCore/Public/DynamicMesh/DynamicMeshOverlay.h b/Engine/Source/Runtime/GeometryCore/Public/DynamicMesh/DynamicMeshOverlay.h
--- a/Engine/Source/Runtime/GeometryCore/Public/DynamicMesh/DynamicMeshOverlay.h
+++ b/Engine/Source/Runtime/GeometryCore/Public/DynamicMesh/DynamicMeshOverlay.h
@@ -21,9 +21,16 @@ class FDynamicMesh3;
 /**
  * TDynamicMeshOverlay is an add-on to a FDynamicMesh3 that allows for per-triangle storage
  * of an "element" (eg like a per-triangle UV or normal). However the elements can be shared
- * between triangles because they are stored in a separate indexable list. As a result,
- * the overlay has it's own topology, IE there may be "seam" boundary edges in the 
- * overlay topology that are not mesh boundary edges in the associated/parent mesh.
+ * between triangles at shared vertices because the elements are stored in a separate indexable
+ * list.
+ *
+ * Each element has one vertex in the parent mesh as its parent, whereas each vertex may be the
+ * parent of multiple elements in cases where neighboring triangles are not sharing a single
+ * element for that vertex. This means that there may be "seam" boundary edges in the
+ * overlay topology that are not mesh boundary edges in the associated/parent mesh, but
+ * the overlay topology will not connect triangles that were not connected in the parent mesh
+ * or create any topologically degenerate triangles, since the parent vids of the elements
+ * of a triangle will have to match up to the vids of the triangle.
  *
  * A "seam" edge is one where at least one of the elements of the triangles on either
  * side of the edge is not shared between the two triangles.
