MeshModelingTools: Refactored modeling tools gizmos to avoid communication between component and proxy and prevent race conditions that resulted in incorrect view information being used in hit tests. Also fixed a minor issue with view plane-parallel circle gizmos, and created option to have scaling handles align better to view.

#rnx

diff --git a/Engine/Plugins/Experimental/GizmoEdMode/Source/LightGizmos/Private/DirectionalLightGizmo.cpp b/Engine/Plugins/Experimental/GizmoEdMode/Source/LightGizmos/Private/DirectionalLightGizmo.cpp
--- a/Engine/Plugins/Experimental/GizmoEdMode/Source/LightGizmos/Private/DirectionalLightGizmo.cpp
+++ b/Engine/Plugins/Experimental/GizmoEdMode/Source/LightGizmos/Private/DirectionalLightGizmo.cpp
@@ -2,7 +2,9 @@
 
 #include "DirectionalLightGizmo.h"
 #include "BaseGizmos/GizmoBoxComponent.h"
+#include "BaseGizmos/GizmoViewContext.h"
 #include "Components/SphereComponent.h"
+#include "ContextObjectStore.h"
 #include "Engine/CollisionProfile.h"
 #include "Kismet/KismetMathLibrary.h"
 #include "BaseGizmos/GizmoMath.h"
@@ -21,6 +23,11 @@ UInteractiveGizmo* UDirectionalLightGizmoBuilder::BuildGizmo(const FToolBuilderS
 {
 	UDirectionalLightGizmo* NewGizmo = NewObject<UDirectionalLightGizmo>(SceneState.GizmoManager);
 	NewGizmo->SetWorld(SceneState.World);
+	
+	UGizmoViewContext* GizmoViewContext = SceneState.ToolManager->GetContextObjectStore()->FindContext<UGizmoViewContext>();
+	check(GizmoViewContext && GizmoViewContext->IsValidLowLevel());
+	NewGizmo->SetGizmoViewContext(GizmoViewContext);
+
 	return NewGizmo;
 }
 
@@ -203,6 +210,11 @@ void UDirectionalLightGizmo::SetWorld(UWorld* InWorld)
 	World = InWorld;
 }
 
+void UDirectionalLightGizmo::SetGizmoViewContext(UGizmoViewContext* GizmoViewContextIn)
+{
+	GizmoViewContext = GizmoViewContextIn;
+}
+
 void UDirectionalLightGizmo::OnBeginDrag(const FInputDeviceRay& Ray)
 {
 	FVector Start = Ray.WorldRay.Origin;
@@ -351,7 +363,8 @@ void UDirectionalLightGizmo::CreateGizmoHandles()
 	FActorSpawnParameters SpawnInfo;
 	GizmoActor = World->SpawnActor<ADirectionalLightGizmoActor>(FVector::ZeroVector, FRotator::ZeroRotator, SpawnInfo);
 
-	GizmoActor->Arrow = AGizmoActor::AddDefaultLineHandleComponent(World, GizmoActor, FLinearColor::Red, FVector::YAxisVector, FVector::XAxisVector, ArrowLength, true);
+	GizmoActor->Arrow = AGizmoActor::AddDefaultLineHandleComponent(World, GizmoActor, GizmoViewContext,
+		FLinearColor::Red, FVector::YAxisVector, FVector::XAxisVector, ArrowLength, true);
 }
 
 void UDirectionalLightGizmo::UpdateGizmoHandles()
diff --git a/Engine/Plugins/Experimental/GizmoEdMode/Source/LightGizmos/Private/DirectionalLightGizmo.h b/Engine/Plugins/Experimental/GizmoEdMode/Source/LightGizmos/Private/DirectionalLightGizmo.h
--- a/Engine/Plugins/Experimental/GizmoEdMode/Source/LightGizmos/Private/DirectionalLightGizmo.h
+++ b/Engine/Plugins/Experimental/GizmoEdMode/Source/LightGizmos/Private/DirectionalLightGizmo.h
@@ -16,6 +16,8 @@
 
 #include "DirectionalLightGizmo.generated.h"
 
+class UGizmoViewContext;
+
 UCLASS()
 class UDirectionalLightGizmoBuilder : public UInteractiveGizmoBuilder
 {
@@ -72,6 +74,7 @@ public:
 
 	virtual void SetSelectedObject(ADirectionalLight* InLight);
 	virtual void SetWorld(UWorld* InWorld);
+	virtual void SetGizmoViewContext(UGizmoViewContext* GizmoViewContextIn);
 
 	virtual void OnBeginDrag(const FInputDeviceRay& Ray);
 	virtual void OnUpdateDrag(const FInputDeviceRay& Ray);
@@ -103,6 +106,10 @@ private:
 	UPROPERTY()
 	ADirectionalLightGizmoActor* GizmoActor;
 
+	/** Used to properly render the handle gizmo. */
+	UPROPERTY()
+	UGizmoViewContext* GizmoViewContext;
+
 	/** The current target of the gizmo */
 	UPROPERTY()
 	ADirectionalLight* LightActor;
diff --git a/Engine/Plugins/Experimental/GizmoEdMode/Source/LightGizmos/Private/LightGizmosModule.cpp b/Engine/Plugins/Experimental/GizmoEdMode/Source/LightGizmos/Private/LightGizmosModule.cpp
--- a/Engine/Plugins/Experimental/GizmoEdMode/Source/LightGizmos/Private/LightGizmosModule.cpp
+++ b/Engine/Plugins/Experimental/GizmoEdMode/Source/LightGizmos/Private/LightGizmosModule.cpp
@@ -1,11 +1,13 @@
 // Copyright Epic Games, Inc. All Rights Reserved.
 
 #include "LightGizmosModule.h"
+#include "ContextObjectStore.h"
 #include "Editor.h"
 #include "EditorModeRegistry.h"
 #include "EditorModeManager.h"
 #include "Misc/CoreDelegates.h"
 #include "GizmoEdMode.h"
+#include "BaseGizmos/GizmoViewContext.h"
 #include "PointLightGizmoFactory.h"
 #include "DirectionalLightGizmoFactory.h"
 #include "SpotLightGizmoFactory.h"
@@ -37,6 +39,15 @@ void FLightGizmosModule::OnLevelEditorCreated(TSharedPtr<ILevelEditor> InLevelEd
 			 // Register the factories and gizmos if we are entering the new Gizmo Mode
 			if (IsEnteringMode && GizmoEdMode)
 			{
+				// UGizmoViewContext is needed for DirectionalLightGizmo and SpotLightGizmo
+				UContextObjectStore* ContextStore = LevelEditorPtr.Pin()->GetEditorModeManager().GetInteractiveToolsContext()->ContextObjectStore;
+				UGizmoViewContext* GizmoViewContext = ContextStore->FindContext<UGizmoViewContext>();
+				if (!GizmoViewContext)
+				{
+					GizmoViewContext = NewObject<UGizmoViewContext>();
+					ContextStore->AddContextObject(GizmoViewContext);
+				}
+
 				LevelEditorPtr.Pin()->GetEditorModeManager().GetInteractiveToolsContext()->GizmoManager->RegisterGizmoType(PointLightGizmoType, NewObject<UPointLightGizmoBuilder>());
 				LevelEditorPtr.Pin()->GetEditorModeManager().GetInteractiveToolsContext()->GizmoManager->RegisterGizmoType(ScalableConeGizmoType, NewObject<UScalableConeGizmoBuilder>());
 				LevelEditorPtr.Pin()->GetEditorModeManager().GetInteractiveToolsContext()->GizmoManager->RegisterGizmoType(SpotLightGizmoType, NewObject<USpotLightGizmoBuilder>());
diff --git a/Engine/Plugins/Experimental/GizmoEdMode/Source/LightGizmos/Private/SpotLightGizmo.cpp b/Engine/Plugins/Experimental/GizmoEdMode/Source/LightGizmos/Private/SpotLightGizmo.cpp
--- a/Engine/Plugins/Experimental/GizmoEdMode/Source/LightGizmos/Private/SpotLightGizmo.cpp
+++ b/Engine/Plugins/Experimental/GizmoEdMode/Source/LightGizmos/Private/SpotLightGizmo.cpp
@@ -2,11 +2,13 @@
 
 #include "SpotLightGizmo.h"
 #include "Components/SpotLightComponent.h"
+#include "ContextObjectStore.h"
 #include "ScalableConeGizmo.h"
 #include "LightGizmosModule.h"
 #include "Engine/CollisionProfile.h"
 #include "Kismet/KismetMathLibrary.h"
 #include "BaseGizmos/GizmoMath.h"
+#include "BaseGizmos/GizmoViewContext.h"
 #include "Components/SphereComponent.h"
 #include "BaseGizmos/GizmoLineHandleComponent.h"
 #include "BaseBehaviors/MouseHoverBehavior.h"
@@ -19,6 +21,11 @@ UInteractiveGizmo* USpotLightGizmoBuilder::BuildGizmo(const FToolBuilderState& S
 {
 	USpotLightGizmo* NewGizmo = NewObject<USpotLightGizmo>(SceneState.GizmoManager);
 	NewGizmo->SetWorld(SceneState.World);
+
+	UGizmoViewContext* GizmoViewContext = SceneState.ToolManager->GetContextObjectStore()->FindContext<UGizmoViewContext>();
+	check(GizmoViewContext && GizmoViewContext->IsValidLowLevel());
+	NewGizmo->SetGizmoViewContext(GizmoViewContext);
+
 	return NewGizmo;
 }
 
@@ -273,6 +280,11 @@ void USpotLightGizmo::SetWorld(UWorld* InWorld)
 	World = InWorld;
 }
 
+void USpotLightGizmo::SetGizmoViewContext(UGizmoViewContext* GizmoViewContextIn)
+{
+	GizmoViewContext = GizmoViewContextIn;
+}
+
 void USpotLightGizmo::OnBeginDrag(const FInputDeviceRay& Ray)
 {
 	FVector Start = Ray.WorldRay.Origin;
@@ -385,7 +397,8 @@ void USpotLightGizmo::CreateAttenuationScaleGizmo()
 	GizmoActor = World->SpawnActor<ASpotLightGizmoActor>(FVector::ZeroVector, FRotator::ZeroRotator, SpawnInfo);
 
 	// The handle to scale attenuation is line handle component
-	GizmoActor->AttenuationScaleHandle = AGizmoActor::AddDefaultLineHandleComponent(World, GizmoActor, FLinearColor::Blue, FVector::YAxisVector, FVector::XAxisVector, 60.f, true);
+	GizmoActor->AttenuationScaleHandle = AGizmoActor::AddDefaultLineHandleComponent(World, GizmoActor, GizmoViewContext,
+		FLinearColor::Blue, FVector::YAxisVector, FVector::XAxisVector, 60.f, true);
 	GizmoActor->AttenuationScaleHandle->SetRelativeLocation(FVector(LightActor->SpotLightComponent->AttenuationRadius, 0, 0));
 
 	TransformProxy->OnTransformChanged.AddUObject(this, &USpotLightGizmo::OnTransformChanged);
diff --git a/Engine/Plugins/Experimental/GizmoEdMode/Source/LightGizmos/Private/SpotLightGizmo.h b/Engine/Plugins/Experimental/GizmoEdMode/Source/LightGizmos/Private/SpotLightGizmo.h
--- a/Engine/Plugins/Experimental/GizmoEdMode/Source/LightGizmos/Private/SpotLightGizmo.h
+++ b/Engine/Plugins/Experimental/GizmoEdMode/Source/LightGizmos/Private/SpotLightGizmo.h
@@ -11,6 +11,8 @@
 
 #include "SpotLightGizmo.generated.h"
 
+class UGizmoViewContext;
+
 UCLASS()
 class USpotLightGizmoBuilder : public UInteractiveGizmoBuilder
 {
@@ -73,6 +75,7 @@ public:
 	void CreateAttenuationScaleGizmo();
 
 	void SetWorld(UWorld* InWorld);
+	void SetGizmoViewContext(UGizmoViewContext* GizmoViewContextIn);
 
 	USubTransformProxy* GetTransformProxy();
 
@@ -112,6 +115,10 @@ private:
 	UPROPERTY()
 	ASpotLightGizmoActor* GizmoActor;
 
+	/** Used to properly render the handle gizmo. */
+	UPROPERTY()
+	UGizmoViewContext* GizmoViewContext;
+
 	UPROPERTY()
 	FVector DragStartWorldPosition;
 
diff --git a/Engine/Source/Editor/Experimental/EditorInteractiveToolsFramework/Private/EdModeInteractiveToolsContext.cpp b/Engine/Source/Editor/Experimental/EditorInteractiveToolsFramework/Private/EdModeInteractiveToolsContext.cpp
--- a/Engine/Source/Editor/Experimental/EditorInteractiveToolsFramework/Private/EdModeInteractiveToolsContext.cpp
+++ b/Engine/Source/Editor/Experimental/EditorInteractiveToolsFramework/Private/EdModeInteractiveToolsContext.cpp
@@ -3,6 +3,8 @@
 
 #include "EdModeInteractiveToolsContext.h"
 
+#include "BaseGizmos/GizmoViewContext.h"
+#include "ContextObjectStore.h"
 #include "Editor.h"
 #include "EditorViewportClient.h"
 #include "EditorModeManager.h"
@@ -570,6 +572,9 @@ void UEdModeInteractiveToolsContext::Initialize(IToolsContextQueriesAPI* Queries
 	});
 
 	InvalidationTimestamp = 0;
+
+	// This gets set up in UInteractiveToolsContext::Initialize;
+	GizmoViewContext = ToolManager->GetContextObjectStore()->FindContext<UGizmoViewContext>();
 }
 
 void UEdModeInteractiveToolsContext::Shutdown()
@@ -818,16 +823,11 @@ void UEdModeInteractiveToolsContext::Render(const FSceneView* View, FViewport* V
 	// FEditorViewportClient, which passes it's own Viewport down. So, this cast should be valid (for now)
 	FEditorViewportClient* ViewportClient = static_cast<FEditorViewportClient*>(Viewport->GetClient());
 
-	// Update the global currently-focused FSceneView variable, which GizmoArrowComponent and friends will
-	// use to know when they are seeing the SceneView they should use to recalculate their size/visibility/etc.
-	// This could go away if we could move that functionality out of the RenderProxy (tricky given that it needs
-	// to respond to each FSceneView...)
-	if (ViewportClient == EditorModeManager->GetHoveredViewportClient())
+	// Update the currently-hovered scene view information, which GizmoArrowComponent and friends will
+	// use to recalculate their size/visibility/etc.
+	if (GizmoViewContext && ViewportClient == EditorModeManager->GetHoveredViewportClient())
 	{
-		// This locks internally and so no need to do on Render thread, and possibly better to do immediately (?)
-		//ENQUEUE_RENDER_COMMAND(BlerBlerBler)( [View](FRHICommandListImmediate& RHICmdList) {
-			GizmoRenderingUtil::SetGlobalFocusedEditorSceneView(View);
-		//});
+		GizmoViewContext->ResetFromSceneView(*View);
 	}
 
 	// Render Tool and Gizmos
diff --git a/Engine/Source/Editor/Experimental/EditorInteractiveToolsFramework/Public/EdModeInteractiveToolsContext.h b/Engine/Source/Editor/Experimental/EditorInteractiveToolsFramework/Public/EdModeInteractiveToolsContext.h
--- a/Engine/Source/Editor/Experimental/EditorInteractiveToolsFramework/Public/EdModeInteractiveToolsContext.h
+++ b/Engine/Source/Editor/Experimental/EditorInteractiveToolsFramework/Public/EdModeInteractiveToolsContext.h
@@ -13,6 +13,7 @@
 class FEdMode;
 class FEditorModeTools;
 class FEditorViewportClient;
+class UGizmoViewContext;
 class FSceneView;
 class FViewport;
 class UMaterialInterface;
@@ -126,6 +127,11 @@ protected:
 	/** Input event instance used to keep track of various button states, etc, that we cannot directly query on-demand */
 	FInputDeviceState CurrentMouseState;
 
+	// An object in which we save the current scene view information that gizmos can use on the game thread
+	// to figure out how big the gizmo is for hit testing. Lives in the context store, but we keep a pointer here
+	// to avoid having to look for it.
+	UGizmoViewContext* GizmoViewContext = nullptr;
+
 	// Utility function to convert viewport x/y from mouse events (and others?) into scene ray.
 	// Copy-pasted from other Editor code, seems kind of expensive?
 	static FRay GetRayFromMousePos(FEditorViewportClient* ViewportClient, FViewport* Viewport, int MouseX, int MouseY);
diff --git a/Engine/Source/Editor/UnrealEd/Private/EditorModeManager.cpp b/Engine/Source/Editor/UnrealEd/Private/EditorModeManager.cpp
--- a/Engine/Source/Editor/UnrealEd/Private/EditorModeManager.cpp
+++ b/Engine/Source/Editor/UnrealEd/Private/EditorModeManager.cpp
@@ -1533,6 +1533,7 @@ bool FEditorModeTools::MouseEnter( FEditorViewportClient* InViewportClient, FVie
 
 bool FEditorModeTools::MouseLeave( FEditorViewportClient* InViewportClient, FViewport* Viewport )
 {
+	HoveredViewportClient = nullptr;
 	bool bHandled = InteractiveToolsContext->MouseLeave(InViewportClient, Viewport);
 
 	ForEachEdMode<ILegacyEdModeViewportInterface>([&bHandled, InViewportClient, Viewport](ILegacyEdModeViewportInterface* Mode)
@@ -1573,6 +1574,12 @@ bool FEditorModeTools::ReceivedFocus( FEditorViewportClient* InViewportClient, F
 
 bool FEditorModeTools::LostFocus( FEditorViewportClient* InViewportClient, FViewport* Viewport )
 {
+	// Note that we don't reset FocusedViewportClient intentionally. EdModeInteractiveToolsContext
+	// only ticks its objects once for the focused viewport to avoid multi-ticking, so if we cleared
+	// it here, we'd stop ticking things in the level editor when clicking out of the viewport.
+	// TODO: Conceptually, we should probably clear FocusedViewportClient here, but also have a
+	// LastFocusedViewportClient property that we don't clear, to use in ticking.
+
 	bool bHandled = false;
 	ForEachEdMode<ILegacyEdModeViewportInterface>([&bHandled, InViewportClient, Viewport](ILegacyEdModeViewportInterface* Mode)
 		{
diff --git a/Engine/Source/Runtime/Experimental/InteractiveToolsFramework/Private/BaseGizmos/GizmoActor.cpp b/Engine/Source/Runtime/Experimental/InteractiveToolsFramework/Private/BaseGizmos/GizmoActor.cpp
--- a/Engine/Source/Runtime/Experimental/InteractiveToolsFramework/Private/BaseGizmos/GizmoActor.cpp
+++ b/Engine/Source/Runtime/Experimental/InteractiveToolsFramework/Private/BaseGizmos/GizmoActor.cpp
@@ -28,12 +28,15 @@ AGizmoActor::AGizmoActor()
 
 
 UGizmoArrowComponent* AGizmoActor::AddDefaultArrowComponent(
-	UWorld* World, AActor* Actor,
+	UWorld* World, AActor* Actor, UGizmoViewContext* GizmoViewContext,
 	const FLinearColor& Color, const FVector& LocalDirection, const float Length)
 {
+	ensure(GizmoViewContext);
+
 	UGizmoArrowComponent* NewArrow = NewObject<UGizmoArrowComponent>(Actor);
 	Actor->AddInstanceComponent(NewArrow);
 	NewArrow->AttachToComponent(Actor->GetRootComponent(), FAttachmentTransformRules::KeepRelativeTransform);
+	NewArrow->SetGizmoViewContext(GizmoViewContext);
 	NewArrow->Direction = LocalDirection;
 	NewArrow->Color = Color;
 	NewArrow->Length = Length;
@@ -44,12 +47,15 @@ UGizmoArrowComponent* AGizmoActor::AddDefaultArrowComponent(
 
 
 UGizmoRectangleComponent* AGizmoActor::AddDefaultRectangleComponent(
-	UWorld* World, AActor* Actor,
+	UWorld* World, AActor* Actor, UGizmoViewContext* GizmoViewContext,
 	const FLinearColor& Color, const FVector& PlaneAxis1, const FVector& PlaneAxisx2)
 {
+	ensure(GizmoViewContext);
+
 	UGizmoRectangleComponent* NewRectangle = NewObject<UGizmoRectangleComponent>(Actor);
 	Actor->AddInstanceComponent(NewRectangle);
 	NewRectangle->AttachToComponent(Actor->GetRootComponent(), FAttachmentTransformRules::KeepRelativeTransform);
+	NewRectangle->SetGizmoViewContext(GizmoViewContext);
 	NewRectangle->DirectionX = PlaneAxis1;
 	NewRectangle->DirectionY = PlaneAxisx2;
 	NewRectangle->Color = Color;
@@ -61,12 +67,15 @@ UGizmoRectangleComponent* AGizmoActor::AddDefaultRectangleComponent(
 
 
 UGizmoCircleComponent* AGizmoActor::AddDefaultCircleComponent(
-	UWorld* World, AActor* Actor,
+	UWorld* World, AActor* Actor, UGizmoViewContext* GizmoViewContext,
 	const FLinearColor& Color, const FVector& PlaneNormal, float Radius)
 {
+	ensure(GizmoViewContext);
+
 	UGizmoCircleComponent* NewCircle = NewObject<UGizmoCircleComponent>(Actor);
 	Actor->AddInstanceComponent(NewCircle);
 	NewCircle->AttachToComponent(Actor->GetRootComponent(), FAttachmentTransformRules::KeepRelativeTransform);
+	NewCircle->SetGizmoViewContext(GizmoViewContext);
 	NewCircle->Normal = PlaneNormal;
 	NewCircle->Color = Color;
 	NewCircle->Radius = Radius;
@@ -77,14 +86,17 @@ UGizmoCircleComponent* AGizmoActor::AddDefaultCircleComponent(
 
 
 UGizmoBoxComponent* AGizmoActor::AddDefaultBoxComponent(
-	UWorld* World, AActor* Actor,
+	UWorld* World, AActor* Actor, UGizmoViewContext* GizmoViewContext,
 	const FLinearColor& Color, const FVector& Origin,
 	const FVector& Dimensions
 )
 {
+	ensure(GizmoViewContext);
+
 	UGizmoBoxComponent* NewBox = NewObject<UGizmoBoxComponent>(Actor);
 	Actor->AddInstanceComponent(NewBox);
 	NewBox->AttachToComponent(Actor->GetRootComponent(), FAttachmentTransformRules::KeepRelativeTransform);
+	NewBox->SetGizmoViewContext(GizmoViewContext);
 	NewBox->Origin = Origin;
 	NewBox->Color = Color;
 	NewBox->Dimensions = Dimensions;
@@ -94,13 +106,16 @@ UGizmoBoxComponent* AGizmoActor::AddDefaultBoxComponent(
 
 
 UGizmoLineHandleComponent* AGizmoActor::AddDefaultLineHandleComponent(
-	UWorld* World, AActor* Actor,
+	UWorld* World, AActor* Actor, UGizmoViewContext* GizmoViewContext,
 	const FLinearColor& Color, const FVector& HandleNormal, const FVector& LocalDirection,
     const float Length, const bool bImageScale)
 {
+	ensure(GizmoViewContext);
+
 	UGizmoLineHandleComponent* LineHandle = NewObject<UGizmoLineHandleComponent>(Actor);
 	Actor->AddInstanceComponent(LineHandle);
 	LineHandle->AttachToComponent(Actor->GetRootComponent(), FAttachmentTransformRules::KeepRelativeTransform);
+	LineHandle->SetGizmoViewContext(GizmoViewContext);
 	LineHandle->Normal = HandleNormal;
 	LineHandle->Direction = LocalDirection;
 	LineHandle->Length = Length;
diff --git a/Engine/Source/Runtime/Experimental/InteractiveToolsFramework/Private/BaseGizmos/GizmoArrowComponent.cpp b/Engine/Source/Runtime/Experimental/InteractiveToolsFramework/Private/BaseGizmos/GizmoArrowComponent.cpp
--- a/Engine/Source/Runtime/Experimental/InteractiveToolsFramework/Private/BaseGizmos/GizmoArrowComponent.cpp
+++ b/Engine/Source/Runtime/Experimental/InteractiveToolsFramework/Private/BaseGizmos/GizmoArrowComponent.cpp
@@ -2,11 +2,59 @@
  
 #include "BaseGizmos/GizmoArrowComponent.h"
 #include "BaseGizmos/GizmoRenderingUtil.h"
+#include "BaseGizmos/GizmoViewContext.h"
 #include "PrimitiveSceneProxy.h"
 #include "Materials/Material.h"
 #include "MaterialShared.h"
 #include "SceneManagement.h"
 
+namespace GizmoArrowComponentLocals
+{
+	const float ARROW_RENDERVISIBILITY_DOT_THRESHOLD = 0.965f; //~15 degrees
+
+	template <typename SceneViewOrGizmoViewContext>
+	bool GetWorldEndpoints(bool bIsViewDependent, const SceneViewOrGizmoViewContext* View,
+		const FVector& WorldOrigin, const FVector& Direction, float Gap, 
+		float Length, bool bUseWorldAxes,
+		TFunctionRef<FVector(const FVector&)> VectorTransform, 
+		FVector& StartPointOut, FVector& EndPointOut, float& PixelToWorldScaleOut)
+	{
+		FVector ArrowDirection = bUseWorldAxes ? Direction : VectorTransform(Direction);
+		float LengthScale = 1;
+
+		if (bIsViewDependent)
+		{
+			// direction to origin of gizmo
+			FVector ViewDirection = View->IsPerspectiveProjection() ?
+				WorldOrigin - View->ViewLocation : View->GetViewDirection();
+			ViewDirection.Normalize();
+
+			bool bFlipped = (FVector::DotProduct(ViewDirection, ArrowDirection) > 0);
+			ArrowDirection = (bFlipped) ? -ArrowDirection : ArrowDirection;
+
+			//bRenderVisibilityOut = FMath::Abs(FVector::DotProduct(ArrowDirection, ViewDirection)) < 0.985f;   // ~10 degrees
+			bool bRenderVisibility =
+				FMath::Abs(FVector::DotProduct(ArrowDirection, ViewDirection)) < ARROW_RENDERVISIBILITY_DOT_THRESHOLD;   // ~15 degrees
+
+			if (!bRenderVisibility)
+			{
+				return false;
+			}
+
+			LengthScale = GizmoRenderingUtil::CalculateLocalPixelToWorldScale(View, WorldOrigin);
+		}
+
+		PixelToWorldScaleOut = LengthScale;
+
+		double StartDist = LengthScale * Gap;
+		double EndDist = LengthScale * (Gap + Length);
+
+		StartPointOut = WorldOrigin + StartDist * ArrowDirection;
+		EndPointOut = WorldOrigin + EndDist * ArrowDirection;
+
+		return true;
+	}
+}
 
 class FGizmoArrowComponentSceneProxy final : public FPrimitiveSceneProxy
 {
@@ -30,8 +78,7 @@ public:
 
 	virtual void GetDynamicMeshElements(const TArray<const FSceneView*>& Views, const FSceneViewFamily& ViewFamily, uint32 VisibilityMap, FMeshElementCollector& Collector) const override
 	{
-		// try to find focused scene view. May return nullptr.
-		const FSceneView* FocusedView = GizmoRenderingUtil::FindFocusedEditorSceneView(Views, ViewFamily, VisibilityMap);
+		using namespace GizmoArrowComponentLocals;
 
 		const FMatrix& LocalToWorldMatrix = GetLocalToWorld();
 		FVector Origin = LocalToWorldMatrix.TransformPosition(FVector::ZeroVector);
@@ -42,55 +89,30 @@ public:
 			{
 				const FSceneView* View = Views[ViewIndex];
 				FPrimitiveDrawInterface* PDI = Collector.GetPDI(ViewIndex);
-				bool bIsFocusedView = (FocusedView != nullptr && View == FocusedView);
-				bool bIsOrtho = !View->IsPerspectiveProjection();
-
-				// direction to origin of gizmo
-				FVector ViewDirection = 
-					(bIsOrtho) ?  (View->GetViewDirection()) : (Origin - View->ViewLocation);
-				ViewDirection.Normalize();
 
 				bool bWorldAxis = (bExternalWorldLocalState) ? (*bExternalWorldLocalState) : false;
-				FVector ArrowDirection = (bWorldAxis) ? Direction : FVector{ LocalToWorldMatrix.TransformVector(Direction) };
-				bool bFlipped = (FVector::DotProduct(ViewDirection, ArrowDirection) > 0);
-				ArrowDirection = (bFlipped) ? -ArrowDirection : ArrowDirection;
-				if (bIsFocusedView && bFlippedExternal != nullptr)
-				{
-					*bFlippedExternal = bFlipped;
-				}
+				FVector StartPoint, EndPoint;
+				float PixelToWorldScale = 0;
 
-				//bRenderVisibility = FMath::Abs(FVector::DotProduct(ArrowDirection, ViewDirection)) < 0.985f;   // ~10 degrees
-				bool bRenderVisibility = FMath::Abs(FVector::DotProduct(ArrowDirection, ViewDirection)) < 0.965f;   // ~15 degrees
+				bool bIsViewDependent = (bExternalIsViewDependent) ? (*bExternalIsViewDependent) : false;
+				bool bRenderVisibility = GetWorldEndpoints(bIsViewDependent, View, Origin, 
+					Direction, Gap, Length, bWorldAxis,
+					[&LocalToWorldMatrix](const FVector& VectorIn) {
+						return FVector{ LocalToWorldMatrix.TransformVector(VectorIn) };
+					}, StartPoint, EndPoint, PixelToWorldScale);
 
-				if (bIsFocusedView && bExternalRenderVisibility != nullptr)
-				{
-					*bExternalRenderVisibility = bRenderVisibility;
-				}
-				if (bRenderVisibility == false)
+				if (!bRenderVisibility)
 				{
 					continue;
 				}
 
-				float PixelToWorldScale = GizmoRenderingUtil::CalculateLocalPixelToWorldScale(View, Origin);
-				float LengthScale = PixelToWorldScale;
-				if (bIsFocusedView && ExternalDynamicPixelToWorldScale != nullptr)
-				{
-					*ExternalDynamicPixelToWorldScale = PixelToWorldScale;
-				}
-
 				float UseThickness = (bExternalHoverState != nullptr && *bExternalHoverState == true) ?
 					(HoverThicknessMultiplier * Thickness) : (Thickness);
-				if (!bIsOrtho)
+				if (View->IsPerspectiveProjection())
 				{
 					UseThickness *= (View->FOV / 90.0);		// compensate for FOV scaling in Gizmos...
 				}
 
-				double StartDist = LengthScale * Gap;
-				double EndDist = LengthScale * (Gap + Length);
-
-				FVector StartPoint = Origin + StartDist*ArrowDirection;
-				FVector EndPoint = Origin + EndDist*ArrowDirection;
-
 				PDI->DrawLine(StartPoint, EndPoint, Color, SDPG_Foreground, UseThickness, 0.0f, true);
 			}
 		}
@@ -115,22 +137,6 @@ public:
 	virtual uint32 GetMemoryFootprint(void) const override { return sizeof *this + GetAllocatedSize(); }
 	uint32 GetAllocatedSize(void) const { return FPrimitiveSceneProxy::GetAllocatedSize(); }
 
-
-	void SetExternalFlip(bool* bFlipped)
-	{
-		bFlippedExternal = bFlipped;
-	}
-
-	void SetExternalDynamicPixelToWorldScale(float* DynamicPixelToWorldScale)
-	{
-		ExternalDynamicPixelToWorldScale = DynamicPixelToWorldScale;
-	}
-
-	void SetExternalRenderVisibility(bool* bRenderVisibility)
-	{
-		bExternalRenderVisibility = bRenderVisibility;
-	}
-
 	void SetExternalHoverState(bool* HoverState)
 	{
 		bExternalHoverState = HoverState;
@@ -141,6 +147,11 @@ public:
 		bExternalWorldLocalState = bWorldLocalState;
 	}
 
+	void SetExternalIsViewDependent(bool* bExternalIsViewDependentIn)
+	{
+		bExternalIsViewDependent = bExternalIsViewDependentIn;
+	}
+
 private:
 	FLinearColor Color;
 	FVector Direction;
@@ -152,22 +163,16 @@ private:
 	// set on Component for use in ::GetDynamicMeshElements()
 	bool* bExternalHoverState = nullptr;
 	bool* bExternalWorldLocalState = nullptr;
-
-	// set in ::GetDynamicMeshElements() for use by Component hit testing
-	bool* bFlippedExternal = nullptr;
-	float* ExternalDynamicPixelToWorldScale = nullptr;
-	bool* bExternalRenderVisibility = nullptr;
+	bool* bExternalIsViewDependent = nullptr;
 };
 
 
 FPrimitiveSceneProxy* UGizmoArrowComponent::CreateSceneProxy()
 {
 	FGizmoArrowComponentSceneProxy* NewProxy = new FGizmoArrowComponentSceneProxy(this);
-	NewProxy->SetExternalFlip(&bFlipped);
-	NewProxy->SetExternalDynamicPixelToWorldScale(&DynamicPixelToWorldScale);
-	NewProxy->SetExternalRenderVisibility(&bRenderVisibility);
 	NewProxy->SetExternalHoverState(&bHovering);
 	NewProxy->SetExternalWorldLocalState(&bWorld);
+	NewProxy->SetExternalIsViewDependent(&bIsViewDependent);
 	return NewProxy;
 }
 
@@ -178,27 +183,27 @@ FBoxSphereBounds UGizmoArrowComponent::CalcBounds(const FTransform& LocalToWorld
 
 bool UGizmoArrowComponent::LineTraceComponent(FHitResult& OutHit, const FVector Start, const FVector End, const FCollisionQueryParams& Params)
 {
-	if (bRenderVisibility == false)
-	{
-		return false;
-	}
+	using namespace GizmoArrowComponentLocals;
 
 	const FTransform& Transform = this->GetComponentToWorld();
+	FVector UseOrigin = Transform.TransformPosition(FVector::ZeroVector);
 
-	FVector UseDirection = (bFlipped) ? -Direction : Direction;
-	float LengthScale = DynamicPixelToWorldScale;
-	double StartDist = LengthScale * Gap;
-	double EndDist = LengthScale * (Gap + Length);
+	// Copy what is done in the proxy object, but get data from the gizmo view context.
+	FVector StartPoint, EndPoint;
+	float PixelToWorldScale = 0;
+	bool bRenderVisibility = GetWorldEndpoints(bIsViewDependent, GizmoViewContext, UseOrigin, Direction, Gap, Length, bWorld,
+		[&Transform](const FVector& VectorIn) { return Transform.TransformVector(VectorIn); },
+		StartPoint, EndPoint, PixelToWorldScale);
 
-	UseDirection = (bWorld) ? UseDirection : Transform.TransformVector(UseDirection);
-	FVector UseOrigin = Transform.TransformPosition(FVector::ZeroVector);
-	FVector Point0 = UseOrigin + StartDist * UseDirection;
-	FVector Point1 = UseOrigin + EndDist * UseDirection;
+	if (!bRenderVisibility)
+	{
+		return false;
+	}
 
 	FVector NearestArrow, NearestLine;
-	FMath::SegmentDistToSegmentSafe(Point0, Point1, Start, End, NearestArrow, NearestLine);
+	FMath::SegmentDistToSegmentSafe(StartPoint, EndPoint, Start, End, NearestArrow, NearestLine);
 	double Distance = FVector::Distance(NearestArrow, NearestLine);
-	if (Distance > PixelHitDistanceThreshold*DynamicPixelToWorldScale)
+	if (Distance > PixelHitDistanceThreshold * PixelToWorldScale)
 	{
 		return false;
 	}
diff --git a/Engine/Source/Runtime/Experimental/InteractiveToolsFramework/Private/BaseGizmos/GizmoBoxComponent.cpp b/Engine/Source/Runtime/Experimental/InteractiveToolsFramework/Private/BaseGizmos/GizmoBoxComponent.cpp
--- a/Engine/Source/Runtime/Experimental/InteractiveToolsFramework/Private/BaseGizmos/GizmoBoxComponent.cpp
+++ b/Engine/Source/Runtime/Experimental/InteractiveToolsFramework/Private/BaseGizmos/GizmoBoxComponent.cpp
@@ -2,8 +2,49 @@
 
 #include "BaseGizmos/GizmoBoxComponent.h"
 #include "BaseGizmos/GizmoRenderingUtil.h"
+#include "BaseGizmos/GizmoViewContext.h"
 #include "PrimitiveSceneProxy.h"
 
+namespace GizmoBoxComponentLocals
+{
+
+	template <typename SceneViewOrGizmoViewContext>
+	void GetBoxDirections(bool bIsViewDependent, const SceneViewOrGizmoViewContext* View, const FVector& WorldOrigin,
+		const FVector& DirectionX, const FVector& DirectionY, const FVector& DirectionZ, 
+		bool bEnableFlipping, bool bWorldAxis,
+		TFunctionRef<FVector(const FVector&)> VectorTransform,
+		FVector& UseDirectionXOut, FVector& UseDirectionYOut, FVector& UseDirectionZOut,
+		bool& bFlippedXOut, bool& bFlippedYOut, bool& bFlippedZOut
+		)
+	{
+
+		UseDirectionXOut = (bWorldAxis) ? DirectionX : VectorTransform(DirectionX);
+		UseDirectionYOut = (bWorldAxis) ? DirectionY : VectorTransform(DirectionY);
+		UseDirectionZOut = (bWorldAxis) ? DirectionZ : VectorTransform(DirectionZ);
+		bFlippedXOut = false;
+		bFlippedYOut = false;
+		bFlippedZOut = false;
+
+		if (bIsViewDependent)
+		{
+			// direction to origin of gizmo
+			FVector ViewDirection =
+				View->IsPerspectiveProjection() ? WorldOrigin - View->ViewLocation : View->GetViewDirection();
+			ViewDirection.Normalize();
+
+			bFlippedXOut = (FVector::DotProduct(ViewDirection, UseDirectionXOut) > 0);
+			UseDirectionXOut = (bEnableFlipping && bFlippedXOut) ? -UseDirectionXOut : UseDirectionXOut;
+
+			bFlippedYOut = (FVector::DotProduct(ViewDirection, UseDirectionYOut) > 0);
+			UseDirectionYOut = (bEnableFlipping && bFlippedYOut) ? -UseDirectionYOut : UseDirectionYOut;
+
+			bFlippedZOut = (FVector::DotProduct(ViewDirection, UseDirectionZOut) > 0);
+			UseDirectionZOut = (bEnableFlipping && bFlippedZOut) ? -UseDirectionZOut : UseDirectionZOut;
+		}
+	}
+
+}
+
 
 class FGizmoBoxComponentSceneProxy final : public FPrimitiveSceneProxy
 {
@@ -31,8 +72,7 @@ public:
 
 	virtual void GetDynamicMeshElements(const TArray<const FSceneView*>& Views, const FSceneViewFamily& ViewFamily, uint32 VisibilityMap, FMeshElementCollector& Collector) const override
 	{
-		// try to find focused scene view. May return nullptr.
-		const FSceneView* FocusedView = GizmoRenderingUtil::FindFocusedEditorSceneView(Views, ViewFamily, VisibilityMap);
+		using namespace GizmoBoxComponentLocals;
 
 		const FMatrix& LocalToWorldMatrix = GetLocalToWorld();
 		FVector WorldOrigin = LocalToWorldMatrix.TransformPosition(FVector::ZeroVector);
@@ -50,52 +90,33 @@ public:
 			{
 				const FSceneView* View = Views[ViewIndex];
 				FPrimitiveDrawInterface* PDI = Collector.GetPDI(ViewIndex);
-				bool bIsFocusedView = (FocusedView != nullptr && View == FocusedView);
 				bool bIsOrtho = !View->IsPerspectiveProjection();
 
-				float PixelToWorldScale = GizmoRenderingUtil::CalculateLocalPixelToWorldScale(View, WorldOrigin);
-				float LengthScale = PixelToWorldScale;
-				if (bIsFocusedView && ExternalDynamicPixelToWorldScale != nullptr)
-				{
-					*ExternalDynamicPixelToWorldScale = PixelToWorldScale;
-				}
-
-				// direction to origin of gizmo
-				FVector ViewDirection = 
-					(bIsOrtho) ? (View->GetViewDirection()) : (WorldOrigin - View->ViewLocation);
-				ViewDirection.Normalize();
-
 				bool bWorldAxis = (bExternalWorldLocalState) ? (*bExternalWorldLocalState) : false;
 
-				FVector UseDirectionX = (bWorldAxis) ? DirectionX : FVector{ LocalToWorldMatrix.TransformVector(DirectionX) };
-				bool bFlippedX = (FVector::DotProduct(ViewDirection, UseDirectionX) > 0);
-				UseDirectionX = (bEnableFlipping && bFlippedX) ? -UseDirectionX : UseDirectionX;
-				if (bIsFocusedView && bFlippedXExternal != nullptr)
-				{
-					*bFlippedXExternal = bFlippedX;
-				}
+				FVector UseDirectionX, UseDirectionY, UseDirectionZ;
+				bool bFlippedX = false;
+				bool bFlippedY = false;
+				bool bFlippedZ = false;
 
-				FVector UseDirectionY = (bWorldAxis) ? DirectionY : FVector{ LocalToWorldMatrix.TransformVector(DirectionY) };
-				bool bFlippedY = (FVector::DotProduct(ViewDirection, UseDirectionY) > 0);
-				UseDirectionY = (bEnableFlipping && bFlippedY) ? -UseDirectionY : UseDirectionY;
-				if (bIsFocusedView && bFlippedYExternal != nullptr)
-				{
-					*bFlippedYExternal = bFlippedY;
-				}
-
-				FVector UseDirectionZ = (bWorldAxis) ? DirectionZ : FVector{ LocalToWorldMatrix.TransformVector(DirectionZ) };
-				bool bFlippedZ = (FVector::DotProduct(ViewDirection, UseDirectionZ) > 0);
-				UseDirectionZ = (bEnableFlipping && bFlippedZ) ? -UseDirectionZ : UseDirectionZ;
-				if (bIsFocusedView && bFlippedZExternal != nullptr)
-				{
-					*bFlippedZExternal = bFlippedZ;
-				}
+				bool bIsViewDependent = (bExternalIsViewDependent) ? (*bExternalIsViewDependent) : false;
+				GetBoxDirections(bIsViewDependent, View, WorldOrigin, DirectionX, DirectionY, DirectionZ, bEnableFlipping,
+					bWorldAxis, [&LocalToWorldMatrix](const FVector& VectorIn) {
+						return FVector{ LocalToWorldMatrix.TransformVector(VectorIn) };
+					}, UseDirectionX, UseDirectionY, UseDirectionZ, bFlippedX, bFlippedY, bFlippedZ);
 
 				FVector UseCenter(
 					(bFlippedX) ? -LocalCenter.X : LocalCenter.X,
 					(bFlippedY) ? -LocalCenter.Y : LocalCenter.Y,
 					(bFlippedZ) ? -LocalCenter.Z : LocalCenter.Z
 				);
+
+				float LengthScale = 1;
+				if (bIsViewDependent)
+				{
+					LengthScale = GizmoRenderingUtil::CalculateLocalPixelToWorldScale(View, WorldOrigin);
+				}
+
 				FVector WorldCenter = WorldOrigin
 					+ LengthScale * LocalCenter.X * UseDirectionX
 					+ LengthScale * LocalCenter.Y * UseDirectionY
@@ -125,6 +146,9 @@ public:
 
 				if (bRemoveHiddenLines)
 				{
+					FVector ViewDirection =
+						View->IsPerspectiveProjection() ? WorldOrigin - View->ViewLocation : View->GetViewDirection();
+
 					// find box corner direction that is most aligned with view direction. That's the corner we will hide.
 					float MaxDot = -999999;
 					int MaxDotIndex = -1;
@@ -181,18 +205,6 @@ public:
 	virtual uint32 GetMemoryFootprint(void) const override { return sizeof *this + GetAllocatedSize(); }
 	uint32 GetAllocatedSize(void) const { return FPrimitiveSceneProxy::GetAllocatedSize(); }
 
-	void SetExternalFlip(bool* bFlippedX, bool* bFlippedY, bool* bFlippedZ)
-	{
-		bFlippedXExternal = bFlippedX;
-		bFlippedYExternal = bFlippedY;
-		bFlippedZExternal = bFlippedZ;
-	}
-
-	void SetExternalDynamicPixelToWorldScale(float* DynamicPixelToWorldScale)
-	{
-		ExternalDynamicPixelToWorldScale = DynamicPixelToWorldScale;
-	}
-
 	void SetExternalHoverState(bool* HoverState)
 	{
 		bExternalHoverState = HoverState;
@@ -203,6 +215,10 @@ public:
 		bExternalWorldLocalState = bWorldLocalState;
 	}
 
+	void SetExternalIsViewDependent(bool* bExternalIsViewDependentIn)
+	{
+		bExternalIsViewDependent = bExternalIsViewDependentIn;
+	}
 
 private:
 	FLinearColor Color;
@@ -217,12 +233,7 @@ private:
 	// set on Component for use in ::GetDynamicMeshElements()
 	bool* bExternalHoverState = nullptr;
 	bool* bExternalWorldLocalState = nullptr;
-
-	// set in ::GetDynamicMeshElements() for use by Component hit testing
-	bool* bFlippedXExternal = nullptr;
-	bool* bFlippedYExternal = nullptr;
-	bool* bFlippedZExternal = nullptr;
-	float* ExternalDynamicPixelToWorldScale = nullptr;
+	bool* bExternalIsViewDependent = nullptr;
 };
 
 
@@ -230,13 +241,9 @@ private:
 FPrimitiveSceneProxy* UGizmoBoxComponent::CreateSceneProxy()
 {
 	FGizmoBoxComponentSceneProxy* NewProxy = new FGizmoBoxComponentSceneProxy(this);
-	if (bEnableAxisFlip)
-	{
-		NewProxy->SetExternalFlip(&bFlippedX, &bFlippedY, &bFlippedZ);
-	}
-	NewProxy->SetExternalDynamicPixelToWorldScale(&DynamicPixelToWorldScale);
 	NewProxy->SetExternalHoverState(&bHovering);
 	NewProxy->SetExternalWorldLocalState(&bWorld);
+	NewProxy->SetExternalIsViewDependent(&bIsViewDependent);
 	return NewProxy;
 }
 
@@ -247,10 +254,7 @@ FBoxSphereBounds UGizmoBoxComponent::CalcBounds(const FTransform& LocalToWorld)
 
 bool UGizmoBoxComponent::LineTraceComponent(FHitResult& OutHit, const FVector Start, const FVector End, const FCollisionQueryParams& Params)
 {
-	if (bRenderVisibility == false)
-	{
-		return false;
-	}
+	using namespace GizmoBoxComponentLocals;
 
 	const FTransform& Transform = this->GetComponentToWorld();
 
@@ -263,11 +267,37 @@ bool UGizmoBoxComponent::LineTraceComponent(FHitResult& OutHit, const FVector St
 	StartLocal = InvRotation * StartLocal;
 	EndLocal = InvRotation * EndLocal;
 
+	bool bFlippedX = false; 
+	bool bFlippedY = false;
+	bool bFlippedZ = false;
+
+	if (bIsViewDependent)
+	{
+		FVector UseDirectionX, UseDirectionY, UseDirectionZ; // not used
+		GetBoxDirections(bIsViewDependent, GizmoViewContext,
+			Transform.TransformPosition(FVector::ZeroVector), // Origin
+			Rotation * FVector(1, 0, 0), // DirectionX
+			Rotation * FVector(0, 1, 0), // DirectionY
+			Rotation * FVector(0, 0, 1), // DirectionZ
+			bEnableAxisFlip, bWorld,
+			[&Transform](const FVector& VectorIn) { return Transform.TransformVector(VectorIn); },
+			UseDirectionX, UseDirectionY, UseDirectionZ, bFlippedX, bFlippedY, bFlippedZ);
+	}
+
 	FVector UseOrigin(
 		(bEnableAxisFlip && bFlippedX) ? -Origin.X : Origin.X,
 		(bEnableAxisFlip && bFlippedY) ? -Origin.Y : Origin.Y,
 		(bEnableAxisFlip && bFlippedZ) ? -Origin.Z : Origin.Z
 	);
+
+	float DynamicPixelToWorldScale = 1;
+
+	if (bIsViewDependent)
+	{
+		DynamicPixelToWorldScale = GizmoRenderingUtil::CalculateLocalPixelToWorldScale(
+			GizmoViewContext, UseOrigin);
+	}
+
 	UseOrigin *= DynamicPixelToWorldScale;
 
 	FVector ScaledDims = DynamicPixelToWorldScale * Dimensions;
diff --git a/Engine/Source/Runtime/Experimental/InteractiveToolsFramework/Private/BaseGizmos/GizmoCircleComponent.cpp b/Engine/Source/Runtime/Experimental/InteractiveToolsFramework/Private/BaseGizmos/GizmoCircleComponent.cpp
--- a/Engine/Source/Runtime/Experimental/InteractiveToolsFramework/Private/BaseGizmos/GizmoCircleComponent.cpp
+++ b/Engine/Source/Runtime/Experimental/InteractiveToolsFramework/Private/BaseGizmos/GizmoCircleComponent.cpp
@@ -1,10 +1,28 @@
 // Copyright Epic Games, Inc. All Rights Reserved.
- 
+
 #include "BaseGizmos/GizmoCircleComponent.h"
 #include "BaseGizmos/GizmoRenderingUtil.h"
 #include "BaseGizmos/GizmoMath.h"
+#include "BaseGizmos/GizmoViewContext.h"
 #include "PrimitiveSceneProxy.h"
 
+namespace GizmoCircleComponentLocals
+{
+	const float RENDER_VISIBILITY_DOT_THRESHOLD = 0.05;
+	const float VIEW_PLANE_PARALLEL_DOT_THRESHOLD = 0.95;
+
+	template <typename SceneViewOrGizmoViewContext>
+	void GetVisibility(const SceneViewOrGizmoViewContext* View, const FVector& ViewDirection,
+		const FVector& GizmoPlaneWorldNormal, const FVector& WorldOrigin,
+		bool& bRenderVisibilityOut, bool& bIsViewPlaneParallelOut)
+	{
+		FVector ViewPlaneDirection = View->GetViewDirection();
+
+		bRenderVisibilityOut = FMath::Abs(ViewDirection.Dot(GizmoPlaneWorldNormal)) >= RENDER_VISIBILITY_DOT_THRESHOLD;
+		bIsViewPlaneParallelOut = FMath::Abs(ViewPlaneDirection.Dot(GizmoPlaneWorldNormal)) >= VIEW_PLANE_PARALLEL_DOT_THRESHOLD;
+	}
+
+}
 
 class FGizmoCircleComponentSceneProxy final : public FPrimitiveSceneProxy
 {
@@ -30,8 +48,7 @@ public:
 
 	virtual void GetDynamicMeshElements(const TArray<const FSceneView*>& Views, const FSceneViewFamily& ViewFamily, uint32 VisibilityMap, FMeshElementCollector& Collector) const override
 	{
-		// try to find focused scene view. May return nullptr.
-		const FSceneView* FocusedView = GizmoRenderingUtil::FindFocusedEditorSceneView(Views, ViewFamily, VisibilityMap);
+		using namespace GizmoCircleComponentLocals;
 
 		const FMatrix& LocalToWorldMatrix = GetLocalToWorld();
 		FVector Origin = LocalToWorldMatrix.TransformPosition(FVector::ZeroVector);
@@ -44,21 +61,14 @@ public:
 			{
 				const FSceneView* View = Views[ViewIndex];
 				FPrimitiveDrawInterface* PDI = Collector.GetPDI(ViewIndex);
-				bool bIsFocusedView = (FocusedView != nullptr && View == FocusedView);
 				bool bIsOrtho = !View->IsPerspectiveProjection();
-				FVector UpVector = View->GetViewUp();
 				FVector ViewVector = View->GetViewDirection();
 
-				// direction to origin of gizmo
-				FVector GizmoViewDirection = 
-					(bIsOrtho) ? (View->GetViewDirection()) : (Origin - View->ViewLocation);
-				GizmoViewDirection.Normalize();
-
-				float PixelToWorldScale = GizmoRenderingUtil::CalculateLocalPixelToWorldScale(View, Origin);
-				float LengthScale = PixelToWorldScale;
-				if (bIsFocusedView && ExternalDynamicPixelToWorldScale != nullptr)
+				float LengthScale = 1;
+				bool bIsViewDependent = (bExternalIsViewDependent) ? (*bExternalIsViewDependent) : false;
+				if (bIsViewDependent)
 				{
-					*ExternalDynamicPixelToWorldScale = PixelToWorldScale;
+					LengthScale = GizmoRenderingUtil::CalculateLocalPixelToWorldScale(View, Origin);
 				}
 
 				double UseRadius = LengthScale * Radius;
@@ -83,15 +93,6 @@ public:
 
 				if (bViewAligned)
 				{
-					//if (bIsOrtho)		// skip in ortho views?
-					//{
-					//	if (bIsFocusedView && bExternalRenderVisibility != nullptr)
-					//	{
-					//		*bExternalRenderVisibility = false;
-					//	}
-					//	continue;
-					//}
-
 					FVector WorldOrigin = LocalToWorldMatrix.TransformPosition(FVector::ZeroVector);
 					WorldOrigin += 0.001 * ViewVector;
 					FVector WorldPlaneX, WorldPlaneY;
@@ -116,19 +117,19 @@ public:
 					FVector WorldPlaneY = (bWorldAxis) ? PlaneY : FVector{ LocalToWorldMatrix.TransformVector(PlaneY) };
 
 					FVector PlaneWorldNormal = (bWorldAxis) ? Normal : FVector{ LocalToWorldMatrix.TransformVector(Normal) };
-					double ViewDot = FVector::DotProduct(GizmoViewDirection, PlaneWorldNormal);
-					bool bOnEdge = FMath::Abs(ViewDot) < 0.05;
-					bool bIsViewPlaneParallel = FMath::Abs(ViewDot) > 0.95;
-					if (bIsFocusedView && bExternalIsViewPlaneParallel != nullptr)
-					{
-						*bExternalIsViewPlaneParallel = bIsViewPlaneParallel;
-					}
 
-					bool bRenderVisibility = !bOnEdge;
-					if (bIsFocusedView && bExternalRenderVisibility != nullptr)
+					// direction to origin of gizmo
+					FVector GizmoViewDirection =
+						(bIsOrtho) ? (View->GetViewDirection()) : (Origin - View->ViewLocation);
+					GizmoViewDirection.Normalize();
+
+					bool bRenderVisibility = true;
+					bool bIsViewPlaneParallel = false;
+					if (bIsViewDependent)
 					{
-						*bExternalRenderVisibility = bRenderVisibility;
+						GetVisibility(View, GizmoViewDirection, PlaneWorldNormal, WorldOrigin, bRenderVisibility, bIsViewPlaneParallel);
 					}
+
 					if (bRenderVisibility)
 					{
 						if (bIsViewPlaneParallel)
@@ -192,22 +193,6 @@ public:
 	virtual uint32 GetMemoryFootprint(void) const override { return sizeof *this + GetAllocatedSize(); }
 	uint32 GetAllocatedSize(void) const { return FPrimitiveSceneProxy::GetAllocatedSize(); }
 
-
-	void SetExternalDynamicPixelToWorldScale(float* DynamicPixelToWorldScale)
-	{
-		ExternalDynamicPixelToWorldScale = DynamicPixelToWorldScale;
-	}
-
-	void SetExternalRenderVisibility(bool* bRenderVisibility)
-	{
-		bExternalRenderVisibility = bRenderVisibility;
-	}
-
-	void SetExternalIsViewPlaneParallel(bool* bIsViewPlaneParallel)
-	{
-		bExternalIsViewPlaneParallel = bIsViewPlaneParallel;
-	}
-
 	void SetExternalHoverState(bool* HoverState)
 	{
 		bExternalHoverState = HoverState;
@@ -218,6 +203,10 @@ public:
 		bExternalWorldLocalState = bWorldLocalState;
 	}
 
+	void SetExternalIsViewDependent(bool* bExternalIsViewDependentIn)
+	{
+		bExternalIsViewDependent = bExternalIsViewDependentIn;
+	}
 
 private:
 	FLinearColor Color;
@@ -232,11 +221,7 @@ private:
 	// set on Component for use in ::GetDynamicMeshElements()
 	bool* bExternalHoverState = nullptr;
 	bool* bExternalWorldLocalState = nullptr;
-
-	// set in ::GetDynamicMeshElements() for use by Component hit testing
-	float* ExternalDynamicPixelToWorldScale = nullptr;
-	bool* bExternalRenderVisibility = nullptr;
-	bool* bExternalIsViewPlaneParallel = nullptr;
+	bool* bExternalIsViewDependent = nullptr;
 };
 
 
@@ -245,11 +230,9 @@ private:
 FPrimitiveSceneProxy* UGizmoCircleComponent::CreateSceneProxy()
 {
 	FGizmoCircleComponentSceneProxy* NewProxy = new FGizmoCircleComponentSceneProxy(this);
-	NewProxy->SetExternalDynamicPixelToWorldScale(&DynamicPixelToWorldScale);
-	NewProxy->SetExternalIsViewPlaneParallel(&bCircleIsViewPlaneParallel);
-	NewProxy->SetExternalRenderVisibility(&bRenderVisibility);
 	NewProxy->SetExternalHoverState(&bHovering);
 	NewProxy->SetExternalWorldLocalState(&bWorld);
+	NewProxy->SetExternalIsViewDependent(&bIsViewDependent);
 	return NewProxy;
 }
 
@@ -260,17 +243,34 @@ FBoxSphereBounds UGizmoCircleComponent::CalcBounds(const FTransform& LocalToWorl
 
 bool UGizmoCircleComponent::LineTraceComponent(FHitResult& OutHit, const FVector Start, const FVector End, const FCollisionQueryParams& Params)
 {
-	if (bRenderVisibility == false)
+	using namespace GizmoCircleComponentLocals;
+
+	const FTransform& Transform = this->GetComponentToWorld();
+	FVector WorldOrigin = Transform.TransformPosition(FVector::ZeroVector);
+	FVector WorldNormal = (bWorld) ? Normal : Transform.TransformVector(Normal);
+
+	bool bRenderVisibility = true;
+	bool bCircleIsViewPlaneParallel = false;
+	float PixelToWorldScale = 1;
+	if (bIsViewDependent)
 	{
-		return false;
+		FVector ViewDirection = GizmoViewContext->IsPerspectiveProjection() ? (WorldOrigin - GizmoViewContext->ViewLocation)
+			: GizmoViewContext->GetViewDirection();
+		ViewDirection.Normalize();
+
+		GetVisibility(GizmoViewContext, ViewDirection, WorldNormal, WorldOrigin, bRenderVisibility, bCircleIsViewPlaneParallel);
+
+		if (!bRenderVisibility)
+		{
+			return false;
+		}
+
+		PixelToWorldScale = GizmoRenderingUtil::CalculateLocalPixelToWorldScale(GizmoViewContext, WorldOrigin);
 	}
 
-	float LengthScale = DynamicPixelToWorldScale;
+	float LengthScale = PixelToWorldScale;
 	double UseRadius = LengthScale * Radius;
 
-	const FTransform& Transform = this->GetComponentToWorld();
-	FVector WorldNormal = (bWorld) ? Normal : Transform.TransformVector(Normal);
-	FVector WorldOrigin = Transform.TransformPosition(FVector::ZeroVector);
 
 	FRay Ray(Start, End - Start, false);
 
@@ -290,7 +290,7 @@ bool UGizmoCircleComponent::LineTraceComponent(FHitResult& OutHit, const FVector
 	FVector NearestRay = Ray.ClosestPoint(NearestCircle);
 
 	double Distance = FVector::Distance(NearestCircle, NearestRay);
-	if (Distance > PixelHitDistanceThreshold*DynamicPixelToWorldScale)
+	if (Distance > PixelHitDistanceThreshold * PixelToWorldScale)
 	{
 		return false;
 	}
diff --git a/Engine/Source/Runtime/Experimental/InteractiveToolsFramework/Private/BaseGizmos/GizmoLineHandleComponent.cpp b/Engine/Source/Runtime/Experimental/InteractiveToolsFramework/Private/BaseGizmos/GizmoLineHandleComponent.cpp
--- a/Engine/Source/Runtime/Experimental/InteractiveToolsFramework/Private/BaseGizmos/GizmoLineHandleComponent.cpp
+++ b/Engine/Source/Runtime/Experimental/InteractiveToolsFramework/Private/BaseGizmos/GizmoLineHandleComponent.cpp
@@ -32,9 +32,6 @@ public:
 
 	virtual void GetDynamicMeshElements(const TArray<const FSceneView*>& Views, const FSceneViewFamily& ViewFamily, uint32 VisibilityMap, FMeshElementCollector& Collector) const override
 	{
-		// try to find focused scene view. May return nullptr.
-		const FSceneView* FocusedView = GizmoRenderingUtil::FindFocusedEditorSceneView(Views, ViewFamily, VisibilityMap);
-
 		FVector LocalOffset = Direction;
 		if (ExternalDistance != nullptr)
 		{
@@ -55,16 +52,15 @@ public:
 			{
 				const FSceneView* View = Views[ViewIndex];
 				FPrimitiveDrawInterface* PDI = Collector.GetPDI(ViewIndex);
-				bool bIsFocusedView = (FocusedView != nullptr && View == FocusedView);
 				bool bIsOrtho = !View->IsPerspectiveProjection();
 				FVector UpVector = View->GetViewUp();
 				FVector ViewVector = View->GetViewDirection();
 
-				float PixelToWorldScale = GizmoRenderingUtil::CalculateLocalPixelToWorldScale(View, WorldDiskOrigin);
-				float LengthScale = PixelToWorldScale;
-				if (bIsFocusedView && ExternalDynamicPixelToWorldScale != nullptr)
+				float LengthScale = 1;
+				bool bIsViewDependent = (bExternalIsViewDependent) ? (*bExternalIsViewDependent) : false;
+				if (bIsViewDependent)
 				{
-					*ExternalDynamicPixelToWorldScale = PixelToWorldScale;
+					LengthScale = GizmoRenderingUtil::CalculateLocalPixelToWorldScale(View, WorldDiskOrigin);
 				}
 
 				FVector ScaledIntevalStart = -LengthScale * (WorldIntervalEnd - WorldDiskOrigin) + WorldDiskOrigin;
@@ -106,11 +102,6 @@ public:
 	uint32 GetAllocatedSize(void) const { return FPrimitiveSceneProxy::GetAllocatedSize(); }
 
 
-	void SetExternalDynamicPixelToWorldScale(float* DynamicPixelToWorldScale)
-	{
-		ExternalDynamicPixelToWorldScale = DynamicPixelToWorldScale;
-	}
-
 	void SetExternalHoverState(bool* HoverState)
 	{
 		bExternalHoverState = HoverState;
@@ -121,6 +112,11 @@ public:
 		bExternalWorldLocalState = bWorldLocalState;
 	}
 
+	void SetExternalIsViewDependent(bool* bExternalIsViewDependentIn)
+	{
+		bExternalIsViewDependent = bExternalIsViewDependentIn;
+	}
+
 	void SetLengthScale(float* Distance)
 	{
 		ExternalDistance = Distance;
@@ -140,19 +136,17 @@ private:
 	bool* bExternalHoverState = nullptr;
 	bool* bExternalWorldLocalState = nullptr;
 	float* ExternalDistance = nullptr;
-
-	// set in ::GetDynamicMeshElements() for use by Component hit testing
-	float* ExternalDynamicPixelToWorldScale = nullptr;
+	bool* bExternalIsViewDependent = nullptr;
 };
 
 
 FPrimitiveSceneProxy* UGizmoLineHandleComponent::CreateSceneProxy()
 {
 	FGizmoLineHandleComponentSceneProxy* NewProxy = new FGizmoLineHandleComponentSceneProxy(this);
-	NewProxy->SetExternalDynamicPixelToWorldScale(&DynamicPixelToWorldScale);
 	NewProxy->SetExternalHoverState(&bHovering);
 	NewProxy->SetExternalWorldLocalState(&bWorld);
 	NewProxy->SetLengthScale(&Length);
+	NewProxy->SetExternalIsViewDependent(&bIsViewDependent);
 	return NewProxy;
 }
 
@@ -168,13 +162,19 @@ FBoxSphereBounds UGizmoLineHandleComponent::CalcBounds(const FTransform& LocalTo
 
 bool UGizmoLineHandleComponent::LineTraceComponent(FHitResult& OutHit, const FVector Start, const FVector End, const FCollisionQueryParams& Params)
 {
-	float LengthScale = (bImageScale) ? DynamicPixelToWorldScale : 1.f;
+	const FTransform& Transform = this->GetComponentToWorld();
+	FVector WorldBaseOrigin = Transform.TransformPosition(FVector::ZeroVector);
+	float PixelToWorldScale = 1;
+	if (bIsViewDependent)
+	{
+		PixelToWorldScale = GizmoRenderingUtil::CalculateLocalPixelToWorldScale(GizmoViewContext, WorldBaseOrigin);
+	}
+
+	float LengthScale = (bImageScale) ? PixelToWorldScale : 1.f;
 	double UseHandleSize = LengthScale * HandleSize;
 	FVector LocalOffset = Length * Direction;
 
-	const FTransform& Transform = this->GetComponentToWorld();
 	FVector HandleDir = (bWorld) ? Normal : Transform.TransformVector(Normal);
-	FVector WorldBaseOrigin = Transform.TransformPosition(FVector::ZeroVector);
 	FVector WorldHandleOrigin = Transform.TransformPosition(LocalOffset);
 
 	FVector BaseToHandle = WorldHandleOrigin - WorldBaseOrigin;
@@ -189,7 +189,7 @@ bool UGizmoLineHandleComponent::LineTraceComponent(FHitResult& OutHit, const FVe
 	FVector NearestOnHandle, NearestOnLine;
 	FMath::SegmentDistToSegmentSafe(HandleStart, HandleEnd, Start, End, NearestOnHandle, NearestOnLine);
 	double Distance = FVector::Distance(NearestOnHandle, NearestOnLine);
-	if (Distance > PixelHitDistanceThreshold*DynamicPixelToWorldScale)
+	if (Distance > PixelHitDistanceThreshold * PixelToWorldScale)
 	{
 		return false;
 	}
diff --git a/Engine/Source/Runtime/Experimental/InteractiveToolsFramework/Private/BaseGizmos/GizmoRectangleComponent.cpp b/Engine/Source/Runtime/Experimental/InteractiveToolsFramework/Private/BaseGizmos/GizmoRectangleComponent.cpp
--- a/Engine/Source/Runtime/Experimental/InteractiveToolsFramework/Private/BaseGizmos/GizmoRectangleComponent.cpp
+++ b/Engine/Source/Runtime/Experimental/InteractiveToolsFramework/Private/BaseGizmos/GizmoRectangleComponent.cpp
@@ -2,8 +2,76 @@
 
 #include "BaseGizmos/GizmoRectangleComponent.h"
 #include "BaseGizmos/GizmoRenderingUtil.h"
+#include "BaseGizmos/GizmoViewContext.h"
 #include "PrimitiveSceneProxy.h"
 
+namespace GizmoRectangleComponentLocals
+{
+	const float RECTANGLE_RENDERVISIBILITY_DOT_THRESHOLD = 0.25;
+
+	template <typename SceneViewOrGizmoViewContext>
+	bool GetWorldCorners(bool bIsViewDependent, const SceneViewOrGizmoViewContext* View, 
+		const FVector& WorldOrigin, const FVector& DirectionX, const FVector& DirectionY, 
+		float OffsetX, float OffsetY, float LengthX, float LengthY,
+		bool bUseWorldAxes, bool bOrientYAccordingToCamera, 
+		TFunctionRef<FVector(const FVector&)> VectorTransform,
+		TArray<FVector>& CornersOut)
+	{
+		FVector UseDirectionX = (bUseWorldAxes) ? DirectionX : VectorTransform(DirectionX);
+		FVector UseDirectionY = (bUseWorldAxes) ? DirectionY : VectorTransform(DirectionY);
+		float LengthScale = 1;
+
+		if (bIsViewDependent)
+		{
+			bool bIsOrtho = !View->IsPerspectiveProjection();
+
+			// direction to origin of gizmo
+			FVector ViewDirection =
+				View->IsPerspectiveProjection() ? WorldOrigin - View->ViewLocation : View->GetViewDirection();
+			ViewDirection.Normalize();
+
+			bool bFlippedX = (FVector::DotProduct(ViewDirection, UseDirectionX) > 0);
+			UseDirectionX = (bFlippedX) ? -UseDirectionX : UseDirectionX;
+
+			if (bOrientYAccordingToCamera)
+			{
+				// See if by rotating the y axis around the x axis 90 degrees, we end up with y that is less
+				// colinear to our view ray.
+				FVector RotatedY = UseDirectionY.RotateAngleAxis(90, UseDirectionX);
+				if (FMath::Abs(RotatedY.Dot(ViewDirection)) < FMath::Abs(UseDirectionY.Dot(ViewDirection)))
+				{
+					UseDirectionY = RotatedY;
+				}
+			}
+			bool bFlippedY = (FVector::DotProduct(ViewDirection, UseDirectionY) > 0);
+			UseDirectionY = (bFlippedY) ? -UseDirectionY : UseDirectionY;
+
+			FVector PlaneNormal = FVector::CrossProduct(UseDirectionX, UseDirectionY);
+			bool bRenderVisibility =
+				FMath::Abs(FVector::DotProduct(PlaneNormal, ViewDirection)) > RECTANGLE_RENDERVISIBILITY_DOT_THRESHOLD;
+
+			if (!bRenderVisibility)
+			{
+				return false;
+			}
+
+			LengthScale = GizmoRenderingUtil::CalculateLocalPixelToWorldScale(View, WorldOrigin);
+		}
+
+		double UseOffsetX = LengthScale * OffsetX;
+		double UseOffsetLengthX = LengthScale * (OffsetX + LengthX);
+		double UseOffsetY = LengthScale * OffsetY;
+		double UseOffsetLengthY = LengthScale * (OffsetY + LengthY);
+
+		CornersOut.SetNum(4);
+		CornersOut[0] = WorldOrigin + UseOffsetX * UseDirectionX + UseOffsetY * UseDirectionY;
+		CornersOut[1] = WorldOrigin + UseOffsetLengthX * UseDirectionX + UseOffsetY * UseDirectionY;
+		CornersOut[2] = WorldOrigin + UseOffsetLengthX * UseDirectionX + UseOffsetLengthY * UseDirectionY;
+		CornersOut[3] = WorldOrigin + UseOffsetX * UseDirectionX + UseOffsetLengthY * UseDirectionY;
+
+		return true;
+	}
+}
 
 class FGizmoRectangleComponentSceneProxy final : public FPrimitiveSceneProxy
 {
@@ -19,6 +87,7 @@ public:
 		Color(InComponent->Color),
 		DirectionX(InComponent->DirectionX),
 		DirectionY(InComponent->DirectionY),
+		bOrientYAccordingToCamera(InComponent->bOrientYAccordingToCamera),
 		OffsetX(InComponent->OffsetX),
 		OffsetY(InComponent->OffsetY),
 		LengthX(InComponent->LengthX),
@@ -31,8 +100,7 @@ public:
 
 	virtual void GetDynamicMeshElements(const TArray<const FSceneView*>& Views, const FSceneViewFamily& ViewFamily, uint32 VisibilityMap, FMeshElementCollector& Collector) const override
 	{
-		// try to find focused scene view. May return nullptr.
-		const FSceneView* FocusedView = GizmoRenderingUtil::FindFocusedEditorSceneView(Views, ViewFamily, VisibilityMap);
+		using namespace GizmoRectangleComponentLocals;
 
 		const FMatrix& LocalToWorldMatrix = GetLocalToWorld();
 		FVector Origin = LocalToWorldMatrix.TransformPosition(FVector::ZeroVector);
@@ -43,82 +111,43 @@ public:
 			{
 				const FSceneView* View = Views[ViewIndex];
 				FPrimitiveDrawInterface* PDI = Collector.GetPDI(ViewIndex);
-				bool bIsFocusedView = (FocusedView != nullptr && View == FocusedView);
-				bool bIsOrtho = !View->IsPerspectiveProjection();
-
-				// direction to origin of gizmo
-				FVector ViewDirection = 
-					(bIsOrtho) ? (View->GetViewDirection()) : (Origin - View->ViewLocation);
-				ViewDirection.Normalize();
 
 				bool bWorldAxis = (bExternalWorldLocalState) ? (*bExternalWorldLocalState) : false;
-				FVector UseDirectionX = (bWorldAxis) ? DirectionX : FVector{ LocalToWorldMatrix.TransformVector(DirectionX) };
-				bool bFlippedX = (FVector::DotProduct(ViewDirection, UseDirectionX) > 0);
-				UseDirectionX = (bFlippedX) ? -UseDirectionX : UseDirectionX;
-				if (bIsFocusedView && bFlippedXExternal != nullptr)
-				{
-					*bFlippedXExternal = bFlippedX;
-				}
-
-				FVector UseDirectionY = (bWorldAxis) ? DirectionY : FVector{ LocalToWorldMatrix.TransformVector(DirectionY) };
-				bool bFlippedY = (FVector::DotProduct(ViewDirection, UseDirectionY) > 0);
-				UseDirectionY = (bFlippedY) ? -UseDirectionY : UseDirectionY;
-				if (bIsFocusedView && bFlippedYExternal != nullptr)
-				{
-					*bFlippedYExternal = bFlippedY;
-				}
-
-				FVector PlaneNormal = FVector::CrossProduct(UseDirectionX, UseDirectionY);
-				bool bRenderVisibility = FMath::Abs(FVector::DotProduct(PlaneNormal, ViewDirection)) > 0.25f;
-				if (bIsFocusedView && bExternalRenderVisibility != nullptr)
-				{
-					*bExternalRenderVisibility = bRenderVisibility;
-				}
-				if (bRenderVisibility == false)
+				
+				TArray<FVector> Corners;
+				bool bIsViewDependent = (bExternalIsViewDependent) ? (*bExternalIsViewDependent) : false;
+				bool bRenderVisibility = GetWorldCorners(bIsViewDependent, View, Origin, DirectionX, DirectionY,
+					OffsetX, OffsetY, LengthX, LengthY, bWorldAxis, bOrientYAccordingToCamera,
+					[&LocalToWorldMatrix](const FVector& VectorIn) { return FVector{ LocalToWorldMatrix.TransformVector(VectorIn) }; },
+					Corners);
+
+				if (!bRenderVisibility)
 				{
 					continue;
 				}
-
-				float PixelToWorldScale = GizmoRenderingUtil::CalculateLocalPixelToWorldScale(View, Origin);
-				float LengthScale = PixelToWorldScale;
-				if (bIsFocusedView && ExternalDynamicPixelToWorldScale != nullptr)
-				{
-					*ExternalDynamicPixelToWorldScale = PixelToWorldScale;
-				}
-
-
+				
 				float UseThickness = (bExternalHoverState != nullptr && *bExternalHoverState == true) ?
 					(HoverThicknessMultiplier * Thickness) : (Thickness);
-				if (!bIsOrtho)
+				if (View->IsPerspectiveProjection())
 				{
 					UseThickness *= (View->FOV / 90.0);		// compensate for FOV scaling in Gizmos...
 				}
 
-				double UseOffsetX = LengthScale * OffsetX;
-				double UseOffsetLengthX = LengthScale * (OffsetX + LengthX);
-				double UseOffsetY = LengthScale * OffsetY;
-				double UseOffsetLengthY = LengthScale * (OffsetY + LengthY);
-
-				FVector Point00 = Origin + UseOffsetX*UseDirectionX + UseOffsetY*UseDirectionY;
-				FVector Point10 = Origin + UseOffsetLengthX*UseDirectionX + UseOffsetY*UseDirectionY;
-				FVector Point11 = Origin + UseOffsetLengthX*UseDirectionX + UseOffsetLengthY*UseDirectionY;
-				FVector Point01 = Origin + UseOffsetX*UseDirectionX + UseOffsetLengthY*UseDirectionY;
-
 				if (SegmentFlags & 0x1)
 				{
-					PDI->DrawLine(Point00, Point10, Color, SDPG_Foreground, UseThickness, 0.0f, true);
+					PDI->DrawLine(Corners[0], Corners[1], Color, SDPG_Foreground, UseThickness, 0.0f, true);
 				}
 				if (SegmentFlags & 0x2)
 				{
-					PDI->DrawLine(Point10, Point11, Color, SDPG_Foreground, UseThickness, 0.0f, true);
+					PDI->DrawLine(Corners[1], Corners[2], Color, SDPG_Foreground, UseThickness, 0.0f, true);
 				}
 				if (SegmentFlags & 0x4)
 				{
-					PDI->DrawLine(Point11, Point01, Color, SDPG_Foreground, UseThickness, 0.0f, true);
+					PDI->DrawLine(Corners[2], Corners[3], Color, SDPG_Foreground, UseThickness, 0.0f, true);
 				}
 				if (SegmentFlags & 0x8)
 				{
-					PDI->DrawLine(Point01, Point00, Color, SDPG_Foreground, UseThickness, 0.0f, true);
+					PDI->DrawLine(Corners[3], Corners[0], Color, SDPG_Foreground, UseThickness, 0.0f, true);
 				}
 
 
@@ -163,22 +192,6 @@ public:
 	virtual uint32 GetMemoryFootprint(void) const override { return sizeof *this + GetAllocatedSize(); }
 	uint32 GetAllocatedSize(void) const { return FPrimitiveSceneProxy::GetAllocatedSize(); }
 
-	void SetExternalFlip(bool* bFlippedX, bool* bFlippedY)
-	{
-		bFlippedXExternal = bFlippedX;
-		bFlippedYExternal = bFlippedY;
-	}
-
-	void SetExternalDynamicPixelToWorldScale(float* DynamicPixelToWorldScale)
-	{
-		ExternalDynamicPixelToWorldScale = DynamicPixelToWorldScale;
-	}
-
-	void SetExternalRenderVisibility(bool* bRenderVisibility)
-	{
-		bExternalRenderVisibility = bRenderVisibility;
-	}
-
 	void SetExternalHoverState(bool* HoverState)
 	{
 		bExternalHoverState = HoverState;
@@ -189,10 +202,15 @@ public:
 		bExternalWorldLocalState = bWorldLocalState;
 	}
 
+	void SetExternalIsViewDependent(bool* bExternalIsViewDependentIn)
+	{
+		bExternalIsViewDependent = bExternalIsViewDependentIn;
+	}
 
 private:
 	FLinearColor Color;
 	FVector DirectionX, DirectionY;
+	bool bOrientYAccordingToCamera;
 	float OffsetX, OffsetY;
 	float LengthX, LengthY;
 	float Thickness;
@@ -202,12 +220,7 @@ private:
 	// set on Component for use in ::GetDynamicMeshElements()
 	bool* bExternalHoverState = nullptr;
 	bool* bExternalWorldLocalState = nullptr;
-
-	// set in ::GetDynamicMeshElements() for use by Component hit testing
-	bool* bFlippedXExternal = nullptr;
-	bool* bFlippedYExternal = nullptr;
-	float* ExternalDynamicPixelToWorldScale = nullptr;
-	bool* bExternalRenderVisibility = nullptr;
+	bool* bExternalIsViewDependent = nullptr;
 };
 
 
@@ -215,11 +228,9 @@ private:
 FPrimitiveSceneProxy* UGizmoRectangleComponent::CreateSceneProxy()
 {
 	FGizmoRectangleComponentSceneProxy* NewProxy = new FGizmoRectangleComponentSceneProxy(this);
-	NewProxy->SetExternalFlip(&bFlippedX, &bFlippedY);
-	NewProxy->SetExternalDynamicPixelToWorldScale(&DynamicPixelToWorldScale);
-	NewProxy->SetExternalRenderVisibility(&bRenderVisibility);
 	NewProxy->SetExternalHoverState(&bHovering);
 	NewProxy->SetExternalWorldLocalState(&bWorld);
+	NewProxy->SetExternalIsViewDependent(&bIsViewDependent);
 	return NewProxy;
 }
 
@@ -232,31 +243,23 @@ FBoxSphereBounds UGizmoRectangleComponent::CalcBounds(const FTransform& LocalToW
 
 bool UGizmoRectangleComponent::LineTraceComponent(FHitResult& OutHit, const FVector Start, const FVector End, const FCollisionQueryParams& Params)
 {
+	using namespace GizmoRectangleComponentLocals;
+
+	const FTransform& Transform = GetComponentToWorld();
+
+	FVector UseOrigin = Transform.TransformPosition(FVector::ZeroVector);
+
+	TArray<FVector> Corners;
+	bool bRenderVisibility = GetWorldCorners(bIsViewDependent, GizmoViewContext, UseOrigin, DirectionX, DirectionY, 
+		OffsetX, OffsetY, LengthX, LengthY,	bWorld, bOrientYAccordingToCamera,
+		[&Transform](const FVector& VectorIn) { return Transform.TransformVector(VectorIn); },
+		Corners);
+
 	if (bRenderVisibility == false)
 	{
 		return false;
 	}
 
-	const FTransform& Transform = this->GetComponentToWorld();
-
-	FVector UseDirectionX = (bFlippedX) ? -DirectionX : DirectionX;
-	UseDirectionX = (bWorld) ? UseDirectionX : Transform.TransformVector(UseDirectionX);
-	FVector UseDirectionY = (bFlippedY) ? -DirectionY : DirectionY;
-	UseDirectionY = (bWorld) ? UseDirectionY : Transform.TransformVector(UseDirectionY);
-	FVector UseOrigin = Transform.TransformPosition(FVector::ZeroVector);
-
-	float LengthScale = DynamicPixelToWorldScale;
-	double UseOffsetX = LengthScale * OffsetX;
-	double UseOffsetLengthX = LengthScale * (OffsetX + LengthX);
-	double UseOffsetY = LengthScale * OffsetY;
-	double UseOffsetLengthY = LengthScale * (OffsetY + LengthY);
-
-	FVector Points[4] = {
-		UseOrigin + UseOffsetX*UseDirectionX + UseOffsetY*UseDirectionY,
-		UseOrigin + UseOffsetLengthX*UseDirectionX + UseOffsetY*UseDirectionY,
-		UseOrigin + UseOffsetLengthX*UseDirectionX + UseOffsetLengthY*UseDirectionY,
-		UseOrigin + UseOffsetX*UseDirectionX + UseOffsetLengthY*UseDirectionY
-	};
 	static const int Triangles[2][3] = { {0,1,2}, {0,2,3} };
 
 	for (int j = 0; j < 2; ++j)
@@ -264,7 +267,7 @@ bool UGizmoRectangleComponent::LineTraceComponent(FHitResult& OutHit, const FVec
 		const int* Triangle = Triangles[j];
 		FVector HitPoint, HitNormal;
 		if (FMath::SegmentTriangleIntersection(Start, End,
-			Points[Triangle[0]], Points[Triangle[1]], Points[Triangle[2]], 
+			Corners[Triangle[0]], Corners[Triangle[1]], Corners[Triangle[2]],
 			HitPoint, HitNormal) )
 		{
 			OutHit.Component = this;
diff --git a/Engine/Source/Runtime/Experimental/InteractiveToolsFramework/Private/BaseGizmos/GizmoRenderingUtil.cpp b/Engine/Source/Runtime/Experimental/InteractiveToolsFramework/Private/BaseGizmos/GizmoRenderingUtil.cpp
--- a/Engine/Source/Runtime/Experimental/InteractiveToolsFramework/Private/BaseGizmos/GizmoRenderingUtil.cpp
+++ b/Engine/Source/Runtime/Experimental/InteractiveToolsFramework/Private/BaseGizmos/GizmoRenderingUtil.cpp
@@ -1,6 +1,7 @@
 // Copyright Epic Games, Inc. All Rights Reserved.
 
 #include "BaseGizmos/GizmoRenderingUtil.h"
+#include "BaseGizmos/GizmoViewContext.h"
 #include "RHI.h"
 #include "Math/UnrealMathUtility.h"
 #include "Math/NumericLimits.h"
@@ -17,132 +18,81 @@ static bool bGlobalUseCurrentSceneViewTracking = true;
 static bool bGlobalUseCurrentSceneViewTracking = false;
 #endif
 
-
-void GizmoRenderingUtil::SetGlobalFocusedEditorSceneView(const FSceneView* View)
+namespace GizmoRenderingUtilLocals
 {
-	GlobalCurrentSceneViewLock.Lock();
-	GlobalCurrentSceneView = View;
-	GlobalCurrentSceneViewLock.Unlock();
-}
-
-void GizmoRenderingUtil::SetGlobalFocusedSceneViewTrackingEnabled(bool bEnabled)
-{
-	bGlobalUseCurrentSceneViewTracking = bEnabled;
-}
-
-
-const FSceneView* GizmoRenderingUtil::FindFocusedEditorSceneView(
-	const TArray<const FSceneView*>& Views,
-	const FSceneViewFamily& ViewFamily,
-	uint32 VisibilityMap)
-{
-	// we are likely being called from a rendering thread GetDynamicMeshElements() function
-	const FSceneView* GlobalEditorView = nullptr;
-	GlobalCurrentSceneViewLock.Lock();
-	GlobalEditorView = GlobalCurrentSceneView;
-	GlobalCurrentSceneViewLock.Unlock();
-
-	// if we only have one view, and we are not tracking active view, just use that one
-	if (bGlobalUseCurrentSceneViewTracking == false && Views.Num() == 1)
+	static double VectorDifferenceSqr(const FVector2D& A, const FVector2D& B)
 	{
-		return Views[0];
+		double ax = A.X, ay = A.Y;
+		double bx = B.X, by = B.Y;
+		ax -= bx;
+		ay -= by;
+		return ax * ax + ay * ay;
 	}
 
-	// otherwise try to find the view that the game thread set for us
-	int32 VisibleViews = 0;
-	for (int32 ViewIndex = 0; ViewIndex < Views.Num(); ViewIndex++)
+	static double VectorDifferenceSqr(const FVector& A, const FVector& B)
 	{
-		if (VisibilityMap & (1 << ViewIndex))
-		{
-			VisibleViews++;
-			if (Views[ViewIndex] == GlobalEditorView)
-			{
-				return Views[ViewIndex];
-			}
-		}
+		double ax = A.X, ay = A.Y, az = A.Z;
+		double bx = B.X, by = B.Y, bz = B.Z;
+		ax -= bx;
+		ay -= by;
+		az -= bz;
+		return ax * ax + ay * ay + az * az;
 	}
 
-	// if we did not find our view, but only one view is visible, we can speculatively return that one
-	if (bGlobalUseCurrentSceneViewTracking == false && VisibleViews == 1)
+	// duplicates FSceneView::WorldToPixel but in double where possible (unfortunately WorldToScreen still in float)
+	template <typename SceneViewOrGizmoViewContext>
+	static FVector2D WorldToPixelDouble(const SceneViewOrGizmoViewContext* View, const FVector& Location)
 	{
-		for (int32 ViewIndex = 0; ViewIndex < Views.Num(); ViewIndex++)
-		{
-			if (VisibilityMap & (1 << ViewIndex))
-			{
-				return Views[ViewIndex];
-			}
-		}
-	}
-
-	// ok give up
-	if (bGlobalUseCurrentSceneViewTracking == false)
-	{
-		return Views[0];
-	}
-	else
-	{
-		return nullptr;
-	}
-}
-
+		FVector4 ScreenPoint = View->WorldToScreen(Location);
 
+		double InvW = (ScreenPoint.W > 0.0 ? 1.0 : -1.0) / (double)ScreenPoint.W;
+		double Y = (GProjectionSignY > 0.0) ? (double)ScreenPoint.Y : 1.0 - (double)ScreenPoint.Y;
 
+		const FIntRect& UnscaledViewRect = View->UnscaledViewRect;
+		double PosX = (double)UnscaledViewRect.Min.X + (0.5 + (double)ScreenPoint.X * 0.5 * InvW) * (double)UnscaledViewRect.Width();
+		double PosY = (double)UnscaledViewRect.Min.Y + (0.5 - Y * 0.5 * InvW) * (double)UnscaledViewRect.Height();
 
-static double VectorDifferenceSqr(const FVector2D& A, const FVector2D& B)
-{
-	double ax = A.X, ay = A.Y;
-	double bx = B.X, by = B.Y;
-	ax -= bx;
-	ay -= by;
-	return ax*ax + ay*ay;
-}
+		return FVector2D((float)PosX, (float)PosY);
+	}
 
-static double VectorDifferenceSqr(const FVector& A, const FVector& B)
-{
-	double ax = A.X, ay = A.Y, az = A.Z;
-	double bx = B.X, by = B.Y, bz = B.Z;
-	ax -= bx; 
-	ay -= by;
-	az -= bz;
-	return ax*ax + ay*ay + az*az;
+	// Shared code between the GizmoViewContext and FSceneView versions of CalculateLocalPixelToWorldScale. We don't expose
+	// this template itself to the user since it's not entirely intuitive that this works on two very different types.
+	template <typename SceneViewOrGizmoViewContext>
+	float CalculateLocalPixelToWorldScale(
+		const SceneViewOrGizmoViewContext* View,
+		const FVector& Location)
+	{
+		// To calculate this scale at Location, we project Location to screen and also project a second
+		// point at a small distance from Location in a camera-perpendicular plane, then measure 2D/3D distance ratio. 
+		// However, because some of the computations are done in float, there will be enormous numerical error 
+		// when the camera is very far from the location if the distance is relatively small. The "W" value
+		// below gives us a sense of this distance, so we make the offset relative to that
+		// (this does do one redundant WorldToScreen)
+		FVector4 LocationScreenPoint = View->WorldToScreen(Location);
+		float OffsetDelta = LocationScreenPoint.W * 0.01f;
+
+		FVector2D PixelA = WorldToPixelDouble(View, Location);
+		FVector OffsetPointWorld = Location + OffsetDelta * View->GetViewRight() + OffsetDelta * View->GetViewUp();
+		FVector2D PixelB = WorldToPixelDouble(View, OffsetPointWorld);
+
+		double PixelDeltaSqr = VectorDifferenceSqr(PixelA, PixelB);
+		double WorldDeltaSqr = VectorDifferenceSqr(Location, OffsetPointWorld);
+		return (float)(sqrt(WorldDeltaSqr / PixelDeltaSqr));
+	}
 }
 
-// duplicates FSceneView::WorldToPixel but in double where possible (unfortunately WorldToScreen still in float)
-static FVector2D WorldToPixelDouble(const FSceneView* View, const FVector& Location)
+float GizmoRenderingUtil::CalculateLocalPixelToWorldScale(
+	const FSceneView* View,
+	const FVector& Location)
 {
-	FVector4 ScreenPoint = View->WorldToScreen(Location);
-
-	double InvW = (ScreenPoint.W > 0.0 ? 1.0 : -1.0) / (double)ScreenPoint.W;
-	double Y = (GProjectionSignY > 0.0) ? (double)ScreenPoint.Y : 1.0 - (double)ScreenPoint.Y;
-
-	const FIntRect& UnscaledViewRect = View->UnscaledViewRect;
-	double PosX = (double)UnscaledViewRect.Min.X + (0.5 + (double)ScreenPoint.X * 0.5 * InvW) * (double)UnscaledViewRect.Width();
-	double PosY = (double)UnscaledViewRect.Min.Y + (0.5 - Y * 0.5 * InvW) * (double)UnscaledViewRect.Height();
-
-	return FVector2D((float)PosX, (float)PosY);
+	return GizmoRenderingUtilLocals::CalculateLocalPixelToWorldScale(View, Location);
 }
 
-
 float GizmoRenderingUtil::CalculateLocalPixelToWorldScale(
-	const FSceneView* View,
+	const UGizmoViewContext* ViewContext,
 	const FVector& Location)
 {
-	// To calculate this scale at Location, we project Location to screen and also project a second
-	// point at a small distance from Location in a camera-perpendicular plane, then measure 2D/3D distance ratio. 
-	// However, because some of the computations are done in float, there will be enormous numerical error 
-	// when the camera is very far from the location if the distance is relatively small. The "W" value
-	// below gives us a sense of this distance, so we make the offset relative to that
-	// (this does do one redundant WorldToScreen)
-	FVector4 LocationScreenPoint = View->WorldToScreen(Location);
-	float OffsetDelta = LocationScreenPoint.W * 0.01f;
-
-	FVector2D PixelA = WorldToPixelDouble(View, Location);
-	FVector OffsetPointWorld = Location + OffsetDelta*View->GetViewRight() + OffsetDelta*View->GetViewUp();
-	FVector2D PixelB = WorldToPixelDouble(View, OffsetPointWorld);
-
-	double PixelDeltaSqr = VectorDifferenceSqr(PixelA, PixelB);
-	double WorldDeltaSqr = VectorDifferenceSqr(Location, OffsetPointWorld);
-	return (float)(sqrt(WorldDeltaSqr / PixelDeltaSqr));
+	return GizmoRenderingUtilLocals::CalculateLocalPixelToWorldScale(ViewContext, Location);
 }
 
 float GizmoRenderingUtil::CalculateViewDependentScaleAndFlatten(
diff --git a/Engine/Source/Runtime/Experimental/InteractiveToolsFramework/Private/BaseGizmos/GizmoViewContext.cpp b/Engine/Source/Runtime/Experimental/InteractiveToolsFramework/Private/BaseGizmos/GizmoViewContext.cpp
new file mode 100644
--- /dev/null
+++ b/Engine/Source/Runtime/Experimental/InteractiveToolsFramework/Private/BaseGizmos/GizmoViewContext.cpp
@@ -0,0 +1,5 @@
+// Copyright Epic Games, Inc. All Rights Reserved.
+
+#include "BaseGizmos/GizmoViewContext.h"
+
+#include "SceneView.h"
diff --git a/Engine/Source/Runtime/Experimental/InteractiveToolsFramework/Private/BaseGizmos/IntervalGizmo.cpp b/Engine/Source/Runtime/Experimental/InteractiveToolsFramework/Private/BaseGizmos/IntervalGizmo.cpp
--- a/Engine/Source/Runtime/Experimental/InteractiveToolsFramework/Private/BaseGizmos/IntervalGizmo.cpp
+++ b/Engine/Source/Runtime/Experimental/InteractiveToolsFramework/Private/BaseGizmos/IntervalGizmo.cpp
@@ -6,12 +6,14 @@
 
 #include "BaseGizmos/GizmoComponents.h"
 #include "BaseGizmos/GizmoLineHandleComponent.h"
+#include "BaseGizmos/GizmoViewContext.h"
 
 // need this to implement hover
 #include "BaseGizmos/GizmoBaseComponent.h"
 
 #include "Components/SphereComponent.h"
 #include "Components/PrimitiveComponent.h"
+#include "ContextObjectStore.h"
 #include "Engine/World.h"
 #include "Engine/CollisionProfile.h"
 
@@ -136,7 +138,7 @@ AIntervalGizmoActor::AIntervalGizmoActor()
 	SphereComponent->SetCollisionProfileName(UCollisionProfile::NoCollision_ProfileName);
 }
 
-AIntervalGizmoActor* AIntervalGizmoActor::ConstructDefaultIntervalGizmo(UWorld* World)
+AIntervalGizmoActor* AIntervalGizmoActor::ConstructDefaultIntervalGizmo(UWorld* World, UGizmoViewContext* GizmoViewContext)
 {
 	FActorSpawnParameters SpawnInfo;
 	AIntervalGizmoActor* NewActor = World->SpawnActor<AIntervalGizmoActor>(FVector::ZeroVector, FRotator::ZeroRotator, SpawnInfo);
@@ -144,9 +146,9 @@ AIntervalGizmoActor* AIntervalGizmoActor::ConstructDefaultIntervalGizmo(UWorld*
 	const FLinearColor MintGreen(152 / 255.f, 1.f, 152 / 255.f);
 	
 	// add components 
-	NewActor->UpIntervalComponent      = AddDefaultLineHandleComponent(World, NewActor, MintGreen, FVector(0, 1, 0), FVector(0, 0, 1));
-	NewActor->DownIntervalComponent    = AddDefaultLineHandleComponent(World, NewActor, MintGreen, FVector(0, 1, 0), FVector(0, 0, 1));
-	NewActor->ForwardIntervalComponent = AddDefaultLineHandleComponent(World, NewActor, MintGreen, FVector(1, 0, 0), FVector(0, 1, 0));
+	NewActor->UpIntervalComponent      = AddDefaultLineHandleComponent(World, NewActor, GizmoViewContext, MintGreen, FVector(0, 1, 0), FVector(0, 0, 1));
+	NewActor->DownIntervalComponent    = AddDefaultLineHandleComponent(World, NewActor, GizmoViewContext, MintGreen, FVector(0, 1, 0), FVector(0, 0, 1));
+	NewActor->ForwardIntervalComponent = AddDefaultLineHandleComponent(World, NewActor, GizmoViewContext, MintGreen, FVector(1, 0, 0), FVector(0, 1, 0));
 
 
 	return NewActor;
@@ -159,8 +161,11 @@ UInteractiveGizmo* UIntervalGizmoBuilder::BuildGizmo(const FToolBuilderState& Sc
 	UIntervalGizmo* NewGizmo = NewObject<UIntervalGizmo>(SceneState.GizmoManager);
 	NewGizmo->SetWorld(SceneState.World);
 
+	UGizmoViewContext* GizmoViewContext = SceneState.ToolManager->GetContextObjectStore()->FindContext<UGizmoViewContext>();
+	check(GizmoViewContext && GizmoViewContext->IsValidLowLevel());
+
 	// use default gizmo actor if client has not given us a new builder
-	NewGizmo->SetGizmoActorBuilder((GizmoActorBuilder) ? GizmoActorBuilder : MakeShared<FIntervalGizmoActorFactory>());
+	NewGizmo->SetGizmoActorBuilder(GizmoActorBuilder ? GizmoActorBuilder : MakeShared<FIntervalGizmoActorFactory>(GizmoViewContext));
 
 	// override default hover function if proposed
 	if (UpdateHoverFunction)
diff --git a/Engine/Source/Runtime/Experimental/InteractiveToolsFramework/Private/BaseGizmos/RepositionableTransformGizmo.cpp b/Engine/Source/Runtime/Experimental/InteractiveToolsFramework/Private/BaseGizmos/RepositionableTransformGizmo.cpp
--- a/Engine/Source/Runtime/Experimental/InteractiveToolsFramework/Private/BaseGizmos/RepositionableTransformGizmo.cpp
+++ b/Engine/Source/Runtime/Experimental/InteractiveToolsFramework/Private/BaseGizmos/RepositionableTransformGizmo.cpp
@@ -5,10 +5,12 @@
 #include "BaseBehaviors/ClickDragBehavior.h"
 #include "BaseGizmos/AxisAngleGizmo.h"
 #include "BaseGizmos/AxisPositionGizmo.h"
+#include "BaseGizmos/GizmoViewContext.h"
 #include "BaseGizmos/HitTargets.h"
 #include "BaseGizmos/PlanePositionGizmo.h"
 #include "BaseGizmos/TransformProxy.h"
 #include "BaseGizmos/TransformSources.h"
+#include "ContextObjectStore.h"
 
 #include "InteractiveGizmo.h"
 
@@ -20,8 +22,11 @@ UInteractiveGizmo* URepositionableTransformGizmoBuilder::BuildGizmo(const FToolB
 	UTransformGizmo* NewGizmo = NewObject<URepositionableTransformGizmo>(SceneState.GizmoManager);
 	NewGizmo->SetWorld(SceneState.World);
 
+	UGizmoViewContext* GizmoViewContext = SceneState.ToolManager->GetContextObjectStore()->FindContext<UGizmoViewContext>();
+	check(GizmoViewContext && GizmoViewContext->IsValidLowLevel());
+
 	// use default gizmo actor if client has not given us a new builder
-	NewGizmo->SetGizmoActorBuilder((GizmoActorBuilder) ? GizmoActorBuilder : MakeShared<FTransformGizmoActorFactory>());
+	NewGizmo->SetGizmoActorBuilder(GizmoActorBuilder ? GizmoActorBuilder : MakeShared<FTransformGizmoActorFactory>(GizmoViewContext));
 
 	NewGizmo->SetSubGizmoBuilderIdentifiers(AxisPositionBuilderIdentifier, PlanePositionBuilderIdentifier, AxisAngleBuilderIdentifier);
 
diff --git a/Engine/Source/Runtime/Experimental/InteractiveToolsFramework/Private/BaseGizmos/TransformGizmo.cpp b/Engine/Source/Runtime/Experimental/InteractiveToolsFramework/Private/BaseGizmos/TransformGizmo.cpp
--- a/Engine/Source/Runtime/Experimental/InteractiveToolsFramework/Private/BaseGizmos/TransformGizmo.cpp
+++ b/Engine/Source/Runtime/Experimental/InteractiveToolsFramework/Private/BaseGizmos/TransformGizmo.cpp
@@ -12,12 +12,14 @@
 #include "BaseGizmos/GizmoCircleComponent.h"
 #include "BaseGizmos/GizmoBoxComponent.h"
 #include "BaseGizmos/GizmoLineHandleComponent.h"
+#include "BaseGizmos/GizmoViewContext.h"
 
 // need this to implement hover
 #include "BaseGizmos/GizmoBaseComponent.h"
 
 #include "Components/SphereComponent.h"
 #include "Components/PrimitiveComponent.h"
+#include "ContextObjectStore.h"
 #include "Engine/World.h"
 #include "Engine/CollisionProfile.h"
 
@@ -37,9 +39,9 @@ ATransformGizmoActor::ATransformGizmoActor()
 
 
 
-ATransformGizmoActor* ATransformGizmoActor::ConstructDefault3AxisGizmo(UWorld* World)
+ATransformGizmoActor* ATransformGizmoActor::ConstructDefault3AxisGizmo(UWorld* World, UGizmoViewContext* GizmoViewContext)
 {
-	return ConstructCustom3AxisGizmo(World, 
+	return ConstructCustom3AxisGizmo(World, GizmoViewContext,
 		ETransformGizmoSubElements::TranslateAllAxes |
 		ETransformGizmoSubElements::TranslateAllPlanes |
 		ETransformGizmoSubElements::RotateAllAxes |
@@ -51,7 +53,7 @@ ATransformGizmoActor* ATransformGizmoActor::ConstructDefault3AxisGizmo(UWorld* W
 
 
 ATransformGizmoActor* ATransformGizmoActor::ConstructCustom3AxisGizmo(
-	UWorld* World,
+	UWorld* World, UGizmoViewContext* GizmoViewContext,
 	ETransformGizmoSubElements Elements)
 {
 	FActorSpawnParameters SpawnInfo;
@@ -62,7 +64,7 @@ ATransformGizmoActor* ATransformGizmoActor::ConstructCustom3AxisGizmo(
 
 	auto MakeAxisArrowFunc = [&](const FLinearColor& Color, const FVector& Axis)
 	{
-		UGizmoArrowComponent* Component = AddDefaultArrowComponent(World, NewActor, Color, Axis, 60.0f);
+		UGizmoArrowComponent* Component = AddDefaultArrowComponent(World, NewActor, GizmoViewContext, Color, Axis, 60.0f);
 		Component->Gap = 20.0f;
 		Component->Thickness = GizmoLineThickness;
 		Component->NotifyExternalPropertyUpdates();
@@ -84,7 +86,7 @@ ATransformGizmoActor* ATransformGizmoActor::ConstructCustom3AxisGizmo(
 
 	auto MakePlaneRectFunc = [&](const FLinearColor& Color, const FVector& Axis0, const FVector& Axis1)
 	{
-		UGizmoRectangleComponent* Component = AddDefaultRectangleComponent(World, NewActor, Color, Axis0, Axis1);
+		UGizmoRectangleComponent* Component = AddDefaultRectangleComponent(World, NewActor, GizmoViewContext, Color, Axis0, Axis1);
 		Component->LengthX = Component->LengthY = 30.0f;
 		Component->SegmentFlags = 0x2 | 0x4;
 		Component->Thickness = GizmoLineThickness;
@@ -106,7 +108,7 @@ ATransformGizmoActor* ATransformGizmoActor::ConstructCustom3AxisGizmo(
 
 	auto MakeAxisRotateCircleFunc = [&](const FLinearColor& Color, const FVector& Axis)
 	{
-		UGizmoCircleComponent* Component = AddDefaultCircleComponent(World, NewActor, Color, Axis, 120.0f);
+		UGizmoCircleComponent* Component = AddDefaultCircleComponent(World, NewActor, GizmoViewContext, Color, Axis, 120.0f);
 		Component->Thickness = GizmoLineThickness;
 		Component->NotifyExternalPropertyUpdates();
 		return Component;
@@ -136,6 +138,7 @@ ATransformGizmoActor* ATransformGizmoActor::ConstructCustom3AxisGizmo(
 		UGizmoCircleComponent* SphereEdge = NewObject<UGizmoCircleComponent>(NewActor);
 		NewActor->AddInstanceComponent(SphereEdge);
 		SphereEdge->AttachToComponent(NewActor->GetRootComponent(), FAttachmentTransformRules::KeepRelativeTransform);
+		SphereEdge->SetGizmoViewContext(GizmoViewContext);
 		SphereEdge->Color = FLinearColor::Gray;
 		SphereEdge->Thickness = 1.0f;
 		SphereEdge->Radius = 120.0f;
@@ -148,7 +151,7 @@ ATransformGizmoActor* ATransformGizmoActor::ConstructCustom3AxisGizmo(
 	if ((Elements & ETransformGizmoSubElements::ScaleUniform) != ETransformGizmoSubElements::None)
 	{
 		float BoxSize = 14.0f;
-		UGizmoBoxComponent* ScaleComponent = AddDefaultBoxComponent(World, NewActor, FLinearColor::Black, 
+		UGizmoBoxComponent* ScaleComponent = AddDefaultBoxComponent(World, NewActor, GizmoViewContext, FLinearColor::Black,
 			FVector(BoxSize/2, BoxSize/2, BoxSize/2), FVector(BoxSize, BoxSize, BoxSize));
 		NewActor->UniformScale = ScaleComponent;
 	}
@@ -157,10 +160,11 @@ ATransformGizmoActor* ATransformGizmoActor::ConstructCustom3AxisGizmo(
 
 	auto MakeAxisScaleFunc = [&](const FLinearColor& Color, const FVector& Axis0, const FVector& Axis1)
 	{
-		UGizmoRectangleComponent* ScaleComponent = AddDefaultRectangleComponent(World, NewActor, Color, Axis0, Axis1);
+		UGizmoRectangleComponent* ScaleComponent = AddDefaultRectangleComponent(World, NewActor, GizmoViewContext, Color, Axis0, Axis1);
 		ScaleComponent->OffsetX = 140.0f; ScaleComponent->OffsetY = -10.0f;
 		ScaleComponent->LengthX = 7.0f; ScaleComponent->LengthY = 20.0f;
 		ScaleComponent->Thickness = GizmoLineThickness;
+		ScaleComponent->bOrientYAccordingToCamera = true;
 		ScaleComponent->NotifyExternalPropertyUpdates();
 		ScaleComponent->SegmentFlags = 0x1 | 0x2 | 0x4; // | 0x8;
 		return ScaleComponent;
@@ -181,7 +185,7 @@ ATransformGizmoActor* ATransformGizmoActor::ConstructCustom3AxisGizmo(
 
 	auto MakePlaneScaleFunc = [&](const FLinearColor& Color, const FVector& Axis0, const FVector& Axis1)
 	{
-		UGizmoRectangleComponent* ScaleComponent = AddDefaultRectangleComponent(World, NewActor, Color, Axis0, Axis1);
+		UGizmoRectangleComponent* ScaleComponent = AddDefaultRectangleComponent(World, NewActor, GizmoViewContext, Color, Axis0, Axis1);
 		ScaleComponent->OffsetX = ScaleComponent->OffsetY = 120.0f;
 		ScaleComponent->LengthX = ScaleComponent->LengthY = 20.0f;
 		ScaleComponent->Thickness = GizmoLineThickness;
@@ -211,7 +215,7 @@ ATransformGizmoActor* ATransformGizmoActor::ConstructCustom3AxisGizmo(
 
 ATransformGizmoActor* FTransformGizmoActorFactory::CreateNewGizmoActor(UWorld* World) const
 {
-	return ATransformGizmoActor::ConstructCustom3AxisGizmo(World, EnableElements);
+	return ATransformGizmoActor::ConstructCustom3AxisGizmo(World, GizmoViewContext, EnableElements);
 }
 
 
@@ -221,8 +225,11 @@ UInteractiveGizmo* UTransformGizmoBuilder::BuildGizmo(const FToolBuilderState& S
 	UTransformGizmo* NewGizmo = NewObject<UTransformGizmo>(SceneState.GizmoManager);
 	NewGizmo->SetWorld(SceneState.World);
 
+	UGizmoViewContext* GizmoViewContext = SceneState.ToolManager->GetContextObjectStore()->FindContext<UGizmoViewContext>();
+	check(GizmoViewContext && GizmoViewContext->IsValidLowLevel());
+
 	// use default gizmo actor if client has not given us a new builder
-	NewGizmo->SetGizmoActorBuilder( (GizmoActorBuilder) ? GizmoActorBuilder : MakeShared<FTransformGizmoActorFactory>() );
+	NewGizmo->SetGizmoActorBuilder(GizmoActorBuilder ? GizmoActorBuilder : MakeShared<FTransformGizmoActorFactory>(GizmoViewContext));
 
 	NewGizmo->SetSubGizmoBuilderIdentifiers(AxisPositionBuilderIdentifier, PlanePositionBuilderIdentifier, AxisAngleBuilderIdentifier);
 
@@ -247,7 +254,6 @@ void UTransformGizmo::SetWorld(UWorld* WorldIn)
 	this->World = WorldIn;
 }
 
-
 void UTransformGizmo::SetGizmoActorBuilder(TSharedPtr<FTransformGizmoActorFactory> Builder)
 {
 	GizmoActorBuilder = Builder;
diff --git a/Engine/Source/Runtime/Experimental/InteractiveToolsFramework/Private/BaseGizmos/TransformGizmoUtil.cpp b/Engine/Source/Runtime/Experimental/InteractiveToolsFramework/Private/BaseGizmos/TransformGizmoUtil.cpp
--- a/Engine/Source/Runtime/Experimental/InteractiveToolsFramework/Private/BaseGizmos/TransformGizmoUtil.cpp
+++ b/Engine/Source/Runtime/Experimental/InteractiveToolsFramework/Private/BaseGizmos/TransformGizmoUtil.cpp
@@ -10,6 +10,7 @@
 #include "BaseGizmos/AxisPositionGizmo.h"
 #include "BaseGizmos/PlanePositionGizmo.h"
 #include "BaseGizmos/AxisAngleGizmo.h"
+#include "BaseGizmos/GizmoViewContext.h"
 #include "BaseGizmos/TransformGizmo.h"
 #include "BaseGizmos/RepositionableTransformGizmo.h"
 
@@ -46,7 +47,14 @@ void UTransformGizmoContextObject::RegisterGizmosWithManager(UInteractiveToolMan
 	TransformBuilder->AxisAngleBuilderIdentifier = DefaultAxisAngleBuilderIdentifier;
 	GizmoManager->RegisterGizmoType(DefaultThreeAxisTransformBuilderIdentifier, TransformBuilder);
 
-	GizmoActorBuilder = MakeShared<FTransformGizmoActorFactory>();
+	UGizmoViewContext* GizmoViewContext = ToolManager->GetContextObjectStore()->FindContext<UGizmoViewContext>();
+	if (!GizmoViewContext)
+	{
+		GizmoViewContext = NewObject<UGizmoViewContext>();
+		ToolManager->GetContextObjectStore()->AddContextObject(GizmoViewContext);
+	}
+
+	GizmoActorBuilder = MakeShared<FTransformGizmoActorFactory>(GizmoViewContext);
 
 	UTransformGizmoBuilder* CustomThreeAxisBuilder = NewObject<UTransformGizmoBuilder>();
 	CustomThreeAxisBuilder->AxisPositionBuilderIdentifier = DefaultAxisPositionBuilderIdentifier;
diff --git a/Engine/Source/Runtime/Experimental/InteractiveToolsFramework/Private/InteractiveGizmoManager.cpp b/Engine/Source/Runtime/Experimental/InteractiveToolsFramework/Private/InteractiveGizmoManager.cpp
--- a/Engine/Source/Runtime/Experimental/InteractiveToolsFramework/Private/InteractiveGizmoManager.cpp
+++ b/Engine/Source/Runtime/Experimental/InteractiveToolsFramework/Private/InteractiveGizmoManager.cpp
@@ -5,11 +5,13 @@
 #include "InteractiveToolsContext.h"
 
 #include "BaseGizmos/AxisPositionGizmo.h"
+#include "BaseGizmos/GizmoViewContext.h"
 #include "BaseGizmos/PlanePositionGizmo.h"
 #include "BaseGizmos/AxisAngleGizmo.h"
 #include "BaseGizmos/TransformGizmo.h"
 #include "BaseGizmos/RepositionableTransformGizmo.h"
 #include "BaseGizmos/ScalableSphereGizmo.h"
+#include "ContextObjectStore.h"
 
 #define LOCTEXT_NAMESPACE "UInteractiveGizmoManager"
 
@@ -300,7 +302,13 @@ void UInteractiveGizmoManager::RegisterDefaultGizmos()
 	TransformBuilder->AxisAngleBuilderIdentifier = DefaultAxisAngleBuilderIdentifier;
 	RegisterGizmoType(DefaultThreeAxisTransformBuilderIdentifier, TransformBuilder);
 
-	GizmoActorBuilder = MakeShared<FTransformGizmoActorFactory>();
+	UGizmoViewContext* GizmoViewContext = GetContextObjectStore()->FindContext<UGizmoViewContext>();
+	if (!GizmoViewContext)
+	{
+		GizmoViewContext = NewObject<UGizmoViewContext>();
+		GetContextObjectStore()->AddContextObject(GizmoViewContext);
+	}
+	GizmoActorBuilder = MakeShared<FTransformGizmoActorFactory>(GizmoViewContext);
 
 	UTransformGizmoBuilder* CustomThreeAxisBuilder = NewObject<UTransformGizmoBuilder>();
 	CustomThreeAxisBuilder->AxisPositionBuilderIdentifier = DefaultAxisPositionBuilderIdentifier;
diff --git a/Engine/Source/Runtime/Experimental/InteractiveToolsFramework/Private/InteractiveToolsContext.cpp b/Engine/Source/Runtime/Experimental/InteractiveToolsFramework/Private/InteractiveToolsContext.cpp
--- a/Engine/Source/Runtime/Experimental/InteractiveToolsFramework/Private/InteractiveToolsContext.cpp
+++ b/Engine/Source/Runtime/Experimental/InteractiveToolsFramework/Private/InteractiveToolsContext.cpp
@@ -125,10 +125,11 @@ void UInteractiveToolsContext::Initialize(IToolsContextQueriesAPI* QueriesAPI, I
 	InputRouter = CreateInputRouterFunc(InitInfo);
 	InitInfo.InputRouter = InputRouter;
 
+	// Context store needs to be built before managers in case they want to use it
+	ContextObjectStore = CreateContextStoreFunc(InitInfo);
 	ToolManager = CreateToolManagerFunc(InitInfo);
 	TargetManager = CreateToolTargetManagerFunc(InitInfo);
 	GizmoManager = CreateGizmoManagerFunc(InitInfo);
-	ContextObjectStore = CreateContextStoreFunc(InitInfo);
 }
 
 
diff --git a/Engine/Source/Runtime/Experimental/InteractiveToolsFramework/Public/BaseGizmos/AxisPositionGizmo.h b/Engine/Source/Runtime/Experimental/InteractiveToolsFramework/Public/BaseGizmos/AxisPositionGizmo.h
--- a/Engine/Source/Runtime/Experimental/InteractiveToolsFramework/Public/BaseGizmos/AxisPositionGizmo.h
+++ b/Engine/Source/Runtime/Experimental/InteractiveToolsFramework/Public/BaseGizmos/AxisPositionGizmo.h
@@ -11,6 +11,7 @@
 
 
 class UClickDragInputBehavior;
+class UGizmoViewContext;
 
 UCLASS()
 class INTERACTIVETOOLSFRAMEWORK_API UAxisPositionGizmoBuilder : public UInteractiveGizmoBuilder
@@ -73,6 +74,9 @@ public:
 	UPROPERTY()
 	TScriptInterface<IGizmoFloatParameterSource> ParameterSource;
 
+	UPROPERTY()
+	UGizmoViewContext* GizmoViewContext;
+
 	/** The HitTarget provides a hit-test against some 3D element (presumably a visual widget) that controls when interaction can start */
 	UPROPERTY()
 	TScriptInterface<IGizmoClickTarget> HitTarget;
diff --git a/Engine/Source/Runtime/Experimental/InteractiveToolsFramework/Public/BaseGizmos/GizmoActor.h b/Engine/Source/Runtime/Experimental/InteractiveToolsFramework/Public/BaseGizmos/GizmoActor.h
--- a/Engine/Source/Runtime/Experimental/InteractiveToolsFramework/Public/BaseGizmos/GizmoActor.h
+++ b/Engine/Source/Runtime/Experimental/InteractiveToolsFramework/Public/BaseGizmos/GizmoActor.h
@@ -11,6 +11,7 @@ class UGizmoRectangleComponent;
 class UGizmoCircleComponent;
 class UGizmoBoxComponent;
 class UGizmoLineHandleComponent;
+class UGizmoViewContext;
 
 
 /**
@@ -27,31 +28,31 @@ public:
 
 	/** Add standard arrow component to Actor, generally used for axis-translation */
 	static UGizmoArrowComponent* AddDefaultArrowComponent(
-		UWorld* World, AActor* Actor,
+		UWorld* World, AActor* Actor, UGizmoViewContext* GizmoViewContext,
 		const FLinearColor& Color, const FVector& LocalDirection, const float Length = 80.0f
 	);
 	/** Add standard rectangle component to Actor, generally used for plane-translation */
 	static UGizmoRectangleComponent* AddDefaultRectangleComponent(
-		UWorld* World, AActor* Actor,
+		UWorld* World, AActor* Actor, UGizmoViewContext* GizmoViewContext,
 		const FLinearColor& Color, const FVector& PlaneAxis1, const FVector& PlaneAxisx2
 	);
 	/** Add standard circle component to Actor, generally used for axis-rotation */
 	static UGizmoCircleComponent* AddDefaultCircleComponent(
-		UWorld* World, AActor* Actor,
+		UWorld* World, AActor* Actor, UGizmoViewContext* GizmoViewContext,
 		const FLinearColor& Color, const FVector& PlaneNormal,
 		float Radius = 120.0f
 	);
 
 	/** Add standard 3D box component to Actor. By default the box is axis-aligned, centered at the specified Origin */
 	static UGizmoBoxComponent* AddDefaultBoxComponent(
-		UWorld* World, AActor* Actor,
+		UWorld* World, AActor* Actor, UGizmoViewContext* GizmoViewContext,
 		const FLinearColor& Color, const FVector& Origin, 
 		const FVector& Dimensions = FVector(20.0f, 20.0f, 20.0f)
 	);
 
 	/** Add standard disk component to Actor, generally used for handles */
 	static UGizmoLineHandleComponent* AddDefaultLineHandleComponent(
-		UWorld* World, AActor* Actor,
+		UWorld* World, AActor* Actor, UGizmoViewContext* GizmoViewContext,
 		const FLinearColor& Color, const FVector& PlaneNormal, const FVector& LocalDirection,
 		const float Length = 60.f, const bool bImageScale = false
 	);
diff --git a/Engine/Source/Runtime/Experimental/InteractiveToolsFramework/Public/BaseGizmos/GizmoArrowComponent.h b/Engine/Source/Runtime/Experimental/InteractiveToolsFramework/Public/BaseGizmos/GizmoArrowComponent.h
--- a/Engine/Source/Runtime/Experimental/InteractiveToolsFramework/Public/BaseGizmos/GizmoArrowComponent.h
+++ b/Engine/Source/Runtime/Experimental/InteractiveToolsFramework/Public/BaseGizmos/GizmoArrowComponent.h
@@ -38,10 +38,4 @@ private:
 	//~ Begin USceneComponent Interface.
 	virtual FBoxSphereBounds CalcBounds(const FTransform& LocalToWorld) const override;
 	//~ Begin USceneComponent Interface.
-
-	// if true, we drew along -Direction instead of Direction, and so should hit-test accordingly
-	bool bFlipped;
-
-	// gizmo visibility
-	bool bRenderVisibility = true;
 };
\ No newline at end of file
diff --git a/Engine/Source/Runtime/Experimental/InteractiveToolsFramework/Public/BaseGizmos/GizmoBaseComponent.h b/Engine/Source/Runtime/Experimental/InteractiveToolsFramework/Public/BaseGizmos/GizmoBaseComponent.h
--- a/Engine/Source/Runtime/Experimental/InteractiveToolsFramework/Public/BaseGizmos/GizmoBaseComponent.h
+++ b/Engine/Source/Runtime/Experimental/InteractiveToolsFramework/Public/BaseGizmos/GizmoBaseComponent.h
@@ -7,6 +7,8 @@
 #include "Components/PrimitiveComponent.h"
 #include "GizmoBaseComponent.generated.h"
 
+class UGizmoViewContext;
+
 /**
  * Base class for simple Components intended to be used as part of 3D Gizmos.
  * Contains common properties and utility functions.
@@ -45,6 +47,11 @@ public:
 	UPROPERTY(EditAnywhere, Category = Options)
 	float PixelHitDistanceThreshold = 7.0f;
 
+	void SetGizmoViewContext(UGizmoViewContext* GizmoViewContextIn)
+	{
+		GizmoViewContext = GizmoViewContextIn;
+		bIsViewDependent = (GizmoViewContext != nullptr);
+	}
 
 public:
 	UFUNCTION()
@@ -67,12 +74,17 @@ public:
 
 
 protected:
-	// scale factor between pixel distances and world distances at Gizmo origin
-	float DynamicPixelToWorldScale = 1.0f;
 
 	// hover state
 	bool bHovering = false;
 
 	// world/local coordinates state
 	bool bWorld = false;
+
+	UPROPERTY()
+	UGizmoViewContext* GizmoViewContext = nullptr;
+
+	// True when GizmoViewContext is not null. Here as a boolean so it
+	// can be pointed to by the proxy to see what it should do.
+	bool bIsViewDependent = false;
 };
\ No newline at end of file
diff --git a/Engine/Source/Runtime/Experimental/InteractiveToolsFramework/Public/BaseGizmos/GizmoBoxComponent.h b/Engine/Source/Runtime/Experimental/InteractiveToolsFramework/Public/BaseGizmos/GizmoBoxComponent.h
--- a/Engine/Source/Runtime/Experimental/InteractiveToolsFramework/Public/BaseGizmos/GizmoBoxComponent.h
+++ b/Engine/Source/Runtime/Experimental/InteractiveToolsFramework/Public/BaseGizmos/GizmoBoxComponent.h
@@ -46,12 +46,4 @@ private:
 	virtual FBoxSphereBounds CalcBounds(const FTransform& LocalToWorld) const override;
 	//~ Begin USceneComponent Interface.
 
-	// if true, we drew along -Direction instead of Direction, and so should hit-test accordingly
-	bool bFlippedX;
-	bool bFlippedY;
-	bool bFlippedZ;
-
-	// gizmo visibility
-	bool bRenderVisibility = true;
-
 };
\ No newline at end of file
diff --git a/Engine/Source/Runtime/Experimental/InteractiveToolsFramework/Public/BaseGizmos/GizmoCircleComponent.h b/Engine/Source/Runtime/Experimental/InteractiveToolsFramework/Public/BaseGizmos/GizmoCircleComponent.h
--- a/Engine/Source/Runtime/Experimental/InteractiveToolsFramework/Public/BaseGizmos/GizmoCircleComponent.h
+++ b/Engine/Source/Runtime/Experimental/InteractiveToolsFramework/Public/BaseGizmos/GizmoCircleComponent.h
@@ -50,10 +50,4 @@ private:
 	//~ Begin USceneComponent Interface.
 	virtual FBoxSphereBounds CalcBounds(const FTransform& LocalToWorld) const override;
 	//~ Begin USceneComponent Interface.
-
-	// dynamic gizmo visibility as set by proxy
-	bool bRenderVisibility = true;
-
-	// set to true if circle is fully visible, in which case we ignore bOnlyAllowFrontFacingHits
-	bool bCircleIsViewPlaneParallel = false;
 };
\ No newline at end of file
diff --git a/Engine/Source/Runtime/Experimental/InteractiveToolsFramework/Public/BaseGizmos/GizmoRectangleComponent.h b/Engine/Source/Runtime/Experimental/InteractiveToolsFramework/Public/BaseGizmos/GizmoRectangleComponent.h
--- a/Engine/Source/Runtime/Experimental/InteractiveToolsFramework/Public/BaseGizmos/GizmoRectangleComponent.h
+++ b/Engine/Source/Runtime/Experimental/InteractiveToolsFramework/Public/BaseGizmos/GizmoRectangleComponent.h
@@ -23,6 +23,13 @@ public:
 	UPROPERTY(EditAnywhere, Category = Options)
 	FVector DirectionY = FVector(0, 1, 0);
 
+	// When true, instead of using the provided DirectionY, the component will
+	// use a direction orthogonal to the camera direction and DirectionX. This
+	// keeps the rectangle pinned along DirectionX but spun to be flatter
+	// relative the camera.
+	UPROPERTY(EditAnywhere, Category = Options)
+	bool bOrientYAccordingToCamera = false;
+
 	UPROPERTY(EditAnywhere, Category = Options)
 	float OffsetX = 0.0f;
 
@@ -41,7 +48,6 @@ public:
 	UPROPERTY(EditAnywhere, Category = Options)
 	uint8 SegmentFlags = 0x1 | 0x2 | 0x4 | 0x8;
 
-
 private:
 	//~ Begin UPrimitiveComponent Interface.
 	virtual FPrimitiveSceneProxy* CreateSceneProxy() override;
@@ -52,12 +58,4 @@ private:
 	//~ Begin USceneComponent Interface.
 	virtual FBoxSphereBounds CalcBounds(const FTransform& LocalToWorld) const override;
 	//~ Begin USceneComponent Interface.
-
-	// if true, we drew along -Direction instead of Direction, and so should hit-test accordingly
-	bool bFlippedX;
-	bool bFlippedY;
-
-	// gizmo visibility
-	bool bRenderVisibility = true;
-
 };
\ No newline at end of file
diff --git a/Engine/Source/Runtime/Experimental/InteractiveToolsFramework/Public/BaseGizmos/GizmoRenderingUtil.h b/Engine/Source/Runtime/Experimental/InteractiveToolsFramework/Public/BaseGizmos/GizmoRenderingUtil.h
--- a/Engine/Source/Runtime/Experimental/InteractiveToolsFramework/Public/BaseGizmos/GizmoRenderingUtil.h
+++ b/Engine/Source/Runtime/Experimental/InteractiveToolsFramework/Public/BaseGizmos/GizmoRenderingUtil.h
@@ -6,40 +6,24 @@
 #include "Components/PrimitiveComponent.h"
 #include "SceneView.h"
 
+class UGizmoViewContext;
+
 /**
  * Utility functions for standard GizmoComponent rendering
  */
 namespace GizmoRenderingUtil
 {
-	/**
-	 * This is a hack. When we are inside the ::GetDynamicMeshElements() of a UGizmoBaseComponent RenderProxy, we need 
-	 * to draw to a set of FSceneView and we do not know which one is focused. UGizmoBaseComponent implementations generally
-	 * try to adapt their rendering to the current view, and some of that (currently) can only happen at the
-	 * level of the RenderProxy. So we need to tell the Component which FSceneView has focus, and we do that
-	 * by calling this function from the EdMode.
-	 */
-	INTERACTIVETOOLSFRAMEWORK_API void SetGlobalFocusedEditorSceneView(const FSceneView* View);
-
-	/**
-	 * This function searches the input Views list for the global currently-focused FSceneView that was last set in
-	 * SetGlobalFocusedEditorSceneView(). Note that the Views pointers are different every frame, and 
-	 * we are calling SetGlobalFocusedEditorSceneView() on the Game thread but this function from the Render Thread.
-	 * So it is entirely possible that we will "miss" our FSceneView and this function will return NULL.
-	 * Implementations need to handle this case.
-	 */
-	INTERACTIVETOOLSFRAMEWORK_API const FSceneView* FindFocusedEditorSceneView(
-		const TArray<const FSceneView*>& Views,
-		const FSceneViewFamily& ViewFamily,
-		uint32 VisibilityMap);
-
 
 	/**
-	 * Turn on/off the above global focused scene view tracking. When disabled, FindFocusedEditorSceneView()
-	 * will return the first visible FSceneView, or the first element of the Views array otherwise.
-	 * Enabled by default in Editor builds, disabled by default in Runtime builds.
+	 * @return Conversion factor between pixel and world-space coordinates at 3D point Location in View.
+	 * @warning This is a local estimate and is increasingly incorrect as the 3D point gets further from Location
 	 */
-	INTERACTIVETOOLSFRAMEWORK_API void SetGlobalFocusedSceneViewTrackingEnabled(bool bEnabled);
-
+	INTERACTIVETOOLSFRAMEWORK_API float CalculateLocalPixelToWorldScale(
+		const FSceneView* View,
+		const FVector& Location);
+	INTERACTIVETOOLSFRAMEWORK_API float CalculateLocalPixelToWorldScale(
+		const UGizmoViewContext* ViewContext,
+		const FVector& Location);
 
 	/**
 	 * @return Conversion factor between pixel and world-space coordinates at 3D point Location in View.
diff --git a/Engine/Source/Runtime/Experimental/InteractiveToolsFramework/Public/BaseGizmos/GizmoViewContext.h b/Engine/Source/Runtime/Experimental/InteractiveToolsFramework/Public/BaseGizmos/GizmoViewContext.h
new file mode 100644
--- /dev/null
+++ b/Engine/Source/Runtime/Experimental/InteractiveToolsFramework/Public/BaseGizmos/GizmoViewContext.h
@@ -0,0 +1,70 @@
+// Copyright Epic Games, Inc. All Rights Reserved.
+
+#pragma once
+
+#include "CoreMinimal.h"
+
+#include "SceneView.h"
+
+#include "GizmoViewContext.generated.h"
+
+/** 
+ * A context object that is meant to hold the scene information for the hovered viewport
+ * on a game thread, to be used by a gizmo later for hit testing. The organization mirrors
+ * FSceneView so that functions could be written in a templated way to use either FSceneView
+ * or UGizmoViewContext, though UGizmoViewContext only keeps the needed data.
+ */
+UCLASS()
+class INTERACTIVETOOLSFRAMEWORK_API UGizmoViewContext : public UObject
+{
+	GENERATED_BODY()
+public:
+	
+	// Wrapping class for the matrices so that they can be accessed in the same way
+	// that they are accessed in FSceneView.
+	class FMatrices
+	{
+	public:
+		void ResetFromSceneView(const FSceneView& SceneView)
+		{
+			ViewMatrix = SceneView.ViewMatrices.GetViewMatrix();
+			ViewProjectionMatrix = SceneView.ViewMatrices.GetViewProjectionMatrix();
+		}
+
+		const FMatrix& GetViewMatrix() const { return ViewMatrix; }
+		const FMatrix& GetViewProjectionMatrix() const { return ViewProjectionMatrix; }
+
+	protected:
+		FMatrix ViewMatrix;
+		FMatrix ViewProjectionMatrix;
+	};
+
+	// Use this to reinitialize the object each frame for the hovered viewport.
+	void ResetFromSceneView(const FSceneView& SceneView)
+	{
+		UnscaledViewRect = SceneView.UnscaledViewRect;
+		ViewMatrices.ResetFromSceneView(SceneView);
+		bIsPerspectiveProjection = SceneView.IsPerspectiveProjection();
+		ViewLocation = SceneView.ViewLocation;
+	}
+
+	// FSceneView-like functions/properties:
+	FVector GetViewRight() const { return ViewMatrices.GetViewMatrix().GetColumn(0); }
+	FVector GetViewUp() const { return ViewMatrices.GetViewMatrix().GetColumn(1); }
+	FVector GetViewDirection() const { return ViewMatrices.GetViewMatrix().GetColumn(2); }
+
+	// As a function just for similarity with FSceneView
+	bool IsPerspectiveProjection() const { return bIsPerspectiveProjection; }
+
+	FVector4 WorldToScreen(const FVector& WorldPoint) const
+	{
+		return ViewMatrices.GetViewProjectionMatrix().TransformFVector4(FVector4(WorldPoint, 1));
+	}
+
+	FMatrices ViewMatrices;
+	FIntRect UnscaledViewRect;
+	FVector	ViewLocation;
+
+protected:
+	bool bIsPerspectiveProjection;
+};
\ No newline at end of file
diff --git a/Engine/Source/Runtime/Experimental/InteractiveToolsFramework/Public/BaseGizmos/IntervalGizmo.h b/Engine/Source/Runtime/Experimental/InteractiveToolsFramework/Public/BaseGizmos/IntervalGizmo.h
--- a/Engine/Source/Runtime/Experimental/InteractiveToolsFramework/Public/BaseGizmos/IntervalGizmo.h
+++ b/Engine/Source/Runtime/Experimental/InteractiveToolsFramework/Public/BaseGizmos/IntervalGizmo.h
@@ -19,6 +19,7 @@ class IGizmoStateTarget;
 class UGizmoComponentAxisSource;
 class IGizmoFloatParameterSource;
 class UGizmoLocalFloatParameterSource;
+class UGizmoViewContext;
 
 
 /**
@@ -55,7 +56,8 @@ public:
 	 * Create a new instance of AIntervalGizmoActor and populate the various
 	 * sub-components with standard GizmoXComponent instances suitable for a 3-interval Gizmo
 	 */
-	static AIntervalGizmoActor* ConstructDefaultIntervalGizmo(UWorld* World);
+	static AIntervalGizmoActor* ConstructDefaultIntervalGizmo(UWorld* World, 
+		UGizmoViewContext* GizmoViewContext);
 
 };
 
@@ -74,6 +76,10 @@ public:
 class INTERACTIVETOOLSFRAMEWORK_API FIntervalGizmoActorFactory
 {
 public:
+	FIntervalGizmoActorFactory(UGizmoViewContext* GizmoViewContextIn)
+		: GizmoViewContext(GizmoViewContextIn)
+	{
+	}
 
 	/**
 	 * @param World the UWorld to create the new Actor in
@@ -81,8 +87,15 @@ public:
 	 */
 	virtual AIntervalGizmoActor* CreateNewGizmoActor(UWorld* World) const
 	{
-		return AIntervalGizmoActor::ConstructDefaultIntervalGizmo(World);
+		return AIntervalGizmoActor::ConstructDefaultIntervalGizmo(World, GizmoViewContext);
 	}
+
+protected:
+	/**
+	 * Needs to be set (and kept alive elsewhere) so that created handle gizmos can
+	 * adjust their length according to view (when not using world-scaling).
+	 */
+	UGizmoViewContext* GizmoViewContext = nullptr;
 };
 
 
diff --git a/Engine/Source/Runtime/Experimental/InteractiveToolsFramework/Public/BaseGizmos/TransformGizmo.h b/Engine/Source/Runtime/Experimental/InteractiveToolsFramework/Public/BaseGizmos/TransformGizmo.h
--- a/Engine/Source/Runtime/Experimental/InteractiveToolsFramework/Public/BaseGizmos/TransformGizmo.h
+++ b/Engine/Source/Runtime/Experimental/InteractiveToolsFramework/Public/BaseGizmos/TransformGizmo.h
@@ -19,6 +19,7 @@ class IGizmoStateTarget;
 class UGizmoConstantFrameAxisSource;
 class UGizmoComponentAxisSource;
 class UGizmoTransformChangeStateTarget;
+class UGizmoViewContext;
 class FTransformGizmoTransformChange;
 class FRay;
 
@@ -131,7 +132,7 @@ public:
 	 * sub-components with standard GizmoXComponent instances suitable for a 3-axis transformer Gizmo
 	 */
 	static ATransformGizmoActor* ConstructDefault3AxisGizmo(
-		UWorld* World
+		UWorld* World, UGizmoViewContext* GizmoViewContext
 	);
 
 	/**
@@ -139,7 +140,7 @@ public:
 	 * specified by Elements with standard GizmoXComponent instances suitable for a 3-axis transformer Gizmo
 	 */
 	static ATransformGizmoActor* ConstructCustom3AxisGizmo(
-		UWorld* World,
+		UWorld* World, UGizmoViewContext* GizmoViewContext,
 		ETransformGizmoSubElements Elements
 	);
 };
@@ -162,6 +163,11 @@ public:
 class INTERACTIVETOOLSFRAMEWORK_API FTransformGizmoActorFactory
 {
 public:
+	FTransformGizmoActorFactory(UGizmoViewContext* GizmoViewContextIn)
+		: GizmoViewContext(GizmoViewContextIn)
+	{
+	}
+
 	/** Only these members of the ATransformGizmoActor gizmo will be initialized */
 	ETransformGizmoSubElements EnableElements =
 		ETransformGizmoSubElements::TranslateAllAxes |
@@ -176,6 +182,14 @@ public:
 	 * @return new ATransformGizmoActor instance with members initialized with Components suitable for a transformation Gizmo
 	 */
 	virtual ATransformGizmoActor* CreateNewGizmoActor(UWorld* World) const;
+
+protected:
+	/**
+	 * The default gizmos that we use need to have the current view information stored for them via
+	 * the ITF context store so that they can figure out how big they are for hit testing, so this
+	 * pointer needs to be set (and kept alive elsewhere) for the actor factory to work properly.
+	 */
+	UGizmoViewContext* GizmoViewContext = nullptr;
 };
 
 
