Chaos: constraint memory leak fix; leaking FJointConstraint

diff --git a/Engine/Source/Runtime/Experimental/Chaos/Private/PhysicsProxy/JointConstraintProxy.cpp b/Engine/Source/Runtime/Experimental/Chaos/Private/PhysicsProxy/JointConstraintProxy.cpp
--- a/Engine/Source/Runtime/Experimental/Chaos/Private/PhysicsProxy/JointConstraintProxy.cpp
+++ b/Engine/Source/Runtime/Experimental/Chaos/Private/PhysicsProxy/JointConstraintProxy.cpp
@@ -104,6 +104,10 @@ void FJointConstraintPhysicsProxy::DestroyOnPhysicsThread(Chaos::FPBDRigidsSolve
 		auto& JointConstraints = InSolver->GetJointConstraints();
 		JointConstraints.RemoveConstraint(Handle->GetConstraintIndex());
 
+	}
+
+	if (Constraint != nullptr)
+	{
 		delete Constraint;
 		Constraint = nullptr;
 	}
