Chaos: update log value to incoming value instead of old value.

diff --git a/Engine/Source/Runtime/Experimental/Chaos/Private/Chaos/Framework/PhysicsSolverBase.cpp b/Engine/Source/Runtime/Experimental/Chaos/Private/Chaos/Framework/PhysicsSolverBase.cpp
--- a/Engine/Source/Runtime/Experimental/Chaos/Private/Chaos/Framework/PhysicsSolverBase.cpp
+++ b/Engine/Source/Runtime/Experimental/Chaos/Private/Chaos/Framework/PhysicsSolverBase.cpp
@@ -154,7 +154,7 @@ namespace Chaos
 		if (AsyncDt != FixedDt)
 		{
 			AccumulatedTime = 0;
-			UE_LOG(LogChaos, Log, TEXT("FPhysicsSolverBase::AsyncDt:%f"), AsyncDt);
+			UE_LOG(LogChaos, Log, TEXT("FPhysicsSolverBase::AsyncDt:%f"), FixedDt);
 		}
 		AsyncDt = FixedDt;
 	}
