Accelerate Dirty Element List with the help of a 2D hash grid

#[fyi] Michael.Lentine, Benn.Gallagher, Chris.Caulfield

diff --git a/Engine/Source/Programs/HeadlessChaos/Private/HeadlessChaosApp.cpp b/Engine/Source/Programs/HeadlessChaos/Private/HeadlessChaosApp.cpp
--- a/Engine/Source/Programs/HeadlessChaos/Private/HeadlessChaosApp.cpp
+++ b/Engine/Source/Programs/HeadlessChaos/Private/HeadlessChaosApp.cpp
@@ -231,6 +231,7 @@ TEST(BP, BroadphaseTests) {
 	ChaosTest::GridBPTest();
 	ChaosTest::GridBPTest2();
 	ChaosTest::AABBTreeTest();
+	ChaosTest::AABBTreeDirtyGridTest();
 	ChaosTest::AABBTreeTimesliceTest();
 	ChaosTest::BoundingVolumeNoBoundsTest();
 	ChaosTest::BroadphaseCollectionTest();
diff --git a/Engine/Source/Programs/HeadlessChaos/Private/HeadlessChaosTestBroadphase.cpp b/Engine/Source/Programs/HeadlessChaos/Private/HeadlessChaosTestBroadphase.cpp
--- a/Engine/Source/Programs/HeadlessChaos/Private/HeadlessChaosTestBroadphase.cpp
+++ b/Engine/Source/Programs/HeadlessChaos/Private/HeadlessChaosTestBroadphase.cpp
@@ -222,7 +222,7 @@ namespace ChaosTest
 	};
 
 
-	auto BuildBoxes(TUniquePtr<TBox<FReal,3>>& Box, FReal BoxSize = 100, const FVec3& BoxGridDimensions = FVec3(10,10,10))
+	auto BuildBoxes(TUniquePtr<TBox<FReal,3>>& Box, FReal BoxSize = 100, const FVec3& BoxGridDimensions = FVec3(10,10,10), const FVec3 Offset = FVec3(0, 0, 0))
 	{
 		Box = MakeUnique<TBox<FReal, 3>>(FVec3(0, 0, 0), FVec3(BoxSize, BoxSize, BoxSize));
 		auto Boxes = MakeUnique<FGeometryParticles>();
@@ -240,7 +240,7 @@ namespace ChaosTest
 				for (int32 Col = 0; Col < NumCols; ++Col)
 				{
 					Boxes->SetGeometry(Idx, MakeSerializable(Box));
-					Boxes->X(Idx) = FVec3(Col * 100, Row * 100, Height * 100);
+					Boxes->X(Idx) = FVec3(Col * 100, Row * 100, Height * 100) + Offset;
 					Boxes->R(Idx) = FRotation3::Identity;
 					Boxes->LocalBounds(Idx) = Box->BoundingBox();
 					Boxes->HasBounds(Idx) = true;
@@ -554,6 +554,122 @@ namespace ChaosTest
 			EXPECT_EQ(Spatial.NumDirtyElements(),Boxes->Size() - 1);
 		}
 	}
+
+
+	void AABBTreeDirtyGridTest()
+	{
+		using TreeType = TAABBTree<int32, TBoundingVolume<int32>>;
+
+		// Save CVARS
+		int32 DirtyElementGridCellSize = FAABBTreeDirtyGridCVars::DirtyElementGridCellSize;
+		int32 DirtyElementMaxGridCellQueryCount = FAABBTreeDirtyGridCVars::DirtyElementMaxGridCellQueryCount;
+		int32 DirtyElementMaxPhysicalSizeInCells = FAABBTreeDirtyGridCVars::DirtyElementMaxPhysicalSizeInCells;
+		int32 DirtyElementMaxCellCapacity = FAABBTreeDirtyGridCVars::DirtyElementMaxCellCapacity;
+
+		// Set CVARS to known values
+		FAABBTreeDirtyGridCVars::DirtyElementGridCellSize = 100;
+		FAABBTreeDirtyGridCVars::DirtyElementMaxGridCellQueryCount = 10000;
+		FAABBTreeDirtyGridCVars::DirtyElementMaxPhysicalSizeInCells = 20;
+		FAABBTreeDirtyGridCVars::DirtyElementMaxCellCapacity = 20;
+
+
+		// Do the standard tests
+		{
+			TUniquePtr<TBox<FReal, 3>> Box;
+			auto Boxes = BuildBoxes(Box);
+			TreeType Spatial{};
+
+			int32 Idx;
+			for (Idx = 0; Idx < (int32)Boxes->Size(); ++Idx)
+			{
+				Spatial.UpdateElement(Idx, Boxes->WorldSpaceInflatedBounds(Idx), true);
+			}
+			EXPECT_EQ(Spatial.NumDirtyElements(), Boxes->Size());
+
+			SpatialTestHelper(Spatial, Boxes.Get(), Box);
+		}
+
+
+		// Repeat the standard tests with low cell capacity and different cell sizes
+		{
+			// Set CVARS to known values
+			FAABBTreeDirtyGridCVars::DirtyElementGridCellSize = 44;
+			FAABBTreeDirtyGridCVars::DirtyElementMaxCellCapacity = 2;
+
+			TUniquePtr<TBox<FReal, 3>> Box;
+			auto Boxes = BuildBoxes(Box);
+			TreeType Spatial{};
+
+			int32 Idx;
+			for (Idx = 0; Idx < (int32)Boxes->Size(); ++Idx)
+			{
+				Spatial.UpdateElement(Idx, Boxes->WorldSpaceInflatedBounds(Idx), true);
+			}
+			EXPECT_EQ(Spatial.NumDirtyElements(), Boxes->Size());
+
+			SpatialTestHelper(Spatial, Boxes.Get(), Box);
+		}
+
+
+		// Make sure we get the same results, with and without the grid for sweeps and raycasts
+		{
+			FAABBTreeDirtyGridCVars::DirtyElementMaxCellCapacity = 7;
+			TUniquePtr<TBox<FReal, 3>> Box;
+			auto Boxes = BuildBoxes(Box, 100, FVec3(40, 40, 1), FVec3(-2000, -2000, -50));
+
+			for (float Angle = 0.0; Angle < 2 * PI; Angle += (10.0f / 360.0f) * 2.0f * PI)
+			{
+
+				FVec3 Direction{ FMath::Cos(Angle), FMath::Sin(Angle), 0 };
+				// With the grid
+				FVisitor VisitorGrid(FVec3(53, 27, 0), Direction, 0, *Boxes);
+				VisitorGrid.HalfExtents = FVec3(102, 20, 2);
+				{
+					FAABBTreeDirtyGridCVars::DirtyElementGridCellSize = 100;
+
+					TreeType Spatial{};
+
+					int32 Idx;
+					for (Idx = 0; Idx < (int32)Boxes->Size(); ++Idx)
+					{
+						Spatial.UpdateElement(Idx, Boxes->WorldSpaceInflatedBounds(Idx), true);
+					}
+					EXPECT_EQ(Spatial.NumDirtyElements(), Boxes->Size());
+
+					Spatial.Raycast(VisitorGrid.Start, VisitorGrid.Dir, 1900, VisitorGrid);
+					Spatial.Sweep(VisitorGrid.Start, VisitorGrid.Dir, 1800, VisitorGrid.HalfExtents, VisitorGrid);
+				}
+
+				// Without the grid
+				FVisitor VisitorNoGrid(FVec3(53, 27, 0), Direction, 0, *Boxes);
+				VisitorNoGrid.HalfExtents = FVec3(102, 20, 2);
+				{
+					FAABBTreeDirtyGridCVars::DirtyElementGridCellSize = 0;
+					TreeType Spatial{};
+
+					int32 Idx;
+					for (Idx = 0; Idx < (int32)Boxes->Size(); ++Idx)
+					{
+						Spatial.UpdateElement(Idx, Boxes->WorldSpaceInflatedBounds(Idx), true);
+					}
+					EXPECT_EQ(Spatial.NumDirtyElements(), Boxes->Size());
+
+					Spatial.Raycast(VisitorNoGrid.Start, VisitorNoGrid.Dir, 1900, VisitorNoGrid);
+					Spatial.Sweep(VisitorNoGrid.Start, VisitorNoGrid.Dir, 1800, VisitorNoGrid.HalfExtents, VisitorNoGrid);
+				}
+				// These will be in the same order, but we can drop this requirement in the future
+				EXPECT_TRUE(VisitorNoGrid.Instances == VisitorGrid.Instances);  
+			}
+			
+		}
+		
+
+		// Restore CVARS
+		FAABBTreeDirtyGridCVars::DirtyElementGridCellSize = DirtyElementGridCellSize;
+		FAABBTreeDirtyGridCVars::DirtyElementMaxGridCellQueryCount = DirtyElementMaxGridCellQueryCount;
+		FAABBTreeDirtyGridCVars::DirtyElementMaxPhysicalSizeInCells = DirtyElementMaxPhysicalSizeInCells;
+		FAABBTreeDirtyGridCVars::DirtyElementMaxCellCapacity = DirtyElementMaxCellCapacity;
+	}
 	
 
 	void AABBTreeTimesliceTest()
diff --git a/Engine/Source/Programs/HeadlessChaos/Public/HeadlessChaosTestBroadphase.h b/Engine/Source/Programs/HeadlessChaos/Public/HeadlessChaosTestBroadphase.h
--- a/Engine/Source/Programs/HeadlessChaos/Public/HeadlessChaosTestBroadphase.h
+++ b/Engine/Source/Programs/HeadlessChaos/Public/HeadlessChaosTestBroadphase.h
@@ -13,6 +13,8 @@ namespace ChaosTest {
 
 	void AABBTreeTest();
 
+	void AABBTreeDirtyGridTest();
+
 	void AABBTreeTimesliceTest();
 
 	void BroadphaseCollectionTest();
diff --git a/Engine/Source/Runtime/Experimental/Chaos/Private/Chaos/AABBTree.cpp b/Engine/Source/Runtime/Experimental/Chaos/Private/Chaos/AABBTree.cpp
--- a/Engine/Source/Runtime/Experimental/Chaos/Private/Chaos/AABBTree.cpp
+++ b/Engine/Source/Runtime/Experimental/Chaos/Private/Chaos/AABBTree.cpp
@@ -4,3 +4,18 @@
 
 int32 FAABBTreeCVars::UpdateDirtyElementPayloadData = 1;
 FAutoConsoleVariableRef FAABBTreeCVars::CVarUpdateDirtyElementPayloadData(TEXT("p.aabbtree.updatedirtyelementpayloads"), FAABBTreeCVars::UpdateDirtyElementPayloadData, TEXT("Allow AABB tree elements to update internal payload data when they recieve a payload update"));
+
+int32 FAABBTreeDirtyGridCVars::DirtyElementGridCellSize = 0; //1000; 0 means disabled
+FAutoConsoleVariableRef FAABBTreeDirtyGridCVars::CVarDirtyElementGridCellSize(TEXT("p.aabbtree.DirtyElementGridCellSize"), FAABBTreeDirtyGridCVars::DirtyElementGridCellSize, TEXT("DirtyElement Grid acceleration structure cell size in cm. 0 or less will disable the feature"));
+
+int32 FAABBTreeDirtyGridCVars::DirtyElementMaxGridCellQueryCount = 340;
+FAutoConsoleVariableRef FAABBTreeDirtyGridCVars::CVarDirtyElementMaxGridCellQueryCount(TEXT("p.aabbtree.DirtyElementMaxGridCellQueryCount"), FAABBTreeDirtyGridCVars::DirtyElementMaxGridCellQueryCount, TEXT("Maximum grid cells to query (per raycast for example) in DirtyElement grid acceleration structure before falling back to brute force"));
+
+int32 FAABBTreeDirtyGridCVars::DirtyElementMaxPhysicalSizeInCells = 16;
+FAutoConsoleVariableRef FAABBTreeDirtyGridCVars::CVarDirtyElementMaxPhysicalSizeInCells(TEXT("p.aabbtree.DirtyElementMaxPhysicalSizeInCells"), FAABBTreeDirtyGridCVars::DirtyElementMaxPhysicalSizeInCells, TEXT("If a dirty element stradles more than this number of cells, it will no be added to the grid acceleration structure"));
+
+int32 FAABBTreeDirtyGridCVars::DirtyElementMaxCellCapacity = 32;
+FAutoConsoleVariableRef FAABBTreeDirtyGridCVars::CVarDirtyElementMaxCellCapacity(TEXT("p.aabbtree.DirtyElementMaxCellCapacity"), FAABBTreeDirtyGridCVars::DirtyElementMaxCellCapacity, TEXT("The maximum number of dirty elements that can be added to a single grid cell before spilling to slower flat list"));
+
+CSV_DEFINE_CATEGORY(ChaosPhysicsTimers, true);
+
diff --git a/Engine/Source/Runtime/Experimental/Chaos/Public/Chaos/AABBTree.h b/Engine/Source/Runtime/Experimental/Chaos/Public/Chaos/AABBTree.h
--- a/Engine/Source/Runtime/Experimental/Chaos/Public/Chaos/AABBTree.h
+++ b/Engine/Source/Runtime/Experimental/Chaos/Public/Chaos/AABBTree.h
@@ -2,6 +2,7 @@
 #pragma once
 
 #include "Chaos/AABB.h"
+#include "Chaos/AABBTreeDirtyGridUtils.h"
 #include "Chaos/Defines.h"
 #include "Chaos/GeometryParticles.h"
 #include "Chaos/Transform.h"
@@ -9,6 +10,9 @@
 #include "Chaos/ISpatialAcceleration.h"
 #include "Templates/Models.h"
 #include "Chaos/BoundingVolume.h"
+#include "ProfilingDebugging/CsvProfiler.h"
+
+CSV_DECLARE_CATEGORY_EXTERN(ChaosPhysicsTimers);
 
 struct FAABBTreeCVars
 {
@@ -16,6 +20,21 @@ struct FAABBTreeCVars
 	static FAutoConsoleVariableRef CVarUpdateDirtyElementPayloadData;
 };
 
+struct CHAOS_API FAABBTreeDirtyGridCVars
+{
+	static int32 DirtyElementGridCellSize;
+	static FAutoConsoleVariableRef CVarDirtyElementGridCellSize;
+
+	static int32 DirtyElementMaxGridCellQueryCount;
+	static FAutoConsoleVariableRef CVarDirtyElementMaxGridCellQueryCount;
+
+	static int32 DirtyElementMaxPhysicalSizeInCells;
+	static FAutoConsoleVariableRef CVarDirtyElementMaxPhysicalSizeInCells;
+
+	static int32 DirtyElementMaxCellCapacity;
+	static FAutoConsoleVariableRef CVarDirtyElementMaxCellCapacity;
+};
+
 namespace Chaos
 {
 
@@ -285,11 +304,13 @@ struct FAABBTreePayloadInfo
 	int32 GlobalPayloadIdx;
 	int32 DirtyPayloadIdx;
 	int32 LeafIdx;
+	int32 DirtyGridOverflowIdx;
 
-	FAABBTreePayloadInfo(int32 InGlobalPayloadIdx = INDEX_NONE, int32 InDirtyIdx = INDEX_NONE, int32 InLeafIdx = INDEX_NONE)
+	FAABBTreePayloadInfo(int32 InGlobalPayloadIdx = INDEX_NONE, int32 InDirtyIdx = INDEX_NONE, int32 InLeafIdx = INDEX_NONE, int32 InDirtyGridOverflowIdx = INDEX_NONE)
 		: GlobalPayloadIdx(InGlobalPayloadIdx)
 		, DirtyPayloadIdx(InDirtyIdx)
 		, LeafIdx(InLeafIdx)
+		, DirtyGridOverflowIdx(InDirtyGridOverflowIdx)
 	{}
 
 	void Serialize(FArchive& Ar)
@@ -297,6 +318,7 @@ struct FAABBTreePayloadInfo
 		Ar << GlobalPayloadIdx;
 		Ar << DirtyPayloadIdx;
 		Ar << LeafIdx;
+		Ar << DirtyGridOverflowIdx;
 	}
 };
 
@@ -329,6 +351,24 @@ static void UpdateElementHelper(T& InElem, const T& InFrom)
 	}
 }
 
+struct DirtyGridHashEntry
+{
+	DirtyGridHashEntry()
+	{
+		Index = 0;
+		Count = 0;
+	}
+
+	DirtyGridHashEntry(const DirtyGridHashEntry& Other)
+	{
+		Index = Other.Index;
+		Count = Other.Count;
+	}
+
+	int32 Index;  // Index into FlattenedCellArrayOfDirtyIndices
+	int32 Count;  // Number of valid entries from Index in FlattenedCellArrayOfDirtyIndices
+};
+
 template <typename TPayloadType, typename TLeafType, bool bMutable = true>
 class TAABBTree final : public ISpatialAcceleration<TPayloadType, FReal, 3>
 {
@@ -349,6 +389,7 @@ public:
 		, MaxPayloadBounds(DefaultMaxPayloadBounds)
 		, MaxNumToProcess(DefaultMaxNumToProcess)
 	{
+		GetCVars();
 	}
 
 	virtual void Reset() override
@@ -356,6 +397,9 @@ public:
 		Nodes.Reset();
 		Leaves.Reset();
 		DirtyElements.Reset();
+		CellHashToFlatArray.Reset();
+		FlattenedCellArrayOfDirtyIndices.Reset();
+		DirtyElementsGridOverflow.Reset();
 		GlobalPayloads.Reset();
 		PayloadToInfo.Reset();
 
@@ -487,6 +531,194 @@ public:
 		return QueryImp<EAABBQueryType::Overlap>(FVec3(), VoidData, FVec3(), QueryBounds, Visitor, VoidData.Dir, VoidData.InvDir, VoidData.bParallel);
 	}
 
+	// This is to make sure important parameters are not changed during inopportune times
+	void GetCVars()
+	{
+		DirtyElementGridCellSize = (FReal) FAABBTreeDirtyGridCVars::DirtyElementGridCellSize;
+		if (DirtyElementGridCellSize > SMALL_NUMBER)
+		{
+			DirtyElementGridCellSizeInv = 1.0f / DirtyElementGridCellSize;
+		}
+		else
+		{
+			DirtyElementGridCellSizeInv = 1.0f;
+		}
+
+		DirtyElementMaxGridCellQueryCount = FAABBTreeDirtyGridCVars::DirtyElementMaxGridCellQueryCount;
+		DirtyElementMaxPhysicalSizeInCells = FAABBTreeDirtyGridCVars::DirtyElementMaxPhysicalSizeInCells;
+		DirtyElementMaxCellCapacity = FAABBTreeDirtyGridCVars::DirtyElementMaxCellCapacity;
+	}
+
+	FORCEINLINE_DEBUGGABLE bool DirtyElementGridEnabled() const
+	{
+		return DirtyElementGridCellSize > 0.0f &&
+			DirtyElementMaxGridCellQueryCount > 0 &&
+			DirtyElementMaxPhysicalSizeInCells > 0 &&
+			DirtyElementMaxCellCapacity > 0;
+	}
+
+	FORCEINLINE_DEBUGGABLE bool EnoughSpaceInGridCell(int32 Hash)
+	{
+		DirtyGridHashEntry* HashEntry = CellHashToFlatArray.Find(Hash);
+		if (HashEntry)
+		{
+			if (HashEntry->Count >= DirtyElementMaxCellCapacity) // Checking if we are at capacity
+			{
+				return false;
+			}
+		}
+
+		return true;
+	}
+
+	// Returns true if there was enough space in the cell to add the new dirty element index and if the element is not already in the cell 
+	//(The second condition should never be true for the current implementation)
+	FORCEINLINE_DEBUGGABLE bool AddNewDirtyParticleIndexToGridCell(int32 Hash, int32 NewDirtyIndex)
+	{
+		DirtyGridHashEntry* HashEntry = CellHashToFlatArray.Find(Hash);
+		if (HashEntry)
+		{
+			if (HashEntry->Count < DirtyElementMaxCellCapacity && ensure(InsertValueIntoSortedSubArray(FlattenedCellArrayOfDirtyIndices, NewDirtyIndex, HashEntry->Index, HashEntry->Count)))
+			{
+				++(HashEntry->Count);
+				return true;
+			}
+		}
+		else
+		{
+			DirtyGridHashEntry& NewHashEntry = CellHashToFlatArray.Add(Hash);
+			NewHashEntry.Index = FlattenedCellArrayOfDirtyIndices.Num(); // End of flat array
+			NewHashEntry.Count = 1;
+			FlattenedCellArrayOfDirtyIndices.AddUninitialized(DirtyElementMaxCellCapacity);
+			FlattenedCellArrayOfDirtyIndices[NewHashEntry.Index] = NewDirtyIndex;
+			return true;
+		}
+		return false;
+	}
+
+	// Returns true if the dirty particle was in the grid and successfully deleted
+	FORCEINLINE_DEBUGGABLE bool DeleteDirtyParticleIndexFromGridCell(int32 Hash, int32 DirtyIndex)
+	{
+		DirtyGridHashEntry* HashEntry = CellHashToFlatArray.Find(Hash);
+		if (HashEntry && HashEntry->Count >= 1)
+		{
+			if (DeleteValueFromSortedSubArray(FlattenedCellArrayOfDirtyIndices, DirtyIndex, HashEntry->Index, HashEntry->Count))
+			{
+				--(HashEntry->Count);
+				// Not deleting cell when it gets empty, it may get reused or will be deleted when the AABBTree is rebuilt
+				return true;
+			}
+		}
+		return false;
+	}
+
+	FORCEINLINE_DEBUGGABLE void DeleteDirtyParticleEverywhere(int32 DeleteDirtyParticleIdx, int32 DeleteDirtyGridOverflowIdx)
+	{
+		if (DeleteDirtyGridOverflowIdx == INDEX_NONE)
+		{
+			// Remove this element from the Grid
+			DoForOverlappedCells(DirtyElements[DeleteDirtyParticleIdx].Bounds, DirtyElementGridCellSize, DirtyElementGridCellSizeInv, [&](int32 Hash) {
+				ensure(DeleteDirtyParticleIndexFromGridCell(Hash, DeleteDirtyParticleIdx));
+				});
+		}
+		else
+		{
+			// remove element from the grid overflow
+			ensure(DirtyElementsGridOverflow[DeleteDirtyGridOverflowIdx] == DeleteDirtyParticleIdx);
+
+			if (DeleteDirtyGridOverflowIdx + 1 < DirtyElementsGridOverflow.Num())
+			{
+				auto LastOverflowPayload = DirtyElements[DirtyElementsGridOverflow.Last()].Payload;
+				PayloadToInfo.FindChecked(LastOverflowPayload).DirtyGridOverflowIdx = DeleteDirtyGridOverflowIdx;
+			}
+			DirtyElementsGridOverflow.RemoveAtSwap(DeleteDirtyGridOverflowIdx);
+		}
+
+		if (DeleteDirtyParticleIdx + 1 < DirtyElements.Num())
+		{
+			// Now rename the last element in DirtyElements in both the grid and the overflow
+			// So that it is correct after swapping Dirty elements in next step
+			int32 LastDirtyElementIndex = DirtyElements.Num() - 1;
+			auto LastDirtyPayload = DirtyElements[LastDirtyElementIndex].Payload;
+			int32 LastDirtyGridOverflowIdx = PayloadToInfo.FindChecked(LastDirtyPayload).DirtyGridOverflowIdx;
+			if (LastDirtyGridOverflowIdx == INDEX_NONE)
+			{
+				// Rename this element in the Grid
+				DoForOverlappedCells(DirtyElements.Last().Bounds, DirtyElementGridCellSize, DirtyElementGridCellSizeInv, [&](int32 Hash) {
+					ensure(DeleteDirtyParticleIndexFromGridCell(Hash, LastDirtyElementIndex));
+					ensure(AddNewDirtyParticleIndexToGridCell(Hash, DeleteDirtyParticleIdx));
+					});
+			}
+			else
+			{
+				// Rename element in overflow instead
+				DirtyElementsGridOverflow[LastDirtyGridOverflowIdx] = DeleteDirtyParticleIdx;
+			}
+
+			// Copy the Payload to the new index
+			
+			PayloadToInfo.FindChecked(LastDirtyPayload).DirtyPayloadIdx = DeleteDirtyParticleIdx;
+		}
+		DirtyElements.RemoveAtSwap(DeleteDirtyParticleIdx);
+	}
+
+	FORCEINLINE_DEBUGGABLE int32 AddDirtyElementToGrid(const TAABB<FReal, 3>& NewBounds, int32 NewDirtyElement)
+	{
+		bool bAddToGrid = !TooManyOverlapQueryCells(NewBounds, DirtyElementGridCellSizeInv, DirtyElementMaxPhysicalSizeInCells);
+		if (bAddToGrid)
+		{
+			DoForOverlappedCells(NewBounds, DirtyElementGridCellSize, DirtyElementGridCellSizeInv, [&](int32 Hash) {
+				if (!EnoughSpaceInGridCell(Hash))
+				{
+					bAddToGrid = false;
+				}
+				});
+		}
+
+		if (bAddToGrid)
+		{
+			DoForOverlappedCells(NewBounds, DirtyElementGridCellSize, DirtyElementGridCellSizeInv, [&](int32 Hash) {
+				ensure(AddNewDirtyParticleIndexToGridCell(Hash, NewDirtyElement));
+				});
+		}
+		else
+		{
+			int32 NewOverflowIndex = DirtyElementsGridOverflow.Add(NewDirtyElement);
+			return NewOverflowIndex;
+		}
+
+		return INDEX_NONE;
+	}
+
+	FORCEINLINE_DEBUGGABLE int32 UpdateDirtyElementInGrid(const TAABB<FReal, 3>& NewBounds, int32 DirtyElementIndex, int32 DirtyGridOverflowIdx)
+	{
+		if (DirtyGridOverflowIdx == INDEX_NONE)
+		{
+			const TAABB<FReal, 3>& OldBounds = DirtyElements[DirtyElementIndex].Bounds;
+
+			// Delete element in cells that are no longer overlapping
+			DoForOverlappedCellsExclude(OldBounds, NewBounds, DirtyElementGridCellSize, DirtyElementGridCellSizeInv, [&](int32 Hash) -> bool {
+				ensure(DeleteDirtyParticleIndexFromGridCell(Hash, DirtyElementIndex));
+				return true;
+				});
+
+			// Add to new overlapped cells
+			if (!DoForOverlappedCellsExclude(NewBounds, OldBounds, DirtyElementGridCellSize, DirtyElementGridCellSizeInv, [&](int32 Hash) -> bool {
+					return AddNewDirtyParticleIndexToGridCell(Hash, DirtyElementIndex);
+				}))
+			{
+				// Was not able to add it to the grid , so delete element from grid
+				DoForOverlappedCells(NewBounds, DirtyElementGridCellSize, DirtyElementGridCellSizeInv, [&](int32 Hash) {
+						DeleteDirtyParticleIndexFromGridCell(Hash, DirtyElementIndex);
+					});
+				// Add to overflow
+				int32 NewOverflowIndex = DirtyElementsGridOverflow.Add(DirtyElementIndex);
+				return NewOverflowIndex;
+			}
+		}
+		return DirtyGridOverflowIdx;
+	}
+
 	virtual void RemoveElement(const TPayloadType& Payload)
 	{
 		if (ensure(bMutable))
@@ -496,6 +728,7 @@ public:
 				if (PayloadInfo->GlobalPayloadIdx != INDEX_NONE)
 				{
 					ensure(PayloadInfo->DirtyPayloadIdx == INDEX_NONE);
+					ensure(PayloadInfo->DirtyGridOverflowIdx == INDEX_NONE);
 					ensure(PayloadInfo->LeafIdx == INDEX_NONE);
 					if (PayloadInfo->GlobalPayloadIdx + 1 < GlobalPayloads.Num())
 					{
@@ -506,12 +739,19 @@ public:
 				}
 				else if (PayloadInfo->DirtyPayloadIdx != INDEX_NONE)
 				{
-					if (PayloadInfo->DirtyPayloadIdx + 1 < DirtyElements.Num())
+					if (DirtyElementGridEnabled())
 					{
-						auto LastDirtyPayload = DirtyElements.Last().Payload;
-						PayloadToInfo.FindChecked(LastDirtyPayload).DirtyPayloadIdx = PayloadInfo->DirtyPayloadIdx;
+						DeleteDirtyParticleEverywhere(PayloadInfo->DirtyPayloadIdx, PayloadInfo->DirtyGridOverflowIdx);
+					}
+					else
+					{
+						if (PayloadInfo->DirtyPayloadIdx + 1 < DirtyElements.Num())
+						{
+							auto LastDirtyPayload = DirtyElements.Last().Payload;
+							PayloadToInfo.FindChecked(LastDirtyPayload).DirtyPayloadIdx = PayloadInfo->DirtyPayloadIdx;
+						}
+						DirtyElements.RemoveAtSwap(PayloadInfo->DirtyPayloadIdx);
 					}
-					DirtyElements.RemoveAtSwap(PayloadInfo->DirtyPayloadIdx);
 				}
 				else if (ensure(PayloadInfo->LeafIdx != INDEX_NONE))
 				{
@@ -525,6 +765,7 @@ public:
 
 	virtual void UpdateElement(const TPayloadType& Payload, const FAABB3& NewBounds, bool bHasBounds) override
 	{
+		//CSV_SCOPED_TIMING_STAT(ChaosPhysicsTimers, AABBTreeUpdateElement)
 		if (ensure(bMutable))
 		{
 			FAABBTreePayloadInfo* PayloadInfo = PayloadToInfo.Find(Payload);
@@ -568,11 +809,22 @@ public:
 				if (PayloadInfo->DirtyPayloadIdx == INDEX_NONE)
 				{
 					PayloadInfo->DirtyPayloadIdx = DirtyElements.Add(FElement{ Payload, NewBounds });
+					if (DirtyElementGridEnabled())
+					{
+						//CSV_SCOPED_TIMING_STAT(ChaosPhysicsTimers, AABBAddElement)
+						PayloadInfo->DirtyGridOverflowIdx = AddDirtyElementToGrid(NewBounds, PayloadInfo->DirtyPayloadIdx);
+					}
 				}
 				else
 				{
-					DirtyElements[PayloadInfo->DirtyPayloadIdx].Bounds = NewBounds;
-					UpdateElementHelper(DirtyElements[PayloadInfo->DirtyPayloadIdx].Payload, Payload);
+					const int32 DirtyElementIndex = PayloadInfo->DirtyPayloadIdx;
+					if (DirtyElementGridEnabled())
+					{
+						//CSV_SCOPED_TIMING_STAT(ChaosPhysicsTimers, AABBUpElement)
+						PayloadInfo->DirtyGridOverflowIdx = UpdateDirtyElementInGrid(NewBounds, DirtyElementIndex, PayloadInfo->DirtyGridOverflowIdx);
+					}
+					DirtyElements[DirtyElementIndex].Bounds = NewBounds;
+					UpdateElementHelper(DirtyElements[DirtyElementIndex].Payload, Payload);
 				}
 
 				// Handle something that previously did not have bounds that may be in global elements.
@@ -604,14 +856,22 @@ public:
 				// Handle something that previously had bounds that may be in dirty elements.
 				if (PayloadInfo->DirtyPayloadIdx != INDEX_NONE)
 				{
-					if (PayloadInfo->DirtyPayloadIdx + 1 < DirtyElements.Num())
+					if (DirtyElementGridEnabled())
+					{
+						DeleteDirtyParticleEverywhere(PayloadInfo->DirtyPayloadIdx, PayloadInfo->DirtyGridOverflowIdx);
+					}
+					else
 					{
-						auto LastDirtyPayload = DirtyElements.Last().Payload;
-						PayloadToInfo.FindChecked(LastDirtyPayload).DirtyPayloadIdx = PayloadInfo->DirtyPayloadIdx;
+						if (PayloadInfo->DirtyPayloadIdx + 1 < DirtyElements.Num())
+						{
+							auto LastDirtyPayload = DirtyElements.Last().Payload;
+							PayloadToInfo.FindChecked(LastDirtyPayload).DirtyPayloadIdx = PayloadInfo->DirtyPayloadIdx;
+						}
+						DirtyElements.RemoveAtSwap(PayloadInfo->DirtyPayloadIdx);
 					}
-					DirtyElements.RemoveAtSwap(PayloadInfo->DirtyPayloadIdx);
 
 					PayloadInfo->DirtyPayloadIdx = INDEX_NONE;
+					PayloadInfo->DirtyGridOverflowIdx = INDEX_NONE;
 				}
 			}
 		}
@@ -671,6 +931,13 @@ public:
 		Ar << MaxChildrenInLeaf;
 		Ar << MaxTreeDepth;
 		Ar << MaxPayloadBounds;
+
+		if (Ar.IsLoading())
+		{
+			// Disable the Grid until it is rebuilt
+			DirtyElementGridCellSize = 0.0f;
+			DirtyElementGridCellSizeInv = 1.0f;
+		}
 	}
 
 private:
@@ -703,6 +970,85 @@ private:
 		*this = NewTree;
 	}
 
+	// Returns true if the query should continue
+	// Execute a function for all cells found in a query as well as the overflow 
+	template <typename FunctionType>
+	bool DoForHitGridCellsAndOverflow(TArray<DirtyGridHashEntry>& HashEntryForOverlappedCells, FunctionType Function) const
+	{
+
+		// Now merge and iterate the lists of elements found in the overlapping cells
+		bool DoneWithGridElements = false;
+		bool DoneWithNonGridElements = false;
+		int NonGridElementIter = 0;
+		while (!DoneWithGridElements || !DoneWithNonGridElements)
+		{
+			// Get the next dirty element index
+
+			int32 SmallestDirtyParticleIndex = INT_MAX; // Best dirty particle index to find
+
+			if (!DoneWithGridElements)
+			{
+				// Find the next smallest index 
+				// This will start slowing down if we are overlapping a lot of cells
+				DoneWithGridElements = true;
+				for (const DirtyGridHashEntry& HashEntry : HashEntryForOverlappedCells)
+				{
+					int32 Count = HashEntry.Count;
+					if (Count > 0)
+					{
+						int32 DirtyParticleIndex = FlattenedCellArrayOfDirtyIndices[HashEntry.Index];
+						if (DirtyParticleIndex < SmallestDirtyParticleIndex)
+						{
+							SmallestDirtyParticleIndex = DirtyParticleIndex;
+							DoneWithGridElements = false;
+						}
+					}
+				}
+			}
+
+			// Now skip all elements with the same best index
+			if (!DoneWithGridElements)
+			{
+				for (DirtyGridHashEntry& HashEntry : HashEntryForOverlappedCells)
+				{
+					int32 Count = HashEntry.Count;
+					if (Count > 0)
+					{
+						int32 DirtyParticleIndex = FlattenedCellArrayOfDirtyIndices[HashEntry.Index];
+						if (DirtyParticleIndex == SmallestDirtyParticleIndex)
+						{
+							++HashEntry.Index; // Increment Index
+							--HashEntry.Count; // Decrement count
+						}
+					}
+				}
+			}
+
+			DoneWithNonGridElements = NonGridElementIter >= DirtyElementsGridOverflow.Num();
+			if (DoneWithGridElements && !DoneWithNonGridElements)
+			{
+				SmallestDirtyParticleIndex = DirtyElementsGridOverflow[NonGridElementIter];
+				++NonGridElementIter;
+			}
+
+			// Elements that are in the overflow should not also be in the grid
+			ensure(DoneWithGridElements || PayloadToInfo.Find(DirtyElements[SmallestDirtyParticleIndex].Payload)->DirtyGridOverflowIdx == INDEX_NONE);
+
+			if ((!DoneWithGridElements || !DoneWithNonGridElements))
+			{
+				const int32 Index = SmallestDirtyParticleIndex;
+				const auto& Elem = DirtyElements[Index];
+
+				if (!Function(Elem))
+				{
+					return false;
+				}
+			}
+		}
+		return true;
+	}
+	
+
 	template <EAABBQueryType Query, typename TQueryFastData, typename SQVisitor>
 	bool QueryImp(const FVec3& RESTRICT Start, TQueryFastData& CurData, const FVec3 QueryHalfExtents, const FAABB3& QueryBounds, SQVisitor& Visitor, const FVec3& Dir, const FVec3 InvDir, const bool bParallel[3]) const
 	{
@@ -710,9 +1056,7 @@ private:
 		FVec3 TmpPosition;
 		FReal TOI = 0;
 		const void* QueryData = Visitor.GetQueryData();
-
 		{
-
 			//QUICK_SCOPE_CYCLE_COUNTER(QueryGlobal);
 
 			for(const auto& Elem : GlobalPayloads)
@@ -744,15 +1088,13 @@ private:
 		}
 
 		if (bMutable)
-		{
-			//QUICK_SCOPE_CYCLE_COUNTER(QueryDirty);
-
-			for (const auto& Elem : DirtyElements)
+		{	// Returns true if we should continue
+			auto IntersectAndVisit = [&](const FElement& Elem) -> bool
 			{
 				const auto& InstanceBounds = Elem.Bounds;
 				if (PrePreFilterHelper(Elem.Payload, QueryData))
 				{
-					continue;
+					return true;
 				}
 
 				if (TAABBTreeIntersectionHelper<TQueryFastData, Query>::Intersects(Start, CurData, TOI, TmpPosition, InstanceBounds, QueryBounds, QueryHalfExtents, Dir, InvDir, bParallel))
@@ -773,8 +1115,75 @@ private:
 						return false;
 					}
 				}
+				return true;
+			};
+
+			//QUICK_SCOPE_CYCLE_COUNTER(QueryDirty);
+			bool bUseGrid = false;
+
+			if (DirtyElementGridEnabled())
+			{
+				if (Query == EAABBQueryType::Overlap)
+				{
+					bUseGrid = !TooManyOverlapQueryCells(QueryBounds, DirtyElementGridCellSizeInv, DirtyElementMaxGridCellQueryCount);
+				}
+				else if (Query == EAABBQueryType::Raycast)
+				{
+					bUseGrid = !TooManyRaycastQueryCells(Start, CurData.Dir, CurData.CurrentLength, DirtyElementGridCellSizeInv, DirtyElementMaxGridCellQueryCount);
+				}
+				else if (Query == EAABBQueryType::Sweep)
+				{
+					bUseGrid = !TooManySweepQueryCells(QueryHalfExtents, Start, CurData.Dir, CurData.CurrentLength, DirtyElementGridCellSizeInv, DirtyElementMaxGridCellQueryCount);
+				}
 			}
+			
+			if (bUseGrid)
+			{
+				TArray<DirtyGridHashEntry> HashEntryForOverlappedCells;
+				auto AddHashEntry = [&](int32 QueryCellHash)
+				{
+					const DirtyGridHashEntry* HashEntry = CellHashToFlatArray.Find(QueryCellHash);
+					if (HashEntry)
+					{
+						HashEntryForOverlappedCells.Add(*HashEntry);
+					}
+				};
 
+				if (Query == EAABBQueryType::Overlap)
+				{
+					DoForOverlappedCells(QueryBounds, DirtyElementGridCellSize, DirtyElementGridCellSizeInv, AddHashEntry);
+				}
+				else if (Query == EAABBQueryType::Raycast)
+				{
+					DoForRaycastIntersectCells(Start, CurData.Dir, CurData.CurrentLength, DirtyElementGridCellSize, DirtyElementGridCellSizeInv, AddHashEntry);
+				}
+				else if (Query == EAABBQueryType::Sweep)
+				{
+					DoForSweepIntersectCells(QueryHalfExtents, Start, CurData.Dir, CurData.CurrentLength, DirtyElementGridCellSize , DirtyElementGridCellSizeInv ,
+					[&](FReal X, FReal Y)
+					{
+						int32 QueryCellHash = HashCoordinates(X, Y, DirtyElementGridCellSizeInv);
+						const DirtyGridHashEntry* HashEntry = CellHashToFlatArray.Find(QueryCellHash);
+						if (HashEntry)
+						{
+							HashEntryForOverlappedCells.Add(*HashEntry);
+						}
+					});
+				}
+
+				if (!DoForHitGridCellsAndOverflow(HashEntryForOverlappedCells, IntersectAndVisit))
+				{
+					return false;
+				}
+			}  // end overlap
+
+			else for (const auto& Elem : DirtyElements)
+			{
+				if (!IntersectAndVisit(Elem))
+				{
+					return false;
+				}
+			}
 		}
 
 		struct FNodeQueueEntry
@@ -882,8 +1291,12 @@ private:
 		Leaves.Reset();
 		Nodes.Reset();
 		DirtyElements.Reset();
+		CellHashToFlatArray.Reset(); 
+		FlattenedCellArrayOfDirtyIndices.Reset();
+		DirtyElementsGridOverflow.Reset();
 		PayloadToInfo.Reset();
 		NumProcessedThisSlice = 0;
+		GetCVars();  // Safe to copy CVARS here
 
 		WorkSnapshot.Bounds = FAABB3::EmptyAABB();
 
@@ -922,7 +1335,7 @@ private:
 				{
 					if (bMutable)
 					{
-						PayloadToInfo.Add(Payload, FAABBTreePayloadInfo{ GlobalPayloads.Num(), INDEX_NONE, INDEX_NONE });
+						PayloadToInfo.Add(Payload, FAABBTreePayloadInfo{ GlobalPayloads.Num(), INDEX_NONE, INDEX_NONE, INDEX_NONE });
 					}
 					GlobalPayloads.Add(FElement{ Payload, ElemBounds });
 				}
@@ -1269,6 +1682,14 @@ private:
 		, Nodes(Other.Nodes)
 		, Leaves(Other.Leaves)
 		, DirtyElements(Other.DirtyElements)
+		, CellHashToFlatArray(Other.CellHashToFlatArray)
+		, FlattenedCellArrayOfDirtyIndices(Other.FlattenedCellArrayOfDirtyIndices)
+		, DirtyElementsGridOverflow(Other.DirtyElementsGridOverflow)
+		, DirtyElementGridCellSize(Other.DirtyElementGridCellSize)
+		, DirtyElementGridCellSizeInv(Other.DirtyElementGridCellSizeInv)
+		, DirtyElementMaxGridCellQueryCount(Other.DirtyElementMaxGridCellQueryCount)
+		, DirtyElementMaxPhysicalSizeInCells(Other.DirtyElementMaxPhysicalSizeInCells)
+		, DirtyElementMaxCellCapacity(Other.DirtyElementMaxCellCapacity)
 		, GlobalPayloads(Other.GlobalPayloads)
 		, PayloadToInfo(Other.PayloadToInfo)
 		, MaxChildrenInLeaf(Other.MaxChildrenInLeaf)
@@ -1293,6 +1714,17 @@ private:
 			Nodes = Rhs.Nodes;
 			Leaves = Rhs.Leaves;
 			DirtyElements = Rhs.DirtyElements;
+			
+			CellHashToFlatArray = Rhs.CellHashToFlatArray;
+			FlattenedCellArrayOfDirtyIndices = Rhs.FlattenedCellArrayOfDirtyIndices;
+			DirtyElementsGridOverflow = Rhs.DirtyElementsGridOverflow;
+			
+			DirtyElementGridCellSize = Rhs.DirtyElementGridCellSize;
+			DirtyElementGridCellSizeInv = Rhs.DirtyElementGridCellSizeInv;
+			DirtyElementMaxGridCellQueryCount = Rhs.DirtyElementMaxGridCellQueryCount;
+			DirtyElementMaxPhysicalSizeInCells = Rhs.DirtyElementMaxPhysicalSizeInCells;
+			DirtyElementMaxCellCapacity = Rhs.DirtyElementMaxCellCapacity;
+			
 			GlobalPayloads = Rhs.GlobalPayloads;
 			PayloadToInfo = Rhs.PayloadToInfo;
 			MaxChildrenInLeaf = Rhs.MaxChildrenInLeaf;
@@ -1308,6 +1740,18 @@ private:
 	TArray<FNode> Nodes;
 	TArray<TLeafType> Leaves;
 	TArray<FElement> DirtyElements;
+
+	// Data needed for DirtyElement2DAccelerationGrid
+	TMap<int32, DirtyGridHashEntry> CellHashToFlatArray; // Index, size into flat grid structure (FlattenedCellArrayOfDirtyIndices)
+	TArray<int32> FlattenedCellArrayOfDirtyIndices; // Array of indices into dirty Elements array, indices are always sorted within a 2D cell
+	TArray<int32> DirtyElementsGridOverflow; // Array of indices of DirtyElements that is not in the grid for some reason
+	// Copy of CVARS
+	FReal DirtyElementGridCellSize;
+	FReal DirtyElementGridCellSizeInv;
+	int32 DirtyElementMaxGridCellQueryCount;
+	int32 DirtyElementMaxPhysicalSizeInCells;
+	int32 DirtyElementMaxCellCapacity;
+
 	TArray<FElement> GlobalPayloads;
 	TArrayAsMap<TPayloadType, FAABBTreePayloadInfo> PayloadToInfo;
 
diff --git a/Engine/Source/Runtime/Experimental/Chaos/Public/Chaos/AABBTreeDirtyGridUtils.h b/Engine/Source/Runtime/Experimental/Chaos/Public/Chaos/AABBTreeDirtyGridUtils.h
new file mode 100644
--- /dev/null
+++ b/Engine/Source/Runtime/Experimental/Chaos/Public/Chaos/AABBTreeDirtyGridUtils.h
@@ -0,0 +1,534 @@
+// Copyright Epic Games, Inc. All Rights Reserved.
+#pragma once
+
+#include "Chaos/AABB.h"
+#include "Chaos/Defines.h"
+#include "ChaosLog.h"
+
+namespace Chaos
+{
+
+	FORCEINLINE_DEBUGGABLE uint32 InterleaveWithZeros(uint16 input)
+	{
+		uint32 Intermediate = (uint32)input;
+		Intermediate = (Intermediate ^ (Intermediate << 8)) & 0x00ff00ff;
+		Intermediate = (Intermediate ^ (Intermediate << 4)) & 0x0f0f0f0f;
+		Intermediate = (Intermediate ^ (Intermediate << 2)) & 0x33333333;
+		Intermediate = (Intermediate ^ (Intermediate << 1)) & 0x55555555;
+		return Intermediate;
+	}
+
+	FORCEINLINE_DEBUGGABLE int32 HashCoordinates(FReal Xcoordinate, FReal Ycoordinate, FReal DirtyElementGridCellSizeInv)
+	{
+		// Requirement: Hash should change for adjacent cells
+		uint16 X = (uint16)FMath::Floor(Xcoordinate * DirtyElementGridCellSizeInv);
+		uint16 Y = (uint16)FMath::Floor(Ycoordinate * DirtyElementGridCellSizeInv);
+
+		return (int32)(InterleaveWithZeros(X) | (InterleaveWithZeros(Y) << 1));
+	}
+
+	FORCEINLINE_DEBUGGABLE int32 GetDirtyCellIndexFromWorldCoordinate(FReal Coordinate, FReal DirtyElementGridCellSizeInv)
+	{
+		return (int32)(FMath::Floor(Coordinate * DirtyElementGridCellSizeInv));
+	}
+
+	FORCEINLINE_DEBUGGABLE bool TooManyOverlapQueryCells(const TAABB<FReal, 3>& AABB, FReal DirtyElementGridCellSizeInv, int32 MaximumOverlap)
+	{
+		int32 XsampleCount = (int32)(FMath::Floor(AABB.Max().X * DirtyElementGridCellSizeInv)) - (int32)(FMath::Floor(AABB.Min().X * DirtyElementGridCellSizeInv)) + 1;
+		int32 YsampleCount = (int32)(FMath::Floor(AABB.Max().Y * DirtyElementGridCellSizeInv)) - (int32)(FMath::Floor(AABB.Min().Y * DirtyElementGridCellSizeInv)) + 1;
+
+		if (XsampleCount * YsampleCount <= MaximumOverlap)
+		{
+			return false;
+		}
+		return true;
+	}
+
+	template <typename FunctionType>
+	FORCEINLINE_DEBUGGABLE void DoForOverlappedCells(const TAABB<FReal, 3>& AABB, FReal DirtyElementGridCellSize, FReal DirtyElementGridCellSizeInv, FunctionType Function)
+	{
+		int32 XsampleCount = (int32)(FMath::Floor(AABB.Max().X * DirtyElementGridCellSizeInv)) - (int32)(FMath::Floor(AABB.Min().X * DirtyElementGridCellSizeInv)) + 1;
+		int32 YsampleCount = (int32)(FMath::Floor(AABB.Max().Y * DirtyElementGridCellSizeInv)) - (int32)(FMath::Floor(AABB.Min().Y * DirtyElementGridCellSizeInv)) + 1;
+
+		FReal CurrentX = AABB.Min().X;
+		for (int32 XsampleIndex = 0; XsampleIndex < XsampleCount; XsampleIndex++)
+		{
+			FReal CurrentY = AABB.Min().Y;
+			for (int32 YsampleIndex = 0; YsampleIndex < YsampleCount; YsampleIndex++)
+			{
+				Function(HashCoordinates(CurrentX, CurrentY, DirtyElementGridCellSizeInv));
+				CurrentY += DirtyElementGridCellSize;
+			}
+			CurrentX += DirtyElementGridCellSize;
+		}
+	}
+
+	// Only execute function for new Cells not covered in old (Set difference: {Cells spanned by AABB} - { Cells spanned by AABBExclude})
+	template <typename FunctionType>
+	FORCEINLINE_DEBUGGABLE bool DoForOverlappedCellsExclude(const TAABB<FReal, 3>& AABB, const TAABB<FReal, 3>& AABBExclude, FReal DirtyElementGridCellSize, FReal DirtyElementGridCellSizeInv, FunctionType Function)
+	{
+
+		int32 NewCellStartX = GetDirtyCellIndexFromWorldCoordinate(AABB.Min().X, DirtyElementGridCellSizeInv);
+		int32 NewCellStartY = GetDirtyCellIndexFromWorldCoordinate(AABB.Min().Y, DirtyElementGridCellSizeInv);
+
+		int32 NewCellEndX = GetDirtyCellIndexFromWorldCoordinate(AABB.Max().X, DirtyElementGridCellSizeInv);
+		int32 NewCellEndY = GetDirtyCellIndexFromWorldCoordinate(AABB.Max().Y, DirtyElementGridCellSizeInv);
+
+		int32 OldCellStartX = GetDirtyCellIndexFromWorldCoordinate(AABBExclude.Min().X, DirtyElementGridCellSizeInv);
+		int32 OldCellStartY = GetDirtyCellIndexFromWorldCoordinate(AABBExclude.Min().Y, DirtyElementGridCellSizeInv);
+
+		int32 OldCellEndX = GetDirtyCellIndexFromWorldCoordinate(AABBExclude.Max().X, DirtyElementGridCellSizeInv);
+		int32 OldCellEndY = GetDirtyCellIndexFromWorldCoordinate(AABBExclude.Max().Y, DirtyElementGridCellSizeInv);
+
+		// Early out here
+		if (OldCellStartX <= NewCellStartX &&
+			OldCellStartY <= NewCellStartY &&
+			OldCellEndX >= NewCellEndX &&
+			OldCellEndY >= NewCellEndY)
+		{
+			return true;
+		}
+
+		for (int32 X = NewCellStartX; X <= NewCellEndX; X++)
+		{
+			for (int32 Y = NewCellStartY; Y <= NewCellEndY; Y++)
+			{
+				if (!(X >= OldCellStartX && X <= OldCellEndX && Y >= OldCellStartY && Y <= OldCellEndY))
+				{
+					if (!Function(HashCoordinates((FReal)X * DirtyElementGridCellSize, (FReal)Y * DirtyElementGridCellSize, DirtyElementGridCellSizeInv)))
+						return false; // Failure early out
+				}
+			}
+		}
+		return true;
+	}
+
+	FORCEINLINE_DEBUGGABLE int32 FindInSortedArray(const TArray<int32>& Array, int32 FindValue, int32 StartIndex, int32 EndIndex)
+	{
+		int32 TestIndex = (EndIndex + StartIndex) / 2;
+		int32 TestValue = Array[TestIndex];
+		if (TestValue == FindValue)
+		{
+			return TestIndex;
+		}
+
+		if (StartIndex == EndIndex)
+		{
+			return INDEX_NONE;
+		}
+
+		if (TestValue < FindValue)
+		{
+			// tail-recursion
+			return FindInSortedArray(Array, FindValue, TestIndex + 1, EndIndex);
+		}
+
+		if (StartIndex == TestIndex)
+		{
+			return INDEX_NONE;
+		}
+
+		// tail-recursion
+		return FindInSortedArray(Array, FindValue, StartIndex, TestIndex - 1);
+	}
+
+	FORCEINLINE_DEBUGGABLE int32 FindInsertIndexIntoSortedArray(const TArray<int32>& Array, int32 FindValue, int32 StartIndex, int32 EndIndex)
+	{
+		int32 TestIndex = (EndIndex + StartIndex) / 2;
+		int32 TestValue = Array[TestIndex];
+		if (TestValue == FindValue)
+		{
+			return INDEX_NONE; // Already in array
+		}
+
+		if (StartIndex == EndIndex)
+		{
+			return FindValue > TestValue ? StartIndex + 1 : StartIndex;
+		}
+
+		if (TestValue < FindValue)
+		{
+			// tail-recursion
+			return FindInsertIndexIntoSortedArray(Array, FindValue, TestIndex + 1, EndIndex);
+		}
+
+		if (StartIndex == TestIndex)
+		{
+			return TestIndex;
+		}
+
+		// tail-recursion
+		return FindInsertIndexIntoSortedArray(Array, FindValue, StartIndex, TestIndex - 1);
+	}
+
+	// Prerequisites: The array must be sorted from StartIndex to StartIndex + Count -1, and must have one element past StartIndex + Count -1 allocated
+	// returns false if the value was already in the array and therefore not added again
+	FORCEINLINE_DEBUGGABLE bool InsertValueIntoSortedSubArray(TArray<int32>& Array, int32 Value, int32 StartIndex, int32 Count)
+	{
+		// We must keep everything sorted
+		if (Count > 0)
+		{
+			int32 EndIndex = StartIndex + Count - 1;
+			int32 InsertIndex = FindInsertIndexIntoSortedArray(Array, Value, StartIndex, EndIndex);
+			//ensure(InsertIndex != INDEX_NONE) // Just for debugging
+			if (InsertIndex != INDEX_NONE)
+			{
+				for (int32 Index = EndIndex + 1; Index > InsertIndex; Index--)
+				{
+					Array[Index] = Array[Index - 1];
+				}
+				Array[InsertIndex] = Value;
+			}
+			else
+			{
+				return false;
+			}
+		}
+		else
+		{
+			Array[StartIndex] = Value;
+		}
+		return true;
+	}
+
+	// Prerequisites: The array must be sorted from StartIndex to EndIndex. 
+	// The extra element won't be deallocated
+	// returns true if the element has been found and successfully deleted
+	FORCEINLINE_DEBUGGABLE bool DeleteValueFromSortedSubArray(TArray<int32>& Array, int32 Value, int32 StartIndex, int32 Count)
+	{
+		if (ensure(Count > 0))
+		{
+			int32 EndIndex = StartIndex + Count - 1;
+			int32 DeleteIndex = FindInSortedArray(Array, Value, StartIndex, EndIndex);
+			if (DeleteIndex != INDEX_NONE)
+			{
+				for (int32 Index = DeleteIndex; Index < EndIndex; Index++)
+				{
+					Array[Index] = Array[Index + 1];
+				}
+				return true;
+			}
+		}
+		return false;
+	}
+
+	FORCEINLINE_DEBUGGABLE bool TooManySweepQueryCells(const TVec3<FReal>& QueryHalfExtents, const TVector<FReal, 3>& StartPoint, const TVector<FReal, 3>& Dir, FReal Length, FReal DirtyElementGridCellSizeInv, int32 DirtyElementMaxGridCellQueryCount)
+	{
+		int EstimatedNumberOfCells = ((int32)(QueryHalfExtents.X * 2 * DirtyElementGridCellSizeInv) + 2) * ((int32)(QueryHalfExtents.Y * 2 * DirtyElementGridCellSizeInv) + 2) +
+			((int32)(FMath::Max(QueryHalfExtents.X, QueryHalfExtents.Y) * 2 * DirtyElementGridCellSizeInv) + 2) * ((int32)(Length * DirtyElementGridCellSizeInv) + 2);
+
+		return EstimatedNumberOfCells > DirtyElementMaxGridCellQueryCount;
+	}
+
+
+	// This function should be called with a dominant x direction only!
+	template <typename FunctionType>
+	FORCEINLINE_DEBUGGABLE void DoForSweepIntersectCellsImp(const FReal QueryHalfExtentsX, const FReal QueryHalfExtentsY, const FReal StartPointX, const FReal StartPointY, const FReal EndPointX, const FReal EndPointY, FReal DirtyElementGridCellSize, FReal DirtyElementGridCellSizeInv, FunctionType InFunction)
+	{
+		// Use 2 paths (Line0 and Line1) that traces the shape of the swept AABB and fill between them
+		// Example of one of the cases (XDirectionDominant, Fill Up):
+		//                            Line0
+		//                        #############
+		//                      #             
+		//                    #               
+		//            Line0 #                 
+		//                #                   #
+		//              #                   #
+		//            #                   #   ----> Dx
+		//          #       ^           #
+		//        #         |         # Line1
+		//                 Fill     #
+		//                  |     #
+		//                      #
+		//        #############
+		//           Line1    ^TurningPointForLine1
+
+
+		FReal DeltaX = EndPointX - StartPointX;
+		FReal DeltaY = EndPointY - StartPointY;
+
+		FReal AbsDx = FMath::Abs(DeltaX);
+		FReal AbsDy = FMath::Abs(DeltaY);
+
+		bool DxTooSmall = AbsDx <= SMALL_NUMBER;
+		bool DyTooSmall = AbsDy <= SMALL_NUMBER;
+
+		int DeltaCelIndexX;
+		int DeltaCelIndexY;
+		FReal DtDy = 0; // DeltaTime over DeltaX
+		FReal DtDx = 0;
+
+		if (DxTooSmall)
+		{
+			// This is just the overlap case (no casting along a ray)
+			DeltaCelIndexX = 1;
+			DeltaCelIndexY = 1;
+		}
+		else
+		{
+			DeltaCelIndexX = DeltaX >= 0 ? 1 : -1;
+			DeltaCelIndexY = DeltaY >= 0 ? 1 : -1;
+		}
+
+		// Use parametric description of the lines here (t is the parameter and is positive along the ray)
+		// x = Dx/Dt*t + x0
+		// y = Dy/Dt*t + x0
+		DtDx = (FReal)DeltaCelIndexX;
+		DtDy = DyTooSmall ? 1 : (FReal)DeltaCelIndexX * DeltaX / DeltaY;
+
+		// Calculate all the bounds we need
+		FReal XEndPointExpanded = EndPointX + (DeltaCelIndexX >= 0 ? QueryHalfExtentsX : -QueryHalfExtentsX);
+		FReal YEndPointExpanded = EndPointY + (DeltaCelIndexY >= 0 ? QueryHalfExtentsY : -QueryHalfExtentsY);
+		FReal XStartPointExpanded = StartPointX + (DeltaCelIndexX >= 0 ? -QueryHalfExtentsX : QueryHalfExtentsX);
+		FReal YStartPointExpanded = StartPointY + (DeltaCelIndexY >= 0 ? -QueryHalfExtentsY : QueryHalfExtentsY);
+		int32 TurningPointForLine1; // This is where we need to change direction for line 2
+		TurningPointForLine1 = GetDirtyCellIndexFromWorldCoordinate(StartPointX + (DeltaCelIndexX >= 0 ? QueryHalfExtentsX : -QueryHalfExtentsX), DirtyElementGridCellSizeInv);
+
+		// Line0 current position
+		FReal X0 = XStartPointExpanded;
+		FReal Y0 = YStartPointExpanded + QueryHalfExtentsY * (FReal)DeltaCelIndexY * 2;
+
+		// Line1 current position
+		FReal X1 = XStartPointExpanded;
+		FReal Y1 = YStartPointExpanded;
+
+		int32 CurrentCellIndexX0 = (int32)FMath::Floor(X0 * DirtyElementGridCellSizeInv);
+		int32 CurrentCellIndexY0 = (int32)FMath::Floor(Y0 * DirtyElementGridCellSizeInv);
+
+		int32 CurrentCellIndexX1 = (int32)FMath::Floor(X1 * DirtyElementGridCellSizeInv);
+		int32 CurrentCellIndexY1 = (int32)FMath::Floor(Y1 * DirtyElementGridCellSizeInv);
+
+		int32 LastCellIndexX = (int32)FMath::Floor(XEndPointExpanded * DirtyElementGridCellSizeInv);
+		int32 LastCellIndexY = (int32)FMath::Floor(YEndPointExpanded * DirtyElementGridCellSizeInv);
+
+		bool Done = false;
+		while (!Done)
+		{
+			// Advance Line 0 crossing a horizontal border here (angle is 45 degrees or less)
+
+			if (CurrentCellIndexY0 * DeltaCelIndexY < LastCellIndexY * DeltaCelIndexY && !DyTooSmall)
+			{
+				FReal CrossingVerticleCellBorderT = std::numeric_limits<FReal>::max();
+				FReal CrossingHorizontalCellBorderT = std::numeric_limits<FReal>::max();
+				CrossingVerticleCellBorderT = DtDx * ((FReal)(CurrentCellIndexX0 + (DeltaCelIndexX > 0 ? 1 : 0)) * DirtyElementGridCellSize - X0);
+				CrossingHorizontalCellBorderT = DtDy * ((FReal)(CurrentCellIndexY0 + (DeltaCelIndexY > 0 ? 1 : 0)) * DirtyElementGridCellSize - Y0);
+				if (CrossingHorizontalCellBorderT < CrossingVerticleCellBorderT)
+				{
+					X0 += CrossingHorizontalCellBorderT * (1 / DtDx);  // DtDx is always 1 or -1
+					Y0 += CrossingHorizontalCellBorderT * (1 / DtDy);  // Abs(DtDy) >= 1
+					CurrentCellIndexY0 += DeltaCelIndexY;
+				}
+			}
+
+			for (int CurrentFillCellIndexY = CurrentCellIndexY1; CurrentFillCellIndexY * DeltaCelIndexY <= CurrentCellIndexY0 * DeltaCelIndexY; CurrentFillCellIndexY += DeltaCelIndexY)
+			{
+				InFunction((FReal)CurrentCellIndexX0 * DirtyElementGridCellSize, (FReal)CurrentFillCellIndexY * DirtyElementGridCellSize);
+			}
+
+			// Advance line 0 crossing vertical cell borders
+			{
+				if (CurrentCellIndexY0 != LastCellIndexY && !DyTooSmall)
+				{
+					FReal CrossingVerticleCellBorderT = std::numeric_limits<FReal>::max();
+					CrossingVerticleCellBorderT = DtDx * ((FReal)(CurrentCellIndexX0 + (DeltaCelIndexX > 0 ? 1 : 0)) * DirtyElementGridCellSize - X0);
+					X0 += CrossingVerticleCellBorderT * (1 / DtDx);
+					Y0 += CrossingVerticleCellBorderT * (1 / DtDy);
+				}
+				else
+				{
+					X0 += DirtyElementGridCellSize * (FReal)DeltaCelIndexX;
+				}
+				CurrentCellIndexX0 += DeltaCelIndexX;
+			}
+
+			// Advance line 1
+			if (CurrentCellIndexX1 != LastCellIndexX)
+			{
+				if (CurrentCellIndexX1 * DeltaCelIndexX < TurningPointForLine1 * DeltaCelIndexX)
+				{
+					X1 += DirtyElementGridCellSize * (FReal)DeltaCelIndexX;
+				}
+				else
+				{
+					if (CurrentCellIndexX1 == TurningPointForLine1)
+					{
+						// Put Line position exactly at the turning point
+						X1 = StartPointX + (DeltaCelIndexX >= 0 ? QueryHalfExtentsX : -QueryHalfExtentsX);
+					}
+
+					FReal CrossingVerticleCellBorderT = std::numeric_limits<FReal>::max();
+					FReal CrossingHorizontalCellBorderT = std::numeric_limits<FReal>::max();
+					if (!DxTooSmall)
+					{
+						CrossingVerticleCellBorderT = DtDx * ((FReal)(CurrentCellIndexX0 + (DeltaCelIndexX > 0 ? 1 : 0)) * DirtyElementGridCellSize - X1);
+					}
+
+					if (!DyTooSmall)
+					{
+						CrossingHorizontalCellBorderT = DtDy * ((FReal)(CurrentCellIndexY0 + (DeltaCelIndexY > 0 ? 1 : 0)) * DirtyElementGridCellSize - Y1);
+					}
+
+					if (CrossingHorizontalCellBorderT < CrossingVerticleCellBorderT)
+					{
+						CurrentCellIndexY1 += DeltaCelIndexY;
+					}
+
+					if (!DxTooSmall)
+					{
+						X1 += CrossingVerticleCellBorderT * (1 / DtDx);
+					}
+
+					if (!DyTooSmall)
+					{
+						Y1 += CrossingVerticleCellBorderT * (1 / DtDy);
+					}
+				}
+				CurrentCellIndexX1 += DeltaCelIndexX;
+			}
+			Done = (CurrentCellIndexY0 == LastCellIndexY) && (DeltaCelIndexX * CurrentCellIndexX0 > LastCellIndexX * DeltaCelIndexX);
+		}
+	}
+
+	template <typename FunctionType>
+	FORCEINLINE_DEBUGGABLE void DoForSweepIntersectCells(const TVec3<FReal>& QueryHalfExtents, const TVector<FReal, 3>& StartPoint, const TVector<FReal, 3>& Dir, FReal Length, FReal DirtyElementGridCellSize, FReal DirtyElementGridCellSizeInv, FunctionType InFunction)
+	{
+		const TVector<FReal, 3>& EndPoint = StartPoint + Length * Dir;
+		FReal DeltaX = EndPoint.X - StartPoint.X;
+		FReal DeltaY = EndPoint.Y - StartPoint.Y;
+
+		FReal AbsDx = FMath::Abs(DeltaX);
+		FReal AbsDy = FMath::Abs(DeltaY);
+
+
+		bool XDirectionDominant = AbsDx >= AbsDy;
+
+		if (XDirectionDominant)
+		{
+			DoForSweepIntersectCellsImp(QueryHalfExtents.X, QueryHalfExtents.Y, StartPoint.X, StartPoint.Y, EndPoint.X, EndPoint.Y, DirtyElementGridCellSize, DirtyElementGridCellSizeInv, InFunction);
+		}
+		else
+		{
+			// Swap Y and X
+			DoForSweepIntersectCellsImp(QueryHalfExtents.Y, QueryHalfExtents.X, StartPoint.Y, StartPoint.X, EndPoint.Y, EndPoint.X, DirtyElementGridCellSize, DirtyElementGridCellSizeInv, [&](FReal X, FReal Y) {InFunction(Y, X); });
+		}
+
+	}
+
+	FORCEINLINE_DEBUGGABLE bool TooManyRaycastQueryCells(const TVector<FReal, 3>& StartPoint, const TVector<FReal, 3>& Dir, const FReal Length, FReal DirtyElementGridCellSizeInv, int32 DirtyElementMaxGridCellQueryCount)
+	{
+		const TVector<FReal, 3>& EndPoint = StartPoint + Length * Dir;
+
+		int32 FirstCellIndexX = (int32)FMath::Floor(StartPoint.X * DirtyElementGridCellSizeInv);
+		int32 FirstCellIndexY = (int32)FMath::Floor(StartPoint.Y * DirtyElementGridCellSizeInv);
+
+		int32 LastCellIndexX = (int32)FMath::Floor(EndPoint.X * DirtyElementGridCellSizeInv);
+		int32 LastCellIndexY = (int32)FMath::Floor(EndPoint.Y * DirtyElementGridCellSizeInv);
+
+		// This will be equal to the Manhattan distance 
+		int CellCount = FMath::Abs(FirstCellIndexX - LastCellIndexX) + FMath::Abs(FirstCellIndexY - LastCellIndexY);
+
+		if (CellCount > DirtyElementMaxGridCellQueryCount)
+		{
+			return true;
+		}
+
+		return false;
+
+	}
+
+	template <typename FunctionType>
+	FORCEINLINE_DEBUGGABLE void DoForRaycastIntersectCells(const TVector<FReal, 3>& StartPoint, const TVector<FReal, 3>& Dir, FReal Length, FReal DirtyElementGridCellSize, FReal DirtyElementGridCellSizeInv, FunctionType InFunction)
+	{
+		const TVector<FReal, 3>& EndPoint = StartPoint + Length * Dir;
+
+		int32 CurrentCellIndexX = (int32)FMath::Floor(StartPoint.X * DirtyElementGridCellSizeInv);
+		int32 CurrentCellIndexY = (int32)FMath::Floor(StartPoint.Y * DirtyElementGridCellSizeInv);
+
+		int32 LastCellIndexX = (int32)FMath::Floor(EndPoint.X * DirtyElementGridCellSizeInv);
+		int32 LastCellIndexY = (int32)FMath::Floor(EndPoint.Y * DirtyElementGridCellSizeInv);
+
+		FReal DeltaX = EndPoint.X - StartPoint.X;
+		FReal DeltaY = EndPoint.Y - StartPoint.Y;
+
+		FReal AbsDx = FMath::Abs(DeltaX);
+		FReal AbsDy = FMath::Abs(DeltaY);
+
+		bool DxTooSmall = AbsDx <= SMALL_NUMBER;
+		bool DyTooSmall = AbsDy <= SMALL_NUMBER;
+
+		if (DxTooSmall && DyTooSmall)
+		{
+			InFunction(HashCoordinates(StartPoint.X, StartPoint.Y, DirtyElementGridCellSizeInv));
+			return;
+		}
+
+		int DeltaCelIndexX = DeltaX >= 0 ? 1 : -1;
+		int DeltaCelIndexY = DeltaY >= 0 ? 1 : -1;
+
+		FReal DtDy = 0; // DeltaTime over DeltaX
+		FReal DtDx = 0;
+
+		bool XDirectionDominant = AbsDx >= AbsDy;
+		// Use parametric description of line here (t is the parameter and is positive along the ray)
+		// x = Dx/Dt*t + x0
+		// y = Dy/Dt*t + x0
+		if (XDirectionDominant)
+		{
+			DtDx = (FReal)DeltaCelIndexX;
+			DtDy = DyTooSmall ? 1 : (FReal)DeltaCelIndexX * DeltaX / DeltaY;
+		}
+		else
+		{
+			DtDx = DxTooSmall ? 1 : (FReal)DeltaCelIndexY * DeltaY / DeltaX;
+			DtDy = (FReal)DeltaCelIndexY;
+		}
+
+		FReal X = StartPoint.X;
+		FReal Y = StartPoint.Y;
+
+		bool Done = false;
+		while (!Done)
+		{
+			InFunction(HashCoordinates((FReal)CurrentCellIndexX * DirtyElementGridCellSize, (FReal)CurrentCellIndexY * DirtyElementGridCellSize, DirtyElementGridCellSizeInv));
+			FReal CrossingVerticleCellBorderT = std::numeric_limits<FReal>::max();
+			FReal CrossingHorizontalCellBorderT = std::numeric_limits<FReal>::max();
+			if (!DxTooSmall)
+			{
+				CrossingVerticleCellBorderT = DtDx * ((FReal)(CurrentCellIndexX + (DeltaCelIndexX > 0 ? 1 : 0)) * DirtyElementGridCellSize - X);
+			}
+
+			if (!DyTooSmall)
+			{
+				CrossingHorizontalCellBorderT = DtDy * ((FReal)(CurrentCellIndexY + (DeltaCelIndexY > 0 ? 1 : 0)) * DirtyElementGridCellSize - Y);
+			}
+
+			FReal SmallestT;
+			if (CrossingVerticleCellBorderT <= CrossingHorizontalCellBorderT)
+			{
+				CurrentCellIndexX += DeltaCelIndexX;
+				SmallestT = CrossingVerticleCellBorderT;
+			}
+			else
+			{
+				CurrentCellIndexY += DeltaCelIndexY;
+				SmallestT = CrossingHorizontalCellBorderT;
+			}
+
+			if (!DxTooSmall)
+			{
+				X += SmallestT * (1 / DtDx);
+			}
+
+			if (!DyTooSmall)
+			{
+				Y += SmallestT * (1 / DtDy);
+			}
+
+			if (DeltaCelIndexX * CurrentCellIndexX > DeltaCelIndexX * LastCellIndexX || DeltaCelIndexY * CurrentCellIndexY > DeltaCelIndexY * LastCellIndexY)
+			{
+				Done = true;
+			}
+		}
+	}
+
+}
