Chaos Fracture Editor
-- Name of newly created cluster bone is no longer derived from selection, nor are bones renamed following the cluster operation.

diff --git a/Engine/Source/Runtime/Experimental/Chaos/Private/GeometryCollection/GeometryCollectionClusteringUtility.cpp b/Engine/Source/Runtime/Experimental/Chaos/Private/GeometryCollection/GeometryCollectionClusteringUtility.cpp
--- a/Engine/Source/Runtime/Experimental/Chaos/Private/GeometryCollection/GeometryCollectionClusteringUtility.cpp
+++ b/Engine/Source/Runtime/Experimental/Chaos/Private/GeometryCollection/GeometryCollectionClusteringUtility.cpp
@@ -23,7 +23,6 @@ void FGeometryCollectionClusteringUtility::ClusterBonesUnderNewNode(FGeometryCol
 	// New Bone Setup takes level/parent from the first of the Selected Bones
 	int32 SourceBoneIndex = InsertAtIndex;
 	int32 OriginalParentIndex = Parents[SourceBoneIndex];
-	BoneNames[NewBoneIndex] = BoneNames[SourceBoneIndex];
 	Parents[NewBoneIndex] = OriginalParentIndex;
 	Children[NewBoneIndex] = TSet<int32>(SelectedBones);
 	SimType[NewBoneIndex] = FGeometryCollection::ESimulationTypes::FST_Clustered;
@@ -42,30 +41,8 @@ void FGeometryCollectionClusteringUtility::ClusterBonesUnderNewNode(FGeometryCol
 		Children[OriginalParentIndex].Add(NewBoneIndex);
 	}
 
-	// update all the bone names from here on down the tree to the leaves
-	if (Parents[NewBoneIndex] != FGeometryCollection::Invalid)
-	{
-		RecursivelyUpdateChildBoneNames(Parents[NewBoneIndex], Children, BoneNames);
-	}
-	else
-	{
-		// #todo: how should we get the appropriate actor's name or invent a name here?
-		BoneNames[NewBoneIndex] = "ClusterBone";
-		RecursivelyUpdateChildBoneNames(NewBoneIndex, Children, BoneNames);
-	}
-
-	//
-	// determine original parents of moved nodes so we can update their childrens names
-	//
-	TArray<int32> ParentsToUpdateNames;
-	for (int32 SourceElement : SelectedBones)
-	{
-		ParentsToUpdateNames.AddUnique(Parents[SourceElement]);
-	}
-	for (int32 NodeIndex : ParentsToUpdateNames)
-	{
-		RecursivelyUpdateChildBoneNames(NodeIndex, Children, BoneNames);
-	}
+	// Update new cluster's bone name
+	BoneNames[NewBoneIndex] = FString::Printf(TEXT("ClusterBone_%d"), NewBoneIndex);
 
 	if (Validate)
 	{
