Backout 16553421, fixes perf issues

Change description:
Make ApplyKinematicTargets update bounds and accelerations structure.
Fix edge case in which accel structure rebuild completes on sim step in which kinematics are still interpolating, where game thread would be given structure with particle bounds at interpolated position, not target.

diff --git a/Engine/Source/Programs/HeadlessChaos/Private/HeadlessChaosTestPhysicsScene.cpp b/Engine/Source/Programs/HeadlessChaos/Private/HeadlessChaosTestPhysicsScene.cpp
--- a/Engine/Source/Programs/HeadlessChaos/Private/HeadlessChaosTestPhysicsScene.cpp
+++ b/Engine/Source/Programs/HeadlessChaos/Private/HeadlessChaosTestPhysicsScene.cpp
@@ -546,7 +546,8 @@ namespace ChaosTest {
 		}
 	}
 	
-	GTEST_TEST(EngineInterface, KinematicTargetsPassingGTWrongAccelBoundsBeforeHittingTarget)
+	// Disabled until we move fix with kineamtic bounds update on PushToPhysicsState into this branch. Might also need to remove bounds computation in ApplyKinematicTarget.
+	GTEST_TEST(EngineInterface, DISABLED_KinematicTargetsPassingGTWrongAccelBoundsBeforeHittingTarget)
 	{
 		// This test is designed to catch an edge case with kinematic targets (or other things interpolated over multiple physics steps), and acceleration structure bounds.
 		// Timestep is setup such that 1 GT frame = 10 physics steps, we have to make sure that if a non-final step gives an acceleration structure to game thread, in which
diff --git a/Engine/Source/Runtime/Experimental/Chaos/Private/Chaos/Framework/PhysicsSolverBase.cpp b/Engine/Source/Runtime/Experimental/Chaos/Private/Chaos/Framework/PhysicsSolverBase.cpp
--- a/Engine/Source/Runtime/Experimental/Chaos/Private/Chaos/Framework/PhysicsSolverBase.cpp
+++ b/Engine/Source/Runtime/Experimental/Chaos/Private/Chaos/Framework/PhysicsSolverBase.cpp
@@ -238,15 +238,6 @@ namespace Chaos
 		SpatialData.SpatialIdx = Particle->SpatialIdx();
 		SpatialData.AccelerationHandle = AccelerationHandle;
 		SpatialData.SyncTimestamp = MarshallingManager.GetExternalTimestamp_External();
-		SpatialData.bInterpolatedMovement = false;
-
-		if (bDelete == false)
-		{
-			if (FKinematicGeometryParticle* Kinematic = Particle->CastToKinematicParticle())
-			{
-				SpatialData.bInterpolatedMovement = Kinematic->IsKinematicTargetDirty();
-			}
-		}
 	}
 
 	void FPhysicsSolverBase::EnqueueSimcallbackRewindRegisteration(ISimCallbackObject* Callback)
diff --git a/Engine/Source/Runtime/Experimental/Chaos/Private/Chaos/PBDRigidsEvolution.cpp b/Engine/Source/Runtime/Experimental/Chaos/Private/Chaos/PBDRigidsEvolution.cpp
--- a/Engine/Source/Runtime/Experimental/Chaos/Private/Chaos/PBDRigidsEvolution.cpp
+++ b/Engine/Source/Runtime/Experimental/Chaos/Private/Chaos/PBDRigidsEvolution.cpp
@@ -198,7 +198,7 @@ namespace Chaos
 			DirtyParticle(Particle);
 		}
 
-		ComputeIntermediateSpatialAcceleration(/*bOnLastSubstep=*/true, /*bBlock=*/false);
+		ComputeIntermediateSpatialAcceleration();
 	}
 
 	FPBDRigidsEvolutionBase::~FPBDRigidsEvolutionBase()
@@ -446,17 +446,17 @@ namespace Chaos
 		}
 	}
 
-	void FPBDRigidsEvolutionBase::FlushInternalAccelerationQueue(bool bOnLastSubstep)
+	void FPBDRigidsEvolutionBase::FlushInternalAccelerationQueue()
 	{
 		for (const FPendingSpatialData& PendingData : InternalAccelerationQueue.PendingData)
 		{
 			ApplyParticlePendingData(PendingData, *InternalAcceleration, false);
 		}
-		InternalAcceleration->SetSyncTimestamp(LatestExternalTimestampConsumed_Internal, bOnLastSubstep);
+		InternalAcceleration->SetSyncTimestamp(LatestExternalTimestampConsumed_Internal);
 		InternalAccelerationQueue.Reset();
 	}
 
-	void FPBDRigidsEvolutionBase::FlushAsyncAccelerationQueue(bool bOnLastSubstep)
+	void FPBDRigidsEvolutionBase::FlushAsyncAccelerationQueue()
 	{
 		for (const FPendingSpatialData& PendingData : AsyncAccelerationQueue.PendingData)
 		{
@@ -478,8 +478,8 @@ namespace Chaos
 		//other queues are no longer needed since we've flushed all operations and now have a pristine structure
 		InternalAccelerationQueue.Reset();
 
-		AsyncInternalAcceleration->SetSyncTimestamp(LatestExternalTimestampConsumed_Internal, bOnLastSubstep);
-		AsyncExternalAcceleration->SetSyncTimestamp(LatestExternalTimestampConsumed_Internal, bOnLastSubstep);
+		AsyncInternalAcceleration->SetSyncTimestamp(LatestExternalTimestampConsumed_Internal);
+		AsyncExternalAcceleration->SetSyncTimestamp(LatestExternalTimestampConsumed_Internal);
 	}
 
 	//TODO: make static and _External suffix
@@ -490,12 +490,7 @@ namespace Chaos
 		for (int32 Idx = ExternalQueue.PendingData.Num() - 1; Idx >=0; --Idx)
 		{
 			const FPendingSpatialData& SpatialData = ExternalQueue.PendingData[Idx];
-
-			// If acceleration structure is not from final substep of frame, and particle is interpolating movement in sim, position in accleration structure
-			// is not at target, however GT is position/rendering at target.
-			bool bParticleIsInterpolating =  SpatialData.bInterpolatedMovement && !Acceleration.GetIsFromLastSimSubstep();
-
-			if(SpatialData.SyncTimestamp > SyncTimestamp || bParticleIsInterpolating)
+			if(SpatialData.SyncTimestamp > SyncTimestamp)
 			{
 				//operation still pending so update structure
 				//note: do we care about roll over? if game ticks at 60fps we'd get 385+ days
@@ -554,7 +549,7 @@ namespace Chaos
 		ExternalStructuresPool.Enqueue(Structure);
 	}
 	
-	void FPBDRigidsEvolutionBase::ComputeIntermediateSpatialAcceleration(bool bOnLastSubstep, bool bBlock)
+	void FPBDRigidsEvolutionBase::ComputeIntermediateSpatialAcceleration(bool bBlock)
 	{
 		LLM_SCOPE(ELLMTag::ChaosAcceleration);
 		SCOPE_CYCLE_COUNTER(STAT_ComputeIntermediateSpatialAcceleration);
@@ -569,7 +564,7 @@ namespace Chaos
 			InternalAcceleration = GetFreeSpatialAcceleration_Internal();
 			AsyncInternalAcceleration = GetFreeSpatialAcceleration_Internal();
 			AsyncExternalAcceleration = GetFreeSpatialAcceleration_Internal();
-			FlushInternalAccelerationQueue(bOnLastSubstep);
+			FlushInternalAccelerationQueue();
 
 			// Give game thread an empty structure to pop
 			ExternalStructuresQueue.Enqueue(AsyncExternalAcceleration);
@@ -595,7 +590,7 @@ namespace Chaos
 
 				bNeedsReset = true;
 
-				FlushAsyncAccelerationQueue(bOnLastSubstep);
+				FlushAsyncAccelerationQueue();
 
 				//swap acceleration structure for new one
 				std::swap(InternalAcceleration, AsyncInternalAcceleration);
@@ -608,7 +603,7 @@ namespace Chaos
 			}
 			else
 			{
-				FlushInternalAccelerationQueue(bOnLastSubstep);
+				FlushInternalAccelerationQueue();
 			}
 			
 			if (bCanStartAsyncTasks)
@@ -619,7 +614,7 @@ namespace Chaos
 		}
 		else
 		{
-			FlushInternalAccelerationQueue(bOnLastSubstep);
+			FlushInternalAccelerationQueue();
 		}
 	}
 
@@ -650,9 +645,9 @@ namespace Chaos
 	void FPBDRigidsEvolutionBase::FlushSpatialAcceleration()
 	{
 		//force build acceleration structure with latest data
-		ComputeIntermediateSpatialAcceleration(/*bOnLastSubstep=*/true, /*bBlock=*/true);
-		ComputeIntermediateSpatialAcceleration(/*bOnLastSubstep=*/true, /*bBlock=*/true);	//having to do it multiple times because of the various caching involved over multiple frames.
-		ComputeIntermediateSpatialAcceleration(/*bOnLastSubstep=*/true, /*bBlock=*/true);
+		ComputeIntermediateSpatialAcceleration(true);
+		ComputeIntermediateSpatialAcceleration(true);	//having to do it multiple times because of the various caching involved over multiple frames.
+		ComputeIntermediateSpatialAcceleration(true);
 	}
 
 	void FPBDRigidsEvolutionBase::RebuildSpatialAccelerationForPerfTest()
@@ -710,7 +705,7 @@ namespace Chaos
 		{
 			if (Ar.CustomVer(FExternalPhysicsCustomObjectVersion::GUID) >= FExternalPhysicsCustomObjectVersion::FlushEvolutionInternalAccelerationQueue)
 			{
-				FlushInternalAccelerationQueue(false);
+				FlushInternalAccelerationQueue();
 			}
 
 			if (Ar.CustomVer(FExternalPhysicsCustomObjectVersion::GUID) < FExternalPhysicsCustomObjectVersion::SerializeMultiStructures)
diff --git a/Engine/Source/Runtime/Experimental/Chaos/Private/Chaos/PBDRigidsEvolutionGBF.cpp b/Engine/Source/Runtime/Experimental/Chaos/Private/Chaos/PBDRigidsEvolutionGBF.cpp
--- a/Engine/Source/Runtime/Experimental/Chaos/Private/Chaos/PBDRigidsEvolutionGBF.cpp
+++ b/Engine/Source/Runtime/Experimental/Chaos/Private/Chaos/PBDRigidsEvolutionGBF.cpp
@@ -206,8 +206,7 @@ void FPBDRigidsEvolutionGBF::AdvanceOneTimeStepImpl(const FReal Dt,const FSubSte
 	}
 	{
 		SCOPE_CYCLE_COUNTER(STAT_Evolution_ComputeIntermediateSpatialAcceleration);
-		const bool bOnLastSubstep = SubStepInfo.Step == SubStepInfo.NumSteps - 1;
-		Base::ComputeIntermediateSpatialAcceleration(bOnLastSubstep, /*bBlock=*/false);
+		Base::ComputeIntermediateSpatialAcceleration();
 	}
 
 	{
diff --git a/Engine/Source/Runtime/Experimental/Chaos/Private/PBDRigidsSolver.cpp b/Engine/Source/Runtime/Experimental/Chaos/Private/PBDRigidsSolver.cpp
--- a/Engine/Source/Runtime/Experimental/Chaos/Private/PBDRigidsSolver.cpp
+++ b/Engine/Source/Runtime/Experimental/Chaos/Private/PBDRigidsSolver.cpp
@@ -318,7 +318,7 @@ namespace Chaos
 				// Editor will tick with 0 DT, this will guarantee acceleration structure is still processing even if we don't advance evolution.
 				if (MDeltaTime < MinDeltaTime)
 				{
-					MSolver->GetEvolution()->ComputeIntermediateSpatialAcceleration(MSubStepInfo.Step == MSubStepInfo.NumSteps - 1, /*bBlock=*/false);
+					MSolver->GetEvolution()->ComputeIntermediateSpatialAcceleration();
 				}
 
 
diff --git a/Engine/Source/Runtime/Experimental/Chaos/Public/Chaos/ISpatialAcceleration.h b/Engine/Source/Runtime/Experimental/Chaos/Public/Chaos/ISpatialAcceleration.h
--- a/Engine/Source/Runtime/Experimental/Chaos/Public/Chaos/ISpatialAcceleration.h
+++ b/Engine/Source/Runtime/Experimental/Chaos/Public/Chaos/ISpatialAcceleration.h
@@ -246,7 +246,7 @@ class CHAOS_API ISpatialAcceleration
 public:
 
 	ISpatialAcceleration(SpatialAccelerationType InType = static_cast<SpatialAccelerationType>(ESpatialAcceleration::Unknown))
-		: Type(InType), SyncTimestamp(0), bIsFromLastSimSubstep(false), AsyncTimeSlicingComplete(true)
+		: Type(InType), SyncTimestamp(0), AsyncTimeSlicingComplete(true)
 	{}
 
 	ISpatialAcceleration(ESpatialAcceleration InType)
@@ -339,16 +339,10 @@ public:
 		return SyncTimestamp;
 	}
 
-	bool GetIsFromLastSimSubstep() const
-	{
-		return bIsFromLastSimSubstep;
-	}
-
 	/** Call this whenever updating the acceleration structure for a new sync point */
-	void SetSyncTimestamp(int32 InTimestamp, bool bInFromLastSimSubstep)
+	void SetSyncTimestamp(int32 InTimestamp)
 	{
 		SyncTimestamp = InTimestamp;
-		bIsFromLastSimSubstep = bInFromLastSimSubstep;
 	}
 
 protected:
@@ -357,7 +351,6 @@ protected:
 private:
 	SpatialAccelerationType Type;
 	int32 SyncTimestamp;	//The set of inputs the acceleration structure is in sync with. GT moves forward in time and enqueues inputs
-	bool bIsFromLastSimSubstep; // Did this acceleration structure finish being built on last sim substep of frame?
 	bool AsyncTimeSlicingComplete;
 };
 
diff --git a/Engine/Source/Runtime/Experimental/Chaos/Public/Chaos/PBDRigidsEvolution.h b/Engine/Source/Runtime/Experimental/Chaos/Public/Chaos/PBDRigidsEvolution.h
--- a/Engine/Source/Runtime/Experimental/Chaos/Public/Chaos/PBDRigidsEvolution.h
+++ b/Engine/Source/Runtime/Experimental/Chaos/Public/Chaos/PBDRigidsEvolution.h
@@ -653,8 +653,6 @@ public:
 		check(StepFraction > (FReal)0);
 		check(StepFraction <= (FReal)1);
 
-		bool bPositionChanged = false;
-
 		// @todo(ccaulfield): optimize. Depending on the number of kinematics relative to the number that have 
 		// targets set, it may be faster to process a command list rather than iterate over them all each frame. 
 		const FReal MinDt = 1e-6f;
@@ -707,7 +705,6 @@ public:
 					FVec3 W = FRotation3::CalculateAngularVelocity(CurrentR, NewR, Dt);
 					Particle.W() = W;
 				}
-				bPositionChanged = true;
 				Particle.X() = NewX;
 				Particle.R() = NewR;
 				Particles.MarkTransientDirtyParticle(Particle.Handle());
@@ -717,7 +714,6 @@ public:
 			case EKinematicTargetMode::Velocity:
 			{
 				// Move based on velocity
-				bPositionChanged = true;
 				Particle.X() = Particle.X() + Particle.V() * Dt;
 				Particle.R() = FRotation3::IntegrateRotationWithAngularVelocity(Particle.R(), Particle.W(), Dt);
 				Particles.MarkTransientDirtyParticle(Particle.Handle());
@@ -735,17 +731,17 @@ public:
 				Rigid->PreV() = Rigid->V();
 				Rigid->PreW() = Rigid->W();
 
-				if (Particle.HasBounds() && bPositionChanged)
+				// Update the world bounds
+			/*	if (Rigid->HasBounds())
 				{
-					const FAABB3& LocalBounds = Particle.LocalBounds();
-					FAABB3 WorldSpaceBounds = LocalBounds.TransformedAABB(FRigidTransform3(Particle.X(), Particle.R()));
+					const FAABB3& LocalBounds = Rigid->LocalBounds();
+					FAABB3 WorldSpaceBounds = LocalBounds.TransformedAABB(FRigidTransform3(Rigid->X(), Rigid->R()));
 					if (Rigid->CCDEnabled())
 					{
 						WorldSpaceBounds.ThickenSymmetrically(Rigid->V() * Dt);
 					}
-					Particle.SetWorldSpaceInflatedBounds(WorldSpaceBounds);
-					DirtyParticle(Particle);
-				}
+					Rigid->SetWorldSpaceInflatedBounds(WorldSpaceBounds);
+				}*/
 				
 			}
 		}
@@ -764,7 +760,7 @@ public:
 	CHAOS_API void RebuildSpatialAccelerationForPerfTest();
 
 	/* Ticks computation of acceleration structures. Normally handled by Advance, but if not advancing can be called to incrementally build structures.*/
-	CHAOS_API void ComputeIntermediateSpatialAcceleration(bool bOnLastSubstep, bool bBlock);
+	CHAOS_API void ComputeIntermediateSpatialAcceleration(bool bBlock = false);
 
 	CHAOS_API const FPBDConstraintGraph& GetConstraintGraph() const { return ConstraintGraph; }
 	CHAOS_API FPBDConstraintGraph& GetConstraintGraph() { return ConstraintGraph; }
@@ -881,8 +877,8 @@ protected:
 		}
 	}
 
-	void FlushInternalAccelerationQueue(bool bOnLastSubstep);
-	void FlushAsyncAccelerationQueue(bool bOnLastSubstep);
+	void FlushInternalAccelerationQueue();
+	void FlushAsyncAccelerationQueue();
 	void WaitOnAccelerationStructure();
 
 	TArray<FForceRule> ForceRules;
diff --git a/Engine/Source/Runtime/Experimental/Chaos/Public/Chaos/PendingSpatialData.h b/Engine/Source/Runtime/Experimental/Chaos/Public/Chaos/PendingSpatialData.h
--- a/Engine/Source/Runtime/Experimental/Chaos/Public/Chaos/PendingSpatialData.h
+++ b/Engine/Source/Runtime/Experimental/Chaos/Public/Chaos/PendingSpatialData.h
@@ -16,16 +16,10 @@ struct FPendingSpatialData
 	FSpatialAccelerationIdx SpatialIdx;
 	int32 SyncTimestamp;	//indicates the inputs timestamp associated with latest change. Only relevant for external queue
 	bool bDelete;
-	
-	// If true, movement may take course over multiple physics steps generated from single GT input.
-	// (Example: Kinematic targets)
-	// Flag used to make sure we don't stomp over game thread bounds with position that is still interpolating when passing accel structure to GT.
-	bool bInterpolatedMovement;
 
 	FPendingSpatialData()
-		: SyncTimestamp(0)
-		, bDelete(false)
-		, bInterpolatedMovement(false)
+	: SyncTimestamp(0)
+	, bDelete(false)
 	{}
 
 	void Serialize(FChaosArchive& Ar)
