ToolsFramework: support multiple saved versions of a UInteractiveToolPropertySet. Previously only a single global version of a particular UInteractiveToolPropertySet could be saved via ::SaveProperties(), despite the Tool argument. This change adds an internal map keyed on a custom string. This allows different Tools using the same PropertySet to have their own unique saved properties if desired, or share a saved version between a set of Tools. A change to initialization is also made - during the first SaveProperties() call that spawns the new saved property set, the current/initial property values are no longer ovewritten with the default values in ::SaveRestoreProperties(). This allows a Tool to customize the default values before calling Restore().

An unused ::GetPropertyCache() function was also removed.

#rnx

diff --git a/Engine/Plugins/Experimental/MeshModelingToolset/Source/MeshModelingTools/Private/Properties/MeshMaterialProperties.cpp b/Engine/Plugins/Experimental/MeshModelingToolset/Source/MeshModelingTools/Private/Properties/MeshMaterialProperties.cpp
--- a/Engine/Plugins/Experimental/MeshModelingToolset/Source/MeshModelingTools/Private/Properties/MeshMaterialProperties.cpp
+++ b/Engine/Plugins/Experimental/MeshModelingToolset/Source/MeshModelingTools/Private/Properties/MeshMaterialProperties.cpp
@@ -17,9 +17,9 @@ UNewMeshMaterialProperties::UNewMeshMaterialProperties()
 	Material = CreateDefaultSubobject<UMaterialInterface>(TEXT("MATERIAL"));
 }
 
-void UExistingMeshMaterialProperties::RestoreProperties(UInteractiveTool* RestoreToTool)
+void UExistingMeshMaterialProperties::RestoreProperties(UInteractiveTool* RestoreToTool, const FString& CacheIdentifier)
 {
-	Super::RestoreProperties(RestoreToTool);
+	Super::RestoreProperties(RestoreToTool, CacheIdentifier);
 	Setup();
 }
 
diff --git a/Engine/Plugins/Experimental/MeshModelingToolset/Source/MeshModelingTools/Public/Properties/MeshMaterialProperties.h b/Engine/Plugins/Experimental/MeshModelingToolset/Source/MeshModelingTools/Public/Properties/MeshMaterialProperties.h
--- a/Engine/Plugins/Experimental/MeshModelingToolset/Source/MeshModelingTools/Public/Properties/MeshMaterialProperties.h
+++ b/Engine/Plugins/Experimental/MeshModelingToolset/Source/MeshModelingTools/Public/Properties/MeshMaterialProperties.h
@@ -82,7 +82,7 @@ public:
 	UMaterialInstanceDynamic* CheckerMaterial = nullptr;
 
 	// Needs custom restore in order to call setup
-	virtual void RestoreProperties(UInteractiveTool* RestoreToTool) override;
+	virtual void RestoreProperties(UInteractiveTool* RestoreToTool, const FString& CacheIdentifier = TEXT("")) override;
 
 	void Setup();
 
diff --git a/Engine/Source/Runtime/Experimental/InteractiveToolsFramework/Private/InteractiveTool.cpp b/Engine/Source/Runtime/Experimental/InteractiveToolsFramework/Private/InteractiveTool.cpp
--- a/Engine/Source/Runtime/Experimental/InteractiveToolsFramework/Private/InteractiveTool.cpp
+++ b/Engine/Source/Runtime/Experimental/InteractiveToolsFramework/Private/InteractiveTool.cpp
@@ -168,21 +168,25 @@ TArray<UObject*> UInteractiveTool::GetToolProperties(bool bEnabledOnly) const
 	return Properties;
 }
 
-void
-UInteractiveToolPropertySet::SaveProperties(UInteractiveTool* SaveFromTool)
+void UInteractiveToolPropertySet::SaveProperties(UInteractiveTool* SaveFromTool, const FString& CacheIdentifier)
 {
-	SaveRestoreProperties(SaveFromTool, true);
+	SaveRestoreProperties(SaveFromTool, CacheIdentifier, true);
 }
 
-void
-UInteractiveToolPropertySet::RestoreProperties(UInteractiveTool* RestoreToTool)
+void UInteractiveToolPropertySet::RestoreProperties(UInteractiveTool* RestoreToTool, const FString& CacheIdentifier)
 {
-	SaveRestoreProperties(RestoreToTool, false);
+	SaveRestoreProperties(RestoreToTool, CacheIdentifier, false);
 }
 
-void UInteractiveToolPropertySet::SaveRestoreProperties(UInteractiveTool* RestoreToTool, bool bSaving)
+void UInteractiveToolPropertySet::SaveRestoreProperties(UInteractiveTool* RestoreToTool, const FString& CacheIdentifier, bool bSaving)
 {
-	UInteractiveToolPropertySet* PropertyCache = GetDynamicPropertyCache();
+	bool bWasCreated = false;
+	UInteractiveToolPropertySet* PropertyCache = GetDynamicPropertyCache(CacheIdentifier, bWasCreated);
+	if (bWasCreated && !bSaving)
+	{
+		// if this is the first time we have seen this property set, then we don't have any values to Restore
+		return;
+	}
 	for ( FProperty* Prop : TFieldRange<FProperty>(GetClass()) )
 	{
 #if WITH_EDITOR
diff --git a/Engine/Source/Runtime/Experimental/InteractiveToolsFramework/Public/InteractiveTool.h b/Engine/Source/Runtime/Experimental/InteractiveToolsFramework/Public/InteractiveTool.h
--- a/Engine/Source/Runtime/Experimental/InteractiveToolsFramework/Public/InteractiveTool.h
+++ b/Engine/Source/Runtime/Experimental/InteractiveToolsFramework/Public/InteractiveTool.h
@@ -237,49 +237,59 @@ public:
 	// Setting saving/serialization
 	//
 	/**
-	 * Save and restore values of current Tool Properties between tool invocations
+	 * Save and restore values of current Tool Properties between tool invocations. The standard usage of
+	 * this setup is to call PropertySet->RestoreProperties() in the UInteractiveTool::Setup() implementation,
+	 * and PropertySet->SaveProperties() in the UInteractiveTool::Shutdown() implementation.
 	 *
 	 * The default behaviour of these functions is to Save or Restore every property in the property set.  It is not
 	 * necessary to save/restore all possible Properties (in many cases this would not make sense), so individual
-	 * properties may be skipped by adding the "TransientToolProperty" tag to their metadata on a property by property
-	 * basis.
+	 * properties may be skipped by adding the "TransientToolProperty" tag to their metadata on a property by property basis.
 	 *
 	 * Property sets which need more exotic behaviour upon Save and Restore may override these routines
 	 *
-	 * GetPropertyCache() and GetDynamicPropertyCache() can be used to return an instance of either the static or the
-	 * dynamic type of the specified property set subclass which may be used as a place to save/restore these properties
-	 * by customized Save/Restore functions
+	 * GetDynamicPropertyCache() can be used to return an instance of the specified property set subclass which 
+	 * may be used as a place to save/restore these properties by customized Save/Restore functions
+	 * 
+	 * Note: the current design of this system assumes that the CDO will keep the referenced objects alive.
+	 * This assumption is incorrect in Runtime builds, and some external mechanism must be used to keep
+	 * the elements in the CachedPropertiesMap alive.
 	 */
-	virtual void SaveProperties(UInteractiveTool* SaveFromTool);
-	virtual void RestoreProperties(UInteractiveTool* RestoreToTool);
-private:
 
-	// Utility func used to implement the default Save/RestoreProperties funcs
-	void SaveRestoreProperties(UInteractiveTool* RestoreToTool, bool bSaving);
-protected:
 	/**
-	 * GetPropertyCache returns a class-internal object that subclasses can use to save/restore properties.
-	 * If the subclass is UMyPropertySet, this function should only ever be called as GetPropertyCache<UMyPropertySet>().
+	 * Save the values of this PropertySet with the given CacheIdentifier.
+	 * The Tool parameter is currently ignored.
 	 */
-	template<typename ObjType>
-	static ObjType* GetPropertyCache()
-	{
-		ObjType* CDO = GetMutableDefault<ObjType>();
-		if (CDO->CachedProperties == nullptr)
-		{
-			CDO->CachedProperties = NewObject<ObjType>();
-		}
-		return CastChecked<ObjType>(CDO->CachedProperties);
-	}
+	virtual void SaveProperties(UInteractiveTool* SaveFromTool, const FString& CacheIdentifier = TEXT(""));
 
-	UInteractiveToolPropertySet* GetDynamicPropertyCache()
+	/**
+	 * Restore the values of the Property Set with the given CacheIdentifier.
+	 * The Tool parameter is currently ignored.
+	 */
+	virtual void RestoreProperties(UInteractiveTool* RestoreToTool, const FString& CacheIdentifier = TEXT(""));
+
+protected:
+	// Utility func used to implement the default Save/RestoreProperties funcs
+	virtual void SaveRestoreProperties(UInteractiveTool* RestoreToTool, const FString& CacheIdentifier, bool bSaving);
+
+	/**
+	 * GetDynamicPropertyCache return class-internal objects that subclasses can use to save/restore properties.
+	 * @param CacheIdentifier multiple versions of the UInteractiveToolPropertySet can be stored, CacheIdentifier indicates which one to use
+	 * @param bWasCreatedOut true is returned here if this is the first time the object was seen
+	 * @return instance of the current subclass that can be used to save/restore values
+	 */
+	TObjectPtr<UInteractiveToolPropertySet> GetDynamicPropertyCache(const FString& CacheIdentifier, bool& bWasCreatedOut)
 	{
+		bWasCreatedOut = false;
 		UInteractiveToolPropertySet* CDO = GetMutableDefault<UInteractiveToolPropertySet>(GetClass());
-		if (CDO->CachedProperties == nullptr)
+		TObjectPtr<UInteractiveToolPropertySet>* Found = CDO->CachedPropertiesMap.Find(CacheIdentifier);
+		if (Found == nullptr)
 		{
-			CDO->CachedProperties = NewObject<UInteractiveToolPropertySet>((UObject*)GetTransientPackage(), GetClass());
+			TObjectPtr<UInteractiveToolPropertySet> NewPropCache = NewObject<UInteractiveToolPropertySet>((UObject*)GetTransientPackage(), GetClass());
+			CDO->CachedPropertiesMap.Add(CacheIdentifier, NewPropCache);
+			bWasCreatedOut = true;
+			return NewPropCache;
 		}
-		return CDO->CachedProperties;
+		return *Found;
 	}
 
 public:
@@ -293,9 +303,8 @@ public:
 #endif
 
 protected:
-	// CachedProperties should only ever be set to an instance of the subclass, ideally via GetPropertyCache().
 	UPROPERTY()
-	TObjectPtr<UInteractiveToolPropertySet> CachedProperties = nullptr;
+	TMap<FString, TObjectPtr<UInteractiveToolPropertySet>> CachedPropertiesMap;
 
 	// Controls whether a property set is shown in the UI.  Transient so that disabling a PropertySet in one tool doesn't disable it in others.
 	UPROPERTY(meta=(TransientToolProperty))
