Fracture Mode Editor
-- correct behavior for shift and ctrl selection
-- Select Inverse abides ViewLevel

diff --git a/Engine/Plugins/Experimental/ChaosEditor/Source/FractureEditor/Private/FractureEditorMode.cpp b/Engine/Plugins/Experimental/ChaosEditor/Source/FractureEditor/Private/FractureEditorMode.cpp
--- a/Engine/Plugins/Experimental/ChaosEditor/Source/FractureEditor/Private/FractureEditorMode.cpp
+++ b/Engine/Plugins/Experimental/ChaosEditor/Source/FractureEditor/Private/FractureEditorMode.cpp
@@ -180,7 +180,7 @@ bool UFractureEditorMode::HandleClick(FEditorViewportClient* InViewportClient, H
 			TArray<int32> BoneIndices({GeometryCollectionProxy->BoneIndex});
 
 			FScopedTransaction Transaction(FractureTransactionContexts::SelectBoneContext, LOCTEXT("SelectGeometryCollectionBoneTransaction", "Select Bone"), GeometryCollectionProxy->Component);
-			FFractureSelectionTools::ToggleSelectedBones(GeometryCollectionProxy->Component, BoneIndices, !Click.IsControlDown());
+			FFractureSelectionTools::ToggleSelectedBones(GeometryCollectionProxy->Component, BoneIndices, !(Click.IsControlDown()||Click.IsShiftDown()), Click.IsShiftDown());
 			GeometryCollectionProxy->Component->Modify();
 
 			if (Toolkit.IsValid())
diff --git a/Engine/Plugins/Experimental/ChaosEditor/Source/FractureEditor/Private/FractureEditorModeToolkit.cpp b/Engine/Plugins/Experimental/ChaosEditor/Source/FractureEditor/Private/FractureEditorModeToolkit.cpp
--- a/Engine/Plugins/Experimental/ChaosEditor/Source/FractureEditor/Private/FractureEditorModeToolkit.cpp
+++ b/Engine/Plugins/Experimental/ChaosEditor/Source/FractureEditor/Private/FractureEditorModeToolkit.cpp
@@ -1348,7 +1348,7 @@ void FFractureEditorModeToolkit::OnOutlinerBoneSelectionChanged(UGeometryCollect
 		if (SelectedBones.Num())
 		{
 
-			FFractureSelectionTools::ToggleSelectedBones(RootComponent, SelectedBones, true);
+			FFractureSelectionTools::ToggleSelectedBones(RootComponent, SelectedBones, true, false);
 			OutlinerView->SetBoneSelection(RootComponent, SelectedBones, true);
 		}
 		else
diff --git a/Engine/Plugins/Experimental/ChaosEditor/Source/FractureEditor/Private/FractureSelectionTools.cpp b/Engine/Plugins/Experimental/ChaosEditor/Source/FractureEditor/Private/FractureSelectionTools.cpp
--- a/Engine/Plugins/Experimental/ChaosEditor/Source/FractureEditor/Private/FractureSelectionTools.cpp
+++ b/Engine/Plugins/Experimental/ChaosEditor/Source/FractureEditor/Private/FractureSelectionTools.cpp
@@ -5,7 +5,7 @@
 #include "GeometryCollection/GeometryCollectionClusteringUtility.h"
 #include "EditorSupportDelegates.h"
 
-void FFractureSelectionTools::ToggleSelectedBones(UGeometryCollectionComponent* GeometryCollectionComponent, TArray<int32>& BoneIndices, bool bClearCurrentSelection)
+void FFractureSelectionTools::ToggleSelectedBones(UGeometryCollectionComponent* GeometryCollectionComponent, TArray<int32>& BoneIndices, bool bClearCurrentSelection, bool bAdd)
 {
 	if (GeometryCollectionComponent)
 	{
@@ -24,7 +24,7 @@ void FFractureSelectionTools::ToggleSelectedBones(UGeometryCollectionComponent*
 					EditBoneColor.ResetBoneSelection();
 				}
 		
-				EditBoneColor.ToggleSelectedBones(BoneIndices);
+				EditBoneColor.ToggleSelectedBones(BoneIndices, bAdd);
 
 				// The actual selection made is based on the hierarchy and the view mode
 				if (GeometryCollection)
diff --git a/Engine/Plugins/Experimental/ChaosEditor/Source/FractureEditor/Public/FractureSelectionTools.h b/Engine/Plugins/Experimental/ChaosEditor/Source/FractureEditor/Public/FractureSelectionTools.h
--- a/Engine/Plugins/Experimental/ChaosEditor/Source/FractureEditor/Public/FractureSelectionTools.h
+++ b/Engine/Plugins/Experimental/ChaosEditor/Source/FractureEditor/Public/FractureSelectionTools.h
@@ -17,7 +17,7 @@ public:
 	 * @param BoneIndices					The indices to select
 	 * @param bClearCurrentSelection		Whether or not to clear the current selection set or just append the new indices to it				
 	 */
-	static void ToggleSelectedBones(UGeometryCollectionComponent* GeometryCollectionComponent, TArray<int32>& BoneIndices, bool bClearCurrentSelection);
+	static void ToggleSelectedBones(UGeometryCollectionComponent* GeometryCollectionComponent, TArray<int32>& BoneIndices, bool bClearCurrentSelection, bool bAdd);
 	static void ClearSelectedBones(UGeometryCollectionComponent* GeometryCollectionComponent);
 
 	static void SelectNeighbors(UGeometryCollectionComponent* GeometryCollectionComponent);
diff --git a/Engine/Plugins/Experimental/ChaosEditor/Source/FractureEditor/Public/FractureTool.cpp b/Engine/Plugins/Experimental/ChaosEditor/Source/FractureEditor/Public/FractureTool.cpp
--- a/Engine/Plugins/Experimental/ChaosEditor/Source/FractureEditor/Public/FractureTool.cpp
+++ b/Engine/Plugins/Experimental/ChaosEditor/Source/FractureEditor/Public/FractureTool.cpp
@@ -151,7 +151,7 @@ void UFractureActionTool::Refresh(FFractureToolContext& Context, FFractureEditor
 	else
 	{
 		EditBoneColor.SetSelectedBones(Context.GetSelection());
-		FFractureSelectionTools::ToggleSelectedBones(GeometryCollectionComponent, Context.GetSelection(), true);
+		FFractureSelectionTools::ToggleSelectedBones(GeometryCollectionComponent, Context.GetSelection(), true, false);
 	}
 	
 	Toolkit->UpdateExplodedVectors(GeometryCollectionComponent);
diff --git a/Engine/Source/Runtime/Experimental/Chaos/Private/GeometryCollection/GeometryCollectionClusteringUtility.cpp b/Engine/Source/Runtime/Experimental/Chaos/Private/GeometryCollection/GeometryCollectionClusteringUtility.cpp
--- a/Engine/Source/Runtime/Experimental/Chaos/Private/GeometryCollection/GeometryCollectionClusteringUtility.cpp
+++ b/Engine/Source/Runtime/Experimental/Chaos/Private/GeometryCollection/GeometryCollectionClusteringUtility.cpp
@@ -482,6 +482,32 @@ void FGeometryCollectionClusteringUtility::GetChildBonesFromLevel(const FGeometr
 
 }
 
+void FGeometryCollectionClusteringUtility::GetChildBonesAtLevel(const FGeometryCollection* GeometryCollection, int32 SourceBone, int32 Level, TArray<int32>& BonesOut)
+{
+	check(GeometryCollection);
+
+	if (Level == -1)
+	{
+		GetLeafBones(GeometryCollection, SourceBone, false, BonesOut);
+	}
+	else
+	{
+		const TManagedArray<int32>& Levels = GeometryCollection->GetAttribute<int32>("Level", FGeometryCollection::TransformGroup);
+		const TManagedArray<TSet<int32>>& Children = GeometryCollection->Children;
+		if (Levels[SourceBone] == Level)
+		{
+			BonesOut.Push(SourceBone);
+		}
+		else
+		{
+			for (int32 Child : Children[SourceBone])
+			{
+				GetChildBonesAtLevel(GeometryCollection, Child, Level, BonesOut);
+			}
+		}
+	}
+}
+
 void FGeometryCollectionClusteringUtility::RecursiveAddAllChildren(const TManagedArray<TSet<int32>>& Children, int32 SourceBone, TArray<int32>& BonesOut)
 {
 	BonesOut.AddUnique(SourceBone);
@@ -733,7 +759,7 @@ void FGeometryCollectionClusteringUtility::ContextBasedClusterSelection(
 
 }
 
-void FGeometryCollectionClusteringUtility::GetLeafBones(FGeometryCollection* GeometryCollection, int BoneIndex, bool bOnlyRigids, TArray<int32>& LeafBonesOut)
+void FGeometryCollectionClusteringUtility::GetLeafBones(const FGeometryCollection* GeometryCollection, int BoneIndex, bool bOnlyRigids, TArray<int32>& LeafBonesOut)
 {
 	if (!ensure(BoneIndex >= 0))
 	{
diff --git a/Engine/Source/Runtime/Experimental/Chaos/Public/GeometryCollection/GeometryCollectionClusteringUtility.h b/Engine/Source/Runtime/Experimental/Chaos/Public/GeometryCollection/GeometryCollectionClusteringUtility.h
--- a/Engine/Source/Runtime/Experimental/Chaos/Public/GeometryCollection/GeometryCollectionClusteringUtility.h
+++ b/Engine/Source/Runtime/Experimental/Chaos/Public/GeometryCollection/GeometryCollectionClusteringUtility.h
@@ -64,6 +64,9 @@ public:
 	/** Get list of child bones down from the source bone below the specified hierarchy level */
 	static void GetChildBonesFromLevel(const FGeometryCollection* GeometryCollection, int32 SourceBone, int32 Level, TArray<int32>& BonesOut);
 
+	/** Get list of child bones down from the source bone at the specified hierarchy level */
+	static void GetChildBonesAtLevel(const FGeometryCollection* GeometryCollection, int32 SourceBone, int32 Level, TArray<int32>& BonesOut);
+
 	/** Recursively Add all children to output bone list from source bone down to the leaf nodes */
 	static void RecursiveAddAllChildren(const TManagedArray<TSet<int32>>& Children, int32 SourceBone, TArray<int32>& BonesOut);
 
@@ -104,7 +107,7 @@ public:
 	static void ContextBasedClusterSelection(FGeometryCollection* GeometryCollection, int ViewLevel, const TArray<int32>& SelectedComponentBonesIn, TArray<int32>& SelectedComponentBonesOut, TArray<int32>& HighlightedComponentBonesOut);
 
 	/** return an array of all child leaf nodes below the specified node. If bOnlyRigids is true, the first Rigid node dound is considered a leaf, regardless of an children it might have. */
-	static void GetLeafBones(FGeometryCollection* GeometryCollection, int BoneIndex, bool bOnlyRigids, TArray<int32>& LeafBonesOut);
+	static void GetLeafBones(const FGeometryCollection* GeometryCollection, int BoneIndex, bool bOnlyRigids, TArray<int32>& LeafBonesOut);
 
 	/** move the selected node up a level in direction of root */
 	static void MoveUpOneHierarchyLevel(FGeometryCollection* GeometryCollection, const TArray<int32>& SelectedBones);
diff --git a/Engine/Source/Runtime/Experimental/GeometryCollectionEngine/Private/GeometryCollection/GeometryCollectionComponent.cpp b/Engine/Source/Runtime/Experimental/GeometryCollectionEngine/Private/GeometryCollection/GeometryCollectionComponent.cpp
--- a/Engine/Source/Runtime/Experimental/GeometryCollectionEngine/Private/GeometryCollection/GeometryCollectionComponent.cpp
+++ b/Engine/Source/Runtime/Experimental/GeometryCollectionEngine/Private/GeometryCollection/GeometryCollectionComponent.cpp
@@ -2147,18 +2147,34 @@ void FScopedColorEdit::AppendSelectedBones(const TArray<int32>& SelectedBonesIn)
 	Component->SelectedBones.Append(SelectedBonesIn);
 }
 
-void FScopedColorEdit::ToggleSelectedBones(const TArray<int32>& SelectedBonesIn)
+void FScopedColorEdit::ToggleSelectedBones(const TArray<int32>& SelectedBonesIn, bool bAdd)
 {
 	bUpdated = true;
-	for (int32 BoneIndex : SelectedBonesIn)
-	{
-		if (Component->SelectedBones.Contains(BoneIndex))
-		{
-			Component->SelectedBones.Remove(BoneIndex);
-		}
-		else
+	
+	const UGeometryCollection* GeometryCollection = Component->GetRestCollection();
+	if (GeometryCollection)
+	{ 
+		TSharedPtr<FGeometryCollection, ESPMode::ThreadSafe> GeometryCollectionPtr = GeometryCollection->GetGeometryCollection();
+		for (int32 BoneIndex : SelectedBonesIn)
 		{
-			Component->SelectedBones.Add(BoneIndex);
+		
+			int32 ContextBoneIndex = (GetViewLevel() > -1) ? FGeometryCollectionClusteringUtility::GetParentOfBoneAtSpecifiedLevel(GeometryCollectionPtr.Get(), BoneIndex, GetViewLevel()) : BoneIndex;
+		
+			if (bAdd) // shift select
+			{
+				Component->SelectedBones.Add(BoneIndex);
+			}
+			else // ctrl select (toggle)
+			{ 
+				if (Component->SelectedBones.Contains(ContextBoneIndex))
+				{
+					Component->SelectedBones.Remove(ContextBoneIndex);
+				}
+				else
+				{
+					Component->SelectedBones.Add(ContextBoneIndex);
+				}
+			}
 		}
 	}
 }
@@ -2230,19 +2246,39 @@ void FScopedColorEdit::SelectBones(GeometryCollection::ESelectionMode SelectionM
 			TArray<int32> Roots;
 			FGeometryCollectionClusteringUtility::GetRootBones(GeometryCollectionPtr.Get(), Roots);
 			TArray<int32> NewSelection;
+
 			for (int32 RootElement : Roots)
 			{
-				TArray<int32> LeafBones;
-				FGeometryCollectionClusteringUtility::GetLeafBones(GeometryCollectionPtr.Get(), RootElement, true, LeafBones);
+				if (GetViewLevel() == -1)
+				{
+					TArray<int32> LeafBones;
+					FGeometryCollectionClusteringUtility::GetLeafBones(GeometryCollectionPtr.Get(), RootElement, true, LeafBones);
 
-				for (int32 Element : LeafBones)
+					for (int32 Element : LeafBones)
+					{
+						if (!IsBoneSelected(Element))
+						{
+							NewSelection.Push(Element);
+						}
+					}
+				}
+				else
 				{
-					if (!IsBoneSelected(Element))
+					TArray<int32> ViewLevelBones;
+					FGeometryCollectionClusteringUtility::GetChildBonesAtLevel(GeometryCollectionPtr.Get(), RootElement, GetViewLevel(), ViewLevelBones);
+					for (int32 ViewLevelBone : ViewLevelBones)
 					{
-						NewSelection.Push(Element);
+						if (!IsBoneSelected(ViewLevelBone))
+						{
+							NewSelection.Push(ViewLevelBone);
+							TArray<int32> ChildBones;
+							FGeometryCollectionClusteringUtility::GetChildBonesFromLevel(GeometryCollectionPtr.Get(), ViewLevelBone, GetViewLevel(), ChildBones);
+							NewSelection.Append(ChildBones);
+						}
 					}
 				}
 			}
+			
 			ResetBoneSelection();
 			AppendSelectedBones(NewSelection);
 		}
diff --git a/Engine/Source/Runtime/Experimental/GeometryCollectionEngine/Public/GeometryCollection/GeometryCollectionComponent.h b/Engine/Source/Runtime/Experimental/GeometryCollectionEngine/Public/GeometryCollection/GeometryCollectionComponent.h
--- a/Engine/Source/Runtime/Experimental/GeometryCollectionEngine/Public/GeometryCollection/GeometryCollectionComponent.h
+++ b/Engine/Source/Runtime/Experimental/GeometryCollectionEngine/Public/GeometryCollection/GeometryCollectionComponent.h
@@ -200,7 +200,7 @@ public:
 	bool IsBoneSelected(int BoneIndex) const;
 	void SetSelectedBones(const TArray<int32>& SelectedBonesIn);
 	void AppendSelectedBones(const TArray<int32>& SelectedBonesIn);
-	void ToggleSelectedBones(const TArray<int32>& SelectedBonesIn);
+	void ToggleSelectedBones(const TArray<int32>& SelectedBonesIn, bool bAdd);
 	void AddSelectedBone(int32 BoneIndex);
 	void ClearSelectedBone(int32 BoneIndex);
 	const TArray<int32>& GetSelectedBones() const;
