Chaos : Geometry Collection
- TConvex constructor that uses UE::Geometry::TConvexHull3 (Tyson Brochu)
- Updated fracture to use UE::Geometry::TConvexHull3 generation.

diff --git a/Engine/Source/Programs/HeadlessChaos/Private/HeadlessChaosApp.cpp b/Engine/Source/Programs/HeadlessChaos/Private/HeadlessChaosApp.cpp
--- a/Engine/Source/Programs/HeadlessChaos/Private/HeadlessChaosApp.cpp
+++ b/Engine/Source/Programs/HeadlessChaos/Private/HeadlessChaosApp.cpp
@@ -62,6 +62,7 @@
 #include "GeometryCollection/GeometryCollectionTestSkeletalMeshPhysicsProxy.h"
 #include "GeometryCollection/GeometryCollectionTestSerialization.h"
 
+#include "CompGeom/ExactPredicates.h"
 
 IMPLEMENT_APPLICATION(HeadlessChaos, "HeadlessChaos");
 
@@ -432,7 +433,9 @@ class UEGTestPrinter : public ::testing::EmptyTestEventListener
 
 INT32_MAIN_INT32_ARGC_TCHAR_ARGV()
 {
-    // start up the main loop
+	UE::Geometry::ExactPredicates::GlobalInit();
+
+	// start up the main loop
 	GEngineLoop.PreInit(ArgC, ArgV);
 	FModuleManager::Get().StartProcessingNewlyLoadedObjects();
 	
diff --git a/Engine/Source/Runtime/Experimental/Chaos/Chaos.Build.cs b/Engine/Source/Runtime/Experimental/Chaos/Chaos.Build.cs
--- a/Engine/Source/Runtime/Experimental/Chaos/Chaos.Build.cs
+++ b/Engine/Source/Runtime/Experimental/Chaos/Chaos.Build.cs
@@ -18,8 +18,9 @@ namespace UnrealBuildTool.Rules
 				"ChaosCore",
                 "IntelISPC",
 				"TraceLog",
-                "Voronoi"
-                }
+                "Voronoi",
+				"GeometryCore"
+				}
             );
 
             PublicDefinitions.Add("COMPILE_WITHOUT_UNREAL_SUPPORT=0");
diff --git a/Engine/Source/Runtime/Experimental/Chaos/Private/Chaos/CollisionConvexMesh.cpp b/Engine/Source/Runtime/Experimental/Chaos/Private/Chaos/CollisionConvexMesh.cpp
--- a/Engine/Source/Runtime/Experimental/Chaos/Private/Chaos/CollisionConvexMesh.cpp
+++ b/Engine/Source/Runtime/Experimental/Chaos/Private/Chaos/CollisionConvexMesh.cpp
@@ -2,6 +2,7 @@
 
 #include "Chaos/CollisionConvexMesh.h"
 #include "HAL/IConsoleManager.h"
+#include "CompGeom/ConvexHull3.h"
 
 namespace Chaos
 {
@@ -14,9 +15,171 @@ namespace Chaos
 
 	int32 FConvexBuilder::ComputeHorizonEpsilonFromMeshExtends = 1;
 
+	bool FConvexBuilder::bUseGeometryTConvexHull3 = false;
+
 	FAutoConsoleVariableRef CVarConvexGeometryCheckEnable(TEXT("p.Chaos.ConvexGeometryCheckEnable"), FConvexBuilder::PerformGeometryCheck, TEXT("Perform convex geometry complexity check for Chaos physics."));
 
 	FAutoConsoleVariableRef CVarConvexGeometrySimplifyEnable(TEXT("p.Chaos.PerformGeometryReduction"), FConvexBuilder::PerformGeometryReduction, TEXT("Perform convex geometry simplification to increase performance in Chaos physics."));
 
 	FAutoConsoleVariableRef CVarConvexParticlesWarningThreshold(TEXT("p.Chaos.ConvexParticlesWarningThreshold"), FConvexBuilder::VerticesThreshold, TEXT("Threshold beyond which we warn about collision geometry complexity."));
+	
+	FAutoConsoleVariableRef CvarUseGeometryTConvexHull3(TEXT("p.Chaos.Convex.UseTConvexHull3Builder"), FConvexBuilder::bUseGeometryTConvexHull3, TEXT("Use the newer Geometry Tools code path for generating convex hulls.[def:false]"));
+
+
+
+	void FConvexBuilder::Build(const TArray<FVec3>& InVertices, TArray <TPlaneConcrete<FReal, 3>>& OutPlanes, TArray<TArray<int32>>& OutFaceIndices, TArray<FVec3>& OutVertices, FAABB3& OutLocalBounds)
+	{
+		OutPlanes.Reset();
+		OutVertices.Reset();
+		OutLocalBounds = FAABB3::EmptyAABB();
+
+		const int32 NumVerticesIn = InVertices.Num();
+		if (NumVerticesIn == 0)
+		{
+			return;
+		}
+
+
+		const TArray<FVec3>* VerticesToUse = &InVertices;
+		TArray<FVec3> ModifiedVertices;
+
+		// For triangles and planar shapes, create a very thin prism as a convex
+		auto Inflate = [](const TArray<FVec3>& Source, TArray<FVec3>& Destination, const FVec3& Normal, FReal Inflation)
+		{
+			const int32 NumSource = Source.Num();
+			Destination.Reset();
+			Destination.SetNum(NumSource * 2);
+
+			for (int32 Index = 0; Index < NumSource; ++Index)
+			{
+				Destination[Index] = Source[Index];
+				Destination[NumSource + Index] = Source[Index] + Normal * Inflation;
+			}
+		};
+
+		FVec3 PlanarNormal(0);
+		if (NumVerticesIn == 3)
+		{
+			const bool bIsValidTriangle = IsValidTriangle(InVertices[0], InVertices[1], InVertices[2], PlanarNormal);
+
+			//TODO_SQ_IMPLEMENTATION: should do proper cleanup to avoid this
+			if (ensureMsgf(bIsValidTriangle, TEXT("FConvexBuilder::Build(): Generated invalid triangle!")))
+			{
+				Inflate(InVertices, ModifiedVertices, PlanarNormal, TriQuadPrismInflation());
+				VerticesToUse = &ModifiedVertices;
+				UE_LOG(LogChaos, Verbose, TEXT("Encountered a triangle in convex hull generation. Will prepare a prism of thickness %.5f in place of a triangle."), TriQuadPrismInflation());
+			}
+			else
+			{
+				return;
+			}
+		}
+		else if (IsPlanarShape(InVertices, PlanarNormal))
+		{
+			Inflate(InVertices, ModifiedVertices, PlanarNormal, TriQuadPrismInflation());
+			VerticesToUse = &ModifiedVertices;
+			UE_LOG(LogChaos, Verbose, TEXT("Encountered a planar shape in convex hull generation. Will prepare a prism of thickness %.5f in place of a triangle."), TriQuadPrismInflation());
+		}
+
+		const int32 NumVerticesToUse = VerticesToUse->Num();
+
+		OutLocalBounds = FAABB3((*VerticesToUse)[0], (*VerticesToUse)[0]);
+		for (int32 VertexIndex = 0; VertexIndex < NumVerticesToUse; ++VertexIndex)
+		{
+			OutLocalBounds.GrowToInclude((*VerticesToUse)[VertexIndex]);
+		}
+
+		if (NumVerticesToUse >= 4)
+		{
+			// @todo(chaos) : Experimental path for convex generation. This is still under
+			//                testing, and will eventually replace the convex generation
+			//                in chaos. Currenty defaults to off, but can be enabled with 
+			//                the cvar : p.Chaos.Convex.UseTConvexHull3Builder
+
+			if (bUseGeometryTConvexHull3) // Use the newer Geometry Tools code path for generating convex hulls.
+			{
+				UE::Geometry::TConvexHull3<FReal> HullCompute;
+				if (HullCompute.Solve<FVec3>(*VerticesToUse))
+				{
+					// Get Planes, FaceIndices, Vertices, and LocalBoundingBox
+					TMap<int32, int32> HullVertMap;
+					HullCompute.GetTriangles([&VerticesToUse, &HullVertMap, &OutPlanes, &OutFaceIndices, &OutVertices, &OutLocalBounds](UE::Geometry::FIndex3i Triangle)
+						{
+							for (int32 j = 0; j < 3; ++j)		// From FMeshConvexHull::Compute_FullMesh
+							{
+								int32 Index = Triangle[j];
+								if (HullVertMap.Contains(Index) == false)
+								{
+									const FVec3& OrigPos = (*VerticesToUse)[Index];
+									int32 NewVID = OutVertices.Num();
+									OutVertices.Add(OrigPos);
+									HullVertMap.Add(Index, NewVID);
+									Triangle[j] = NewVID;
+								}
+								else
+								{
+									Triangle[j] = HullVertMap[Index];
+								}
+							}
+
+							OutFaceIndices.Add({ Triangle[0], Triangle[1], Triangle[2] });
+
+							UE::Geometry::THalfspace3<FReal> Halfspace(OutVertices[Triangle[0]], OutVertices[Triangle[1]], OutVertices[Triangle[2]]);
+							const FVec3& Normal = Halfspace.Normal;
+							const FVec3& Point = OutVertices[Triangle[0]];
+							OutPlanes.Add(TPlaneConcrete<FReal, 3>{ Point, Normal });
+						});
+
+					OutLocalBounds = FAABB3(OutVertices[0], OutVertices[0]);
+					for (int32 VertexIndex = 1; VertexIndex < OutVertices.Num(); ++VertexIndex)
+					{
+						OutLocalBounds.GrowToInclude(OutVertices[VertexIndex]);
+					}
+				}
+			}
+			else // Use the older Chaos code path for generating convex hulls.
+			{
+				// @todo(chaos): Deprecate the older convex code path.
+				TArray<TVec3<int32>> Indices;
+				Params BuildParams;
+				BuildParams.HorizonEpsilon = Chaos::FConvexBuilder::SuggestEpsilon(*VerticesToUse);
+				BuildConvexHull(*VerticesToUse, Indices, BuildParams);
+				OutPlanes.Reserve(Indices.Num());
+				TMap<int32, int32> IndexMap; // maps original particle indices to output particle indices
+				int32 NewIdx = 0;
+
+				const auto AddIndex = [&IndexMap, &NewIdx](const int32 OriginalIdx)
+				{
+					if (int32* Idx = IndexMap.Find(OriginalIdx))
+					{
+						return *Idx;
+					}
+					IndexMap.Add(OriginalIdx, NewIdx);
+					return NewIdx++;
+				};
+
+				for (const TVec3<int32>& Idx : Indices)
+				{
+					FVec3 Vs[3] = { (*VerticesToUse)[Idx[0]], (*VerticesToUse)[Idx[1]], (*VerticesToUse)[Idx[2]] };
+					const FVec3 Normal = FVec3::CrossProduct(Vs[1] - Vs[0], Vs[2] - Vs[0]).GetUnsafeNormal();
+					OutPlanes.Add(TPlaneConcrete<FReal, 3>(Vs[0], Normal));
+					TArray<int32> FaceIndices;
+					FaceIndices.SetNum(3);
+					FaceIndices[0] = AddIndex(Idx[0]);
+					FaceIndices[1] = AddIndex(Idx[1]);
+					FaceIndices[2] = AddIndex(Idx[2]);
+					OutFaceIndices.Add(FaceIndices);
+				}
+
+				OutVertices.SetNum(IndexMap.Num());
+				for (const auto& Elem : IndexMap)
+				{
+					OutVertices[Elem.Value] = (*VerticesToUse)[Elem.Key];
+				}
+			}
+		}
+
+		UE_CLOG(OutVertices.Num() == 0, LogChaos, Warning, TEXT("Convex hull generation produced zero convex particles, collision will fail for this primitive."));
+	}
+
 }
diff --git a/Engine/Source/Runtime/Experimental/Chaos/Private/ChaosModule.cpp b/Engine/Source/Runtime/Experimental/Chaos/Private/ChaosModule.cpp
--- a/Engine/Source/Runtime/Experimental/Chaos/Private/ChaosModule.cpp
+++ b/Engine/Source/Runtime/Experimental/Chaos/Private/ChaosModule.cpp
@@ -5,3 +5,11 @@
 #include "Modules/ModuleManager.h"
 
 IMPLEMENT_MODULE(FChaosEngineModule, Chaos);
+
+void FChaosEngineModule::StartupModule()
+{
+	if (FModuleManager::Get().ModuleExists(TEXT("GeometryCore")))
+	{
+		FModuleManager::Get().LoadModule("GeometryCore");
+	}
+}
diff --git a/Engine/Source/Runtime/Experimental/Chaos/Public/Chaos/CollisionConvexMesh.h b/Engine/Source/Runtime/Experimental/Chaos/Public/Chaos/CollisionConvexMesh.h
--- a/Engine/Source/Runtime/Experimental/Chaos/Public/Chaos/CollisionConvexMesh.h
+++ b/Engine/Source/Runtime/Experimental/Chaos/Public/Chaos/CollisionConvexMesh.h
@@ -7,6 +7,7 @@
 #include "TriangleMesh.h"
 #include "Particles.h"
 #include "ChaosLog.h"
+#include "CompGeom/ConvexHull3.h"
 
 #define DEBUG_HULL_GENERATION 0
 
@@ -254,109 +255,7 @@ namespace Chaos
 
 	public:
 
-		static void Build(const TArray<FVec3>& InVertices, TArray <TPlaneConcrete<FReal, 3>>& OutPlanes, TArray<TArray<int32>>& OutFaceIndices, TArray<FVec3>& OutVertices, FAABB3& OutLocalBounds)
-		{
-			OutPlanes.Reset();
-			OutVertices.Reset();
-			OutLocalBounds = FAABB3::EmptyAABB();
-
-			const int32 NumVerticesIn = InVertices.Num();
-			if(NumVerticesIn == 0)
-			{
-				return;
-			}
-
-			const TArray<FVec3>* VerticesToUse = &InVertices;
-			TArray<FVec3> ModifiedVertices;
-
-			// For triangles and planar shapes, create a very thin prism as a convex
-			auto Inflate = [](const TArray<FVec3>& Source, TArray<FVec3>& Destination, const FVec3& Normal, FReal Inflation)
-			{
-				const int32 NumSource = Source.Num();
-				Destination.Reset();
-				Destination.SetNum(NumSource * 2);
-
-				for(int32 Index = 0; Index < NumSource; ++Index)
-				{
-					Destination[Index] = Source[Index];
-					Destination[NumSource + Index] = Source[Index] + Normal * Inflation;
-				}
-			};
-
-			FVec3 PlanarNormal(0);
-			if (NumVerticesIn == 3)
-			{
-				const bool bIsValidTriangle = IsValidTriangle(InVertices[0], InVertices[1], InVertices[2], PlanarNormal);
-
-				//TODO_SQ_IMPLEMENTATION: should do proper cleanup to avoid this
-				if(ensureMsgf(bIsValidTriangle, TEXT("FConvexBuilder::Build(): Generated invalid triangle!")))
-				{
-					Inflate(InVertices, ModifiedVertices, PlanarNormal, TriQuadPrismInflation());
-					VerticesToUse = &ModifiedVertices;
-					UE_LOG(LogChaos, Verbose, TEXT("Encountered a triangle in convex hull generation. Will prepare a prism of thickness %.5f in place of a triangle."), TriQuadPrismInflation());
-				}
-				else
-				{
-					return;
-				}
-			}
-			else if (IsPlanarShape(InVertices, PlanarNormal))
-			{
-				Inflate(InVertices, ModifiedVertices, PlanarNormal, TriQuadPrismInflation());
-				VerticesToUse = &ModifiedVertices;
-				UE_LOG(LogChaos, Verbose, TEXT("Encountered a planar shape in convex hull generation. Will prepare a prism of thickness %.5f in place of a triangle."), TriQuadPrismInflation());
-			}
-
-			const int32 NumVerticesToUse = VerticesToUse->Num();
-
-			OutLocalBounds = FAABB3((*VerticesToUse)[0], (*VerticesToUse)[0]);
-			for (int32 VertexIndex = 0; VertexIndex < NumVerticesToUse; ++VertexIndex)
-			{
-				OutLocalBounds.GrowToInclude((*VerticesToUse)[VertexIndex]);
-			}
-
-			if (NumVerticesToUse >= 4)
-			{
-				TArray<TVec3<int32>> Indices;
-				Params BuildParams;
-				BuildParams.HorizonEpsilon = Chaos::FConvexBuilder::SuggestEpsilon(*VerticesToUse);
-				BuildConvexHull(*VerticesToUse, Indices, BuildParams);
-				OutPlanes.Reserve(Indices.Num());
-				TMap<int32, int32> IndexMap; // maps original particle indices to output particle indices
-				int32 NewIdx = 0;
-
-				const auto AddIndex = [&IndexMap, &NewIdx](const int32 OriginalIdx)
-				{
-					if (int32* Idx = IndexMap.Find(OriginalIdx))
-					{
-						return *Idx;
-					}
-					IndexMap.Add(OriginalIdx, NewIdx);
-					return NewIdx++;
-				};
-
-				for(const TVec3<int32>& Idx : Indices)
-				{
-					FVec3 Vs[3] = { (*VerticesToUse)[Idx[0]], (*VerticesToUse)[Idx[1]], (*VerticesToUse)[Idx[2]] };
-					const FVec3 Normal = FVec3::CrossProduct(Vs[1] - Vs[0], Vs[2] - Vs[0]).GetUnsafeNormal();
-					OutPlanes.Add(TPlaneConcrete<FReal, 3>(Vs[0], Normal));
-					TArray<int32> FaceIndices;
-					FaceIndices.SetNum(3);
-					FaceIndices[0] = AddIndex(Idx[0]);
-					FaceIndices[1] = AddIndex(Idx[1]);
-					FaceIndices[2] = AddIndex(Idx[2]);
-					OutFaceIndices.Add(FaceIndices);
-				}
-
-				OutVertices.SetNum(IndexMap.Num());
-				for(const auto& Elem : IndexMap)
-				{
-					OutVertices[Elem.Value] = (*VerticesToUse)[Elem.Key];
-				}
-			}
-
-			UE_CLOG(OutVertices.Num() == 0, LogChaos, Warning, TEXT("Convex hull generation produced zero convex particles, collision will fail for this primitive."));
-		}
+		static CHAOS_API void Build(const TArray<FVec3>& InVertices, TArray <TPlaneConcrete<FReal, 3>>& OutPlanes, TArray<TArray<int32>>& OutFaceIndices, TArray<FVec3>& OutVertices, FAABB3& OutLocalBounds);
 
 		static void BuildConvexHull(const TArray<FVec3>& InVertices, TArray<TVec3<int32>>& OutIndices, const Params& InParams = Params())
 		{
@@ -645,6 +544,7 @@ namespace Chaos
 		static CHAOS_API int32 PerformGeometryReduction;
 		static CHAOS_API int32 VerticesThreshold;
 		static CHAOS_API int32 ComputeHorizonEpsilonFromMeshExtends;
+		static CHAOS_API bool bUseGeometryTConvexHull3;
 
 	private:
 
diff --git a/Engine/Source/Runtime/Experimental/Chaos/Public/ChaosModule.h b/Engine/Source/Runtime/Experimental/Chaos/Public/ChaosModule.h
--- a/Engine/Source/Runtime/Experimental/Chaos/Public/ChaosModule.h
+++ b/Engine/Source/Runtime/Experimental/Chaos/Public/ChaosModule.h
@@ -8,6 +8,6 @@ class CHAOS_API FChaosEngineModule : public IModuleInterface
 {
 public:
 
-	virtual void StartupModule() override {};
+	virtual void StartupModule() override;
 	virtual void ShutdownModule() override {};
 };
