GeometryCollection
-- reverting changes to PhysicsProxy initialization that caused problems with implicits setup in Headless Chaos

diff --git a/Engine/Source/Runtime/Experimental/Chaos/Private/PhysicsProxy/GeometryCollectionPhysicsProxy.cpp b/Engine/Source/Runtime/Experimental/Chaos/Private/PhysicsProxy/GeometryCollectionPhysicsProxy.cpp
--- a/Engine/Source/Runtime/Experimental/Chaos/Private/PhysicsProxy/GeometryCollectionPhysicsProxy.cpp
+++ b/Engine/Source/Runtime/Experimental/Chaos/Private/PhysicsProxy/GeometryCollectionPhysicsProxy.cpp
@@ -517,10 +517,23 @@ void FGeometryCollectionPhysicsProxy::InitializeDynamicCollection(FGeometryDynam
 	}
 
 	// Process Activity
-	const int32 NumTransforms = DynamicCollection.SimulatableParticles.Num();
-	for (int32 TransformIdx = 0; TransformIdx < NumTransforms; TransformIdx++)
 	{
-		DynamicCollection.SimulatableParticles[TransformIdx] = DynamicCollection.Active[TransformIdx];
+		const int32 NumTransforms = DynamicCollection.SimulatableParticles.Num();
+		if (!RestCollection.HasAttribute(FGeometryCollection::SimulatableParticlesAttribute, FTransformCollection::TransformGroup))
+		{
+			// If no simulation data is available then default to the simulation of just the rigid geometry.
+			for (int32 TransformIdx = 0; TransformIdx < NumTransforms; TransformIdx++)
+			{
+				if (DynamicCollection.Children[TransformIdx].Num())
+				{
+					DynamicCollection.SimulatableParticles[TransformIdx] = false;
+				}
+				else
+				{
+					DynamicCollection.SimulatableParticles[TransformIdx] = DynamicCollection.Active[TransformIdx];
+				}
+			}
+		}
 	}
 }
 
