Chaos: DirtyGrid Floating point accuracy issue fix

diff --git a/Engine/Source/Programs/HeadlessChaos/Private/HeadlessChaosTestBroadphase.cpp b/Engine/Source/Programs/HeadlessChaos/Private/HeadlessChaosTestBroadphase.cpp
--- a/Engine/Source/Programs/HeadlessChaos/Private/HeadlessChaosTestBroadphase.cpp
+++ b/Engine/Source/Programs/HeadlessChaos/Private/HeadlessChaosTestBroadphase.cpp
@@ -662,6 +662,37 @@ namespace ChaosTest
 			}
 			
 		}
+
+		// Test a case that failed before (with an assert)
+		{
+			FAABBTreeDirtyGridCVars::DirtyElementGridCellSize = 1000;
+			FAABBTreeDirtyGridCVars::DirtyElementMaxCellCapacity = 7;
+			
+			TUniquePtr<TBox<FReal, 3>> Box;
+			auto Boxes = BuildBoxes(Box, 100, FVec3(1, 1, 1), FVec3(-3000, -1000, -50)); // Just one box
+			TreeType Spatial{};
+			Spatial.UpdateElement(0, Boxes->WorldSpaceInflatedBounds(0), true);
+
+			// Move the Box
+			Boxes = BuildBoxes(Box, 100, FVec3(1, 1, 1), FVec3(-4000, -1000, -50)); // Change position of box
+			Spatial.UpdateElement(0, Boxes->WorldSpaceInflatedBounds(0), true); // Check for no ensures
+
+			// Move the Box
+			Boxes = BuildBoxes(Box, 100, FVec3(1, 1, 1), FVec3(3000, 1000, -50)); // Change position of box
+			Spatial.UpdateElement(0, Boxes->WorldSpaceInflatedBounds(0), true);
+
+			Boxes = BuildBoxes(Box, 100, FVec3(1, 1, 1), FVec3(4000, 1000, -50)); // Change position of box
+			Spatial.UpdateElement(0, Boxes->WorldSpaceInflatedBounds(0), true); // Check for no ensures
+
+			// Move the Box
+			Boxes = BuildBoxes(Box, 100, FVec3(1, 1, 1), FVec3(-10000003000.0f, -1000, -50)); // Change position of box
+			Spatial.UpdateElement(0, Boxes->WorldSpaceInflatedBounds(0), true); // Check for no ensures
+
+			// Move the Box
+			Boxes = BuildBoxes(Box, 100, FVec3(1, 1, 1), FVec3(-10000004000.0f, -1000, -50)); // Change position of box
+			Spatial.UpdateElement(0, Boxes->WorldSpaceInflatedBounds(0), true); // Check for no ensures
+			
+		}
 		
 
 		// Restore CVARS
diff --git a/Engine/Source/Runtime/Experimental/Chaos/Public/Chaos/AABBTree.h b/Engine/Source/Runtime/Experimental/Chaos/Public/Chaos/AABBTree.h
--- a/Engine/Source/Runtime/Experimental/Chaos/Public/Chaos/AABBTree.h
+++ b/Engine/Source/Runtime/Experimental/Chaos/Public/Chaos/AABBTree.h
@@ -571,16 +571,19 @@ public:
 		return true;
 	}
 
-	// Returns true if there was enough space in the cell to add the new dirty element index and if the element is not already in the cell 
+	// Returns true if there was enough space in the cell to add the new dirty element index or if the element was already added (This second condition should not happen)
 	//(The second condition should never be true for the current implementation)
 	FORCEINLINE_DEBUGGABLE bool AddNewDirtyParticleIndexToGridCell(int32 Hash, int32 NewDirtyIndex)
 	{
 		DirtyGridHashEntry* HashEntry = CellHashToFlatArray.Find(Hash);
 		if (HashEntry)
 		{
-			if (HashEntry->Count < DirtyElementMaxCellCapacity && ensure(InsertValueIntoSortedSubArray(FlattenedCellArrayOfDirtyIndices, NewDirtyIndex, HashEntry->Index, HashEntry->Count)))
+			if (HashEntry->Count < DirtyElementMaxCellCapacity)
 			{
-				++(HashEntry->Count);
+				if (ensure(InsertValueIntoSortedSubArray(FlattenedCellArrayOfDirtyIndices, NewDirtyIndex, HashEntry->Index, HashEntry->Count)))
+				{
+					++(HashEntry->Count);
+				}
 				return true;
 			}
 		}
diff --git a/Engine/Source/Runtime/Experimental/Chaos/Public/Chaos/AABBTreeDirtyGridUtils.h b/Engine/Source/Runtime/Experimental/Chaos/Public/Chaos/AABBTreeDirtyGridUtils.h
--- a/Engine/Source/Runtime/Experimental/Chaos/Public/Chaos/AABBTreeDirtyGridUtils.h
+++ b/Engine/Source/Runtime/Experimental/Chaos/Public/Chaos/AABBTreeDirtyGridUtils.h
@@ -18,24 +18,33 @@ namespace Chaos
 		return Intermediate;
 	}
 
-	FORCEINLINE_DEBUGGABLE int32 HashCoordinates(FReal Xcoordinate, FReal Ycoordinate, FReal DirtyElementGridCellSizeInv)
+	FORCEINLINE_DEBUGGABLE int32 GetDirtyCellIndexFromWorldCoordinate(FReal Coordinate, FReal DirtyElementGridCellSizeInv)
 	{
-		// Requirement: Hash should change for adjacent cells
-		uint16 X = (uint16)FMath::Floor(Xcoordinate * DirtyElementGridCellSizeInv);
-		uint16 Y = (uint16)FMath::Floor(Ycoordinate * DirtyElementGridCellSizeInv);
+		FReal CellIndex = Coordinate * DirtyElementGridCellSizeInv;
+		// Bias values just enough to remove floating point integer inconsistencies
+		// Assume maximum of 1ULP error in DirtyElementGridCellSizeInv
+		FReal FloatingPointMaxCellIndexError = FMath::Abs(CellIndex) * std::numeric_limits<FReal>::epsilon();
+		return (int32)(FMath::Floor(CellIndex + FloatingPointMaxCellIndexError));
+	}
 
-		return (int32)(InterleaveWithZeros(X) | (InterleaveWithZeros(Y) << 1));
+	FORCEINLINE_DEBUGGABLE int32 HashCell(int32 XCell, int32 YCell)
+	{
+		return (int32)(InterleaveWithZeros((uint16)XCell) | (InterleaveWithZeros((uint16)YCell) << 1));
 	}
 
-	FORCEINLINE_DEBUGGABLE int32 GetDirtyCellIndexFromWorldCoordinate(FReal Coordinate, FReal DirtyElementGridCellSizeInv)
+	FORCEINLINE_DEBUGGABLE int32 HashCoordinates(FReal Xcoordinate, FReal Ycoordinate, FReal DirtyElementGridCellSizeInv)
 	{
-		return (int32)(FMath::Floor(Coordinate * DirtyElementGridCellSizeInv));
+		// Requirement: Hash should change for adjacent cells
+		int32 X = GetDirtyCellIndexFromWorldCoordinate(Xcoordinate,DirtyElementGridCellSizeInv);
+		int32 Y = GetDirtyCellIndexFromWorldCoordinate(Ycoordinate,DirtyElementGridCellSizeInv);
+
+		return HashCell(X, Y);
 	}
 
 	FORCEINLINE_DEBUGGABLE bool TooManyOverlapQueryCells(const TAABB<FReal, 3>& AABB, FReal DirtyElementGridCellSizeInv, int32 MaximumOverlap)
 	{
-		int32 XsampleCount = (int32)(FMath::Floor(AABB.Max().X * DirtyElementGridCellSizeInv)) - (int32)(FMath::Floor(AABB.Min().X * DirtyElementGridCellSizeInv)) + 1;
-		int32 YsampleCount = (int32)(FMath::Floor(AABB.Max().Y * DirtyElementGridCellSizeInv)) - (int32)(FMath::Floor(AABB.Min().Y * DirtyElementGridCellSizeInv)) + 1;
+		int32 XsampleCount = GetDirtyCellIndexFromWorldCoordinate(AABB.Max().X, DirtyElementGridCellSizeInv) - GetDirtyCellIndexFromWorldCoordinate(AABB.Min().X, DirtyElementGridCellSizeInv) + 1;
+		int32 YsampleCount = GetDirtyCellIndexFromWorldCoordinate(AABB.Max().Y, DirtyElementGridCellSizeInv) - GetDirtyCellIndexFromWorldCoordinate(AABB.Min().Y, DirtyElementGridCellSizeInv) + 1;
 
 		if (XsampleCount * YsampleCount <= MaximumOverlap)
 		{
@@ -47,8 +56,8 @@ namespace Chaos
 	template <typename FunctionType>
 	FORCEINLINE_DEBUGGABLE void DoForOverlappedCells(const TAABB<FReal, 3>& AABB, FReal DirtyElementGridCellSize, FReal DirtyElementGridCellSizeInv, FunctionType Function)
 	{
-		int32 XsampleCount = (int32)(FMath::Floor(AABB.Max().X * DirtyElementGridCellSizeInv)) - (int32)(FMath::Floor(AABB.Min().X * DirtyElementGridCellSizeInv)) + 1;
-		int32 YsampleCount = (int32)(FMath::Floor(AABB.Max().Y * DirtyElementGridCellSizeInv)) - (int32)(FMath::Floor(AABB.Min().Y * DirtyElementGridCellSizeInv)) + 1;
+		int32 XsampleCount = GetDirtyCellIndexFromWorldCoordinate(AABB.Max().X, DirtyElementGridCellSizeInv) - GetDirtyCellIndexFromWorldCoordinate(AABB.Min().X, DirtyElementGridCellSizeInv) + 1;
+		int32 YsampleCount = GetDirtyCellIndexFromWorldCoordinate(AABB.Max().Y, DirtyElementGridCellSizeInv) - GetDirtyCellIndexFromWorldCoordinate(AABB.Min().Y, DirtyElementGridCellSizeInv) + 1;
 
 		FReal CurrentX = AABB.Min().X;
 		for (int32 XsampleIndex = 0; XsampleIndex < XsampleCount; XsampleIndex++)
@@ -95,7 +104,7 @@ namespace Chaos
 			{
 				if (!(X >= OldCellStartX && X <= OldCellEndX && Y >= OldCellStartY && Y <= OldCellEndY))
 				{
-					if (!Function(HashCoordinates((FReal)X * DirtyElementGridCellSize, (FReal)Y * DirtyElementGridCellSize, DirtyElementGridCellSizeInv)))
+					if (!Function(HashCell(X, Y)))
 						return false; // Failure early out
 				}
 			}
@@ -292,14 +301,14 @@ namespace Chaos
 		FReal X1 = XStartPointExpanded;
 		FReal Y1 = YStartPointExpanded;
 
-		int32 CurrentCellIndexX0 = (int32)FMath::Floor(X0 * DirtyElementGridCellSizeInv);
-		int32 CurrentCellIndexY0 = (int32)FMath::Floor(Y0 * DirtyElementGridCellSizeInv);
+		int32 CurrentCellIndexX0 = GetDirtyCellIndexFromWorldCoordinate(X0, DirtyElementGridCellSizeInv);
+		int32 CurrentCellIndexY0 = GetDirtyCellIndexFromWorldCoordinate(Y0, DirtyElementGridCellSizeInv);
 
-		int32 CurrentCellIndexX1 = (int32)FMath::Floor(X1 * DirtyElementGridCellSizeInv);
-		int32 CurrentCellIndexY1 = (int32)FMath::Floor(Y1 * DirtyElementGridCellSizeInv);
+		int32 CurrentCellIndexX1 = GetDirtyCellIndexFromWorldCoordinate(X1, DirtyElementGridCellSizeInv);
+		int32 CurrentCellIndexY1 = GetDirtyCellIndexFromWorldCoordinate(Y1, DirtyElementGridCellSizeInv);
 
-		int32 LastCellIndexX = (int32)FMath::Floor(XEndPointExpanded * DirtyElementGridCellSizeInv);
-		int32 LastCellIndexY = (int32)FMath::Floor(YEndPointExpanded * DirtyElementGridCellSizeInv);
+		int32 LastCellIndexX = GetDirtyCellIndexFromWorldCoordinate(XEndPointExpanded, DirtyElementGridCellSizeInv);
+		int32 LastCellIndexY = GetDirtyCellIndexFromWorldCoordinate(YEndPointExpanded, DirtyElementGridCellSizeInv);
 
 		bool Done = false;
 		while (!Done)
@@ -418,11 +427,11 @@ namespace Chaos
 	{
 		const TVector<FReal, 3>& EndPoint = StartPoint + Length * Dir;
 
-		int32 FirstCellIndexX = (int32)FMath::Floor(StartPoint.X * DirtyElementGridCellSizeInv);
-		int32 FirstCellIndexY = (int32)FMath::Floor(StartPoint.Y * DirtyElementGridCellSizeInv);
+		int32 FirstCellIndexX = GetDirtyCellIndexFromWorldCoordinate(StartPoint.X, DirtyElementGridCellSizeInv);
+		int32 FirstCellIndexY = GetDirtyCellIndexFromWorldCoordinate(StartPoint.Y, DirtyElementGridCellSizeInv);
 
-		int32 LastCellIndexX = (int32)FMath::Floor(EndPoint.X * DirtyElementGridCellSizeInv);
-		int32 LastCellIndexY = (int32)FMath::Floor(EndPoint.Y * DirtyElementGridCellSizeInv);
+		int32 LastCellIndexX = GetDirtyCellIndexFromWorldCoordinate(EndPoint.X, DirtyElementGridCellSizeInv);
+		int32 LastCellIndexY = GetDirtyCellIndexFromWorldCoordinate(EndPoint.Y, DirtyElementGridCellSizeInv);
 
 		// This will be equal to the Manhattan distance 
 		int CellCount = FMath::Abs(FirstCellIndexX - LastCellIndexX) + FMath::Abs(FirstCellIndexY - LastCellIndexY);
@@ -441,11 +450,11 @@ namespace Chaos
 	{
 		const TVector<FReal, 3>& EndPoint = StartPoint + Length * Dir;
 
-		int32 CurrentCellIndexX = (int32)FMath::Floor(StartPoint.X * DirtyElementGridCellSizeInv);
-		int32 CurrentCellIndexY = (int32)FMath::Floor(StartPoint.Y * DirtyElementGridCellSizeInv);
+		int32 CurrentCellIndexX = GetDirtyCellIndexFromWorldCoordinate(StartPoint.X, DirtyElementGridCellSizeInv);
+		int32 CurrentCellIndexY = GetDirtyCellIndexFromWorldCoordinate(StartPoint.Y, DirtyElementGridCellSizeInv);
 
-		int32 LastCellIndexX = (int32)FMath::Floor(EndPoint.X * DirtyElementGridCellSizeInv);
-		int32 LastCellIndexY = (int32)FMath::Floor(EndPoint.Y * DirtyElementGridCellSizeInv);
+		int32 LastCellIndexX = GetDirtyCellIndexFromWorldCoordinate(EndPoint.X, DirtyElementGridCellSizeInv);
+		int32 LastCellIndexY = GetDirtyCellIndexFromWorldCoordinate(EndPoint.Y, DirtyElementGridCellSizeInv);
 
 		FReal DeltaX = EndPoint.X - StartPoint.X;
 		FReal DeltaY = EndPoint.Y - StartPoint.Y;
@@ -489,7 +498,7 @@ namespace Chaos
 		bool Done = false;
 		while (!Done)
 		{
-			InFunction(HashCoordinates((FReal)CurrentCellIndexX * DirtyElementGridCellSize, (FReal)CurrentCellIndexY * DirtyElementGridCellSize, DirtyElementGridCellSizeInv));
+			InFunction(HashCell(CurrentCellIndexX, CurrentCellIndexY));
 			FReal CrossingVerticleCellBorderT = std::numeric_limits<FReal>::max();
 			FReal CrossingHorizontalCellBorderT = std::numeric_limits<FReal>::max();
 			if (!DxTooSmall)
