Make sure ApplyPlasticityLimits is not processed when one of the invoved particle is disabled
This was causing weird stretches in skeletal mesh when disabling bodies at start and enabling them will moving the skeletal mesh with a non zero velocity

diff --git a/Engine/Source/Runtime/Experimental/Chaos/Private/Chaos/PBDJointConstraints.cpp b/Engine/Source/Runtime/Experimental/Chaos/Private/Chaos/PBDJointConstraints.cpp
--- a/Engine/Source/Runtime/Experimental/Chaos/Private/Chaos/PBDJointConstraints.cpp
+++ b/Engine/Source/Runtime/Experimental/Chaos/Private/Chaos/PBDJointConstraints.cpp
@@ -19,8 +19,6 @@
 #include "PBDJointSolverGaussSeidel.ispc.generated.h"
 #endif
 
-//PRAGMA_DISABLE_OPTIMIZATION
-
 namespace Chaos
 {
 	DECLARE_CYCLE_STAT(TEXT("Joints::Sort"), STAT_Joints_Sort, STATGROUP_ChaosJoint);
@@ -1503,6 +1501,10 @@ namespace Chaos
 		GetConstrainedParticleIndices(ConstraintIndex, Index0, Index1);
 		FGenericParticleHandle Particle0 = FGenericParticleHandle(ConstraintParticles[ConstraintIndex][Index0]);
 		FGenericParticleHandle Particle1 = FGenericParticleHandle(ConstraintParticles[ConstraintIndex][Index1]);
+		if (Particle0->Disabled() || Particle1->Disabled())
+		{
+			return;
+		}
 
 		const FTransformPair &ConstraintFramesLocal = JointSettings.ConnectorTransforms;
 		FTransformPair ConstraintFramesGlobal(ConstraintFramesLocal[Index0] * FRigidTransform3(Particle0->P(), Particle0->Q()), ConstraintFramesLocal[Index1] * FRigidTransform3(Particle1->P(), Particle1->Q()));
