Chaos : Implicits
- Copy path for ImplicitObjectTransformed

diff --git a/Engine/Source/Runtime/Experimental/Chaos/Public/Chaos/ImplicitObjectTransformed.h b/Engine/Source/Runtime/Experimental/Chaos/Public/Chaos/ImplicitObjectTransformed.h
--- a/Engine/Source/Runtime/Experimental/Chaos/Public/Chaos/ImplicitObjectTransformed.h
+++ b/Engine/Source/Runtime/Experimental/Chaos/Public/Chaos/ImplicitObjectTransformed.h
@@ -90,6 +90,13 @@ public:
 		return ImplicitObjectType::Transformed;
 	}
 
+	virtual TUniquePtr<FImplicitObject> Copy() const
+	{
+		TUniquePtr<FImplicitObject> ImplicitCopy = MObject->Copy();
+		return TUniquePtr<FImplicitObject>(new TImplicitObjectTransformed<T, d>(MoveTemp(ImplicitCopy), MTransform));
+	}
+
+
 	const FImplicitObject* GetTransformedObject() const
 	{
 		return MObject.Get();
