Removed extra-size parameters for events that don't need them.

diff --git a/Engine/Source/Runtime/Core/Private/ProfilingDebugging/TagTrace.cpp b/Engine/Source/Runtime/Core/Private/ProfilingDebugging/TagTrace.cpp
--- a/Engine/Source/Runtime/Core/Private/ProfilingDebugging/TagTrace.cpp
+++ b/Engine/Source/Runtime/Core/Private/ProfilingDebugging/TagTrace.cpp
@@ -257,7 +257,7 @@ int32 FTagTrace::AnnounceFNameTag(const FName& Name)
 int32 FTagTrace::AnnounceCustomTag(int32 Tag, int32 ParentTag, const ANSICHAR* Display)
 {		
 	const uint32 DisplayLen = FCStringAnsi::Strlen(Display);
-	UE_TRACE_LOG(Memory, TagSpec, MemAllocChannel, DisplayLen * sizeof(ANSICHAR))
+	UE_TRACE_LOG(Memory, TagSpec, MemAllocChannel)
 		<< TagSpec.Tag(Tag)
 		<< TagSpec.Parent(ParentTag)
 		<< TagSpec.Display(Display, DisplayLen);
diff --git a/Engine/Source/Runtime/Experimental/Chaos/Private/ChaosVisualDebugger/ChaosVisualDebuggerTrace.cpp b/Engine/Source/Runtime/Experimental/Chaos/Private/ChaosVisualDebugger/ChaosVisualDebuggerTrace.cpp
--- a/Engine/Source/Runtime/Experimental/Chaos/Private/ChaosVisualDebugger/ChaosVisualDebuggerTrace.cpp
+++ b/Engine/Source/Runtime/Experimental/Chaos/Private/ChaosVisualDebugger/ChaosVisualDebuggerTrace.cpp
@@ -15,7 +15,7 @@ using namespace Chaos;
 //A very minimal test debug log function
 void ChaosVisualDebugger::ParticlePositionLog(const FVec3& Position)
 {
-	UE_TRACE_LOG(Physics, ParticlePosition, PhysicsChannel, 0)
+	UE_TRACE_LOG(Physics, ParticlePosition, PhysicsChannel)
 					<< ParticlePosition.PositionX(Position.X)
 					<< ParticlePosition.PositionY(Position.Y)
 					<< ParticlePosition.PositionZ(Position.Z);
