GeometryFramework: move UDynamicMesh and DynamicMeshComponent to new Runtime/Experimental/GeometryFramework module

#rnx

diff --git a/Engine/Plugins/Experimental/MeshModelingToolset/Source/MeshModelingTools/MeshModelingTools.Build.cs b/Engine/Plugins/Experimental/MeshModelingToolset/Source/MeshModelingTools/MeshModelingTools.Build.cs
--- a/Engine/Plugins/Experimental/MeshModelingToolset/Source/MeshModelingTools/MeshModelingTools.Build.cs
+++ b/Engine/Plugins/Experimental/MeshModelingToolset/Source/MeshModelingTools/MeshModelingTools.Build.cs
@@ -35,6 +35,7 @@ public class MeshModelingTools : ModuleRules
 				"Eigen",
                 "InteractiveToolsFramework",
 				"GeometryCore",
+				"GeometryFramework",
 				"GeometryAlgorithms",
 				"DynamicMesh",
 				"MeshConversion",
diff --git a/Engine/Plugins/Experimental/MeshModelingToolset/Source/MeshModelingToolsEditorOnly/MeshModelingToolsEditorOnly.Build.cs b/Engine/Plugins/Experimental/MeshModelingToolset/Source/MeshModelingToolsEditorOnly/MeshModelingToolsEditorOnly.Build.cs
--- a/Engine/Plugins/Experimental/MeshModelingToolset/Source/MeshModelingToolsEditorOnly/MeshModelingToolsEditorOnly.Build.cs
+++ b/Engine/Plugins/Experimental/MeshModelingToolset/Source/MeshModelingToolsEditorOnly/MeshModelingToolsEditorOnly.Build.cs
@@ -31,6 +31,7 @@ public class MeshModelingToolsEditorOnly : ModuleRules
 				"StaticMeshDescription",
 				"SkeletalMeshDescription",
 				"GeometryCore",
+				"GeometryFramework",
 				"DynamicMesh",
 				"MeshConversion",
 				"MeshModelingTools",
diff --git a/Engine/Plugins/Experimental/MeshModelingToolset/Source/ModelingComponents/ModelingComponents.Build.cs b/Engine/Plugins/Experimental/MeshModelingToolset/Source/ModelingComponents/ModelingComponents.Build.cs
--- a/Engine/Plugins/Experimental/MeshModelingToolset/Source/ModelingComponents/ModelingComponents.Build.cs
+++ b/Engine/Plugins/Experimental/MeshModelingToolset/Source/ModelingComponents/ModelingComponents.Build.cs
@@ -31,6 +31,7 @@ public class ModelingComponents : ModuleRules
                 "MeshDescription",
 				"StaticMeshDescription",
 				"GeometryCore",
+				"GeometryFramework",
 				"GeometryAlgorithms",
 				"DynamicMesh",
 				"MeshConversion",
diff --git a/Engine/Plugins/Experimental/MeshModelingToolset/Source/ModelingComponents/Private/Changes/MeshChange.cpp b/Engine/Plugins/Experimental/MeshModelingToolset/Source/ModelingComponents/Private/Changes/MeshChange.cpp
deleted file mode 100644
diff --git a/Engine/Plugins/Experimental/MeshModelingToolset/Source/ModelingComponents/Private/Changes/MeshReplacementChange.cpp b/Engine/Plugins/Experimental/MeshModelingToolset/Source/ModelingComponents/Private/Changes/MeshReplacementChange.cpp
deleted file mode 100644
diff --git a/Engine/Plugins/Experimental/MeshModelingToolset/Source/ModelingComponents/Private/Changes/MeshVertexChange.cpp b/Engine/Plugins/Experimental/MeshModelingToolset/Source/ModelingComponents/Private/Changes/MeshVertexChange.cpp
deleted file mode 100644
diff --git a/Engine/Plugins/Experimental/MeshModelingToolset/Source/ModelingComponents/Private/Components/BaseDynamicMeshComponent.cpp b/Engine/Plugins/Experimental/MeshModelingToolset/Source/ModelingComponents/Private/Components/BaseDynamicMeshComponent.cpp
deleted file mode 100644
diff --git a/Engine/Plugins/Experimental/MeshModelingToolset/Source/ModelingComponents/Private/Components/BaseDynamicMeshSceneProxy.h b/Engine/Plugins/Experimental/MeshModelingToolset/Source/ModelingComponents/Private/Components/BaseDynamicMeshSceneProxy.h
deleted file mode 100644
diff --git a/Engine/Plugins/Experimental/MeshModelingToolset/Source/ModelingComponents/Private/Components/DynamicMeshComponent.cpp b/Engine/Plugins/Experimental/MeshModelingToolset/Source/ModelingComponents/Private/Components/DynamicMeshComponent.cpp
deleted file mode 100644
diff --git a/Engine/Plugins/Experimental/MeshModelingToolset/Source/ModelingComponents/Private/Components/DynamicMeshSceneProxy.h b/Engine/Plugins/Experimental/MeshModelingToolset/Source/ModelingComponents/Private/Components/DynamicMeshSceneProxy.h
deleted file mode 100644
diff --git a/Engine/Plugins/Experimental/MeshModelingToolset/Source/ModelingComponents/Private/Drawing/MeshRenderDecomposition.cpp b/Engine/Plugins/Experimental/MeshModelingToolset/Source/ModelingComponents/Private/Drawing/MeshRenderDecomposition.cpp
deleted file mode 100644
diff --git a/Engine/Plugins/Experimental/MeshModelingToolset/Source/ModelingComponents/Private/DynamicMeshActor.cpp b/Engine/Plugins/Experimental/MeshModelingToolset/Source/ModelingComponents/Private/DynamicMeshActor.cpp
deleted file mode 100644
diff --git a/Engine/Plugins/Experimental/MeshModelingToolset/Source/ModelingComponents/Private/UDynamicMesh.cpp b/Engine/Plugins/Experimental/MeshModelingToolset/Source/ModelingComponents/Private/UDynamicMesh.cpp
deleted file mode 100644
diff --git a/Engine/Plugins/Experimental/MeshModelingToolset/Source/ModelingComponents/Public/Changes/MeshChange.h b/Engine/Plugins/Experimental/MeshModelingToolset/Source/ModelingComponents/Public/Changes/MeshChange.h
deleted file mode 100644
diff --git a/Engine/Plugins/Experimental/MeshModelingToolset/Source/ModelingComponents/Public/Changes/MeshReplacementChange.h b/Engine/Plugins/Experimental/MeshModelingToolset/Source/ModelingComponents/Public/Changes/MeshReplacementChange.h
deleted file mode 100644
diff --git a/Engine/Plugins/Experimental/MeshModelingToolset/Source/ModelingComponents/Public/Changes/MeshVertexChange.h b/Engine/Plugins/Experimental/MeshModelingToolset/Source/ModelingComponents/Public/Changes/MeshVertexChange.h
deleted file mode 100644
diff --git a/Engine/Plugins/Experimental/MeshModelingToolset/Source/ModelingComponents/Public/Components/BaseDynamicMeshComponent.h b/Engine/Plugins/Experimental/MeshModelingToolset/Source/ModelingComponents/Public/Components/BaseDynamicMeshComponent.h
deleted file mode 100644
diff --git a/Engine/Plugins/Experimental/MeshModelingToolset/Source/ModelingComponents/Public/Components/DynamicMeshComponent.h b/Engine/Plugins/Experimental/MeshModelingToolset/Source/ModelingComponents/Public/Components/DynamicMeshComponent.h
deleted file mode 100644
diff --git a/Engine/Plugins/Experimental/MeshModelingToolset/Source/ModelingComponents/Public/Components/OctreeDynamicMeshComponent.h b/Engine/Plugins/Experimental/MeshModelingToolset/Source/ModelingComponents/Public/Components/OctreeDynamicMeshComponent.h
--- a/Engine/Plugins/Experimental/MeshModelingToolset/Source/ModelingComponents/Public/Components/OctreeDynamicMeshComponent.h
+++ b/Engine/Plugins/Experimental/MeshModelingToolset/Source/ModelingComponents/Public/Components/OctreeDynamicMeshComponent.h
@@ -3,7 +3,7 @@
 #pragma once
 
 #include "CoreMinimal.h"
-#include "BaseDynamicMeshComponent.h"
+#include "Components/BaseDynamicMeshComponent.h"
 #include "MeshConversionOptions.h"
 #include "UDynamicMesh.h"
 #include "DynamicMesh/DynamicMeshOctree3.h"
diff --git a/Engine/Plugins/Experimental/MeshModelingToolset/Source/ModelingComponents/Public/Drawing/MeshRenderDecomposition.h b/Engine/Plugins/Experimental/MeshModelingToolset/Source/ModelingComponents/Public/Drawing/MeshRenderDecomposition.h
deleted file mode 100644
diff --git a/Engine/Plugins/Experimental/MeshModelingToolset/Source/ModelingComponents/Public/DynamicMeshActor.h b/Engine/Plugins/Experimental/MeshModelingToolset/Source/ModelingComponents/Public/DynamicMeshActor.h
deleted file mode 100644
diff --git a/Engine/Plugins/Experimental/MeshModelingToolset/Source/ModelingComponents/Public/UDynamicMesh.h b/Engine/Plugins/Experimental/MeshModelingToolset/Source/ModelingComponents/Public/UDynamicMesh.h
deleted file mode 100644
diff --git a/Engine/Plugins/Experimental/MeshModelingToolset/Source/ModelingComponentsEditorOnly/ModelingComponentsEditorOnly.Build.cs b/Engine/Plugins/Experimental/MeshModelingToolset/Source/ModelingComponentsEditorOnly/ModelingComponentsEditorOnly.Build.cs
--- a/Engine/Plugins/Experimental/MeshModelingToolset/Source/ModelingComponentsEditorOnly/ModelingComponentsEditorOnly.Build.cs
+++ b/Engine/Plugins/Experimental/MeshModelingToolset/Source/ModelingComponentsEditorOnly/ModelingComponentsEditorOnly.Build.cs
@@ -30,6 +30,7 @@ public class ModelingComponentsEditorOnly : ModuleRules
                 "MeshDescription",
 				"StaticMeshDescription",
 				"GeometryCore",
+				"GeometryFramework",
 				"GeometryAlgorithms",
 				"DynamicMesh",
 				"MeshConversion",
diff --git a/Engine/Plugins/Experimental/ModelingToolsEditorMode/Source/ModelingToolsEditorMode/ModelingToolsEditorMode.Build.cs b/Engine/Plugins/Experimental/ModelingToolsEditorMode/Source/ModelingToolsEditorMode/ModelingToolsEditorMode.Build.cs
--- a/Engine/Plugins/Experimental/ModelingToolsEditorMode/Source/ModelingToolsEditorMode/ModelingToolsEditorMode.Build.cs
+++ b/Engine/Plugins/Experimental/ModelingToolsEditorMode/Source/ModelingToolsEditorMode/ModelingToolsEditorMode.Build.cs
@@ -53,6 +53,7 @@ public class ModelingToolsEditorMode : ModuleRules
 				"StaticMeshDescription",
                 "InteractiveToolsFramework",
 				"EditorInteractiveToolsFramework",
+				"GeometryFramework",
 				"ModelingComponents",
 				"ModelingComponentsEditorOnly",
 				"MeshModelingTools",
diff --git a/Engine/Source/Runtime/Experimental/GeometryFramework/GeometryFramework.Build.cs b/Engine/Source/Runtime/Experimental/GeometryFramework/GeometryFramework.Build.cs
new file mode 100644
--- /dev/null
+++ b/Engine/Source/Runtime/Experimental/GeometryFramework/GeometryFramework.Build.cs
@@ -0,0 +1,28 @@
+// Copyright Epic Games, Inc. All Rights Reserved.
+
+using UnrealBuildTool;
+
+public class GeometryFramework : ModuleRules
+{	
+	public GeometryFramework(ReadOnlyTargetRules Target) : base(Target)
+	{
+        //PCHUsage = ModuleRules.PCHUsageMode.UseExplicitOrSharedPCHs;
+
+        PublicDependencyModuleNames.AddRange(
+			new string[] {
+				"Core",
+				"CoreUObject",
+				"Engine",
+				"RenderCore",
+				"RHI",
+				"PhysicsCore",
+				"InteractiveToolsFramework",
+				"MeshDescription",
+				"StaticMeshDescription",
+				"SkeletalMeshDescription",
+				"GeometryCore",
+				"MeshConversion"
+			}
+		);
+	}
+}
diff --git a/Engine/Source/Runtime/Experimental/GeometryFramework/Private/Changes/MeshChange.cpp b/Engine/Source/Runtime/Experimental/GeometryFramework/Private/Changes/MeshChange.cpp
new file mode 100644
--- /dev/null
+++ b/Engine/Source/Runtime/Experimental/GeometryFramework/Private/Changes/MeshChange.cpp
@@ -0,0 +1,51 @@
+// Copyright Epic Games, Inc. All Rights Reserved.
+
+#include "Changes/MeshChange.h"
+#include "DynamicMesh/DynamicMesh3.h"
+
+using namespace UE::Geometry;
+
+FMeshChange::FMeshChange()
+{
+}
+
+FMeshChange::FMeshChange(TUniquePtr<FDynamicMeshChange> DynamicMeshChangeIn)
+{
+	DynamicMeshChange = MoveTemp(DynamicMeshChangeIn);
+}
+
+void FMeshChange::Apply(UObject* Object)
+{
+	IMeshCommandChangeTarget* ChangeTarget = CastChecked<IMeshCommandChangeTarget>(Object);
+	ChangeTarget->ApplyChange(this, false);
+
+	if (OnChangeAppliedFunc)
+	{
+		OnChangeAppliedFunc(this, Object, true);
+	}
+}
+
+void FMeshChange::Revert(UObject* Object)
+{
+	IMeshCommandChangeTarget* ChangeTarget = CastChecked<IMeshCommandChangeTarget>(Object);
+	ChangeTarget->ApplyChange(this, true);
+
+	if (OnChangeAppliedFunc)
+	{
+		OnChangeAppliedFunc(this, Object, false);
+	}
+}
+
+
+FString FMeshChange::ToString() const
+{
+	return FString(TEXT("Mesh Change"));
+}
+
+
+void FMeshChange::ApplyChangeToMesh(FDynamicMesh3* Mesh, bool bRevert) const
+{
+	DynamicMeshChange->Apply(Mesh, bRevert);
+}
+
+
diff --git a/Engine/Source/Runtime/Experimental/GeometryFramework/Private/Changes/MeshReplacementChange.cpp b/Engine/Source/Runtime/Experimental/GeometryFramework/Private/Changes/MeshReplacementChange.cpp
new file mode 100644
--- /dev/null
+++ b/Engine/Source/Runtime/Experimental/GeometryFramework/Private/Changes/MeshReplacementChange.cpp
@@ -0,0 +1,45 @@
+// Copyright Epic Games, Inc. All Rights Reserved.
+
+#include "Changes/MeshReplacementChange.h"
+#include "DynamicMesh/DynamicMesh3.h"
+
+using namespace UE::Geometry;
+
+FMeshReplacementChange::FMeshReplacementChange()
+{
+}
+
+FMeshReplacementChange::FMeshReplacementChange(TSharedPtr<const FDynamicMesh3, ESPMode::ThreadSafe> BeforeIn, TSharedPtr<const FDynamicMesh3, ESPMode::ThreadSafe> AfterIn)
+{
+	Before = BeforeIn;
+	After = AfterIn;
+}
+
+void FMeshReplacementChange::Apply(UObject* Object)
+{
+	IMeshReplacementCommandChangeTarget* ChangeTarget = CastChecked<IMeshReplacementCommandChangeTarget>(Object);
+	ChangeTarget->ApplyChange(this, false);
+
+	if (OnChangeAppliedFunc)
+	{
+		OnChangeAppliedFunc(this, Object, true);
+	}
+}
+
+void FMeshReplacementChange::Revert(UObject* Object)
+{
+	IMeshReplacementCommandChangeTarget* ChangeTarget = CastChecked<IMeshReplacementCommandChangeTarget>(Object);
+	ChangeTarget->ApplyChange(this, true);
+
+	if (OnChangeAppliedFunc)
+	{
+		OnChangeAppliedFunc(this, Object, false);
+	}
+}
+
+
+FString FMeshReplacementChange::ToString() const
+{
+	return FString(TEXT("Mesh Change"));
+}
+
diff --git a/Engine/Source/Runtime/Experimental/GeometryFramework/Private/Changes/MeshVertexChange.cpp b/Engine/Source/Runtime/Experimental/GeometryFramework/Private/Changes/MeshVertexChange.cpp
new file mode 100644
--- /dev/null
+++ b/Engine/Source/Runtime/Experimental/GeometryFramework/Private/Changes/MeshVertexChange.cpp
@@ -0,0 +1,249 @@
+// Copyright Epic Games, Inc. All Rights Reserved.
+
+#include "Changes/MeshVertexChange.h"
+#include "DynamicMesh/DynamicMesh3.h"
+
+#include "Components/BaseDynamicMeshComponent.h"
+
+using namespace UE::Geometry;
+
+void FMeshVertexChange::Apply(UObject* Object)
+{
+	IMeshVertexCommandChangeTarget* ChangeTarget = CastChecked<IMeshVertexCommandChangeTarget>(Object);
+	ChangeTarget->ApplyChange(this, false);
+}
+
+void FMeshVertexChange::Revert(UObject* Object)
+{
+	IMeshVertexCommandChangeTarget* ChangeTarget = CastChecked<IMeshVertexCommandChangeTarget>(Object);
+	ChangeTarget->ApplyChange(this, true);
+}
+
+
+FString FMeshVertexChange::ToString() const
+{
+	return FString(TEXT("Mesh Vertex Change"));
+}
+
+
+FMeshVertexChangeBuilder::FMeshVertexChangeBuilder()
+{
+	Change = MakeUnique<FMeshVertexChange>();
+}
+
+
+FMeshVertexChangeBuilder::FMeshVertexChangeBuilder(EMeshVertexChangeComponents Components)
+{
+	Change = MakeUnique<FMeshVertexChange>();
+
+	bSavePositions = ((Components & EMeshVertexChangeComponents::VertexPositions) != EMeshVertexChangeComponents::None);
+	Change->bHaveVertexPositions = bSavePositions;
+
+	bSaveColors = ((Components & EMeshVertexChangeComponents::VertexColors) != EMeshVertexChangeComponents::None);
+	Change->bHaveVertexColors = bSaveColors;
+	
+	bSaveOverlayNormals = ((Components & EMeshVertexChangeComponents::OverlayNormals) != EMeshVertexChangeComponents::None);
+	Change->bHaveOverlayNormals = bSaveOverlayNormals;
+}
+
+void FMeshVertexChangeBuilder::UpdateVertex(int VertexID, const FVector3d& OldPosition, const FVector3d& NewPosition)
+{
+	const int* FoundIndex = SavedVertices.Find(VertexID);
+	if (FoundIndex == nullptr)
+	{
+		int NewIndex = Change->Vertices.Num();
+		SavedVertices.Add(VertexID, NewIndex);
+		Change->Vertices.Add(VertexID);
+		Change->OldPositions.Add(OldPosition);
+		Change->NewPositions.Add(NewPosition);
+	} 
+	else
+	{
+		Change->NewPositions[*FoundIndex] = NewPosition;
+	}
+}
+
+
+void FMeshVertexChangeBuilder::UpdateVertexColor(int32 VertexID, const FVector3f& OldColor, const FVector3f& NewColor)
+{
+	const int* FoundIndex = SavedVertices.Find(VertexID);
+	if (FoundIndex == nullptr)
+	{
+		int NewIndex = Change->Vertices.Num();
+		SavedVertices.Add(VertexID, NewIndex);
+		Change->Vertices.Add(VertexID);
+		Change->OldColors.Add(OldColor);
+		Change->NewColors.Add(NewColor);
+	}
+	else
+	{
+		Change->NewColors[*FoundIndex] = NewColor;
+	}
+}
+
+
+void FMeshVertexChangeBuilder::UpdateVertexFinal(int VertexID, const FVector3d& NewPosition)
+{
+	check(SavedVertices.Contains(VertexID));
+
+	const int* Index = SavedVertices.Find(VertexID);
+	if ( Index != nullptr )
+	{
+		Change->NewPositions[*Index] = NewPosition;
+	}
+}
+
+
+
+
+
+void FMeshVertexChangeBuilder::SaveVertexInitial(const FDynamicMesh3* Mesh, int32 VertexID)
+{
+	const int32* FoundIndex = SavedVertices.Find(VertexID);
+	if (FoundIndex == nullptr)
+	{
+		int32 Index = Change->Vertices.Num();
+		SavedVertices.Add(VertexID, Index);
+		Change->Vertices.Add(VertexID);
+		if (bSavePositions)
+		{
+			FVector3d Pos = Mesh->GetVertex(VertexID);
+			Change->OldPositions.Add(Pos);
+			Change->NewPositions.Add(Pos);
+		}
+		if (bSaveColors)
+		{
+			FVector3f Color = Mesh->GetVertexColor(VertexID);
+			Change->OldColors.Add(Color);
+			Change->NewColors.Add(Color);
+		}
+		if (OnNewVertexSaved)
+		{
+			OnNewVertexSaved(VertexID, Index);
+		}
+	}
+	else
+	{
+		int32 Index = *FoundIndex;
+		if (bSavePositions)
+		{
+			Change->NewPositions[Index] = Mesh->GetVertex(VertexID);
+		}
+		if (bSaveColors)
+		{
+			Change->NewColors[Index] = Mesh->GetVertexColor(VertexID);
+		}
+	}
+}
+
+
+void FMeshVertexChangeBuilder::SaveVertexFinal(const FDynamicMesh3* Mesh, int32 VertexID)
+{
+	const int32* FoundIndex = SavedVertices.Find(VertexID);
+	if (FoundIndex != nullptr)
+	{
+		int32 Index = *FoundIndex;
+		if (bSavePositions)
+		{
+			Change->NewPositions[Index] = Mesh->GetVertex(VertexID);
+		}
+		if (bSaveColors)
+		{
+			Change->NewColors[Index] = Mesh->GetVertexColor(VertexID);
+		}
+	}
+}
+
+
+
+
+
+
+
+
+
+
+
+
+void FMeshVertexChangeBuilder::UpdateOverlayNormal(int ElementID, const FVector3f& OldNormal, const FVector3f& NewNormal)
+{
+	const int* FoundIndex = SavedNormalElements.Find(ElementID);
+	if (FoundIndex == nullptr)
+	{
+		int NewIndex = Change->Normals.Num();
+		SavedNormalElements.Add(ElementID, NewIndex);
+		Change->Normals.Add(ElementID);
+		Change->OldNormals.Add(OldNormal);
+		Change->NewNormals.Add(NewNormal);
+	}
+	else
+	{
+		Change->NewNormals[*FoundIndex] = NewNormal;
+	}
+}
+
+void FMeshVertexChangeBuilder::UpdateOverlayNormalFinal(int ElementID, const FVector3f& NewNormal)
+{
+	check(SavedNormalElements.Contains(ElementID));
+	const int* Index = SavedNormalElements.Find(ElementID);
+	if (Index != nullptr)
+	{
+		Change->NewNormals[*Index] = NewNormal;
+	}
+}
+
+
+
+void FMeshVertexChangeBuilder::SaveOverlayNormals(const FDynamicMesh3* Mesh, const TArray<int>& ElementIDs, bool bInitial)
+{
+	if (Mesh->HasAttributes() == false || Mesh->Attributes()->PrimaryNormals() == nullptr)
+	{
+		return;
+	}
+	const FDynamicMeshNormalOverlay* Overlay = Mesh->Attributes()->PrimaryNormals();
+
+	int Num = ElementIDs.Num();
+	if (bInitial)
+	{
+		for (int k = 0; k < Num; ++k)
+		{
+			FVector3f Normal = Overlay->GetElement(ElementIDs[k]);
+			UpdateOverlayNormal(ElementIDs[k], Normal, Normal);
+		}
+	}
+	else
+	{
+		for (int k = 0; k < Num; ++k)
+		{
+			FVector3f Normal = Overlay->GetElement(ElementIDs[k]);
+			UpdateOverlayNormalFinal(ElementIDs[k], Normal);
+		}
+	}
+}
+
+
+void FMeshVertexChangeBuilder::SaveOverlayNormals(const FDynamicMesh3* Mesh, const TSet<int>& ElementIDs, bool bInitial)
+{
+	if (Mesh->HasAttributes() == false || Mesh->Attributes()->PrimaryNormals() == nullptr)
+	{
+		return;
+	}
+	const FDynamicMeshNormalOverlay* Overlay = Mesh->Attributes()->PrimaryNormals();
+
+	if (bInitial)
+	{
+		for (int ElementID : ElementIDs)
+		{
+			FVector3f Normal = Overlay->GetElement(ElementID);
+			UpdateOverlayNormal(ElementID, Normal, Normal);
+		}
+	}
+	else
+	{
+		for (int ElementID : ElementIDs)
+		{
+			FVector3f Normal = Overlay->GetElement(ElementID);
+			UpdateOverlayNormalFinal(ElementID, Normal);
+		}
+	}
+}
\ No newline at end of file
diff --git a/Engine/Source/Runtime/Experimental/GeometryFramework/Private/Components/BaseDynamicMeshComponent.cpp b/Engine/Source/Runtime/Experimental/GeometryFramework/Private/Components/BaseDynamicMeshComponent.cpp
new file mode 100644
--- /dev/null
+++ b/Engine/Source/Runtime/Experimental/GeometryFramework/Private/Components/BaseDynamicMeshComponent.cpp
@@ -0,0 +1,139 @@
+// Copyright Epic Games, Inc. All Rights Reserved. 
+
+#include "Components/BaseDynamicMeshComponent.h"
+#include "Components/BaseDynamicMeshSceneProxy.h"
+
+using namespace UE::Geometry;
+
+UBaseDynamicMeshComponent::UBaseDynamicMeshComponent(const FObjectInitializer& ObjectInitializer)
+	: Super(ObjectInitializer)
+{
+}
+
+
+
+void UBaseDynamicMeshComponent::SetShadowsEnabled(bool bEnabled)
+{
+	// finish any drawing so that we can be certain our SceneProxy is no longer in use before we rebuild it below
+	FlushRenderingCommands();
+
+	SetCastShadow(bEnabled);
+	//bCastDynamicShadow = bEnabled;
+
+	// apparently SceneProxy has to be fully rebuilt to change shadow state
+
+	// this marks the SceneProxy for rebuild, but not immediately, and possibly allows bad things to happen
+	// before the end of the frame
+	//MarkRenderStateDirty();
+
+	// force immediate rebuild of the SceneProxy
+	if (IsRegistered())
+	{
+		ReregisterComponent();
+	}
+}
+
+
+void UBaseDynamicMeshComponent::SetOverrideRenderMaterial(UMaterialInterface* Material)
+{
+	if (OverrideRenderMaterial != Material)
+	{
+		OverrideRenderMaterial = Material;
+		NotifyMaterialSetUpdated();
+	}
+}
+
+void UBaseDynamicMeshComponent::ClearOverrideRenderMaterial()
+{
+	if (OverrideRenderMaterial != nullptr)
+	{
+		OverrideRenderMaterial = nullptr;
+		NotifyMaterialSetUpdated();
+	}
+}
+
+
+
+
+
+void UBaseDynamicMeshComponent::SetSecondaryRenderMaterial(UMaterialInterface* Material)
+{
+	if (SecondaryRenderMaterial != Material)
+	{
+		SecondaryRenderMaterial = Material;
+		NotifyMaterialSetUpdated();
+	}
+}
+
+void UBaseDynamicMeshComponent::ClearSecondaryRenderMaterial()
+{
+	if (SecondaryRenderMaterial != nullptr)
+	{
+		SecondaryRenderMaterial = nullptr;
+		NotifyMaterialSetUpdated();
+	}
+}
+
+
+
+void UBaseDynamicMeshComponent::SetSecondaryBuffersVisibility(bool bSecondaryVisibility)
+{
+	bDrawSecondaryBuffers = bSecondaryVisibility;
+}
+
+bool UBaseDynamicMeshComponent::GetSecondaryBuffersVisibility() const
+{
+	return bDrawSecondaryBuffers;
+}
+
+
+
+
+int32 UBaseDynamicMeshComponent::GetNumMaterials() const
+{
+	return BaseMaterials.Num();
+}
+
+UMaterialInterface* UBaseDynamicMeshComponent::GetMaterial(int32 ElementIndex) const 
+{
+	return (ElementIndex >= 0 && ElementIndex < BaseMaterials.Num()) ? BaseMaterials[ElementIndex] : nullptr;
+}
+
+FMaterialRelevance UBaseDynamicMeshComponent::GetMaterialRelevance(ERHIFeatureLevel::Type InFeatureLevel) const
+{
+	FMaterialRelevance Result = UMeshComponent::GetMaterialRelevance(InFeatureLevel);
+	if (OverrideRenderMaterial)
+	{
+		Result |= OverrideRenderMaterial->GetRelevance_Concurrent(InFeatureLevel);
+	}
+	if (SecondaryRenderMaterial)
+	{
+		Result |= SecondaryRenderMaterial->GetRelevance_Concurrent(InFeatureLevel);
+	}
+	return Result;
+}
+
+void UBaseDynamicMeshComponent::SetMaterial(int32 ElementIndex, UMaterialInterface* Material)
+{
+	check(ElementIndex >= 0);
+	if (ElementIndex >= BaseMaterials.Num())
+	{
+		BaseMaterials.SetNum(ElementIndex + 1, false);
+	}
+	BaseMaterials[ElementIndex] = Material;
+}
+
+
+void UBaseDynamicMeshComponent::GetUsedMaterials(TArray<UMaterialInterface*>& OutMaterials, bool bGetDebugMaterials) const
+{
+	UMeshComponent::GetUsedMaterials(OutMaterials, bGetDebugMaterials);
+	if (OverrideRenderMaterial != nullptr)
+	{
+		OutMaterials.Add(OverrideRenderMaterial);
+	}
+	if (SecondaryRenderMaterial != nullptr)
+	{
+		OutMaterials.Add(SecondaryRenderMaterial);
+	}
+}
+
diff --git a/Engine/Source/Runtime/Experimental/GeometryFramework/Private/Components/DynamicMeshComponent.cpp b/Engine/Source/Runtime/Experimental/GeometryFramework/Private/Components/DynamicMeshComponent.cpp
new file mode 100644
--- /dev/null
+++ b/Engine/Source/Runtime/Experimental/GeometryFramework/Private/Components/DynamicMeshComponent.cpp
@@ -0,0 +1,1130 @@
+// Copyright Epic Games, Inc. All Rights Reserved. 
+
+#include "Components/DynamicMeshComponent.h"
+#include "PrimitiveSceneProxy.h"
+#include "MaterialShared.h"
+#include "Engine/CollisionProfile.h"
+#include "Materials/Material.h"
+#include "Async/Async.h"
+#include "Engine/CollisionProfile.h"
+
+#include "DynamicMesh/DynamicMeshAttributeSet.h"
+#include "DynamicMesh/MeshNormals.h"
+#include "MeshDescriptionToDynamicMesh.h"
+
+#include "Changes/MeshVertexChange.h"
+#include "Changes/MeshChange.h"
+#include "DynamicMesh/MeshTransforms.h"
+
+// default proxy for this component
+#include "Components/DynamicMeshSceneProxy.h"
+
+#include "ExplicitUseGeometryMathTypes.h"		// using UE::Geometry::(math types)
+using namespace UE::Geometry;
+
+namespace
+{
+	// probably should be something defined for the whole tool framework...
+#if WITH_EDITOR
+	static EAsyncExecution DynamicMeshComponentAsyncExecTarget = EAsyncExecution::LargeThreadPool;
+#else
+	static EAsyncExecution DynamicMeshComponentAsyncExecTarget = EAsyncExecution::ThreadPool;
+#endif
+}
+
+
+
+
+UDynamicMeshComponent::UDynamicMeshComponent(const FObjectInitializer& ObjectInitializer)
+	: Super(ObjectInitializer)
+{
+	PrimaryComponentTick.bCanEverTick = false;
+
+	SetCollisionProfileName(UCollisionProfile::NoCollision_ProfileName);
+
+	MeshObject = CreateDefaultSubobject<UDynamicMesh>(TEXT("DynamicMesh"));
+	//MeshObject->SetFlags(RF_Transactional);
+
+	MeshObjectChangedHandle = MeshObject->OnMeshChanged().AddUObject(this, &UDynamicMeshComponent::OnMeshObjectChanged);
+
+	ResetProxy();
+}
+
+
+void UDynamicMeshComponent::PostLoad()
+{
+	Super::PostLoad();
+
+	check(MeshObject != nullptr);
+	MeshObjectChangedHandle = MeshObject->OnMeshChanged().AddUObject(this, &UDynamicMeshComponent::OnMeshObjectChanged);
+
+	ResetProxy();
+
+	// make sure BodySetup is created
+	GetBodySetup();
+}
+
+
+#if WITH_EDITOR
+void UDynamicMeshComponent::PostEditChangeProperty(FPropertyChangedEvent& PropertyChangedEvent)
+{
+	Super::PostEditChangeProperty(PropertyChangedEvent);
+
+	const FName PropName = PropertyChangedEvent.GetPropertyName();
+	if ( (PropName == GET_MEMBER_NAME_CHECKED(UDynamicMeshComponent, bEnableComplexCollision)) ||
+		(PropName == GET_MEMBER_NAME_CHECKED(UDynamicMeshComponent, CollisionType)) ||
+		(PropName == GET_MEMBER_NAME_CHECKED(UDynamicMeshComponent, bDeferCollisionUpdates))  )
+	{
+		if (bDeferCollisionUpdates)
+		{
+			InvalidatePhysicsData();
+		}
+		else
+		{
+			RebuildPhysicsData();
+		}
+	}
+}
+#endif
+
+
+void UDynamicMeshComponent::SetMesh(UE::Geometry::FDynamicMesh3&& MoveMesh)
+{
+	MeshObject->SetMesh(MoveTemp(MoveMesh));
+}
+
+
+void UDynamicMeshComponent::ProcessMesh(
+	TFunctionRef<void(const UE::Geometry::FDynamicMesh3&)> ProcessFunc ) const
+{
+	MeshObject->ProcessMesh(ProcessFunc);
+}
+
+
+void UDynamicMeshComponent::EditMesh(TFunctionRef<void(UE::Geometry::FDynamicMesh3&)> EditFunc,
+										   EDynamicMeshComponentRenderUpdateMode UpdateMode )
+{
+	MeshObject->EditMesh(EditFunc);
+	if (UpdateMode != EDynamicMeshComponentRenderUpdateMode::NoUpdate)
+	{
+		NotifyMeshUpdated();
+	}
+}
+
+
+void UDynamicMeshComponent::SetRenderMeshPostProcessor(TUniquePtr<IRenderMeshPostProcessor> Processor)
+{
+	RenderMeshPostProcessor = MoveTemp(Processor);
+	if (RenderMeshPostProcessor)
+	{
+		if (!RenderMesh)
+		{
+			RenderMesh = MakeUnique<FDynamicMesh3>(*GetMesh());
+		}
+	}
+	else
+	{
+		// No post processor, no render mesh
+		RenderMesh = nullptr;
+	}
+}
+
+FDynamicMesh3* UDynamicMeshComponent::GetRenderMesh()
+{
+	if (RenderMeshPostProcessor && RenderMesh)
+	{
+		return RenderMesh.Get();
+	}
+	else
+	{
+		return GetMesh();
+	}
+}
+
+const FDynamicMesh3* UDynamicMeshComponent::GetRenderMesh() const
+{
+	if (RenderMeshPostProcessor && RenderMesh)
+	{
+		return RenderMesh.Get();
+	}
+	else
+	{
+		return GetMesh();
+	}
+}
+
+
+
+
+void UDynamicMeshComponent::ApplyTransform(const UE::Geometry::FTransform3d& Transform, bool bInvert)
+{
+	MeshObject->EditMesh([&](FDynamicMesh3& EditMesh)
+	{
+		if (bInvert)
+		{
+			MeshTransforms::ApplyTransformInverse(EditMesh, Transform);
+		}
+		else
+		{
+			MeshTransforms::ApplyTransform(EditMesh, Transform);
+		}
+	}, EDynamicMeshChangeType::DeformationEdit);
+}
+
+
+
+void UDynamicMeshComponent::SetTangentsType(EDynamicMeshComponentTangentsMode NewTangentsType)
+{
+	if (NewTangentsType != TangentsType)
+	{
+		TangentsType = NewTangentsType;
+		InvalidateAutoCalculatedTangents();
+	}
+}
+
+void UDynamicMeshComponent::InvalidateAutoCalculatedTangents() 
+{ 
+	bAutoCalculatedTangentsValid = false; 
+}
+
+const UE::Geometry::FMeshTangentsf* UDynamicMeshComponent::GetAutoCalculatedTangents() 
+{ 
+	if (TangentsType == EDynamicMeshComponentTangentsMode::AutoCalculated)
+	{
+		UpdateAutoCalculatedTangents();
+		return (bAutoCalculatedTangentsValid) ? &AutoCalculatedTangents : nullptr;
+	}
+	return nullptr;
+}
+
+void UDynamicMeshComponent::UpdateAutoCalculatedTangents()
+{
+	if (TangentsType == EDynamicMeshComponentTangentsMode::AutoCalculated && bAutoCalculatedTangentsValid == false)
+	{
+		GetDynamicMesh()->ProcessMesh([&](const FDynamicMesh3& Mesh)
+		{
+			const FDynamicMeshUVOverlay* UVOverlay = Mesh.Attributes()->PrimaryUV();
+			const FDynamicMeshNormalOverlay* NormalOverlay = Mesh.Attributes()->PrimaryNormals();
+			AutoCalculatedTangents.SetMesh(&Mesh);
+			AutoCalculatedTangents.ComputeTriVertexTangents(NormalOverlay, UVOverlay, FComputeTangentsOptions());
+			AutoCalculatedTangents.SetMesh(nullptr);
+		});
+
+		bAutoCalculatedTangentsValid = true;
+	}
+}
+
+
+
+
+void UDynamicMeshComponent::UpdateLocalBounds()
+{
+	LocalBounds = GetMesh()->GetBounds(true);
+	if (LocalBounds.MaxDim() <= 0)
+	{
+		LocalBounds = FAxisAlignedBox3d(FVector3d::Zero(), FMathf::ZeroTolerance);
+	}
+}
+
+FDynamicMeshSceneProxy* UDynamicMeshComponent::GetCurrentSceneProxy() 
+{ 
+	if (bProxyValid)
+	{
+		return (FDynamicMeshSceneProxy*)SceneProxy;
+	}
+	return nullptr;
+}
+
+
+void UDynamicMeshComponent::ResetProxy()
+{
+	bProxyValid = false;
+	InvalidateAutoCalculatedTangents();
+
+	// Need to recreate scene proxy to send it over
+	MarkRenderStateDirty();
+	UpdateLocalBounds();
+	UpdateBounds();
+
+	// this seems speculative, ie we may not actually have a mesh update, but we currently ResetProxy() in lots
+	// of places where that is what it means
+	GetDynamicMesh()->PostRealtimeUpdate();
+}
+
+void UDynamicMeshComponent::NotifyMeshUpdated()
+{
+	if (RenderMeshPostProcessor)
+	{
+		RenderMeshPostProcessor->ProcessMesh(*GetMesh(), *RenderMesh);
+	}
+
+	ResetProxy();
+}
+
+
+void UDynamicMeshComponent::FastNotifyColorsUpdated()
+{
+	// should not be using fast paths if we have to run mesh postprocessor
+	if (ensure(!RenderMeshPostProcessor) == false)
+	{
+		RenderMeshPostProcessor->ProcessMesh(*GetMesh(), *RenderMesh);
+		ResetProxy();
+		return;
+	}
+
+	FDynamicMeshSceneProxy* Proxy = GetCurrentSceneProxy();
+	if (Proxy)
+	{
+		if (HasTriangleColorFunction() && Proxy->bUsePerTriangleColor == false )
+		{
+			Proxy->bUsePerTriangleColor = true;
+			Proxy->PerTriangleColorFunc = [this](const FDynamicMesh3* MeshIn, int TriangleID) { return GetTriangleColor(MeshIn, TriangleID); };
+		} 
+		else if ( !HasTriangleColorFunction() && Proxy->bUsePerTriangleColor == true)
+		{
+			Proxy->bUsePerTriangleColor = false;
+			Proxy->PerTriangleColorFunc = nullptr;
+		}
+
+		Proxy->FastUpdateVertices(false, false, true, false);
+		//MarkRenderDynamicDataDirty();
+	}
+	else
+	{
+		ResetProxy();
+	}
+}
+
+
+
+void UDynamicMeshComponent::FastNotifyPositionsUpdated(bool bNormals, bool bColors, bool bUVs)
+{
+	// should not be using fast paths if we have to run mesh postprocessor
+	if (ensure(!RenderMeshPostProcessor) == false)
+	{
+		RenderMeshPostProcessor->ProcessMesh(*GetMesh(), *RenderMesh);
+		ResetProxy();
+		return;
+	}
+
+	FDynamicMeshSceneProxy* Proxy = GetCurrentSceneProxy();
+	if (Proxy)
+	{
+		// calculate bounds while we are updating vertices
+		TFuture<void> UpdateBoundsCalc;
+		UpdateBoundsCalc = Async(DynamicMeshComponentAsyncExecTarget, [this]()
+		{
+			TRACE_CPUPROFILER_EVENT_SCOPE(SimpleDynamicMeshComponent_FastPositionsUpdate_AsyncBoundsUpdate);
+			UpdateLocalBounds();
+		});
+
+		GetCurrentSceneProxy()->FastUpdateVertices(true, bNormals, bColors, bUVs);
+
+		//MarkRenderDynamicDataDirty();
+		MarkRenderTransformDirty();
+		UpdateBoundsCalc.Wait();
+		UpdateBounds();
+
+		GetDynamicMesh()->PostRealtimeUpdate();
+	}
+	else
+	{
+		ResetProxy();
+	}
+}
+
+
+void UDynamicMeshComponent::FastNotifyVertexAttributesUpdated(bool bNormals, bool bColors, bool bUVs)
+{
+	// should not be using fast paths if we have to run mesh postprocessor
+	if (ensure(!RenderMeshPostProcessor) == false)
+	{
+		RenderMeshPostProcessor->ProcessMesh(*GetMesh(), *RenderMesh);
+		ResetProxy();
+		return;
+	}
+
+	FDynamicMeshSceneProxy* Proxy = GetCurrentSceneProxy();
+	if (Proxy && ensure(bNormals || bColors || bUVs) )
+	{
+		GetCurrentSceneProxy()->FastUpdateVertices(false, bNormals, bColors, bUVs);
+		//MarkRenderDynamicDataDirty();
+		//MarkRenderTransformDirty();
+
+		GetDynamicMesh()->PostRealtimeUpdate();
+	}
+	else
+	{
+		ResetProxy();
+	}
+}
+
+
+void UDynamicMeshComponent::FastNotifyVertexAttributesUpdated(EMeshRenderAttributeFlags UpdatedAttributes)
+{
+	// should not be using fast paths if we have to run mesh postprocessor
+	if (ensure(!RenderMeshPostProcessor) == false)
+	{
+		RenderMeshPostProcessor->ProcessMesh(*GetMesh(), *RenderMesh);
+		ResetProxy();
+		return;
+	}
+
+	FDynamicMeshSceneProxy* Proxy = GetCurrentSceneProxy();
+	if (Proxy && ensure(UpdatedAttributes != EMeshRenderAttributeFlags::None))
+	{
+		bool bPositions = (UpdatedAttributes & EMeshRenderAttributeFlags::Positions) != EMeshRenderAttributeFlags::None;
+
+		// calculate bounds while we are updating vertices
+		TFuture<void> UpdateBoundsCalc;
+		if (bPositions)
+		{
+			UpdateBoundsCalc = Async(DynamicMeshComponentAsyncExecTarget, [this]()
+			{
+				TRACE_CPUPROFILER_EVENT_SCOPE(SimpleDynamicMeshComponent_FastVertexAttribUpdate_AsyncBoundsUpdate);
+				UpdateLocalBounds();
+			});
+		}
+
+		GetCurrentSceneProxy()->FastUpdateVertices(bPositions,
+			(UpdatedAttributes & EMeshRenderAttributeFlags::VertexNormals) != EMeshRenderAttributeFlags::None,
+			(UpdatedAttributes & EMeshRenderAttributeFlags::VertexColors) != EMeshRenderAttributeFlags::None,
+			(UpdatedAttributes & EMeshRenderAttributeFlags::VertexUVs) != EMeshRenderAttributeFlags::None);
+
+		if (bPositions)
+		{
+			MarkRenderTransformDirty();
+			UpdateBoundsCalc.Wait();
+			UpdateBounds();
+		}
+
+		GetDynamicMesh()->PostRealtimeUpdate();
+	}
+	else
+	{
+		ResetProxy();
+	}
+}
+
+
+void UDynamicMeshComponent::FastNotifyUVsUpdated()
+{
+	FastNotifyVertexAttributesUpdated(EMeshRenderAttributeFlags::VertexUVs);
+}
+
+
+
+void UDynamicMeshComponent::FastNotifySecondaryTrianglesChanged()
+{
+	// should not be using fast paths if we have to run mesh postprocessor
+	if (ensure(!RenderMeshPostProcessor) == false)
+	{
+		RenderMeshPostProcessor->ProcessMesh(*GetMesh(), *RenderMesh);
+		ResetProxy();
+		return;
+	}
+
+	FDynamicMeshSceneProxy* Proxy = GetCurrentSceneProxy();
+	if (Proxy)
+	{
+		GetCurrentSceneProxy()->FastUpdateAllIndexBuffers();
+		GetDynamicMesh()->PostRealtimeUpdate();
+	}
+	else
+	{
+		ResetProxy();
+	}
+}
+
+
+void UDynamicMeshComponent::FastNotifyTriangleVerticesUpdated(const TArray<int32>& Triangles, EMeshRenderAttributeFlags UpdatedAttributes)
+{
+	// should not be using fast paths if we have to run mesh postprocessor
+	if (ensure(!RenderMeshPostProcessor) == false)
+	{
+		RenderMeshPostProcessor->ProcessMesh(*GetMesh(), *RenderMesh);
+		ResetProxy();
+		return;
+	}
+
+	bool bUpdateSecondarySort = (SecondaryTriFilterFunc) &&
+		((UpdatedAttributes & EMeshRenderAttributeFlags::SecondaryIndexBuffers) != EMeshRenderAttributeFlags::None);
+
+	FDynamicMeshSceneProxy* Proxy = GetCurrentSceneProxy();
+	if (!Proxy)
+	{
+		ResetProxy();
+	}
+	else if ( ! Decomposition )
+	{
+		FastNotifyVertexAttributesUpdated(UpdatedAttributes);
+		if (bUpdateSecondarySort)
+		{
+			Proxy->FastUpdateAllIndexBuffers();
+		}
+		GetDynamicMesh()->PostRealtimeUpdate();
+	}
+	else
+	{
+		// compute list of sets to update
+		TArray<int32> UpdatedSets;
+		{
+			TRACE_CPUPROFILER_EVENT_SCOPE(SimpleDynamicMeshComponent_FastVertexUpdate_FindSets);
+			for (int32 tid : Triangles)
+			{
+				int32 SetID = Decomposition->GetGroupForTriangle(tid);
+				UpdatedSets.AddUnique(SetID);
+			}
+		}
+
+		bool bPositions = (UpdatedAttributes & EMeshRenderAttributeFlags::Positions) != EMeshRenderAttributeFlags::None;
+
+		// calculate bounds while we are updating vertices
+		TFuture<void> UpdateBoundsCalc;
+		if (bPositions)
+		{
+			UpdateBoundsCalc = Async(DynamicMeshComponentAsyncExecTarget, [this]()
+			{
+				TRACE_CPUPROFILER_EVENT_SCOPE(SimpleDynamicMeshComponent_FastVertexUpdate_AsyncBoundsUpdate);
+				UpdateLocalBounds();
+			});
+		}
+
+		// update the render buffers
+		{
+			TRACE_CPUPROFILER_EVENT_SCOPE(SimpleDynamicMeshComponent_FastVertexUpdate_ApplyUpdate);
+			Proxy->FastUpdateVertices(UpdatedSets, bPositions,
+				(UpdatedAttributes & EMeshRenderAttributeFlags::VertexNormals) != EMeshRenderAttributeFlags::None,
+				(UpdatedAttributes & EMeshRenderAttributeFlags::VertexColors) != EMeshRenderAttributeFlags::None,
+				(UpdatedAttributes & EMeshRenderAttributeFlags::VertexUVs) != EMeshRenderAttributeFlags::None);
+		}
+
+		if (bUpdateSecondarySort)
+		{
+			Proxy->FastUpdateIndexBuffers(UpdatedSets);
+		}
+
+		if (bPositions)
+		{
+			TRACE_CPUPROFILER_EVENT_SCOPE(SimpleDynamicMeshComponent_FastVertexUpdate_FinalPositionsUpdate);
+			MarkRenderTransformDirty();
+			UpdateBoundsCalc.Wait();
+			UpdateBounds();
+		}
+
+		GetDynamicMesh()->PostRealtimeUpdate();
+	}
+}
+
+
+
+
+void UDynamicMeshComponent::FastNotifyTriangleVerticesUpdated(const TSet<int32>& Triangles, EMeshRenderAttributeFlags UpdatedAttributes)
+{
+	// should not be using fast paths if we have to run mesh postprocessor
+	if (ensure(!RenderMeshPostProcessor) == false)
+	{
+		RenderMeshPostProcessor->ProcessMesh(*GetMesh(), *RenderMesh);
+		ResetProxy();
+		return;
+	}
+
+	bool bUpdateSecondarySort = (SecondaryTriFilterFunc) &&
+		((UpdatedAttributes & EMeshRenderAttributeFlags::SecondaryIndexBuffers) != EMeshRenderAttributeFlags::None);
+
+	FDynamicMeshSceneProxy* Proxy = GetCurrentSceneProxy();
+	if (!Proxy)
+	{
+		ResetProxy();
+	}
+	else if (!Decomposition)
+	{
+		FastNotifyVertexAttributesUpdated(UpdatedAttributes);
+		if (bUpdateSecondarySort)
+		{
+			Proxy->FastUpdateAllIndexBuffers();
+		}
+		GetDynamicMesh()->PostRealtimeUpdate();
+	}
+	else
+	{
+		// compute list of sets to update
+		TArray<int32> UpdatedSets;
+		{
+			TRACE_CPUPROFILER_EVENT_SCOPE(SimpleDynamicMeshComponent_FastVertexUpdate_FindSets);
+			for (int32 tid : Triangles)
+			{
+				int32 SetID = Decomposition->GetGroupForTriangle(tid);
+				UpdatedSets.AddUnique(SetID);
+			}
+		}
+
+		bool bPositions = (UpdatedAttributes & EMeshRenderAttributeFlags::Positions) != EMeshRenderAttributeFlags::None;
+
+		// calculate bounds while we are updating vertices
+		TFuture<void> UpdateBoundsCalc;
+		if (bPositions)
+		{
+			UpdateBoundsCalc = Async(DynamicMeshComponentAsyncExecTarget, [this]()
+			{
+				TRACE_CPUPROFILER_EVENT_SCOPE(SimpleDynamicMeshComponent_FastVertexUpdate_AsyncBoundsUpdate);
+				UpdateLocalBounds();
+			});
+		}
+
+		// update the render buffers
+		{
+			TRACE_CPUPROFILER_EVENT_SCOPE(SimpleDynamicMeshComponent_FastVertexUpdate_ApplyUpdate);
+			Proxy->FastUpdateVertices(UpdatedSets, bPositions,
+				(UpdatedAttributes & EMeshRenderAttributeFlags::VertexNormals) != EMeshRenderAttributeFlags::None,
+				(UpdatedAttributes & EMeshRenderAttributeFlags::VertexColors) != EMeshRenderAttributeFlags::None,
+				(UpdatedAttributes & EMeshRenderAttributeFlags::VertexUVs) != EMeshRenderAttributeFlags::None);
+		}
+
+		{
+			TRACE_CPUPROFILER_EVENT_SCOPE(SimpleDynamicMeshComponent_FastVertexUpdate_UpdateIndexBuffers);
+			if (bUpdateSecondarySort)
+			{
+				Proxy->FastUpdateIndexBuffers(UpdatedSets);
+			}
+		}
+
+		// finish up, have to wait for background bounds recalculation here
+		if (bPositions)
+		{
+			TRACE_CPUPROFILER_EVENT_SCOPE(SimpleDynamicMeshComponent_FastVertexUpdate_FinalPositionsUpdate);
+			MarkRenderTransformDirty();
+			UpdateBoundsCalc.Wait();
+			UpdateBounds();
+		}
+
+		GetDynamicMesh()->PostRealtimeUpdate();
+	}
+}
+
+
+
+/**
+ * Compute the combined bounding-box of the Triangles array in parallel, by computing
+ * partial boxes for subsets of this array, and then combining those boxes.
+ * TODO: this should move to a pulbic utility function, and possibly the block-based ParallelFor
+ * should be refactored out into something more general, as this pattern is useful in many places...
+ */
+static FAxisAlignedBox3d ParallelComputeROIBounds(const FDynamicMesh3& Mesh, const TArray<int32>& Triangles)
+{
+	FAxisAlignedBox3d FinalBounds = FAxisAlignedBox3d::Empty();
+	FCriticalSection FinalBoundsLock;
+	int32 N = Triangles.Num();
+	constexpr int32 BlockSize = 4096;
+	int32 Blocks = (N / BlockSize) + 1;
+	ParallelFor(Blocks, [&](int bi)
+	{
+		FAxisAlignedBox3d BlockBounds = FAxisAlignedBox3d::Empty();
+		for (int32 k = 0; k < BlockSize; ++k)
+		{
+			int32 i = bi * BlockSize + k;
+			if (i < N)
+			{
+				int32 tid = Triangles[i];
+				const FIndex3i& TriV = Mesh.GetTriangleRef(tid);
+				BlockBounds.Contain(Mesh.GetVertexRef(TriV.A));
+				BlockBounds.Contain(Mesh.GetVertexRef(TriV.B));
+				BlockBounds.Contain(Mesh.GetVertexRef(TriV.C));
+			}
+		}
+		FinalBoundsLock.Lock();
+		FinalBounds.Contain(BlockBounds);
+		FinalBoundsLock.Unlock();
+	});
+	return FinalBounds;
+}
+
+
+
+TFuture<bool> UDynamicMeshComponent::FastNotifyTriangleVerticesUpdated_TryPrecompute(
+	const TArray<int32>& Triangles,
+	TArray<int32>& UpdateSetsOut,
+	FAxisAlignedBox3d& BoundsOut)
+{
+	if ((!!RenderMeshPostProcessor) || (GetCurrentSceneProxy() == nullptr) || (!Decomposition))
+	{
+		// is there a simpler way to do this? cannot seem to just make a TFuture<bool>...
+		return Async(DynamicMeshComponentAsyncExecTarget, []() { return false; });
+	}
+
+	return Async(DynamicMeshComponentAsyncExecTarget, [this, &Triangles, &UpdateSetsOut, &BoundsOut]()
+	{
+		TFuture<void> ComputeBounds = Async(DynamicMeshComponentAsyncExecTarget, [this, &BoundsOut, &Triangles]()
+		{
+			TRACE_CPUPROFILER_EVENT_SCOPE(SimpleDynamicMeshComponent_FastVertexUpdatePrecomp_CalcBounds);
+			BoundsOut = ParallelComputeROIBounds(*GetMesh(), Triangles);
+		});
+
+		TFuture<void> ComputeSets = Async(DynamicMeshComponentAsyncExecTarget, [this, &UpdateSetsOut, &Triangles]()
+		{
+			TRACE_CPUPROFILER_EVENT_SCOPE(SimpleDynamicMeshComponent_FastVertexUpdatePrecomp_FindSets);
+			int32 NumBuffers = Decomposition->Num();
+			TArray<std::atomic<bool>> BufferFlags;
+			BufferFlags.SetNum(NumBuffers);
+			for (int32 k = 0; k < NumBuffers; ++k)
+			{
+				BufferFlags[k] = false;
+			}
+			ParallelFor(Triangles.Num(), [&](int32 k)
+			{
+				int32 SetID = Decomposition->GetGroupForTriangle(Triangles[k]);
+				BufferFlags[SetID] = true;
+			});
+			UpdateSetsOut.Reset();
+			for (int32 k = 0; k < NumBuffers; ++k)
+			{
+				if (BufferFlags[k])
+				{
+					UpdateSetsOut.Add(k);
+				}
+			}
+
+		});
+
+		ComputeSets.Wait();
+		ComputeBounds.Wait();
+
+		return true;
+	});
+}
+
+
+void UDynamicMeshComponent::FastNotifyTriangleVerticesUpdated_ApplyPrecompute(
+	const TArray<int32>& Triangles,
+	EMeshRenderAttributeFlags UpdatedAttributes, 
+	TFuture<bool>& Precompute, 
+	const TArray<int32>& UpdateSets, 
+	const FAxisAlignedBox3d& UpdateSetBounds)
+{
+	Precompute.Wait();
+
+	bool bPrecomputeOK = Precompute.Get();
+	if (bPrecomputeOK == false || GetCurrentSceneProxy() == nullptr )
+	{
+		FastNotifyTriangleVerticesUpdated(Triangles, UpdatedAttributes);
+		return;
+	}
+
+	FDynamicMeshSceneProxy* Proxy = GetCurrentSceneProxy();
+	bool bPositions = (UpdatedAttributes & EMeshRenderAttributeFlags::Positions) != EMeshRenderAttributeFlags::None;
+	bool bUpdateSecondarySort = (SecondaryTriFilterFunc) &&
+		((UpdatedAttributes & EMeshRenderAttributeFlags::SecondaryIndexBuffers) != EMeshRenderAttributeFlags::None);
+
+	{
+		TRACE_CPUPROFILER_EVENT_SCOPE(SimpleDynamicMeshComponent_FastVertexUpdate_ApplyUpdate);
+		Proxy->FastUpdateVertices(UpdateSets, bPositions,
+			(UpdatedAttributes & EMeshRenderAttributeFlags::VertexNormals) != EMeshRenderAttributeFlags::None,
+			(UpdatedAttributes & EMeshRenderAttributeFlags::VertexColors) != EMeshRenderAttributeFlags::None,
+			(UpdatedAttributes & EMeshRenderAttributeFlags::VertexUVs) != EMeshRenderAttributeFlags::None);
+	}
+
+	{
+		TRACE_CPUPROFILER_EVENT_SCOPE(SimpleDynamicMeshComponent_FastVertexUpdate_UpdateIndexBuffers);
+		if (bUpdateSecondarySort)
+		{
+			Proxy->FastUpdateIndexBuffers(UpdateSets);
+		}
+	}
+
+	if (bPositions)
+	{
+		TRACE_CPUPROFILER_EVENT_SCOPE(SimpleDynamicMeshComponent_FastVertexUpdate_FinalPositionsUpdate);
+		MarkRenderTransformDirty();
+		LocalBounds.Contain(UpdateSetBounds);
+		UpdateBounds();
+	}
+
+	GetDynamicMesh()->PostRealtimeUpdate();
+}
+
+
+
+
+
+FPrimitiveSceneProxy* UDynamicMeshComponent::CreateSceneProxy()
+{
+	// if this is not always the case, we have made incorrect assumptions
+	ensure(GetCurrentSceneProxy() == nullptr);
+
+	FDynamicMeshSceneProxy* NewProxy = nullptr;
+	if (GetMesh()->TriangleCount() > 0)
+	{
+		NewProxy = new FDynamicMeshSceneProxy(this);
+
+		if (TriangleColorFunc)
+		{
+			NewProxy->bUsePerTriangleColor = true;
+			NewProxy->PerTriangleColorFunc = [this](const FDynamicMesh3* MeshIn, int TriangleID) { return GetTriangleColor(MeshIn, TriangleID); };
+		}
+
+		if (SecondaryTriFilterFunc)
+		{
+			NewProxy->bUseSecondaryTriBuffers = true;
+			NewProxy->SecondaryTriFilterFunc = [this](const FDynamicMesh3* MeshIn, int32 TriangleID) 
+			{ 
+				return (SecondaryTriFilterFunc) ? SecondaryTriFilterFunc(MeshIn, TriangleID) : false;
+			};
+		}
+
+		if (Decomposition)
+		{
+			NewProxy->InitializeFromDecomposition(Decomposition);
+		}
+		else
+		{
+			NewProxy->Initialize();
+		}
+	}
+
+	bProxyValid = true;
+	return NewProxy;
+}
+
+
+
+void UDynamicMeshComponent::NotifyMaterialSetUpdated()
+{
+	if (GetCurrentSceneProxy() != nullptr)
+	{
+		GetCurrentSceneProxy()->UpdatedReferencedMaterials();
+	}
+}
+
+
+void UDynamicMeshComponent::SetTriangleColorFunction(
+	TUniqueFunction<FColor(const FDynamicMesh3*, int)> TriangleColorFuncIn,
+	EDynamicMeshComponentRenderUpdateMode UpdateMode)
+{
+	TriangleColorFunc = MoveTemp(TriangleColorFuncIn);
+
+	if (UpdateMode == EDynamicMeshComponentRenderUpdateMode::FastUpdate)
+	{
+		FastNotifyColorsUpdated();
+	}
+	else if (UpdateMode == EDynamicMeshComponentRenderUpdateMode::FullUpdate)
+	{
+		NotifyMeshUpdated();
+	}
+}
+
+void UDynamicMeshComponent::ClearTriangleColorFunction(EDynamicMeshComponentRenderUpdateMode UpdateMode)
+{
+	if (TriangleColorFunc)
+	{
+		TriangleColorFunc = nullptr;
+
+		if (UpdateMode == EDynamicMeshComponentRenderUpdateMode::FastUpdate)
+		{
+			FastNotifyColorsUpdated();
+		}
+		else if (UpdateMode == EDynamicMeshComponentRenderUpdateMode::FullUpdate)
+		{
+			NotifyMeshUpdated();
+		}
+	}
+}
+
+bool UDynamicMeshComponent::HasTriangleColorFunction()
+{
+	return !!TriangleColorFunc;
+}
+
+
+
+
+void UDynamicMeshComponent::EnableSecondaryTriangleBuffers(TUniqueFunction<bool(const FDynamicMesh3*, int32)> SecondaryTriFilterFuncIn)
+{
+	SecondaryTriFilterFunc = MoveTemp(SecondaryTriFilterFuncIn);
+	NotifyMeshUpdated();
+}
+
+void UDynamicMeshComponent::DisableSecondaryTriangleBuffers()
+{
+	SecondaryTriFilterFunc = nullptr;
+	NotifyMeshUpdated();
+}
+
+
+void UDynamicMeshComponent::SetExternalDecomposition(TUniquePtr<FMeshRenderDecomposition> DecompositionIn)
+{
+	check(DecompositionIn->Num() > 0);
+	Decomposition = MoveTemp(DecompositionIn);
+	NotifyMeshUpdated();
+}
+
+
+
+FColor UDynamicMeshComponent::GetTriangleColor(const FDynamicMesh3* MeshIn, int TriangleID)
+{
+	if (TriangleColorFunc)
+	{
+		return TriangleColorFunc(MeshIn, TriangleID);
+	}
+	else
+	{
+		return (TriangleID % 2 == 0) ? FColor::Red : FColor::White;
+	}
+}
+
+
+
+FBoxSphereBounds UDynamicMeshComponent::CalcBounds(const FTransform& LocalToWorld) const
+{
+	// can get a tighter box by calculating in world space, but we care more about performance
+	FBox LocalBoundingBox = (FBox)LocalBounds;
+	FBoxSphereBounds Ret(LocalBoundingBox.TransformBy(LocalToWorld));
+	Ret.BoxExtent *= BoundsScale;
+	Ret.SphereRadius *= BoundsScale;
+	return Ret;
+}
+
+
+
+
+void UDynamicMeshComponent::SetInvalidateProxyOnChangeEnabled(bool bEnabled)
+{
+	bInvalidateProxyOnChange = bEnabled;
+}
+
+
+void UDynamicMeshComponent::ApplyChange(const FMeshVertexChange* Change, bool bRevert)
+{
+	// will fire UDynamicMesh::MeshChangedEvent, which will call OnMeshObjectChanged() below to invalidate proxy, fire change events, etc
+	MeshObject->ApplyChange(Change, bRevert);
+}
+
+void UDynamicMeshComponent::ApplyChange(const FMeshChange* Change, bool bRevert)
+{
+	// will fire UDynamicMesh::MeshChangedEvent, which will call OnMeshObjectChanged() below to invalidate proxy, fire change events, etc
+	MeshObject->ApplyChange(Change, bRevert);
+}
+
+void UDynamicMeshComponent::ApplyChange(const FMeshReplacementChange* Change, bool bRevert)
+{
+	// will fire UDynamicMesh::MeshChangedEvent, which will call OnMeshObjectChanged() below to invalidate proxy, fire change events, etc
+	MeshObject->ApplyChange(Change, bRevert);
+}
+
+void UDynamicMeshComponent::OnMeshObjectChanged(UDynamicMesh* ChangedMeshObject, FDynamicMeshChangeInfo ChangeInfo)
+{
+	bool bIsFChange = (
+		ChangeInfo.Type == EDynamicMeshChangeType::MeshChange
+		|| ChangeInfo.Type == EDynamicMeshChangeType::MeshVertexChange
+		|| ChangeInfo.Type == EDynamicMeshChangeType::MeshReplacementChange);
+
+	if (bIsFChange)
+	{
+		if (bInvalidateProxyOnChange)
+		{
+			NotifyMeshUpdated();
+		}
+
+		OnMeshChanged.Broadcast();
+
+		if (ChangeInfo.Type == EDynamicMeshChangeType::MeshVertexChange)
+		{
+			OnMeshVerticesChanged.Broadcast(this, ChangeInfo.VertexChange, ChangeInfo.bIsRevertChange);
+		}
+	}
+	else
+	{
+		NotifyMeshUpdated();
+		OnMeshChanged.Broadcast();
+	}
+
+	// Rebuild body setup. Should this be deferred until proxy creation? Sometimes multiple changes are emitted...
+	// todo: can possibly skip this in some change situations, eg if only changing attributes
+	if (bDeferCollisionUpdates)
+	{
+		InvalidatePhysicsData();
+	}
+	else
+	{
+		RebuildPhysicsData();
+	}
+}
+
+
+void UDynamicMeshComponent::SetDynamicMesh(UDynamicMesh* NewMesh)
+{
+	if (ensure(NewMesh) == false)
+	{
+		return;
+	}
+
+	if (ensure(MeshObject))
+	{
+		MeshObject->OnMeshChanged().Remove(MeshObjectChangedHandle);
+	}
+
+	MeshObject = NewMesh;
+	MeshObjectChangedHandle = MeshObject->OnMeshChanged().AddUObject(this, &UDynamicMeshComponent::OnMeshObjectChanged);
+
+	NotifyMeshUpdated();
+	OnMeshChanged.Broadcast();
+
+	// Rebuild physics data
+	if (bDeferCollisionUpdates)
+	{
+		InvalidatePhysicsData();
+	}
+	else
+	{
+		RebuildPhysicsData();
+	}
+}
+
+
+
+void UDynamicMeshComponent::OnChildAttached(USceneComponent* ChildComponent)
+{
+	Super::OnChildAttached(ChildComponent);
+	OnChildAttachmentModified.Broadcast(ChildComponent, true);
+}
+void UDynamicMeshComponent::OnChildDetached(USceneComponent* ChildComponent)
+{
+	Super::OnChildDetached(ChildComponent);
+	OnChildAttachmentModified.Broadcast(ChildComponent, false);
+}
+
+
+
+
+
+bool UDynamicMeshComponent::GetPhysicsTriMeshData(struct FTriMeshCollisionData* CollisionData, bool InUseAllTriData)
+{
+	// todo: support physical materials
+	// todo: support UPhysicsSettings::Get()->bSupportUVFromHitResults
+
+	// this is something we currently assume, if you hit this ensure, we made a mistake
+	ensure(bEnableComplexCollision);
+
+	ProcessMesh([&](const FDynamicMesh3& Mesh)
+	{
+		TArray<int32> VertexMap;
+		bool bIsSparseV = !Mesh.IsCompactV();
+		if (bIsSparseV)
+		{
+			VertexMap.SetNum(Mesh.MaxVertexID());
+		}
+
+		// copy vertices
+		CollisionData->Vertices.Reserve(Mesh.VertexCount());
+		for (int32 vid : Mesh.VertexIndicesItr())
+		{
+			int32 Index = CollisionData->Vertices.Add((FVector)Mesh.GetVertex(vid));
+			if (bIsSparseV)
+			{
+				VertexMap[vid] = Index;
+			}
+			else
+			{
+				check(vid == Index);
+			}
+		}
+
+		// copy triangles
+		CollisionData->Indices.Reserve(Mesh.TriangleCount());
+		CollisionData->MaterialIndices.Reserve(Mesh.TriangleCount());
+		for (int32 tid : Mesh.TriangleIndicesItr())
+		{
+			FIndex3i Tri = Mesh.GetTriangle(tid);
+			FTriIndices Triangle;
+			Triangle.v0 = (bIsSparseV) ? VertexMap[Tri.A] : Tri.A;
+			Triangle.v1 = (bIsSparseV) ? VertexMap[Tri.B] : Tri.B;
+			Triangle.v2 = (bIsSparseV) ? VertexMap[Tri.C] : Tri.C;
+			CollisionData->Indices.Add(Triangle);
+
+			CollisionData->MaterialIndices.Add(0);		// not supporting physical materials yet
+		}
+
+		CollisionData->bFlipNormals = true;
+		CollisionData->bDeformableMesh = true;
+		CollisionData->bFastCook = true;
+	});
+
+	return true;
+}
+
+bool UDynamicMeshComponent::ContainsPhysicsTriMeshData(bool InUseAllTriData) const
+{
+	return bEnableComplexCollision && ((MeshObject != nullptr) ? (MeshObject->GetTriangleCount() > 0) : false);
+}
+
+bool UDynamicMeshComponent::WantsNegXTriMesh()
+{
+	return true;
+}
+
+UBodySetup* UDynamicMeshComponent::GetBodySetup()
+{
+	if (MeshBodySetup == nullptr)
+	{
+		UBodySetup* NewBodySetup = nullptr;
+		{
+			FGCScopeGuard Scope;
+			NewBodySetup = NewObject<UBodySetup>(this);
+		}
+		NewBodySetup->BodySetupGuid = FGuid::NewGuid();
+		
+		NewBodySetup->bGenerateMirroredCollision = false;
+		NewBodySetup->CollisionTraceFlag = this->CollisionType;
+
+		NewBodySetup->DefaultInstance.SetCollisionProfileName(UCollisionProfile::BlockAll_ProfileName);
+		NewBodySetup->bSupportUVsAndFaceRemap = false; /* bSupportPhysicalMaterialMasks; */
+		
+		SetBodySetup(NewBodySetup);
+	}
+
+	return MeshBodySetup;
+}
+
+void UDynamicMeshComponent::SetBodySetup(UBodySetup* NewSetup)
+{
+	if (ensure(NewSetup))
+	{
+		MeshBodySetup = NewSetup;
+	}
+}
+
+void UDynamicMeshComponent::InvalidatePhysicsData()
+{
+	if (GetBodySetup())
+	{
+		GetBodySetup()->InvalidatePhysicsData();
+		bCollisionUpdatePending = true;
+	}
+}
+
+void UDynamicMeshComponent::RebuildPhysicsData()
+{
+	UBodySetup* BodySetup = GetBodySetup();
+	if (BodySetup)
+	{
+		// New GUID as collision has changed
+		BodySetup->BodySetupGuid = FGuid::NewGuid();
+		// Also we want cooked data for this
+		BodySetup->bHasCookedCollisionData = true;
+
+		BodySetup->CollisionTraceFlag = this->CollisionType;
+
+		BodySetup->InvalidatePhysicsData();
+		BodySetup->CreatePhysicsMeshes();
+		RecreatePhysicsState();
+
+		bCollisionUpdatePending = false;
+	}
+}
+
+void UDynamicMeshComponent::UpdateCollision(bool bOnlyIfPending)
+{
+	if (bOnlyIfPending == false || bCollisionUpdatePending)
+	{
+		RebuildPhysicsData();
+	}
+}
\ No newline at end of file
diff --git a/Engine/Source/Runtime/Experimental/GeometryFramework/Private/Components/DynamicMeshSceneProxy.h b/Engine/Source/Runtime/Experimental/GeometryFramework/Private/Components/DynamicMeshSceneProxy.h
new file mode 100644
--- /dev/null
+++ b/Engine/Source/Runtime/Experimental/GeometryFramework/Private/Components/DynamicMeshSceneProxy.h
@@ -0,0 +1,658 @@
+// Copyright Epic Games, Inc. All Rights Reserved. 
+
+#pragma once
+
+#include "Components/DynamicMeshComponent.h"
+#include "Components/BaseDynamicMeshSceneProxy.h"
+#include "DynamicMesh/MeshTangents.h"
+#include "Async/ParallelFor.h"
+
+/**
+ * Scene Proxy for UDynamicMeshComponent.
+ * 
+ * Based on FProceduralMeshSceneProxy but simplified in various ways.
+ * 
+ * Supports wireframe-on-shaded rendering.
+ * 
+ */
+class FDynamicMeshSceneProxy final : public FBaseDynamicMeshSceneProxy
+{
+	using FIndex2i = UE::Geometry::FIndex2i;
+	using FIndex3i = UE::Geometry::FIndex3i;
+private:
+	FMaterialRelevance MaterialRelevance;
+
+	// note: FBaseDynamicMeshSceneProxy owns and will destroy these
+	TArray<FMeshRenderBufferSet*> RenderBufferSets;
+
+	// if true, we store entire mesh in single RenderBuffers and we can do some optimizations
+	bool bIsSingleBuffer = false;
+
+public:
+	/** Component that created this proxy (is there a way to look this up?) */
+	UDynamicMeshComponent* ParentComponent;
+
+
+	FDynamicMeshSceneProxy(UDynamicMeshComponent* Component)
+		: FBaseDynamicMeshSceneProxy(Component)
+		, MaterialRelevance(Component->GetMaterialRelevance(GetScene().GetFeatureLevel()))
+	{
+		ParentComponent = Component;
+	}
+
+
+	virtual void GetActiveRenderBufferSets(TArray<FMeshRenderBufferSet*>& Buffers) const override
+	{
+		Buffers = RenderBufferSets;
+	}
+
+
+
+	void Initialize()
+	{
+		// allocate buffer sets based on materials
+		ensure(RenderBufferSets.Num() == 0);
+		int32 NumMaterials = GetNumMaterials();
+		if (NumMaterials == 0)
+		{
+			RenderBufferSets.SetNum(1);
+			RenderBufferSets[0] = AllocateNewRenderBufferSet();
+			RenderBufferSets[0]->Material = UMaterial::GetDefaultMaterial(MD_Surface);
+		}
+		else
+		{
+			RenderBufferSets.SetNum(NumMaterials);
+			for (int32 k = 0; k < NumMaterials; ++k)
+			{
+				RenderBufferSets[k] = AllocateNewRenderBufferSet();
+				RenderBufferSets[k]->Material = GetMaterial(k);
+			}
+		}
+
+
+		FDynamicMesh3* Mesh = ParentComponent->GetRenderMesh();
+		if (Mesh->HasAttributes() && Mesh->Attributes()->HasMaterialID() && NumMaterials > 1)
+		{
+			bIsSingleBuffer = false;
+			InitializeByMaterial(RenderBufferSets);
+		}
+		else
+		{
+			bIsSingleBuffer = true;
+			InitializeSingleBufferSet(RenderBufferSets[0]);
+		}
+	}
+
+
+
+	TUniqueFunction<void(int, int, int, const FVector3f&, FVector3f&, FVector3f&)> MakeTangentsFunc(bool bSkipAutoCompute = false)
+	{
+		EDynamicMeshComponentTangentsMode TangentsType = ParentComponent->GetTangentsType();
+		if (TangentsType == EDynamicMeshComponentTangentsMode::ExternallyProvided)
+		{
+			UE::Geometry::FDynamicMeshTangents Tangents(ParentComponent->GetRenderMesh());
+			return [Tangents](int VertexID, int TriangleID, int TriVtxIdx, const FVector3f& Normal, FVector3f& TangentX, FVector3f& TangentY) -> void
+			{
+				Tangents.GetTangentVectors(TriangleID, TriVtxIdx, Normal, TangentX, TangentY);
+			};
+		}
+		else if (TangentsType == EDynamicMeshComponentTangentsMode::AutoCalculated && bSkipAutoCompute == false)
+		{
+			const UE::Geometry::FMeshTangentsf* Tangents = ParentComponent->GetAutoCalculatedTangents();
+			if (Tangents != nullptr )
+			{
+				return [Tangents](int VertexID, int TriangleID, int TriVtxIdx, const FVector3f& Normal, FVector3f& TangentX, FVector3f& TangentY) -> void
+				{
+					Tangents->GetTriangleVertexTangentVectors(TriangleID, TriVtxIdx, TangentX, TangentY);
+				};
+			}
+		}
+
+		// fallback to orthogonal basis
+		return [](int VertexID, int TriangleID, int TriVtxIdx, const FVector3f& Normal, FVector3f& TangentX, FVector3f& TangentY) -> void
+		{
+			VectorUtil::MakePerpVectors(Normal, TangentX, TangentY);
+		};
+	}
+
+
+	/**
+	 * Initialize multiple buffers for the mesh based on the given Decomposition
+	 */
+	void InitializeFromDecomposition(TUniquePtr<FMeshRenderDecomposition>& Decomposition)
+	{
+		ensure(RenderBufferSets.Num() == 0);
+		int32 NumSets = Decomposition->Num();
+		RenderBufferSets.SetNum(NumSets);
+		for (int32 k = 0; k < NumSets; ++k)
+		{
+			RenderBufferSets[k] = AllocateNewRenderBufferSet();
+			RenderBufferSets[k]->Material = Decomposition->GetGroup(k).Material;
+			if (RenderBufferSets[k]->Material == nullptr)
+			{
+				RenderBufferSets[k]->Material = UMaterial::GetDefaultMaterial(MD_Surface);
+			}
+		}
+
+		bIsSingleBuffer = false;
+
+		const FDynamicMesh3* Mesh = ParentComponent->GetRenderMesh();
+		// find suitable overlays
+		const FDynamicMeshUVOverlay* UVOverlay = Mesh->Attributes()->PrimaryUV();
+		const FDynamicMeshNormalOverlay* NormalOverlay = Mesh->Attributes()->PrimaryNormals();
+		const FDynamicMeshColorOverlay* ColorOverlay = Mesh->Attributes()->PrimaryColors();
+		TUniqueFunction<void(int, int, int, const FVector3f&, FVector3f&, FVector3f&)> TangentsFunc = MakeTangentsFunc();
+
+		// init renderbuffers for each set
+		ParallelFor(NumSets, [&](int32 SetIndex)
+		{
+			const FMeshRenderDecomposition::FGroup& Group = Decomposition->GetGroup(SetIndex);
+			const TArray<int32>& Triangles = Group.Triangles;
+			if (Triangles.Num() > 0)
+			{
+				FMeshRenderBufferSet* RenderBuffers = RenderBufferSets[SetIndex];
+				RenderBuffers->Triangles = Triangles;
+				InitializeBuffersFromOverlays(RenderBuffers, Mesh,
+					Triangles.Num(), Triangles,
+					UVOverlay, NormalOverlay, ColorOverlay, TangentsFunc);
+
+				ENQUEUE_RENDER_COMMAND(FDynamicMeshSceneProxyInitializeFromDecomposition)(
+					[RenderBuffers](FRHICommandListImmediate& RHICmdList)
+				{
+					RenderBuffers->Upload();
+				});
+			}
+		});
+	}
+
+
+
+
+
+
+
+
+	/**
+	 * Initialize a single set of mesh buffers for the entire mesh
+	 */
+	void InitializeSingleBufferSet(FMeshRenderBufferSet* RenderBuffers)
+	{
+		const FDynamicMesh3* Mesh = ParentComponent->GetRenderMesh();
+
+		// find suitable overlays
+		TArray<const FDynamicMeshUVOverlay*> UVOverlays;
+		const FDynamicMeshNormalOverlay* NormalOverlay = nullptr;
+		const FDynamicMeshColorOverlay* ColorOverlay = nullptr;
+		if (Mesh->HasAttributes())
+		{
+			const FDynamicMeshAttributeSet* Attributes = Mesh->Attributes();
+			NormalOverlay = Attributes->PrimaryNormals();
+			UVOverlays.SetNum(Attributes->NumUVLayers());
+			for (int32 k = 0; k < UVOverlays.Num(); ++k)
+			{
+				UVOverlays[k] = Attributes->GetUVLayer(k);
+			}
+			ColorOverlay = Attributes->PrimaryColors();
+		}
+		TUniqueFunction<void(int, int, int, const FVector3f&, FVector3f&, FVector3f&)> TangentsFunc = MakeTangentsFunc();
+
+		InitializeBuffersFromOverlays(RenderBuffers, Mesh,
+			Mesh->TriangleCount(), Mesh->TriangleIndicesItr(),
+			UVOverlays, NormalOverlay, ColorOverlay, TangentsFunc);
+
+		ENQUEUE_RENDER_COMMAND(FDynamicMeshSceneProxyInitializeSingle)(
+			[RenderBuffers](FRHICommandListImmediate& RHICmdList)
+		{
+			RenderBuffers->Upload();
+		});
+	}
+
+
+
+	/**
+	 * Initialize the mesh buffers, one per material
+	 */
+	void InitializeByMaterial(TArray<FMeshRenderBufferSet*>& BufferSets)
+	{
+		const FDynamicMesh3* Mesh = ParentComponent->GetRenderMesh();
+		if (ensure(Mesh->HasAttributes() && Mesh->Attributes()->HasMaterialID()) == false)
+		{
+			return;
+		}
+
+		const FDynamicMeshAttributeSet* Attributes = Mesh->Attributes();
+
+		// find suitable overlays
+		const FDynamicMeshMaterialAttribute* MaterialID = Attributes->GetMaterialID();
+		const FDynamicMeshNormalOverlay* NormalOverlay = Mesh->Attributes()->PrimaryNormals();
+		const FDynamicMeshColorOverlay* ColorOverlay = Mesh->Attributes()->PrimaryColors();
+
+		TArray<const FDynamicMeshUVOverlay*> UVOverlays;
+		UVOverlays.SetNum(Attributes->NumUVLayers());
+		for (int32 k = 0; k < UVOverlays.Num(); ++k)
+		{
+			UVOverlays[k] = Attributes->GetUVLayer(k);
+		}
+
+		TUniqueFunction<void(int, int, int, const FVector3f&, FVector3f&, FVector3f&)> TangentsFunc = MakeTangentsFunc();
+
+		// count number of triangles for each material (could do in parallel?)
+		int NumMaterials = BufferSets.Num();
+		TArray<FThreadSafeCounter> Counts;
+		Counts.SetNum(NumMaterials);
+		for (int k = 0; k < NumMaterials; ++k)
+		{
+			Counts[k].Reset();
+		}
+		ParallelFor(Mesh->MaxTriangleID(), [&](int tid)
+		{
+			int MatIdx;
+			MaterialID->GetValue(tid, &MatIdx);
+			if (MatIdx >= 0 && MatIdx < NumMaterials)
+			{
+				Counts[MatIdx].Increment();
+			}
+		});
+
+		// find max count
+		int32 MaxCount = 0;
+		for (FThreadSafeCounter& Count : Counts)
+		{
+			MaxCount = FMath::Max(Count.GetValue(), MaxCount);
+		}
+
+		// init renderbuffers for each material
+		// could do this in parallel but then we need to allocate separate triangle arrays...is it worth it?
+		TArray<int> Triangles;
+		Triangles.Reserve(MaxCount);
+		for (int k = 0; k < NumMaterials; ++k)
+		{
+			if (Counts[k].GetValue() > 0)
+			{
+				FMeshRenderBufferSet* RenderBuffers = BufferSets[k];
+
+				Triangles.Reset();
+				for (int tid : Mesh->TriangleIndicesItr())
+				{
+					int MatIdx;
+					MaterialID->GetValue(tid, &MatIdx);
+					if (MatIdx == k)
+					{
+						Triangles.Add(tid);
+					}
+				}
+
+				InitializeBuffersFromOverlays(RenderBuffers, Mesh,
+					Triangles.Num(), Triangles,
+					UVOverlays, NormalOverlay, ColorOverlay, TangentsFunc);
+
+				RenderBuffers->Triangles = Triangles;
+
+				ENQUEUE_RENDER_COMMAND(FDynamicMeshSceneProxyInitializeByMaterial)(
+					[RenderBuffers](FRHICommandListImmediate& RHICmdList)
+				{
+					RenderBuffers->Upload();
+				});
+			}
+		}
+	}
+
+
+	bool RenderMeshLayoutMatchesRenderBuffers() const
+	{
+		const FDynamicMesh3* Mesh = ParentComponent->GetRenderMesh();
+
+		auto CheckBufferSet = [](const FDynamicMesh3* Mesh, const FMeshRenderBufferSet* BufferSet, int NumTriangles) -> bool
+		{
+			if (BufferSet->Triangles)
+			{
+				for (int TriangleID : BufferSet->Triangles.GetValue())
+				{
+					if (!Mesh->IsTriangle(TriangleID))
+					{
+						return false;
+					}
+				}
+			}
+
+			int NumVertices = NumTriangles * 3;
+			if (BufferSet->PositionVertexBuffer.GetNumVertices() != NumVertices ||
+				BufferSet->StaticMeshVertexBuffer.GetNumVertices() != NumVertices ||
+				BufferSet->ColorVertexBuffer.GetNumVertices() != NumVertices)
+			{
+				return false;
+			}
+
+			return true;
+		};
+
+		if (bIsSingleBuffer)
+		{
+			if (ensure(RenderBufferSets.Num() == 1))
+			{
+				FMeshRenderBufferSet* BufferSet = RenderBufferSets[0];
+				if (BufferSet->TriangleCount != Mesh->TriangleCount())
+				{
+					return false;
+				}
+
+				int NumTriangles = Mesh->TriangleCount();
+
+				if (!CheckBufferSet(Mesh, BufferSet, NumTriangles))
+				{
+					return false;
+				}
+			}
+		}
+		else
+		{
+			for (FMeshRenderBufferSet* BufferSet : RenderBufferSets)
+			{
+				if (ensure(BufferSet->Triangles))
+				{
+					int NumTriangles = BufferSet->Triangles->Num();
+
+					if (!CheckBufferSet(Mesh, BufferSet, NumTriangles))
+					{
+						return false;
+					}
+				}
+			}
+		}
+
+		return true;
+	}
+
+
+
+	/**
+	 * Update the vertex position/normal/color buffers
+	 */
+	void FastUpdateVertices(bool bPositions, bool bNormals, bool bColors, bool bUVs)
+	{
+		// This needs to be rewritten for split-by-material buffers.
+		// Could store triangle set with each buffer, and then rebuild vtx buffer(s) as needed?
+
+		FDynamicMesh3* Mesh = ParentComponent->GetRenderMesh();
+
+		// find suitable overlays and attributes
+		FDynamicMeshNormalOverlay* NormalOverlay = nullptr;
+		if (bNormals && ensure(Mesh->HasAttributes()) )
+		{
+			NormalOverlay = Mesh->Attributes()->PrimaryNormals();
+		}
+		FDynamicMeshUVOverlay* UVOVerlay = nullptr;
+		if (bUVs && ensure(Mesh->HasAttributes()) )
+		{
+			UVOVerlay = Mesh->Attributes()->PrimaryUV();
+		}
+		FDynamicMeshColorOverlay* ColorOverlay = nullptr;
+		if (bColors && ensure(Mesh->HasAttributes()) )
+		{
+			ColorOverlay = Mesh->Attributes()->PrimaryColors();
+		}
+
+		// Currently deferring tangents auto-computation if doing fast vertex update (maybe needs to be exposed?)
+		TUniqueFunction<void(int, int, int, const FVector3f&, FVector3f&, FVector3f&)> TangentsFunc = MakeTangentsFunc(true);
+
+		if (bIsSingleBuffer)
+		{
+			if (ensure(RenderBufferSets.Num() == 1))
+			{
+				FMeshRenderBufferSet* Buffers = RenderBufferSets[0];
+				if (bPositions || bNormals || bColors)
+				{
+					UpdateVertexBuffersFromOverlays(Buffers, Mesh,
+						Mesh->TriangleCount(), Mesh->TriangleIndicesItr(),
+						NormalOverlay, ColorOverlay, TangentsFunc,
+						bPositions, bNormals, bColors);
+				}
+				if (bUVs)
+				{
+					UpdateVertexUVBufferFromOverlays(Buffers, Mesh,
+						Mesh->TriangleCount(), Mesh->TriangleIndicesItr(), UVOVerlay, 0);
+				}
+
+
+				ENQUEUE_RENDER_COMMAND(FDynamicMeshSceneProxyFastUpdateVertices)(
+					[Buffers, bPositions, bNormals, bColors, bUVs](FRHICommandListImmediate& RHICmdList)
+				{
+					Buffers->UploadVertexUpdate(bPositions, bNormals || bUVs, bColors);
+				});
+			}
+		}
+		else
+		{
+			ParallelFor(RenderBufferSets.Num(), [&](int i)
+			{
+				FMeshRenderBufferSet* Buffers = RenderBufferSets[i];
+				if (Buffers->TriangleCount == 0)
+				{
+					return;
+				}
+				if (ensure(Buffers->Triangles.IsSet()))
+				{
+					if (bPositions || bNormals || bColors)
+					{
+						UpdateVertexBuffersFromOverlays(Buffers, Mesh,
+							Buffers->Triangles->Num(), Buffers->Triangles.GetValue(),
+							NormalOverlay, ColorOverlay, TangentsFunc,
+							bPositions, bNormals, bColors);
+					}
+					if (bUVs)
+					{
+						UpdateVertexUVBufferFromOverlays(Buffers, Mesh,
+							Buffers->Triangles->Num(), Buffers->Triangles.GetValue(), UVOVerlay, 0);
+					}
+
+					ENQUEUE_RENDER_COMMAND(FDynamicMeshSceneProxyFastUpdateVertices)(
+						[Buffers, bPositions, bNormals, bColors, bUVs](FRHICommandListImmediate& RHICmdList)
+					{
+						Buffers->UploadVertexUpdate(bPositions, bNormals || bUVs, bColors);
+					});
+				}
+			});
+		}
+	}
+
+
+
+	/**
+	 * Update the vertex position/normal/color buffers
+	 */
+	void FastUpdateVertices(const TArray<int32>& WhichBuffers, bool bPositions, bool bNormals, bool bColors, bool bUVs)
+	{
+		TRACE_CPUPROFILER_EVENT_SCOPE(DynamicMeshProxy_FastUpdateVertices);
+
+		// skip if we have no updates
+		if (bPositions == false && bNormals == false && bColors == false && bUVs == false)
+		{
+			return;
+		}
+
+		// This needs to be rewritten for split-by-material buffers.
+		// Could store triangle set with each buffer, and then rebuild vtx buffer(s) as needed?
+
+		const FDynamicMesh3* Mesh = ParentComponent->GetRenderMesh();
+
+		// find suitable overlays
+		const FDynamicMeshNormalOverlay* NormalOverlay = nullptr;
+		if (bNormals && ensure(Mesh->HasAttributes()))
+		{
+			NormalOverlay = Mesh->Attributes()->PrimaryNormals();
+		}
+		const FDynamicMeshUVOverlay* UVOVerlay = nullptr;
+		if (bUVs && ensure(Mesh->HasAttributes()))
+		{
+			UVOVerlay = Mesh->Attributes()->PrimaryUV();
+		}
+		const FDynamicMeshColorOverlay* ColorOverlay = nullptr;
+		if (bColors && ensure(Mesh->HasAttributes()))
+		{
+			ColorOverlay = Mesh->Attributes()->PrimaryColors();
+		}
+
+		// Currently deferring tangents auto-computation if doing fast vertex update (maybe needs to be exposed?)
+		TUniqueFunction<void(int, int, int, const FVector3f&, FVector3f&, FVector3f&)> TangentsFunc = MakeTangentsFunc(true);
+
+		ParallelFor(WhichBuffers.Num(), [&](int idx)
+		{
+			int32 BufferIndex = WhichBuffers[idx];
+			if ( RenderBufferSets.IsValidIndex(BufferIndex) == false || RenderBufferSets[BufferIndex]->TriangleCount == 0)
+			{
+				return;
+			}
+			FMeshRenderBufferSet* Buffers = RenderBufferSets[BufferIndex];
+			if (ensure(Buffers->Triangles.IsSet()))
+			{
+				if (bPositions || bNormals || bColors)
+				{
+					UpdateVertexBuffersFromOverlays(Buffers, Mesh,
+						Buffers->Triangles->Num(), Buffers->Triangles.GetValue(),
+						NormalOverlay, ColorOverlay, TangentsFunc,
+						bPositions, bNormals, bColors);
+				}
+				if (bUVs)
+				{
+					UpdateVertexUVBufferFromOverlays(Buffers, Mesh,
+						Buffers->Triangles->Num(), Buffers->Triangles.GetValue(), UVOVerlay, 0);
+				}
+
+				ENQUEUE_RENDER_COMMAND(FDynamicMeshSceneProxyFastUpdateVerticesBufferList)(
+					[Buffers, bPositions, bNormals, bColors, bUVs](FRHICommandListImmediate& RHICmdList)
+				{
+					Buffers->TransferVertexUpdateToGPU(bPositions, bNormals, bUVs, bColors);
+				});
+			}
+		});
+	}
+
+
+
+
+	/**
+	 * Update index buffers inside each RenderBuffer set
+	 */
+	void FastUpdateAllIndexBuffers()
+	{
+		FDynamicMesh3* Mesh = ParentComponent->GetRenderMesh();
+
+		// have to wait for all outstanding rendering to finish because the index buffers we are about to edit might be in-use
+		FlushRenderingCommands();
+
+		ParallelFor(RenderBufferSets.Num(), [this, &Mesh](int i)
+		{
+			FMeshRenderBufferSet* Buffers = RenderBufferSets[i];
+
+			FScopeLock BuffersLock(&Buffers->BuffersLock);
+
+			if (Buffers->TriangleCount > 0)
+			{
+				RecomputeRenderBufferTriangleIndexSets(Buffers, Mesh);
+			}
+
+			ENQUEUE_RENDER_COMMAND(FDynamicMeshSceneProxyFastUpdateAllIndexBuffers)(
+				[Buffers](FRHICommandListImmediate& RHICmdList)
+			{
+				Buffers->UploadIndexBufferUpdate();
+			});
+			
+		});
+	}
+
+
+
+	/**
+	 * Update index buffers inside each RenderBuffer set
+	 */
+	void FastUpdateIndexBuffers(const TArray<int32>& WhichBuffers)
+	{
+		FDynamicMesh3* Mesh = ParentComponent->GetRenderMesh();
+
+		// have to wait for all outstanding rendering to finish because the index buffers we are about to edit might be in-use
+		FlushRenderingCommands();
+
+		ParallelFor(WhichBuffers.Num(), [this, &WhichBuffers, &Mesh](int i)
+		{
+			int32 BufferIndex = WhichBuffers[i];
+			if (RenderBufferSets.IsValidIndex(BufferIndex) == false)
+			{
+				return;
+			}
+
+			FMeshRenderBufferSet* Buffers = RenderBufferSets[BufferIndex];
+			FScopeLock BuffersLock(&Buffers->BuffersLock);
+			if (Buffers->TriangleCount > 0)
+			{
+				RecomputeRenderBufferTriangleIndexSets(Buffers, Mesh);
+			}
+
+			ENQUEUE_RENDER_COMMAND(FDynamicMeshSceneProxyFastUpdateSomeIndexBuffers)(
+				[Buffers](FRHICommandListImmediate& RHICmdList)
+			{
+				Buffers->UploadIndexBufferUpdate();
+			});
+		});
+	}
+
+
+
+public:
+
+
+
+	virtual FPrimitiveViewRelevance GetViewRelevance(const FSceneView* View) const override
+	{
+		FPrimitiveViewRelevance Result;
+
+		Result.bDrawRelevance = IsShown(View);
+		Result.bShadowRelevance = IsShadowCast(View);
+		Result.bDynamicRelevance = true;
+		Result.bRenderInMainPass = ShouldRenderInMainPass();
+		Result.bUsesLightingChannels = GetLightingChannelMask() != GetDefaultLightingChannelMask();
+		Result.bTranslucentSelfShadow = bCastVolumetricTranslucentShadow;
+		Result.bRenderCustomDepth = ShouldRenderCustomDepth();
+		// Note that this is actually a getter. One may argue that it is poorly named.
+		MaterialRelevance.SetPrimitiveViewRelevance(Result);
+		Result.bVelocityRelevance = DrawsVelocity() && Result.bOpaque && Result.bRenderInMainPass;
+
+		return Result;
+	}
+
+	virtual void UpdatedReferencedMaterials() override
+	{
+		FBaseDynamicMeshSceneProxy::UpdatedReferencedMaterials();
+
+		// The material relevance may need updating.
+		MaterialRelevance = ParentComponent->GetMaterialRelevance(GetScene().GetFeatureLevel());
+	}
+
+
+	virtual void GetLightRelevance(const FLightSceneProxy* LightSceneProxy, bool& bDynamic, bool& bRelevant, bool& bLightMapped, bool& bShadowMapped) const override
+	{
+		FPrimitiveSceneProxy::GetLightRelevance(LightSceneProxy, bDynamic, bRelevant, bLightMapped, bShadowMapped);
+	}
+
+	virtual bool CanBeOccluded() const override
+	{
+		return !MaterialRelevance.bDisableDepthTest;
+	}
+
+	virtual uint32 GetMemoryFootprint(void) const override { return(sizeof(*this) + GetAllocatedSize()); }
+
+	uint32 GetAllocatedSize(void) const { return(FPrimitiveSceneProxy::GetAllocatedSize()); }
+
+	SIZE_T GetTypeHash() const override
+	{
+		static size_t UniquePointer;
+		return reinterpret_cast<size_t>(&UniquePointer);
+	}
+
+};
+
+
+
+
+
+
+
diff --git a/Engine/Source/Runtime/Experimental/GeometryFramework/Private/Components/MeshRenderDecomposition.cpp b/Engine/Source/Runtime/Experimental/GeometryFramework/Private/Components/MeshRenderDecomposition.cpp
new file mode 100644
--- /dev/null
+++ b/Engine/Source/Runtime/Experimental/GeometryFramework/Private/Components/MeshRenderDecomposition.cpp
@@ -0,0 +1,159 @@
+// Copyright Epic Games, Inc. All Rights Reserved. 
+
+#include "Components/MeshRenderDecomposition.h"
+#include "DynamicMesh/DynamicMesh3.h"
+#include "DynamicMesh/DynamicMeshAttributeSet.h"
+#include "DynamicMesh/DynamicMeshAABBTree3.h"
+#include "Async/ParallelFor.h"
+#include "ComponentSourceInterfaces.h"
+
+using namespace UE::Geometry;
+
+void FMeshRenderDecomposition::BuildAssociations(const FDynamicMesh3* Mesh)
+{
+	TriangleToGroupMap.SetNum(Mesh->MaxTriangleID());
+
+	ParallelFor(Groups.Num(), [&](int32 GroupIndex) 
+	{
+		if (IsGroup(GroupIndex))
+		{
+			const FGroup& Group = GetGroup(GroupIndex);
+			for (int32 tid : Group.Triangles)
+			{
+				TriangleToGroupMap[tid] = GroupIndex;
+			}
+		}
+	});
+}
+
+
+
+
+void FMeshRenderDecomposition::BuildMaterialDecomposition(const FDynamicMesh3* Mesh, const FComponentMaterialSet* MaterialSet, FMeshRenderDecomposition& Decomp)
+{
+	// always have at least one material
+	FComponentMaterialSet LocalMaterials;
+	LocalMaterials.Materials.Add(nullptr);
+	const FComponentMaterialSet* UseMaterials = (ensure(MaterialSet != nullptr)) ? MaterialSet : &LocalMaterials;
+
+	// may not have MaterialID, in that case all triangles will be material 0
+	const FDynamicMeshMaterialAttribute* MaterialID = Mesh->HasAttributes() ? Mesh->Attributes()->GetMaterialID() : nullptr;
+
+	int32 NumMaterials = UseMaterials->Materials.Num();
+	Decomp.Initialize(NumMaterials);
+	for (int32 k = 0; k < NumMaterials; ++k)
+	{
+		FMeshRenderDecomposition::FGroup& Group = Decomp.GetGroup(k);
+		Group.Material = UseMaterials->Materials[k];
+	}
+
+	for (int32 tid : Mesh->TriangleIndicesItr())
+	{
+		int32 MatIdx = 0;
+		if (MaterialID)
+		{
+			MaterialID->GetValue(tid, &MatIdx);
+		}
+		if (MatIdx < NumMaterials)
+		{
+			FMeshRenderDecomposition::FGroup& Group = Decomp.GetGroup(MatIdx);
+			Group.Triangles.Add(tid);
+		}
+	}
+}
+
+
+
+
+
+
+static void CollectSubDecomposition(
+	const FDynamicMesh3* Mesh, 
+	const TArray<int32>& Triangles, 
+	UMaterialInterface* Material, 
+	FMeshRenderDecomposition& Decomp, 
+	int32 MaxChunkSize,
+	FCriticalSection& DecompLock)
+{
+	int32 MaxTrisPerGroup = MaxChunkSize;
+	if (Triangles.Num() < MaxTrisPerGroup)
+	{
+		DecompLock.Lock();
+		int32 i = Decomp.AppendGroup();
+		FMeshRenderDecomposition::FGroup& Group = Decomp.GetGroup(i);
+		DecompLock.Unlock();
+		Group.Triangles = Triangles;
+		Group.Material = Material;
+		return;
+	}
+
+	FDynamicMeshAABBTree3 Spatial(Mesh, false);
+	Spatial.SetBuildOptions(MaxTrisPerGroup);
+	Spatial.Build(Triangles);
+
+	TArray<int32> ActiveSet;
+	TArray<int32> SpillSet;
+
+	FDynamicMeshAABBTree3::FTreeTraversal Collector;
+	Collector.BeginBoxTrianglesF = [&](int32 BoxID)
+	{
+		ActiveSet.Reset();
+		ActiveSet.Reserve(MaxTrisPerGroup);
+	};
+	Collector.NextTriangleF = [&](int32 TriangleID)
+	{
+		ActiveSet.Add(TriangleID);
+	};
+	Collector.EndBoxTrianglesF = [&](int32 BoxID)
+	{
+		if (ActiveSet.Num() > 0)
+		{
+			if (ActiveSet.Num() < 100)
+			{
+				SpillSet.Append(ActiveSet);
+			}
+			else
+			{
+				DecompLock.Lock();
+				int32 i = Decomp.AppendGroup();
+				FMeshRenderDecomposition::FGroup& Group = Decomp.GetGroup(i);
+				DecompLock.Unlock();
+				Group.Triangles = MoveTemp(ActiveSet);
+				Group.Material = Material;
+			}
+		}
+	};
+	Spatial.DoTraversal(Collector);
+
+	if (SpillSet.Num() > 0)
+	{
+		//UE_LOG(LogTemp, Warning, TEXT("SpillSet Size: %d"), SpillSet.Num());
+
+		DecompLock.Lock();
+		int32 i = Decomp.AppendGroup();
+		FMeshRenderDecomposition::FGroup& Group = Decomp.GetGroup(i);
+		DecompLock.Unlock();
+		Group.Triangles = MoveTemp(SpillSet);
+		Group.Material = Material;
+	}
+}
+
+
+
+
+
+
+
+void FMeshRenderDecomposition::BuildChunkedDecomposition(const FDynamicMesh3* Mesh, const FComponentMaterialSet* MaterialSet, FMeshRenderDecomposition& Decomp, int32 MaxChunkSize)
+{
+	TUniquePtr<FMeshRenderDecomposition> MaterialDecomp = MakeUnique<FMeshRenderDecomposition>();
+	BuildMaterialDecomposition(Mesh, MaterialSet, *MaterialDecomp);
+	int32 NumMatGroups = MaterialDecomp->Num();
+
+	FCriticalSection Lock;
+	ParallelFor(NumMatGroups, [&](int32 k)
+	{
+		const FMeshRenderDecomposition::FGroup& Group = MaterialDecomp->GetGroup(k);
+		CollectSubDecomposition(Mesh, Group.Triangles, Group.Material, Decomp, MaxChunkSize, Lock);
+	});
+}
\ No newline at end of file
diff --git a/Engine/Source/Runtime/Experimental/GeometryFramework/Private/DynamicMeshActor.cpp b/Engine/Source/Runtime/Experimental/GeometryFramework/Private/DynamicMeshActor.cpp
new file mode 100644
--- /dev/null
+++ b/Engine/Source/Runtime/Experimental/GeometryFramework/Private/DynamicMeshActor.cpp
@@ -0,0 +1,30 @@
+// Copyright Epic Games, Inc. All Rights Reserved.
+
+#include "DynamicMeshActor.h"
+#include "Materials/Material.h"
+
+
+#define LOCTEXT_NAMESPACE "ADynamicMeshActor"
+
+ADynamicMeshActor::ADynamicMeshActor(const FObjectInitializer& ObjectInitializer)
+	: Super(ObjectInitializer)
+{
+	DynamicMeshComponent = CreateDefaultSubobject<UDynamicMeshComponent>(TEXT("DynamicMeshComponent"));
+	DynamicMeshComponent->SetMobility(EComponentMobility::Movable);
+	DynamicMeshComponent->SetGenerateOverlapEvents(false);
+	DynamicMeshComponent->SetCollisionProfileName(UCollisionProfile::BlockAll_ProfileName);
+
+	DynamicMeshComponent->CollisionType = ECollisionTraceFlag::CTF_UseDefault;
+
+	DynamicMeshComponent->SetMaterial(0, UMaterial::GetDefaultMaterial(MD_Surface));		// is this necessary?
+
+	SetRootComponent(DynamicMeshComponent);
+}
+
+
+
+
+
+
+
+#undef LOCTEXT_NAMESPACE
\ No newline at end of file
diff --git a/Engine/Source/Runtime/Experimental/GeometryFramework/Private/GeometryFrameworkModule.cpp b/Engine/Source/Runtime/Experimental/GeometryFramework/Private/GeometryFrameworkModule.cpp
new file mode 100644
--- /dev/null
+++ b/Engine/Source/Runtime/Experimental/GeometryFramework/Private/GeometryFrameworkModule.cpp
@@ -0,0 +1,20 @@
+// Copyright Epic Games, Inc. All Rights Reserved.
+
+#include "GeometryFrameworkModule.h"
+
+#define LOCTEXT_NAMESPACE "FGeometryFrameworkModule"
+
+void FGeometryFrameworkModule::StartupModule()
+{
+	// This code will execute after your module is loaded into memory; the exact timing is specified in the .uplugin file per-module
+}
+
+void FGeometryFrameworkModule::ShutdownModule()
+{
+	// This function may be called during shutdown to clean up your module.  For modules that support dynamic reloading,
+	// we call this function before unloading the module.
+}
+
+#undef LOCTEXT_NAMESPACE
+	
+IMPLEMENT_MODULE(FGeometryFrameworkModule, GeometryFramework)
\ No newline at end of file
diff --git a/Engine/Source/Runtime/Experimental/GeometryFramework/Private/UDynamicMesh.cpp b/Engine/Source/Runtime/Experimental/GeometryFramework/Private/UDynamicMesh.cpp
new file mode 100644
--- /dev/null
+++ b/Engine/Source/Runtime/Experimental/GeometryFramework/Private/UDynamicMesh.cpp
@@ -0,0 +1,606 @@
+// Copyright Epic Games, Inc. All Rights Reserved.
+
+#include "UDynamicMesh.h"
+
+#include "Changes/MeshVertexChange.h"
+#include "Changes/MeshChange.h"
+#include "Changes/MeshReplacementChange.h"
+#include "Misc/Base64.h"
+#include "Serialization/MemoryWriter.h"
+#include "Serialization/MemoryReader.h"
+#include "UObject/UE5MainStreamObjectVersion.h"
+#include "HAL/IConsoleManager.h"
+#include "Engine/Engine.h"
+#include "Generators/MinimalBoxMeshGenerator.h"
+
+#include "ExplicitUseGeometryMathTypes.h"		// using UE::Geometry::(math types)
+using namespace UE::Geometry;
+
+
+// these cvars are used to support T3D encoding of the internal FDynamicMesh3, see ::ExportCustomProperties() and ::ImportCustomProperties()
+static TAutoConsoleVariable<int32> CVarDynamicMeshTextBasedDupeTriangleCountThreshold(
+	TEXT("geometry.DynamicMesh.TextBasedDupeTriThreshold"),
+	1000,
+	TEXT("Triangle count threshold for text-based UDynamicMesh duplication using Base64. Large values are quite slow."));
+static TAutoConsoleVariable<int32> CVarDynamicMeshDupeHelperTimeout(
+	TEXT("geometry.DynamicMesh.DupeStashTimeout"),
+	5*60,
+	TEXT("Timeout in seconds for references held by internal UDynamicMesh duplication helper system. See FDynamicMeshCopyHelper."));
+
+
+
+UDynamicMesh::UDynamicMesh(const FObjectInitializer& ObjectInitializer)
+	: Super(ObjectInitializer)
+{
+	InitializeNewMesh();
+}
+
+
+void UDynamicMesh::InitializeNewMesh()
+{
+	Mesh = MakeUnique<FDynamicMesh3>(EMeshComponents::FaceGroups);
+	InitializeMesh();
+}
+
+void UDynamicMesh::Reset()
+{
+	FDynamicMeshChangeInfo ChangeInfo;
+	ChangeInfo.Type = EDynamicMeshChangeType::GeneralEdit;
+	EditMeshInternal([this](FDynamicMesh3& EditMesh)
+	{
+		check(&EditMesh == Mesh.Get());		// assuming that EditMesh is internal mesh here...
+		InitializeMesh();
+	}, ChangeInfo);
+}
+
+void UDynamicMesh::ResetToCube()
+{
+	FDynamicMeshChangeInfo ChangeInfo;
+	ChangeInfo.Type = EDynamicMeshChangeType::GeneralEdit;
+	EditMeshInternal([this](FDynamicMesh3& EditMesh)
+	{
+		EditMesh.Clear();
+		FMinimalBoxMeshGenerator BoxGen;
+		BoxGen.Box = FOrientedBox3d(FVector3d::Zero(), 50.0 * FVector3d::One());
+		EditMesh = FDynamicMesh3(&BoxGen.Generate());
+		EditMesh.EnableTriangleGroups();
+		EditMesh.Attributes()->EnableMaterialID();
+	}, ChangeInfo);
+}
+
+
+void UDynamicMesh::InitializeMesh()
+{
+	Mesh->Clear();
+	Mesh->EnableTriangleGroups();
+	Mesh->EnableAttributes();
+	Mesh->Attributes()->EnableMaterialID();
+
+	if (MeshGenerator != nullptr && bEnableMeshGenerator)
+	{
+		MeshGenerator->Generate(*Mesh);
+	}
+}
+
+
+bool UDynamicMesh::IsEmpty() const
+{
+	return Mesh->TriangleCount() == 0;
+}
+
+int32 UDynamicMesh::GetTriangleCount() const
+{
+	return Mesh->TriangleCount();
+}
+
+
+void UDynamicMesh::SetMesh(const UE::Geometry::FDynamicMesh3& MoveMesh)
+{
+	FDynamicMeshChangeInfo ChangeInfo;
+	ChangeInfo.Type = EDynamicMeshChangeType::GeneralEdit;
+	EditMeshInternal([&](FDynamicMesh3& EditMesh) {
+		EditMesh = MoveMesh;
+	}, ChangeInfo);
+}
+
+void UDynamicMesh::SetMesh(UE::Geometry::FDynamicMesh3&& MoveMesh)
+{
+	FDynamicMeshChangeInfo ChangeInfo;
+	ChangeInfo.Type = EDynamicMeshChangeType::GeneralEdit;
+	EditMeshInternal([&](FDynamicMesh3& EditMesh) {
+		EditMesh = MoveTemp(MoveMesh);
+	}, ChangeInfo);
+}
+
+void UDynamicMesh::ProcessMesh(TFunctionRef<void(const UE::Geometry::FDynamicMesh3&)> ProcessFunc) const
+{
+	ProcessFunc(GetMeshRef());
+}
+
+void UDynamicMesh::EditMesh(TFunctionRef<void(FDynamicMesh3&)> EditFunc,
+							EDynamicMeshChangeType ChangeType,
+							EDynamicMeshAttributeChangeFlags ChangeFlags,
+							bool bDeferChangeEvents)
+{
+	FDynamicMeshChangeInfo ChangeInfo;
+	ChangeInfo.Type = ChangeType;
+	ChangeInfo.Flags = ChangeFlags;
+	EditMeshInternal(EditFunc, ChangeInfo, bDeferChangeEvents);
+}
+
+void UDynamicMesh::EditMeshInternal(TFunctionRef<void(FDynamicMesh3&)> EditFunc, const FDynamicMeshChangeInfo& ChangeInfo, bool bDeferChangeEvents)
+{
+	if (!bDeferChangeEvents)
+	{
+		PreMeshChangedEvent.Broadcast(this, ChangeInfo);
+	}
+	EditFunc(GetMeshRef());
+
+	// Enforce our mesh attribute invariants. This should probably be optional to support compute-only UDynamicMeshes....
+	if (Mesh->HasTriangleGroups() == false)
+	{
+		Mesh->EnableTriangleGroups();
+	}
+	if (Mesh->HasAttributes() == false)
+	{
+		Mesh->EnableAttributes();
+	}
+	if (Mesh->Attributes()->HasMaterialID() == false)
+	{
+		Mesh->Attributes()->EnableMaterialID();
+	}
+
+	if (!bDeferChangeEvents)
+	{
+		MeshChangedEvent.Broadcast(this, ChangeInfo);
+		MeshModifiedBPEvent.Broadcast(this);
+	}
+}
+
+
+TUniquePtr<FDynamicMesh3> UDynamicMesh::ExtractMesh()
+{
+	TUniquePtr<FDynamicMesh3> ReturnMesh = MoveTemp(Mesh);
+	InitializeNewMesh();
+	return ReturnMesh;
+}
+
+
+
+void UDynamicMesh::SetMeshGenerator(TObjectPtr<UDynamicMeshGenerator> NewGenerator)
+{
+	MeshGenerator = NewGenerator;
+}
+
+void UDynamicMesh::ClearMeshGenerator()
+{
+	MeshGenerator = nullptr;
+}
+
+void UDynamicMesh::Regenerate()
+{
+	if (MeshGenerator != nullptr && bEnableMeshGenerator)
+	{
+		Reset();
+	}
+}
+
+
+void UDynamicMesh::PostRealtimeUpdate() 
+{ 
+	MeshRealtimeUpdateEvent.Broadcast(this); 
+}
+
+
+void UDynamicMesh::ApplyChange(const FMeshVertexChange* Change, bool bRevert)
+{
+	FDynamicMeshChangeInfo ChangeInfo;
+	ChangeInfo.Type = EDynamicMeshChangeType::MeshVertexChange;
+	ChangeInfo.VertexChange = Change;
+	ChangeInfo.bIsRevertChange = bRevert;
+
+	EditMeshInternal([&](FDynamicMesh3& EditMesh)
+	{
+		bool bHavePositions = Change->bHaveVertexPositions;
+		bool bHaveColors = Change->bHaveVertexColors && EditMesh.HasVertexColors();
+
+		int32 NV = Change->Vertices.Num();
+		const TArray<FVector3d>& Positions = (bRevert) ? Change->OldPositions : Change->NewPositions;
+		const TArray<FVector3f>& Colors = (bRevert) ? Change->OldColors : Change->NewColors;
+		for (int32 k = 0; k < NV; ++k)
+		{
+			int32 vid = Change->Vertices[k];
+			if (EditMesh.IsVertex(vid))
+			{
+				if (bHavePositions)
+				{
+					EditMesh.SetVertex(vid, Positions[k]);
+				}
+				if (bHaveColors)
+				{
+					EditMesh.SetVertexColor(vid, Colors[k]);
+				}
+			}
+		}
+
+		if (Change->bHaveOverlayNormals && EditMesh.HasAttributes() && EditMesh.Attributes()->PrimaryNormals())
+		{
+			FDynamicMeshNormalOverlay* Overlay = EditMesh.Attributes()->PrimaryNormals();
+			int32 NumNormals = Change->Normals.Num();
+			const TArray<FVector3f>& UseNormals = (bRevert) ? Change->OldNormals : Change->NewNormals;
+			for (int32 k = 0; k < NumNormals; ++k)
+			{
+				int32 elemid = Change->Normals[k];
+				if (Overlay->IsElement(elemid))
+				{
+					Overlay->SetElement(elemid, UseNormals[k]);
+				}
+			}
+		}
+	}, ChangeInfo);
+}
+
+void UDynamicMesh::ApplyChange(const FMeshChange* Change, bool bRevert)
+{
+	FDynamicMeshChangeInfo ChangeInfo;
+	ChangeInfo.Type = EDynamicMeshChangeType::MeshChange;
+	ChangeInfo.MeshChange = Change;
+	ChangeInfo.bIsRevertChange = bRevert;
+
+	EditMeshInternal([&](FDynamicMesh3& EditMesh)
+	{
+		Change->ApplyChangeToMesh(&EditMesh, bRevert);
+	}, ChangeInfo);
+}
+
+void UDynamicMesh::ApplyChange(const FMeshReplacementChange* Change, bool bRevert)
+{
+	FDynamicMeshChangeInfo ChangeInfo;
+	ChangeInfo.Type = EDynamicMeshChangeType::MeshReplacementChange;
+	ChangeInfo.ReplaceChange = Change;
+	ChangeInfo.bIsRevertChange = bRevert;
+
+	EditMeshInternal([&](FDynamicMesh3& EditMesh)
+	{
+		EditMesh.Copy(*Change->GetMesh(bRevert));
+	}, ChangeInfo);
+}
+
+
+void UDynamicMesh::Serialize(FArchive& Ar)
+{
+	TRACE_CPUPROFILER_EVENT_SCOPE(UDynamicMesh::Serialize);
+
+	Super::Serialize(Ar);
+
+	// do not serialize mesh for transactions
+	if (Ar.IsTransacting())
+	{
+		return;
+	}
+
+	Ar.UsingCustomVersion(FUE5MainStreamObjectVersion::GUID);
+
+	if (Ar.IsLoading())
+	{
+		InitializeNewMesh();
+	}
+	Ar << *Mesh;
+}
+
+
+
+
+
+
+/**
+ * This is an internal singleton used to support UDynamicMesh::ExportCustomProperties() and UDynamicMesh::ImportCustomProperties() below.
+ * 
+ */
+namespace UE 
+{
+namespace Local
+{
+
+class FDynamicMeshCopyHelper
+{
+protected:
+	struct FStashedMesh
+	{
+		TSoftObjectPtr<UDynamicMesh> SourceMesh;
+		FDateTime Timestamp;
+	};
+
+	static TMap<int32, FStashedMesh> StashedMeshes;
+	static FRandomStream KeyGenerator;
+public:
+
+	static void Initialize()
+	{
+		static bool bInitialized = false;
+		if (!bInitialized)
+		{
+			KeyGenerator = FRandomStream( (int32)FDateTime::UtcNow().ToUnixTimestamp() );
+			bInitialized = true;
+		}
+	}
+
+	static void DiscardExpiredMeshes()
+	{
+		const int32 ExpiryTimeoutInSeconds = CVarDynamicMeshDupeHelperTimeout.GetValueOnGameThread();
+
+		TArray<int32> ToRemove;
+		for (TPair<int32, FStashedMesh>& Pair : StashedMeshes)
+		{
+			if ((FDateTime::Now() - Pair.Value.Timestamp).GetTotalSeconds() > (double)ExpiryTimeoutInSeconds)
+			{
+				ToRemove.Add(Pair.Key);
+			}
+		}
+		for (int32 id : ToRemove)
+		{
+			StashedMeshes.Remove(id);
+		}
+	}
+
+	static int32 StashMeshReference(UDynamicMesh* SourceMesh)
+	{
+		Initialize();
+		int32 MeshKey = FMath::Abs((int32)KeyGenerator.GetUnsignedInt());
+
+		FStashedMesh NewMesh;
+		NewMesh.Timestamp = FDateTime::Now();
+		NewMesh.SourceMesh = SourceMesh;
+		StashedMeshes.Add(MeshKey, MoveTemp(NewMesh));
+
+		DiscardExpiredMeshes();
+		return MeshKey;
+	}
+
+	static TUniquePtr<FDynamicMesh3> ExtractStashedMesh(int32 MeshKey)
+	{
+		TUniquePtr<FDynamicMesh3> ReturnMesh;
+		FStashedMesh* Found = StashedMeshes.Find(MeshKey);
+		if (Found != nullptr && Found->SourceMesh.IsValid())
+		{
+			ReturnMesh = MakeUnique<FDynamicMesh3>();
+			UDynamicMesh* SourceMesh = Found->SourceMesh.Get();
+			if (SourceMesh)
+			{
+				// should we Cast<UDynamicMesh> here to make sure this is still the same UClass? (due to recyling)
+				SourceMesh->ProcessMesh([&](const FDynamicMesh3& EditMesh)
+				{
+					*ReturnMesh = EditMesh;
+				});
+			}
+		}
+
+		DiscardExpiredMeshes();
+		return ReturnMesh;
+	}
+
+};
+TMap<int32, FDynamicMeshCopyHelper::FStashedMesh> FDynamicMeshCopyHelper::StashedMeshes;
+FRandomStream FDynamicMeshCopyHelper::KeyGenerator;
+
+} } // end namespace UE::Local
+
+
+void UDynamicMesh::ExportCustomProperties(FOutputDevice& Out, uint32 Indent)
+{
+	// ignore empty meshes
+	if (Mesh->TriangleCount() == 0 && Mesh->VertexCount() == 0)
+	{
+		return;
+	}
+
+	// In the Editor, Copy/Paste of Actors and Components is not based on serialization,
+	// but rather on T3D structured text records. This is what allows things to be copied
+	// between Editor sessions, for example. Duplicate functionality in the Editor is implemented
+	// as Copy and Paste, so for Duplicate to work, any non-UProperty data must be serialized via text.
+	//
+	// This poses a problem for large data, such as large meshes. Serializing a mesh with millions of
+	// triangles as structured text is very expensive. An alternative is to serialize it to binary and 
+	// then use Base64 encoding. This is implemented below, however even that is quite slow. 
+	// 
+	// So, we employ a second strategy (read: hack), of just passing the source UObject pointer 
+	// via text. Instead of directly sending the pointer, we store it in FDynamicMeshCopyHelper, and pass
+	// an integer key instead. This allows for somewhat better handling, eg the copy helper could for example
+	// keep the UDynamicMeshes alive (so far this has not been necessary) and the keys are randomized so
+	// even in the unlikely even that two Editor instances end up with the same keys, it would just
+	// result in an correct pasted object, rather than accessing garbage pointers.
+	// 
+	// FDynamicMeshCopyHelper does attempt to discard "old" references, the CVar geometry.DynamicMesh.DupeStashTimeout 
+	// controls the definition of old (currenly 5 minutes). One effect this can have is if one does a
+	// copy and then a paste after the timeout, the mesh will not be found. This could also be problematic
+	// for large full-scene copies of hundreds of objects that takes longer than the timeout (resolvable via the CVar)
+	// 
+	// Obviously this does not work between Editor sessions. So we also optionally do a Base64 binary encoding if
+	// the mesh triangle count is below a CVar threshold (geometry.DynamicMesh.TextBasedDupeTriThreshold) 
+	// defaulting to 1000. Larger meshes can be supported by increasing the CVar value if need be.
+	// 
+	// If the mesh is not found in the FDynamicMeshCopyHelper, and is too large to text-copy, then instead
+	// of leaving an empty mesh, we emit a cube, as empty meshes can be problematic. A warning is also
+	// printed to the Output Log, pointing the user to the CVars.
+	// 
+	// Possible todos: UObject Reuse/Recycling
+	//
+
+	Super::ExportCustomProperties(Out, Indent);
+	Out.Logf(TEXT("%sCustomProperties "), FCString::Spc(Indent));
+	Out.Logf(TEXT("MeshData "));
+
+	// stash copy to circumvent expensive T3D generation/parsing
+	int32 StashedKey = UE::Local::FDynamicMeshCopyHelper::StashMeshReference(this);
+	Out.Logf(TEXT("MESHKEY=%d "), StashedKey);
+
+	if (Mesh->TriangleCount() < CVarDynamicMeshTextBasedDupeTriangleCountThreshold.GetValueOnGameThread())
+	{
+		// serialize our mesh
+		TArray<uint8> MeshWriteBuffer;
+		FMemoryWriter MemWriter(MeshWriteBuffer);
+		Mesh->Serialize(MemWriter);
+
+		FString Base64String = FBase64::Encode(MeshWriteBuffer);
+		//TArray<uint8> MeshReadBuffer;
+		//ensure(FBase64::Decode(Base64String, MeshReadBuffer));		// test decode
+
+		// Base64 encoding uses the '/' character, but T3D interprets '//' as some kind of
+		// terminator (?). If it occurs then the string passed to ImportCustomProperties() will
+		// come back as full of nullptrs. So we will swap in '-' here, and swap back to '/' in ImportCustomProperties()
+		for (int32 k = 0; k < Base64String.Len(); ++k)
+		{
+			if (Base64String[k] == '/')
+			{
+				Base64String[k] = '-';
+			}
+		}
+
+		Out.Logf(TEXT("MESHDATALEN=%d MESHDATA=%s"), Base64String.Len(), *Base64String);
+	}
+
+	Out.Logf(TEXT("\r\n"));
+}
+
+
+void UDynamicMesh::ImportCustomProperties(const TCHAR* SourceText, FFeedbackContext* Warn)
+{
+	Super::ImportCustomProperties(SourceText, Warn);
+
+	if (FParse::Command(&SourceText, TEXT("MeshData")))	
+	{
+		static const TCHAR MeshKeyToken[] = TEXT("MESHKEY=");
+		const TCHAR* FoundMeshKeyStart = FCString::Strifind(SourceText, MeshKeyToken);
+		if (FoundMeshKeyStart)
+		{
+			SourceText = FoundMeshKeyStart + FCString::Strlen(MeshKeyToken);
+			int32 MeshKey = FCString::Atoi(SourceText);
+			TUniquePtr<FDynamicMesh3> FoundMesh = UE::Local::FDynamicMeshCopyHelper::ExtractStashedMesh(MeshKey);
+			if (FoundMesh.IsValid() && FoundMesh->TriangleCount() > 0)
+			{
+				InitializeNewMesh();
+				SetMesh(MoveTemp(*FoundMesh));
+				return;
+			}
+		}
+
+		static const TCHAR MeshDataLenToken[] = TEXT("MESHDATALEN=");
+		const TCHAR* FoundMeshDataLenStart = FCString::Strifind(SourceText, MeshDataLenToken);
+		if (FoundMeshDataLenStart)
+		{
+			SourceText = FoundMeshDataLenStart + FCString::Strlen(MeshDataLenToken);
+			int32 MeshDataLen = FCString::Atoi(SourceText);
+
+			static const TCHAR MeshDataToken[] = TEXT("MESHDATA=");
+			const TCHAR* FoundMeshDataStart = FCString::Strifind(SourceText, MeshDataToken);
+			if (FoundMeshDataStart)
+			{
+				SourceText = FoundMeshDataStart + FCString::Strlen(MeshDataToken);
+				FString MeshData(MeshDataLen, SourceText);
+
+				// fix-up the hack applied to the Base64-encoded string in ExportCustomProperties()
+				for (int32 k = 0; k < MeshData.Len(); ++k)
+				{
+					if (MeshData[k] == '-')
+					{
+						MeshData[k] = '/';
+					}
+				}
+
+				TArray<uint8> MeshReadBuffer;
+				bool bDecoded = FBase64::Decode(MeshData, MeshReadBuffer);
+				if (bDecoded)
+				{
+					FMemoryReader MemReader(MeshReadBuffer);
+					FDynamicMesh3 NewMesh;
+					NewMesh.Serialize(MemReader);
+
+					InitializeNewMesh();
+					SetMesh(MoveTemp(NewMesh));
+					return;
+				}
+			}
+		}
+
+		// if we got here we failed. Rather than produce an empty mesh, we generate a small cube
+		UE_LOG(LogTemp, Warning, TEXT("UDynamicMesh text-based property serialization incomplete, generating box as placeholder. Try increasing geometry.DynamicMesh.TextBasedDupeTriThreshold, or geometry.DynamicMesh.DupeStashTimeout."))
+		FMinimalBoxMeshGenerator BoxGen;
+		BoxGen.Box = FOrientedBox3d(FVector3d::Zero(), 50.0 * FVector3d::One());
+		FDynamicMesh3 GenMesh(&BoxGen.Generate());
+		GenMesh.Attributes()->EnableMaterialID();
+		SetMesh(MoveTemp(GenMesh));
+	}
+}
+
+
+
+
+
+//
+// Pool support for blueprints
+//
+
+static TAutoConsoleVariable<int32> CVarDynamicMeshPoolMaxPoolSizeThreshold(
+	TEXT("geometry.DynamicMesh.MaxPoolSize"),
+	1000,
+	TEXT("Maximum number of meshes a UDynamicMeshPool will allow to be in the pool before running garbage collection"));
+
+
+
+UDynamicMesh* UDynamicMeshPool::RequestMesh()
+{
+	if (CachedMeshes.Num() > 0)
+	{
+		return CachedMeshes.Pop(false);
+	}
+	UDynamicMesh* NewMesh = NewObject<UDynamicMesh>();
+
+	// If we have allocated more meshes than our safety threshold, drop our holds on the existing meshes.
+	// This will allow them to be garbage-collected (eventually)
+	if (!ensure(AllCreatedMeshes.Num() < CVarDynamicMeshPoolMaxPoolSizeThreshold.GetValueOnGameThread()))
+	{
+		UE_LOG(LogTemp, Warning, TEXT("UDynamicMeshPool Threshold of %d Allocated Meshes exceeded! Releasing references to all current meshes and forcing a garbage collection."), CVarDynamicMeshPoolMaxPoolSizeThreshold.GetValueOnGameThread());
+		AllCreatedMeshes.Reset();
+		GEngine->ForceGarbageCollection(true);
+	}
+
+	AllCreatedMeshes.Add(NewMesh);
+	return NewMesh;
+}
+
+
+
+void UDynamicMeshPool::ReturnMesh(UDynamicMesh* Mesh)
+{
+	if (ensure(Mesh))
+	{
+		Mesh->Reset();
+		if (ensure(CachedMeshes.Contains(Mesh) == false))
+		{
+			CachedMeshes.Add(Mesh);
+		}
+	}
+}
+
+
+void UDynamicMeshPool::ReturnAllMeshes()
+{
+	CachedMeshes = AllCreatedMeshes;
+	for (UDynamicMesh* Mesh : CachedMeshes)
+	{
+		if (Mesh)
+		{
+			Mesh->Reset();
+		}
+	}
+
+	// TODO: this may be vestigial code, unclear how it could be hit
+	int32 Removed = CachedMeshes.RemoveAll([](UDynamicMesh* Mesh) { return Mesh == nullptr; });
+	ensure(Removed == 0);
+}
+
+void UDynamicMeshPool::FreeAllMeshes()
+{
+	CachedMeshes.Reset();
+	AllCreatedMeshes.Reset();
+}
\ No newline at end of file
diff --git a/Engine/Source/Runtime/Experimental/GeometryFramework/Public/Changes/MeshChange.h b/Engine/Source/Runtime/Experimental/GeometryFramework/Public/Changes/MeshChange.h
new file mode 100644
--- /dev/null
+++ b/Engine/Source/Runtime/Experimental/GeometryFramework/Public/Changes/MeshChange.h
@@ -0,0 +1,70 @@
+// Copyright Epic Games, Inc. All Rights Reserved.
+
+#pragma once
+
+#include "CoreMinimal.h"
+#include "UObject/Interface.h"
+#include "InteractiveToolChange.h"
+#include "VectorTypes.h"
+#include "DynamicMesh/DynamicMeshChangeTracker.h"
+#include "MeshChange.generated.h"
+
+//class FDynamicMeshChange;		// need to refactor this out of DynamicMeshChangeTracker
+
+
+
+
+/**
+ * FMeshChange represents an undoable change to a FDynamicMesh3.
+ * Currently only valid to call Apply/Revert when the Object is a one of several components backed by FDynamicMesh: UDynamicMeshComponent, UOctreeDynamicMeshComponent, UPreviewMesh
+ */
+class GEOMETRYFRAMEWORK_API FMeshChange : public FToolCommandChange
+{
+public:
+	FMeshChange();
+	FMeshChange(TUniquePtr<UE::Geometry::FDynamicMeshChange> DynamicMeshChangeIn);
+
+	TUniquePtr<UE::Geometry::FDynamicMeshChange> DynamicMeshChange;
+
+	/** This function is called on Apply and Revert (last argument is true on Apply)*/
+	TFunction<void(FMeshChange*, UObject*, bool)> OnChangeAppliedFunc;
+
+	/** Makes the change to the object */
+	virtual void Apply(UObject* Object) override;
+
+	/** Reverts change to the object */
+	virtual void Revert(UObject* Object) override;
+
+	/** Describes this change (for debugging) */
+	virtual FString ToString() const override;
+
+
+	/** 
+	 * Apply embedded DynamicMeshChange to given Mesh. This function is for 
+	 * change-targets to call, when passed a FMeshChange to apply to a Mesh they own. 
+	 * This allows FMeshChange subclasses to customize the change behavior if necessary.
+	 * The default behavior just forwards the call to DynamicMeshChange->Apply(Mesh, bRevert).
+	 */
+	virtual void ApplyChangeToMesh(UE::Geometry::FDynamicMesh3* Mesh, bool bRevert) const;
+};
+
+
+
+
+
+UINTERFACE()
+class GEOMETRYFRAMEWORK_API UMeshCommandChangeTarget : public UInterface
+{
+	GENERATED_BODY()
+};
+/**
+ * IMeshCommandChangeTarget is an interface which is used to apply a mesh change
+ */
+class GEOMETRYFRAMEWORK_API IMeshCommandChangeTarget
+{
+	GENERATED_BODY()
+public:
+	virtual void ApplyChange(const FMeshChange* Change, bool bRevert) = 0;
+};
+
+
diff --git a/Engine/Source/Runtime/Experimental/GeometryFramework/Public/Changes/MeshReplacementChange.h b/Engine/Source/Runtime/Experimental/GeometryFramework/Public/Changes/MeshReplacementChange.h
new file mode 100644
--- /dev/null
+++ b/Engine/Source/Runtime/Experimental/GeometryFramework/Public/Changes/MeshReplacementChange.h
@@ -0,0 +1,64 @@
+// Copyright Epic Games, Inc. All Rights Reserved.
+
+#pragma once
+
+#include "CoreMinimal.h"
+#include "UObject/Interface.h"
+#include "InteractiveToolChange.h"
+#include "GeometryBase.h"
+#include "MeshReplacementChange.generated.h"
+
+
+PREDECLARE_USE_GEOMETRY_CLASS(FDynamicMesh3);
+
+
+/**
+ * FMeshReplacementChange represents an undoable *complete* change to a FDynamicMesh3.
+ * Currently only valid to call Apply/Revert when the Object is a UDynamicMeshComponent
+ */
+class GEOMETRYFRAMEWORK_API FMeshReplacementChange : public FToolCommandChange
+{
+	TSharedPtr<const FDynamicMesh3, ESPMode::ThreadSafe> Before, After;
+
+public:
+	FMeshReplacementChange();
+	FMeshReplacementChange(TSharedPtr<const FDynamicMesh3, ESPMode::ThreadSafe> Before, TSharedPtr<const FDynamicMesh3, ESPMode::ThreadSafe> After);
+
+	const TSharedPtr<const FDynamicMesh3, ESPMode::ThreadSafe>& GetMesh(bool bRevert) const
+	{
+		return bRevert ? Before : After;
+	}
+
+	/** This function is called on Apply and Revert (last argument is true on Apply)*/
+	TFunction<void(FMeshReplacementChange*, UObject*, bool)> OnChangeAppliedFunc;
+
+	/** Makes the change to the object */
+	virtual void Apply(UObject* Object) override;
+
+	/** Reverts change to the object */
+	virtual void Revert(UObject* Object) override;
+
+	/** Describes this change (for debugging) */
+	virtual FString ToString() const override;
+};
+
+
+
+
+
+UINTERFACE()
+class GEOMETRYFRAMEWORK_API UMeshReplacementCommandChangeTarget : public UInterface
+{
+	GENERATED_BODY()
+};
+/**
+ * IMeshReplacementCommandChangeTarget is an interface which is used to apply a mesh replacement change
+ */
+class GEOMETRYFRAMEWORK_API IMeshReplacementCommandChangeTarget
+{
+	GENERATED_BODY()
+public:
+	virtual void ApplyChange(const FMeshReplacementChange* Change, bool bRevert) = 0;
+};
+
+
diff --git a/Engine/Source/Runtime/Experimental/GeometryFramework/Public/Changes/MeshVertexChange.h b/Engine/Source/Runtime/Experimental/GeometryFramework/Public/Changes/MeshVertexChange.h
new file mode 100644
--- /dev/null
+++ b/Engine/Source/Runtime/Experimental/GeometryFramework/Public/Changes/MeshVertexChange.h
@@ -0,0 +1,139 @@
+// Copyright Epic Games, Inc. All Rights Reserved.
+
+#pragma once
+
+#include "CoreMinimal.h"
+#include "UObject/Interface.h"
+#include "InteractiveToolChange.h"
+#include "VectorTypes.h"
+#include "MeshVertexChange.generated.h"
+
+PREDECLARE_USE_GEOMETRY_CLASS(FDynamicMesh3);
+
+/**
+ * FMeshVertexChange represents an reversible change to a set of vertex positions.
+ * Currently only a UDynamicMeshComponent target is supported.
+ * 
+ * @todo support optionally storing old/new normals and tangents
+ * @todo support applying to a StaticMeshComponent/MeshDescription ?
+ */
+class GEOMETRYFRAMEWORK_API FMeshVertexChange : public FToolCommandChange
+{
+public:
+	bool bHaveVertexPositions = true;
+	bool bHaveVertexColors = false;
+	TArray<int32> Vertices;
+	TArray<FVector3d> OldPositions;
+	TArray<FVector3d> NewPositions;
+	TArray<FVector3f> OldColors;
+	TArray<FVector3f> NewColors;
+
+	bool bHaveOverlayNormals = false;
+	TArray<int32> Normals;
+	TArray<FVector3f> OldNormals;
+	TArray<FVector3f> NewNormals;
+
+	/** Makes the change to the object */
+	virtual void Apply(UObject* Object) override;
+
+	/** Reverts change to the object */
+	virtual void Revert(UObject* Object) override;
+
+	/** Describes this change (for debugging) */
+	virtual FString ToString() const override;
+};
+
+
+enum class EMeshVertexChangeComponents : uint8
+{
+	None = 0,
+	VertexPositions = 1,
+	VertexColors = 2,
+	OverlayNormals = 16
+};
+ENUM_CLASS_FLAGS(EMeshVertexChangeComponents);
+
+
+/**
+ * FMeshVertexChangeBuilder can be used to construct a FMeshVertexChange.
+ */
+class GEOMETRYFRAMEWORK_API FMeshVertexChangeBuilder
+{
+public:
+	TUniquePtr<FMeshVertexChange> Change;
+	TMap<int32, int32> SavedVertices;
+	bool bSavePositions = true;
+	bool bSaveColors = false;
+
+	bool bSaveOverlayNormals = false;
+	TMap<int32, int32> SavedNormalElements;
+
+	/** If set, this function is called whenever a newly-seen VertexID is saved, parameters are (VertexID, Index) into saved-vertices array */
+	TUniqueFunction<void(int32, int32)> OnNewVertexSaved = nullptr;
+
+	FMeshVertexChangeBuilder();
+	explicit FMeshVertexChangeBuilder(EMeshVertexChangeComponents Components);
+
+	void SaveVertexInitial(const FDynamicMesh3* Mesh, int32 VertexID);
+	void SaveVertexFinal(const FDynamicMesh3* Mesh, int32 VertexID);
+
+	template<typename Enumerable>
+	void SaveVertices(const FDynamicMesh3* Mesh, Enumerable Enum, bool bInitial);
+
+	void SaveOverlayNormals(const FDynamicMesh3* Mesh, const TArray<int32>& ElementIDs, bool bInitial);
+	void SaveOverlayNormals(const FDynamicMesh3* Mesh, const TSet<int32>& ElementIDs, bool bInitial);
+
+
+public:
+	// currently only used in vertex sculpt tool. cannot be used if bSaveColors = true
+	void UpdateVertex(int32 VertexID, const FVector3d& OldPosition, const FVector3d& NewPosition);
+
+	// currently only used in element paint tool. Can only be used if bSaveColors=true and bSavePositions=false
+	void UpdateVertexColor(int32 VertexID, const FVector3f& OldColor, const FVector3f& NewColor);
+
+protected:
+	void UpdateVertexFinal(int32 VertexID, const FVector3d& NewPosition);
+
+	void UpdateOverlayNormal(int32 ElementID, const FVector3f& OldNormal, const FVector3f& NewNormal);
+	void UpdateOverlayNormalFinal(int32 ElementID, const FVector3f& NewNormal);
+};
+
+
+
+UINTERFACE()
+class GEOMETRYFRAMEWORK_API UMeshVertexCommandChangeTarget : public UInterface
+{
+	GENERATED_BODY()
+};
+/**
+ * IMeshVertexCommandChangeTarget is an interface which is used to apply a FMeshVertexChange
+ */
+class GEOMETRYFRAMEWORK_API IMeshVertexCommandChangeTarget
+{
+	GENERATED_BODY()
+public:
+	virtual void ApplyChange(const FMeshVertexChange* Change, bool bRevert) = 0;
+};
+
+
+
+
+
+template<typename Enumerable>
+void FMeshVertexChangeBuilder::SaveVertices(const FDynamicMesh3* Mesh, Enumerable Enum, bool bInitial)
+{
+	if (bInitial)
+	{
+		for (int32 k : Enum)
+		{
+			SaveVertexInitial(Mesh, k);
+		}
+	}
+	else
+	{
+		for (int32 k : Enum)
+		{
+			SaveVertexFinal(Mesh, k);
+		}
+	}
+}
diff --git a/Engine/Source/Runtime/Experimental/GeometryFramework/Public/Components/BaseDynamicMeshComponent.h b/Engine/Source/Runtime/Experimental/GeometryFramework/Public/Components/BaseDynamicMeshComponent.h
new file mode 100644
--- /dev/null
+++ b/Engine/Source/Runtime/Experimental/GeometryFramework/Public/Components/BaseDynamicMeshComponent.h
@@ -0,0 +1,339 @@
+// Copyright Epic Games, Inc. All Rights Reserved.
+
+#pragma once
+
+#include "CoreMinimal.h"
+#include "Components/MeshComponent.h"
+#include "InteractiveToolObjects.h"
+#include "Changes/MeshVertexChange.h"
+#include "Changes/MeshChange.h"
+#include "Changes/MeshReplacementChange.h"
+#include "MeshConversionOptions.h"
+#include "DynamicMesh/DynamicMesh3.h"		// todo replace with predeclaration (lots of fallout)
+#include "UDynamicMesh.h"
+
+#include "BaseDynamicMeshComponent.generated.h"
+
+// predecl
+struct FMeshDescription;
+class FMeshVertexChange;
+class FMeshChange;
+class FBaseDynamicMeshSceneProxy;
+using UE::Geometry::FDynamicMesh3;
+
+/**
+ * EMeshRenderAttributeFlags is used to identify different mesh rendering attributes, for things
+ * like fast-update functions
+ */
+enum class EMeshRenderAttributeFlags : uint8
+{
+	None = 0,
+	Positions = 1,
+	VertexColors = 1<<1,
+	VertexNormals = 1<<2,
+	VertexUVs = 1<<3,
+	SecondaryIndexBuffers = 1<<4,
+
+	AllVertexAttribs = Positions | VertexColors | VertexNormals | VertexUVs
+};
+ENUM_CLASS_FLAGS(EMeshRenderAttributeFlags);
+
+
+
+/**
+ * Tangent calculation modes
+ */
+UENUM()
+enum class EDynamicMeshComponentTangentsMode : uint8
+{
+	/** Tangents are not used/available, proceed accordingly (eg generate arbitrary orthogonal basis) */
+	NoTangents,
+	/** Tangents will be automatically calculated on demand. Note that mesh changes due to tangents calculation will *not* be broadcast via MeshChange events! */
+	AutoCalculated,
+	/** Tangents are externally provided via the FDynamicMesh3 AttributeSet */
+	ExternallyProvided
+};
+
+
+/**
+ * UBaseDynamicMeshComponent is a base interface for a UMeshComponent based on a UDynamicMesh.
+ */
+UCLASS(Abstract, hidecategories = (LOD), ClassGroup = Rendering)
+class GEOMETRYFRAMEWORK_API UBaseDynamicMeshComponent : 
+	public UMeshComponent, 
+	public IToolFrameworkComponent, 
+	public IMeshVertexCommandChangeTarget, 
+	public IMeshCommandChangeTarget, 
+	public IMeshReplacementCommandChangeTarget
+{
+	GENERATED_UCLASS_BODY()
+
+
+	//===============================================================================================================
+	// UBaseDynamicMeshComponent API. Subclasses must implement these functions
+	//
+public:
+	/**
+	 * initialize the internal mesh from a DynamicMesh
+	 */
+	virtual void SetMesh(UE::Geometry::FDynamicMesh3&& MoveMesh)
+	{
+		unimplemented();
+	}
+
+	/**
+	 * @return pointer to internal mesh
+	 * @warning avoid usage of this function, access via GetDynamicMesh() instead
+	 */
+	virtual FDynamicMesh3* GetMesh()
+	{
+		unimplemented();
+		return nullptr;
+	}
+
+	/**
+	 * @return pointer to internal mesh
+	 */
+	virtual const FDynamicMesh3* GetMesh() const
+	{
+		unimplemented();
+		return nullptr;
+	}
+
+	/**
+	 * Allow external code to read the internal mesh.
+	 */
+	virtual void ProcessMesh(TFunctionRef<void(const UE::Geometry::FDynamicMesh3&)> ProcessFunc) const
+	{
+		unimplemented();
+	}
+
+	/**
+	 * @return the child UDynamicMesh
+	 */
+	UFUNCTION(BlueprintCallable, Category = "Dynamic Mesh Component")
+	virtual UDynamicMesh* GetDynamicMesh() 
+	{ 
+		unimplemented();
+		return nullptr; 
+	}
+
+	/**
+	 * Call this if you update the mesh via GetMesh()
+	 * @todo should provide a function that calls a lambda to modify the mesh, and only return const mesh pointer
+	 */
+	virtual void NotifyMeshUpdated()
+	{
+		unimplemented();
+	}
+
+	/**
+	 * Apply a vertex deformation change to the internal mesh  (implements IMeshVertexCommandChangeTarget)
+	 */
+	virtual void ApplyChange(const FMeshVertexChange* Change, bool bRevert) override
+	{
+		unimplemented();
+	}
+
+	/**
+	 * Apply a general mesh change to the internal mesh  (implements IMeshCommandChangeTarget)
+	 */
+	virtual void ApplyChange(const FMeshChange* Change, bool bRevert) override
+	{
+		unimplemented();
+	}
+
+	/**
+	* Apply a full mesh replacement change to the internal mesh  (implements IMeshReplacementCommandChangeTarget)
+	*/
+	virtual void ApplyChange(const FMeshReplacementChange* Change, bool bRevert) override
+	{
+		unimplemented();
+	}
+
+	/**
+	 * Apply a transform to the mesh
+	 */
+	virtual void ApplyTransform(const UE::Geometry::FTransform3d& Transform, bool bInvert)
+	{
+		unimplemented();
+	}
+
+
+protected:
+
+	/**
+	 * Subclass must implement this to return scene proxy if available, or nullptr
+	 */
+	virtual FBaseDynamicMeshSceneProxy* GetBaseSceneProxy()
+	{
+		unimplemented();
+		return nullptr;
+	}
+
+	/**
+	 * Subclass must implement this to notify allocated proxies of updated materials
+	 */
+	virtual void NotifyMaterialSetUpdated()
+	{
+		unimplemented();
+	}
+
+
+
+	//===============================================================================================================
+	// Built-in Wireframe-on-Shaded Rendering support. The wireframe looks terrible but this is a convenient 
+	// way to enable/disable it.
+	//
+public:
+	/**
+	 * if true, we always show the wireframe on top of the shaded mesh, even when not in wireframe mode
+	 * @todo: this should not be public, access via Set/Get once all usage is cleaned up
+	 */
+	UPROPERTY(EditAnywhere, BlueprintReadWrite, Category = "Dynamic Mesh Component", meta = (DisplayName = "Wireframe Overlay") )
+	bool bExplicitShowWireframe = false;
+
+	/**
+	 * Configure whether wireframe rendering is enabled or not
+	 */
+	UFUNCTION(BlueprintCallable, Category = "Dynamic Mesh Component")
+	virtual void SetEnableWireframeRenderPass(bool bEnable) { bExplicitShowWireframe = bEnable; }
+
+	/**
+	 * @return true if wireframe rendering pass is enabled
+	 */
+	UFUNCTION(BlueprintCallable, Category = "Dynamic Mesh Component")
+	virtual bool GetEnableWireframeRenderPass() const { return bExplicitShowWireframe; }
+
+
+
+	//===============================================================================================================
+	// API for changing Rendering settings. Although some of these settings are available publicly
+	// on the Component (in some cases as public members), generally changing them requires more complex 
+	// Rendering invalidation.
+	//
+public:
+
+	UFUNCTION(BlueprintCallable, Category = "Dynamic Mesh Component")
+	virtual void SetShadowsEnabled(bool bEnabled);
+
+	UFUNCTION(BlueprintCallable, Category = "Dynamic Mesh Component")
+	virtual bool GetShadowsEnabled() const { return CastShadow; }
+
+
+
+	//===============================================================================================================
+	// Override rendering material support. If an Override material is set, then it
+	// will be used during drawing of all mesh buffers except Secondary buffers.
+	//
+public:
+	/**
+	 * Set an active override render material. This should replace all materials during rendering.
+	 */
+	UFUNCTION(BlueprintCallable, Category = "Dynamic Mesh Component")
+	virtual void SetOverrideRenderMaterial(UMaterialInterface* Material);
+
+	/**
+	 * Clear any active override render material
+	 */
+	UFUNCTION(BlueprintCallable, Category = "Dynamic Mesh Component")
+	virtual void ClearOverrideRenderMaterial();
+	
+	/**
+	 * @return true if an override render material is currently enabled for the given MaterialIndex
+	 */
+	UFUNCTION(BlueprintCallable, Category = "Dynamic Mesh Component")
+	virtual bool HasOverrideRenderMaterial(int k) const
+	{
+		return OverrideRenderMaterial != nullptr;
+	}
+
+	/**
+	 * @return active override render material for the given MaterialIndex
+	 */
+	UFUNCTION(BlueprintCallable, Category = "Dynamic Mesh Component")
+	virtual UMaterialInterface* GetOverrideRenderMaterial(int MaterialIndex) const
+	{
+		return OverrideRenderMaterial;
+	}
+
+protected:
+	UPROPERTY()
+	TObjectPtr<UMaterialInterface> OverrideRenderMaterial = nullptr;
+
+
+
+
+
+	//===============================================================================================================
+	// Secondary Render Buffers support. This requires implementation in subclasses. It allows
+	// a subset of the mesh triangles to be moved to a separate set of render buffers, which
+	// can then have a separate material (eg to highlight faces), or be shown/hidden independently.
+	// 
+public:
+	/**
+	 * Set an active secondary render material. 
+	 */
+	UFUNCTION(BlueprintCallable, Category = "Dynamic Mesh Component")
+	virtual void SetSecondaryRenderMaterial(UMaterialInterface* Material);
+
+	/**
+	 * Clear any active secondary render material
+	 */
+	UFUNCTION(BlueprintCallable, Category = "Dynamic Mesh Component")
+	virtual void ClearSecondaryRenderMaterial();
+
+	/**
+	 * @return true if a secondary render material is set
+	 */
+	virtual bool HasSecondaryRenderMaterial() const
+	{
+		return SecondaryRenderMaterial != nullptr;
+	}
+
+	/**
+	 * @return active secondary render material
+	 */
+	UFUNCTION(BlueprintCallable, Category = "Dynamic Mesh Component")
+	virtual UMaterialInterface* GetSecondaryRenderMaterial() const
+	{
+		return SecondaryRenderMaterial;
+	}
+
+	/**
+	 * Show/Hide the secondary triangle buffers. Does not invalidate SceneProxy.
+	 */
+	UFUNCTION(BlueprintCallable, Category = "Dynamic Mesh Component")
+	virtual void SetSecondaryBuffersVisibility(bool bSetVisible);
+
+	/**
+	 * @return true if secondary buffers are currently set to be visible
+	 */
+	UFUNCTION(BlueprintCallable, Category = "Dynamic Mesh Component")
+	virtual bool GetSecondaryBuffersVisibility() const;
+
+protected:
+	UPROPERTY()
+	TObjectPtr<UMaterialInterface> SecondaryRenderMaterial = nullptr;
+
+	bool bDrawSecondaryBuffers = true;
+
+
+
+
+
+	//===============================================================================================================
+	// Standard Component interfaces
+	//
+public:
+
+	// UMeshComponent Interface.
+	virtual int32 GetNumMaterials() const override;
+	virtual UMaterialInterface* GetMaterial(int32 ElementIndex) const override;
+	virtual FMaterialRelevance GetMaterialRelevance(ERHIFeatureLevel::Type InFeatureLevel) const override;
+	virtual void SetMaterial(int32 ElementIndex, UMaterialInterface* Material) override;
+	virtual void GetUsedMaterials(TArray<UMaterialInterface*>& OutMaterials, bool bGetDebugMaterials = false) const override;
+
+	UPROPERTY()
+	TArray<TObjectPtr<UMaterialInterface>> BaseMaterials;
+};
diff --git a/Engine/Source/Runtime/Experimental/GeometryFramework/Public/Components/BaseDynamicMeshSceneProxy.h b/Engine/Source/Runtime/Experimental/GeometryFramework/Public/Components/BaseDynamicMeshSceneProxy.h
new file mode 100644
--- /dev/null
+++ b/Engine/Source/Runtime/Experimental/GeometryFramework/Public/Components/BaseDynamicMeshSceneProxy.h
@@ -0,0 +1,937 @@
+// Copyright Epic Games, Inc. All Rights Reserved. 
+
+#pragma once
+
+#include "CoreMinimal.h"
+#include "PrimitiveSceneProxy.h"
+#include "VertexFactory.h"
+#include "LocalVertexFactory.h"
+#include "PrimitiveViewRelevance.h"
+#include "Rendering/StaticMeshVertexBuffer.h"
+#include "Rendering/PositionVertexBuffer.h"
+#include "Rendering/ColorVertexBuffer.h"
+#include "DynamicMeshBuilder.h"
+#include "Components/BaseDynamicMeshComponent.h"
+
+using UE::Geometry::FDynamicMesh3;
+using UE::Geometry::FDynamicMeshAttributeSet;
+using UE::Geometry::FDynamicMeshUVOverlay;
+using UE::Geometry::FDynamicMeshNormalOverlay;
+using UE::Geometry::FDynamicMeshColorOverlay;
+using UE::Geometry::FDynamicMeshMaterialAttribute;
+
+/**
+ * FMeshRenderBufferSet stores a set of RenderBuffers for a mesh
+ */
+class FMeshRenderBufferSet
+{
+public:
+	/** Number of triangles in this renderbuffer set. Note that triangles may be split between IndexBuffer and SecondaryIndexBuffer. */
+	int TriangleCount = 0;
+
+	/** The buffer containing vertex data. */
+	FStaticMeshVertexBuffer StaticMeshVertexBuffer;
+	/** The buffer containing the position vertex data. */
+	FPositionVertexBuffer PositionVertexBuffer;
+	/** The buffer containing the vertex color data. */
+	FColorVertexBuffer ColorVertexBuffer;
+
+	/** triangle indices */
+	FDynamicMeshIndexBuffer32 IndexBuffer;
+
+	/** vertex factory */
+	FLocalVertexFactory VertexFactory;
+
+	/** Material to draw this mesh with */
+	UMaterialInterface* Material = nullptr;
+
+	/**
+	 * Optional list of triangles stored in this buffer. Storing this allows us
+	 * to rebuild the buffers if vertex data changes.
+	 */
+	TOptional<TArray<int>> Triangles;
+
+	/**
+	 * If secondary index buffer is enabled, we populate this index buffer with additional triangles indexing into the same vertex buffers
+	 */
+	bool bEnableSecondaryIndexBuffer = false;
+
+	/**
+	 * partition or subset of IndexBuffer that indexes into same vertex buffers
+	 */
+	FDynamicMeshIndexBuffer32 SecondaryIndexBuffer;
+
+
+	/**
+	 * In situations where we want to *update* the existing Vertex or Index buffers, we need to synchronize
+	 * access between the Game and Render threads. We use this lock to do that.
+	 */
+	FCriticalSection BuffersLock;
+
+
+	FMeshRenderBufferSet(ERHIFeatureLevel::Type FeatureLevelType)
+		: VertexFactory(FeatureLevelType, "FMeshRenderBufferSet")
+	{
+		StaticMeshVertexBuffer.SetUseFullPrecisionUVs(true);
+		StaticMeshVertexBuffer.SetUseHighPrecisionTangentBasis(true);
+	}
+
+
+	virtual ~FMeshRenderBufferSet()
+	{
+		check(IsInRenderingThread());
+
+		if (TriangleCount > 0)
+		{
+			PositionVertexBuffer.ReleaseResource();
+			StaticMeshVertexBuffer.ReleaseResource();
+			ColorVertexBuffer.ReleaseResource();
+			VertexFactory.ReleaseResource();
+			if (IndexBuffer.IsInitialized())
+			{
+				IndexBuffer.ReleaseResource();
+			}
+			if (SecondaryIndexBuffer.IsInitialized())
+			{
+				SecondaryIndexBuffer.ReleaseResource();
+			}
+		}
+	}
+
+
+	/**
+	 * Upload initialized mesh buffers. 
+	 * @warning This can only be called on the Rendering Thread.
+	 */
+	void Upload()
+	{
+		check(IsInRenderingThread());
+
+		if (TriangleCount == 0)
+		{
+			return;
+		}
+
+		InitOrUpdateResource(&this->PositionVertexBuffer);
+		InitOrUpdateResource(&this->StaticMeshVertexBuffer);
+		InitOrUpdateResource(&this->ColorVertexBuffer);
+
+		FLocalVertexFactory::FDataType Data;
+		this->PositionVertexBuffer.BindPositionVertexBuffer(&this->VertexFactory, Data);
+		this->StaticMeshVertexBuffer.BindTangentVertexBuffer(&this->VertexFactory, Data);
+		this->StaticMeshVertexBuffer.BindPackedTexCoordVertexBuffer(&this->VertexFactory, Data);
+		// currently no lightmaps support
+		//this->StaticMeshVertexBuffer.BindLightMapVertexBuffer(&this->VertexFactory, Data, LightMapIndex);
+		this->ColorVertexBuffer.BindColorVertexBuffer(&this->VertexFactory, Data);
+		this->VertexFactory.SetData(Data);
+
+		InitOrUpdateResource(&this->VertexFactory);
+		PositionVertexBuffer.InitResource();
+		StaticMeshVertexBuffer.InitResource();
+		ColorVertexBuffer.InitResource();
+		VertexFactory.InitResource();
+
+		if (IndexBuffer.Indices.Num() > 0)
+		{
+			IndexBuffer.InitResource();
+		}
+		if (bEnableSecondaryIndexBuffer && SecondaryIndexBuffer.Indices.Num() > 0)
+		{
+			SecondaryIndexBuffer.InitResource();
+		}
+	}
+
+
+	void UploadIndexBufferUpdate()
+	{
+		check(IsInRenderingThread());
+		if (IndexBuffer.Indices.Num() > 0)
+		{
+			InitOrUpdateResource(&IndexBuffer);
+		}
+		if (bEnableSecondaryIndexBuffer && SecondaryIndexBuffer.Indices.Num() > 0)
+		{
+			InitOrUpdateResource(&SecondaryIndexBuffer);
+		}
+	}
+
+
+	void UploadVertexUpdate(bool bPositions, bool bMeshAttribs, bool bColors)
+	{
+		check(IsInRenderingThread());
+
+		if (TriangleCount == 0)
+		{
+			return;
+		}
+
+		if (bPositions)
+		{
+			InitOrUpdateResource(&this->PositionVertexBuffer);
+		}
+		if (bMeshAttribs)
+		{
+			InitOrUpdateResource(&this->StaticMeshVertexBuffer);
+		}
+		if (bColors)
+		{
+			InitOrUpdateResource(&this->ColorVertexBuffer);
+		}
+
+		FLocalVertexFactory::FDataType Data;
+		this->PositionVertexBuffer.BindPositionVertexBuffer(&this->VertexFactory, Data);
+		this->StaticMeshVertexBuffer.BindTangentVertexBuffer(&this->VertexFactory, Data);
+		this->StaticMeshVertexBuffer.BindPackedTexCoordVertexBuffer(&this->VertexFactory, Data);
+		// currently no lightmaps support
+		//this->StaticMeshVertexBuffer.BindLightMapVertexBuffer(&this->VertexFactory, Data, LightMapIndex);
+		this->ColorVertexBuffer.BindColorVertexBuffer(&this->VertexFactory, Data);
+		this->VertexFactory.SetData(Data);
+
+		InitOrUpdateResource(&this->VertexFactory);
+	}
+
+
+	void TransferVertexUpdateToGPU(bool bPositions, bool bNormals, bool bTexCoords, bool bColors)
+	{
+		check(IsInRenderingThread());
+		if (TriangleCount == 0)
+		{
+			return;
+		}
+
+		if (bPositions)
+		{
+			auto& VertexBuffer = this->PositionVertexBuffer;
+			void* VertexBufferData = RHILockBuffer(VertexBuffer.VertexBufferRHI, 0, VertexBuffer.GetNumVertices() * VertexBuffer.GetStride(), RLM_WriteOnly);
+			FMemory::Memcpy(VertexBufferData, VertexBuffer.GetVertexData(), VertexBuffer.GetNumVertices() * VertexBuffer.GetStride());
+			RHIUnlockBuffer(VertexBuffer.VertexBufferRHI);
+		}
+		if (bNormals)
+		{
+			auto& VertexBuffer = this->StaticMeshVertexBuffer;
+			void* VertexBufferData = RHILockBuffer(VertexBuffer.TangentsVertexBuffer.VertexBufferRHI, 0, VertexBuffer.GetTangentSize(), RLM_WriteOnly);
+			FMemory::Memcpy(VertexBufferData, VertexBuffer.GetTangentData(), VertexBuffer.GetTangentSize());
+			RHIUnlockBuffer(VertexBuffer.TangentsVertexBuffer.VertexBufferRHI);
+		}
+		if (bColors)
+		{
+			auto& VertexBuffer = this->ColorVertexBuffer;
+			void* VertexBufferData = RHILockBuffer(VertexBuffer.VertexBufferRHI, 0, VertexBuffer.GetNumVertices() * VertexBuffer.GetStride(), RLM_WriteOnly);
+			FMemory::Memcpy(VertexBufferData, VertexBuffer.GetVertexData(), VertexBuffer.GetNumVertices() * VertexBuffer.GetStride());
+			RHIUnlockBuffer(VertexBuffer.VertexBufferRHI);
+		}
+		if (bTexCoords)
+		{
+			auto& VertexBuffer = this->StaticMeshVertexBuffer;
+			void* VertexBufferData = RHILockBuffer(VertexBuffer.TexCoordVertexBuffer.VertexBufferRHI, 0, VertexBuffer.GetTexCoordSize(), RLM_WriteOnly);
+			FMemory::Memcpy(VertexBufferData, VertexBuffer.GetTexCoordData(), VertexBuffer.GetTexCoordSize());
+			RHIUnlockBuffer(VertexBuffer.TexCoordVertexBuffer.VertexBufferRHI);
+		}
+	}
+
+
+
+
+
+
+	/**
+	 * Initializes a render resource, or update it if already initialized.
+	 * @wearning This function can only be called on the Render Thread
+	 */
+	void InitOrUpdateResource(FRenderResource* Resource)
+	{
+		check(IsInRenderingThread());
+
+		if (!Resource->IsInitialized())
+		{
+			Resource->InitResource();
+		}
+		else
+		{
+			Resource->UpdateRHI();
+		}
+	}
+
+
+
+
+protected:
+	friend class FBaseDynamicMeshSceneProxy;
+
+	/**
+	 * Enqueue a command on the Render Thread to destroy the passed in buffer set.
+	 * At this point the buffer set should be considered invalid.
+	 */
+	static void DestroyRenderBufferSet(FMeshRenderBufferSet* BufferSet)
+	{
+		if (BufferSet->TriangleCount == 0)
+		{
+			return;
+		}
+
+		ENQUEUE_RENDER_COMMAND(FMeshRenderBufferSetDestroy)(
+			[BufferSet](FRHICommandListImmediate& RHICmdList)
+		{
+			delete BufferSet;
+		});
+	}
+
+
+};
+
+
+
+
+/**
+ * FBaseDynamicMeshSceneProxy is an abstract base class for a Render Proxy
+ * for a UBaseDynamicMeshComponent, where the assumption is that mesh data
+ * will be stored in FMeshRenderBufferSet instances
+ */
+class FBaseDynamicMeshSceneProxy : public FPrimitiveSceneProxy
+{
+	using FIndex2i = UE::Geometry::FIndex2i;
+	using FIndex3i = UE::Geometry::FIndex3i;
+public:
+	UBaseDynamicMeshComponent* ParentBaseComponent;
+
+	/**
+	 * Constant color assigned to vertices if no other vertex color is specified
+	 */
+	FColor ConstantVertexColor = FColor::White;
+
+	/**
+	 * If true, vertex colors on the FDynamicMesh will be ignored
+	 */
+	bool bIgnoreVertexColors = false;
+
+	/**
+	 * If true, a per-triangle color is used to set vertex colors
+	 */
+	bool bUsePerTriangleColor = false;
+
+	/**
+	 * Per-triangle color function. Only called if bUsePerTriangleColor=true
+	 */
+	TFunction<FColor(const FDynamicMesh3*, int)> PerTriangleColorFunc = nullptr;
+
+
+	/**
+	 * If true, populate secondary buffers using SecondaryTriFilterFunc
+	 */
+	bool bUseSecondaryTriBuffers = false;
+
+	/**
+	 * Filter predicate for secondary triangle index buffer. Only called if bUseSecondaryTriBuffers=true
+	 */
+	TUniqueFunction<bool(const FDynamicMesh3*, int32)> SecondaryTriFilterFunc = nullptr;
+
+protected:
+	// Set of currently-allocated RenderBuffers. We own these pointers and must clean them up.
+	// Must guard access with AllocatedSetsLock!!
+	TSet<FMeshRenderBufferSet*> AllocatedBufferSets;
+
+	// use to control access to AllocatedBufferSets 
+	FCriticalSection AllocatedSetsLock;
+
+public:
+	FBaseDynamicMeshSceneProxy(UBaseDynamicMeshComponent* Component)
+		: FPrimitiveSceneProxy(Component),
+		ParentBaseComponent(Component)
+	{
+	}
+
+
+	virtual ~FBaseDynamicMeshSceneProxy()
+	{
+		// we are assuming in code below that this is always called from the rendering thread
+		check(IsInRenderingThread());
+
+		// destroy all existing renderbuffers
+		for (FMeshRenderBufferSet* BufferSet : AllocatedBufferSets)
+		{
+			FMeshRenderBufferSet::DestroyRenderBufferSet(BufferSet);
+		}
+	}
+
+
+	//
+	// FBaseDynamicMeshSceneProxy API - subclasses must implement these functions
+	//
+
+
+	/**
+	 * Return set of active renderbuffers. Must be implemented by subclass.
+	 * This is the set of render buffers that will be drawn by GetDynamicMeshElements
+	 */
+	virtual void GetActiveRenderBufferSets(TArray<FMeshRenderBufferSet*>& Buffers) const = 0;
+
+
+
+	//
+	// RenderBuffer management
+	//
+
+
+	/**
+	 * Allocates a set of render buffers. FPrimitiveSceneProxy will keep track of these
+	 * buffers and destroy them on destruction.
+	 */
+	virtual FMeshRenderBufferSet* AllocateNewRenderBufferSet()
+	{
+		// should we hang onto these and destroy them in constructor? leaving to subclass seems risky?
+		FMeshRenderBufferSet* RenderBufferSet = new FMeshRenderBufferSet(GetScene().GetFeatureLevel());
+
+		RenderBufferSet->Material = UMaterial::GetDefaultMaterial(MD_Surface);
+
+		AllocatedSetsLock.Lock();
+		AllocatedBufferSets.Add(RenderBufferSet);
+		AllocatedSetsLock.Unlock();
+
+		return RenderBufferSet;
+	}
+
+	/**
+	 * Explicitly release a set of RenderBuffers
+	 */
+	virtual void ReleaseRenderBufferSet(FMeshRenderBufferSet* BufferSet)
+	{
+		FScopeLock Lock(&AllocatedSetsLock);
+		if (ensure(AllocatedBufferSets.Contains(BufferSet)))
+		{
+			AllocatedBufferSets.Remove(BufferSet);
+			Lock.Unlock();
+
+			FMeshRenderBufferSet::DestroyRenderBufferSet(BufferSet);
+		}
+	}
+
+
+	/**
+	 * Initialize rendering buffers from given attribute overlays.
+	 * Creates three vertices per triangle, IE no shared vertices in buffers.
+	 */
+	template<typename TriangleEnumerable>
+	void InitializeBuffersFromOverlays(
+		FMeshRenderBufferSet* RenderBuffers,
+		const FDynamicMesh3* Mesh,
+		int NumTriangles, TriangleEnumerable Enumerable,
+		const FDynamicMeshUVOverlay* UVOverlay,
+		const FDynamicMeshNormalOverlay* NormalOverlay,
+		const FDynamicMeshColorOverlay* ColorOverlay,
+		TFunctionRef<void(int, int, int, const FVector3f&, FVector3f&, FVector3f&)> TangentsFunc,
+		bool bTrackTriangles = false)
+	{
+		TArray<const FDynamicMeshUVOverlay*> UVOverlays;
+		UVOverlays.Add(UVOverlay);
+		InitializeBuffersFromOverlays(RenderBuffers, Mesh, NumTriangles, Enumerable,
+			UVOverlays, NormalOverlay, ColorOverlay, TangentsFunc, bTrackTriangles);
+	}
+
+
+
+	/**
+	 * Initialize rendering buffers from given attribute overlays.
+	 * Creates three vertices per triangle, IE no shared vertices in buffers.
+	 */
+	template<typename TriangleEnumerable>
+	void InitializeBuffersFromOverlays(
+		FMeshRenderBufferSet* RenderBuffers,
+		const FDynamicMesh3* Mesh,
+		int NumTriangles, TriangleEnumerable Enumerable,
+		const TArray<const FDynamicMeshUVOverlay*>& UVOverlays,
+		const FDynamicMeshNormalOverlay* NormalOverlay,
+		const FDynamicMeshColorOverlay* ColorOverlay,
+		TFunctionRef<void(int, int, int, const FVector3f&, FVector3f&, FVector3f&)> TangentsFunc,
+		bool bTrackTriangles = false)
+	{
+		RenderBuffers->TriangleCount = NumTriangles;
+		if (NumTriangles == 0)
+		{
+			return;
+		}
+
+		bool bHaveColors = (ColorOverlay != nullptr) && (bIgnoreVertexColors == false);
+
+		int NumVertices = NumTriangles * 3;
+		int NumUVOverlays = UVOverlays.Num();
+		int NumTexCoords = FMath::Max(1, NumUVOverlays);		// must have at least one tex coord
+		TArray<FIndex3i, TFixedAllocator<MAX_STATIC_TEXCOORDS>> UVTriangles;
+		UVTriangles.SetNum(NumTexCoords);
+
+		{
+			RenderBuffers->PositionVertexBuffer.Init(NumVertices);
+			RenderBuffers->StaticMeshVertexBuffer.Init(NumVertices, NumTexCoords );
+			RenderBuffers->ColorVertexBuffer.Init(NumVertices);
+			RenderBuffers->IndexBuffer.Indices.AddUninitialized(NumTriangles * 3);
+		}
+
+		// build triangle list if requested, or if we are using secondary buffers in which case we need it to filter later
+		bool bBuildTriangleList = bTrackTriangles || bUseSecondaryTriBuffers;
+		if (bBuildTriangleList)
+		{
+			RenderBuffers->Triangles = TArray<int32>();
+		}
+
+		int TriIdx = 0, VertIdx = 0;
+		FVector3f TangentX, TangentY;
+		for (int TriangleID : Enumerable)
+		{
+			FIndex3i Tri = Mesh->GetTriangle(TriangleID);
+			for (int32 k = 0; k < NumTexCoords; ++k)
+			{
+				UVTriangles[k] = (k < NumUVOverlays && UVOverlays[k] != nullptr) ? UVOverlays[k]->GetTriangle(TriangleID) : FIndex3i::Invalid();
+			}
+			FIndex3i TriNormal = (NormalOverlay != nullptr) ? NormalOverlay->GetTriangle(TriangleID) : FIndex3i::Zero();
+			FIndex3i TriColor = (ColorOverlay != nullptr) ? ColorOverlay->GetTriangle(TriangleID) : FIndex3i::Zero();
+
+			FColor UniformTriColor = ConstantVertexColor;
+			if (bUsePerTriangleColor && PerTriangleColorFunc != nullptr)
+			{
+				UniformTriColor = PerTriangleColorFunc(Mesh, TriangleID);
+				bHaveColors = false;
+			}
+
+			for (int j = 0; j < 3; ++j)
+			{
+				RenderBuffers->PositionVertexBuffer.VertexPosition(VertIdx) = (FVector)Mesh->GetVertex(Tri[j]);
+
+				FVector3f Normal = (NormalOverlay != nullptr && TriNormal[j] != FDynamicMesh3::InvalidID) ?
+					NormalOverlay->GetElement(TriNormal[j]) : Mesh->GetVertexNormal(Tri[j]);
+
+				// get tangents
+				TangentsFunc(Tri[j], TriangleID, j, Normal, TangentX, TangentY);
+
+				RenderBuffers->StaticMeshVertexBuffer.SetVertexTangents(VertIdx, (FVector)TangentX, (FVector)TangentY, (FVector)Normal);
+
+				for (int32 k = 0; k < NumTexCoords; ++k)
+				{
+					FVector2f UV = (UVTriangles[k][j] != FDynamicMesh3::InvalidID) ?
+						UVOverlays[k]->GetElement(UVTriangles[k][j]) : FVector2f::Zero();
+					RenderBuffers->StaticMeshVertexBuffer.SetVertexUV(VertIdx, k, (FVector2D)UV);
+				}
+
+				FColor VertexFColor = (bHaveColors && TriColor[j] != FDynamicMesh3::InvalidID) ?
+					FColor(ColorOverlay->GetElement(TriColor[j])) : UniformTriColor;
+
+				RenderBuffers->ColorVertexBuffer.VertexColor(VertIdx) = VertexFColor;
+
+				RenderBuffers->IndexBuffer.Indices[TriIdx++] = VertIdx;		// currently TriIdx == VertIdx so we don't really need both...
+				VertIdx++;
+			}
+
+			if (bBuildTriangleList)
+			{
+				RenderBuffers->Triangles->Add(TriangleID);
+			}
+		}
+
+		// split triangles into secondary buffer (at bit redudant since we just built IndexBuffer, but we may optionally duplicate triangles in the future
+		if (bUseSecondaryTriBuffers)
+		{
+			RenderBuffers->bEnableSecondaryIndexBuffer = true;
+			UpdateSecondaryTriangleBuffer(RenderBuffers, Mesh, false);
+		}
+	}
+
+
+
+
+	/**
+	 * Filter the triangles in a FMeshRenderBufferSet into the SecondaryIndexBuffer.
+	 * Requires that RenderBuffers->Triangles has been initialized.
+	 * @param bDuplicate if set, then primary IndexBuffer is unmodified and SecondaryIndexBuffer contains duplicates. Otherwise triangles are sorted via predicate into either primary or secondary.
+	 */
+	void UpdateSecondaryTriangleBuffer(
+		FMeshRenderBufferSet* RenderBuffers,
+		const FDynamicMesh3* Mesh,
+		bool bDuplicate)
+	{
+		if (ensure(bUseSecondaryTriBuffers == true && RenderBuffers->Triangles.IsSet()) == false)
+		{
+			return;
+		}
+
+		const TArray<int32>& TriangleIDs = RenderBuffers->Triangles.GetValue();
+		int NumTris = TriangleIDs.Num();
+		TArray<uint32>& Indices = RenderBuffers->IndexBuffer.Indices;
+		TArray<uint32>& SecondaryIndices = RenderBuffers->SecondaryIndexBuffer.Indices;
+
+		RenderBuffers->SecondaryIndexBuffer.Indices.Reset();
+		if (bDuplicate == false)
+		{
+			RenderBuffers->IndexBuffer.Indices.Reset();
+		}
+		for ( int k = 0; k < NumTris; ++k)
+		{
+			int TriangleID = TriangleIDs[k];
+			bool bInclude = SecondaryTriFilterFunc(Mesh, TriangleID);
+			if (bInclude)
+			{
+				SecondaryIndices.Add(3*k);
+				SecondaryIndices.Add(3*k + 1);
+				SecondaryIndices.Add(3*k + 2);
+			} 
+			else if (bDuplicate == false)
+			{
+				Indices.Add(3*k);
+				Indices.Add(3*k + 1);
+				Indices.Add(3*k + 2);
+			}
+		}
+	}
+
+
+	/**
+	 * RecomputeRenderBufferTriangleIndexSets re-sorts the existing set of triangles in a FMeshRenderBufferSet
+	 * into primary and secondary index buffers. Note that UploadIndexBufferUpdate() must be called
+	 * after this function!
+	 */
+	void RecomputeRenderBufferTriangleIndexSets(
+		FMeshRenderBufferSet* RenderBuffers, 
+		const FDynamicMesh3* Mesh)
+	{
+		if (RenderBuffers->TriangleCount == 0)
+		{
+			return;
+		}
+		if (ensure(RenderBuffers->Triangles.IsSet() && RenderBuffers->Triangles->Num() > 0) == false)
+		{
+			return;
+		}
+
+		//bool bDuplicate = false;		// flag for future use, in case we want to draw all triangles in primary and duplicates in secondary...
+		RenderBuffers->IndexBuffer.Indices.Reset();
+		RenderBuffers->SecondaryIndexBuffer.Indices.Reset();
+		
+		TArray<uint32>& Indices = RenderBuffers->IndexBuffer.Indices;
+		TArray<uint32>& SecondaryIndices = RenderBuffers->SecondaryIndexBuffer.Indices;
+		const TArray<int32>& TriangleIDs = RenderBuffers->Triangles.GetValue();
+
+		int NumTris = TriangleIDs.Num();
+		for (int k = 0; k < NumTris; ++k)
+		{
+			int TriangleID = TriangleIDs[k];
+			bool bInclude = SecondaryTriFilterFunc(Mesh, TriangleID);
+			if (bInclude)
+			{
+				SecondaryIndices.Add(3 * k);
+				SecondaryIndices.Add(3 * k + 1);
+				SecondaryIndices.Add(3 * k + 2);
+			}
+			else // if (bDuplicate == false)
+			{
+				Indices.Add(3 * k);
+				Indices.Add(3 * k + 1);
+				Indices.Add(3 * k + 2);
+			}
+		}
+	}
+
+
+
+	/**
+	 * Update vertex positions/normals/colors of an existing set of render buffers.
+	 * Assumes that buffers were created with unshared vertices, ie three vertices per triangle, eg by InitializeBuffersFromOverlays()
+	 */
+	template<typename TriangleEnumerable>
+	void UpdateVertexBuffersFromOverlays(
+		FMeshRenderBufferSet* RenderBuffers,
+		const FDynamicMesh3* Mesh,
+		int NumTriangles, TriangleEnumerable Enumerable,
+		const FDynamicMeshNormalOverlay* NormalOverlay,
+		const FDynamicMeshColorOverlay* ColorOverlay,
+		TFunctionRef<void(int, int, int, const FVector3f&, FVector3f&, FVector3f&)> TangentsFunc,
+		bool bUpdatePositions = true,
+		bool bUpdateNormals = false,
+		bool bUpdateColors = false)
+	{
+		if (RenderBuffers->TriangleCount == 0)
+		{
+			return;
+		}
+
+		bool bHaveColors = (ColorOverlay != nullptr) && (bIgnoreVertexColors == false);
+
+		int NumVertices = NumTriangles * 3;
+		if ( (bUpdatePositions && ensure(RenderBuffers->PositionVertexBuffer.GetNumVertices() == NumVertices) == false )
+			|| (bUpdateNormals && ensure(RenderBuffers->StaticMeshVertexBuffer.GetNumVertices() == NumVertices) == false )
+			|| (bUpdateColors && ensure(RenderBuffers->ColorVertexBuffer.GetNumVertices() == NumVertices) == false ) )
+		{
+			return;
+		}
+
+		int VertIdx = 0;
+		FVector3f TangentX, TangentY;
+		for (int TriangleID : Enumerable)
+		{
+			FIndex3i Tri = Mesh->GetTriangle(TriangleID);
+
+			FIndex3i TriNormal = (bUpdateNormals && NormalOverlay != nullptr) ? NormalOverlay->GetTriangle(TriangleID) : FIndex3i::Zero();
+			FIndex3i TriColor = (bUpdateColors && ColorOverlay != nullptr) ? ColorOverlay->GetTriangle(TriangleID) : FIndex3i::Zero();
+
+			FColor UniformTriColor = ConstantVertexColor;
+			if (bUpdateColors && bUsePerTriangleColor && PerTriangleColorFunc != nullptr)
+			{
+				UniformTriColor = PerTriangleColorFunc(Mesh, TriangleID);
+				bHaveColors = false;
+			}
+
+			for (int j = 0; j < 3; ++j)
+			{
+				if (bUpdatePositions)
+				{
+					RenderBuffers->PositionVertexBuffer.VertexPosition(VertIdx) = (FVector)Mesh->GetVertex(Tri[j]);
+				}
+
+				if (bUpdateNormals)
+				{
+					// get normal and tangent
+					FVector3f Normal = (NormalOverlay != nullptr && TriNormal[j] != FDynamicMesh3::InvalidID) ?
+						NormalOverlay->GetElement(TriNormal[j]) : Mesh->GetVertexNormal(Tri[j]);
+					TangentsFunc(Tri[j], TriangleID, j, Normal, TangentX, TangentY);
+
+					RenderBuffers->StaticMeshVertexBuffer.SetVertexTangents(VertIdx, (FVector)TangentX, (FVector)TangentY, (FVector)Normal);
+				}
+
+				if (bUpdateColors)
+				{
+					FColor VertexFColor = (bHaveColors  && TriColor[j] != FDynamicMesh3::InvalidID) ?
+						FColor(ColorOverlay->GetElement(TriColor[j])) : UniformTriColor;
+					RenderBuffers->ColorVertexBuffer.VertexColor(VertIdx) = VertexFColor;
+				}
+
+				VertIdx++;
+			}
+		}
+	}
+
+
+
+	/**
+	 * Update vertex uvs of an existing set of render buffers.
+	 * Assumes that buffers were created with unshared vertices, ie three vertices per triangle, eg by InitializeBuffersFromOverlays()
+	 */
+	template<typename TriangleEnumerable>
+	void UpdateVertexUVBufferFromOverlays(
+		FMeshRenderBufferSet* RenderBuffers,
+		const FDynamicMesh3* Mesh,
+		int32 NumTriangles, TriangleEnumerable Enumerable,
+		const FDynamicMeshUVOverlay* UVOverlay, int32 UVIndex)
+	{
+		if (RenderBuffers->TriangleCount == 0)
+		{
+			return;
+		}
+		int NumVertices = NumTriangles * 3;
+		if ( ensure(RenderBuffers->StaticMeshVertexBuffer.GetNumVertices() == NumVertices) == false)
+		{
+			return;
+		}
+
+		int VertIdx = 0;
+		for (int TriangleID : Enumerable)
+		{
+			FIndex3i UVTri = UVOverlay->GetTriangle(TriangleID);
+			for (int j = 0; j < 3; ++j)
+			{
+				FVector2f UV = UVOverlay->GetElement(UVTri[j]);
+				RenderBuffers->StaticMeshVertexBuffer.SetVertexUV(VertIdx, UVIndex, (FVector2D)UV);
+				VertIdx++;
+			}
+		}
+	}
+
+
+
+
+	/**
+	 * @return number of active materials
+	 */
+	virtual int32 GetNumMaterials() const
+	{
+		return ParentBaseComponent->GetNumMaterials();
+	}
+
+	/**
+	 * Safe GetMaterial function that will never return nullptr
+	 */
+	virtual UMaterialInterface* GetMaterial(int32 k) const
+	{
+		UMaterialInterface* Material = ParentBaseComponent->GetMaterial(k);
+		return (Material != nullptr) ? Material : UMaterial::GetDefaultMaterial(MD_Surface);
+	}
+
+
+	/**
+	 * This needs to be called if the set of active materials changes, otherwise
+	 * the check in FPrimitiveSceneProxy::VerifyUsedMaterial() will fail if an override
+	 * material is set, if materials change, etc, etc
+	 */
+	virtual void UpdatedReferencedMaterials()
+	{
+#if WITH_EDITOR
+		TArray<UMaterialInterface*> Materials;
+		ParentBaseComponent->GetUsedMaterials(Materials, true);
+
+		// Temporarily disable material verification while the enqueued render command is in flight.
+		// The original value for bVerifyUsedMaterials gets restored when the command is executed.
+		// If we do not do this, material verification might spuriously fail in cases where the render command for changing
+		// the verfifcation material is still in flight but the render thread is already trying to render the mesh.
+		const uint8 bRestoreVerifyUsedMaterials = bVerifyUsedMaterials;
+		bVerifyUsedMaterials = false;
+
+		ENQUEUE_RENDER_COMMAND(FMeshRenderBufferSetDestroy)(
+			[this, Materials, bRestoreVerifyUsedMaterials](FRHICommandListImmediate& RHICmdList)
+		{
+			this->SetUsedMaterialForVerification(Materials);
+			this->bVerifyUsedMaterials = bRestoreVerifyUsedMaterials;
+		});
+#endif
+	}
+
+
+	//
+	// FBaseDynamicMeshSceneProxy implementation
+	//
+
+
+	/**
+	 * Render set of active RenderBuffers returned by GetActiveRenderBufferSets
+	 */
+	virtual void GetDynamicMeshElements(
+		const TArray<const FSceneView*>& Views, 
+		const FSceneViewFamily& ViewFamily, 
+		uint32 VisibilityMap, 
+		FMeshElementCollector& Collector) const override
+	{
+		QUICK_SCOPE_CYCLE_COUNTER(STAT_OctreeDynamicMeshSceneProxy_GetDynamicMeshElements);
+
+		const bool bWireframe = (AllowDebugViewmodes() && ViewFamily.EngineShowFlags.Wireframe)
+			|| ParentBaseComponent->GetEnableWireframeRenderPass();
+
+		// set up wireframe material. Probably bad to reference GEngine here...also this material is very bad?
+		FMaterialRenderProxy* WireframeMaterialProxy = nullptr;
+		if (bWireframe)
+		{
+			FColoredMaterialRenderProxy* WireframeMaterialInstance = new FColoredMaterialRenderProxy(
+				GEngine->WireframeMaterial ? GEngine->WireframeMaterial->GetRenderProxy() : nullptr,
+				FLinearColor(0, 0.5f, 1.f)
+			);
+			Collector.RegisterOneFrameMaterialProxy(WireframeMaterialInstance);
+			WireframeMaterialProxy = WireframeMaterialInstance;
+		}
+
+		ESceneDepthPriorityGroup DepthPriority = SDPG_World;
+
+		TArray<FMeshRenderBufferSet*> Buffers;
+		GetActiveRenderBufferSets(Buffers);
+
+		FMaterialRenderProxy* SecondaryMaterialProxy =
+			ParentBaseComponent->HasSecondaryRenderMaterial() ? ParentBaseComponent->GetSecondaryRenderMaterial()->GetRenderProxy() : nullptr;
+		bool bDrawSecondaryBuffers = ParentBaseComponent->GetSecondaryBuffersVisibility();
+
+		for (int32 ViewIndex = 0; ViewIndex < Views.Num(); ViewIndex++)
+		{
+			if (VisibilityMap & (1 << ViewIndex))
+			{
+				const FSceneView* View = Views[ViewIndex];
+
+				bool bHasPrecomputedVolumetricLightmap;
+				FMatrix PreviousLocalToWorld;
+				int32 SingleCaptureIndex;
+				bool bOutputVelocity;
+				GetScene().GetPrimitiveUniformShaderParameters_RenderThread(GetPrimitiveSceneInfo(), bHasPrecomputedVolumetricLightmap, PreviousLocalToWorld, SingleCaptureIndex, bOutputVelocity);
+
+				// Draw the mesh.
+				for (FMeshRenderBufferSet* BufferSet : Buffers)
+				{
+					UMaterialInterface* UseMaterial = BufferSet->Material;
+					if (ParentBaseComponent->HasOverrideRenderMaterial(0))
+					{
+						UseMaterial = ParentBaseComponent->GetOverrideRenderMaterial(0);
+					}
+					FMaterialRenderProxy* MaterialProxy = UseMaterial->GetRenderProxy();
+
+					if (BufferSet->TriangleCount == 0)
+					{
+						continue;
+					}
+
+					// lock buffers so that they aren't modified while we are submitting them
+					FScopeLock BuffersLock(&BufferSet->BuffersLock);
+
+					// do we need separate one of these for each MeshRenderBufferSet?
+					FDynamicPrimitiveUniformBuffer& DynamicPrimitiveUniformBuffer = Collector.AllocateOneFrameResource<FDynamicPrimitiveUniformBuffer>();
+					DynamicPrimitiveUniformBuffer.Set(GetLocalToWorld(), PreviousLocalToWorld, GetBounds(), GetLocalBounds(), true, bHasPrecomputedVolumetricLightmap, DrawsVelocity(), bOutputVelocity);
+
+					if (BufferSet->IndexBuffer.Indices.Num() > 0)
+					{
+						// Unlike most meshes, which just use the wireframe material in wireframe mode, we draw the wireframe on top of the normal material if needed,
+						// as this is easier to interpret. However, we do not do this in ortho viewports, where it frequently causes the our edit gizmo to be hidden 
+						// beneath the material. So, only draw the base material if we are in perspective mode, or we're in ortho but not in wireframe.
+						if (View->IsPerspectiveProjection() || !(AllowDebugViewmodes() && ViewFamily.EngineShowFlags.Wireframe))
+						{
+							DrawBatch(Collector, *BufferSet, BufferSet->IndexBuffer, MaterialProxy, false, DepthPriority, ViewIndex, DynamicPrimitiveUniformBuffer);
+						}
+						if (bWireframe)
+						{
+							DrawBatch(Collector, *BufferSet, BufferSet->IndexBuffer, WireframeMaterialProxy, true, DepthPriority, ViewIndex, DynamicPrimitiveUniformBuffer);
+						}
+					}
+
+					// draw secondary buffer if we have it, falling back to base material if we don't have the Secondary material
+					FMaterialRenderProxy* UseSecondaryMaterialProxy = (SecondaryMaterialProxy != nullptr) ? SecondaryMaterialProxy : MaterialProxy;
+					if (bDrawSecondaryBuffers && BufferSet->SecondaryIndexBuffer.Indices.Num() > 0 && UseSecondaryMaterialProxy != nullptr)
+					{
+						DrawBatch(Collector, *BufferSet, BufferSet->SecondaryIndexBuffer, UseSecondaryMaterialProxy, false, DepthPriority, ViewIndex, DynamicPrimitiveUniformBuffer);
+						if (bWireframe)
+						{
+							DrawBatch(Collector, *BufferSet, BufferSet->SecondaryIndexBuffer, UseSecondaryMaterialProxy, true, DepthPriority, ViewIndex, DynamicPrimitiveUniformBuffer);
+						}
+					}
+				}
+			}
+		}
+	}
+
+
+
+	/**
+	 * Draw a single-frame FMeshBatch for a FMeshRenderBufferSet
+	 */
+	virtual void DrawBatch(FMeshElementCollector& Collector,
+		const FMeshRenderBufferSet& RenderBuffers,
+		const FDynamicMeshIndexBuffer32& IndexBuffer,
+		FMaterialRenderProxy* UseMaterial,
+		bool bWireframe,
+		ESceneDepthPriorityGroup DepthPriority,
+		int ViewIndex,
+		FDynamicPrimitiveUniformBuffer& DynamicPrimitiveUniformBuffer) const
+	{
+		FMeshBatch& Mesh = Collector.AllocateMesh();
+		FMeshBatchElement& BatchElement = Mesh.Elements[0];
+		BatchElement.IndexBuffer = &IndexBuffer;
+		Mesh.bWireframe = bWireframe;
+		Mesh.VertexFactory = &RenderBuffers.VertexFactory;
+		Mesh.MaterialRenderProxy = UseMaterial;
+
+		BatchElement.PrimitiveUniformBufferResource = &DynamicPrimitiveUniformBuffer.UniformBuffer;
+
+		BatchElement.FirstIndex = 0;
+		BatchElement.NumPrimitives = IndexBuffer.Indices.Num() / 3;
+		BatchElement.MinVertexIndex = 0;
+		BatchElement.MaxVertexIndex = RenderBuffers.PositionVertexBuffer.GetNumVertices() - 1;
+		Mesh.ReverseCulling = IsLocalToWorldDeterminantNegative();
+		Mesh.Type = PT_TriangleList;
+		Mesh.DepthPriorityGroup = DepthPriority;
+		Mesh.bCanApplyViewModeOverrides = false;
+		Collector.AddMesh(ViewIndex, Mesh);
+	}
+
+
+
+
+
+
+
+};
\ No newline at end of file
diff --git a/Engine/Source/Runtime/Experimental/GeometryFramework/Public/Components/DynamicMeshComponent.h b/Engine/Source/Runtime/Experimental/GeometryFramework/Public/Components/DynamicMeshComponent.h
new file mode 100644
--- /dev/null
+++ b/Engine/Source/Runtime/Experimental/GeometryFramework/Public/Components/DynamicMeshComponent.h
@@ -0,0 +1,522 @@
+// Copyright Epic Games, Inc. All Rights Reserved.
+
+#pragma once
+
+#include "CoreMinimal.h"
+#include "BaseDynamicMeshComponent.h"
+#include "MeshConversionOptions.h"
+#include "Components/MeshRenderDecomposition.h"
+#include "DynamicMesh/MeshTangents.h"
+#include "TransformTypes.h"
+#include "Async/Future.h"
+#include "UDynamicMesh.h"
+#include "PhysicsEngine/BodySetup.h"
+
+#include "DynamicMeshComponent.generated.h"
+
+// predecl
+struct FMeshDescription;
+
+/** internal FPrimitiveSceneProxy defined in DynamicMeshSceneProxy.h */
+class FDynamicMeshSceneProxy;
+class FBaseDynamicMeshSceneProxy;
+
+/**
+ * Interface for a render mesh processor. Use this to process the Mesh stored in UDynamicMeshComponent before
+ * sending it off for rendering.
+ * NOTE: This is called whenever the Mesh is updated and before rendering, so performance matters.
+ */
+class GEOMETRYFRAMEWORK_API IRenderMeshPostProcessor
+{
+public:
+	virtual ~IRenderMeshPostProcessor() = default;
+
+	virtual void ProcessMesh(const FDynamicMesh3& Mesh, FDynamicMesh3& OutRenderMesh) = 0;
+};
+
+
+/** Render data update hint */
+UENUM()
+enum class EDynamicMeshComponentRenderUpdateMode
+{
+	/** Do not update render data */
+	NoUpdate = 0,
+	/** Invalidate overlay of internal component, rebuilding all render data */
+	FullUpdate = 1,
+	/** Attempt to do partial update of render data if possible */
+	FastUpdate = 2
+};
+
+
+/** 
+ * UDynamicMeshComponent is a mesh component similar to UProceduralMeshComponent,
+ * except it bases the renderable geometry off an internal UDynamicMesh instance (which
+ * encapsulates a FDynamicMesh3). 
+ * 
+ * There is extensive support for partial updates to render buffers, customizing colors,
+ * internally decomposing the mesh into separate chunks for more efficient render updates,
+ * and support for attaching a 'Postprocessor' to generate a render mesh on-the-fly
+ * See comment sections below for details.
+ * 
+ */
+UCLASS(hidecategories = (LOD), meta = (BlueprintSpawnableComponent), ClassGroup = Rendering)
+class GEOMETRYFRAMEWORK_API UDynamicMeshComponent : public UBaseDynamicMeshComponent, public IInterface_CollisionDataProvider
+{
+	GENERATED_UCLASS_BODY()
+
+
+
+	//===============================================================================================================
+	// Mesh Access. Usage via GetDynamicMesh() or SetMesh()/ProcessMesh()/EditMesh() is preferred, the GetMesh() 
+	// pointer access exist largely to support existing code from before UDynamicMesh was added.
+public:
+	/**
+	 * @return pointer to internal mesh
+	 * @warning avoid usage of this function, access via GetDynamicMesh() instead
+	 */
+	virtual FDynamicMesh3* GetMesh() override { return MeshObject->GetMeshPtr(); }
+
+	/**
+	 * @return pointer to internal mesh
+	 * @warning avoid usage of this function, access via GetDynamicMesh() instead
+	 */
+	virtual const FDynamicMesh3* GetMesh() const override { return MeshObject->GetMeshPtr(); }
+
+	/**
+	 * @return the child UDynamicMesh
+	 */
+	//UFUNCTION(BlueprintCallable, Category = "Dynamic Mesh Component")
+	virtual UDynamicMesh* GetDynamicMesh() override { return MeshObject; }
+
+	/**
+	 * Set the child UDynamicMesh. This can be used to 'share' a UDynamicMesh between Component instances.
+	 * @warning Currently this is somewhat risky, it is on the caller/clients to make sure that the actual mesh is not being simultaneously modified on multiple threads
+	 */
+	UFUNCTION(BlueprintCallable, Category = "Dynamic Mesh Component")
+	void SetDynamicMesh(UDynamicMesh* NewMesh);
+
+	/**
+	 * initialize the internal mesh from a DynamicMesh
+	 */
+	virtual void SetMesh(UE::Geometry::FDynamicMesh3&& MoveMesh) override;
+
+	/**
+	 * Allow external code to read the internal mesh.
+	 */
+	virtual void ProcessMesh(TFunctionRef<void(const UE::Geometry::FDynamicMesh3&)> ProcessFunc) const;
+
+	/**
+	 * Allow external code to to edit the internal mesh.
+	 */
+	virtual void EditMesh(TFunctionRef<void(UE::Geometry::FDynamicMesh3&)> EditFunc,
+						  EDynamicMeshComponentRenderUpdateMode UpdateMode = EDynamicMeshComponentRenderUpdateMode::FullUpdate);
+
+	/**
+	 * Apply transform to internal mesh. In some cases this can be more efficient than a general edit.
+	 * @param bInvert if true, inverse tranform is applied instead of forward transform
+	 */
+	virtual void ApplyTransform(const UE::Geometry::FTransform3d& Transform, bool bInvert) override;
+
+
+
+protected:
+	/**
+	 * Internal FDynamicMesh is stored inside a UDynamicMesh container, which allows it to be
+	 * used from BP, shared with other UObjects, and so on
+	 */
+	UPROPERTY(Instanced)
+	TObjectPtr<UDynamicMesh> MeshObject;
+
+
+
+
+	//===============================================================================================================
+	// RenderBuffer Update API. These functions can be used by external code (and internally in some places)
+	// to tell the Component that the Mesh data has been modified in some way, and that the RenderBuffers in the RenderProxy
+	// need to be updated (or rebuilt entirely). On large meshes a full rebuild is expensive, so there are quite a few
+	// variants that can be used to minimize the amount of data updated in different situations.
+	//
+public:
+	/**
+	 * Call this if you update the mesh via GetMesh(). This will destroy the existing RenderProxy and create a new one.
+	 * @todo should provide a function that calls a lambda to modify the mesh, and only return const mesh pointer
+	 */
+	virtual void NotifyMeshUpdated() override;
+
+	/**
+	 * Call this instead of NotifyMeshUpdated() if you have only updated the vertex colors (or triangle color function).
+	 * This function will update the existing RenderProxy buffers if possible
+	 */
+	void FastNotifyColorsUpdated();
+
+	/**
+	 * Call this instead of NotifyMeshUpdated() if you have only updated the vertex positions (and possibly some attributes).
+	 * This function will update the existing RenderProxy buffers if possible
+	 */
+	void FastNotifyPositionsUpdated(bool bNormals = false, bool bColors = false, bool bUVs = false);
+
+	/**
+	 * Call this instead of NotifyMeshUpdated() if you have only updated the vertex attributes (but not positions).
+	 * This function will update the existing RenderProxy buffers if possible, rather than create new ones.
+	 */
+	void FastNotifyVertexAttributesUpdated(bool bNormals, bool bColors, bool bUVs);
+
+	/**
+	 * Call this instead of NotifyMeshUpdated() if you have only updated the vertex positions/attributes
+	 * This function will update the existing RenderProxy buffers if possible, rather than create new ones.
+	 */
+	void FastNotifyVertexAttributesUpdated(EMeshRenderAttributeFlags UpdatedAttributes);
+
+	/**
+	 * Call this instead of NotifyMeshUpdated() if you have only updated the vertex uvs.
+	 * This function will update the existing RenderProxy buffers if possible
+	 */
+	void FastNotifyUVsUpdated();
+
+	/**
+	 * Call this instead of NotifyMeshUpdated() if you have only updated secondary triangle sorting.
+	 * This function will update the existing buffers if possible, without rebuilding entire RenderProxy.
+	 */
+	void FastNotifySecondaryTrianglesChanged();
+
+	/**
+	 * This function updates vertex positions/attributes of existing SceneProxy render buffers if possible, for the given triangles.
+	 * If a FMeshRenderDecomposition has not been explicitly set, call is forwarded to FastNotifyVertexAttributesUpdated()
+	 */
+	void FastNotifyTriangleVerticesUpdated(const TArray<int32>& Triangles, EMeshRenderAttributeFlags UpdatedAttributes);
+
+	/**
+	 * This function updates vertex positions/attributes of existing SceneProxy render buffers if possible, for the given triangles.
+	 * If a FMeshRenderDecomposition has not been explicitly set, call is forwarded to FastNotifyVertexAttributesUpdated()
+	 */
+	void FastNotifyTriangleVerticesUpdated(const TSet<int32>& Triangles, EMeshRenderAttributeFlags UpdatedAttributes);
+
+	/**
+	 * If a Decomposition is set on this Component, and everything is currently valid (proxy/etc), precompute the set of
+	 * buffers that will be modified, as well as the bounds of the modified region. These are both computed in parallel.
+	 * Use FastNotifyTriangleVerticesUpdated_ApplyPrecompute() with the returned future to apply this precomputation.
+	 * @return a future that will (eventually) return true if the precompute is OK, and (immediately) false if it is not
+	 */
+	TFuture<bool> FastNotifyTriangleVerticesUpdated_TryPrecompute(const TArray<int32>& Triangles, TArray<int32>& UpdateSetsOut, UE::Geometry::FAxisAlignedBox3d& BoundsOut);
+
+	/**
+	 * This function updates vertex positions/attributes of existing SceneProxy render buffers if possible, for the given triangles.
+	 * The assumption is that FastNotifyTriangleVerticesUpdated_TryPrecompute() was used to get the Precompute future, this function
+	 * will Wait() until it is done and then use the UpdateSets and UpdateSetBounds that were computed (must be the same variables
+	 * passed to FastNotifyTriangleVerticesUpdated_TryPrecompute). 
+	 * If the Precompute future returns false, then we forward the call to FastNotifyTriangleVerticesUpdated(), which will do more work.
+	 */
+	void FastNotifyTriangleVerticesUpdated_ApplyPrecompute(const TArray<int32>& Triangles, EMeshRenderAttributeFlags UpdatedAttributes,
+		TFuture<bool>& Precompute, const TArray<int32>& UpdateSets, const UE::Geometry::FAxisAlignedBox3d& UpdateSetBounds);
+
+
+
+
+
+	//===============================================================================================================
+	// Change Support. These changes are primarily used for Undo/Redo, however there is no strict assumption
+	// about this internally, objects of these change types could also be used to perform more structured editing.
+	// (Note that these functions simply forward the change events to the child UDynamicMesh, which will
+	// 	post a mesh-change event that 
+	//
+public:
+	/**
+	 * Apply a vertex deformation change to the mesh
+	 */
+	virtual void ApplyChange(const FMeshVertexChange* Change, bool bRevert) override;
+
+	/**
+	 * Apply a general mesh change to the mesh
+	 */
+	virtual void ApplyChange(const FMeshChange* Change, bool bRevert) override;
+
+	/**
+	* Apply a mesh replacement change to mesh
+	*/
+	virtual void ApplyChange(const FMeshReplacementChange* Change, bool bRevert) override;
+
+	/**
+	 * This delegate fires when the mesh has been changed
+	 */
+	FSimpleMulticastDelegate OnMeshChanged;
+
+	/**
+	 * This delegate fires when the mesh vertices have been changed via an FMeshVertexChange
+	 */
+	DECLARE_MULTICAST_DELEGATE_ThreeParams(FMeshVerticesModified, UDynamicMeshComponent*, const FMeshVertexChange*, bool);
+	FMeshVerticesModified OnMeshVerticesChanged;
+
+	/**
+	 * When a FMeshChange or FMeshVertexChange is applied, by default we currently fully invalidate the render proxy. However in certain
+	 * realtime situations (eg like Sculpting tools) it can be critical to undo/redo performance to do more optimized render data updates
+	 * (eg using one of the FastXYZ functions above). To allow for that, the full proxy invalidation on change can be (temporarily!) disabled 
+	 * using this function. 
+	 */
+	void SetInvalidateProxyOnChangeEnabled(bool bEnabled);
+
+	/** @return true if InvalidateProxyOnChange is enabled (default) */
+	bool GetInvalidateProxyOnChangeEnabled() const { return bInvalidateProxyOnChange; }
+
+protected:
+	/** If false, we don't completely invalidate the RenderProxy when ApplyChange() is called (assumption is it will be handled elsewhere) */
+	bool bInvalidateProxyOnChange = true;
+
+	/** Handle for OnMeshObjectChanged which is registered with MeshObject::OnMeshChanged delegate */
+	FDelegateHandle MeshObjectChangedHandle;
+
+	/** Called whenever internal MeshObject is modified, fires OnMeshChanged and OnMeshVerticesChanged above */
+	void OnMeshObjectChanged(UDynamicMesh* ChangedMeshObject, FDynamicMeshChangeInfo ChangeInfo);
+
+
+
+
+	//===============================================================================================================
+	// Support for specifying per-triangle colors as vertex colors. This allows external code to dynamically override 
+	// the vertex colors on the rendered mesh. The lambda that is passed is held for the lifetime of the Component and
+	// must remain valid. A Material that uses the vertex colors must be applied, otherwise setting this override will
+	// have no visible effect. If the colors change externally, FastNotifyColorsUpdated() can be used to do the 
+	// minimal vertex buffer updates necessary in the RenderProxy
+	//
+public:
+	/** Clear an active triangle color function if one exists, and update the mesh */
+	virtual void SetTriangleColorFunction(TUniqueFunction<FColor(const FDynamicMesh3*, int)> TriangleColorFuncIn,
+										  EDynamicMeshComponentRenderUpdateMode UpdateMode = EDynamicMeshComponentRenderUpdateMode::FastUpdate);
+
+	/** Clear an active triangle color function if one exists, and update the mesh */
+	virtual void ClearTriangleColorFunction(EDynamicMeshComponentRenderUpdateMode UpdateMode = EDynamicMeshComponentRenderUpdateMode::FastUpdate);
+
+	/** @return true if a triangle color function is configured */
+	virtual bool HasTriangleColorFunction();
+
+protected:
+	/** If this function is set, we will use these colors instead of vertex colors */
+	TUniqueFunction<FColor(const FDynamicMesh3*, int)> TriangleColorFunc = nullptr;
+
+	/** This function is passed via lambda to the RenderProxy to be able to access TriangleColorFunc */
+	FColor GetTriangleColor(const FDynamicMesh3* Mesh, int TriangleID);
+
+
+
+
+
+	//===============================================================================================================
+	// Support for Secondary triangle index buffers. When this is configured, then triangles identified
+	// by the filtering predicate function will be placed in a second set of RenderBuffers at the SceneProxy level.
+	// This can be combined with the SecondaryRenderMaterial support in UBaseDynamicMeshComponent to draw
+	// that triangle set with a different material, to efficiently accomplish UI features like highlighting a
+	// subset of mesh triangles.
+	//
+public:
+	/**
+	 * If Secondary triangle buffers are enabled, then we will filter triangles that pass the given predicate
+	 * function into a second index buffer. These triangles will be drawn with the Secondary render material
+	 * that is set in the BaseDynamicMeshComponent. Calling this function invalidates the SceneProxy.
+	 */
+	virtual void EnableSecondaryTriangleBuffers(TUniqueFunction<bool(const FDynamicMesh3*, int32)> SecondaryTriFilterFunc);
+
+	/**
+	 * Disable secondary triangle buffers. This invalidates the SceneProxy.
+	 */
+	virtual void DisableSecondaryTriangleBuffers();
+
+protected:
+	TUniqueFunction<bool(const FDynamicMesh3*, int32)> SecondaryTriFilterFunc = nullptr;
+
+
+
+
+	//===============================================================================================================
+	// Support for a Render Decomposition, which is basically a segmentation of the mesh triangles into
+	// subsets which will be turned into separate RenderBuffers in the Render Proxy. If this is configured,
+	// then various of the FastNotifyXYZUpdated() functions above will only need to rebuild the RenderBuffers
+	// that include affected triangles. The FMeshRenderDecomposition implementation has various options for
+	// building decompositions based on material, spatial clustering, etc.
+	//
+public:
+	/**
+	 * Configure a decomposition of the mesh, which will result in separate render buffers for each 
+	 * decomposition triangle group. Invalidates existing SceneProxy.
+	 */
+	virtual void SetExternalDecomposition(TUniquePtr<FMeshRenderDecomposition> Decomposition);
+
+protected:
+	TUniquePtr<FMeshRenderDecomposition> Decomposition;
+
+
+
+
+	//===============================================================================================================
+	// IRenderMeshPostProcessor Support. If a RenderMesh Postprocessor is configured, then instead of directly 
+	// passing the internal mesh to the RenderProxy, IRenderMeshPostProcessor::PostProcess is applied to populate
+	// the internal RenderMesh which is passed instead. This allows things like Displacement or Subdivision to be 
+	// done on-the-fly at the rendering level (which is potentially more efficient).
+	//
+public:
+	/**
+	 * Add a render mesh processor, to be called before the mesh is sent for rendering.
+	 */
+	virtual void SetRenderMeshPostProcessor(TUniquePtr<IRenderMeshPostProcessor> Processor);
+
+	/**
+	 * The SceneProxy should call these functions to get the post-processed RenderMesh. (See IRenderMeshPostProcessor.)
+	 */
+	virtual FDynamicMesh3* GetRenderMesh();
+
+	/**
+	 * The SceneProxy should call these functions to get the post-processed RenderMesh. (See IRenderMeshPostProcessor.)
+	 */
+	virtual const FDynamicMesh3* GetRenderMesh() const;
+
+protected:
+	TUniquePtr<IRenderMeshPostProcessor> RenderMeshPostProcessor;
+	TUniquePtr<FDynamicMesh3> RenderMesh;
+
+
+
+
+	//===============================================================================================================
+	// Support for Component attachment change notifications via delegates. Standard UE
+	// Actor/Component hierarchy does not generally provide these capabilities, but in some use
+	// cases (eg procedural mesh Actors) we need to know things like when the Component set inside
+	// an Actor is modified.
+public:
+	DECLARE_MULTICAST_DELEGATE_TwoParams(FComponentChildrenChangedDelegate, USceneComponent*, bool);
+
+	/**
+	 * The OnChildAttached() and OnChildDetached() implementations (from USceneComponent API) broadcast this delegate. This
+	 * allows Actors that have UDynamicMeshComponent's to respond to changes in their Component hierarchy.
+	 */
+	FComponentChildrenChangedDelegate OnChildAttachmentModified;
+
+
+
+
+	//===============================================================================================================
+	// Triangle-Vertex Tangents support. The default behavior is to not use Tangents, this will lead to incorrect
+	// rendering for any material with Normal Maps and some other shaders.
+	// If TangentsType == EDynamicMeshComponentTangentsMode::ExternallyProvided, the Tangent and Bitangent attributes of
+	//   the FDynamicMesh3 AttributeSet are used at the SceneProxy level, the Component is not involved
+	// If TangentsType == EDynamicMeshComponentTangentsMode::AutoCalculated, the Tangents are computed internally using
+	//   a fast MikkT approximation via FMeshTangentsf. They will be recomputed when the mesh is modified, however
+	//   they are *not* recomputed when using the Fast Update functions above (in that case InvalidateAutoCalculatedTangents() 
+	//   can be used to force recomputation)
+	//
+public:
+	UFUNCTION(BlueprintCallable, Category = "Dynamic Mesh Component")
+	void SetTangentsType(EDynamicMeshComponentTangentsMode NewTangentsType);
+
+	UFUNCTION(BlueprintCallable, Category = "Dynamic Mesh Component")
+	EDynamicMeshComponentTangentsMode GetTangentsType() const { return TangentsType; }
+
+	/** This function marks the auto tangents as dirty, they will be recomputed before they are used again */
+	virtual void InvalidateAutoCalculatedTangents();
+
+	/** @return AutoCalculated Tangent Set, which may require that they be recomputed, or nullptr if not enabled/available */
+	const UE::Geometry::FMeshTangentsf* GetAutoCalculatedTangents();
+
+protected:
+	/** How should Tangents be calculated/handled */
+	UPROPERTY()
+	EDynamicMeshComponentTangentsMode TangentsType = EDynamicMeshComponentTangentsMode::NoTangents;
+
+	/** true if AutoCalculatedTangents has been computed for current mesh */
+	bool bAutoCalculatedTangentsValid = false;
+
+	/** Set of per-triangle-vertex tangents computed for the current mesh. Only valid if bAutoCalculatedTangentsValid == true */
+	UE::Geometry::FMeshTangentsf AutoCalculatedTangents;
+
+	void UpdateAutoCalculatedTangents();
+
+
+
+	//===============================================================================================================
+	//
+	// Physics APIs
+	//
+
+public:
+	virtual bool GetPhysicsTriMeshData(struct FTriMeshCollisionData* CollisionData, bool InUseAllTriData) override;
+	virtual bool ContainsPhysicsTriMeshData(bool InUseAllTriData) const override;
+	virtual bool WantsNegXTriMesh() override;
+
+	virtual UBodySetup* GetBodySetup() override;
+	virtual void SetBodySetup(UBodySetup* NewSetup);
+
+	/** Type of Collision Geometry to use for this Mesh */
+	UPROPERTY(EditAnywhere, BlueprintReadWrite, Category = "Dynamic Mesh Component|Collision")
+	TEnumAsByte<enum ECollisionTraceFlag> CollisionType = ECollisionTraceFlag::CTF_UseSimpleAsComplex;
+
+	/** If true, current mesh will be used as Complex Collision source mesh. Independent of the CollisionType setting. */
+	UPROPERTY(EditAnywhere, BlueprintReadWrite, Category = "Dynamic Mesh Component|Collision");
+	bool bEnableComplexCollision = false;
+
+	/** If true, updates to the mesh will not result in immediate collision regeneration. Useful when the mesh will be modified multiple times before collision is needed. */
+	UPROPERTY(EditAnywhere, BlueprintReadWrite, Category = "Dynamic Mesh Component|Collision");
+	bool bDeferCollisionUpdates = false;
+
+protected:
+	UPROPERTY(Instanced)
+	TObjectPtr<UBodySetup> MeshBodySetup;
+
+	virtual void InvalidatePhysicsData();
+	virtual void RebuildPhysicsData();
+
+	bool bCollisionUpdatePending = false;
+	virtual void UpdateCollision(bool bOnlyIfPending);
+
+protected:
+
+	//
+	// standard Component internals, for computing bounds and managing the SceneProxy
+	//
+
+	/** Current local-space bounding box of Mesh */
+	UE::Geometry::FAxisAlignedBox3d LocalBounds;
+
+	/** Recompute LocalBounds from the current Mesh */
+	void UpdateLocalBounds();
+
+	/**
+	 * This is called to tell our RenderProxy about modifications to the material set.
+	 * We need to pass this on for things like material validation in the Editor.
+	 */
+	virtual void NotifyMaterialSetUpdated();
+
+	/**
+	 * If the render proxy is invalidated (eg by MarkRenderStateDirty()), it will be destroyed at the end of
+	 * the frame, but the base SceneProxy pointer is not nulled out immediately. As a result if we call various
+	 * partial-update functions after invalidating the proxy, they may be operating on an invalid proxy.
+	 * So we have to keep track of proxy-valid state ourselves.
+	 */
+	bool bProxyValid = false;
+
+	virtual FBaseDynamicMeshSceneProxy* GetBaseSceneProxy() override { return (FBaseDynamicMeshSceneProxy*)GetCurrentSceneProxy(); }
+
+	/** 
+	 * @return current render proxy, if valid, otherwise nullptr 
+	 */
+	FDynamicMeshSceneProxy* GetCurrentSceneProxy();
+
+
+	/**
+	 * Fully invalidate all rendering data for this Component. Current Proxy will be discarded, Bounds and possibly Tangents recomputed, etc
+	 */
+	void ResetProxy();
+
+	//~ Begin UPrimitiveComponent Interface.
+	virtual FPrimitiveSceneProxy* CreateSceneProxy() override;
+
+	//~ USceneComponent Interface.
+	virtual FBoxSphereBounds CalcBounds(const FTransform& LocalToWorld) const override;
+	virtual void OnChildAttached(USceneComponent* ChildComponent) override;
+	virtual void OnChildDetached(USceneComponent* ChildComponent) override;
+
+	//~ UObject Interface.
+	virtual void PostLoad() override;
+#if WITH_EDITOR
+	void PostEditChangeProperty(FPropertyChangedEvent& PropertyChangedEvent) override;
+#endif
+
+
+
+};
diff --git a/Engine/Source/Runtime/Experimental/GeometryFramework/Public/Components/MeshRenderDecomposition.h b/Engine/Source/Runtime/Experimental/GeometryFramework/Public/Components/MeshRenderDecomposition.h
new file mode 100644
--- /dev/null
+++ b/Engine/Source/Runtime/Experimental/GeometryFramework/Public/Components/MeshRenderDecomposition.h
@@ -0,0 +1,97 @@
+// Copyright Epic Games, Inc. All Rights Reserved.
+
+#pragma once
+
+#include "CoreMinimal.h"
+#include "GeometryBase.h"
+
+PREDECLARE_USE_GEOMETRY_CLASS(FDynamicMesh3);
+class UMaterialInterface;
+struct FComponentMaterialSet;
+
+/**
+ * FMeshRenderDecomposition represents a decomposition of a mesh into "chunks" of triangles, with associated materials.
+ * This is passed to the rendering components to split a mesh into multiple RenderBuffers, for more efficient updating.
+ */
+class GEOMETRYFRAMEWORK_API FMeshRenderDecomposition
+{
+public:
+
+	// Movable
+	FMeshRenderDecomposition(FMeshRenderDecomposition&&) = default;
+	FMeshRenderDecomposition& operator=(FMeshRenderDecomposition&&) = default;
+	// TArray<TUniquePtr> member cannot be default-constructed (in this case we just make it NonCopyable)
+	FMeshRenderDecomposition() = default;
+	FMeshRenderDecomposition(const FMeshRenderDecomposition&) = delete;
+	FMeshRenderDecomposition& operator=(const FMeshRenderDecomposition&) = delete;
+
+	struct FGroup
+	{
+		TArray<int32> Triangles;
+		UMaterialInterface* Material;
+	};
+	TArray<TUniquePtr<FGroup>> Groups;
+
+
+	/** Mapping from TriangleID to Groups array index. Initialized by BuildAssociations() */
+	TArray<int32> TriangleToGroupMap;
+
+	void Initialize(int32 Count)
+	{
+		Groups.SetNum(Count);
+		for (int32 k = 0; k < Count; ++k)
+		{
+			Groups[k] = MakeUnique<FGroup>();
+		}
+	}
+
+	int32 AppendGroup()
+	{
+		int32 N = Groups.Num();
+		Groups.SetNum(N + 1);
+		Groups[N] = MakeUnique<FGroup>();
+		return N;
+	}
+
+	int32 Num() const
+	{
+		return Groups.Num();
+	}
+
+	bool IsGroup(int32 Index) const
+	{
+		return Groups[Index] != nullptr;
+	}
+
+	FGroup& GetGroup(int32 Index)
+	{
+		return *Groups[Index];
+	}
+
+	const FGroup& GetGroup(int32 Index) const
+	{
+		return *Groups[Index];
+	}
+
+	int32 GetGroupForTriangle(int32 TriangleID) const
+	{
+		return TriangleToGroupMap[TriangleID];
+	}
+
+	/**
+	 * Construct mappings between mesh and groups (eg TriangleToGroupMap)
+	 */
+	void BuildAssociations(const FDynamicMesh3* Mesh);
+
+
+	/**
+	 * Build decomposition with one group for each MaterialID of mesh
+	 */
+	static void BuildMaterialDecomposition(const FDynamicMesh3* Mesh, const FComponentMaterialSet* MaterialSet, FMeshRenderDecomposition& Decomp);
+
+	/**
+	 * Build per-material decomposition, and then split each of those into chunks of at most MaxChunkSize
+	 * (actual chunk sizes will be highly variable and some may be very small...)
+	 */
+	static void BuildChunkedDecomposition(const FDynamicMesh3* Mesh, const FComponentMaterialSet* MaterialSet, FMeshRenderDecomposition& Decomp, int32 MaxChunkSize = 1 << 14 /* 16k */ );
+};
diff --git a/Engine/Source/Runtime/Experimental/GeometryFramework/Public/DynamicMeshActor.h b/Engine/Source/Runtime/Experimental/GeometryFramework/Public/DynamicMeshActor.h
new file mode 100644
--- /dev/null
+++ b/Engine/Source/Runtime/Experimental/GeometryFramework/Public/DynamicMeshActor.h
@@ -0,0 +1,29 @@
+// Copyright Epic Games, Inc. All Rights Reserved.
+
+#pragma once
+
+#include "CoreMinimal.h"
+#include "UDynamicMesh.h"
+#include "Components/DynamicMeshComponent.h"
+
+#include "DynamicMeshActor.generated.h"
+
+
+/**
+ * ADynamicMeshActor is an Actor that has a USimpleDynamicMeshComponent as it's RootObject.
+ * Used for experimental procedural geometry generation Blueprints. Currently no physics support.
+ */
+UCLASS(ConversionRoot, ComponentWrapperClass, ClassGroup=DynamicMesh, meta = (ChildCanTick))
+class GEOMETRYFRAMEWORK_API ADynamicMeshActor : public AActor
+{
+	GENERATED_UCLASS_BODY()
+private:
+	UPROPERTY(Category = DynamicMeshActor, VisibleAnywhere, BlueprintReadOnly, meta = (ExposeFunctionCategories = "Mesh,Rendering,Physics,Components|StaticMesh", AllowPrivateAccess = "true"))
+	TObjectPtr<class UDynamicMeshComponent> DynamicMeshComponent;
+
+public:
+	UFUNCTION(BlueprintCallable, Category = DynamicMeshActor)
+		UDynamicMeshComponent* GetDynamicMeshComponent() const { return DynamicMeshComponent; }
+};
+
+
diff --git a/Engine/Source/Runtime/Experimental/GeometryFramework/Public/GeometryFrameworkModule.h b/Engine/Source/Runtime/Experimental/GeometryFramework/Public/GeometryFrameworkModule.h
new file mode 100644
--- /dev/null
+++ b/Engine/Source/Runtime/Experimental/GeometryFramework/Public/GeometryFrameworkModule.h
@@ -0,0 +1,15 @@
+// Copyright Epic Games, Inc. All Rights Reserved.
+
+#pragma once
+
+#include "CoreMinimal.h"
+#include "Modules/ModuleManager.h"
+
+class FGeometryFrameworkModule : public IModuleInterface
+{
+public:
+
+	/** IModuleInterface implementation */
+	virtual void StartupModule() override;
+	virtual void ShutdownModule() override;
+};
diff --git a/Engine/Source/Runtime/Experimental/GeometryFramework/Public/UDynamicMesh.h b/Engine/Source/Runtime/Experimental/GeometryFramework/Public/UDynamicMesh.h
new file mode 100644
--- /dev/null
+++ b/Engine/Source/Runtime/Experimental/GeometryFramework/Public/UDynamicMesh.h
@@ -0,0 +1,396 @@
+// Copyright Epic Games, Inc. All Rights Reserved.
+
+#pragma once
+
+#include "CoreMinimal.h"
+#include "DynamicMesh/DynamicMesh3.h"
+
+#include "Changes/MeshVertexChange.h"
+#include "Changes/MeshChange.h"
+#include "Changes/MeshReplacementChange.h"
+
+#include "UDynamicMesh.generated.h"
+
+
+
+/**
+ * UDynamicMeshGenerator is an abstract base class for an implementation that can
+ * mutate an input mesh into an output mesh. A subclass of this class can be attached to
+ * a UDynamicMesh to allow for arbitrarily-complex procedural generation
+ */
+UCLASS(Abstract)
+class GEOMETRYFRAMEWORK_API UDynamicMeshGenerator : public UObject
+{
+	GENERATED_BODY()
+public:
+	/** Apply a change to MeshInOut */
+	virtual void Generate(FDynamicMesh3& MeshInOut) { };
+};
+
+
+
+
+
+/**
+ * EDynamicMeshChangeType is used by FDynamicMeshChangeInfo to indicate a "type" of mesh change
+ */
+UENUM(BlueprintType)
+enum class EDynamicMeshChangeType : uint8
+{
+	GeneralEdit = 0,
+	MeshChange = 1,
+	MeshReplacementChange = 2,
+	MeshVertexChange = 3,
+	
+	DeformationEdit = 4,
+	AttributeEdit = 5
+};
+
+UENUM(BlueprintType)
+enum class EDynamicMeshAttributeChangeFlags : uint8
+{
+	Unknown = 0,
+	MeshTopology = 1 << 0,
+
+	VertexPositions = 1 << 1,
+	NormalsTangents = 1 << 2,
+	VertexColors = 1 << 3,
+	UVs = 1 << 4,
+	TriangleGroups = 1 << 5
+};
+ENUM_CLASS_FLAGS(EDynamicMeshAttributeChangeFlags)
+
+/**
+ * FDynamicMeshChangeInfo stores information about a change to a UDynamicMesh.
+ * This struct is emitted by the UDynamicMesh OnPreMeshChanged() and OnMeshChanged() delegates.
+ */
+USTRUCT(BlueprintType)
+struct GEOMETRYFRAMEWORK_API FDynamicMeshChangeInfo
+{
+	GENERATED_BODY()
+
+	UPROPERTY(BlueprintReadOnly, EditAnywhere, Category = "DynamicMeshChangeInfo")
+	EDynamicMeshChangeType Type = EDynamicMeshChangeType::GeneralEdit;
+
+	UPROPERTY(BlueprintReadOnly, EditAnywhere, Category = "DynamicMeshChangeInfo")
+	EDynamicMeshAttributeChangeFlags Flags = EDynamicMeshAttributeChangeFlags::Unknown;
+
+	// for changes that are an FChange, indicates whether this is an 'Apply' or 'Revert' of the FChange
+	UPROPERTY(BlueprintReadOnly, EditAnywhere, Category = "DynamicMeshChangeInfo")
+	bool bIsRevertChange = false;
+
+	//
+	// internals
+	//
+	const FMeshChange* MeshChange = nullptr;
+	const FMeshReplacementChange* ReplaceChange = nullptr;
+	const FMeshVertexChange* VertexChange = nullptr;
+};
+
+// These delegates are used by UDynamicMesh
+DECLARE_MULTICAST_DELEGATE_TwoParams(FOnDynamicMeshChanged, UDynamicMesh*, FDynamicMeshChangeInfo);
+DECLARE_DYNAMIC_MULTICAST_DELEGATE_OneParam(FOnDynamicMeshModifiedBP, UDynamicMesh*, Mesh);
+
+
+
+/**
+ * UDynamicMesh is a UObject container for a FDynamicMesh3. 
+ */
+UCLASS(BlueprintType)
+class GEOMETRYFRAMEWORK_API UDynamicMesh : public UObject,
+	public IMeshVertexCommandChangeTarget, 
+	public IMeshCommandChangeTarget, 
+	public IMeshReplacementCommandChangeTarget
+{
+	GENERATED_UCLASS_BODY()
+public:
+
+	/**
+	 * Clear the internal mesh to an empty mesh.
+	 * This *does not* allocate a new mesh, so any existing mesh pointers/refs are still valid
+	 */
+	UFUNCTION(BlueprintCallable, Category = "Dynamic Mesh")
+	void Reset();
+
+	/**
+	 * Clear the internal mesh to a 100x100x100 cube with base at the origin.
+	 * This this instead of Reset() if an initially-empty mesh is undesirable (eg for a Component)
+	 */
+	UFUNCTION(BlueprintCallable, Category = "Dynamic Mesh")
+	void ResetToCube();
+
+	//
+	// Native access/modification functions
+	//
+
+	/**
+	 * Reset the internal mesh data and then optionally run the MeshGenerator
+	 */
+	virtual void InitializeMesh();
+
+	/**
+	 * @return true if the mesh has no triangles
+	 */
+	UFUNCTION(BlueprintCallable, Category = "Dynamic Mesh")
+	bool IsEmpty() const;
+
+	/**
+	 * @return number of triangles in the mesh
+	 */
+	UFUNCTION(BlueprintCallable, Category = "Dynamic Mesh")
+	int32 GetTriangleCount() const;
+
+
+	/** 
+	 * @return Const reference to the internal FDynamicMesh3.
+	 * @warning Calling ProcessMesh() is preferred! This interface may be  removed in the future
+	 */
+	const UE::Geometry::FDynamicMesh3& GetMeshRef() const { return *Mesh; }
+	/**
+	 * @return Const pointer to the internal FDynamicMesh3.
+	 * @warning Calling ProcessMesh() is preferred! This interface may be  removed in the future
+	 */
+	const UE::Geometry::FDynamicMesh3* GetMeshPtr() const { return Mesh.Get(); }
+
+	/**
+	 * @return Writable reference to the internal FDynamicMesh3. 
+	 * @warning Calling EditMesh() is preferred! This interface may be  removed in the future
+	 */
+	UE::Geometry::FDynamicMesh3& GetMeshRef() { return *Mesh; }
+	/**
+	 * @return Writable pointer to the internal FDynamicMesh3.
+	 * @warning Calling EditMesh() is preferred! This interface may be  removed in the future
+	 */
+	UE::Geometry::FDynamicMesh3* GetMeshPtr() { return Mesh.Get(); }
+
+
+	/** Replace the internal mesh with a copy of MoveMesh */
+	void SetMesh(const UE::Geometry::FDynamicMesh3& MoveMesh);
+
+	/** Replace the internal mesh with the data in MoveMesh */
+	void SetMesh(UE::Geometry::FDynamicMesh3&& MoveMesh);
+
+	/**
+	 * Apply ProcessFunc to the internal Mesh
+	 */
+	void ProcessMesh(TFunctionRef<void(const UE::Geometry::FDynamicMesh3&)> ProcessFunc) const;
+
+	/**
+	 * Apply EditFunc to the internal mesh.
+	 * This will broadcast PreMeshChangedEvent, then call EditFunc(), then broadcast MeshChangedEvent and MeshModifiedBPEvent
+	 */
+	void EditMesh(TFunctionRef<void(UE::Geometry::FDynamicMesh3&)> EditFunc,
+				  EDynamicMeshChangeType ChangeType = EDynamicMeshChangeType::GeneralEdit,
+				  EDynamicMeshAttributeChangeFlags ChangeFlags = EDynamicMeshAttributeChangeFlags::Unknown,
+				  bool bDeferChangeEvents = false);
+
+	/**
+	 * Take ownership of the internal Mesh, and have it replaced with a new mesh
+	 */
+	TUniquePtr<UE::Geometry::FDynamicMesh3> ExtractMesh();
+
+
+	//
+	// Change support
+	//
+
+
+	// IMeshVertexCommandChangeTarget implementation, allows a FVertexChange to be applied to the mesh
+	void ApplyChange(const FMeshVertexChange* Change, bool bRevert);
+
+	// IMeshCommandChangeTarget implementation, allows a FMeshChange to be applied to the mesh
+	void ApplyChange(const FMeshChange* Change, bool bRevert);
+
+	// IMeshReplacementCommandChangeTarget implementation, allows a FMeshReplacementChange to be applied to the mesh
+	void ApplyChange(const FMeshReplacementChange* Change, bool bRevert);
+
+
+	//
+	// Event support
+	//
+protected:
+	/** Broadcast before the internal mesh is modified by Reset funcs, SetMesh(), EditMesh(), and ApplyChange()'s above */
+	FOnDynamicMeshChanged PreMeshChangedEvent;
+	/** Broadcast after the internal mesh is modified, in the same cases as PreMeshChangedEvent */
+	FOnDynamicMeshChanged MeshChangedEvent;
+
+public:
+	/** Broadcast before the internal mesh is modified by Reset funcs, SetMesh(), EditMesh(), and ApplyChange()'s above */
+	FOnDynamicMeshChanged& OnPreMeshChanged() { return PreMeshChangedEvent; }
+	/** Broadcast after the internal mesh is modified, in the same cases as OnPreMeshChanged */
+	FOnDynamicMeshChanged& OnMeshChanged() { return MeshChangedEvent; }
+
+public:
+	/** Blueprintable event called when mesh is modified, in the same cases as OnMeshChanged */
+	UPROPERTY(BlueprintAssignable, meta = (DisplayName = "MeshModified"))
+	FOnDynamicMeshModifiedBP MeshModifiedBPEvent;
+
+
+	//
+	// Realtime Update support.
+	// This is intended to be used in situations where the internal mesh is being
+	// modified directly, ie instead of via EditMesh(), and we would like to
+	// notify listeners of these changes. Generally EditMesh() is preferred but in certain
+	// cases (eg like 3D sculpting) it is too complex to refactor all the mesh updates
+	// into EditMesh() calls (eg some are done async/etc). So, code that does those
+	// kinds of modifications can call PostRealtimeUpdate() to let any interested
+	// parties know that the mesh is actively changing.
+	//
+public:
+	DECLARE_MULTICAST_DELEGATE_OneParam(FOnMeshRealtimeUpdate, UDynamicMesh*);
+
+	/**
+	 * Multicast delegate that is broadcast whenever PostRealtimeUpdate() is called
+	 */
+	FOnMeshRealtimeUpdate& OnMeshRealtimeUpdate() { return MeshRealtimeUpdateEvent; }
+
+	/**
+	 * Broadcasts FOnMeshRealtimeUpdate
+	 */
+	virtual void PostRealtimeUpdate();
+
+protected:
+	FOnMeshRealtimeUpdate MeshRealtimeUpdateEvent;
+
+
+
+protected:
+	/**
+	 * Mesh data, owned by this object.
+	 * By default will have Attributes enabled and MaterialID enabled.
+	 */
+	TUniquePtr<UE::Geometry::FDynamicMesh3> Mesh;
+
+	/**
+	 * Allocate a new Mesh (ie pointer will change) and then call InitializeMesh()
+	 */
+	void InitializeNewMesh();
+
+	/**
+	 * Internal function that edits the Mesh, but broadcasts PreMeshChangedEvent and MeshChangedEvent, and then MeshModifiedBPEvent
+	 */
+	void EditMeshInternal(TFunctionRef<void(UE::Geometry::FDynamicMesh3&)> EditFunc, const FDynamicMeshChangeInfo& ChangeInfo, bool bDeferChangeEvents = false);
+
+
+public:
+	// serialize Mesh to an Archive
+	virtual void Serialize(FArchive& Archive) override;
+
+	// serialize Mesh to/from T3D
+	virtual void ExportCustomProperties(FOutputDevice& Out, uint32 Indent) override;
+	virtual void ImportCustomProperties(const TCHAR* SourceText, FFeedbackContext* Warn) override;
+
+
+
+	//
+	// (Preliminary) Procedural Generator support. If a Generator is configured, then each
+	// time this mesh is Reset(), it will call MeshGenerator->Generate(). The idea is that 
+	// generator will be set to something that emits a new mesh based on external data,
+	// for example a procedural primitive Actor/Component could configure a Generator that
+	// generates the primitive mesh based on Actor settings.
+	//
+
+protected:
+
+	/**
+	 * Active mesh generator. If configured, and bEnableMeshGenerator is true, then MeshGenerator->Generate() 
+	 * will be called when this mesh is Reset(). The Regenerate() function below can be used to force regeneration.
+	 */
+	UPROPERTY(Instanced)
+	TObjectPtr<UDynamicMeshGenerator> MeshGenerator;
+
+public:
+
+	/**
+	 * Controls whether the active Generator (if configured) will be applied when rebuilding the mesh
+	 */
+	UPROPERTY(EditAnywhere, Category = "Dynamic Mesh")
+	bool bEnableMeshGenerator = false;
+
+	/**
+	 * Set the active mesh generator. Clears if nullptr
+	 */
+	virtual void SetMeshGenerator(TObjectPtr<UDynamicMeshGenerator> NewGenerator);
+
+	/**
+	 * Set the active mesh generator. Clears if nullptr
+	 */
+	virtual void ClearMeshGenerator();
+
+	/**
+	 * Reset() the mesh, which will re-run the active MeshGenerator, if bEnableMeshGenerator
+	 */
+	virtual void Regenerate();
+
+};
+
+
+
+
+
+
+/**
+ * UDynamicMeshPool manages a Pool of UDynamicMesh objects. This allows
+ * the meshes to be re-used instead of being garbage-collected. This
+ * is intended to be used by Blueprints that need to do procedural geometry
+ * operations that generate temporary meshes, as these will commonly run their
+ * construction scripts many times as the user (eg) manipulates parameters,
+ * and constantly spawning new UDynamicMesh instances will result in enormous
+ * memory usage hanging around until GC runs.
+ *
+ * Usage is to call RequestMesh() to take ownership of an available UDynamicMesh (which
+ * will allocate a new one if the pool is empty) and ReturnMesh() to return it to the pool.
+ *
+ * ReturnAllMeshes() can be called to return all allocated meshes.
+ *
+ * In both cases, there is nothing preventing you from still holding on to the mesh.
+ * So, be careful.
+ *
+ * FreeAllMeshes() calls ReturnAllMeshes() and then releases the pool's references to
+ * the allocated meshes, so they can be Garbage Collected
+ *
+ * If you Request() more meshes than you Return(), the Pool will still be holding on to
+ * references to those meshes, and they will never be Garbage Collected (ie memory leak).
+ * As a failsafe, if the number of allocated meshes exceeds geometry.DynamicMesh.MaxPoolSize,
+ * the Pool will release all it's references and run garbage collection on the next call to RequestMesh().
+ * (Do not rely on this as a memory management strategy)
+ *
+ * An alternate strategy that could be employed here is for the Pool to not hold
+ * references to meshes it has provided, only those that have been explicitly returned.
+ * Then non-returned meshes would simply be garbage-collected, however it allows
+ * potentially a large amount of memory to be consumed until that occurs.
+ *
+ * UDynamicMesh::Reset() is called on the object returned to the Pool, which clears
+ * the internal FDynamicMesh3 (which uses normal C++ memory management, so no garbage collection involved)
+ * So the Pool does not re-use mesh memory, only the UObject containers.
+ */
+UCLASS(Transient)
+class GEOMETRYFRAMEWORK_API UDynamicMeshPool : public UObject
+{
+	GENERATED_BODY()
+public:
+	/** @return an available UDynamicMesh from the pool (possibly allocating a new mesh) */
+	UFUNCTION(BlueprintCallable, Category="Dynamic Mesh")
+	UDynamicMesh* RequestMesh();
+
+	/** Release a UDynamicMesh returned by RequestMesh() back to the pool */
+	UFUNCTION(BlueprintCallable, Category = "Dynamic Mesh")
+	void ReturnMesh(UDynamicMesh* Mesh);
+
+	/** Release all GeneratedMeshes back to the pool */
+	UFUNCTION(BlueprintCallable, Category = "Dynamic Mesh")
+	void ReturnAllMeshes();
+
+	/** Release all GeneratedMeshes back to the pool and allow them to be garbage collected */
+	UFUNCTION(BlueprintCallable, Category = "Dynamic Mesh")
+	void FreeAllMeshes();
+
+protected:
+	/** Meshes in the pool that are available */
+	UPROPERTY()
+	TArray<UDynamicMesh*> CachedMeshes;
+
+	/** All meshes the pool has allocated */
+	UPROPERTY()
+	TArray<UDynamicMesh*> AllCreatedMeshes;
+};
\ No newline at end of file
