Fix for no unity issue
#rnx

diff --git a/Engine/Source/Runtime/Experimental/GeometryFramework/Public/Components/BaseDynamicMeshSceneProxy.h b/Engine/Source/Runtime/Experimental/GeometryFramework/Public/Components/BaseDynamicMeshSceneProxy.h
--- a/Engine/Source/Runtime/Experimental/GeometryFramework/Public/Components/BaseDynamicMeshSceneProxy.h
+++ b/Engine/Source/Runtime/Experimental/GeometryFramework/Public/Components/BaseDynamicMeshSceneProxy.h
@@ -12,6 +12,7 @@
 #include "Rendering/ColorVertexBuffer.h"
 #include "DynamicMeshBuilder.h"
 #include "Components/BaseDynamicMeshComponent.h"
+#include "Materials/Material.h"
 
 using UE::Geometry::FDynamicMesh3;
 using UE::Geometry::FDynamicMeshAttributeSet;
