async log includes force disable info

diff --git a/Engine/Source/Runtime/Experimental/Chaos/Private/Chaos/Framework/PhysicsSolverBase.cpp b/Engine/Source/Runtime/Experimental/Chaos/Private/Chaos/Framework/PhysicsSolverBase.cpp
--- a/Engine/Source/Runtime/Experimental/Chaos/Private/Chaos/Framework/PhysicsSolverBase.cpp
+++ b/Engine/Source/Runtime/Experimental/Chaos/Private/Chaos/Framework/PhysicsSolverBase.cpp
@@ -146,17 +146,17 @@ namespace Chaos
 		, bStealAdvanceTasksForTesting(false)
 #endif
 	{
-		UE_LOG(LogChaos, Log, TEXT("FPhysicsSolverBase::AsyncDt:%f"), AsyncDt);
+		UE_LOG(LogChaos, Log, TEXT("FPhysicsSolverBase::AsyncDt:%f"), IsUsingAsyncResults() ? AsyncDt : -1);
 	}
 
 	void FPhysicsSolverBase::EnableAsyncMode(FReal FixedDt)
 	{
+		AsyncDt = FixedDt;
 		if (AsyncDt != FixedDt)
 		{
 			AccumulatedTime = 0;
-			UE_LOG(LogChaos, Log, TEXT("FPhysicsSolverBase::AsyncDt:%f"), FixedDt);
+			UE_LOG(LogChaos, Log, TEXT("FPhysicsSolverBase::AsyncDt:%f"), IsUsingAsyncResults() ? AsyncDt : -1);
 		}
-		AsyncDt = FixedDt;
 	}
 
 	void FPhysicsSolverBase::DisableAsyncMode()
