Fix for LWC compile errors

diff --git a/Engine/Plugins/FX/Niagara/Source/NiagaraVertexFactories/Public/NiagaraGPURayTracingTransformsShader.h b/Engine/Plugins/FX/Niagara/Source/NiagaraVertexFactories/Public/NiagaraGPURayTracingTransformsShader.h
--- a/Engine/Plugins/FX/Niagara/Source/NiagaraVertexFactories/Public/NiagaraGPURayTracingTransformsShader.h
+++ b/Engine/Plugins/FX/Niagara/Source/NiagaraVertexFactories/Public/NiagaraGPURayTracingTransformsShader.h
@@ -43,12 +43,12 @@ class NIAGARAVERTEXFACTORIES_API FNiagaraGPURayTracingTransformsCS : public FGlo
 		SHADER_PARAMETER(uint32,	RenderVisibilityValue)
 		SHADER_PARAMETER(uint32,	MeshIndexValue)
 
-		SHADER_PARAMETER(FVector,	DefaultPosition)
+		SHADER_PARAMETER(FVector3f,	DefaultPosition)
 		SHADER_PARAMETER(FVector4,	DefaultRotation)
-		SHADER_PARAMETER(FVector,	DefaultScale)
-		SHADER_PARAMETER(FVector,	MeshScale)
+		SHADER_PARAMETER(FVector3f,	DefaultScale)
+		SHADER_PARAMETER(FVector3f,	MeshScale)
 
-		SHADER_PARAMETER(FMatrix,	LocalTransform)
+		SHADER_PARAMETER(FMatrix44f,	LocalTransform)
 
 		SHADER_PARAMETER_SRV(Buffer<float>,	ParticleDataFloatBuffer)
 		//SHADER_PARAMETER_SRV(Buffer<half>,	ParticleDataHalfBuffer)
diff --git a/Engine/Source/Runtime/Experimental/Chaos/Private/Chaos/ChaosDebugDraw.cpp b/Engine/Source/Runtime/Experimental/Chaos/Private/Chaos/ChaosDebugDraw.cpp
--- a/Engine/Source/Runtime/Experimental/Chaos/Private/Chaos/ChaosDebugDraw.cpp
+++ b/Engine/Source/Runtime/Experimental/Chaos/Private/Chaos/ChaosDebugDraw.cpp
@@ -1163,7 +1163,7 @@ namespace Chaos
 				const FChaosDebugDrawSettings& DebugDrawSettings = GetChaosDebugDrawSettings(Settings);
 				for (auto& Particle : ParticlesView)
 				{				
-					FColor Color = (ColorScale * DebugDrawSettings.ShapesColorsPerState.GetColorFromState(Particle.ObjectState())).ToFColor(false);
+					FColor Color = ((float)ColorScale * DebugDrawSettings.ShapesColorsPerState.GetColorFromState(Particle.ObjectState())).ToFColor(false);
 					DrawParticleShapesImpl(SpaceTransform, GetHandleHelper(&Particle), Color, DebugDrawSettings);
 				}
 			}
@@ -1176,7 +1176,7 @@ namespace Chaos
 				const FChaosDebugDrawSettings& DebugDrawSettings = GetChaosDebugDrawSettings(Settings);
 				for (auto& Particle : ParticlesView)
 				{
-					FColor Color = (ColorScale * DebugDrawSettings.ShapesColorsPerState.GetColorFromState(Particle.ObjectState())).ToFColor(false);
+					FColor Color = ((float)ColorScale * DebugDrawSettings.ShapesColorsPerState.GetColorFromState(Particle.ObjectState())).ToFColor(false);
 					DrawParticleShapesImpl(SpaceTransform, GetHandleHelper(&Particle), Color, DebugDrawSettings);
 				}
 			}
@@ -1189,7 +1189,7 @@ namespace Chaos
 				const FChaosDebugDrawSettings& DebugDrawSettings = GetChaosDebugDrawSettings(Settings);
 				for (auto& Particle : ParticlesView)
 				{
-					FColor Color = (ColorScale * DebugDrawSettings.ShapesColorsPerState.GetColorFromState(Particle.ObjectState())).ToFColor(false);
+					FColor Color = ((float)ColorScale * DebugDrawSettings.ShapesColorsPerState.GetColorFromState(Particle.ObjectState())).ToFColor(false);
 					DrawParticleShapesImpl(SpaceTransform, GetHandleHelper(&Particle), Color, DebugDrawSettings);
 				}
 			}
