Chaos
- fix ContainerId passed to constraint graph when a constrained particle is disabled and re-enabled
- fix ParticleHandle::Island() for sleeping particles
- clean up constraint handle
- fix particle's constraint array natvis for VS debugger

diff --git a/Engine/Source/Programs/HeadlessChaos/Private/HeadlessChaosTestGraph.cpp b/Engine/Source/Programs/HeadlessChaos/Private/HeadlessChaosTestGraph.cpp
--- a/Engine/Source/Programs/HeadlessChaos/Private/HeadlessChaosTestGraph.cpp
+++ b/Engine/Source/Programs/HeadlessChaos/Private/HeadlessChaosTestGraph.cpp
@@ -24,19 +24,16 @@ namespace ChaosTest {
 
 
 	template<int32 T_TYPEID>
-	class TMockGraphConstraintHandle : public FConstraintHandle
+	class TMockGraphConstraintHandle : public TContainerConstraintHandle<TMockGraphConstraints<T_TYPEID>>
 	{
 	public:
 		TMockGraphConstraintHandle(TMockGraphConstraints<T_TYPEID>* InConstraintContainer, int32 ConstraintIndex)
-			: FConstraintHandle(FConstraintHandle::EType::Invalid, ConstraintIndex)
-			, ConstraintContainer(InConstraintContainer)
+			: TContainerConstraintHandle<TMockGraphConstraints<T_TYPEID>>(InConstraintContainer, ConstraintIndex)
 		{
 		}
 
 		virtual void SetEnabled(bool InEnabled) {};
 		virtual bool IsEnabled() const { return true; };
-
-		TMockGraphConstraints<T_TYPEID>* ConstraintContainer;
 	};
 
 	/**
@@ -45,7 +42,7 @@ namespace ChaosTest {
 	 * by using containers with different T_TYPEIDs.
 	 */
 	template<int32 T_TYPEID>
-	class TMockGraphConstraints
+	class TMockGraphConstraints : public FPBDConstraintContainer
 	{
 	public:
 		using FConstraintContainerHandle = TMockGraphConstraintHandle<T_TYPEID>;
@@ -54,6 +51,11 @@ namespace ChaosTest {
 			TVec2<int32> ConstrainedParticles;
 		};
 
+		TMockGraphConstraints()
+			: FPBDConstraintContainer(EConstraintContainerType::Invalid)	// Hmmm
+		{
+		}
+
 		int32 NumConstraints() const { return Constraints.Num(); }
 
 		TVec2<int32> ConstraintParticleIndices(const int32 ConstraintIndex) const
diff --git a/Engine/Source/Runtime/Experimental/Chaos/Private/Chaos/Collision/PBDCollisionConstraintHandle.cpp b/Engine/Source/Runtime/Experimental/Chaos/Private/Chaos/Collision/PBDCollisionConstraintHandle.cpp
--- a/Engine/Source/Runtime/Experimental/Chaos/Private/Chaos/Collision/PBDCollisionConstraintHandle.cpp
+++ b/Engine/Source/Runtime/Experimental/Chaos/Private/Chaos/Collision/PBDCollisionConstraintHandle.cpp
@@ -16,12 +16,12 @@ namespace Chaos
 	{
 		if (GetType() == FCollisionConstraintBase::FType::SinglePoint)
 		{
-			return ConstraintContainer->Constraints.SinglePointConstraints[ConstraintIndex];
+			return ConcreteContainer()->Constraints.SinglePointConstraints[ConstraintIndex];
 		}
 		else
 		{
 			check(GetType() == FCollisionConstraintBase::FType::SinglePointSwept);
-			return ConstraintContainer->Constraints.SinglePointSweptConstraints[ConstraintIndex];
+			return ConcreteContainer()->Constraints.SinglePointSweptConstraints[ConstraintIndex];
 		}
 	}
 
@@ -29,37 +29,37 @@ namespace Chaos
 	{
 		if (GetType() == FCollisionConstraintBase::FType::SinglePoint)
 		{
-			return ConstraintContainer->Constraints.SinglePointConstraints[ConstraintIndex];
+			return ConcreteContainer()->Constraints.SinglePointConstraints[ConstraintIndex];
 		}
 		else
 		{
 			check(GetType() == FCollisionConstraintBase::FType::SinglePointSwept);
-			return ConstraintContainer->Constraints.SinglePointSweptConstraints[ConstraintIndex];
+			return ConcreteContainer()->Constraints.SinglePointSweptConstraints[ConstraintIndex];
 		}
 	}
 
 	const FRigidBodyPointContactConstraint& FPBDCollisionConstraintHandle::GetPointContact() const
 	{ 
 		check(GetType() == FCollisionConstraintBase::FType::SinglePoint);
-		return ConstraintContainer->Constraints.SinglePointConstraints[ConstraintIndex]; 
+		return ConcreteContainer()->Constraints.SinglePointConstraints[ConstraintIndex];
 	}
 	
 	FRigidBodyPointContactConstraint& FPBDCollisionConstraintHandle::GetPointContact() 
 	{ 
 		check(GetType() == FCollisionConstraintBase::FType::SinglePoint);
-		return ConstraintContainer->Constraints.SinglePointConstraints[ConstraintIndex];
+		return ConcreteContainer()->Constraints.SinglePointConstraints[ConstraintIndex];
 	}
 
 	const FRigidBodySweptPointContactConstraint& FPBDCollisionConstraintHandle::GetSweptPointContact() const
 	{
 		check(GetType() == FCollisionConstraintBase::FType::SinglePointSwept);
-		return ConstraintContainer->Constraints.SinglePointSweptConstraints[ConstraintIndex];
+		return ConcreteContainer()->Constraints.SinglePointSweptConstraints[ConstraintIndex];
 	}
 	
 	FRigidBodySweptPointContactConstraint& FPBDCollisionConstraintHandle::GetSweptPointContact()
 	{
 		check(GetType() == FCollisionConstraintBase::FType::SinglePointSwept);
-		return ConstraintContainer->Constraints.SinglePointSweptConstraints[ConstraintIndex];
+		return ConcreteContainer()->Constraints.SinglePointSweptConstraints[ConstraintIndex];
 	}
 
 }
\ No newline at end of file
diff --git a/Engine/Source/Runtime/Experimental/Chaos/Private/Chaos/PBDCollisionConstraints.cpp b/Engine/Source/Runtime/Experimental/Chaos/Private/Chaos/PBDCollisionConstraints.cpp
--- a/Engine/Source/Runtime/Experimental/Chaos/Private/Chaos/PBDCollisionConstraints.cpp
+++ b/Engine/Source/Runtime/Experimental/Chaos/Private/Chaos/PBDCollisionConstraints.cpp
@@ -102,7 +102,8 @@ namespace Chaos
 		const int32 InApplyPairIterations /*= 1*/,
 		const int32 InApplyPushOutPairIterations /*= 1*/,
 		const FReal InRestitutionThreshold /*= (FReal)0*/)
-		: bInAppendOperation(false)
+		: FPBDConstraintContainer(EConstraintContainerType::Collision)
+		, bInAppendOperation(false)
 		, Particles(InParticles)
 		, NumActivePointConstraints(0)
 		, NumActiveSweptPointConstraints(0)
diff --git a/Engine/Source/Runtime/Experimental/Chaos/Private/Chaos/PBDConstraintContainer.cpp b/Engine/Source/Runtime/Experimental/Chaos/Private/Chaos/PBDConstraintContainer.cpp
--- a/Engine/Source/Runtime/Experimental/Chaos/Private/Chaos/PBDConstraintContainer.cpp
+++ b/Engine/Source/Runtime/Experimental/Chaos/Private/Chaos/PBDConstraintContainer.cpp
@@ -3,7 +3,9 @@
 
 namespace Chaos
 {
-	FPBDConstraintContainer::FPBDConstraintContainer()
+	FPBDConstraintContainer::FPBDConstraintContainer(EConstraintContainerType InConstraintType)
+		: ConstraintContainerType(InConstraintType)
+		, ContainerId(INDEX_NONE)
 	{
 	}
 
diff --git a/Engine/Source/Runtime/Experimental/Chaos/Private/Chaos/PBDConstraintGraph.cpp b/Engine/Source/Runtime/Experimental/Chaos/Private/Chaos/PBDConstraintGraph.cpp
--- a/Engine/Source/Runtime/Experimental/Chaos/Private/Chaos/PBDConstraintGraph.cpp
+++ b/Engine/Source/Runtime/Experimental/Chaos/Private/Chaos/PBDConstraintGraph.cpp
@@ -311,33 +311,27 @@ void FPBDConstraintGraph::RemoveConstraint(const uint32 InContainerId, FConstrai
 {
 	check(InConstraintHandle);
 
-	int32 EdgeIndex = InConstraintHandle->GetConstraintGraphIndex();
-
-	int32 Index0 = INDEX_NONE, Index1 = INDEX_NONE;
-
-	int32* PNodeIndex0 = (ConstrainedParticles[0]) ? ParticleToNodeIndex.Find(ConstrainedParticles[0]) : nullptr;
-	int32* PNodeIndex1 = (ConstrainedParticles[1]) ? ParticleToNodeIndex.Find(ConstrainedParticles[1]) : nullptr;
-	if (ensure(PNodeIndex0 || PNodeIndex1))
+	const int32 EdgeIndex = InConstraintHandle->GetConstraintGraphIndex();
+	if (ensure(EdgeIndex != INDEX_NONE))
 	{
-		// we should have an valid index at this point
-		if (PNodeIndex0)
+		const int32* PNodeIndex0 = (ConstrainedParticles[0]) ? ParticleToNodeIndex.Find(ConstrainedParticles[0]) : nullptr;
+		const int32* PNodeIndex1 = (ConstrainedParticles[1]) ? ParticleToNodeIndex.Find(ConstrainedParticles[1]) : nullptr;
+		if (ensure(PNodeIndex0 || PNodeIndex1))
 		{
-			Nodes[*PNodeIndex0].Edges.Remove(EdgeIndex);
-			Index0 = *PNodeIndex0;
-		}
-		if (PNodeIndex1)
-		{
-			Nodes[*PNodeIndex1].Edges.Remove(EdgeIndex);
-			Index1 = *PNodeIndex1;
+			// we should have an valid index at this point
+			if (PNodeIndex0)
+			{
+				Nodes[*PNodeIndex0].Edges.Remove(EdgeIndex);
+			}
+			if (PNodeIndex1)
+			{
+				Nodes[*PNodeIndex1].Edges.Remove(EdgeIndex);
+			}
 		}
-	}
 
-	// we need to find the edge index that matches this constraint
-	// if the particles are being removed first and the constraint is being disabled
-	if (ensure(EdgeIndex != INDEX_NONE))
-	{
-		InConstraintHandle->SetConstraintGraphIndex(INDEX_NONE);
 		Edges[EdgeIndex] = FGraphEdge();
+
+		InConstraintHandle->SetConstraintGraphIndex(INDEX_NONE);
 	}
 }
 
@@ -640,7 +634,7 @@ void FPBDConstraintGraph::ComputeIslands(const TParticleView<FPBDRigidParticles>
 	check(IslandToParticles.Num() == IslandToConstraints.Num());
 	check(IslandToParticles.Num() == IslandToData.Num());
 	// @todo(ccaulfield): make a more complex unit test to check island integrity
-	//checkSlow(DebugCheckGraph());
+	checkSlow(DebugCheckGraph());
 }
 
 
diff --git a/Engine/Source/Runtime/Experimental/Chaos/Private/Chaos/PBDConstraintRule.cpp b/Engine/Source/Runtime/Experimental/Chaos/Private/Chaos/PBDConstraintRule.cpp
--- a/Engine/Source/Runtime/Experimental/Chaos/Private/Chaos/PBDConstraintRule.cpp
+++ b/Engine/Source/Runtime/Experimental/Chaos/Private/Chaos/PBDConstraintRule.cpp
@@ -24,8 +24,8 @@ namespace Chaos
 	template<class T_CONSTRAINTS>
 	void TPBDConstraintGraphRuleImpl<T_CONSTRAINTS>::BindToGraph(FPBDConstraintGraph& InContactGraph, uint32 InContainerId)
 	{
+		Constraints.SetContainerId(InContainerId);
 		ConstraintGraph = &InContactGraph;
-		ContainerId = InContainerId;
 	}
 
 	template<class T_CONSTRAINTS>
@@ -44,7 +44,7 @@ namespace Chaos
 		{
 			if (ConstraintHandle->IsEnabled())
 			{
-				ConstraintGraph->AddConstraint(ContainerId, ConstraintHandle, ConstraintHandle->GetConstrainedParticles());
+				ConstraintGraph->AddConstraint(GetContainerId(), ConstraintHandle, ConstraintHandle->GetConstrainedParticles());
 			}
 		}
 
diff --git a/Engine/Source/Runtime/Experimental/Chaos/Private/Chaos/PBDJointConstraints.cpp b/Engine/Source/Runtime/Experimental/Chaos/Private/Chaos/PBDJointConstraints.cpp
--- a/Engine/Source/Runtime/Experimental/Chaos/Private/Chaos/PBDJointConstraints.cpp
+++ b/Engine/Source/Runtime/Experimental/Chaos/Private/Chaos/PBDJointConstraints.cpp
@@ -40,79 +40,79 @@ namespace Chaos
 
 	
 	FPBDJointConstraintHandle::FPBDJointConstraintHandle(FConstraintContainer* InConstraintContainer, int32 InConstraintIndex)
-		: TContainerConstraintHandle<FPBDJointConstraints>(StaticType(), InConstraintContainer, InConstraintIndex)
+		: TContainerConstraintHandle<FPBDJointConstraints>(InConstraintContainer, InConstraintIndex)
 	{
 	}
 
 	
 	void FPBDJointConstraintHandle::CalculateConstraintSpace(FVec3& OutXa, FMatrix33& OutRa, FVec3& OutXb, FMatrix33& OutRb) const
 	{
-		ConstraintContainer->CalculateConstraintSpace(ConstraintIndex, OutXa, OutRa, OutXb, OutRb);
+		ConcreteContainer()->CalculateConstraintSpace(ConstraintIndex, OutXa, OutRa, OutXb, OutRb);
 	}
 
 
 	int32 FPBDJointConstraintHandle::GetConstraintIsland() const
 	{
-		return ConstraintContainer->GetConstraintIsland(ConstraintIndex);
+		return ConcreteContainer()->GetConstraintIsland(ConstraintIndex);
 	}
 
 
 	int32 FPBDJointConstraintHandle::GetConstraintLevel() const
 	{
-		return ConstraintContainer->GetConstraintLevel(ConstraintIndex);
+		return ConcreteContainer()->GetConstraintLevel(ConstraintIndex);
 	}
 
 
 	int32 FPBDJointConstraintHandle::GetConstraintColor() const
 	{
-		return ConstraintContainer->GetConstraintColor(ConstraintIndex);
+		return ConcreteContainer()->GetConstraintColor(ConstraintIndex);
 	}
 
 
 	bool FPBDJointConstraintHandle::IsConstraintEnabled() const
 	{
-		return ConstraintContainer->IsConstraintEnabled(ConstraintIndex);
+		return ConcreteContainer()->IsConstraintEnabled(ConstraintIndex);
 	}
 
 	bool FPBDJointConstraintHandle::IsConstraintBreaking() const
 	{
-		return ConstraintContainer->IsConstraintBreaking(ConstraintIndex);
+		return ConcreteContainer()->IsConstraintBreaking(ConstraintIndex);
 	}
 
 	void FPBDJointConstraintHandle::ClearConstraintBreaking()
 	{
-		return ConstraintContainer->ClearConstraintBreaking(ConstraintIndex);
+		return ConcreteContainer()->ClearConstraintBreaking(ConstraintIndex);
 	}
 
 	FVec3 FPBDJointConstraintHandle::GetLinearImpulse() const
 	{
-		return ConstraintContainer->GetConstraintLinearImpulse(ConstraintIndex);
+		return ConcreteContainer()->GetConstraintLinearImpulse(ConstraintIndex);
 	}
 
 	FVec3 FPBDJointConstraintHandle::GetAngularImpulse() const
 	{
-		return ConstraintContainer->GetConstraintAngularImpulse(ConstraintIndex);
+		return ConcreteContainer()->GetConstraintAngularImpulse(ConstraintIndex);
 	}
 
 	
 	const FPBDJointSettings& FPBDJointConstraintHandle::GetSettings() const
 	{
-		return ConstraintContainer->GetConstraintSettings(ConstraintIndex);
+		return ConcreteContainer()->GetConstraintSettings(ConstraintIndex);
 	}
 
 	void FPBDJointConstraintHandle::SetSettings(const FPBDJointSettings& Settings)
 	{
-		ConstraintContainer->SetConstraintSettings(ConstraintIndex, Settings);
+		ConcreteContainer()->SetConstraintSettings(ConstraintIndex, Settings);
 	}
 
 	TVector<FGeometryParticleHandle*, 2> FPBDJointConstraintHandle::GetConstrainedParticles() const 
 	{ 
-		return ConstraintContainer->GetConstrainedParticles(ConstraintIndex); 
+		return ConcreteContainer()->GetConstrainedParticles(ConstraintIndex);
 	}
 
 	void FPBDJointConstraintHandle::SetConstraintEnabled(bool bInEnabled)
 	{
-		return ConstraintContainer->SetConstraintEnabled(ConstraintIndex, bInEnabled);
+		return ConcreteContainer()->SetConstraintEnabled(ConstraintIndex, bInEnabled);
 	}
 
 	//
@@ -318,7 +318,8 @@ namespace Chaos
 
 	
 	FPBDJointConstraints::FPBDJointConstraints(const FPBDJointSolverSettings& InSettings)
-		: Settings(InSettings)
+		: FPBDConstraintContainer(EConstraintContainerType::Joint)
+		, Settings(InSettings)
 		, bJointsDirty(false)
 		, bUpdateVelocityInApplyConstraints(false)
 		, PreApplyCallback(nullptr)
diff --git a/Engine/Source/Runtime/Experimental/Chaos/Private/Chaos/PBDPositionConstraints.cpp b/Engine/Source/Runtime/Experimental/Chaos/Private/Chaos/PBDPositionConstraints.cpp
--- a/Engine/Source/Runtime/Experimental/Chaos/Private/Chaos/PBDPositionConstraints.cpp
+++ b/Engine/Source/Runtime/Experimental/Chaos/Private/Chaos/PBDPositionConstraints.cpp
@@ -6,7 +6,7 @@ namespace Chaos
 {
 	TVector<FGeometryParticleHandle*, 2> FPBDPositionConstraintHandle::GetConstrainedParticles() const 
 	{ 
-		return ConstraintContainer->GetConstrainedParticles(ConstraintIndex); 
+		return ConcreteContainer()->GetConstrainedParticles(ConstraintIndex);
 	}
 
 	bool FPBDPositionConstraints::Apply(const FReal Dt, const TArray<FConstraintContainerHandle*>& ConstraintHandles, const int32 It, const int32 NumIts) const
diff --git a/Engine/Source/Runtime/Experimental/Chaos/Private/Chaos/PBDRigidDynamicSpringConstraints.cpp b/Engine/Source/Runtime/Experimental/Chaos/Private/Chaos/PBDRigidDynamicSpringConstraints.cpp
--- a/Engine/Source/Runtime/Experimental/Chaos/Private/Chaos/PBDRigidDynamicSpringConstraints.cpp
+++ b/Engine/Source/Runtime/Experimental/Chaos/Private/Chaos/PBDRigidDynamicSpringConstraints.cpp
@@ -8,7 +8,7 @@ using namespace Chaos;
 
 TVec2<FGeometryParticleHandle*> FPBDRigidDynamicSpringConstraintHandle::GetConstrainedParticles() const
 {
-	return ConstraintContainer->GetConstrainedParticles(ConstraintIndex);
+	return ConcreteContainer()->GetConstrainedParticles(ConstraintIndex);
 }
 
 void FPBDRigidDynamicSpringConstraints::UpdatePositionBasedState(const FReal Dt)
diff --git a/Engine/Source/Runtime/Experimental/Chaos/Private/Chaos/PBDRigidSpringConstraints.cpp b/Engine/Source/Runtime/Experimental/Chaos/Private/Chaos/PBDRigidSpringConstraints.cpp
--- a/Engine/Source/Runtime/Experimental/Chaos/Private/Chaos/PBDRigidSpringConstraints.cpp
+++ b/Engine/Source/Runtime/Experimental/Chaos/Private/Chaos/PBDRigidSpringConstraints.cpp
@@ -12,28 +12,28 @@ namespace Chaos
 
 	const TVector<FVec3, 2>& FPBDRigidSpringConstraintHandle::GetConstraintPositions() const
 	{
-		return ConstraintContainer->GetConstraintPositions(ConstraintIndex);
+		return ConcreteContainer()->GetConstraintPositions(ConstraintIndex);
 
 	}
 
 	void FPBDRigidSpringConstraintHandle::SetConstraintPositions(const TVector<FVec3, 2>& ConstraintPositions)
 	{
-		ConstraintContainer->SetConstraintPositions(ConstraintIndex, ConstraintPositions);
+		ConcreteContainer()->SetConstraintPositions(ConstraintIndex, ConstraintPositions);
 	}
 	
 	TVector<typename FPBDRigidSpringConstraintHandle::FGeometryParticleHandle*, 2> FPBDRigidSpringConstraintHandle::GetConstrainedParticles() const
 	{ 
-		return ConstraintContainer->GetConstrainedParticles(ConstraintIndex);
+		return ConcreteContainer()->GetConstrainedParticles(ConstraintIndex);
 	}
 
 	FReal FPBDRigidSpringConstraintHandle::GetRestLength() const
 	{
-		return ConstraintContainer->GetRestLength(ConstraintIndex);
+		return ConcreteContainer()->GetRestLength(ConstraintIndex);
 	}
 
 	void FPBDRigidSpringConstraintHandle::SetRestLength(const FReal SpringLength)
 	{
-		ConstraintContainer->SetRestLength(ConstraintIndex, SpringLength);
+		ConcreteContainer()->SetRestLength(ConstraintIndex, SpringLength);
 	}
 
 	//
@@ -41,6 +41,7 @@ namespace Chaos
 	//
 
 	FPBDRigidSpringConstraints::FPBDRigidSpringConstraints()
+		: FPBDConstraintContainer(EConstraintContainerType::RigidSpring)
 	{}
 
 	FPBDRigidSpringConstraints::~FPBDRigidSpringConstraints()
diff --git a/Engine/Source/Runtime/Experimental/Chaos/Private/Chaos/PBDSuspensionConstraints.cpp b/Engine/Source/Runtime/Experimental/Chaos/Private/Chaos/PBDSuspensionConstraints.cpp
--- a/Engine/Source/Runtime/Experimental/Chaos/Private/Chaos/PBDSuspensionConstraints.cpp
+++ b/Engine/Source/Runtime/Experimental/Chaos/Private/Chaos/PBDSuspensionConstraints.cpp
@@ -5,29 +5,29 @@
 namespace Chaos
 {
 	FPBDSuspensionConstraintHandle::FPBDSuspensionConstraintHandle(FConstraintContainer* InConstraintContainer, int32 InConstraintIndex) : 
-		TContainerConstraintHandle<FPBDSuspensionConstraints>(StaticType(), InConstraintContainer, InConstraintIndex)
+		TContainerConstraintHandle<FPBDSuspensionConstraints>(InConstraintContainer, InConstraintIndex)
 	{
 	}
 
 	const FPBDSuspensionSettings& FPBDSuspensionConstraintHandle::GetSettings() const
 	{
-		return ConstraintContainer->GetSettings(ConstraintIndex);
+		return ConcreteContainer()->GetSettings(ConstraintIndex);
 	}
 
 	FPBDSuspensionSettings& FPBDSuspensionConstraintHandle::GetSettings()
 	{
-		return ConstraintContainer->GetSettings(ConstraintIndex);
+		return ConcreteContainer()->GetSettings(ConstraintIndex);
 	}
 
 	void FPBDSuspensionConstraintHandle::SetSettings(const FPBDSuspensionSettings& Settings)
 	{
-		ConstraintContainer->SetSettings(ConstraintIndex, Settings);
+		ConcreteContainer()->SetSettings(ConstraintIndex, Settings);
 	}
 
 
 	TVec2<FGeometryParticleHandle*> FPBDSuspensionConstraintHandle::GetConstrainedParticles() const
 	{
-		return ConstraintContainer->GetConstrainedParticles(ConstraintIndex);
+		return ConcreteContainer()->GetConstrainedParticles(ConstraintIndex);
 	}
 
 
diff --git a/Engine/Source/Runtime/Experimental/Chaos/Public/Chaos/Collision/PBDCollisionConstraintHandle.h b/Engine/Source/Runtime/Experimental/Chaos/Public/Chaos/Collision/PBDCollisionConstraintHandle.h
--- a/Engine/Source/Runtime/Experimental/Chaos/Public/Chaos/Collision/PBDCollisionConstraintHandle.h
+++ b/Engine/Source/Runtime/Experimental/Chaos/Public/Chaos/Collision/PBDCollisionConstraintHandle.h
@@ -26,11 +26,11 @@ namespace Chaos
 		{}
 
 		FPBDCollisionConstraintHandle(FPBDCollisionConstraints* InConstraintContainer, int32 InConstraintIndex, typename FCollisionConstraintBase::FType InType)
-			: TContainerConstraintHandle<FPBDCollisionConstraints>(StaticType(), InConstraintContainer, InConstraintIndex)
+			: TContainerConstraintHandle<FPBDCollisionConstraints>(InConstraintContainer, InConstraintIndex)
 			, ConstraintType(InType)
 		{
 		}
-		static FConstraintHandle::EType StaticType() { return FConstraintHandle::EType::Collision; }
+		static EConstraintContainerType StaticType() { return EConstraintContainerType::Collision; }
 
 
 		FHandleKey GetKey();
@@ -88,7 +88,7 @@ namespace Chaos
 	protected:
 		typename FCollisionConstraintBase::FType ConstraintType;
 		using Base::ConstraintIndex;
-		using Base::ConstraintContainer;
+		using Base::ConcreteContainer;
 
 
 	};
diff --git a/Engine/Source/Runtime/Experimental/Chaos/Public/Chaos/ConstraintHandle.h b/Engine/Source/Runtime/Experimental/Chaos/Public/Chaos/ConstraintHandle.h
--- a/Engine/Source/Runtime/Experimental/Chaos/Public/Chaos/ConstraintHandle.h
+++ b/Engine/Source/Runtime/Experimental/Chaos/Public/Chaos/ConstraintHandle.h
@@ -16,27 +16,18 @@ namespace Chaos
 	class CHAOS_API FConstraintHandle
 	{
 	public:
-		enum class EType : uint8
-		{
-			Invalid = 0,
-			Collision,
-			RigidSpring,
-			DynamicSpring,
-			Position,
-			Joint,
-			Suspension
-		};
+		using FGeometryParticleHandle = TGeometryParticleHandle<FReal, 3>;
 
+		FConstraintHandle() : ConstraintContainer(nullptr), ConstraintIndex(INDEX_NONE), ConstraintGraphIndex(INDEX_NONE) { }
 
-		using FGeometryParticleHandle = TGeometryParticleHandle<FReal, 3>;
+		FConstraintHandle(FPBDConstraintContainer* InContainer, int32 InConstraintIndex): ConstraintContainer(InContainer), ConstraintIndex(InConstraintIndex), ConstraintGraphIndex(INDEX_NONE) {}
 
-		FConstraintHandle() : Type(EType::Invalid), ConstraintIndex(INDEX_NONE), ConstraintGraphIndex(INDEX_NONE) { }
-		FConstraintHandle(EType InType, int32 InConstraintIndex): Type(InType), ConstraintIndex(InConstraintIndex), ConstraintGraphIndex(INDEX_NONE) {}
 		virtual ~FConstraintHandle() {}
 
 		bool IsValid() const
 		{
-			return (ConstraintIndex != INDEX_NONE && IsEnabled());
+			// @todo(chaos): why does IsValid() also check IsEnabled()?
+			return (ConstraintContainer != nullptr) && (ConstraintIndex != INDEX_NONE) && IsEnabled();
 		}
 
 		int32 GetConstraintIndex() const
@@ -54,17 +45,29 @@ namespace Chaos
 			ConstraintGraphIndex = InIndex;
 		}
 
+		bool IsInConstraintGraph() const
+		{
+			return (ConstraintGraphIndex != INDEX_NONE);
+		}
+
 
-		virtual void SetEnabled(bool InEnabled) = 0;
-		virtual bool IsEnabled() const = 0;
+		// Implemented in ConstraintContainer.h
+		int32 GetContainerId() const;
 
-		template<typename T>  T* As() { return T::StaticType() == Type ? static_cast<T*>(this) : nullptr; }
-		template<typename T>  const T* As() const { return T::StaticType() == Type ? static_cast<const T*>(this) : nullptr; }
+		// Implemented in ConstraintContainer.h
+		void SetEnabled(bool InEnabled);
+
+		// Implemented in ConstraintContainer.h
+		bool IsEnabled() const;
+
+		// Implemented in ConstraintContainer.h
+		template<typename T>  T* As();
+		template<typename T>  const T* As() const;
 
 	protected:
 		friend class FPBDConstraintContainer;
 
-		EType Type;
+		FPBDConstraintContainer* ConstraintContainer;
 		int32 ConstraintIndex;
 		int32 ConstraintGraphIndex; // @todo(chaos): move constraint graph index to base constraint container
 	};
@@ -81,18 +84,33 @@ namespace Chaos
 		using FGeometryParticleHandle = typename Base::FGeometryParticleHandle;
 		using FConstraintContainer = T_CONTAINER;
 
-		TContainerConstraintHandle() : ConstraintContainer(nullptr) {}
-		TContainerConstraintHandle(Base::EType InType, FConstraintContainer* InConstraintContainer, int32 InConstraintIndex)
-			: FConstraintHandle(InType, InConstraintIndex), ConstraintContainer(InConstraintContainer) {}
-
-		void RemoveConstraint() { ConstraintContainer->RemoveConstraint(ConstraintIndex); }
+		TContainerConstraintHandle()
+		{
+		}
+		
+		TContainerConstraintHandle(FConstraintContainer* InConstraintContainer, int32 InConstraintIndex)
+			: FConstraintHandle(InConstraintContainer, InConstraintIndex)
+		{
+		}
 
-		void SetEnabled(bool InEnabled) override { ConstraintContainer->SetConstraintEnabled(ConstraintIndex,InEnabled); }
-		bool IsEnabled() const override { return ConstraintContainer->IsConstraintEnabled(ConstraintIndex); }
+		// @todo(chaos): Make this a virtual on FConstraintContainer and move to base class
+		void RemoveConstraint()
+		{
+			ConcreteContainer()->RemoveConstraint(ConstraintIndex);
+		}
 
 	protected:
+		FConstraintContainer* ConcreteContainer()
+		{
+			return static_cast<FConstraintContainer*>(ConstraintContainer);
+		}
+
+		const FConstraintContainer* ConcreteContainer() const
+		{
+			return static_cast<const FConstraintContainer*>(ConstraintContainer);
+		}
+
 		using Base::ConstraintIndex;
-		FConstraintContainer* ConstraintContainer;
 	};
 
 
diff --git a/Engine/Source/Runtime/Experimental/Chaos/Public/Chaos/GeometryParticles.h b/Engine/Source/Runtime/Experimental/Chaos/Public/Chaos/GeometryParticles.h
--- a/Engine/Source/Runtime/Experimental/Chaos/Public/Chaos/GeometryParticles.h
+++ b/Engine/Source/Runtime/Experimental/Chaos/Public/Chaos/GeometryParticles.h
@@ -250,6 +250,8 @@ namespace Chaos
 
 	using FShapesArray = TArray<TUniquePtr<FPerShapeData>, TInlineAllocator<1>>;
 
+	using FConstraintHandleArray = TArray<FConstraintHandle*>;
+
 	void CHAOS_API UpdateShapesArrayFromGeometry(FShapesArray& ShapesArray, TSerializablePtr<FImplicitObject> Geometry, const FRigidTransform3& ActorTM, IPhysicsProxyBase* Proxy);
 
 
@@ -620,7 +622,7 @@ namespace Chaos
 			return WeakHandle;
 		}
 
-		CHAOS_API TArray<FConstraintHandle*>& ParticleConstraints(const int32 Index)
+		CHAOS_API FConstraintHandleArray& ParticleConstraints(const int32 Index)
 		{
 			return MParticleConstraints[Index];
 		}
@@ -772,7 +774,7 @@ public:
 		TArrayCollectionArray<void*> MUserData;
 		TArrayCollectionArray<FSyncState> MSyncState;
 		TArrayCollectionArray<FWeakParticleHandle> MWeakParticleHandle;
-		TArrayCollectionArray<TArray<FConstraintHandle*> > MParticleConstraints;
+		TArrayCollectionArray<FConstraintHandleArray> MParticleConstraints;
 
 		void UpdateShapesArray(const int32 Index)
 		{
diff --git a/Engine/Source/Runtime/Experimental/Chaos/Public/Chaos/PBDConstraintContainer.h b/Engine/Source/Runtime/Experimental/Chaos/Public/Chaos/PBDConstraintContainer.h
--- a/Engine/Source/Runtime/Experimental/Chaos/Public/Chaos/PBDConstraintContainer.h
+++ b/Engine/Source/Runtime/Experimental/Chaos/Public/Chaos/PBDConstraintContainer.h
@@ -6,6 +6,18 @@
 
 namespace Chaos
 {
+	// @todo(chaos): get rid of this - we should be able to extend constraint types without having to add to an enum
+	enum class EConstraintContainerType : uint8
+	{
+		Invalid = 0,
+		Collision,
+		RigidSpring,
+		DynamicSpring,
+		Position,
+		Joint,
+		Suspension
+	};
+
 	/**
 	 * Base class for containers of constraints.
 	 * A Constraint Container holds an array of constraints and provides methods to allocate and deallocate constraints
@@ -14,17 +26,80 @@ namespace Chaos
 	class CHAOS_API FPBDConstraintContainer
 	{
 	public:
-		FPBDConstraintContainer();
+		FPBDConstraintContainer(EConstraintContainerType InType);
 
 		virtual ~FPBDConstraintContainer();
 
+		// The ContainerId is used by the Evolution to map constraints graph items back to constraints
+		int32 GetContainerId() const
+		{
+			return ContainerId;
+		}
+
+		void SetContainerId(int32 InContainerId)
+		{
+			ContainerId = InContainerId;
+		}
+
+		EConstraintContainerType GetConstraintType() const
+		{
+			return ConstraintContainerType;
+		}
+
 		virtual void SetConstraintEnabled(int32 ConstraintIndex, bool bEnabled) { }
 		virtual bool IsConstraintEnabled(int32 ConstraintIndex) const { return true; }
 		virtual void DisconnectConstraints(const TSet<TGeometryParticleHandle<FReal, 3>*>&) {}
 
 	protected:
-		// friend access to the Constraint Handle's container API
 		int32 GetConstraintIndex(const FConstraintHandle* ConstraintHandle) const;
 		void SetConstraintIndex(FConstraintHandle* ConstraintHandle, int32 ConstraintIndex) const;
+
+	private:
+		EConstraintContainerType ConstraintContainerType;
+		int32 ContainerId;
 	};
+
+
+	//
+	// From ConstraintHandle.h
+	//
+
+	inline int32 FConstraintHandle::GetContainerId() const
+	{
+		if (ConstraintContainer != nullptr)
+		{
+			return ConstraintContainer->GetContainerId();
+		}
+		return INDEX_NONE;
+	}
+
+	inline void FConstraintHandle::SetEnabled(bool bInEnabled)
+	{
+		if (ConstraintContainer != nullptr)
+		{
+			ConstraintContainer->SetConstraintEnabled(ConstraintIndex, bInEnabled);
+		}
+	}
+
+	inline bool FConstraintHandle::IsEnabled() const
+	{
+		if (ConstraintContainer != nullptr)
+		{
+			return ConstraintContainer->IsConstraintEnabled(ConstraintIndex);
+		}
+		return false;
+	}
+
+	template<typename T>
+	inline T* FConstraintHandle::As()
+	{
+		return ((ConstraintContainer != nullptr) && (T::StaticType() == ConstraintContainer->GetConstraintType())) ? static_cast<T*>(this) : nullptr;
+	}
+
+	template<typename T>
+	inline const T* FConstraintHandle::As() const
+	{
+		return ((ConstraintContainer != nullptr) && (T::StaticType() == ConstraintContainer->GetConstraintType())) ? static_cast<const T*>(this) : nullptr;
+	}
+
 }
\ No newline at end of file
diff --git a/Engine/Source/Runtime/Experimental/Chaos/Public/Chaos/PBDConstraintRule.h b/Engine/Source/Runtime/Experimental/Chaos/Public/Chaos/PBDConstraintRule.h
--- a/Engine/Source/Runtime/Experimental/Chaos/Public/Chaos/PBDConstraintRule.h
+++ b/Engine/Source/Runtime/Experimental/Chaos/Public/Chaos/PBDConstraintRule.h
@@ -223,10 +223,14 @@ namespace Chaos
 
 		virtual int32 NumConstraints() const override;
 
+		int32 GetContainerId() const
+		{
+			return Constraints.GetContainerId();
+		}
+
 	protected:
 		FConstraints& Constraints;
 		FPBDConstraintGraph* ConstraintGraph;
-		uint32 ContainerId;
 	};
 
 	/**
@@ -242,6 +246,8 @@ namespace Chaos
 		using FConstraintContainerHandle = typename FConstraints::FConstraintContainerHandle;
 		using FConstraintList = TArray<FConstraintContainerHandle*>;
 
+		using Base::GetContainerId;
+
 		TPBDConstraintIslandRule(FConstraints& InConstraints, int32 InPriority = 0)
 			: TPBDConstraintGraphRuleImpl<T_CONSTRAINTS>(InConstraints, InPriority)
 		{
@@ -283,7 +289,7 @@ namespace Chaos
 			for (int32 ConstraintDataIndex : ConstraintDataIndices)
 			{
 				const typename FPBDConstraintGraph::FConstraintData& ConstraintData = ConstraintGraph->GetConstraintData(ConstraintDataIndex);
-				if (ConstraintData.GetContainerId() == ContainerId)
+				if (ConstraintData.GetContainerId() == GetContainerId())
 				{
 					IslandConstraintList.Add(ConstraintData.GetConstraintHandle()->As<FConstraintContainerHandle>());
 				}
@@ -299,7 +305,6 @@ namespace Chaos
 	private:
 		using Base::Constraints;
 		using Base::ConstraintGraph;
-		using Base::ContainerId;
 
 		const TArray<FConstraintContainerHandle*>& GetIslandConstraints(int32 Island) const
 		{
@@ -433,7 +438,7 @@ namespace Chaos
 
 		virtual void UpdateAccelerationStructures(const int32 Island) override
 		{
-			GraphColor.ComputeColor(Island, *ConstraintGraph, ContainerId);
+			GraphColor.ComputeColor(Island, *ConstraintGraph, Constraints.GetContainerId());
 		}
 
 		virtual void SetUseContactGraph(const bool bInUseContactGraph) override
@@ -468,7 +473,6 @@ namespace Chaos
 	private:
 		using Base::Constraints;
 		using Base::ConstraintGraph;
-		using Base::ContainerId;
 
 		const TArray<typename FConstraints::FConstraintContainerHandle*>& GetLevelColorConstraints(const typename FPBDConstraintColor::FLevelToColorToConstraintListMap& LevelToColorToConstraintListMap, int32 Level, int32 Color) const
 		{
diff --git a/Engine/Source/Runtime/Experimental/Chaos/Public/Chaos/PBDJointConstraints.h b/Engine/Source/Runtime/Experimental/Chaos/Public/Chaos/PBDJointConstraints.h
--- a/Engine/Source/Runtime/Experimental/Chaos/Public/Chaos/PBDJointConstraints.h
+++ b/Engine/Source/Runtime/Experimental/Chaos/Public/Chaos/PBDJointConstraints.h
@@ -30,7 +30,7 @@ namespace Chaos
 
 		FPBDJointConstraintHandle();
 		FPBDJointConstraintHandle(FConstraintContainer* InConstraintContainer, int32 InConstraintIndex);
-		static FConstraintHandle::EType StaticType() { return FConstraintHandle::EType::Joint; }
+		static EConstraintContainerType StaticType() { return EConstraintContainerType::Joint; }
 
 		void SetConstraintEnabled(bool bEnabled);
 
@@ -52,7 +52,7 @@ namespace Chaos
 
 	protected:
 		using Base::ConstraintIndex;
-		using Base::ConstraintContainer;
+		using Base::ConcreteContainer;
 	};
 
 	class CHAOS_API FPBDJointState
diff --git a/Engine/Source/Runtime/Experimental/Chaos/Public/Chaos/PBDPositionConstraints.h b/Engine/Source/Runtime/Experimental/Chaos/Public/Chaos/PBDPositionConstraints.h
--- a/Engine/Source/Runtime/Experimental/Chaos/Public/Chaos/PBDPositionConstraints.h
+++ b/Engine/Source/Runtime/Experimental/Chaos/Public/Chaos/PBDPositionConstraints.h
@@ -19,13 +19,13 @@ namespace Chaos
 
 		FPBDPositionConstraintHandle() {}
 		FPBDPositionConstraintHandle(FConstraintContainer* InConstraintContainer, int32 InConstraintIndex) 
-			: TContainerConstraintHandle<FPBDPositionConstraints>(StaticType(), InConstraintContainer, InConstraintIndex) {}
-		static FConstraintHandle::EType StaticType() { return FConstraintHandle::EType::Position; }
+			: TContainerConstraintHandle<FPBDPositionConstraints>(InConstraintContainer, InConstraintIndex) {}
+		static EConstraintContainerType StaticType() { return EConstraintContainerType::Position; }
 		TVector<FGeometryParticleHandle*, 2> GetConstrainedParticles() const;
 
 	protected:
 		using Base::ConstraintIndex;
-		using Base::ConstraintContainer;
+		using Base::ConcreteContainer;
 	};
 
 	class FPBDPositionConstraints : public FPBDConstraintContainer
@@ -39,11 +39,13 @@ namespace Chaos
 		using FHandles = TArray<FConstraintContainerHandle*>;
 
 		FPBDPositionConstraints(const FReal InStiffness = (FReal)1.)
-			: Stiffness(InStiffness)
+			: FPBDConstraintContainer(EConstraintContainerType::Position)
+			, Stiffness(InStiffness)
 		{}
 
 		FPBDPositionConstraints(TArray<FVec3>&& Locations, TArray<FPBDRigidParticleHandle*>&& InConstrainedParticles, const FReal InStiffness = (FReal)1.)
-			: Targets(MoveTemp(Locations)), ConstrainedParticles(MoveTemp(InConstrainedParticles)), Stiffness(InStiffness)
+			: FPBDConstraintContainer(EConstraintContainerType::Position)
+			, Targets(MoveTemp(Locations)), ConstrainedParticles(MoveTemp(InConstrainedParticles)), Stiffness(InStiffness)
 		{
 			if (ConstrainedParticles.Num() > 0)
 			{
diff --git a/Engine/Source/Runtime/Experimental/Chaos/Public/Chaos/PBDRigidDynamicSpringConstraints.h b/Engine/Source/Runtime/Experimental/Chaos/Public/Chaos/PBDRigidDynamicSpringConstraints.h
--- a/Engine/Source/Runtime/Experimental/Chaos/Public/Chaos/PBDRigidDynamicSpringConstraints.h
+++ b/Engine/Source/Runtime/Experimental/Chaos/Public/Chaos/PBDRigidDynamicSpringConstraints.h
@@ -21,16 +21,16 @@ namespace Chaos
 		}
 
 		FPBDRigidDynamicSpringConstraintHandle(FConstraintContainer* InConstraintContainer, int32 InConstraintIndex) 
-			: TContainerConstraintHandle<FPBDRigidDynamicSpringConstraints>(StaticType(),InConstraintContainer, InConstraintIndex) 
+			: TContainerConstraintHandle<FPBDRigidDynamicSpringConstraints>(InConstraintContainer, InConstraintIndex) 
 		{
 		}
 
-		static FConstraintHandle::EType StaticType() { return FConstraintHandle::EType::DynamicSpring; }
+		static EConstraintContainerType StaticType() { return EConstraintContainerType::DynamicSpring; }
 		TVec2<FGeometryParticleHandle*> GetConstrainedParticles() const;
 
 	protected:
 		using Base::ConstraintIndex;
-		using Base::ConstraintContainer;
+		using Base::ConcreteContainer;
 	};
 
 	class CHAOS_API FPBDRigidDynamicSpringConstraints : public FPBDConstraintContainer
@@ -44,11 +44,18 @@ namespace Chaos
 		using FHandles = TArray<FConstraintContainerHandle*>;
 
 		FPBDRigidDynamicSpringConstraints(const FReal InStiffness = (FReal)1.)
-			: CreationThreshold(1), MaxSprings(1), Stiffness(InStiffness) 
+			: FPBDConstraintContainer(EConstraintContainerType::DynamicSpring)
+			, CreationThreshold(1)
+			, MaxSprings(1)
+			, Stiffness(InStiffness) 
 		{}
 
 		FPBDRigidDynamicSpringConstraints(TArray<FConstrainedParticlePair>&& InConstraints, const FReal InCreationThreshold = (FReal)1., const int32 InMaxSprings = 1, const FReal InStiffness = (FReal)1.)
-			: Constraints(MoveTemp(InConstraints)), CreationThreshold(InCreationThreshold), MaxSprings(InMaxSprings), Stiffness(InStiffness)
+			: FPBDConstraintContainer(EConstraintContainerType::DynamicSpring)
+			, Constraints(MoveTemp(InConstraints))
+			, CreationThreshold(InCreationThreshold)
+			, MaxSprings(InMaxSprings)
+			, Stiffness(InStiffness)
 		{
 			if (Constraints.Num() > 0)
 			{
diff --git a/Engine/Source/Runtime/Experimental/Chaos/Public/Chaos/PBDRigidSpringConstraints.h b/Engine/Source/Runtime/Experimental/Chaos/Public/Chaos/PBDRigidSpringConstraints.h
--- a/Engine/Source/Runtime/Experimental/Chaos/Public/Chaos/PBDRigidSpringConstraints.h
+++ b/Engine/Source/Runtime/Experimental/Chaos/Public/Chaos/PBDRigidSpringConstraints.h
@@ -22,11 +22,11 @@ namespace Chaos
 		}
 		
 		FPBDRigidSpringConstraintHandle(FConstraintContainer* InConstraintContainer, int32 InConstraintIndex) 
-		: TContainerConstraintHandle<FPBDRigidSpringConstraints>(StaticType(), InConstraintContainer, InConstraintIndex) 
+		: TContainerConstraintHandle<FPBDRigidSpringConstraints>(InConstraintContainer, InConstraintIndex) 
 		{
 		}
 
-		static FConstraintHandle::EType StaticType() { return FConstraintHandle::EType::RigidSpring; }
+		static EConstraintContainerType StaticType() { return EConstraintContainerType::RigidSpring; }
 
 		const TVector<FVec3, 2>& GetConstraintPositions() const;
 		void SetConstraintPositions(const TVector<FVec3, 2>& ConstraintPositions);
diff --git a/Engine/Source/Runtime/Experimental/Chaos/Public/Chaos/PBDRigidsEvolution.h b/Engine/Source/Runtime/Experimental/Chaos/Public/Chaos/PBDRigidsEvolution.h
--- a/Engine/Source/Runtime/Experimental/Chaos/Public/Chaos/PBDRigidsEvolution.h
+++ b/Engine/Source/Runtime/Experimental/Chaos/Public/Chaos/PBDRigidsEvolution.h
@@ -30,6 +30,10 @@ extern CHAOS_API int32 ForceNoCollisionIntoSQ;
 
 namespace Chaos
 {
+namespace Collisions
+{
+	void CHAOS_API ResetChaosCollisionCounters();
+}
 
 extern CHAOS_API int32 FixBadAccelerationStructureRemoval;
 
@@ -467,29 +471,32 @@ public:
 	}
 
 	/** remove a list of constraints from the constraint graph (see AddConstraintsToConstraintGraph) */
-	CHAOS_API void RemoveConstraintsFromConstraintGraph(const TArray<FConstraintHandle*>& Constraints)
+	CHAOS_API void RemoveConstraintsFromConstraintGraph(const FConstraintHandleArray& Constraints)
 	{
 		for (FConstraintHandle* BaseConstraintHandle : Constraints)
 		{
 			if (FPBDJointConstraintHandle* ConstraintHandle = BaseConstraintHandle->As<FPBDJointConstraintHandle>())
 			{
 				// if it is already disabled then it will have already been removed from the graph
-				if (ConstraintHandle->IsConstraintEnabled())
+				if (ConstraintHandle->IsInConstraintGraph())
 				{
-					ConstraintGraph.RemoveConstraint(ConstraintHandle->GetConstraintIndex(), ConstraintHandle, ConstraintHandle->GetConstrainedParticles());
+					ConstraintGraph.RemoveConstraint(ConstraintHandle->GetContainerId(), ConstraintHandle, ConstraintHandle->GetConstrainedParticles());
 				}
 			}
 		}
 	}
 
 	/** Add a list of constraints to the constraint graph (see RemoveConstraintsFromConstraintGraph) */
-	CHAOS_API void AddConstraintsToConstraintGraph(const TArray<FConstraintHandle*>& Constraints)
+	CHAOS_API void AddConstraintsToConstraintGraph(const FConstraintHandleArray& Constraints)
 	{
 		for (FConstraintHandle* BaseConstraintHandle : Constraints)
 		{
 			if (FPBDJointConstraintHandle* ConstraintHandle = BaseConstraintHandle->As<FPBDJointConstraintHandle>())
 			{
-				ConstraintGraph.AddConstraint(ConstraintHandle->GetConstraintIndex(), ConstraintHandle, ConstraintHandle->GetConstrainedParticles());
+				if (!ensure(ConstraintHandle->IsInConstraintGraph()))
+				{
+					ConstraintGraph.AddConstraint(ConstraintHandle->GetContainerId(), ConstraintHandle, ConstraintHandle->GetConstrainedParticles());
+				}
 			}
 		}
 	}
diff --git a/Engine/Source/Runtime/Experimental/Chaos/Public/Chaos/PBDSuspensionConstraints.h b/Engine/Source/Runtime/Experimental/Chaos/Public/Chaos/PBDSuspensionConstraints.h
--- a/Engine/Source/Runtime/Experimental/Chaos/Public/Chaos/PBDSuspensionConstraints.h
+++ b/Engine/Source/Runtime/Experimental/Chaos/Public/Chaos/PBDSuspensionConstraints.h
@@ -24,7 +24,7 @@ namespace Chaos
 
 		FPBDSuspensionConstraintHandle() {}
 		FPBDSuspensionConstraintHandle(FConstraintContainer* InConstraintContainer, int32 InConstraintIndex);
-		static FConstraintHandle::EType StaticType() { return FConstraintHandle::EType::Suspension; }
+		static EConstraintContainerType StaticType() { return EConstraintContainerType::Suspension; }
 
 		FPBDSuspensionSettings& GetSettings();
 		const FPBDSuspensionSettings& GetSettings() const;
@@ -35,7 +35,7 @@ namespace Chaos
 
 	protected:
 		using Base::ConstraintIndex;
-		using Base::ConstraintContainer;
+		using Base::ConcreteContainer;
 	};
 
 	class CHAOS_API FPBDSuspensionConstraints : public FPBDConstraintContainer
@@ -47,11 +47,13 @@ namespace Chaos
 		using FHandles = TArray<FConstraintContainerHandle*>;
 
 		FPBDSuspensionConstraints(const FPBDSuspensionSolverSettings& InSolverSettings = FPBDSuspensionSolverSettings())
-			: SolverSettings(InSolverSettings)
+			: FPBDConstraintContainer(EConstraintContainerType::Suspension)
+			, SolverSettings(InSolverSettings)
 		{}
 
 		FPBDSuspensionConstraints(TArray<FVec3>&& Locations, TArray<TGeometryParticleHandle<FReal,3>*>&& InConstrainedParticles, TArray<FVec3>&& InLocalOffset, TArray<FPBDSuspensionSettings>&& InConstraintSettings)
-			: ConstrainedParticles(MoveTemp(InConstrainedParticles)), SuspensionLocalOffset(MoveTemp(InLocalOffset)), ConstraintSettings(MoveTemp(InConstraintSettings))
+			: FPBDConstraintContainer(EConstraintContainerType::Suspension)
+			, ConstrainedParticles(MoveTemp(InConstrainedParticles)), SuspensionLocalOffset(MoveTemp(InLocalOffset)), ConstraintSettings(MoveTemp(InConstraintSettings))
 		{
 			if (ConstrainedParticles.Num() > 0)
 			{
diff --git a/Engine/Source/Runtime/Experimental/Chaos/Public/Chaos/ParticleHandle.h b/Engine/Source/Runtime/Experimental/Chaos/Public/Chaos/ParticleHandle.h
--- a/Engine/Source/Runtime/Experimental/Chaos/Public/Chaos/ParticleHandle.h
+++ b/Engine/Source/Runtime/Experimental/Chaos/Public/Chaos/ParticleHandle.h
@@ -546,7 +546,7 @@ public:
 		return GeometryParticles->WeakParticleHandle(ParticleIdx);
 	}
 
-	TArray<FConstraintHandle*>& ParticleConstraints()
+	FConstraintHandleArray& ParticleConstraints()
 	{
 		return GeometryParticles->ParticleConstraints(ParticleIdx);
 	}
@@ -1475,7 +1475,7 @@ public:
 
 	int32 Island() const
 	{
-		if (MHandle->CastToRigidParticle() && MHandle->ObjectState() == EObjectStateType::Dynamic)
+		if ((MHandle->CastToRigidParticle() != nullptr) && IsDynamic())
 		{
 			return MHandle->CastToRigidParticle()->Island();
 		}
