GeometryCollectionSceneProxy
-- check for existence of PerBoneHitProxy before accessing it

diff --git a/Engine/Source/Runtime/Experimental/GeometryCollectionEngine/Private/GeometryCollection/GeometryCollectionSceneProxy.cpp b/Engine/Source/Runtime/Experimental/GeometryCollectionEngine/Private/GeometryCollection/GeometryCollectionSceneProxy.cpp
--- a/Engine/Source/Runtime/Experimental/GeometryCollectionEngine/Private/GeometryCollection/GeometryCollectionSceneProxy.cpp
+++ b/Engine/Source/Runtime/Experimental/GeometryCollectionEngine/Private/GeometryCollection/GeometryCollectionSceneProxy.cpp
@@ -400,7 +400,10 @@ void FGeometryCollectionSceneProxy::SetConstantData_RenderThread(FGeometryCollec
 			{
 				// One proxy per bone
 				const int32 ProxyIndex = ConstantData->BoneMap[i];
-				HitProxyIdBuffer.VertexColor(i) = PerBoneHitProxies[ProxyIndex]->Id.GetColor();
+				if (PerBoneHitProxies.Num())
+				{ 
+					HitProxyIdBuffer.VertexColor(i) = PerBoneHitProxies[ProxyIndex]->Id.GetColor();
+				}
 			}
 #endif
 		});
