Chaos : Colliison Group ID
- All non dynamic rigids to use collision group ids.

diff --git a/Engine/Source/Runtime/Experimental/Chaos/Public/Chaos/ParticleHandle.h b/Engine/Source/Runtime/Experimental/Chaos/Public/Chaos/ParticleHandle.h
--- a/Engine/Source/Runtime/Experimental/Chaos/Public/Chaos/ParticleHandle.h
+++ b/Engine/Source/Runtime/Experimental/Chaos/Public/Chaos/ParticleHandle.h
@@ -1222,7 +1222,7 @@ public:
 
 	int32 CollisionGroup() const
 	{
-		if (MHandle->CastToRigidParticle() && MHandle->ObjectState() == EObjectStateType::Dynamic)
+		if (MHandle->CastToRigidParticle())
 		{
 			return MHandle->CastToRigidParticle()->CollisionGroup();
 		}
