Chaos : Plasticity
- Add plasticity controls for expansion and contraction.

diff --git a/Engine/Source/Editor/DetailCustomizations/Private/PhysicsConstraintComponentDetails.cpp b/Engine/Source/Editor/DetailCustomizations/Private/PhysicsConstraintComponentDetails.cpp
--- a/Engine/Source/Editor/DetailCustomizations/Private/PhysicsConstraintComponentDetails.cpp
+++ b/Engine/Source/Editor/DetailCustomizations/Private/PhysicsConstraintComponentDetails.cpp
@@ -214,7 +214,6 @@ void FPhysicsConstraintComponentDetails::AddLinearLimits(IDetailLayoutBuilder& D
 	IDetailCategoryBuilder& LinearLimitCat = DetailBuilder.EditCategory("Linear Limits");
 	TSharedPtr<IPropertyHandle> LinearConstraintProperty = ProfilePropertiesProperty->GetChildHandle(GET_MEMBER_NAME_CHECKED(FConstraintProfileProperties, LinearLimit));
 
-
 	TSharedPtr<IPropertyHandle> LinearXMotionProperty = LinearConstraintProperty->GetChildHandle(GET_MEMBER_NAME_CHECKED(FLinearConstraint, XMotion));
 	TSharedPtr<IPropertyHandle> LinearYMotionProperty = LinearConstraintProperty->GetChildHandle(GET_MEMBER_NAME_CHECKED(FLinearConstraint, YMotion));
 	TSharedPtr<IPropertyHandle> LinearZMotionProperty = LinearConstraintProperty->GetChildHandle(GET_MEMBER_NAME_CHECKED(FLinearConstraint, ZMotion));
@@ -335,6 +334,7 @@ void FPhysicsConstraintComponentDetails::AddLinearLimits(IDetailLayoutBuilder& D
 	LinearLimitCat.AddProperty(ProfilePropertiesProperty->GetChildHandle(GET_MEMBER_NAME_CHECKED(FConstraintProfileProperties, LinearBreakThreshold)));
 
 	LinearLimitCat.AddProperty(ProfilePropertiesProperty->GetChildHandle(GET_MEMBER_NAME_CHECKED(FConstraintProfileProperties, bLinearPlasticity)).ToSharedRef());
+	LinearLimitCat.AddProperty(ProfilePropertiesProperty->GetChildHandle(GET_MEMBER_NAME_CHECKED(FConstraintProfileProperties, LinearPlasticityType)).ToSharedRef());
 	LinearLimitCat.AddProperty(ProfilePropertiesProperty->GetChildHandle(GET_MEMBER_NAME_CHECKED(FConstraintProfileProperties, LinearPlasticityThreshold)).ToSharedRef());
 }
 
diff --git a/Engine/Source/Runtime/Engine/Classes/PhysicsEngine/ConstraintInstance.h b/Engine/Source/Runtime/Engine/Classes/PhysicsEngine/ConstraintInstance.h
--- a/Engine/Source/Runtime/Engine/Classes/PhysicsEngine/ConstraintInstance.h
+++ b/Engine/Source/Runtime/Engine/Classes/PhysicsEngine/ConstraintInstance.h
@@ -121,10 +121,14 @@ struct ENGINE_API FConstraintProfileProperties
 	UPROPERTY(EditAnywhere, AdvancedDisplay, Category = Linear)
 	uint8 bLinearPlasticity : 1;
 
+	/** Whether linear plasticity has a operation mode [free]*/
+	UPROPERTY(EditAnywhere, AdvancedDisplay, Category = Linear)
+	TEnumAsByte<enum EConstraintPlasticityType> LinearPlasticityType;
+
 	FConstraintProfileProperties();
 
 	/** Updates physx joint properties from unreal properties (limits, drives, flags, etc...) */
-	void Update_AssumesLocked(const FPhysicsConstraintHandle& InConstraintRef, float AverageMass, float UseScale) const;
+	void Update_AssumesLocked(const FPhysicsConstraintHandle& InConstraintRef, float AverageMass, float UseScale, bool InInitialize = false) const;
 
 	/** Updates joint breakable properties (threshold, etc...)*/
 	void UpdateBreakable_AssumesLocked(const FPhysicsConstraintHandle& InConstraintRef) const;
@@ -544,10 +548,11 @@ public:
 	*	@param bInLinearPlasticity 	Whether it is possible to reset the target angles
 	*	@param InLinearPlasticityThreshold	Delta from target needed to reset the target joint
 	*/
-	void SetLinearPlasticity(bool bInLinearPlasticity, float InLinearPlasticityThreshold)
+	void SetLinearPlasticity(bool bInLinearPlasticity, float InLinearPlasticityThreshold, EConstraintPlasticityType InLinearPlasticityType)
 	{
 		ProfileInstance.bLinearPlasticity = bInLinearPlasticity;
 		ProfileInstance.LinearPlasticityThreshold = InLinearPlasticityThreshold;
+		ProfileInstance.LinearPlasticityType = InLinearPlasticityType;
 		UpdatePlasticity();
 	}
 
@@ -563,6 +568,12 @@ public:
 		return ProfileInstance.LinearPlasticityThreshold;
 	}
 
+	/** Gets Plasticity Type from joint */
+	TEnumAsByte<enum EConstraintPlasticityType> GetLinearPlasticityType() const
+	{
+		return ProfileInstance.LinearPlasticityType;
+	}
+
 	/** Sets the Angular Breakable properties
 	*	@param bInAngularBreakable		Whether it is possible to break the joint with angular force
 	*	@param InAngularBreakThreshold	Torque needed to break the joint
diff --git a/Engine/Source/Runtime/Engine/Classes/PhysicsEngine/ConstraintInstanceBlueprintLibrary.h b/Engine/Source/Runtime/Engine/Classes/PhysicsEngine/ConstraintInstanceBlueprintLibrary.h
--- a/Engine/Source/Runtime/Engine/Classes/PhysicsEngine/ConstraintInstanceBlueprintLibrary.h
+++ b/Engine/Source/Runtime/Engine/Classes/PhysicsEngine/ConstraintInstanceBlueprintLibrary.h
@@ -188,7 +188,8 @@ public:
 		static void SetLinearPlasticity(
 			UPARAM(ref) FConstraintInstanceAccessor& Accessor,
 			bool bLinearPlasticity,
-			float LinearPlasticityThreshold
+			float LinearPlasticityThreshold,
+			TEnumAsByte<EConstraintPlasticityType> PlasticityType
 		);
 
 	/** Sets Constraint Linear Plasticity properties
@@ -200,7 +201,8 @@ public:
 		static void GetLinearPlasticity(
 			UPARAM(ref) FConstraintInstanceAccessor& Accessor,
 			bool& bLinearPlasticity,
-			float& LinearPlasticityThreshold
+			float& LinearPlasticityThreshold,
+			TEnumAsByte<EConstraintPlasticityType>& PlasticityType
 		);
 
 	//---------------------------------------------------------------------------------------------------
diff --git a/Engine/Source/Runtime/Engine/Classes/PhysicsEngine/PhysicsConstraintComponent.h b/Engine/Source/Runtime/Engine/Classes/PhysicsEngine/PhysicsConstraintComponent.h
--- a/Engine/Source/Runtime/Engine/Classes/PhysicsEngine/PhysicsConstraintComponent.h
+++ b/Engine/Source/Runtime/Engine/Classes/PhysicsEngine/PhysicsConstraintComponent.h
@@ -278,7 +278,7 @@ public:
 	*	@param LinearPlasticityThreshold	Percent deformation needed to reset the rest length of the joint
 	*/
 	UFUNCTION(BlueprintCallable, Category = "Physics|Components|PhysicsConstraint")
-	void SetLinearPlasticity(bool bLinearPlasticity, float LinearPlasticityThreshold);
+	void SetLinearPlasticity(bool bLinearPlasticity, float LinearPlasticityThreshold, EConstraintPlasticityType PlasticityType);
 
 	/** Sets the Angular Breakable properties
 	*	@param bAngularBreakable		Whether it is possible to break the joint with angular force
diff --git a/Engine/Source/Runtime/Engine/Private/PhysicsEngine/ConstraintInstance.cpp b/Engine/Source/Runtime/Engine/Private/PhysicsEngine/ConstraintInstance.cpp
--- a/Engine/Source/Runtime/Engine/Private/PhysicsEngine/ConstraintInstance.cpp
+++ b/Engine/Source/Runtime/Engine/Private/PhysicsEngine/ConstraintInstance.cpp
@@ -131,7 +131,7 @@ FConstraintProfileProperties::FConstraintProfileProperties()
 	, ProjectionLinearAlpha(1.0f)
 	, ProjectionAngularAlpha(0.0f)
 	, LinearBreakThreshold(300.f)
-	, LinearPlasticityThreshold(10.f)
+	, LinearPlasticityThreshold(0.1f)
 	, AngularBreakThreshold(500.f)
 	, AngularPlasticityThreshold(10.f)
 	, bDisableCollision(false)
@@ -142,6 +142,7 @@ FConstraintProfileProperties::FConstraintProfileProperties()
 	, bAngularPlasticity(false)
 	, bLinearBreakable(false)
 	, bLinearPlasticity(false)
+	, LinearPlasticityType(EConstraintPlasticityType::CCPT_Free)
 {
 }
 
@@ -192,7 +193,7 @@ void FConstraintProfileProperties::UpdatePlasticity_AssumesLocked(const FPhysics
 	const float LinearPlasticityLimit = bLinearPlasticity ? LinearPlasticityThreshold : FLT_MAX;
 	const float AngularPlasticityLimit = bAngularPlasticity ? FMath::DegreesToRadians(AngularPlasticityThreshold) : MAX_FLT;
 
-	FPhysicsInterface::SetPlasticityLimits_AssumesLocked(InConstraintRef, LinearPlasticityLimit, AngularPlasticityLimit);
+	FPhysicsInterface::SetPlasticityLimits_AssumesLocked(InConstraintRef, LinearPlasticityLimit, AngularPlasticityLimit, LinearPlasticityType);
 #endif
 }
 
@@ -502,7 +503,7 @@ void FConstraintInstance::InitConstraint_AssumesLocked(const FPhysicsActorHandle
 	// update mass
 	UpdateAverageMass_AssumesLocked(ActorRef1, ActorRef2);
 
-	ProfileInstance.Update_AssumesLocked(ConstraintHandle, AverageMass, bScaleLinearLimits ? LastKnownScale : 1.f);
+	ProfileInstance.Update_AssumesLocked(ConstraintHandle, AverageMass, bScaleLinearLimits ? LastKnownScale : 1.f, true);
 
 	// Put the bodies back to sleep both bodies were asleep
 	if (bActor1WasAsleep && bActor2WasAsleep)
@@ -519,7 +520,7 @@ void FConstraintInstance::InitConstraint_AssumesLocked(const FPhysicsActorHandle
 	}
 }
 
-void FConstraintProfileProperties::Update_AssumesLocked(const FPhysicsConstraintHandle& InConstraintRef, float AverageMass, float UseScale) const
+void FConstraintProfileProperties::Update_AssumesLocked(const FPhysicsConstraintHandle& InConstraintRef, float AverageMass, float UseScale, bool InInitialize) const
 {
 	// flags and projection settings
 	UpdateConstraintFlags_AssumesLocked(InConstraintRef);
@@ -532,12 +533,8 @@ void FConstraintProfileProperties::Update_AssumesLocked(const FPhysicsConstraint
 	UpdateBreakable_AssumesLocked(InConstraintRef);
 	UpdatePlasticity_AssumesLocked(InConstraintRef);
 
-	// Motors
-	FPhysicsInterface::UpdateLinearDrive_AssumesLocked(InConstraintRef, LinearDrive);
-	FPhysicsInterface::UpdateAngularDrive_AssumesLocked(InConstraintRef, AngularDrive);
-
 	// Target
-	FPhysicsInterface::UpdateDriveTarget_AssumesLocked(InConstraintRef, LinearDrive, AngularDrive);
+	FPhysicsInterface::UpdateDriveTarget_AssumesLocked(InConstraintRef, LinearDrive, AngularDrive, InInitialize);
 }
 
 void FConstraintInstance::TermConstraint()
diff --git a/Engine/Source/Runtime/Engine/Private/PhysicsEngine/ConstraintInstanceBlueprintLibrary.cpp b/Engine/Source/Runtime/Engine/Private/PhysicsEngine/ConstraintInstanceBlueprintLibrary.cpp
--- a/Engine/Source/Runtime/Engine/Private/PhysicsEngine/ConstraintInstanceBlueprintLibrary.cpp
+++ b/Engine/Source/Runtime/Engine/Private/PhysicsEngine/ConstraintInstanceBlueprintLibrary.cpp
@@ -208,30 +208,34 @@ void UConstraintInstanceBlueprintLibrary::GetLinearBreakable(
 void UConstraintInstanceBlueprintLibrary::SetLinearPlasticity(
 	FConstraintInstanceAccessor& Accessor,
 	bool bLinearPlasticity,
-	float LinearPlasticityThreshold
+	float LinearPlasticityThreshold,
+	TEnumAsByte<EConstraintPlasticityType> PlasticityType
 	)
 {
 	if (FConstraintInstance* ConstraintInstance = Accessor.Get())
 	{
-		ConstraintInstance->SetLinearPlasticity(bLinearPlasticity, LinearPlasticityThreshold);
+		ConstraintInstance->SetLinearPlasticity(bLinearPlasticity, LinearPlasticityThreshold,PlasticityType);
 	}
 }
 
 void UConstraintInstanceBlueprintLibrary::GetLinearPlasticity(
 	FConstraintInstanceAccessor& Accessor,
 	bool& bLinearPlasticity,
-	float& LinearPlasticityThreshold
+	float& LinearPlasticityThreshold,
+	TEnumAsByte<EConstraintPlasticityType>& PlasticityType
 	)
 {
 	if (FConstraintInstance* ConstraintInstance = Accessor.Get())
 	{
 		bLinearPlasticity = ConstraintInstance->HasLinearPlasticity();
 		LinearPlasticityThreshold = ConstraintInstance->GetLinearPlasticityThreshold();
+		PlasticityType = ConstraintInstance->GetLinearPlasticityType();
 	}
 	else
 	{
 		bLinearPlasticity = false;
 		LinearPlasticityThreshold = 0.0f;
+		PlasticityType = EConstraintPlasticityType::CCPT_Free;
 	}
 }
 
diff --git a/Engine/Source/Runtime/Engine/Private/PhysicsEngine/Experimental/PhysInterface_Chaos.cpp b/Engine/Source/Runtime/Engine/Private/PhysicsEngine/Experimental/PhysInterface_Chaos.cpp
--- a/Engine/Source/Runtime/Engine/Private/PhysicsEngine/Experimental/PhysInterface_Chaos.cpp
+++ b/Engine/Source/Runtime/Engine/Private/PhysicsEngine/Experimental/PhysInterface_Chaos.cpp
@@ -371,7 +371,7 @@ void FPhysInterface_Chaos::UpdateTwistLimitParams_AssumesLocked(const FPhysicsCo
 	}
 }
 
-void FPhysInterface_Chaos::UpdateLinearDrive_AssumesLocked(const FPhysicsConstraintHandle& InConstraintRef, const FLinearDriveConstraint& InDriveParams)
+void FPhysInterface_Chaos::UpdateLinearDrive_AssumesLocked(const FPhysicsConstraintHandle& InConstraintRef, const FLinearDriveConstraint& InDriveParams, bool InInitialize)
 {
 	if (InConstraintRef.IsValid() && InConstraintRef.Constraint->IsType(Chaos::EConstraintType::JointConstraintType))
 	{
@@ -391,7 +391,7 @@ void FPhysInterface_Chaos::UpdateLinearDrive_AssumesLocked(const FPhysicsConstra
 				Constraint->SetLinearPositionDriveXEnabled(InDriveParams.XDrive.bEnablePositionDrive);
 				Constraint->SetLinearPositionDriveYEnabled(InDriveParams.YDrive.bEnablePositionDrive);
 				Constraint->SetLinearPositionDriveZEnabled(InDriveParams.ZDrive.bEnablePositionDrive);
-				if (FMath::IsNearlyEqual(Constraint->GetLinearPlasticityLimit(), (Chaos::FReal)FLT_MAX))
+				if (InInitialize || FMath::IsNearlyEqual(Constraint->GetLinearPlasticityLimit(), (Chaos::FReal)FLT_MAX))
 				{
 					Constraint->SetLinearDrivePositionTarget(InDriveParams.PositionTarget);
 				}
@@ -413,7 +413,7 @@ void FPhysInterface_Chaos::UpdateLinearDrive_AssumesLocked(const FPhysicsConstra
 	}
 }
 
-void FPhysInterface_Chaos::UpdateAngularDrive_AssumesLocked(const FPhysicsConstraintHandle& InConstraintRef, const FAngularDriveConstraint& InDriveParams)
+void FPhysInterface_Chaos::UpdateAngularDrive_AssumesLocked(const FPhysicsConstraintHandle& InConstraintRef, const FAngularDriveConstraint& InDriveParams, bool InInitialize)
 {
 	if (InConstraintRef.IsValid() && InConstraintRef.Constraint->IsType(Chaos::EConstraintType::JointConstraintType))
 	{
@@ -440,7 +440,7 @@ void FPhysInterface_Chaos::UpdateAngularDrive_AssumesLocked(const FPhysicsConstr
 					Constraint->SetAngularSLerpPositionDriveEnabled(InDriveParams.SlerpDrive.bEnablePositionDrive);
 				}
 
-				if (FMath::IsNearlyEqual(Constraint->GetAngularPlasticityLimit(), (Chaos::FReal)FLT_MAX))
+				if (InInitialize || FMath::IsNearlyEqual(Constraint->GetAngularPlasticityLimit(), (Chaos::FReal)FLT_MAX))
 				{
 					// Plastic joints should not be re-targeted after initialization. 
 					Constraint->SetAngularDrivePositionTarget(Chaos::FRotation3(InDriveParams.OrientationTarget.Quaternion()));
@@ -481,12 +481,12 @@ void FPhysInterface_Chaos::UpdateAngularDrive_AssumesLocked(const FPhysicsConstr
 	}
 }
 
-void FPhysInterface_Chaos::UpdateDriveTarget_AssumesLocked(const FPhysicsConstraintHandle& InConstraintRef, const FLinearDriveConstraint& InLinDrive, const FAngularDriveConstraint& InAngDrive)
+void FPhysInterface_Chaos::UpdateDriveTarget_AssumesLocked(const FPhysicsConstraintHandle& InConstraintRef, const FLinearDriveConstraint& InLinDrive, const FAngularDriveConstraint& InAngDrive, bool InInitialize)
 {
 	if (InConstraintRef.IsValid())
 	{
-		UpdateLinearDrive_AssumesLocked(InConstraintRef, InLinDrive);
-		UpdateAngularDrive_AssumesLocked(InConstraintRef, InAngDrive);
+		UpdateLinearDrive_AssumesLocked(InConstraintRef, InLinDrive, InInitialize);
+		UpdateAngularDrive_AssumesLocked(InConstraintRef, InAngDrive, InInitialize);
 	}
 }
 
diff --git a/Engine/Source/Runtime/Engine/Private/PhysicsEngine/PhysicsConstraintComponent.cpp b/Engine/Source/Runtime/Engine/Private/PhysicsEngine/PhysicsConstraintComponent.cpp
--- a/Engine/Source/Runtime/Engine/Private/PhysicsEngine/PhysicsConstraintComponent.cpp
+++ b/Engine/Source/Runtime/Engine/Private/PhysicsEngine/PhysicsConstraintComponent.cpp
@@ -670,9 +670,9 @@ void UPhysicsConstraintComponent::SetLinearBreakable(bool bLinearBreakable, floa
 	ConstraintInstance.SetLinearBreakable(bLinearBreakable, LinearBreakThreshold);
 }
 
-void UPhysicsConstraintComponent::SetLinearPlasticity(bool bLinearPlasticity, float LinearPlasticityThreshold)
+void UPhysicsConstraintComponent::SetLinearPlasticity(bool bLinearPlasticity, float LinearPlasticityThreshold, EConstraintPlasticityType PlasticityType)
 {
-	ConstraintInstance.SetLinearPlasticity(bLinearPlasticity, LinearPlasticityThreshold);
+	ConstraintInstance.SetLinearPlasticity(bLinearPlasticity, LinearPlasticityThreshold, PlasticityType);
 }
 
 void UPhysicsConstraintComponent::SetAngularBreakable(bool bAngularBreakable, float AngularBreakThreshold)
diff --git a/Engine/Source/Runtime/Engine/Public/Physics/Experimental/PhysInterface_Chaos.h b/Engine/Source/Runtime/Engine/Public/Physics/Experimental/PhysInterface_Chaos.h
--- a/Engine/Source/Runtime/Engine/Public/Physics/Experimental/PhysInterface_Chaos.h
+++ b/Engine/Source/Runtime/Engine/Public/Physics/Experimental/PhysInterface_Chaos.h
@@ -57,9 +57,9 @@ public:
 	static void UpdateLinearLimitParams_AssumesLocked(const FPhysicsConstraintHandle& InConstraintRef, float InLimit, float InAverageMass, const FLinearConstraint& InParams);
 	static void UpdateConeLimitParams_AssumesLocked(const FPhysicsConstraintHandle& InConstraintRef, float InAverageMass, const FConeConstraint& InParams);
 	static void UpdateTwistLimitParams_AssumesLocked(const FPhysicsConstraintHandle& InConstraintRef, float InAverageMass, const FTwistConstraint& InParams);
-	static void UpdateLinearDrive_AssumesLocked(const FPhysicsConstraintHandle& InConstraintRef, const FLinearDriveConstraint& InDriveParams);
-	static void UpdateAngularDrive_AssumesLocked(const FPhysicsConstraintHandle& InConstraintRef, const FAngularDriveConstraint& InDriveParams);
-	static void UpdateDriveTarget_AssumesLocked(const FPhysicsConstraintHandle& InConstraintRef, const FLinearDriveConstraint& InLinDrive, const FAngularDriveConstraint& InAngDrive);
+	static void UpdateLinearDrive_AssumesLocked(const FPhysicsConstraintHandle& InConstraintRef, const FLinearDriveConstraint& InDriveParams, bool InInitialize = false);
+	static void UpdateAngularDrive_AssumesLocked(const FPhysicsConstraintHandle& InConstraintRef, const FAngularDriveConstraint& InDriveParams, bool InInitialize = false);
+	static void UpdateDriveTarget_AssumesLocked(const FPhysicsConstraintHandle& InConstraintRef, const FLinearDriveConstraint& InLinDrive, const FAngularDriveConstraint& InAngDrive, bool InInitialize = false);
 	
 
 	static bool ExecuteOnUnbrokenConstraintReadOnly(const FPhysicsConstraintHandle& InConstraintRef, TFunctionRef<void(const FPhysicsConstraintHandle&)> Func);
diff --git a/Engine/Source/Runtime/Experimental/Chaos/Private/Chaos/PBDJointConstraints.cpp b/Engine/Source/Runtime/Experimental/Chaos/Private/Chaos/PBDJointConstraints.cpp
--- a/Engine/Source/Runtime/Experimental/Chaos/Private/Chaos/PBDJointConstraints.cpp
+++ b/Engine/Source/Runtime/Experimental/Chaos/Private/Chaos/PBDJointConstraints.cpp
@@ -169,7 +169,8 @@ namespace Chaos
 		, AngularDriveDamping(0)
 		, LinearBreakForce(FLT_MAX)
 		, LinearPlasticityLimit(FLT_MAX)
-		, InitCOMDistance(-1)
+		, LinearPlasticityType(EPlasticityType::Free)
+		, LinearPlasticityInitialDistanceSquared(FLT_MAX)
 		, AngularBreakTorque(FLT_MAX)
 		, AngularPlasticityLimit(FLT_MAX)
 		, UserData(nullptr)
@@ -730,14 +731,22 @@ namespace Chaos
 
 			// Plasticity should not be turned on in the middle of simulation.
 			const bool bUseLinearPlasticity = JointSettings.LinearPlasticityLimit != FLT_MAX;
-			if(bUseLinearPlasticity)
+			if (bUseLinearPlasticity)
 			{
-				const bool bIsCOMDistanceInitialized = JointSettings.InitCOMDistance >= 0;
-				if(!bIsCOMDistanceInitialized)
+				const bool bIsCOMDistanceInitialized = !FMath::IsNearlyEqual(JointSettings.LinearPlasticityInitialDistanceSquared, (FReal)FLT_MAX);
+				if (!bIsCOMDistanceInitialized)
 				{
-					ConstraintSettings[JointIndex].InitCOMDistance = (FParticleUtilitiesXR::GetCoMWorldPosition(Particle0) - FParticleUtilitiesXR::GetCoMWorldPosition(Particle1)).Size();
+					// Joint plasticity is baed on the distance of one of the moment arms of the joint. Typically, plasticity
+					// will get setup from the joint pivot to the child COM (centor of mass), so that is found first. However, when 
+					// the pivot is at the child COM then we fall back to the distance between thge pivot and parent COM.
+					ConstraintSettings[JointIndex].LinearPlasticityInitialDistanceSquared = JointSettings.ConnectorTransforms[1].GetTranslation().SizeSquared();
+					if (FMath::IsNearlyZero(ConstraintSettings[JointIndex].LinearPlasticityInitialDistanceSquared))
+					{
+						ConstraintSettings[JointIndex].LinearPlasticityInitialDistanceSquared = JointSettings.ConnectorTransforms[0].GetTranslation().SizeSquared();
+					}
+					ensureMsgf(!FMath::IsNearlyZero(ConstraintSettings[JointIndex].LinearPlasticityInitialDistanceSquared), TEXT("Plasticity made inactive due to Zero length difference between parent and child rigid body."));
 				}
-			} 
+			}
 		}
 	}
 
@@ -1239,26 +1248,44 @@ namespace Chaos
 		Q1.EnforceShortestArcWith(ConstraintFramesGlobal[0].GetRotation());
 		ConstraintFramesGlobal[1].SetRotation(Q1);
 
-		if(bHasLinearPlasticityLimit)
+		if (bHasLinearPlasticityLimit)
 		{
 			FVec3 LinearDisplacement = ConstraintFramesGlobal[0].InverseTransformPositionNoScale(ConstraintFramesGlobal[1].GetTranslation());
 
 			// @todo(chaos): still need to warn against the case where all position drives are not enabled or all dimensions are locked. Warning should print out the joint names and should only print out once to avoid spamming.
-			for(int32 Axis = 0; Axis < 3; Axis++)
+			for (int32 Axis = 0; Axis < 3; Axis++)
 			{
-				if(!JointSettings.bLinearPositionDriveEnabled[Axis] || JointSettings.LinearMotionTypes[Axis] == EJointMotionType::Locked)
+				if (!JointSettings.bLinearPositionDriveEnabled[Axis] || JointSettings.LinearMotionTypes[Axis] == EJointMotionType::Locked)
 				{
 					LinearDisplacement[Axis] = 0;
 				}
 			}
 			// Assuming that the dimensions which are locked or have no targets are 0. in LinearDrivePositionTarget
-			FReal LinearPlasticityDistanceThreshold = JointSettings.LinearPlasticityLimit * JointSettings.InitCOMDistance;
-			if((LinearDisplacement - JointSettings.LinearDrivePositionTarget).SizeSquared() > LinearPlasticityDistanceThreshold * LinearPlasticityDistanceThreshold)
+			FReal LinearPlasticityDistanceThreshold = JointSettings.LinearPlasticityLimit * JointSettings.LinearPlasticityLimit * JointSettings.LinearPlasticityInitialDistanceSquared;
+			if ((LinearDisplacement - JointSettings.LinearDrivePositionTarget).SizeSquared() > LinearPlasticityDistanceThreshold)
 			{
-				JointSettings.LinearDrivePositionTarget = LinearDisplacement;
+				if (JointSettings.LinearPlasticityType == EPlasticityType::Free)
+				{
+					JointSettings.LinearDrivePositionTarget = LinearDisplacement;
+				}
+				else // EPlasticityType::Shrink || EPlasticityType::Grow
+				{
+					// Shrink and Grow are based on the distance between the joint pivot and the child. 
+					// Note, if the pivot is located at the COM of the child then shrink will not do anything. 
+					FVector CurrentDelta = ConstraintFramesLocal[Index0].TransformPosition(LinearDisplacement);
+					FVector StartDelta = ConstraintFramesLocal[Index0].TransformPosition(JointSettings.LinearDrivePositionTarget);
+
+					if (JointSettings.LinearPlasticityType == EPlasticityType::Shrink && CurrentDelta.SizeSquared() < StartDelta.SizeSquared())
+					{
+						JointSettings.LinearDrivePositionTarget = LinearDisplacement;
+					}
+					else if (JointSettings.LinearPlasticityType == EPlasticityType::Grow && CurrentDelta.SizeSquared() > StartDelta.SizeSquared())
+					{
+						JointSettings.LinearDrivePositionTarget = LinearDisplacement;
+					}
+				}
 			}
 		}
-
 		if(bHasAngularPlasticityLimit)
 		{
 			FRotation3 Swing, Twist; FPBDJointUtilities::DecomposeSwingTwistLocal(ConstraintFramesGlobal[0].GetRotation(), ConstraintFramesGlobal[1].GetRotation(), Swing, Twist);
diff --git a/Engine/Source/Runtime/Experimental/Chaos/Private/PhysicsProxy/JointConstraintProxy.cpp b/Engine/Source/Runtime/Experimental/Chaos/Private/PhysicsProxy/JointConstraintProxy.cpp
--- a/Engine/Source/Runtime/Experimental/Chaos/Private/PhysicsProxy/JointConstraintProxy.cpp
+++ b/Engine/Source/Runtime/Experimental/Chaos/Private/PhysicsProxy/JointConstraintProxy.cpp
@@ -153,6 +153,7 @@ void FJointConstraintPhysicsProxy::PushStateOnGameThread(Chaos::FPBDRigidsSolver
 			if (Constraint->IsDirty(Chaos::EJointConstraintFlags::LinearBreakForce))
 			{
 				JointSettingsBuffer.LinearBreakForce = Constraint->GetLinearBreakForce();
+				JointSettingsBuffer.LinearPlasticityType = Constraint->GetLinearPlasticityType();
 				JointSettingsBuffer.LinearPlasticityLimit = Constraint->GetLinearPlasticityLimit();
 				DirtyFlagsBuffer.MarkDirty(Chaos::EJointConstraintFlags::LinearBreakForce);
 			}
@@ -318,6 +319,7 @@ void FJointConstraintPhysicsProxy::PushStateOnPhysicsThread(Chaos::FPBDRigidsSol
 			if (DirtyFlagsBuffer.IsDirty(Chaos::EJointConstraintFlags::LinearBreakForce))
 			{
 				ConstraintSettings.LinearBreakForce = JointSettingsBuffer.LinearBreakForce;
+				ConstraintSettings.LinearPlasticityType = JointSettingsBuffer.LinearPlasticityType;
 				ConstraintSettings.LinearPlasticityLimit = FMath::Clamp((float)JointSettingsBuffer.LinearPlasticityLimit, 0.f, FLT_MAX);
 			}
 
diff --git a/Engine/Source/Runtime/Experimental/Chaos/Public/Chaos/PBDJointConstraintData.h b/Engine/Source/Runtime/Experimental/Chaos/Public/Chaos/PBDJointConstraintData.h
--- a/Engine/Source/Runtime/Experimental/Chaos/Public/Chaos/PBDJointConstraintData.h
+++ b/Engine/Source/Runtime/Experimental/Chaos/Public/Chaos/PBDJointConstraintData.h
@@ -85,6 +85,10 @@ namespace Chaos
 		//void SetLinearPlasticityLimit(FReal InLinearPlasticityLimit);
 		//FReal GetLinearPlasticityLimit() const
 
+		CONSTRAINT_JOINT_PROPERPETY_IMPL(EPlasticityType, LinearPlasticityType, EJointConstraintFlags::LinearBreakForce, JointSettings.LinearPlasticityType);
+		//void SetLinearPlasticityType(EPlasticityType InLinearPlasticityType);
+		//EPlasticityType GetLinearPlasticityType() const
+
 		CONSTRAINT_JOINT_PROPERPETY_IMPL(FReal, AngularBreakTorque, EJointConstraintFlags::AngularBreakTorque, JointSettings.AngularBreakTorque);
 		//void SetAngularBreakTorque(FReal InAngularBreakTorque);
 		//FReal GetAngularBreakTorque() const
diff --git a/Engine/Source/Runtime/Experimental/Chaos/Public/Chaos/PBDJointConstraintTypes.h b/Engine/Source/Runtime/Experimental/Chaos/Public/Chaos/PBDJointConstraintTypes.h
--- a/Engine/Source/Runtime/Experimental/Chaos/Public/Chaos/PBDJointConstraintTypes.h
+++ b/Engine/Source/Runtime/Experimental/Chaos/Public/Chaos/PBDJointConstraintTypes.h
@@ -32,6 +32,14 @@ namespace Chaos
 		Force,
 	};
 
+	enum class EPlasticityType : int32
+	{
+		Free,
+		Shrink,
+		Grow,
+	};
+
+
 	/**
 	 * The order of the angular constraints (for settings held in vectors etc)
 	 */
@@ -173,7 +181,8 @@ namespace Chaos
 
 		FReal LinearBreakForce;
 		FReal LinearPlasticityLimit;
-		FReal InitCOMDistance;
+		EPlasticityType LinearPlasticityType;
+		FReal LinearPlasticityInitialDistanceSquared;
 		FReal AngularBreakTorque;
 		FReal AngularPlasticityLimit;
 
diff --git a/Engine/Source/Runtime/PhysicsCore/Private/ChaosEngineInterface.cpp b/Engine/Source/Runtime/PhysicsCore/Private/ChaosEngineInterface.cpp
--- a/Engine/Source/Runtime/PhysicsCore/Private/ChaosEngineInterface.cpp
+++ b/Engine/Source/Runtime/PhysicsCore/Private/ChaosEngineInterface.cpp
@@ -1301,12 +1301,13 @@ void FChaosEngineInterface::SetBreakForces_AssumesLocked(const FPhysicsConstrain
 	}
 }
 
-void FChaosEngineInterface::SetPlasticityLimits_AssumesLocked(const FPhysicsConstraintHandle& InConstraintRef, float InLinearPlasticityLimit, float InAngularPlasticityLimit)
+void FChaosEngineInterface::SetPlasticityLimits_AssumesLocked(const FPhysicsConstraintHandle& InConstraintRef, float InLinearPlasticityLimit, float InAngularPlasticityLimit, EConstraintPlasticityType InLinearPlasticityType)
 {
 	if (InConstraintRef.IsValid() && InConstraintRef.Constraint->IsType(Chaos::EConstraintType::JointConstraintType))
 	{
 		if (Chaos::FJointConstraint* Constraint = static_cast<Chaos::FJointConstraint*>(InConstraintRef.Constraint))
 		{
+			Constraint->SetLinearPlasticityType((Chaos::EPlasticityType)InLinearPlasticityType);
 			Constraint->SetLinearPlasticityLimit(InLinearPlasticityLimit);
 			Constraint->SetAngularPlasticityLimit(InAngularPlasticityLimit);
 		}
diff --git a/Engine/Source/Runtime/PhysicsCore/Public/Chaos/ChaosEngineInterface.h b/Engine/Source/Runtime/PhysicsCore/Public/Chaos/ChaosEngineInterface.h
--- a/Engine/Source/Runtime/PhysicsCore/Public/Chaos/ChaosEngineInterface.h
+++ b/Engine/Source/Runtime/PhysicsCore/Public/Chaos/ChaosEngineInterface.h
@@ -164,6 +164,21 @@ namespace PhysicsInterfaceTypes
 	typedef TArray<FPhysicsShapeHandle, TInlineAllocator<NumInlinedPxShapeElements>> FInlineShapeArray;
 }
 
+
+// LINEAR CCPT
+UENUM()
+enum EConstraintPlasticityType
+{
+	/** */
+	CCPT_Free	UMETA(DisplayName = "Free"),
+	/** */
+	CCPT_Shrink	UMETA(DisplayName = "Shirnk"),
+	/** */
+	CCPT_Grow	UMETA(DisplayName = "Grow"),
+
+	CCPT_MAX,
+};
+
 // LINEAR DOF
 UENUM()
 enum ELinearConstraintMotion
@@ -483,7 +498,7 @@ public:
 	static void SetProjectionEnabled_AssumesLocked(const FPhysicsConstraintHandle& InConstraintRef,bool bInProjectionEnabled,float InLinearAlpah = 1.0f,float InAngularAlpha = 0.0f);
 	static void SetParentDominates_AssumesLocked(const FPhysicsConstraintHandle& InConstraintRef,bool bInParentDominates);
 	static void SetBreakForces_AssumesLocked(const FPhysicsConstraintHandle& InConstraintRef,float InLinearBreakForce,float InAngularBreakForce);
-	static void SetPlasticityLimits_AssumesLocked(const FPhysicsConstraintHandle& InConstraintRef, float InLinearPlasticityLimit, float InAngularPlasticityLimit);
+	static void SetPlasticityLimits_AssumesLocked(const FPhysicsConstraintHandle& InConstraintRef, float InLinearPlasticityLimit, float InAngularPlasticityLimit, EConstraintPlasticityType InLinearPlasticityType);
 	static void SetLocalPose(const FPhysicsConstraintHandle& InConstraintRef, const FTransform& InPose, EConstraintFrame::Type InFrame);
 
 	static void SetDrivePosition(const FPhysicsConstraintHandle& InConstraintRef,const FVector& InPosition);
