Chaos : GeometryCollection
- Rigid body mass alignment for leaf geometry.

diff --git a/Engine/Source/Runtime/Experimental/Chaos/Private/GeometryCollection/GeometryCollectionCollisionStructureManager.cpp b/Engine/Source/Runtime/Experimental/Chaos/Private/GeometryCollection/GeometryCollectionCollisionStructureManager.cpp
--- a/Engine/Source/Runtime/Experimental/Chaos/Private/GeometryCollection/GeometryCollectionCollisionStructureManager.cpp
+++ b/Engine/Source/Runtime/Experimental/Chaos/Private/GeometryCollection/GeometryCollectionCollisionStructureManager.cpp
@@ -250,14 +250,22 @@ FCollisionStructureManager::FImplicit*
 FCollisionStructureManager::NewImplicitConvex(
 	const TArray<int32>& ConvexIndices,
 	const TManagedArray<TUniquePtr<Chaos::FConvex>>* ConvexGeometry,
-	const ECollisionTypeEnum CollisionType)
+	const ECollisionTypeEnum CollisionType,
+	const FTransform& MassTransform)
 {
 	if (ConvexIndices.Num())
 	{
 		TArray<TUniquePtr<Chaos::FImplicitObject>> Implicits;
 		for (auto& Index : ConvexIndices)
 		{
-			Chaos::FImplicitObject* Implicit = new Chaos::FConvex((*ConvexGeometry)[Index]->GetVertices(), (*ConvexGeometry)[Index]->GetMargin());
+			TArray<Chaos::FVec3> ConvexVertices = (*ConvexGeometry)[Index]->GetVertices();
+			for (int32 Idx = 0; Idx < ConvexVertices.Num(); Idx++)
+			{
+				ConvexVertices[Idx] = MassTransform.InverseTransformPosition(ConvexVertices[Idx]);
+			}
+
+			Chaos::FReal Margin = (Chaos::FReal)(*ConvexGeometry)[Index]->BoundingBox().Extents().Min() * (Chaos::FReal)0.1;
+			Chaos::FImplicitObject* Implicit = new Chaos::FConvex(ConvexVertices, Margin);
 			UpdateImplicitFlags(Implicit, CollisionType);
 			Implicits.Add(TUniquePtr<Chaos::FImplicitObject>(Implicit));
 		}
diff --git a/Engine/Source/Runtime/Experimental/Chaos/Private/PhysicsProxy/GeometryCollectionPhysicsProxy.cpp b/Engine/Source/Runtime/Experimental/Chaos/Private/PhysicsProxy/GeometryCollectionPhysicsProxy.cpp
--- a/Engine/Source/Runtime/Experimental/Chaos/Private/PhysicsProxy/GeometryCollectionPhysicsProxy.cpp
+++ b/Engine/Source/Runtime/Experimental/Chaos/Private/PhysicsProxy/GeometryCollectionPhysicsProxy.cpp
@@ -156,7 +156,6 @@ void PopulateSimulatedParticle(
 	const int16 CollisionGroup)
 {
 	SCOPE_CYCLE_COUNTER(STAT_PopulateSimulatedParticle);
-
 	Handle->SetDisabledLowLevel(false);
 	Handle->SetX(WorldTransform.GetTranslation());
 	Handle->SetV(Chaos::FVec3(0.f));
@@ -688,7 +687,7 @@ void FGeometryCollectionPhysicsProxy::InitializeBodiesPT(Chaos::FPBDRigidsSolver
 
 			FFieldSystemMetaDataProcessingResolution* ResolutionData = new FFieldSystemMetaDataProcessingResolution(EFieldResolutionType::Field_Resolution_Maximum);
 
-			Cmd.MetaData.Add(FFieldSystemMetaData::EMetaType::ECommandData_ProcessingResolution, TUniquePtr<FFieldSystemMetaDataProcessingResolution>(ResolutionData));
+			Cmd.MetaData.Add( FFieldSystemMetaData::EMetaType::ECommandData_ProcessingResolution, TUniquePtr<FFieldSystemMetaDataProcessingResolution>(ResolutionData));
 			Commands.Add(Cmd);
 		}
 		Parameters.InitializationCommands.Empty();
@@ -2166,7 +2165,10 @@ void FGeometryCollectionPhysicsProxy::InitializeSharedCollisionStructures(
 									FCollisionStructureManager::NewImplicitConvex(
 										(*TransformToConvexIndices)[TransformGroupIndex].Array(),
 										ConvexGeometry,
-										SizeSpecificData.CollisionShapesData[0].CollisionType));
+										SizeSpecificData.CollisionShapesData[0].CollisionType,
+										CollectionMassToLocal[TransformGroupIndex]
+									)
+							);
 						}
 					}
 					else if (SizeSpecificData.CollisionShapesData[0].ImplicitType == EImplicitTypeEnum::Chaos_Implicit_Capsule)
@@ -2445,7 +2447,10 @@ void FGeometryCollectionPhysicsProxy::InitializeSharedCollisionStructures(
 								FCollisionStructureManager::NewImplicitConvex(
 									(*TransformToConvexIndices)[TransformGroupIndex].Array(),
 									ConvexGeometry,
-									SizeSpecificData.CollisionShapesData[0].CollisionType));
+									SizeSpecificData.CollisionShapesData[0].CollisionType,
+									CollectionMassToLocal[TransformGroupIndex]
+								)
+							);
 						}
 					}
 					else if (SizeSpecificData.CollisionShapesData[0].ImplicitType == EImplicitTypeEnum::Chaos_Implicit_Capsule)
diff --git a/Engine/Source/Runtime/Experimental/Chaos/Public/GeometryCollection/GeometryCollectionCollisionStructureManager.h b/Engine/Source/Runtime/Experimental/Chaos/Public/GeometryCollection/GeometryCollectionCollisionStructureManager.h
--- a/Engine/Source/Runtime/Experimental/Chaos/Public/GeometryCollection/GeometryCollectionCollisionStructureManager.h
+++ b/Engine/Source/Runtime/Experimental/Chaos/Public/GeometryCollection/GeometryCollectionCollisionStructureManager.h
@@ -107,7 +107,8 @@ public:
 	static FImplicit* NewImplicitConvex(
 		const TArray<int32>& ConvexIndices,
 		const TManagedArray<TUniquePtr<Chaos::FConvex>>* ConvexGeometry,
-		const ECollisionTypeEnum CollisionType);
+		const ECollisionTypeEnum CollisionType,
+		const FTransform& MassTransform);
 
 	static FImplicit* NewImplicitLevelset(
 		Chaos::FErrorReporter ErrorReporter,
