Chaos : Destruction
- Re-enable rotational inertial alignment for geometry collection bodies.
- Remove default rotational initialization for clusters as they created in the correct orientation.
- Bumped the DDC key for chaos.

diff --git a/Engine/Source/Runtime/Experimental/Chaos/Private/Chaos/PBDRigidClustering.cpp b/Engine/Source/Runtime/Experimental/Chaos/Private/Chaos/PBDRigidClustering.cpp
--- a/Engine/Source/Runtime/Experimental/Chaos/Private/Chaos/PBDRigidClustering.cpp
+++ b/Engine/Source/Runtime/Experimental/Chaos/Private/Chaos/PBDRigidClustering.cpp
@@ -388,7 +388,6 @@ namespace Chaos
 		}
 		Parent->W() /= Parent->M();
 		Parent->PreW() = Parent->W();
-		Parent->R() = Chaos::TransformToLocalSpace(InertiaTensor);
 		if (ForceMassOrientation)
 		{
 			Parent->R() = ForceMassOrientation->GetRotation();
diff --git a/Engine/Source/Runtime/Experimental/Chaos/Private/PhysicsProxy/GeometryCollectionPhysicsProxy.cpp b/Engine/Source/Runtime/Experimental/Chaos/Private/PhysicsProxy/GeometryCollectionPhysicsProxy.cpp
--- a/Engine/Source/Runtime/Experimental/Chaos/Private/PhysicsProxy/GeometryCollectionPhysicsProxy.cpp
+++ b/Engine/Source/Runtime/Experimental/Chaos/Private/PhysicsProxy/GeometryCollectionPhysicsProxy.cpp
@@ -2048,7 +2048,6 @@ void FGeometryCollectionPhysicsProxy::InitializeSharedCollisionStructures(
 				// Note: particles already in CoM space, so passing in zero as CoM
 				CalculateInertiaAndRotationOfMass(MassSpaceParticles, TriMesh->GetSurfaceElements(), Density_i, FVec3(0), MassProperties.InertiaTensor, MassProperties.RotationOfMass);
 				CollectionInertiaTensor[TransformGroupIndex] = FVec3(MassProperties.InertiaTensor.M[0][0], MassProperties.InertiaTensor.M[1][1], MassProperties.InertiaTensor.M[2][2]);
-#if false
 				CollectionMassToLocal[TransformGroupIndex] = FTransform(MassProperties.RotationOfMass, MassProperties.CenterOfMass);
 
 
@@ -2062,7 +2061,6 @@ void FGeometryCollectionPhysicsProxy::InitializeSharedCollisionStructures(
 						MassSpaceParticles.X(Idx) = InverseMassRotation.TransformPosition(MassSpaceParticles.X(Idx));
 					}
 				}
-#endif
 			}
 			else
 			{
diff --git a/Engine/Source/Runtime/Experimental/ChaosCore/Public/Chaos/Core.h b/Engine/Source/Runtime/Experimental/ChaosCore/Public/Chaos/Core.h
--- a/Engine/Source/Runtime/Experimental/ChaosCore/Public/Chaos/Core.h
+++ b/Engine/Source/Runtime/Experimental/ChaosCore/Public/Chaos/Core.h
@@ -35,5 +35,5 @@ namespace Chaos
 	// NOTE: if you get a merge conflict on the GUID, you must replace it with a new GUID - do not accept the source or target
 	// or you will likely get DDC version conflicts resulting in crashes during load.
 	// Core version string for Chaos data. Any DDC builder dependent on Chaos for serialization should depend on this version
-	static const FString ChaosVersionString = TEXT("CC808BF8-0958-49F6-B745-09E044451D61");
+	static const FString ChaosVersionString = TEXT("6955386D-0525-4192-8927-74B0E57E8926");
 }
