Chaos - Geometry Collection
- Added CollisionMarginFraction to the size specific data.

diff --git a/Engine/Source/Runtime/Experimental/Chaos/Private/GeometryCollection/GeometryCollectionCollisionStructureManager.cpp b/Engine/Source/Runtime/Experimental/Chaos/Private/GeometryCollection/GeometryCollectionCollisionStructureManager.cpp
--- a/Engine/Source/Runtime/Experimental/Chaos/Private/GeometryCollection/GeometryCollectionCollisionStructureManager.cpp
+++ b/Engine/Source/Runtime/Experimental/Chaos/Private/GeometryCollection/GeometryCollectionCollisionStructureManager.cpp
@@ -251,7 +251,8 @@ FCollisionStructureManager::NewImplicitConvex(
 	const TArray<int32>& ConvexIndices,
 	const TManagedArray<TUniquePtr<Chaos::FConvex>>* ConvexGeometry,
 	const ECollisionTypeEnum CollisionType,
-	const FTransform& MassTransform)
+	const FTransform& MassTransform,
+	const Chaos::FReal CollisionMarginFraction)
 {
 	if (ConvexIndices.Num())
 	{
@@ -264,7 +265,7 @@ FCollisionStructureManager::NewImplicitConvex(
 				ConvexVertices[Idx] = MassTransform.InverseTransformPosition(ConvexVertices[Idx]);
 			}
 
-			Chaos::FReal Margin = (Chaos::FReal)(*ConvexGeometry)[Index]->BoundingBox().Extents().Min() * (Chaos::FReal)0.1;
+			Chaos::FReal Margin = (Chaos::FReal)(*ConvexGeometry)[Index]->BoundingBox().Extents().Min() * CollisionMarginFraction;
 			Chaos::FImplicitObject* Implicit = new Chaos::FConvex(ConvexVertices, Margin);
 			UpdateImplicitFlags(Implicit, CollisionType);
 			Implicits.Add(TUniquePtr<Chaos::FImplicitObject>(Implicit));
diff --git a/Engine/Source/Runtime/Experimental/Chaos/Private/PhysicsProxy/GeometryCollectionPhysicsProxy.cpp b/Engine/Source/Runtime/Experimental/Chaos/Private/PhysicsProxy/GeometryCollectionPhysicsProxy.cpp
--- a/Engine/Source/Runtime/Experimental/Chaos/Private/PhysicsProxy/GeometryCollectionPhysicsProxy.cpp
+++ b/Engine/Source/Runtime/Experimental/Chaos/Private/PhysicsProxy/GeometryCollectionPhysicsProxy.cpp
@@ -2164,7 +2164,8 @@ void FGeometryCollectionPhysicsProxy::InitializeSharedCollisionStructures(
 										(*TransformToConvexIndices)[TransformGroupIndex].Array(),
 										ConvexGeometry,
 										SizeSpecificData.CollisionShapesData[0].CollisionType,
-										CollectionMassToLocal[TransformGroupIndex]
+										CollectionMassToLocal[TransformGroupIndex],
+										(Chaos::FReal)SizeSpecificData.CollisionShapesData[0].CollisionMarginFraction
 									)
 							);
 						}
@@ -2446,8 +2447,8 @@ void FGeometryCollectionPhysicsProxy::InitializeSharedCollisionStructures(
 									(*TransformToConvexIndices)[TransformGroupIndex].Array(),
 									ConvexGeometry,
 									SizeSpecificData.CollisionShapesData[0].CollisionType,
-									CollectionMassToLocal[TransformGroupIndex]
-								)
+									CollectionMassToLocal[TransformGroupIndex],
+									(Chaos::FReal)SizeSpecificData.CollisionShapesData[0].CollisionMarginFraction)
 							);
 						}
 					}
diff --git a/Engine/Source/Runtime/Experimental/Chaos/Public/GeometryCollection/GeometryCollectionCollisionStructureManager.h b/Engine/Source/Runtime/Experimental/Chaos/Public/GeometryCollection/GeometryCollectionCollisionStructureManager.h
--- a/Engine/Source/Runtime/Experimental/Chaos/Public/GeometryCollection/GeometryCollectionCollisionStructureManager.h
+++ b/Engine/Source/Runtime/Experimental/Chaos/Public/GeometryCollection/GeometryCollectionCollisionStructureManager.h
@@ -20,6 +20,8 @@ namespace Chaos
 	using FParticles = TParticles<FReal, 3>;
 }
 
+
+
 class CHAOS_API FCollisionStructureManager
 {
 public:
@@ -108,7 +110,8 @@ public:
 		const TArray<int32>& ConvexIndices,
 		const TManagedArray<TUniquePtr<Chaos::FConvex>>* ConvexGeometry,
 		const ECollisionTypeEnum CollisionType,
-		const FTransform& MassTransform);
+		const FTransform& MassTransform,
+		const Chaos::FReal CollisionMarginFraction);
 
 	static FImplicit* NewImplicitLevelset(
 		Chaos::FErrorReporter ErrorReporter,
diff --git a/Engine/Source/Runtime/Experimental/Chaos/Public/GeometryCollection/GeometryCollectionSimulationCoreTypes.h b/Engine/Source/Runtime/Experimental/Chaos/Public/GeometryCollection/GeometryCollectionSimulationCoreTypes.h
--- a/Engine/Source/Runtime/Experimental/Chaos/Public/GeometryCollection/GeometryCollectionSimulationCoreTypes.h
+++ b/Engine/Source/Runtime/Experimental/Chaos/Public/GeometryCollection/GeometryCollectionSimulationCoreTypes.h
@@ -50,6 +50,7 @@ struct FCollectionCollisionTypeData
 		, LevelSetData()
 		, CollisionParticleData()
 		, CollisionObjectReductionPercentage(0.f)
+		, CollisionMarginFraction(0.1f)
 	{
 	}
 
@@ -58,6 +59,7 @@ struct FCollectionCollisionTypeData
 	FCollectionLevelSetData LevelSetData;
 	FCollectionCollisionParticleData CollisionParticleData;
 	float CollisionObjectReductionPercentage;
+	float CollisionMarginFraction;
 };
 
 struct FSharedSimulationSizeSpecificData
@@ -115,8 +117,9 @@ struct FSharedSimulationParameters
 		, float InMaximumBoundingExtentClamp
 		, float InMinimumInertiaTensorDiagonalClamp
 		, float InMaximumInertiaTensorDiagonalClamp
-		,float InCollisionParticlesFraction
-		,int32 InMaximumCollisionParticleCount)
+		, float InCollisionParticlesFraction
+		, int32 InMaximumCollisionParticleCount
+		, float InCollisionMarginFraction)
 	: bMassAsDensity(InMassAsDensity)
 	, Mass(InMass)
 	, MinimumMassClamp(InMinimumMassClamp)
@@ -132,6 +135,7 @@ struct FSharedSimulationParameters
 		{
 			SizeSpecificData[0].CollisionShapesData[0].CollisionType = InCollisionType;
 			SizeSpecificData[0].CollisionShapesData[0].ImplicitType = InImplicitType;
+			SizeSpecificData[0].CollisionShapesData[0].CollisionMarginFraction = InCollisionMarginFraction;
 			SizeSpecificData[0].CollisionShapesData[0].LevelSetData.MinLevelSetResolution = InMinLevelSetResolution;
 			SizeSpecificData[0].CollisionShapesData[0].LevelSetData.MaxLevelSetResolution = InMaxLevelSetResolution;
 			SizeSpecificData[0].CollisionShapesData[0].LevelSetData.MinClusterLevelSetResolution = InMinClusterLevelSetResolution;
diff --git a/Engine/Source/Runtime/Experimental/GeometryCollectionEngine/Private/GeometryCollection/GeometryCollectionObject.cpp b/Engine/Source/Runtime/Experimental/GeometryCollectionEngine/Private/GeometryCollection/GeometryCollectionObject.cpp
--- a/Engine/Source/Runtime/Experimental/GeometryCollectionEngine/Private/GeometryCollection/GeometryCollectionObject.cpp
+++ b/Engine/Source/Runtime/Experimental/GeometryCollectionEngine/Private/GeometryCollection/GeometryCollectionObject.cpp
@@ -16,6 +16,8 @@
 #include "ProfilingDebugging/CookStats.h"
 #include "EngineUtils.h"
 #include "Engine/StaticMesh.h"
+#include "PhysicsEngine/PhysicsSettings.h"
+
 
 #if WITH_EDITOR
 #include "GeometryCollection/DerivedDataGeometryCollectionCooker.h"
@@ -110,6 +112,7 @@ FGeometryCollectionCollisionTypeData::FGeometryCollectionCollisionTypeData()
 	, LevelSet()
 	, CollisionParticles()
 	, CollisionObjectReductionPercentage(0.0f)
+	, CollisionMarginFraction(0.1)
 {
 }
 
@@ -151,6 +154,7 @@ void FGeometryCollectionSizeSpecificData::PostSerialize(const FArchive& Ar)
 			CollisionShapes[0].CollisionType = CollisionType_DEPRECATED;
 			CollisionShapes[0].ImplicitType = ImplicitType_DEPRECATED;
 			CollisionShapes[0].CollisionObjectReductionPercentage = CollisionObjectReductionPercentage_DEPRECATED;
+			CollisionShapes[0].CollisionMarginFraction = UPhysicsSettings::Get()->SolverOptions.CollisionMarginFraction;
 			CollisionShapes[0].CollisionParticles.CollisionParticlesFraction = CollisionParticlesFraction_DEPRECATED;
 			CollisionShapes[0].CollisionParticles.MaximumCollisionParticles = MaximumCollisionParticles_DEPRECATED;
 			CollisionShapes[0].LevelSet.MinLevelSetResolution = MinLevelSetResolution_DEPRECATED;
@@ -181,6 +185,7 @@ void FillSharedSimulationSizeSpecificData(FSharedSimulationSizeSpecificData& ToD
 			ToData.CollisionShapesData[i].LevelSetData.MinClusterLevelSetResolution = FromData.CollisionShapes[i].LevelSet.MinClusterLevelSetResolution;
 			ToData.CollisionShapesData[i].LevelSetData.MaxClusterLevelSetResolution = FromData.CollisionShapes[i].LevelSet.MaxClusterLevelSetResolution;
 			ToData.CollisionShapesData[i].CollisionObjectReductionPercentage = FromData.CollisionShapes[i].CollisionObjectReductionPercentage;
+			ToData.CollisionShapesData[i].CollisionMarginFraction = FromData.CollisionShapes[i].CollisionMarginFraction;
 			ToData.CollisionShapesData[i].CollisionParticleData.CollisionParticlesFraction = FromData.CollisionShapes[i].CollisionParticles.CollisionParticlesFraction;
 			ToData.CollisionShapesData[i].CollisionParticleData.MaximumCollisionParticles = FromData.CollisionShapes[i].CollisionParticles.MaximumCollisionParticles;
 		}
@@ -202,6 +207,7 @@ FGeometryCollectionSizeSpecificData UGeometryCollection::GeometryCollectionSizeS
 		Data.CollisionShapes[0].LevelSet.MinClusterLevelSetResolution = 25;
 		Data.CollisionShapes[0].LevelSet.MaxClusterLevelSetResolution = 50;
 		Data.CollisionShapes[0].CollisionObjectReductionPercentage = 1.0;
+		Data.CollisionShapes[0].CollisionMarginFraction = UPhysicsSettings::Get()->SolverOptions.CollisionMarginFraction;
 		Data.CollisionShapes[0].CollisionParticles.CollisionParticlesFraction = 1.0;
 		Data.CollisionShapes[0].CollisionParticles.MaximumCollisionParticles = 60;
 	}
@@ -237,6 +243,7 @@ void UGeometryCollection::ValidateSizeSpecificDataDefaults()
 			Data.CollisionShapes[0].LevelSet.MinClusterLevelSetResolution = MinClusterLevelSetResolution_DEPRECATED;
 			Data.CollisionShapes[0].LevelSet.MaxClusterLevelSetResolution = MaxClusterLevelSetResolution_DEPRECATED;
 			Data.CollisionShapes[0].CollisionObjectReductionPercentage = CollisionObjectReductionPercentage_DEPRECATED;
+			Data.CollisionShapes[0].CollisionMarginFraction = UPhysicsSettings::Get()->SolverOptions.CollisionMarginFraction;
 #endif
 			if (Data.CollisionShapes[0].ImplicitType == EImplicitTypeEnum::Chaos_Implicit_LevelSet)
 			{
@@ -288,6 +295,7 @@ void UGeometryCollection::GetSharedSimulationParams(FSharedSimulationParameters&
 		InfSize.CollisionShapes[0].LevelSet.MinClusterLevelSetResolution = SizeSpecificDefault.CollisionShapes[0].LevelSet.MinClusterLevelSetResolution;
 		InfSize.CollisionShapes[0].LevelSet.MaxClusterLevelSetResolution = SizeSpecificDefault.CollisionShapes[0].LevelSet.MaxClusterLevelSetResolution;
 		InfSize.CollisionShapes[0].CollisionObjectReductionPercentage = SizeSpecificDefault.CollisionShapes[0].CollisionObjectReductionPercentage;
+		InfSize.CollisionShapes[0].CollisionMarginFraction = SizeSpecificDefault.CollisionShapes[0].CollisionMarginFraction;
 		InfSize.CollisionShapes[0].CollisionParticles.CollisionParticlesFraction = SizeSpecificDefault.CollisionShapes[0].CollisionParticles.CollisionParticlesFraction;
 		InfSize.CollisionShapes[0].CollisionParticles.MaximumCollisionParticles = SizeSpecificDefault.CollisionShapes[0].CollisionParticles.MaximumCollisionParticles;
 	}
diff --git a/Engine/Source/Runtime/Experimental/GeometryCollectionEngine/Public/GeometryCollection/GeometryCollectionObject.h b/Engine/Source/Runtime/Experimental/GeometryCollectionEngine/Public/GeometryCollection/GeometryCollectionObject.h
--- a/Engine/Source/Runtime/Experimental/GeometryCollectionEngine/Public/GeometryCollection/GeometryCollectionObject.h
+++ b/Engine/Source/Runtime/Experimental/GeometryCollectionEngine/Public/GeometryCollection/GeometryCollectionObject.h
@@ -158,6 +158,14 @@ struct GEOMETRYCOLLECTIONENGINE_API FGeometryCollectionCollisionTypeData
 	UPROPERTY(EditAnywhere, Category = "Collisions")
 	float CollisionObjectReductionPercentage;
 
+	/**
+	* A collision margin is a fraction of size used by some boxes and convex shapes to improve collision detection results.
+	* The core geometry of shapes that support a margin are reduced in size by the margin, and the margin
+	* is added back on during collision detection. The net result is a shape of the same size but with rounded corners.
+	*/
+	UPROPERTY(EditAnywhere, Category = "Collisions", meta = (EditCondition = "ImplicitType == EImplicitTypeEnum::Chaos_Implicit_Convex || ImplicitType == EImplicitTypeEnum::Chaos_Implicit_Box", EditConditionHides))
+	float CollisionMarginFraction;
+
 };
 
 
