Add UPROPERTY markup to prevent garbage collection.

diff --git a/Engine/Source/Runtime/Experimental/InteractiveToolsFramework/Public/BaseGizmos/GizmoArrowObject.h b/Engine/Source/Runtime/Experimental/InteractiveToolsFramework/Public/BaseGizmos/GizmoArrowObject.h
--- a/Engine/Source/Runtime/Experimental/InteractiveToolsFramework/Public/BaseGizmos/GizmoArrowObject.h
+++ b/Engine/Source/Runtime/Experimental/InteractiveToolsFramework/Public/BaseGizmos/GizmoArrowObject.h
@@ -49,7 +49,12 @@ public:
 	UPROPERTY(EditAnywhere, Category = Options)
 	float Offset = 0.0f;
 
+	UPROPERTY()
 	TObjectPtr<UGizmoCylinderObject> CylinderObject;
+
+	UPROPERTY()
 	TObjectPtr<UGizmoConeObject> ConeObject;
+
+	UPROPERTY()
 	TObjectPtr<UGizmoBoxObject> BoxObject;
 };
\ No newline at end of file
diff --git a/Engine/Source/Runtime/Experimental/InteractiveToolsFramework/Public/InteractiveGizmoManager.h b/Engine/Source/Runtime/Experimental/InteractiveToolsFramework/Public/InteractiveGizmoManager.h
--- a/Engine/Source/Runtime/Experimental/InteractiveToolsFramework/Public/InteractiveGizmoManager.h
+++ b/Engine/Source/Runtime/Experimental/InteractiveToolsFramework/Public/InteractiveGizmoManager.h
@@ -21,7 +21,8 @@ struct FActiveGizmo
 {
 	GENERATED_BODY()
 
-	UInteractiveGizmo* Gizmo;
+	UPROPERTY()
+	TObjectPtr<UInteractiveGizmo> Gizmo = nullptr;
 	FString BuilderIdentifier;
 	FString InstanceIdentifier;
 	void* Owner = nullptr;
