Fix crash on fracture convex generation when embedded geometry is a child of cluster or embedded geometry (instead of rigid)

#rnx

diff --git a/Engine/Source/Runtime/Experimental/Chaos/Private/GeometryCollection/GeometryCollectionConvexUtility.cpp b/Engine/Source/Runtime/Experimental/Chaos/Private/GeometryCollection/GeometryCollectionConvexUtility.cpp
--- a/Engine/Source/Runtime/Experimental/Chaos/Private/GeometryCollection/GeometryCollectionConvexUtility.cpp
+++ b/Engine/Source/Runtime/Experimental/Chaos/Private/GeometryCollection/GeometryCollectionConvexUtility.cpp
@@ -175,6 +175,7 @@ void CreateNonoverlappingConvexHulls(
 	TArray<TSet<int32>>& TransformToConvexIndices,
 	const TManagedArray<int32>& SimulationType,
 	int32 LeafType,
+	int32 SkipType,
 	const TManagedArray<int32>& Parents,
 	const TManagedArray<TSet<int32>>* GeoProximity,
 	const TManagedArray<int32>& GeometryToTransformIndex,
@@ -184,9 +185,14 @@ void CreateNonoverlappingConvexHulls(
 	int32 NumBones = TransformToConvexIndices.Num();
 	check(Parents.Num() == NumBones);
 
+	auto SkipBone = [&SimulationType, SkipType](int32 Bone) -> bool
+	{
+		return SimulationType[Bone] == SkipType;
+	};
+
 	auto OnlyConvex = [&TransformToConvexIndices](int32 Bone) -> int32
 	{
-		ensure(TransformToConvexIndices[Bone].Num() == 1);
+		ensure(TransformToConvexIndices[Bone].Num() <= 1);
 		for (int32 ConvexIndex : TransformToConvexIndices[Bone])
 		{
 			return ConvexIndex;
@@ -212,9 +218,9 @@ void CreateNonoverlappingConvexHulls(
 		// TODO: fill out LeafProximity by doing IsColliding calls (may have to construct Children array before this)
 	}
 
-	auto IsColliding2 = [&Convexes](int32 ConvexA, int32 ConvexB)
+	auto IsColliding = [&Convexes](int32 ConvexA, int32 ConvexB)
 	{
-		if (!ensure(Convexes[ConvexA]->NumVertices() > 0 && Convexes[ConvexB]->NumVertices() > 0))
+		if (ConvexA == -1 || ConvexB == -1 || !ensure(Convexes[ConvexA]->NumVertices() > 0 && Convexes[ConvexB]->NumVertices() > 0))
 		{
 			// at least one of the convex hulls was empty, so cannot be colliding
 			// TODO: eventually remove this ensure; and add more handling for empty hulls
@@ -305,13 +311,13 @@ void CreateNonoverlappingConvexHulls(
 		while (ToExpand.Num() > 0)
 		{
 			int32 ToProcess = ToExpand.Pop(false);
-			if (Children[ToProcess].Num() > 0 && SimulationType[ToProcess] != LeafType)
+			if (SimulationType[ToProcess] == LeafType)
 			{
-				ToExpand.Append(Children[ToProcess]);
+				Leaves.Add(ToProcess);
 			}
-			else
+			else if (Children[ToProcess].Num() > 0)
 			{
-				Leaves.Add(ToProcess);
+				ToExpand.Append(Children[ToProcess]);
 			}
 		}
 	};
@@ -379,7 +385,7 @@ void CreateNonoverlappingConvexHulls(
 			{
 				if (BoneB < BoneA)
 				{
-					if (IsColliding2(OnlyConvex(BoneA), OnlyConvex(BoneB)))
+					if (IsColliding(OnlyConvex(BoneA), OnlyConvex(BoneB)))
 					{
 						bool bLeavesCollide = FixLeafCollisions(BoneA, BoneB);
 						int32 Bones[2]{ BoneA, BoneB };
@@ -393,7 +399,7 @@ void CreateNonoverlappingConvexHulls(
 							while (TraverseBones.Num() > 0)
 							{
 								int32 ToProc = TraverseBones.Pop(false);
-								if (IsColliding2(OnlyConvex(OtherBone), OnlyConvex(ToProc)))
+								if (IsColliding(OnlyConvex(OtherBone), OnlyConvex(ToProc)))
 								{
 									ClusterProximity[OtherBone].Add(ToProc);
 									ClusterProximity[ToProc].Add(OtherBone);
@@ -421,7 +427,7 @@ void CreateNonoverlappingConvexHulls(
 		TSet<int32> ToProcess;
 		for (int32 Bone = 0; Bone < NumBones; Bone++)
 		{
-			if (Depths[Bone] == ProcessDepth && Children[Bone].Num() > 0)
+			if (Depths[Bone] == ProcessDepth && Children[Bone].Num() > 0 && !SkipBone(Bone))
 			{
 				TArray<Chaos::FVec3> JoinedHullPts;
 				for (int32 Child : Children[Bone])
@@ -659,7 +665,7 @@ FGeometryCollectionConvexUtility::FGeometryCollectionConvexData FGeometryCollect
 	GeometryCollectionAlgo::GlobalMatrices(GeometryCollection->Transform, GeometryCollection->Parent, GlobalTransformArray);
 	HullsFromGeometry(*GeometryCollection, GlobalTransformArray, Convexes, TransformToConvexIndexArr, GeometryCollection->SimulationType, FGeometryCollection::ESimulationTypes::FST_Rigid);
 
-	CreateNonoverlappingConvexHulls(Convexes, TransformToConvexIndexArr, GeometryCollection->SimulationType, FGeometryCollection::ESimulationTypes::FST_Rigid, GeometryCollection->Parent, GCProximity, GeometryCollection->TransformIndex, FracAllowRemove);
+	CreateNonoverlappingConvexHulls(Convexes, TransformToConvexIndexArr, GeometryCollection->SimulationType, FGeometryCollection::ESimulationTypes::FST_Rigid, FGeometryCollection::ESimulationTypes::FST_None, GeometryCollection->Parent, GCProximity, GeometryCollection->TransformIndex, FracAllowRemove);
 
 	TransformHullsToLocal(GlobalTransformArray, Convexes, TransformToConvexIndexArr);
 
