Ignore Chaos Heightfield SerializedRealSize, which does not match the currently enforced FReal serialization as float.

diff --git a/Engine/Source/Runtime/Experimental/Chaos/Public/Chaos/HeightField.h b/Engine/Source/Runtime/Experimental/Chaos/Public/Chaos/HeightField.h
--- a/Engine/Source/Runtime/Experimental/Chaos/Public/Chaos/HeightField.h
+++ b/Engine/Source/Runtime/Experimental/Chaos/Public/Chaos/HeightField.h
@@ -279,27 +279,28 @@ namespace Chaos
 
 			void SafeSerializeReal(FChaosArchive& Ar, FReal& RealValue, int32 RuntimeRealSize, int32 SerializedRealSize)
 			{
-				if (RuntimeRealSize == SerializedRealSize)
+				// LWC_TODO: Serializer - SafeSerializeReal fix up required. Fall back to FChaosArchive FReal serialization behavior (forced to float) in the short term.
+				//if (RuntimeRealSize == SerializedRealSize)
 				{
 					// same sizes all FReal
 					Ar << RealValue;
 				}
-				else 
-				{
-					// size don't match need to do some conversion
-					if (SerializedRealSize == sizeof(float))
-					{
-						float Value = (float)RealValue;
-						Ar << Value;
-						RealValue = (FReal)Value;
-					}
-					else if (SerializedRealSize == sizeof(double))
-					{
-						double Value = (double)RealValue;
-						Ar << Value;
-						RealValue = (FReal)Value;
-					}
-				}
+				//else 
+				//{
+				//	// size don't match need to do some conversion
+				//	if (SerializedRealSize == sizeof(float))
+				//	{
+				//		float Value = (float)RealValue;
+				//		Ar << Value;
+				//		RealValue = (FReal)Value;
+				//	}
+				//	else if (SerializedRealSize == sizeof(double))
+				//	{
+				//		double Value = (double)RealValue;
+				//		Ar << Value;
+				//		RealValue = (FReal)Value;
+				//	}
+				//}
 			}
 
 			void Serialize(FChaosArchive& Ar)
