MeshModelingTools: Change mode manager's HoveredViewportClient back to acting as LastHoveredViewportClient, to work around MouseLeave/MousEnter mismatch when flying in viewport. Also fix to GizmoBoxComponent not using correct position for scale calculation.

#rnx

diff --git a/Engine/Source/Editor/UnrealEd/Private/EditorModeManager.cpp b/Engine/Source/Editor/UnrealEd/Private/EditorModeManager.cpp
--- a/Engine/Source/Editor/UnrealEd/Private/EditorModeManager.cpp
+++ b/Engine/Source/Editor/UnrealEd/Private/EditorModeManager.cpp
@@ -222,17 +222,13 @@ UWorld* FEditorModeTools::GetWorld() const
 
 FEditorViewportClient* FEditorModeTools::GetHoveredViewportClient() const
 {
-	// This is our best effort right now. However this is somewhat incorrect as if you Hover
-	// on other Viewports they get mouse events, but this value stays on the Focused viewport.
-	// Not sure what to do about this right now.
+	// Note: as per the comment in MouseLeave, this currently acts as LastHoveredViewportClient.
 	return HoveredViewportClient;
 }
 
 FEditorViewportClient* FEditorModeTools::GetFocusedViewportClient() const
 {
-	// This is our best effort right now. However this is somewhat incorrect as if you Hover
-	// on other Viewports they get mouse events, but this value stays on the Focused viewport.
-	// Not sure what to do about this right now.
+	// Note: as per the comment in LostFocus, this actually currently acts as LastFocusedViewportClient.
 	return FocusedViewportClient;
 }
 
@@ -1542,7 +1538,13 @@ bool FEditorModeTools::MouseEnter( FEditorViewportClient* InViewportClient, FVie
 
 bool FEditorModeTools::MouseLeave( FEditorViewportClient* InViewportClient, FViewport* Viewport )
 {
-	HoveredViewportClient = nullptr;
+	// TODO: HoveredViewportClient should be reset here, but there is currently a bug (UE-119516) 
+	// that makes it so that flying in viewports can create mismatches between MouseEnter and MouseLeave.
+	// For this reason, we currently use HoveredViewportClient as if it were LastHoveredViewportClient,
+	// which works for the purposes that we use it for.
+	// If we never fix the bug, we should probably just rename it.
+	//HoveredViewportClient = nullptr;
+	
 	bool bHandled = InteractiveToolsContext->MouseLeave(InViewportClient, Viewport);
 
 	ForEachEdMode<ILegacyEdModeViewportInterface>([&bHandled, InViewportClient, Viewport](ILegacyEdModeViewportInterface* Mode)
diff --git a/Engine/Source/Runtime/Experimental/InteractiveToolsFramework/Private/BaseGizmos/GizmoBoxComponent.cpp b/Engine/Source/Runtime/Experimental/InteractiveToolsFramework/Private/BaseGizmos/GizmoBoxComponent.cpp
--- a/Engine/Source/Runtime/Experimental/InteractiveToolsFramework/Private/BaseGizmos/GizmoBoxComponent.cpp
+++ b/Engine/Source/Runtime/Experimental/InteractiveToolsFramework/Private/BaseGizmos/GizmoBoxComponent.cpp
@@ -271,11 +271,14 @@ bool UGizmoBoxComponent::LineTraceComponent(FHitResult& OutHit, const FVector St
 	bool bFlippedY = false;
 	bool bFlippedZ = false;
 
+	FVector WorldOrigin; // Only used if bIsViewDependent
 	if (bIsViewDependent)
 	{
+		WorldOrigin = Transform.TransformPosition(FVector::ZeroVector);
+
 		FVector UseDirectionX, UseDirectionY, UseDirectionZ; // not used
 		GetBoxDirections(bIsViewDependent, ToRawPtr(GizmoViewContext),
-			Transform.TransformPosition(FVector::ZeroVector), // Origin
+			WorldOrigin,
 			Rotation * FVector(1, 0, 0), // DirectionX
 			Rotation * FVector(0, 1, 0), // DirectionY
 			Rotation * FVector(0, 0, 1), // DirectionZ
@@ -284,7 +287,7 @@ bool UGizmoBoxComponent::LineTraceComponent(FHitResult& OutHit, const FVector St
 			UseDirectionX, UseDirectionY, UseDirectionZ, bFlippedX, bFlippedY, bFlippedZ);
 	}
 
-	FVector UseOrigin(
+	FVector UseCenter(
 		(bEnableAxisFlip && bFlippedX) ? -Origin.X : Origin.X,
 		(bEnableAxisFlip && bFlippedY) ? -Origin.Y : Origin.Y,
 		(bEnableAxisFlip && bFlippedZ) ? -Origin.Z : Origin.Z
@@ -295,15 +298,14 @@ bool UGizmoBoxComponent::LineTraceComponent(FHitResult& OutHit, const FVector St
 	if (bIsViewDependent)
 	{
 		DynamicPixelToWorldScale = GizmoRenderingUtil::CalculateLocalPixelToWorldScale(
-			GizmoViewContext, UseOrigin);
+			GizmoViewContext, WorldOrigin);
 	}
 
-	UseOrigin *= DynamicPixelToWorldScale;
+	UseCenter *= DynamicPixelToWorldScale;
 
 	FVector ScaledDims = DynamicPixelToWorldScale * Dimensions;
-	FBox Box(UseOrigin-0.5f*ScaledDims, UseOrigin+0.5f*ScaledDims);
+	FBox Box(UseCenter - 0.5f * ScaledDims, UseCenter + 0.5f * ScaledDims);
 
-	//FMath::LineBoxIntersection(Box, StartLocal, EndLocal, 
 	FVector Extent(SMALL_NUMBER, SMALL_NUMBER, SMALL_NUMBER);
 	FVector HitLocal, NormalLocal; float HitTime;
 	if (FMath::LineExtentBoxIntersection(Box, StartLocal, EndLocal, Extent, HitLocal, NormalLocal, HitTime) == false)
