Chaos: Fix double->float truncation warning in debug packaging/builds

diff --git a/Engine/Source/Runtime/Experimental/Chaos/Private/PhysicsProxy/SkeletalMeshPhysicsProxy.cpp b/Engine/Source/Runtime/Experimental/Chaos/Private/PhysicsProxy/SkeletalMeshPhysicsProxy.cpp
--- a/Engine/Source/Runtime/Experimental/Chaos/Private/PhysicsProxy/SkeletalMeshPhysicsProxy.cpp
+++ b/Engine/Source/Runtime/Experimental/Chaos/Private/PhysicsProxy/SkeletalMeshPhysicsProxy.cpp
@@ -176,7 +176,7 @@ void FSkeletalMeshPhysicsProxy::CreateRigidBodyCallback(FParticlesType& Particle
 						return true; 
 					};
 				// Precision gets worse in opt builds, so only do the phi test in debug.
-				checkSlow(PointVolumeRegistrationCheck(*SamplePoints, *ImplicitObject, 0.01));
+				checkSlow(PointVolumeRegistrationCheck(*SamplePoints, *ImplicitObject, 0.01f));
 				auto PointBBoxCheck =
 					[&](const auto &InSamplePoints, const Chaos::FAABB3 &BBox) -> bool
 					{
