GeometryCollection
-- Generate RelativeSize ManagedArray on the RestCollection
-- Use RelativeSize to determine SizeSpecificData index

diff --git a/Engine/Plugins/Experimental/ChaosEditor/Source/FractureEditor/Private/FractureEditorModeToolkit.cpp b/Engine/Plugins/Experimental/ChaosEditor/Source/FractureEditor/Private/FractureEditorModeToolkit.cpp
--- a/Engine/Plugins/Experimental/ChaosEditor/Source/FractureEditor/Private/FractureEditorModeToolkit.cpp
+++ b/Engine/Plugins/Experimental/ChaosEditor/Source/FractureEditor/Private/FractureEditorModeToolkit.cpp
@@ -1130,6 +1130,25 @@ void FFractureEditorModeToolkit::UpdateGeometryComponentAttributes(UGeometryColl
 					UpdateVolumes(GeometryCollection, MassSpaceParticles, Idx);
 				}
 			}
+
+			// Normalize Size
+			TManagedArray<float>& Size = GeometryCollection->GetAttribute<float>("Size", FTransformCollection::TransformGroup);
+			const TManagedArray<int32>& SimulationType = GeometryCollection->SimulationType;
+			float MaxSize = 0.0f;
+			for (int32 Idx = 0; Idx < Size.Num(); ++Idx)
+			{
+				if (SimulationType[Idx] == FGeometryCollection::ESimulationTypes::FST_Rigid)
+				{ 
+					MaxSize = FMath::Max(Size[Idx], MaxSize);
+				}
+			}
+			if (MaxSize > 0.0)
+			{ 
+				for (int32 Idx = 0; Idx < Size.Num(); ++Idx)
+				{ 
+					Size[Idx] /= MaxSize;
+				}
+			}
 		}
 	}
 	
@@ -1145,6 +1164,12 @@ void FFractureEditorModeToolkit::UpdateVolumes(FGeometryCollectionPtr GeometryCo
 	const TManagedArray<int32>& TransformToGeometryIndex = GeometryCollection->TransformToGeometryIndex;
 	const TManagedArray<FIntVector>& Indices = GeometryCollection->Indices;
 	TManagedArray<float>& Volumes = GeometryCollection->GetAttribute<float>("Volume", FTransformCollection::TransformGroup);
+	
+	if (!GeometryCollection->HasAttribute("Size", FTransformCollection::TransformGroup))
+	{
+		GeometryCollection->AddAttribute<float>("Size", FTransformCollection::TransformGroup);
+	}
+	TManagedArray<float>& Size = GeometryCollection->GetAttribute<float>("Size", FTransformCollection::TransformGroup);
 
 	if (SimulationType[TransformIndex] == FGeometryCollection::ESimulationTypes::FST_Rigid)
 	{
@@ -1166,19 +1191,23 @@ void FFractureEditorModeToolkit::UpdateVolumes(FGeometryCollectionPtr GeometryCo
 
 			// Since we're only interested in relative mass, we assume density = 1.0
 			Volumes[TransformIndex] = Volume; // todo(lwc) potential conversion from double to float
+			Size[TransformIndex] = FGenericPlatformMath::Pow(static_cast<float>(Volume), (1.0f/3.0f));
 		}	
 	}
 	else if (SimulationType[TransformIndex] == FGeometryCollection::ESimulationTypes::FST_Clustered)
 	{
 		// Recurse to children and sum the volumes for this node
 		float LocalVolume = 0.0;
+		float LocalSize = 0.0;
 		for (int32 ChildIndex : Children[TransformIndex])
 		{
 			UpdateVolumes(GeometryCollection, MassSpaceParticles, ChildIndex);
 			LocalVolume += Volumes[ChildIndex];
+			LocalSize += Size[ChildIndex];
 		}
 
 		Volumes[TransformIndex] = LocalVolume;
+		Size[TransformIndex] = LocalSize;
 	}
 	else
 	{
diff --git a/Engine/Plugins/Experimental/ChaosEditor/Source/FractureEditor/Private/SGeometryCollectionHistogram.cpp b/Engine/Plugins/Experimental/ChaosEditor/Source/FractureEditor/Private/SGeometryCollectionHistogram.cpp
--- a/Engine/Plugins/Experimental/ChaosEditor/Source/FractureEditor/Private/SGeometryCollectionHistogram.cpp
+++ b/Engine/Plugins/Experimental/ChaosEditor/Source/FractureEditor/Private/SGeometryCollectionHistogram.cpp
@@ -82,6 +82,11 @@ void FGeometryCollectionHistogramItem::SetInspectedAttribute(EInspectedAttribute
 				}
 				break;
 
+			case EInspectedAttributeEnum::Size:
+				InspectedValue = Collection->GetAttribute<float>(TEXT("Size"), TEXT("Transform"))[BoneIndex];
+				HoverString = FString::Printf(TEXT("%d: %.2f"), BoneIndex, InspectedValue);
+				break;
+
 			default:
 				// Invalid inspection attribute
 				check(false);
diff --git a/Engine/Plugins/Experimental/ChaosEditor/Source/FractureEditor/Private/SGeometryCollectionHistogram.h b/Engine/Plugins/Experimental/ChaosEditor/Source/FractureEditor/Private/SGeometryCollectionHistogram.h
--- a/Engine/Plugins/Experimental/ChaosEditor/Source/FractureEditor/Private/SGeometryCollectionHistogram.h
+++ b/Engine/Plugins/Experimental/ChaosEditor/Source/FractureEditor/Private/SGeometryCollectionHistogram.h
@@ -29,7 +29,8 @@ enum class EInspectedAttributeEnum : uint8
 {
 	Volume = 0					UMETA(DisplayName = "Volume"),
 	Level = 1					UMETA(DisplayName = "Level"),
-	InitialDynamicState = 3		UMETA(DisplayName = "InitialDynamicState")
+	InitialDynamicState = 3		UMETA(DisplayName = "InitialDynamicState"),
+	Size = 4					UMETA(DisplayName = "RelativeSize")
 };
 
 /** Settings for Histogram configuration. **/
diff --git a/Engine/Source/Runtime/Experimental/Chaos/Private/PhysicsProxy/GeometryCollectionPhysicsProxy.cpp b/Engine/Source/Runtime/Experimental/Chaos/Private/PhysicsProxy/GeometryCollectionPhysicsProxy.cpp
--- a/Engine/Source/Runtime/Experimental/Chaos/Private/PhysicsProxy/GeometryCollectionPhysicsProxy.cpp
+++ b/Engine/Source/Runtime/Experimental/Chaos/Private/PhysicsProxy/GeometryCollectionPhysicsProxy.cpp
@@ -1017,17 +1017,28 @@ FGeometryCollectionPhysicsProxy::BuildClusters(
 
 	if (Parameters.bUseSizeSpecificDamageThresholds)
 	{
-		const TManagedArray<FGeometryDynamicCollection::FSharedImplicit>& Implicit = DynamicCollection.Implicits;
-		if (Implicit[CollectionClusterIndex] && Implicit[CollectionClusterIndex]->HasBoundingBox())
+		// If RelativeSize is available, use that to determine SizeSpecific index, otherwise, fall back to bounds volume.
+		int32 SizeSpecificIdx = 0;
+		if (Parameters.RestCollection->HasAttribute("Size", FTransformCollection::TransformGroup))
 		{
-			FBox LocalBoundingBox(Implicit[CollectionClusterIndex]->BoundingBox().Min(), Implicit[CollectionClusterIndex]->BoundingBox().Max());
-			const int32 SizeSpecificIdx = FindSizeSpecificIdx(Parameters.Shared.SizeSpecificData, LocalBoundingBox);
-			if (0<= SizeSpecificIdx && SizeSpecificIdx < Parameters.Shared.SizeSpecificData.Num() )
+			const TManagedArray<float>& RelativeSize = Parameters.RestCollection->GetAttribute<float>("Size", FTransformCollection::TransformGroup);
+			SizeSpecificIdx = FindSizeSpecificIdx(Parameters.Shared.SizeSpecificData, RelativeSize[CollectionClusterIndex]);
+		}
+		else
+		{
+			const TManagedArray<FGeometryDynamicCollection::FSharedImplicit>& Implicit = DynamicCollection.Implicits;
+			if (Implicit[CollectionClusterIndex] && Implicit[CollectionClusterIndex]->HasBoundingBox())
 			{
-				const FSharedSimulationSizeSpecificData& SizeSpecificData = Parameters.Shared.SizeSpecificData[SizeSpecificIdx];
-				Damage = SizeSpecificData.DamageThreshold;
+				FBox LocalBoundingBox(Implicit[CollectionClusterIndex]->BoundingBox().Min(), Implicit[CollectionClusterIndex]->BoundingBox().Max());
+				SizeSpecificIdx = FindSizeSpecificIdx(Parameters.Shared.SizeSpecificData, LocalBoundingBox);
 			}
 		}
+
+		if (0 <= SizeSpecificIdx && SizeSpecificIdx < Parameters.Shared.SizeSpecificData.Num())
+		{
+			const FSharedSimulationSizeSpecificData& SizeSpecificData = Parameters.Shared.SizeSpecificData[SizeSpecificIdx];
+			Damage = SizeSpecificData.DamageThreshold;
+		}
 	}
 
 	Parent->SetStrains(Damage);
@@ -1820,6 +1831,12 @@ void FGeometryCollectionPhysicsProxy::InitializeSharedCollisionStructures(
 	TManagedArray<TSet<int32>>* TransformToConvexIndices = RestCollection.FindAttribute<TSet<int32>>("TransformToConvexIndices", FTransformCollection::TransformGroup);
 	TManagedArray<TUniquePtr<Chaos::FConvex>>* ConvexGeometry = RestCollection.FindAttribute<TUniquePtr<Chaos::FConvex>>("ConvexHull", "Convex");
 
+	bool bUseRelativeSize = RestCollection.HasAttribute(TEXT("Size"), FTransformCollection::TransformGroup);
+	if (!bUseRelativeSize)
+	{
+		UE_LOG(LogChaos, Warning, TEXT("Relative Size not found on Rest Collection. Using bounds volume for SizeSpecificData indexing instead."));
+	}
+
 
 	// @todo(chaos_transforms) : do we still use this?
 	TManagedArray<FTransform>& CollectionMassToLocal =
@@ -2091,7 +2108,18 @@ void FGeometryCollectionPhysicsProxy::InitializeSharedCollisionStructures(
 				InstanceBoundingBox = FBox(MassProperties.CenterOfMass, MassProperties.CenterOfMass);
 			}
 
-			const int32 SizeSpecificIdx = FindSizeSpecificIdx(SharedParams.SizeSpecificData, InstanceBoundingBox);
+			// If we have a normalized Size available, use that to determine SizeSpecific index, otherwise fall back on Bounds volume.
+			int32 SizeSpecificIdx;
+			if (bUseRelativeSize)
+			{
+				const TManagedArray<float>& RelativeSize = RestCollection.GetAttribute<float>(TEXT("Size"), FTransformCollection::TransformGroup);
+				SizeSpecificIdx = FindSizeSpecificIdx(SharedParams.SizeSpecificData, RelativeSize[TransformGroupIndex]);
+			}
+			else
+			{
+				SizeSpecificIdx = FindSizeSpecificIdx(SharedParams.SizeSpecificData, InstanceBoundingBox);
+			}
+			
 			const FSharedSimulationSizeSpecificData& SizeSpecificData = SharedParams.SizeSpecificData[SizeSpecificIdx];
 
 			if (SizeSpecificData.CollisionShapesData.Num())
