Fix a priority crash in FPBDConstraintGraph::RemoveConstraint(), re-instating a previous implemetaion of that function.

diff --git a/Engine/Source/Runtime/Experimental/Chaos/Private/Chaos/PBDConstraintGraph.cpp b/Engine/Source/Runtime/Experimental/Chaos/Private/Chaos/PBDConstraintGraph.cpp
--- a/Engine/Source/Runtime/Experimental/Chaos/Private/Chaos/PBDConstraintGraph.cpp
+++ b/Engine/Source/Runtime/Experimental/Chaos/Private/Chaos/PBDConstraintGraph.cpp
@@ -307,29 +307,61 @@ void FPBDConstraintGraph::AddConstraint(const uint32 InContainerId, FConstraintH
 	}
 }
 
+int32 FPBDConstraintGraph::GetConstraintIndexFromNodeAndConstraintHandle(const FGraphNode& Node, const FConstraintHandle* ConstraintHandle)
+{
+	for (const int32 NodeEdgeIndex : Node.Edges)
+	{
+		const FConstraintData& EdgeData = Edges[NodeEdgeIndex].Data;
+		if (EdgeData.GetConstraintHandle() == ConstraintHandle)
+		{
+			return NodeEdgeIndex;
+		}
+	}
+
+	return INDEX_NONE;
+}
+
 void FPBDConstraintGraph::RemoveConstraint(const uint32 InContainerId, FConstraintHandle* InConstraintHandle, const TVector<FGeometryParticleHandle*, 2>& ConstrainedParticles)
 {
 	check(InConstraintHandle);
 
-	const int32 EdgeIndex = InConstraintHandle->GetConstraintGraphIndex();
-	if (ensure(EdgeIndex != INDEX_NONE))
-	{
-		const int32* PNodeIndex0 = (ConstrainedParticles[0]) ? ParticleToNodeIndex.Find(ConstrainedParticles[0]) : nullptr;
-		const int32* PNodeIndex1 = (ConstrainedParticles[1]) ? ParticleToNodeIndex.Find(ConstrainedParticles[1]) : nullptr;
-		if (ensure(PNodeIndex0 || PNodeIndex1))
+	const int32* PNodeIndex0 = (ConstrainedParticles[0]) ? ParticleToNodeIndex.Find(ConstrainedParticles[0]) : nullptr;
+	const int32* PNodeIndex1 = (ConstrainedParticles[1]) ? ParticleToNodeIndex.Find(ConstrainedParticles[1]) : nullptr;
+
+	if (ensure(PNodeIndex0 || PNodeIndex1))
+	{ 
+		// we need to find the edge index that matches this constraint
+		// we could go through the entire edge list and find the matching InContainerId, 
+		// but each node has an edge array that will contain it 
+		int32 EdgeIndex = INDEX_NONE;
+		if (PNodeIndex0)
+		{
+			EdgeIndex = GetConstraintIndexFromNodeAndConstraintHandle(Nodes[*PNodeIndex0], InConstraintHandle);
+		}
+		if (PNodeIndex1 && EdgeIndex == INDEX_NONE)
+		{
+			EdgeIndex = GetConstraintIndexFromNodeAndConstraintHandle(Nodes[*PNodeIndex1], InConstraintHandle);
+		}
+
+		if (EdgeIndex != INDEX_NONE)
 		{
-			// we should have an valid index at this point
+			int32 Index0 = INDEX_NONE, Index1 = INDEX_NONE;
 			if (PNodeIndex0)
 			{
 				Nodes[*PNodeIndex0].Edges.Remove(EdgeIndex);
+				Index0 = *PNodeIndex0;
 			}
 			if (PNodeIndex1)
 			{
 				Nodes[*PNodeIndex1].Edges.Remove(EdgeIndex);
+				Index1 = *PNodeIndex1;
+			}
+			if (Edges[EdgeIndex].FirstNode == Index0 && Edges[EdgeIndex].SecondNode == Index1)
+			{
+				Edges[EdgeIndex] = FGraphEdge();
 			}
-		}
 
-		Edges[EdgeIndex] = FGraphEdge();
+		}
 
 		InConstraintHandle->SetConstraintGraphIndex(INDEX_NONE);
 	}
diff --git a/Engine/Source/Runtime/Experimental/Chaos/Public/Chaos/PBDConstraintGraph.h b/Engine/Source/Runtime/Experimental/Chaos/Public/Chaos/PBDConstraintGraph.h
--- a/Engine/Source/Runtime/Experimental/Chaos/Public/Chaos/PBDConstraintGraph.h
+++ b/Engine/Source/Runtime/Experimental/Chaos/Public/Chaos/PBDConstraintGraph.h
@@ -248,6 +248,7 @@ namespace Chaos
 		void ParticleRemove(FGeometryParticleHandle* RemovedParticle);
 		int32 GetNextNodeIndex();
 		const TArray<int32>& GetUpdatedNodes() const { return UpdatedNodes; }
+		int32 GetConstraintIndexFromNodeAndConstraintHandle(const FGraphNode& Node, const FConstraintHandle* ConstraintHandle);
 
 		TArray<FGraphNode> Nodes;
 		TArray<FGraphEdge> Edges;
