Update non-overlapping convex algorithm to try more plane options + be more robust

#rnx

diff --git a/Engine/Source/Runtime/Experimental/Chaos/Private/Chaos/CollisionConvexMesh.cpp b/Engine/Source/Runtime/Experimental/Chaos/Private/Chaos/CollisionConvexMesh.cpp
--- a/Engine/Source/Runtime/Experimental/Chaos/Private/Chaos/CollisionConvexMesh.cpp
+++ b/Engine/Source/Runtime/Experimental/Chaos/Private/Chaos/CollisionConvexMesh.cpp
@@ -122,7 +122,8 @@ namespace Chaos
 								}
 							}
 
-							OutFaceIndices.Add({ Triangle[0], Triangle[1], Triangle[2] });
+							// Winding is backwards from what Chaos expects
+							OutFaceIndices.Add({ Triangle[0], Triangle[2], Triangle[1] });
 
 							UE::Geometry::THalfspace3<FReal> Halfspace(OutVertices[Triangle[0]], OutVertices[Triangle[1]], OutVertices[Triangle[2]]);
 							const FVec3& Normal = Halfspace.Normal;
diff --git a/Engine/Source/Runtime/Experimental/Chaos/Private/GeometryCollection/GeometryCollectionConvexUtility.cpp b/Engine/Source/Runtime/Experimental/Chaos/Private/GeometryCollection/GeometryCollectionConvexUtility.cpp
--- a/Engine/Source/Runtime/Experimental/Chaos/Private/GeometryCollection/GeometryCollectionConvexUtility.cpp
+++ b/Engine/Source/Runtime/Experimental/Chaos/Private/GeometryCollection/GeometryCollectionConvexUtility.cpp
@@ -230,7 +230,127 @@ void CreateNonoverlappingConvexHulls(
 		return GJKIntersection(*Convexes[ConvexA], *Convexes[ConvexB], IdentityTransform);
 	};
 
-	auto FixCollisionWithCut2 = [&Convexes](int32 ConvexA, int32 ConvexB)
+	auto GetConvexSpan = [](const Chaos::FConvex& Convex, const Chaos::FVec3& Center, const Chaos::FVec3& Normal) -> Chaos::FVec2
+	{
+		int32 NumVertices = Convex.NumVertices();
+		if (!ensure(NumVertices > 0))
+		{
+			return Chaos::FVec2(0, 0);
+		}
+		Chaos::FReal AlongFirst = (Convex.GetVertex(0) - Center).Dot(Normal);
+		Chaos::FVec2 Range(AlongFirst, AlongFirst);
+		for (int Idx = 1; Idx < NumVertices; Idx++)
+		{
+			Chaos::FReal Along = (Convex.GetVertex(Idx) - Center).Dot(Normal);
+			if (Along < Range.X)
+			{
+				Range.X = Along;
+			}
+			else if (Along > Range.Y)
+			{
+				Range.Y = Along;
+			}
+		}
+		return Range;
+	};
+
+	// Score separating plane direction based on how well it separates (lower is better)
+	//  and also compute new center for plane + normal. Normal is either the input normal or flipped.
+	auto ScoreCutPlane = [&GetConvexSpan](
+		const Chaos::FConvex& A, const Chaos::FConvex& B,
+		const FChaosPlane& Plane, bool bOneSidedCut,
+		Chaos::FVec3& OutCenter, Chaos::FVec3& OutNormal) -> Chaos::FReal
+	{
+		bool bRangeAValid = false, bRangeBValid = false;
+		Chaos::FVec2 RangeA = GetConvexSpan(A, Plane.X(), Plane.Normal());
+		Chaos::FVec2 RangeB = GetConvexSpan(B, Plane.X(), Plane.Normal());
+		Chaos::FVec2 Union(FMath::Min(RangeA.X, RangeB.X), FMath::Max(RangeA.Y, RangeB.Y));
+		// no intersection -- cut plane is separating, this is ideal!
+		if (RangeA.X > RangeB.Y || RangeA.Y < RangeB.X)
+		{
+			if (RangeA.X > RangeB.Y)
+			{
+				OutCenter = Plane.X() + Plane.Normal() * ((RangeA.X + RangeB.Y) * (Chaos::FReal).5);
+				OutNormal = -Plane.Normal();
+			}
+			else
+			{
+				OutCenter = Plane.X() + Plane.Normal() * ((RangeA.Y + RangeB.X) * (Chaos::FReal).5);
+				OutNormal = Plane.Normal();
+			}
+			return 0;
+		}
+		// there was an intersection; find the actual mid plane-center and score it
+		Chaos::FVec2 Intersection(FMath::Max(RangeA.X, RangeB.X), FMath::Min(RangeA.Y, RangeB.Y));
+		Chaos::FReal IntersectionMid = (Intersection.X + Intersection.Y) * (Chaos::FReal).5;
+
+		// Decide which side of the plane is kept/removed
+		Chaos::FVec2 BiggerRange = RangeA;
+		Chaos::FReal Sign = 1;
+		if (RangeA.Y - RangeA.X < RangeB.Y - RangeB.X)
+		{
+			BiggerRange = RangeB;
+			Sign = -1;
+		}
+		if (IntersectionMid - BiggerRange.X < BiggerRange.Y - IntersectionMid)
+		{
+			Sign *= -1;
+		}
+		OutNormal = Sign * Plane.Normal();
+
+		Chaos::FReal IntersectionCut = IntersectionMid;
+		if (bOneSidedCut) // if cut is one-sided, move the plane to the far end of Convex B (which it should not cut)
+		{
+			// which end depends on which way the output cut plane is oriented
+			if (Sign > 0)
+			{
+				IntersectionCut = RangeB.X;
+			}
+			else
+			{
+				IntersectionCut = RangeB.Y;
+			}
+		}
+		OutCenter = Plane.X() + Plane.Normal() * IntersectionCut;
+
+		// Simple score favors small intersection span relative to union span
+		// TODO: consider other metrics; e.g. something more directly ~ percent cut away
+		return (Intersection.Y - Intersection.X) / (Union.Y - Union.X);
+	};
+
+	// Search cut plane options for the most promising one -- 
+	// Usually GJK gives a good cut plane, but it can fail badly so we also test a simple 'difference between centers' plane
+	// TODO: consider adding more plane options to the search
+	auto FindCutPlane = [&ScoreCutPlane](const Chaos::FConvex& A, const Chaos::FConvex& B,
+		Chaos::FVec3 CloseA, Chaos::FVec3 CloseB, Chaos::FVec3 Normal,
+		bool bOneSidedCut,
+		Chaos::FVec3& OutCenter, Chaos::FVec3& OutNormal) -> bool
+	{
+		OutCenter = (CloseA + CloseB) * .5;
+		OutNormal = Normal;
+		Chaos::FVec3 GJKCenter;
+		Chaos::FReal BestScore = ScoreCutPlane(A, B, FChaosPlane(OutCenter, OutNormal), bOneSidedCut, OutCenter, OutNormal);
+		Chaos::FVec3 MassSepNormal = (B.GetCenterOfMass() - A.GetCenterOfMass());
+		if (MassSepNormal.Normalize() && BestScore > 0)
+		{
+			Chaos::FVec3 MassSepCenter = (A.GetCenterOfMass() + B.GetCenterOfMass()) * .5;
+			Chaos::FReal ScoreB = ScoreCutPlane(A, B,
+				FChaosPlane(MassSepCenter, MassSepNormal), bOneSidedCut, MassSepCenter, MassSepNormal);
+			if (ScoreB < BestScore)
+			{
+				BestScore = ScoreB;
+				OutCenter = MassSepCenter;
+				OutNormal = MassSepNormal;
+			}
+		}
+		if (BestScore == 0)
+		{
+			return false;
+		}
+		return true;
+	};
+
+	auto FixCollisionWithCut = [&Convexes, &FindCutPlane](int32 ConvexA, int32 ConvexB)
 	{
 		if (!ensure(Convexes[ConvexA]->NumVertices() > 0 && Convexes[ConvexB]->NumVertices() > 0))
 		{
@@ -245,7 +365,12 @@ void CreateNonoverlappingConvexHulls(
 		bool bCollide = Chaos::GJKPenetration(*Convexes[ConvexA], *Convexes[ConvexB], IdentityTransform, Depth, CloseA, CloseB, Normal, OutIdxA, OutIdxB);
 		if (bCollide)
 		{
-			FChaosPlane CutPlane((CloseA + CloseB) * .5, Normal);
+			Chaos::FVec3 IdealCenter, IdealNormal;
+			if (!FindCutPlane(*Convexes[ConvexA], *Convexes[ConvexB], CloseA, CloseB, Normal, false, IdealCenter, IdealNormal))
+			{
+				return false;
+			}
+			FChaosPlane CutPlane(IdealCenter, IdealNormal);
 			TArray<Chaos::FVec3> CutHullPts;
 			if (CutHull(*Convexes[ConvexA], CutPlane, true, CutHullPts))
 			{
@@ -298,7 +423,7 @@ void CreateNonoverlappingConvexHulls(
 			{
 				if (NbrIdx < HullIdx && TransformToConvexIndices[NbrIdx].Num() > 0)
 				{
-					FixCollisionWithCut2(OnlyConvex(HullIdx), OnlyConvex(NbrIdx));
+					FixCollisionWithCut(OnlyConvex(HullIdx), OnlyConvex(NbrIdx));
 				}
 			}
 		}
@@ -322,7 +447,7 @@ void CreateNonoverlappingConvexHulls(
 		}
 	};
 
-	auto FixLeafCollisions = [&AddLeaves, &FixCollisionWithCut2, &OnlyConvex](int32 BoneA, int32 BoneB)
+	auto FixLeafCollisions = [&AddLeaves, &FixCollisionWithCut, &OnlyConvex](int32 BoneA, int32 BoneB)
 	{
 		TArray<int32> LeavesA, LeavesB;
 		AddLeaves(BoneA, LeavesA);
@@ -332,7 +457,7 @@ void CreateNonoverlappingConvexHulls(
 		{
 			for (int32 LeafB : LeavesB)
 			{
-				bool bCollides = FixCollisionWithCut2(OnlyConvex(LeafA), OnlyConvex(LeafB));
+				bool bCollides = FixCollisionWithCut(OnlyConvex(LeafA), OnlyConvex(LeafB));
 				bAnyCollides |= bCollides;
 			}
 		}
@@ -438,13 +563,16 @@ void CreateNonoverlappingConvexHulls(
 					}
 				}
 				int32 ConvexIdx = OnlyConvex(Bone);
-				*Convexes[ConvexIdx] = Chaos::FConvex(JoinedHullPts, KINDA_SMALL_NUMBER);
-				NonLeafVolumes[ConvexIdx] = Convexes[ConvexIdx]->GetVolume();
+				if (ConvexIdx)
+				{
+					*Convexes[ConvexIdx] = Chaos::FConvex(JoinedHullPts, KINDA_SMALL_NUMBER);
+					NonLeafVolumes[ConvexIdx] = Convexes[ConvexIdx]->GetVolume();
+				}
 			}
 		}
 	}
 
-	auto CutIfOk = [&Convexes, &OnlyConvex, &NonLeafVolumes, &FracAllowRemove](bool bOneSidedCut, int32 ConvexA, int32 ConvexB) -> bool
+	auto CutIfOk = [&Convexes, &OnlyConvex, &NonLeafVolumes, &FindCutPlane, &FracAllowRemove](bool bOneSidedCut, int32 ConvexA, int32 ConvexB) -> bool
 	{
 		Chaos::FReal Depth;
 		Chaos::FVec3 CloseA, CloseB, Normal;
@@ -453,12 +581,9 @@ void CreateNonoverlappingConvexHulls(
 		bool bCollide = Chaos::GJKPenetration(*Convexes[ConvexA], *Convexes[ConvexB], IdentityTransform, Depth, CloseA, CloseB, Normal, OutIdxA, OutIdxB);
 		if (bCollide)
 		{
-			Chaos::FVec3 CutPt = CloseB;
-			if (!bOneSidedCut)
-			{
-				CutPt = (CloseA + CloseB) * .5;
-			}
-			FChaosPlane CutPlane(CutPt, Normal);
+			Chaos::FVec3 IdealCenter, IdealNormal;
+			FindCutPlane(*Convexes[ConvexA], *Convexes[ConvexB], CloseA, CloseB, Normal, bOneSidedCut, IdealCenter, IdealNormal);
+			FChaosPlane CutPlane(IdealCenter, IdealNormal);
 
 			// Tentatively create the clipped hulls
 			Chaos::FConvex CutHullA, CutHullB;
@@ -466,6 +591,10 @@ void CreateNonoverlappingConvexHulls(
 			TArray<Chaos::FVec3> CutHullPts;
 			if (CutHull(*Convexes[ConvexA], CutPlane, true, CutHullPts))
 			{
+				if (CutHullPts.Num() < 4) // immediate reject zero-volume results
+				{
+					return false;
+				}
 				CutHullA = Chaos::FConvex(CutHullPts, KINDA_SMALL_NUMBER);
 				bCreatedA = true;
 			}
@@ -614,14 +743,32 @@ void HullsFromGeometry(
 			int32 VStart = Geometry.VertexStart[GeomIdx];
 			int32 VCount = Geometry.VertexCount[GeomIdx];
 			int32 VEnd = VStart + VCount;
-			TArray<Chaos::FVec3> CutHullPts;
-			CutHullPts.Reserve(VCount);
+			TArray<Chaos::FVec3> HullPts;
+			HullPts.Reserve(VCount);
 			for (int32 VIdx = VStart; VIdx < VEnd; VIdx++)
 			{
-				CutHullPts.Add(GlobalVertices[VIdx]);
+				HullPts.Add(GlobalVertices[VIdx]);
+			}
+			ensure(HullPts.Num() > 0);
+			int32 ConvexIdx = Convexes.Add(MakeUnique<Chaos::FConvex>(HullPts, KINDA_SMALL_NUMBER));
+			if (Convexes[ConvexIdx]->NumVertices() == 0 && HullPts.Num() > 0)
+			{
+				// if we've failed to make a convex hull, add a tiny bounding box just to ensure every geometry has a hull
+				Chaos::FAABB3 AABB = Convexes[ConvexIdx]->BoundingBox();
+				AABB.Thicken(.0001);
+				Chaos::FVec3 Min = AABB.Min();
+				Chaos::FVec3 Max = AABB.Max();
+				HullPts.Add(Min);
+				HullPts.Add(Max);
+				HullPts.Add(Chaos::FVec3(Min.X, Min.Y, Max.Z));
+				HullPts.Add(Chaos::FVec3(Min.X, Max.Y, Max.Z));
+				HullPts.Add(Chaos::FVec3(Max.X, Min.Y, Max.Z));
+				HullPts.Add(Chaos::FVec3(Max.X, Max.Y, Min.Z));
+				HullPts.Add(Chaos::FVec3(Max.X, Min.Y, Min.Z));
+				HullPts.Add(Chaos::FVec3(Min.X, Max.Y, Min.Z));
+				*Convexes[ConvexIdx] = Chaos::FConvex(HullPts, KINDA_SMALL_NUMBER);
 			}
-			ensure(CutHullPts.Num() > 0);
-			int32 ConvexIdx = Convexes.Add(MakeUnique<Chaos::FConvex>(CutHullPts, KINDA_SMALL_NUMBER));
+			ensure(Convexes[ConvexIdx]->NumVertices() > 0);
 			TransformToConvexIndices[Idx].Add(ConvexIdx);
 		}
 	}
