CIS fix (missing cast) for convex utility

diff --git a/Engine/Source/Runtime/Experimental/Chaos/Private/GeometryCollection/GeometryCollectionConvexUtility.cpp b/Engine/Source/Runtime/Experimental/Chaos/Private/GeometryCollection/GeometryCollectionConvexUtility.cpp
--- a/Engine/Source/Runtime/Experimental/Chaos/Private/GeometryCollection/GeometryCollectionConvexUtility.cpp
+++ b/Engine/Source/Runtime/Experimental/Chaos/Private/GeometryCollection/GeometryCollectionConvexUtility.cpp
@@ -755,7 +755,7 @@ void HullsFromGeometry(
 			{
 				// if we've failed to make a convex hull, add a tiny bounding box just to ensure every geometry has a hull
 				Chaos::FAABB3 AABB = Convexes[ConvexIdx]->BoundingBox();
-				AABB.Thicken(.0001);
+				AABB.Thicken((Chaos::FReal).0001);
 				Chaos::FVec3 Min = AABB.Min();
 				Chaos::FVec3 Max = AABB.Max();
 				HullPts.Add(Min);
