Fix for LWC compile error due to FVector2 using floats even with LWC enabled
#rnx

diff --git a/Engine/Source/Runtime/Experimental/Chaos/Private/GeometryCollection/GeometryCollectionConvexUtility.cpp b/Engine/Source/Runtime/Experimental/Chaos/Private/GeometryCollection/GeometryCollectionConvexUtility.cpp
--- a/Engine/Source/Runtime/Experimental/Chaos/Private/GeometryCollection/GeometryCollectionConvexUtility.cpp
+++ b/Engine/Source/Runtime/Experimental/Chaos/Private/GeometryCollection/GeometryCollectionConvexUtility.cpp
@@ -241,7 +241,7 @@ void CreateNonoverlappingConvexHulls(
 		Chaos::FVec2 Range(AlongFirst, AlongFirst);
 		for (int Idx = 1; Idx < NumVertices; Idx++)
 		{
-			Chaos::FReal Along = (Convex.GetVertex(Idx) - Center).Dot(Normal);
+			const float Along = static_cast<float>((Convex.GetVertex(Idx) - Center).Dot(Normal));
 			if (Along < Range.X)
 			{
 				Range.X = Along;
