MeshModelingTools: Fix gizmo box component not taking into account world gizmo mode in hit testing.

#rnx

diff --git a/Engine/Source/Runtime/Experimental/InteractiveToolsFramework/Private/BaseGizmos/GizmoBoxComponent.cpp b/Engine/Source/Runtime/Experimental/InteractiveToolsFramework/Private/BaseGizmos/GizmoBoxComponent.cpp
--- a/Engine/Source/Runtime/Experimental/InteractiveToolsFramework/Private/BaseGizmos/GizmoBoxComponent.cpp
+++ b/Engine/Source/Runtime/Experimental/InteractiveToolsFramework/Private/BaseGizmos/GizmoBoxComponent.cpp
@@ -257,15 +257,20 @@ bool UGizmoBoxComponent::LineTraceComponent(FHitResult& OutHit, const FVector St
 	using namespace GizmoBoxComponentLocals;
 
 	const FTransform& Transform = this->GetComponentToWorld();
+	FTransform TransformToUse = Transform;
+	if (bWorld)
+	{
+		TransformToUse.SetRotation(FQuat::Identity);
+	}
 
 	// transform points into local space
-	FVector StartLocal = Transform.InverseTransformPosition(Start);
-	FVector EndLocal = Transform.InverseTransformPosition(End);
+	FVector StartLocal = TransformToUse.InverseTransformPosition(Start);
+	FVector EndLocal = TransformToUse.InverseTransformPosition(End);
 
-	// transform into rotation space
-	FQuat InvRotation = Rotation.Inverse();
-	StartLocal = InvRotation * StartLocal;
-	EndLocal = InvRotation * EndLocal;
+	// transform into box-specific rotation space
+	FQuat InvBoxRotation = Rotation.Inverse();
+	StartLocal = InvBoxRotation * StartLocal;
+	EndLocal = InvBoxRotation * EndLocal;
 
 	bool bFlippedX = false; 
 	bool bFlippedY = false;
@@ -274,6 +279,11 @@ bool UGizmoBoxComponent::LineTraceComponent(FHitResult& OutHit, const FVector St
 	FVector WorldOrigin; // Only used if bIsViewDependent
 	if (bIsViewDependent)
 	{
+		// It doesn't matter whether we use TransformToUse or Transform in this block since
+		// the rotation component does not affect world origin and the vector transforms in
+		// GetBoxDirections are only used if bWorld is false. We'll use the original transform
+		// here just to line up with what our scene proxy is doing.
+
 		WorldOrigin = Transform.TransformPosition(FVector::ZeroVector);
 
 		FVector UseDirectionX, UseDirectionY, UseDirectionZ; // not used
@@ -313,7 +323,7 @@ bool UGizmoBoxComponent::LineTraceComponent(FHitResult& OutHit, const FVector St
 		return false;
 	}
 
-	FVector HitWorld = Transform.TransformPosition(Rotation * HitLocal);
+	FVector HitWorld = TransformToUse.TransformPosition(Rotation * HitLocal);
 
 	OutHit.Component = this;
 	OutHit.ImpactPoint = HitWorld;
