Chaos AABBTree dirty grid, Raycasts and Sweeps are now done in local coordinates to improve precision in very large worlds

diff --git a/Engine/Source/Programs/HeadlessChaos/Private/HeadlessChaosTestBroadphase.cpp b/Engine/Source/Programs/HeadlessChaos/Private/HeadlessChaosTestBroadphase.cpp
--- a/Engine/Source/Programs/HeadlessChaos/Private/HeadlessChaosTestBroadphase.cpp
+++ b/Engine/Source/Programs/HeadlessChaos/Private/HeadlessChaosTestBroadphase.cpp
@@ -615,14 +615,15 @@ namespace ChaosTest
 		{
 			FAABBTreeDirtyGridCVars::DirtyElementMaxCellCapacity = 7;
 			TUniquePtr<TBox<FReal, 3>> Box;
-			auto Boxes = BuildBoxes(Box, 100, FVec3(40, 40, 1), FVec3(-2000, -2000, -50));
+			FVec3 LargeOffset(10000000, 10000000, 10000000); // Test for floating point precision errors at large world offsets
+			auto Boxes = BuildBoxes(Box, 100, FVec3(40, 40, 1), FVec3(-2000, -2000, -50) + LargeOffset);
 
 			for (float Angle = 0.0; Angle < 2 * PI; Angle += (10.0f / 360.0f) * 2.0f * PI)
 			{
 
 				FVec3 Direction{ FMath::Cos(Angle), FMath::Sin(Angle), 0 };
 				// With the grid
-				FVisitor VisitorGrid(FVec3(53, 27, 0), Direction, 0, *Boxes);
+				FVisitor VisitorGrid(FVec3(53, 27, 0) + LargeOffset, Direction, 0, *Boxes);
 				VisitorGrid.HalfExtents = FVec3(102, 20, 2);
 				{
 					FAABBTreeDirtyGridCVars::DirtyElementGridCellSize = 100;
@@ -641,7 +642,7 @@ namespace ChaosTest
 				}
 
 				// Without the grid
-				FVisitor VisitorNoGrid(FVec3(53, 27, 0), Direction, 0, *Boxes);
+				FVisitor VisitorNoGrid(FVec3(53, 27, 0) + LargeOffset, Direction, 0, *Boxes);
 				VisitorNoGrid.HalfExtents = FVec3(102, 20, 2);
 				{
 					FAABBTreeDirtyGridCVars::DirtyElementGridCellSize = 0;
diff --git a/Engine/Source/Runtime/Experimental/Chaos/Public/Chaos/AABBTreeDirtyGridUtils.h b/Engine/Source/Runtime/Experimental/Chaos/Public/Chaos/AABBTreeDirtyGridUtils.h
--- a/Engine/Source/Runtime/Experimental/Chaos/Public/Chaos/AABBTreeDirtyGridUtils.h
+++ b/Engine/Source/Runtime/Experimental/Chaos/Public/Chaos/AABBTreeDirtyGridUtils.h
@@ -221,7 +221,7 @@ namespace Chaos
 		return false;
 	}
 
-	FORCEINLINE_DEBUGGABLE bool TooManySweepQueryCells(const TVec3<FReal>& QueryHalfExtents, const TVector<FReal, 3>& StartPoint, const TVector<FReal, 3>& Dir, FReal Length, FReal DirtyElementGridCellSizeInv, int32 DirtyElementMaxGridCellQueryCount)
+	FORCEINLINE_DEBUGGABLE bool TooManySweepQueryCells(const TVec3<FReal>& QueryHalfExtents, const FVec3& StartPoint, const FVec3& Dir, FReal Length, FReal DirtyElementGridCellSizeInv, int32 DirtyElementMaxGridCellQueryCount)
 	{
 		int EstimatedNumberOfCells = ((int32)(QueryHalfExtents.X * 2 * DirtyElementGridCellSizeInv) + 2) * ((int32)(QueryHalfExtents.Y * 2 * DirtyElementGridCellSizeInv) + 2) +
 			((int32)(FMath::Max(QueryHalfExtents.X, QueryHalfExtents.Y) * 2 * DirtyElementGridCellSizeInv) + 2) * ((int32)(Length * DirtyElementGridCellSizeInv) + 2);
@@ -231,8 +231,9 @@ namespace Chaos
 
 
 	// This function should be called with a dominant x direction only!
+	// Todo: Refactor: Use TVectors consistently
 	template <typename FunctionType>
-	FORCEINLINE_DEBUGGABLE void DoForSweepIntersectCellsImp(const FReal QueryHalfExtentsX, const FReal QueryHalfExtentsY, const FReal StartPointX, const FReal StartPointY, const FReal EndPointX, const FReal EndPointY, FReal DirtyElementGridCellSize, FReal DirtyElementGridCellSizeInv, FunctionType InFunction)
+	FORCEINLINE_DEBUGGABLE void DoForSweepIntersectCellsImp(const FReal QueryHalfExtentsX, const FReal QueryHalfExtentsY, const FReal StartPointX, const FReal StartPointY, const FReal RayX, const FReal RayY, FReal DirtyElementGridCellSize, FReal DirtyElementGridCellSizeInv, FunctionType InFunction)
 	{
 		// Use 2 paths (Line0 and Line1) that traces the shape of the swept AABB and fill between them
 		// Example of one of the cases (XDirectionDominant, Fill Up):
@@ -253,8 +254,15 @@ namespace Chaos
 		//           Line1    ^TurningPointForLine1
 
 
-		FReal DeltaX = EndPointX - StartPointX;
-		FReal DeltaY = EndPointY - StartPointY;
+		// The Reference Cell's origin is used for the local coordinates origin
+		const int32 ReferenceCellIndexX = GetDirtyCellIndexFromWorldCoordinate(StartPointX, DirtyElementGridCellSizeInv);
+		const int32 ReferenceCellIndexY = GetDirtyCellIndexFromWorldCoordinate(StartPointY, DirtyElementGridCellSizeInv);
+		const TVector<FReal, 2> LocalCoordinatesOrigin{ (FReal)ReferenceCellIndexX * DirtyElementGridCellSize, (FReal)ReferenceCellIndexY * DirtyElementGridCellSize};
+		const TVector<FReal, 2> StartPointLocal{ StartPointX - LocalCoordinatesOrigin.X, StartPointY - LocalCoordinatesOrigin.Y };
+		const TVector<FReal, 2> EndPointLocal{ StartPointLocal.X + RayX, StartPointLocal.Y + RayY};
+
+		FReal DeltaX = RayX;
+		FReal DeltaY = RayY;
 
 		FReal AbsDx = FMath::Abs(DeltaX);
 		FReal AbsDy = FMath::Abs(DeltaY);
@@ -286,12 +294,12 @@ namespace Chaos
 		DtDy = DyTooSmall ? 1 : (FReal)DeltaCelIndexX * DeltaX / DeltaY;
 
 		// Calculate all the bounds we need
-		FReal XEndPointExpanded = EndPointX + (DeltaCelIndexX >= 0 ? QueryHalfExtentsX : -QueryHalfExtentsX);
-		FReal YEndPointExpanded = EndPointY + (DeltaCelIndexY >= 0 ? QueryHalfExtentsY : -QueryHalfExtentsY);
-		FReal XStartPointExpanded = StartPointX + (DeltaCelIndexX >= 0 ? -QueryHalfExtentsX : QueryHalfExtentsX);
-		FReal YStartPointExpanded = StartPointY + (DeltaCelIndexY >= 0 ? -QueryHalfExtentsY : QueryHalfExtentsY);
+		FReal XEndPointExpanded = EndPointLocal.X + (DeltaCelIndexX >= 0 ? QueryHalfExtentsX : -QueryHalfExtentsX);
+		FReal YEndPointExpanded = EndPointLocal.Y + (DeltaCelIndexY >= 0 ? QueryHalfExtentsY : -QueryHalfExtentsY);
+		FReal XStartPointExpanded = StartPointLocal.X + (DeltaCelIndexX >= 0 ? -QueryHalfExtentsX : QueryHalfExtentsX);
+		FReal YStartPointExpanded = StartPointLocal.Y + (DeltaCelIndexY >= 0 ? -QueryHalfExtentsY : QueryHalfExtentsY);
 		int32 TurningPointForLine1; // This is where we need to change direction for line 2
-		TurningPointForLine1 = GetDirtyCellIndexFromWorldCoordinate(StartPointX + (DeltaCelIndexX >= 0 ? QueryHalfExtentsX : -QueryHalfExtentsX), DirtyElementGridCellSizeInv);
+		TurningPointForLine1 = GetDirtyCellIndexFromWorldCoordinate(StartPointLocal.X + (DeltaCelIndexX >= 0 ? QueryHalfExtentsX : -QueryHalfExtentsX), DirtyElementGridCellSizeInv);
 
 		// Line0 current position
 		FReal X0 = XStartPointExpanded;
@@ -301,14 +309,14 @@ namespace Chaos
 		FReal X1 = XStartPointExpanded;
 		FReal Y1 = YStartPointExpanded;
 
-		int32 CurrentCellIndexX0 = GetDirtyCellIndexFromWorldCoordinate(X0, DirtyElementGridCellSizeInv);
-		int32 CurrentCellIndexY0 = GetDirtyCellIndexFromWorldCoordinate(Y0, DirtyElementGridCellSizeInv);
+		int32 CurrentCellIndexX0 = GetDirtyCellIndexFromWorldCoordinate(X0 + LocalCoordinatesOrigin.X, DirtyElementGridCellSizeInv);
+		int32 CurrentCellIndexY0 = GetDirtyCellIndexFromWorldCoordinate(Y0 + LocalCoordinatesOrigin.Y, DirtyElementGridCellSizeInv);
 
-		int32 CurrentCellIndexX1 = GetDirtyCellIndexFromWorldCoordinate(X1, DirtyElementGridCellSizeInv);
-		int32 CurrentCellIndexY1 = GetDirtyCellIndexFromWorldCoordinate(Y1, DirtyElementGridCellSizeInv);
+		int32 CurrentCellIndexX1 = GetDirtyCellIndexFromWorldCoordinate(X1 + LocalCoordinatesOrigin.X, DirtyElementGridCellSizeInv);
+		int32 CurrentCellIndexY1 = GetDirtyCellIndexFromWorldCoordinate(Y1 + LocalCoordinatesOrigin.Y, DirtyElementGridCellSizeInv);
 
-		int32 LastCellIndexX = GetDirtyCellIndexFromWorldCoordinate(XEndPointExpanded, DirtyElementGridCellSizeInv);
-		int32 LastCellIndexY = GetDirtyCellIndexFromWorldCoordinate(YEndPointExpanded, DirtyElementGridCellSizeInv);
+		int32 LastCellIndexX = GetDirtyCellIndexFromWorldCoordinate(XEndPointExpanded + LocalCoordinatesOrigin.X, DirtyElementGridCellSizeInv);
+		int32 LastCellIndexY = GetDirtyCellIndexFromWorldCoordinate(YEndPointExpanded + LocalCoordinatesOrigin.Y, DirtyElementGridCellSizeInv);
 
 		bool Done = false;
 		while (!Done)
@@ -319,8 +327,8 @@ namespace Chaos
 			{
 				FReal CrossingVerticleCellBorderT = std::numeric_limits<FReal>::max();
 				FReal CrossingHorizontalCellBorderT = std::numeric_limits<FReal>::max();
-				CrossingVerticleCellBorderT = DtDx * ((FReal)(CurrentCellIndexX0 + (DeltaCelIndexX > 0 ? 1 : 0)) * DirtyElementGridCellSize - X0);
-				CrossingHorizontalCellBorderT = DtDy * ((FReal)(CurrentCellIndexY0 + (DeltaCelIndexY > 0 ? 1 : 0)) * DirtyElementGridCellSize - Y0);
+				CrossingVerticleCellBorderT = DtDx * ((FReal)(CurrentCellIndexX0 - ReferenceCellIndexX + (DeltaCelIndexX > 0 ? 1 : 0)) * DirtyElementGridCellSize - X0);
+				CrossingHorizontalCellBorderT = DtDy * ((FReal)(CurrentCellIndexY0 - ReferenceCellIndexY + (DeltaCelIndexY > 0 ? 1 : 0)) * DirtyElementGridCellSize - Y0);
 				if (CrossingHorizontalCellBorderT < CrossingVerticleCellBorderT)
 				{
 					X0 += CrossingHorizontalCellBorderT * (1 / DtDx);  // DtDx is always 1 or -1
@@ -339,7 +347,7 @@ namespace Chaos
 				if (CurrentCellIndexY0 != LastCellIndexY && !DyTooSmall)
 				{
 					FReal CrossingVerticleCellBorderT = std::numeric_limits<FReal>::max();
-					CrossingVerticleCellBorderT = DtDx * ((FReal)(CurrentCellIndexX0 + (DeltaCelIndexX > 0 ? 1 : 0)) * DirtyElementGridCellSize - X0);
+					CrossingVerticleCellBorderT = DtDx * ((FReal)(CurrentCellIndexX0 - ReferenceCellIndexX + (DeltaCelIndexX > 0 ? 1 : 0)) * DirtyElementGridCellSize - X0);
 					X0 += CrossingVerticleCellBorderT * (1 / DtDx);
 					Y0 += CrossingVerticleCellBorderT * (1 / DtDy);
 				}
@@ -353,28 +361,28 @@ namespace Chaos
 			// Advance line 1
 			if (CurrentCellIndexX1 != LastCellIndexX)
 			{
-				if (CurrentCellIndexX1 * DeltaCelIndexX < TurningPointForLine1 * DeltaCelIndexX)
+				if ((CurrentCellIndexX1 - ReferenceCellIndexX) * DeltaCelIndexX < TurningPointForLine1 * DeltaCelIndexX)
 				{
 					X1 += DirtyElementGridCellSize * (FReal)DeltaCelIndexX;
 				}
 				else
 				{
-					if (CurrentCellIndexX1 == TurningPointForLine1)
+					if ((CurrentCellIndexX1 - ReferenceCellIndexX) == TurningPointForLine1)
 					{
 						// Put Line position exactly at the turning point
-						X1 = StartPointX + (DeltaCelIndexX >= 0 ? QueryHalfExtentsX : -QueryHalfExtentsX);
+						X1 = StartPointLocal.X + (DeltaCelIndexX >= 0 ? QueryHalfExtentsX : -QueryHalfExtentsX);
 					}
 
 					FReal CrossingVerticleCellBorderT = std::numeric_limits<FReal>::max();
 					FReal CrossingHorizontalCellBorderT = std::numeric_limits<FReal>::max();
 					if (!DxTooSmall)
 					{
-						CrossingVerticleCellBorderT = DtDx * ((FReal)(CurrentCellIndexX0 + (DeltaCelIndexX > 0 ? 1 : 0)) * DirtyElementGridCellSize - X1);
+						CrossingVerticleCellBorderT = DtDx * ((FReal)(CurrentCellIndexX0 - ReferenceCellIndexX + (DeltaCelIndexX > 0 ? 1 : 0)) * DirtyElementGridCellSize - X1);
 					}
 
 					if (!DyTooSmall)
 					{
-						CrossingHorizontalCellBorderT = DtDy * ((FReal)(CurrentCellIndexY0 + (DeltaCelIndexY > 0 ? 1 : 0)) * DirtyElementGridCellSize - Y1);
+						CrossingHorizontalCellBorderT = DtDy * ((FReal)(CurrentCellIndexY0 - ReferenceCellIndexY + (DeltaCelIndexY > 0 ? 1 : 0)) * DirtyElementGridCellSize - Y1);
 					}
 
 					if (CrossingHorizontalCellBorderT < CrossingVerticleCellBorderT)
@@ -399,33 +407,28 @@ namespace Chaos
 	}
 
 	template <typename FunctionType>
-	FORCEINLINE_DEBUGGABLE void DoForSweepIntersectCells(const TVec3<FReal>& QueryHalfExtents, const TVector<FReal, 3>& StartPoint, const TVector<FReal, 3>& Dir, FReal Length, FReal DirtyElementGridCellSize, FReal DirtyElementGridCellSizeInv, FunctionType InFunction)
+	FORCEINLINE_DEBUGGABLE void DoForSweepIntersectCells(const TVec3<FReal>& QueryHalfExtents, const FVec3& StartPoint, const FVec3& Dir, FReal Length, FReal DirtyElementGridCellSize, FReal DirtyElementGridCellSizeInv, FunctionType InFunction)
 	{
-		const TVector<FReal, 3>& EndPoint = StartPoint + Length * Dir;
-		FReal DeltaX = EndPoint.X - StartPoint.X;
-		FReal DeltaY = EndPoint.Y - StartPoint.Y;
-
-		FReal AbsDx = FMath::Abs(DeltaX);
-		FReal AbsDy = FMath::Abs(DeltaY);
-
+		FReal AbsDx = FMath::Abs(Dir.X);
+		FReal AbsDy = FMath::Abs(Dir.Y);
 
 		bool XDirectionDominant = AbsDx >= AbsDy;
 
 		if (XDirectionDominant)
 		{
-			DoForSweepIntersectCellsImp(QueryHalfExtents.X, QueryHalfExtents.Y, StartPoint.X, StartPoint.Y, EndPoint.X, EndPoint.Y, DirtyElementGridCellSize, DirtyElementGridCellSizeInv, InFunction);
+			DoForSweepIntersectCellsImp(QueryHalfExtents.X, QueryHalfExtents.Y, StartPoint.X, StartPoint.Y, Dir.X * Length, Dir.Y * Length, DirtyElementGridCellSize, DirtyElementGridCellSizeInv, InFunction);
 		}
 		else
 		{
 			// Swap Y and X
-			DoForSweepIntersectCellsImp(QueryHalfExtents.Y, QueryHalfExtents.X, StartPoint.Y, StartPoint.X, EndPoint.Y, EndPoint.X, DirtyElementGridCellSize, DirtyElementGridCellSizeInv, [&](FReal X, FReal Y) {InFunction(Y, X); });
+			DoForSweepIntersectCellsImp(QueryHalfExtents.Y, QueryHalfExtents.X, StartPoint.Y, StartPoint.X, Dir.Y * Length, Dir.X * Length, DirtyElementGridCellSize, DirtyElementGridCellSizeInv, [&](FReal X, FReal Y) {InFunction(Y, X); });
 		}
 
 	}
 
-	FORCEINLINE_DEBUGGABLE bool TooManyRaycastQueryCells(const TVector<FReal, 3>& StartPoint, const TVector<FReal, 3>& Dir, const FReal Length, FReal DirtyElementGridCellSizeInv, int32 DirtyElementMaxGridCellQueryCount)
+	FORCEINLINE_DEBUGGABLE bool TooManyRaycastQueryCells(const FVec3& StartPoint, const FVec3& Dir, const FReal Length, FReal DirtyElementGridCellSizeInv, int32 DirtyElementMaxGridCellQueryCount)
 	{
-		const TVector<FReal, 3>& EndPoint = StartPoint + Length * Dir;
+		FVec3 EndPoint = StartPoint + Length * Dir;
 
 		int32 FirstCellIndexX = GetDirtyCellIndexFromWorldCoordinate(StartPoint.X, DirtyElementGridCellSizeInv);
 		int32 FirstCellIndexY = GetDirtyCellIndexFromWorldCoordinate(StartPoint.Y, DirtyElementGridCellSizeInv);
@@ -446,18 +449,28 @@ namespace Chaos
 	}
 
 	template <typename FunctionType>
-	FORCEINLINE_DEBUGGABLE void DoForRaycastIntersectCells(const TVector<FReal, 3>& StartPoint, const TVector<FReal, 3>& Dir, FReal Length, FReal DirtyElementGridCellSize, FReal DirtyElementGridCellSizeInv, FunctionType InFunction)
+	FORCEINLINE_DEBUGGABLE void DoForRaycastIntersectCells(const FVec3& StartPoint, const FVec3& Dir, FReal Length, FReal DirtyElementGridCellSize, FReal DirtyElementGridCellSizeInv, FunctionType InFunction)
 	{
-		const TVector<FReal, 3>& EndPoint = StartPoint + Length * Dir;
+		
+
+		const int32 FirstCellIndexX = GetDirtyCellIndexFromWorldCoordinate(StartPoint.X, DirtyElementGridCellSizeInv);
+		const int32 FirstCellIndexY = GetDirtyCellIndexFromWorldCoordinate(StartPoint.Y, DirtyElementGridCellSizeInv);
+
+		int32 CurrentCellIndexX = FirstCellIndexX;
+		int32 CurrentCellIndexY = FirstCellIndexY;
+
+		// Note: local coordinates are relative to the StartPoint cell origin
+		const FVec3 LocalCoordinatesOrigin{ (FReal)CurrentCellIndexX * DirtyElementGridCellSize, (FReal)CurrentCellIndexY * DirtyElementGridCellSize,  StartPoint.Z};
 
-		int32 CurrentCellIndexX = GetDirtyCellIndexFromWorldCoordinate(StartPoint.X, DirtyElementGridCellSizeInv);
-		int32 CurrentCellIndexY = GetDirtyCellIndexFromWorldCoordinate(StartPoint.Y, DirtyElementGridCellSizeInv);
+		const FVec3 StartPointLocal = StartPoint - LocalCoordinatesOrigin;
 
+		const FVec3 EndPointLocal = StartPointLocal + Length * Dir; // Local coordinates are relative to the start point
+		const FVec3 EndPoint = LocalCoordinatesOrigin + EndPointLocal;
 		int32 LastCellIndexX = GetDirtyCellIndexFromWorldCoordinate(EndPoint.X, DirtyElementGridCellSizeInv);
 		int32 LastCellIndexY = GetDirtyCellIndexFromWorldCoordinate(EndPoint.Y, DirtyElementGridCellSizeInv);
 
-		FReal DeltaX = EndPoint.X - StartPoint.X;
-		FReal DeltaY = EndPoint.Y - StartPoint.Y;
+		FReal DeltaX = EndPointLocal.X - StartPointLocal.X;
+		FReal DeltaY = EndPointLocal.Y - StartPointLocal.Y;
 
 		FReal AbsDx = FMath::Abs(DeltaX);
 		FReal AbsDy = FMath::Abs(DeltaY);
@@ -492,8 +505,9 @@ namespace Chaos
 			DtDy = (FReal)DeltaCelIndexY;
 		}
 
-		FReal X = StartPoint.X;
-		FReal Y = StartPoint.Y;
+		// These are local coordinates
+		FReal X = StartPointLocal.X;
+		FReal Y = StartPointLocal.Y;
 
 		bool Done = false;
 		while (!Done)
@@ -503,12 +517,12 @@ namespace Chaos
 			FReal CrossingHorizontalCellBorderT = std::numeric_limits<FReal>::max();
 			if (!DxTooSmall)
 			{
-				CrossingVerticleCellBorderT = DtDx * ((FReal)(CurrentCellIndexX + (DeltaCelIndexX > 0 ? 1 : 0)) * DirtyElementGridCellSize - X);
+				CrossingVerticleCellBorderT = DtDx * ((FReal)(CurrentCellIndexX - FirstCellIndexX + (DeltaCelIndexX > 0 ? 1 : 0)) * DirtyElementGridCellSize - X);
 			}
 
 			if (!DyTooSmall)
 			{
-				CrossingHorizontalCellBorderT = DtDy * ((FReal)(CurrentCellIndexY + (DeltaCelIndexY > 0 ? 1 : 0)) * DirtyElementGridCellSize - Y);
+				CrossingHorizontalCellBorderT = DtDy * ((FReal)(CurrentCellIndexY - FirstCellIndexY + (DeltaCelIndexY > 0 ? 1 : 0)) * DirtyElementGridCellSize - Y);
 			}
 
 			FReal SmallestT;
