Add Plasticity Deformation Event

diff --git a/Engine/Source/Runtime/Engine/Classes/Components/SkeletalMeshComponent.h b/Engine/Source/Runtime/Engine/Classes/Components/SkeletalMeshComponent.h
--- a/Engine/Source/Runtime/Engine/Classes/Components/SkeletalMeshComponent.h
+++ b/Engine/Source/Runtime/Engine/Classes/Components/SkeletalMeshComponent.h
@@ -816,6 +816,10 @@ public:
 	UPROPERTY(BlueprintAssignable)
 	FConstraintBrokenSignature OnConstraintBroken;
 
+	/** Notification when constraint plasticity drive target changes. */
+	UPROPERTY(BlueprintAssignable)
+	FPlasticDeformationEventSignature OnPlasticDeformation;
+
 	/** Class of the object responsible for  */
 	UPROPERTY(EditAnywhere, Category = Clothing)
 	TSubclassOf<class UClothingSimulationFactory> ClothingSimulationFactory;
@@ -1487,6 +1491,9 @@ private:
 	/** Wrapper that calls our constraint broken delegate */
 	void OnConstraintBrokenWrapper(int32 ConstraintIndex);
 
+	/** Wrapper that calls our constraint plasticity delegate */
+	void OnPlasticDeformationWrapper(int32 ConstraintIndex);
+
 	/** 
 	 * Morph Target Curves. This will override AnimInstance MorphTargetCurves
 	 * if same curve is found
diff --git a/Engine/Source/Runtime/Engine/Classes/Engine/EngineTypes.h b/Engine/Source/Runtime/Engine/Classes/Engine/EngineTypes.h
--- a/Engine/Source/Runtime/Engine/Classes/Engine/EngineTypes.h
+++ b/Engine/Source/Runtime/Engine/Classes/Engine/EngineTypes.h
@@ -878,6 +878,7 @@ enum ECollisionChannel
 };
 
 DECLARE_DELEGATE_OneParam(FOnConstraintBroken, int32 /*ConstraintIndex*/);
+DECLARE_DELEGATE_OneParam(FOnPlasticDeformation, int32 /*ConstraintIndex*/);
 
 #define COLLISION_GIZMO ECC_EngineTraceChannel1
 
@@ -3988,6 +3989,9 @@ struct FComponentSocketDescription
 /** Dynamic delegate to use by components that want to route the broken-event into blueprints */
 DECLARE_DYNAMIC_MULTICAST_DELEGATE_OneParam(FConstraintBrokenSignature, int32, ConstraintIndex);
 
+/** Dynamic delegate to use by components that want to route the pasticity deformation event into blueprints */
+DECLARE_DYNAMIC_MULTICAST_DELEGATE_OneParam(FPlasticDeformationEventSignature, int32, ConstraintIndex);
+
 
 /**
  * Reference to an editor collection of assets. This allows an editor-only picker UI
diff --git a/Engine/Source/Runtime/Engine/Classes/PhysicsEngine/ConstraintInstance.h b/Engine/Source/Runtime/Engine/Classes/PhysicsEngine/ConstraintInstance.h
--- a/Engine/Source/Runtime/Engine/Classes/PhysicsEngine/ConstraintInstance.h
+++ b/Engine/Source/Runtime/Engine/Classes/PhysicsEngine/ConstraintInstance.h
@@ -169,12 +169,17 @@ struct ENGINE_API FConstraintInstanceBase
 	/** Set the constraint broken delegate. */
 	void SetConstraintBrokenDelegate(FOnConstraintBroken InConstraintBrokenDelegate);
 
+	/** Set the plastic deformation delegate. */
+	void SetPlasticDeformationDelegate(FOnPlasticDeformation InPlasticDeformationDelegate);
+
 	protected:
 
 		FOnConstraintBroken OnConstraintBrokenDelegate;
+		FOnPlasticDeformation OnPlasticDeformationDelegate;
 
 		friend struct FConstraintBrokenDelegateData;
 		friend struct FConstraintBrokenDelegateWrapper;
+		friend struct FPlasticDeformationDelegateWrapper;
 
 };
 
@@ -810,10 +815,10 @@ public:
 	void SetLinearLimitSize(float NewLimitSize);
 
 	/** Create physics engine constraint. */
-	void InitConstraint(FBodyInstance* Body1, FBodyInstance* Body2, float Scale, UObject* DebugOwner, FOnConstraintBroken InConstraintBrokenDelegate = FOnConstraintBroken());
+	void InitConstraint(FBodyInstance* Body1, FBodyInstance* Body2, float Scale, UObject* DebugOwner, FOnConstraintBroken InConstraintBrokenDelegate = FOnConstraintBroken(), FOnPlasticDeformation InPlasticDeformationDelegate = FOnPlasticDeformation());
 
 	/** Create physics engine constraint using physx actors. */
-	void InitConstraint_AssumesLocked(const FPhysicsActorHandle& ActorRef1, const FPhysicsActorHandle& ActorRef2, float InScale, FOnConstraintBroken InConstraintBrokenDelegate = FOnConstraintBroken());
+	void InitConstraint_AssumesLocked(const FPhysicsActorHandle& ActorRef1, const FPhysicsActorHandle& ActorRef2, float InScale, FOnConstraintBroken InConstraintBrokenDelegate = FOnConstraintBroken(), FOnPlasticDeformation InPlasticDeformationDelegate = FOnPlasticDeformation());
 
 	/** Terminate physics engine constraint */
 	void TermConstraint();
diff --git a/Engine/Source/Runtime/Engine/Classes/PhysicsEngine/PhysicsConstraintComponent.h b/Engine/Source/Runtime/Engine/Classes/PhysicsEngine/PhysicsConstraintComponent.h
--- a/Engine/Source/Runtime/Engine/Classes/PhysicsEngine/PhysicsConstraintComponent.h
+++ b/Engine/Source/Runtime/Engine/Classes/PhysicsEngine/PhysicsConstraintComponent.h
@@ -58,6 +58,10 @@ class ENGINE_API UPhysicsConstraintComponent : public USceneComponent
 	UPROPERTY(BlueprintAssignable)
 	FConstraintBrokenSignature OnConstraintBroken;
 
+	/** Notification when constraint plasticity drive target changes. */
+	UPROPERTY(BlueprintAssignable)
+	FPlasticDeformationEventSignature OnPlasticDeformation;
+
 public:
 	/** All constraint settings */
 	UPROPERTY(EditAnywhere, Category=ConstraintComponent, meta=(ShowOnlyInnerProperties))
@@ -362,11 +366,17 @@ protected:
 	/** Routes the FConstraint callback to the dynamic delegate */
 	void OnConstraintBrokenHandler(FConstraintInstance* BrokenConstraint);
 
+	/** Routes the FConstraint callback to the dynamic delegate */
+	void OnPlasticDeformationHandler(FConstraintInstance* Constraint);
+
 	/** Returns the scale of the constraint as it will be passed into the ConstraintInstance*/
 	float GetConstraintScale() const;
 
 private:
 	/** Wrapper that calls our constraint broken delegate */
 	void OnConstraintBrokenWrapper(int32 ConstraintIndex);
+
+	/** Wrapper that calls our constraint plasticity delegate */
+	void OnPlasticDeformationWrapper(int32 ConstraintIndex);
 };
 
diff --git a/Engine/Source/Runtime/Engine/Private/PhysicsEngine/ConstraintInstance.cpp b/Engine/Source/Runtime/Engine/Private/PhysicsEngine/ConstraintInstance.cpp
--- a/Engine/Source/Runtime/Engine/Private/PhysicsEngine/ConstraintInstance.cpp
+++ b/Engine/Source/Runtime/Engine/Private/PhysicsEngine/ConstraintInstance.cpp
@@ -458,7 +458,7 @@ void FConstraintInstance::UpdateAverageMass_AssumesLocked(const FPhysicsActorHan
 /** 
  *	Create physics engine constraint.
  */
-void FConstraintInstance::InitConstraint(FBodyInstance* Body1, FBodyInstance* Body2, float InScale, UObject* DebugOwner, FOnConstraintBroken InConstraintBrokenDelegate)
+void FConstraintInstance::InitConstraint(FBodyInstance* Body1, FBodyInstance* Body2, float InScale, UObject* DebugOwner, FOnConstraintBroken InConstraintBrokenDelegate, FOnPlasticDeformation InPlasticDeformationDelegate)
 {
 	FPhysicsActorHandle Actor1;
 	FPhysicsActorHandle Actor2;
@@ -472,15 +472,16 @@ void FConstraintInstance::InitConstraint(FBodyInstance* Body1, FBodyInstance* Bo
 
 		FPhysicsCommand::ExecuteWrite(Actor1, Actor2, [&](const FPhysicsActorHandle& ActorA, const FPhysicsActorHandle& ActorB)
 		{
-			InitConstraint_AssumesLocked(ActorA, ActorB, InScale, InConstraintBrokenDelegate);
+			InitConstraint_AssumesLocked(ActorA, ActorB, InScale, InConstraintBrokenDelegate, InPlasticDeformationDelegate);
 		});
 	}
 
 }
 
-void FConstraintInstance::InitConstraint_AssumesLocked(const FPhysicsActorHandle& ActorRef1, const FPhysicsActorHandle& ActorRef2, float InScale, FOnConstraintBroken InConstraintBrokenDelegate)
+void FConstraintInstance::InitConstraint_AssumesLocked(const FPhysicsActorHandle& ActorRef1, const FPhysicsActorHandle& ActorRef2, float InScale, FOnConstraintBroken InConstraintBrokenDelegate, FOnPlasticDeformation InPlasticDeformationDelegate)
 {
 	OnConstraintBrokenDelegate = InConstraintBrokenDelegate;
+	OnPlasticDeformationDelegate = InPlasticDeformationDelegate;
 	LastKnownScale = InScale;
 
 	UserData = FChaosUserData(this);
diff --git a/Engine/Source/Runtime/Engine/Private/PhysicsEngine/Experimental/PhysScene_Chaos.cpp b/Engine/Source/Runtime/Engine/Private/PhysicsEngine/Experimental/PhysScene_Chaos.cpp
--- a/Engine/Source/Runtime/Engine/Private/PhysicsEngine/Experimental/PhysScene_Chaos.cpp
+++ b/Engine/Source/Runtime/Engine/Private/PhysicsEngine/Experimental/PhysScene_Chaos.cpp
@@ -1677,6 +1677,17 @@ void FPhysScene_Chaos::OnSyncBodies(Chaos::FPhysicsSolverBase* Solver)
 				Constraint->GetOutputData().bIsBreaking = false;
 			}
 
+			if (Constraint->GetOutputData().bDriveTargetChanged)
+			{
+				if (FConstraintInstanceBase* ConstraintInstance = (Constraint) ? FPhysicsUserData_Chaos::Get<FConstraintInstanceBase>(Constraint->GetUserData()) : nullptr)
+				{
+					FPlasticDeformationDelegateWrapper CPD(ConstraintInstance);
+					CPD.DispatchPlasticDeformation();
+				}
+
+				Constraint->GetOutputData().bDriveTargetChanged = false;
+			}
+
 		};
 
 		Solver->PullPhysicsStateForEachDirtyProxy_External(RigidLambda, ConstraintLambda);
@@ -1734,6 +1745,18 @@ void FConstraintBrokenDelegateWrapper::DispatchOnBroken()
 	OnConstraintBrokenDelegate.ExecuteIfBound(ConstraintIndex);
 }
 
+FPlasticDeformationDelegateWrapper::FPlasticDeformationDelegateWrapper(FConstraintInstanceBase* ConstraintInstance)
+	: OnPlasticDeformationDelegate(ConstraintInstance->OnPlasticDeformationDelegate)
+	, ConstraintIndex(ConstraintInstance->ConstraintIndex)
+{
+
+}
+
+void FPlasticDeformationDelegateWrapper::DispatchPlasticDeformation()
+{
+	OnPlasticDeformationDelegate.ExecuteIfBound(ConstraintIndex);
+}
+
 void FPhysScene_Chaos::ResimNFrames(const int32 NumFramesRequested)
 {
 	//needs to run on physics thread from a special location
diff --git a/Engine/Source/Runtime/Engine/Private/PhysicsEngine/PhysicsConstraintComponent.cpp b/Engine/Source/Runtime/Engine/Private/PhysicsEngine/PhysicsConstraintComponent.cpp
--- a/Engine/Source/Runtime/Engine/Private/PhysicsEngine/PhysicsConstraintComponent.cpp
+++ b/Engine/Source/Runtime/Engine/Private/PhysicsEngine/PhysicsConstraintComponent.cpp
@@ -244,6 +244,12 @@ void UPhysicsConstraintComponent::OnConstraintBrokenWrapper(int32 ConstraintInde
 	OnConstraintBroken.Broadcast(ConstraintIndex);
 }
 
+/** Wrapper that calls our plasticity delegate */
+void UPhysicsConstraintComponent::OnPlasticDeformationWrapper(int32 ConstraintIndex)
+{
+	OnPlasticDeformation.Broadcast(ConstraintIndex);
+}
+
 void UPhysicsConstraintComponent::InitComponentConstraint()
 {
 	// First we convert world space position of constraint into local space frames
@@ -269,6 +275,12 @@ void UPhysicsConstraintComponent::OnConstraintBrokenHandler(FConstraintInstance*
 	OnConstraintBroken.Broadcast(BrokenConstraint->ConstraintIndex);
 }
 
+void UPhysicsConstraintComponent::OnPlasticDeformationHandler(FConstraintInstance* Constraint)
+{
+	OnPlasticDeformation.Broadcast(Constraint->ConstraintIndex);
+}
+
+
 float UPhysicsConstraintComponent::GetConstraintScale() const
 {
 	return GetComponentScale().GetAbsMin();
diff --git a/Engine/Source/Runtime/Engine/Private/SkeletalMeshComponentPhysics.cpp b/Engine/Source/Runtime/Engine/Private/SkeletalMeshComponentPhysics.cpp
--- a/Engine/Source/Runtime/Engine/Private/SkeletalMeshComponentPhysics.cpp
+++ b/Engine/Source/Runtime/Engine/Private/SkeletalMeshComponentPhysics.cpp
@@ -559,6 +559,11 @@ void USkeletalMeshComponent::OnConstraintBrokenWrapper(int32 ConstraintIndex)
 	OnConstraintBroken.Broadcast(ConstraintIndex);
 }
 
+void USkeletalMeshComponent::OnPlasticDeformationWrapper(int32 ConstraintIndex)
+{
+	OnPlasticDeformation.Broadcast(ConstraintIndex);
+}
+
 void USkeletalMeshComponent::InitCollisionRelationships()
 {
 #if WITH_CHAOS
@@ -956,7 +961,9 @@ void USkeletalMeshComponent::InstantiatePhysicsAsset_Internal(const UPhysicsAsse
 			{
 				ScalePosition(Body1, Scale, ConInst->Pos1);
 				ScalePosition(Body2, Scale, ConInst->Pos2);
-				ConInst->InitConstraint(Body1, Body2, Scale, OwningComponent, OwningComponent ? FOnConstraintBroken::CreateUObject(OwningComponent, &USkeletalMeshComponent::OnConstraintBrokenWrapper) : FOnConstraintBroken());
+				ConInst->InitConstraint(Body1, Body2, Scale, OwningComponent
+					, OwningComponent ? FOnConstraintBroken::CreateUObject(OwningComponent, &USkeletalMeshComponent::OnConstraintBrokenWrapper) : FOnConstraintBroken()
+					, OwningComponent ? FOnPlasticDeformation::CreateUObject(OwningComponent, &USkeletalMeshComponent::OnPlasticDeformationWrapper) : FOnPlasticDeformation());
 			}
 		}
 	}
diff --git a/Engine/Source/Runtime/Engine/Public/Physics/Experimental/PhysScene_Chaos.h b/Engine/Source/Runtime/Engine/Public/Physics/Experimental/PhysScene_Chaos.h
--- a/Engine/Source/Runtime/Engine/Public/Physics/Experimental/PhysScene_Chaos.h
+++ b/Engine/Source/Runtime/Engine/Public/Physics/Experimental/PhysScene_Chaos.h
@@ -76,6 +76,15 @@ struct FConstraintBrokenDelegateWrapper
 	int32 ConstraintIndex;
 };
 
+struct FPlasticDeformationDelegateWrapper
+{
+	FPlasticDeformationDelegateWrapper(FConstraintInstanceBase* ConstraintInstance);
+
+	void DispatchPlasticDeformation();
+
+	FOnPlasticDeformation OnPlasticDeformationDelegate;
+	int32 ConstraintIndex;
+};
 
 /**
 * Low level Chaos scene used when building custom simulations that don't exist in the main world physics scene.
diff --git a/Engine/Source/Runtime/Experimental/Chaos/Private/Chaos/PBDJointConstraints.cpp b/Engine/Source/Runtime/Experimental/Chaos/Private/Chaos/PBDJointConstraints.cpp
--- a/Engine/Source/Runtime/Experimental/Chaos/Private/Chaos/PBDJointConstraints.cpp
+++ b/Engine/Source/Runtime/Experimental/Chaos/Private/Chaos/PBDJointConstraints.cpp
@@ -84,6 +84,16 @@ namespace Chaos
 		return ConcreteContainer()->ClearConstraintBreaking(ConstraintIndex);
 	}
 
+	bool FPBDJointConstraintHandle::IsDriveTargetChanged() const
+	{
+		return ConcreteContainer()->IsDriveTargetChanged(ConstraintIndex);
+	}
+
+	void FPBDJointConstraintHandle::ClearDriveTargetChanged()
+	{
+		return ConcreteContainer()->ClearDriveTargetChanged(ConstraintIndex);
+	}
+
 	FVec3 FPBDJointConstraintHandle::GetLinearImpulse() const
 	{
 		return ConcreteContainer()->GetConstraintLinearImpulse(ConstraintIndex);
@@ -269,6 +279,8 @@ namespace Chaos
 		, Color(INDEX_NONE)
 		, IslandSize(0)
 		, bDisabled(false)
+		, bBreaking(false)
+		, bDriveTargetChanged(false)
 		, LinearImpulse(FVec3(0))
 		, AngularImpulse(FVec3(0))
 	{
@@ -520,6 +532,16 @@ namespace Chaos
 		ConstraintStates[ConstraintIndex].bBreaking = false;
 	}
 
+	bool FPBDJointConstraints::IsDriveTargetChanged(int32 ConstraintIndex) const
+	{
+		return ConstraintStates[ConstraintIndex].bDriveTargetChanged;
+	}
+
+	void FPBDJointConstraints::ClearDriveTargetChanged(int32 ConstraintIndex)
+	{
+		ConstraintStates[ConstraintIndex].bDriveTargetChanged = false;
+	}
+
 	void FPBDJointConstraints::SetConstraintEnabled(int32 ConstraintIndex, bool bEnabled)
 	{
 		const FGenericParticleHandle Particle0 = FGenericParticleHandle(ConstraintParticles[ConstraintIndex][0]);
@@ -546,6 +568,11 @@ namespace Chaos
 		ConstraintStates[ConstraintIndex].bBreaking = bBreaking;
 	}
 
+	void FPBDJointConstraints::SetDriveTargetChanged(int32 ConstraintIndex, bool bTargetChanged)
+	{
+		ConstraintStates[ConstraintIndex].bDriveTargetChanged = bTargetChanged;
+	}
+
 	void FPBDJointConstraints::BreakConstraint(int32 ConstraintIndex)
 	{
 		SetConstraintEnabled(ConstraintIndex, false);
@@ -1267,6 +1294,7 @@ namespace Chaos
 				if (JointSettings.LinearPlasticityType == EPlasticityType::Free)
 				{
 					JointSettings.LinearDrivePositionTarget = LinearDisplacement;
+					SetDriveTargetChanged(ConstraintIndex, true);
 				}
 				else // EPlasticityType::Shrink || EPlasticityType::Grow
 				{
@@ -1278,10 +1306,12 @@ namespace Chaos
 					if (JointSettings.LinearPlasticityType == EPlasticityType::Shrink && CurrentDelta.SizeSquared() < StartDelta.SizeSquared())
 					{
 						JointSettings.LinearDrivePositionTarget = LinearDisplacement;
+						SetDriveTargetChanged(ConstraintIndex, true);
 					}
 					else if (JointSettings.LinearPlasticityType == EPlasticityType::Grow && CurrentDelta.SizeSquared() > StartDelta.SizeSquared())
 					{
 						JointSettings.LinearDrivePositionTarget = LinearDisplacement;
+						SetDriveTargetChanged(ConstraintIndex, true);
 					}
 				}
 			}
@@ -1307,6 +1337,7 @@ namespace Chaos
 			if (AngleDeg > JointSettings.AngularPlasticityLimit)
 			{
 				JointSettings.AngularDrivePositionTarget = AngularDisplacement;
+				SetDriveTargetChanged(ConstraintIndex, true);
 			}
 		}
 	}
diff --git a/Engine/Source/Runtime/Experimental/Chaos/Private/PhysicsProxy/JointConstraintProxy.cpp b/Engine/Source/Runtime/Experimental/Chaos/Private/PhysicsProxy/JointConstraintProxy.cpp
--- a/Engine/Source/Runtime/Experimental/Chaos/Private/PhysicsProxy/JointConstraintProxy.cpp
+++ b/Engine/Source/Runtime/Experimental/Chaos/Private/PhysicsProxy/JointConstraintProxy.cpp
@@ -47,14 +47,16 @@ void FJointConstraintPhysicsProxy::BufferPhysicsResults(Chaos::FDirtyJointConstr
 	Buffer.SetProxy(*this);
 	if (Constraint != nullptr && Constraint->IsValid() )
 	{
-		if (Handle != nullptr && (Handle->IsValid() || Handle->IsConstraintBreaking()))
+		if (Handle != nullptr && (Handle->IsValid() || Handle->IsConstraintBreaking() || Handle->IsDriveTargetChanged()))
 		{
 			Buffer.OutputData.bIsBreaking = Handle->IsConstraintBreaking();
 			Buffer.OutputData.bIsBroken = !Handle->IsConstraintEnabled();
+			Buffer.OutputData.bDriveTargetChanged = Handle->IsDriveTargetChanged();
 			Buffer.OutputData.Force = Handle->GetLinearImpulse();
 			Buffer.OutputData.Torque = Handle->GetAngularImpulse();
 
 			Handle->ClearConstraintBreaking(); // it's a single frame event, so reset
+			Handle->ClearDriveTargetChanged(); // it's a single frame event, so reset
 		}
 	}
 }
@@ -64,10 +66,11 @@ bool FJointConstraintPhysicsProxy::PullFromPhysicsState(const Chaos::FDirtyJoint
 {
 	if (Constraint != nullptr && Constraint->IsValid())
 	{
-		if (Handle != nullptr && (Handle->IsValid() || Buffer.OutputData.bIsBreaking))
+		if (Handle != nullptr && (Handle->IsValid() || Buffer.OutputData.bIsBreaking || Buffer.OutputData.bDriveTargetChanged))
 		{
 			Constraint->GetOutputData().bIsBreaking = Buffer.OutputData.bIsBreaking;
 			Constraint->GetOutputData().bIsBroken = Buffer.OutputData.bIsBroken;
+			Constraint->GetOutputData().bDriveTargetChanged = Buffer.OutputData.bDriveTargetChanged;
 			Constraint->GetOutputData().Force = Buffer.OutputData.Force;
 			Constraint->GetOutputData().Torque = Buffer.OutputData.Torque;
 		}
diff --git a/Engine/Source/Runtime/Experimental/Chaos/Public/Chaos/PBDJointConstraintData.h b/Engine/Source/Runtime/Experimental/Chaos/Public/Chaos/PBDJointConstraintData.h
--- a/Engine/Source/Runtime/Experimental/Chaos/Public/Chaos/PBDJointConstraintData.h
+++ b/Engine/Source/Runtime/Experimental/Chaos/Public/Chaos/PBDJointConstraintData.h
@@ -304,6 +304,7 @@ namespace Chaos
 			// Output properties
 			bool bIsBreaking = false;
 			bool bIsBroken = false;
+			bool bDriveTargetChanged = false;
 			FVector Force = FVector(0);
 			FVector Torque = FVector(0);
 		};
diff --git a/Engine/Source/Runtime/Experimental/Chaos/Public/Chaos/PBDJointConstraints.h b/Engine/Source/Runtime/Experimental/Chaos/Public/Chaos/PBDJointConstraints.h
--- a/Engine/Source/Runtime/Experimental/Chaos/Public/Chaos/PBDJointConstraints.h
+++ b/Engine/Source/Runtime/Experimental/Chaos/Public/Chaos/PBDJointConstraints.h
@@ -42,6 +42,8 @@ namespace Chaos
 		bool IsConstraintEnabled() const;
 		bool IsConstraintBreaking() const;
 		void ClearConstraintBreaking();
+		bool IsDriveTargetChanged() const;
+		void ClearDriveTargetChanged();
 		FVec3 GetLinearImpulse() const;
 		FVec3 GetAngularImpulse() const;
 
@@ -66,6 +68,7 @@ namespace Chaos
 		int32 IslandSize;
 		bool bDisabled;
 		bool bBreaking;
+		bool bDriveTargetChanged;
 		FVec3 LinearImpulse;
 		FVec3 AngularImpulse;
 	};
@@ -143,6 +146,16 @@ namespace Chaos
 		 */
 		void ClearConstraintBreaking(int32 ConstraintIndex);
 
+		/*
+		 * Whether the drive target has changed
+		 */
+		bool IsDriveTargetChanged(int32 ConstraintIndex) const;
+
+		/*
+		 * Clear the drive target state
+		 */
+		void ClearDriveTargetChanged(int32 ConstraintIndex);
+
 		/*
 		 * Enable or disable a constraints
 		 */
@@ -153,6 +166,11 @@ namespace Chaos
 		 */
 		void SetConstraintBreaking(int32 ConstraintIndex, bool bBreaking);
 
+		/*
+		* Set Drive Target Changed State
+		*/
+		void SetDriveTargetChanged(int32 ConstraintIndex, bool bTargetChanged);
+
 		/*
 		 * Force a constraints to break
 		 */
diff --git a/Engine/Source/Runtime/Experimental/Chaos/Public/Chaos/PullPhysicsDataImp.h b/Engine/Source/Runtime/Experimental/Chaos/Public/Chaos/PullPhysicsDataImp.h
--- a/Engine/Source/Runtime/Experimental/Chaos/Public/Chaos/PullPhysicsDataImp.h
+++ b/Engine/Source/Runtime/Experimental/Chaos/Public/Chaos/PullPhysicsDataImp.h
@@ -64,6 +64,7 @@ struct FDirtyGeometryCollectionData : public TBasePullData<FGeometryCollectionPh
 struct FJointConstraintOutputData {
 	bool bIsBreaking = false;
 	bool bIsBroken = false;
+	bool bDriveTargetChanged = false;
 	FVector Force = FVector(0);
 	FVector Torque = FVector(0);
 };
