CIS - Fix FReal to float implicit conversion.

diff --git a/Engine/Source/Runtime/Experimental/Chaos/Public/Chaos/PBDLongRangeConstraintsBase.h b/Engine/Source/Runtime/Experimental/Chaos/Public/Chaos/PBDLongRangeConstraintsBase.h
--- a/Engine/Source/Runtime/Experimental/Chaos/Public/Chaos/PBDLongRangeConstraintsBase.h
+++ b/Engine/Source/Runtime/Experimental/Chaos/Public/Chaos/PBDLongRangeConstraintsBase.h
@@ -45,7 +45,7 @@ public:
 	void SetStiffness(const FVec2& InStiffness) { Stiffness.SetWeightedValue(InStiffness); }
 
 	// Set the scale low and high value of the scale weight map
-	void SetScale(const FVec2& InScale) { Scale = InScale.ClampAxes((FReal)0.01, (FReal)10.); }
+	void SetScale(const FVec2& InScale) { Scale = FVec2(FMath::Clamp(InScale.X, (FReal)0.01, (FReal)10.), FMath::Clamp(InScale.Y, (FReal)0.01, (FReal)10.)); }
 
 	// Set stiffness offset and range, as well as the simulation stiffness exponent
 	void ApplyProperties(const FReal Dt, const int32 NumIterations);
