Chaos: fix crash in geometry collection proxy when when opening Editor

diff --git a/Engine/Source/Runtime/Experimental/Chaos/Private/PhysicsProxy/GeometryCollectionPhysicsProxy.cpp b/Engine/Source/Runtime/Experimental/Chaos/Private/PhysicsProxy/GeometryCollectionPhysicsProxy.cpp
--- a/Engine/Source/Runtime/Experimental/Chaos/Private/PhysicsProxy/GeometryCollectionPhysicsProxy.cpp
+++ b/Engine/Source/Runtime/Experimental/Chaos/Private/PhysicsProxy/GeometryCollectionPhysicsProxy.cpp
@@ -1815,8 +1815,8 @@ void FGeometryCollectionPhysicsProxy::InitializeSharedCollisionStructures(
 	TManagedArray<FVector3f>& CollectionInertiaTensor =
 		RestCollection.AddAttribute<FVector3f>(
 			TEXT("InertiaTensor"), FTransformCollection::TransformGroup);
-	TManagedArray<FReal>& CollectionMass =
-		RestCollection.AddAttribute<FReal>(
+	TManagedArray<FRealSingle>& CollectionMass =
+		RestCollection.AddAttribute<FRealSingle>(
 			TEXT("Mass"), FTransformCollection::TransformGroup);
 	TManagedArray<TUniquePtr<FSimplicial>>& CollectionSimplicials =
 		RestCollection.AddAttribute<TUniquePtr<FSimplicial>>(
@@ -2058,7 +2058,7 @@ void FGeometryCollectionPhysicsProxy::InitializeSharedCollisionStructures(
 
 			const FReal Mass_i = FMath::Max(DesiredDensity * Volume_i, SharedParams.MinimumMassClamp);
 			const FReal Density_i = Mass_i / Volume_i;
-			CollectionMass[TransformGroupIndex] = Mass_i;
+			CollectionMass[TransformGroupIndex] = (FRealSingle)Mass_i;
 
 			if (InertiaComputationNeeded[GeometryIndex])
 			{
@@ -2393,7 +2393,7 @@ void FGeometryCollectionPhysicsProxy::InitializeSharedCollisionStructures(
 				const FMatrix& InertiaMatrix = CollectionSpaceParticles->I(ClusterTransformIdx);
 				const FVector InertiaDiagonal(InertiaMatrix.M[0][0], InertiaMatrix.M[1][1], InertiaMatrix.M[2][2]);
 				CollectionInertiaTensor[ClusterTransformIdx] = (FVector3f)InertiaDiagonal;	// LWC_TODO: Precision loss
-				CollectionMass[ClusterTransformIdx] = CollectionSpaceParticles->M(ClusterTransformIdx);
+				CollectionMass[ClusterTransformIdx] = (FRealSingle)CollectionSpaceParticles->M(ClusterTransformIdx);
 
 				const int32 SizeSpecificIdx = FindSizeSpecificIdx(SharedParams.SizeSpecificData, InstanceBoundingBox);
 				const FSharedSimulationSizeSpecificData& SizeSpecificData = SharedParams.SizeSpecificData[SizeSpecificIdx];
