GeometryCollectionComponent
-- Removed restrictions on RestPose setting to allow ChaosCaching to set rest pose during play.

diff --git a/Engine/Source/Runtime/Experimental/GeometryCollectionEngine/Private/GeometryCollection/GeometryCollectionComponent.cpp b/Engine/Source/Runtime/Experimental/GeometryCollectionEngine/Private/GeometryCollection/GeometryCollectionComponent.cpp
--- a/Engine/Source/Runtime/Experimental/GeometryCollectionEngine/Private/GeometryCollection/GeometryCollectionComponent.cpp
+++ b/Engine/Source/Runtime/Experimental/GeometryCollectionEngine/Private/GeometryCollection/GeometryCollectionComponent.cpp
@@ -860,54 +860,50 @@ void UGeometryCollectionComponent::RefreshEmbeddedGeometry()
 
 void UGeometryCollectionComponent::SetRestState(TArray<FTransform>&& InRestTransforms)
 {
-	// This method should never be called on a simulating component.
-	if (!PhysicsProxy)
-	{ 
-		RestTransforms = InRestTransforms;
+	RestTransforms = InRestTransforms;
 	
-		CalculateGlobalMatrices();
+	CalculateGlobalMatrices();
 
-		if (DynamicCollection)
-		{
-			SetInitialTransforms(RestTransforms);
-		}
+	if (DynamicCollection)
+	{
+		SetInitialTransforms(RestTransforms);
+	}
 
-		FGeometryCollectionDynamicData* DynamicData = GDynamicDataPool.Allocate();
-		DynamicData->SetAllTransforms(GlobalMatrices);
-		DynamicData->IsDynamic = true;
+	FGeometryCollectionDynamicData* DynamicData = GDynamicDataPool.Allocate();
+	DynamicData->SetAllTransforms(GlobalMatrices);
+	DynamicData->IsDynamic = true;
 
-		if (SceneProxy)
+	if (SceneProxy)
+	{
+		if (SceneProxy->IsNaniteMesh())
 		{
-			if (SceneProxy->IsNaniteMesh())
-			{
 
-	#if WITH_EDITOR
-				// We need to do this in case we're controlled by Sequencer in editor, which doesn't invoke PostEditChangeProperty
-				SendRenderTransform_Concurrent();
-	#endif
+#if WITH_EDITOR
+			// We need to do this in case we're controlled by Sequencer in editor, which doesn't invoke PostEditChangeProperty
+			SendRenderTransform_Concurrent();
+#endif
 
-				FNaniteGeometryCollectionSceneProxy* GeometryCollectionSceneProxy = static_cast<FNaniteGeometryCollectionSceneProxy*>(SceneProxy);
-				ENQUEUE_RENDER_COMMAND(SendRenderDynamicData)(
-					[GeometryCollectionSceneProxy, DynamicData](FRHICommandListImmediate& RHICmdList)
-					{
-						GeometryCollectionSceneProxy->SetDynamicData_RenderThread(DynamicData);
-					}
-				);
-			}
-			else
-			{
-				FGeometryCollectionSceneProxy* GeometryCollectionSceneProxy = static_cast<FGeometryCollectionSceneProxy*>(SceneProxy);
-				ENQUEUE_RENDER_COMMAND(SendRenderDynamicData)(
-					[GeometryCollectionSceneProxy, DynamicData](FRHICommandListImmediate& RHICmdList)
-					{
-						GeometryCollectionSceneProxy->SetDynamicData_RenderThread(DynamicData);
-					}
-				);
-			}
+			FNaniteGeometryCollectionSceneProxy* GeometryCollectionSceneProxy = static_cast<FNaniteGeometryCollectionSceneProxy*>(SceneProxy);
+			ENQUEUE_RENDER_COMMAND(SendRenderDynamicData)(
+				[GeometryCollectionSceneProxy, DynamicData](FRHICommandListImmediate& RHICmdList)
+				{
+					GeometryCollectionSceneProxy->SetDynamicData_RenderThread(DynamicData);
+				}
+			);
+		}
+		else
+		{
+			FGeometryCollectionSceneProxy* GeometryCollectionSceneProxy = static_cast<FGeometryCollectionSceneProxy*>(SceneProxy);
+			ENQUEUE_RENDER_COMMAND(SendRenderDynamicData)(
+				[GeometryCollectionSceneProxy, DynamicData](FRHICommandListImmediate& RHICmdList)
+				{
+					GeometryCollectionSceneProxy->SetDynamicData_RenderThread(DynamicData);
+				}
+			);
 		}
-
-		RefreshEmbeddedGeometry();
 	}
+
+	RefreshEmbeddedGeometry();
 }
 
 void UGeometryCollectionComponent::InitializeComponent()
