GeometryCollection
-- prevent creation of particles for embedded geometry

diff --git a/Engine/Source/Runtime/Experimental/Chaos/Private/PhysicsProxy/GeometryCollectionPhysicsProxy.cpp b/Engine/Source/Runtime/Experimental/Chaos/Private/PhysicsProxy/GeometryCollectionPhysicsProxy.cpp
--- a/Engine/Source/Runtime/Experimental/Chaos/Private/PhysicsProxy/GeometryCollectionPhysicsProxy.cpp
+++ b/Engine/Source/Runtime/Experimental/Chaos/Private/PhysicsProxy/GeometryCollectionPhysicsProxy.cpp
@@ -581,7 +581,7 @@ void FGeometryCollectionPhysicsProxy::InitializeBodiesPT(Chaos::FPBDRigidsSolver
 		for (int32 Idx = 0; Idx < SimulatableParticles.Num(); ++Idx)
 		{
 			NumSimulatedParticles += SimulatableParticles[Idx];
-			if (SimulatableParticles[Idx] && !RestCollection->IsClustered(Idx))
+			if (SimulatableParticles[Idx] && !RestCollection->IsClustered(Idx) && RestCollection->IsGeometry(Idx))
 			{
 				NumRigids++;
 				UniqueIndices.Add(GTParticles[Idx]->UniqueIdx());
