Chaos: Fixed non-threadsafe data access in geometry collection proxy. GT particle was having its position updated after recalculation on the PT. Instead now the particle is force-queued as a dirty particle to ensure we get a position update where required.

diff --git a/Engine/Source/Runtime/Experimental/Chaos/Private/PhysicsProxy/GeometryCollectionPhysicsProxy.cpp b/Engine/Source/Runtime/Experimental/Chaos/Private/PhysicsProxy/GeometryCollectionPhysicsProxy.cpp
--- a/Engine/Source/Runtime/Experimental/Chaos/Private/PhysicsProxy/GeometryCollectionPhysicsProxy.cpp
+++ b/Engine/Source/Runtime/Experimental/Chaos/Private/PhysicsProxy/GeometryCollectionPhysicsProxy.cpp
@@ -805,18 +805,23 @@ void FGeometryCollectionPhysicsProxy::InitializeBodiesPT(Chaos::FPBDRigidsSolver
 						SolverClusterID[ChildTransformIndex] = RigidChildren[RigidChildrenIdx++]->CastToClustered()->ClusterIds().Id;;
 					}
 					SolverClusterID[TransformGroupIndex] = Handle->ClusterIds().Id;					
-					
-					// Cluster transform has been recalculated based on children - copy to the GT particle (not threadsafe - just testing)
-					GTParticle->SetX(Handle->X());
-					GTParticle->SetR(Handle->R());
-					GTParticle->UpdateShapeBounds();
 
 					SolverClusterHandles[TransformGroupIndex] = Handle;
 					SolverParticleHandles[TransformGroupIndex] = Handle;
 					HandleToTransformGroupIndex.Add(Handle, TransformGroupIndex);
 					Handle->SetPhysicsProxy(this);
 
+					// Dirty for SQ
 					RigidsSolver->GetEvolution()->DirtyParticle(*Handle);
+
+					// If we're not simulating we would normally not write any results back to the game thread.
+					// This will force a single write in this case because we've updated the transform on the cluster
+					// and it should be updated on the game thread also
+					// #TODO Consider building this information at edit-time / offline
+					if(!Parameters.Simulating)
+					{
+						RigidsSolver->GetEvolution()->GetParticles().MarkTransientDirtyParticle(Handle);
+					}
 				}
 			}
 
@@ -1749,6 +1754,7 @@ bool FGeometryCollectionPhysicsProxy::PullFromPhysicsState(const Chaos::FDirtyGe
 
 				GTParticles[TransformGroupIndex]->SetX(ParticleToWorld.GetTranslation());
 				GTParticles[TransformGroupIndex]->SetR(ParticleToWorld.GetRotation());
+				GTParticles[TransformGroupIndex]->UpdateShapeBounds();
 			}
 
 			DynamicCollection.DynamicState[TransformGroupIndex] = TargetResults.DynamicState[TransformGroupIndex];
