Check Field RootNode in code when bad BP setup to avoid UE to crash

diff --git a/Engine/Source/Runtime/Experimental/FieldSystem/Source/FieldSystemEngine/Private/Field/FieldSystemComponent.cpp b/Engine/Source/Runtime/Experimental/FieldSystem/Source/FieldSystemEngine/Private/Field/FieldSystemComponent.cpp
--- a/Engine/Source/Runtime/Experimental/FieldSystem/Source/FieldSystemEngine/Private/Field/FieldSystemComponent.cpp
+++ b/Engine/Source/Runtime/Experimental/FieldSystem/Source/FieldSystemEngine/Private/Field/FieldSystemComponent.cpp
@@ -137,7 +137,7 @@ bool UFieldSystemComponent::HasValidPhysicsState() const
 void UFieldSystemComponent::DispatchFieldCommand(const FFieldSystemCommand& InCommand, const bool IsTransient)
 {
 	using namespace Chaos;
-	if (HasValidPhysicsState())
+	if (HasValidPhysicsState() && InCommand.RootNode)
 	{
 		const FName Name = GetOwner() ? *GetOwner()->GetName() : TEXT("");
 		
@@ -281,7 +281,10 @@ void UFieldSystemComponent::AddFieldCommand(bool Enabled, EFieldPhysicsType Targ
 		if (FieldSystem)
 		{
 			FFieldSystemCommand Command = FFieldObjectCommands::CreateFieldCommand(Target, Field, MetaData);
-			SetupConstructionFields.Add(Command);
+			if (Command.RootNode)
+			{
+				SetupConstructionFields.Add(Command);
+			}
 		}
 
 		BufferCommands.AddFieldCommand(GetFieldPhysicsName(Target), Field, MetaData);
diff --git a/Engine/Source/Runtime/Experimental/GeometryCollectionEngine/Private/GeometryCollection/GeometryCollectionComponent.cpp b/Engine/Source/Runtime/Experimental/GeometryCollectionEngine/Private/GeometryCollection/GeometryCollectionComponent.cpp
--- a/Engine/Source/Runtime/Experimental/GeometryCollectionEngine/Private/GeometryCollection/GeometryCollectionComponent.cpp
+++ b/Engine/Source/Runtime/Experimental/GeometryCollectionEngine/Private/GeometryCollection/GeometryCollectionComponent.cpp
@@ -2445,7 +2445,7 @@ bool UGeometryCollectionComponent::GetIsObjectDynamic() const
 
 void UGeometryCollectionComponent::DispatchFieldCommand(const FFieldSystemCommand& InCommand)
 {
-	if (PhysicsProxy)
+	if (PhysicsProxy && InCommand.RootNode)
 	{
 		FChaosSolversModule* ChaosModule = FChaosSolversModule::GetModule();
 		checkSlow(ChaosModule);
@@ -2478,7 +2478,11 @@ void UGeometryCollectionComponent::GetInitializationCommands(TArray<FFieldSystem
 				{
 					for (int32 CommandIndex = 0; CommandIndex < NumCommands; ++CommandIndex)
 					{
-						CombinedCommmands.Add(FieldSystemActor->GetFieldSystemComponent()->ConstructionCommands.BuildFieldCommand(CommandIndex));
+						const FFieldSystemCommand NewCommand = FieldSystemActor->GetFieldSystemComponent()->ConstructionCommands.BuildFieldCommand(CommandIndex);
+						if (NewCommand.RootNode)
+						{
+							CombinedCommmands.Emplace(NewCommand);
+						}
 					}
 				}
 				// Legacy path : only there for old levels. New ones will have the commands directly stored onto the component
@@ -2487,14 +2491,17 @@ void UGeometryCollectionComponent::GetInitializationCommands(TArray<FFieldSystem
 					const FName Name = GetOwner() ? *GetOwner()->GetName() : TEXT("");
 					for (const FFieldSystemCommand& Command : FieldSystemActor->GetFieldSystemComponent()->GetFieldSystem()->Commands)
 					{
-						FFieldSystemCommand NewCommand = { Command.TargetAttribute, Command.RootNode->NewCopy() };
-						NewCommand.InitFieldNodes(0.0, Name);
-
-						for (auto& Elem : Command.MetaData)
+						if (Command.RootNode)
 						{
-							NewCommand.MetaData.Add(Elem.Key, TUniquePtr<FFieldSystemMetaData>(Elem.Value->NewCopy()));
+							FFieldSystemCommand NewCommand = { Command.TargetAttribute, Command.RootNode->NewCopy() };
+							NewCommand.InitFieldNodes(0.0, Name);
+
+							for (const TPair<FFieldSystemMetaData::EMetaType, TUniquePtr<FFieldSystemMetaData>>& Elem : Command.MetaData)
+							{
+								NewCommand.MetaData.Add(Elem.Key, TUniquePtr<FFieldSystemMetaData>(Elem.Value->NewCopy()));
+							}
+							CombinedCommmands.Emplace(NewCommand);
 						}
-						CombinedCommmands.Add(NewCommand);
 					}
 				}
 			}
