ToolsFramework: add UTransformProxy::OnTransformChangedUndoRedo delegate that is fired when the transform is updated by a FTransformProxyChange, to allow special-case handling in this situation
#rnx

diff --git a/Engine/Source/Runtime/Experimental/InteractiveToolsFramework/Private/BaseGizmos/TransformProxy.cpp b/Engine/Source/Runtime/Experimental/InteractiveToolsFramework/Private/BaseGizmos/TransformProxy.cpp
--- a/Engine/Source/Runtime/Experimental/InteractiveToolsFramework/Private/BaseGizmos/TransformProxy.cpp
+++ b/Engine/Source/Runtime/Experimental/InteractiveToolsFramework/Private/BaseGizmos/TransformProxy.cpp
@@ -192,6 +192,7 @@ void FTransformProxyChange::Apply(UObject* Object)
 	bool bSavedSetPivotMode = Proxy->bSetPivotMode;
 	Proxy->bSetPivotMode = bSetPivotMode;
 	Proxy->SetTransform(To);
+	Proxy->OnTransformChangedUndoRedo.Broadcast(Proxy, To);
 	Proxy->bSetPivotMode = bSavedSetPivotMode;
 }
 
@@ -202,6 +203,7 @@ void FTransformProxyChange::Revert(UObject* Object)
 	bool bSavedSetPivotMode = Proxy->bSetPivotMode;
 	Proxy->bSetPivotMode = bSetPivotMode;
 	Proxy->SetTransform(From);
+	Proxy->OnTransformChangedUndoRedo.Broadcast(Proxy, From);
 	Proxy->bSetPivotMode = bSavedSetPivotMode;
 }
 
diff --git a/Engine/Source/Runtime/Experimental/InteractiveToolsFramework/Public/BaseGizmos/TransformProxy.h b/Engine/Source/Runtime/Experimental/InteractiveToolsFramework/Public/BaseGizmos/TransformProxy.h
--- a/Engine/Source/Runtime/Experimental/InteractiveToolsFramework/Public/BaseGizmos/TransformProxy.h
+++ b/Engine/Source/Runtime/Experimental/InteractiveToolsFramework/Public/BaseGizmos/TransformProxy.h
@@ -96,6 +96,12 @@ public:
 	DECLARE_MULTICAST_DELEGATE_TwoParams(FOnTransformChanged, UTransformProxy*, FTransform);
 	FOnTransformChanged OnTransformChanged;
 
+	/**
+	 * This delegate is fired whenever the transform is changed by a FTransformProxyChange, ie on undo/redo.
+	 * OnTransformChanged is also fired in those cases, this extra event can be used when undo/redo state changes need to be handled specifically.
+	 */
+	FOnTransformChanged OnTransformChangedUndoRedo;
+
 	/** This delegate is fired when BeginTransformEditSequence() is called to indicate that a sequence of transform updates has started */
 	DECLARE_MULTICAST_DELEGATE_OneParam(FOnBeginTransformEdit, UTransformProxy*);
 	FOnBeginTransformEdit OnBeginTransformEdit;
