ToolsFramework: add ToolBuilder::PostBuildTool() and PostSetupTool(), call from ToolManager::DoPostBuild()/ToolManager::DoPostSetup(). These are hooks which simplify re-use of existing Tool implementations, not currently used internally.
#rnx

diff --git a/Engine/Source/Runtime/Experimental/InteractiveToolsFramework/Private/InteractiveToolManager.cpp b/Engine/Source/Runtime/Experimental/InteractiveToolsFramework/Private/InteractiveToolManager.cpp
--- a/Engine/Source/Runtime/Experimental/InteractiveToolsFramework/Private/InteractiveToolManager.cpp
+++ b/Engine/Source/Runtime/Experimental/InteractiveToolsFramework/Private/InteractiveToolManager.cpp
@@ -56,11 +56,13 @@ void UInteractiveToolManager::Shutdown()
 
 void UInteractiveToolManager::DoPostBuild(EToolSide Side, UInteractiveTool* InBuiltTool, UInteractiveToolBuilder* InToolBuilder, const FToolBuilderState& InBuilderState)
 {
+	InToolBuilder->PostBuildTool(InBuiltTool, InBuilderState);
 	OnToolPostBuild.Broadcast(this, Side, InBuiltTool, InToolBuilder, InBuilderState);
 }
 
-void UInteractiveToolManager::DoPostSetup(EToolSide Side, UInteractiveTool* InInteractiveTool)
+void UInteractiveToolManager::DoPostSetup(EToolSide Side, UInteractiveTool* InInteractiveTool, UInteractiveToolBuilder* InToolBuilder, const FToolBuilderState& InBuilderState)
 {
+	InToolBuilder->PostSetupTool(InInteractiveTool, InBuilderState);
 	OnToolPostSetup.Broadcast(this, Side, InInteractiveTool);
 }
 
@@ -205,7 +207,7 @@ bool UInteractiveToolManager::ActivateToolInternal(EToolSide Side)
 	DoPostBuild(Side, ActiveLeftTool, ActiveLeftBuilder, InputState);
 
 	ActiveLeftTool->Setup();
-	DoPostSetup(Side, ActiveLeftTool);
+	DoPostSetup(Side, ActiveLeftTool, ActiveLeftBuilder, InputState);
 
 	// register new active input behaviors
 	InputRouter->RegisterSource(ActiveLeftTool);
diff --git a/Engine/Source/Runtime/Experimental/InteractiveToolsFramework/Public/InteractiveToolBuilder.h b/Engine/Source/Runtime/Experimental/InteractiveToolsFramework/Public/InteractiveToolBuilder.h
--- a/Engine/Source/Runtime/Experimental/InteractiveToolsFramework/Public/InteractiveToolBuilder.h
+++ b/Engine/Source/Runtime/Experimental/InteractiveToolsFramework/Public/InteractiveToolBuilder.h
@@ -45,6 +45,29 @@ public:
 		return nullptr;
 	}
 
+
+	/**
+	 * PostBuildTool is called by the ToolManager after a new Tool instance has been built.
+	 * This allows clients to extend a ToolBuilder with additional customization/etc, instead of having to entirely replace the existing BuildTool()
+	 * @param Tool the new Tool instance
+	 * @param SceneState the scene state used to create the new Tool
+	 */
+	virtual void PostBuildTool(UInteractiveTool* Tool, const FToolBuilderState& SceneState) const
+	{
+		return;
+	}
+
+	/**
+	 * PostBuildTool is called by the ToolManager after a new Tool instance has been built and Setup() has been called.
+	 * This allows clients to make modifications to Tools like changing initial parameters without subclassing the Tool
+	 * @param Tool the new Tool instance
+	 * @param SceneState the scene state used to create the new Tool
+	 */
+	virtual void PostSetupTool(UInteractiveTool* Tool, const FToolBuilderState& SceneState) const
+	{
+		return;
+	}
+
 protected:
 	/**
 	 * Gives the target requirements of the associated tool. Usually, it is the tool builder
diff --git a/Engine/Source/Runtime/Experimental/InteractiveToolsFramework/Public/InteractiveToolManager.h b/Engine/Source/Runtime/Experimental/InteractiveToolsFramework/Public/InteractiveToolManager.h
--- a/Engine/Source/Runtime/Experimental/InteractiveToolsFramework/Public/InteractiveToolManager.h
+++ b/Engine/Source/Runtime/Experimental/InteractiveToolsFramework/Public/InteractiveToolManager.h
@@ -78,7 +78,7 @@ protected:
 	virtual void DoPostBuild(EToolSide Side, UInteractiveTool* InBuiltTool, UInteractiveToolBuilder* InToolBuilder, const FToolBuilderState& InBuilderState);
 	
 	/** Called immediately after a tool's Setup is called. Broadcasts OnToolPostSetup. */
-	virtual void DoPostSetup(EToolSide Side, UInteractiveTool* InInteractiveTool);
+	virtual void DoPostSetup(EToolSide Side, UInteractiveTool* InInteractiveTool, UInteractiveToolBuilder* InToolBuilder, const FToolBuilderState& InBuilderState);
 
 public:
 
