Chaos : Plasticity
- Reversed the parent/child indexing to allow the plasticity rest length from the joint position to the child.

#okforgithum public, internal

diff --git a/Engine/Source/Runtime/Experimental/Chaos/Private/Chaos/PBDJointConstraints.cpp b/Engine/Source/Runtime/Experimental/Chaos/Private/Chaos/PBDJointConstraints.cpp
--- a/Engine/Source/Runtime/Experimental/Chaos/Private/Chaos/PBDJointConstraints.cpp
+++ b/Engine/Source/Runtime/Experimental/Chaos/Private/Chaos/PBDJointConstraints.cpp
@@ -766,10 +766,10 @@ namespace Chaos
 					// Joint plasticity is baed on the distance of one of the moment arms of the joint. Typically, plasticity
 					// will get setup from the joint pivot to the child COM (centor of mass), so that is found first. However, when 
 					// the pivot is at the child COM then we fall back to the distance between thge pivot and parent COM.
-					ConstraintSettings[JointIndex].LinearPlasticityInitialDistanceSquared = JointSettings.ConnectorTransforms[1].GetTranslation().SizeSquared();
+					ConstraintSettings[JointIndex].LinearPlasticityInitialDistanceSquared = JointSettings.ConnectorTransforms[Index1].GetTranslation().SizeSquared();
 					if (FMath::IsNearlyZero(ConstraintSettings[JointIndex].LinearPlasticityInitialDistanceSquared))
 					{
-						ConstraintSettings[JointIndex].LinearPlasticityInitialDistanceSquared = JointSettings.ConnectorTransforms[0].GetTranslation().SizeSquared();
+						ConstraintSettings[JointIndex].LinearPlasticityInitialDistanceSquared = JointSettings.ConnectorTransforms[Index0].GetTranslation().SizeSquared();
 					}
 					ensureMsgf(!FMath::IsNearlyZero(ConstraintSettings[JointIndex].LinearPlasticityInitialDistanceSquared), TEXT("Plasticity made inactive due to Zero length difference between parent and child rigid body."));
 				}
@@ -1300,8 +1300,8 @@ namespace Chaos
 				{
 					// Shrink and Grow are based on the distance between the joint pivot and the child. 
 					// Note, if the pivot is located at the COM of the child then shrink will not do anything. 
-					FVector CurrentDelta = ConstraintFramesLocal[Index0].TransformPosition(LinearDisplacement);
-					FVector StartDelta = ConstraintFramesLocal[Index0].TransformPosition(JointSettings.LinearDrivePositionTarget);
+					FVec3 StartDelta = ConstraintFramesLocal[Index1].InverseTransformPositionNoScale(JointSettings.LinearDrivePositionTarget);
+					FVec3 CurrentDelta = ConstraintFramesGlobal[Index1].InverseTransformPositionNoScale(Particle1->P());
 
 					if (JointSettings.LinearPlasticityType == EPlasticityType::Shrink && CurrentDelta.SizeSquared() < StartDelta.SizeSquared())
 					{
