Chaos : Convex
Disable convex face simplification when using CoreGeo convex generation

diff --git a/Engine/Source/Runtime/Experimental/Chaos/Public/Chaos/Convex.h b/Engine/Source/Runtime/Experimental/Chaos/Public/Chaos/Convex.h
--- a/Engine/Source/Runtime/Experimental/Chaos/Public/Chaos/Convex.h
+++ b/Engine/Source/Runtime/Experimental/Chaos/Public/Chaos/Convex.h
@@ -103,10 +103,17 @@ namespace Chaos
 			TParticles<FReal, 3> VertexParticles(CopyTemp(Vertices));
 			CalculateVolumeAndCenterOfMass(VertexParticles, FaceIndices, Volume, CenterOfMass);
 
-			// @todo(chaos): DistanceTolerance should be based on size, or passed in
-			const FReal DistanceTolerance = 1.0f;
-			FConvexBuilder::MergeFaces(Planes, FaceIndices, Vertices, DistanceTolerance);
-			CHAOS_ENSURE(Planes.Num() == FaceIndices.Num());
+			// @todo(chaos):																				 should be based on size, or passed in
+			if (!FConvexBuilder::bUseGeometryTConvexHull3)
+			{
+				// @todo(convex) : TConvexHull3 does not need to merge faces, and 
+				// it appears that this code path can leave the convex in an
+				// undefined state. We should consider removing the merge faces when
+				// we transition to the UE::Geometry convex generation. 
+				const FReal DistanceTolerance = 1.0f;
+				FConvexBuilder::MergeFaces(Planes, FaceIndices, Vertices, DistanceTolerance);
+				CHAOS_ENSURE(Planes.Num() == FaceIndices.Num());
+			}
 
 			CreateStructureData(MoveTemp(FaceIndices));
 
