Chaos : Change Implicit Trimesh to use Float instead of double even in LWC mode
this reduce the memory footprint for vertices. the internal BVH needs more work

diff --git a/Engine/Source/Programs/HeadlessChaos/Private/HeadlessChaosTestEPA.cpp b/Engine/Source/Programs/HeadlessChaos/Private/HeadlessChaosTestEPA.cpp
--- a/Engine/Source/Programs/HeadlessChaos/Private/HeadlessChaosTestEPA.cpp
+++ b/Engine/Source/Programs/HeadlessChaos/Private/HeadlessChaosTestEPA.cpp
@@ -1201,7 +1201,7 @@ namespace ChaosTest
 
 			FConvex ConvexBox(MoveTemp(ConvexBoxSurfaceParticles), 0.0f);
 
-			FParticles TrimeshParticles(
+			FTriangleMeshImplicitObject::ParticlesType TrimeshParticles(
 				{
 					{50.0000000, 50.0000000, -8.04061356e-15},
 					{50.0000000, -50.0000000, 8.04061356e-15},
diff --git a/Engine/Source/Programs/HeadlessChaos/Private/HeadlessChaosTestMostOpposing.cpp b/Engine/Source/Programs/HeadlessChaos/Private/HeadlessChaosTestMostOpposing.cpp
--- a/Engine/Source/Programs/HeadlessChaos/Private/HeadlessChaosTestMostOpposing.cpp
+++ b/Engine/Source/Programs/HeadlessChaos/Private/HeadlessChaosTestMostOpposing.cpp
@@ -25,7 +25,7 @@ namespace ChaosTest
 		TArray<uint16> DummyMaterials;
 		int32 FaceIndex;
 
-		FParticles Particles;
+		FTriangleMeshImplicitObject::ParticlesType Particles;
 		Particles.AddParticles(6);
 		Particles.X(0) = FVec3(1, 1, 1);
 		Particles.X(1) = FVec3(5, 1, 1);
diff --git a/Engine/Source/Programs/HeadlessChaos/Private/HeadlessChaosTestPhysicsScene.cpp b/Engine/Source/Programs/HeadlessChaos/Private/HeadlessChaosTestPhysicsScene.cpp
--- a/Engine/Source/Programs/HeadlessChaos/Private/HeadlessChaosTestPhysicsScene.cpp
+++ b/Engine/Source/Programs/HeadlessChaos/Private/HeadlessChaosTestPhysicsScene.cpp
@@ -73,7 +73,7 @@ namespace ChaosTest {
 
 		bool VisitRaycast(const TSpatialVisitorData<FPayload>& Data, FQueryFastData& CurData)
 		{
-			float OutTime = 0;
+			FReal OutTime = 0;
 			FVec3 OutPos;
 			FVec3 OutNorm;
 			int32 FaceIdx;
diff --git a/Engine/Source/Programs/HeadlessChaos/Private/HeadlessChaosTestRaycast.cpp b/Engine/Source/Programs/HeadlessChaos/Private/HeadlessChaosTestRaycast.cpp
--- a/Engine/Source/Programs/HeadlessChaos/Private/HeadlessChaosTestRaycast.cpp
+++ b/Engine/Source/Programs/HeadlessChaos/Private/HeadlessChaosTestRaycast.cpp
@@ -294,7 +294,7 @@ namespace ChaosTest
 		TArray<uint16> DummyMaterials;
 		int32 FaceIndex;
 
-		FParticles Particles;
+		FTriangleMeshImplicitObject::ParticlesType Particles;
 		Particles.AddParticles(3);
 		Particles.X(0) = FVec3(1, 1, 1);
 		Particles.X(1) = FVec3(5, 1, 1);
diff --git a/Engine/Source/Programs/HeadlessChaos/Private/HeadlessChaosTestSweep.cpp b/Engine/Source/Programs/HeadlessChaos/Private/HeadlessChaosTestSweep.cpp
--- a/Engine/Source/Programs/HeadlessChaos/Private/HeadlessChaosTestSweep.cpp
+++ b/Engine/Source/Programs/HeadlessChaos/Private/HeadlessChaosTestSweep.cpp
@@ -17,7 +17,7 @@ namespace ChaosTest
 		// Trimesh is from SM_Cattus_POI_Rib, this was a real world failure that is now fixed.
 
 		using namespace Chaos;
-		FParticles TrimeshParticles(
+		FTriangleMeshImplicitObject::ParticlesType TrimeshParticles(
 		{
 			{29.0593967, -5.21321106, -10.2669592},
 			{34.5006638, -3.16600156, -14.5092020},
diff --git a/Engine/Source/Runtime/Engine/Private/PhysicsEngine/Experimental/ChaosCooking.cpp b/Engine/Source/Runtime/Engine/Private/PhysicsEngine/Experimental/ChaosCooking.cpp
--- a/Engine/Source/Runtime/Engine/Private/PhysicsEngine/Experimental/ChaosCooking.cpp
+++ b/Engine/Source/Runtime/Engine/Private/PhysicsEngine/Experimental/ChaosCooking.cpp
@@ -44,7 +44,7 @@ namespace Chaos
 			}
 
 			// Build particle list #BG Maybe allow TParticles to copy vectors?
-			Chaos::TParticles<Chaos::FReal, 3> TriMeshParticles;
+			Chaos::FTriangleMeshImplicitObject::ParticlesType TriMeshParticles;
 			TriMeshParticles.AddParticles(FinalVerts.Num());
 
 			const int32 NumVerts = FinalVerts.Num();
@@ -216,7 +216,7 @@ namespace Chaos
 			}
 
 			// Build particle list #BG Maybe allow TParticles to copy vectors?
-			Chaos::TParticles<Chaos::FReal, 3> TriMeshParticles;
+			Chaos::FTriangleMeshImplicitObject::ParticlesType TriMeshParticles;
 			TriMeshParticles.AddParticles(FinalVerts.Num());
 
 			const int32 NumVerts = FinalVerts.Num();
diff --git a/Engine/Source/Runtime/Experimental/Chaos/Private/Chaos/TriangleMeshImplicitObject.cpp b/Engine/Source/Runtime/Experimental/Chaos/Private/Chaos/TriangleMeshImplicitObject.cpp
--- a/Engine/Source/Runtime/Experimental/Chaos/Private/Chaos/TriangleMeshImplicitObject.cpp
+++ b/Engine/Source/Runtime/Experimental/Chaos/Private/Chaos/TriangleMeshImplicitObject.cpp
@@ -39,7 +39,7 @@ void ScaleTransformHelper(const FVec3& TriMeshScale, const FRigidTransform3& Que
 }
 
 template <typename IdxType>
-void TransformVertsHelper(const FVec3& TriMeshScale, int32 TriIdx, const FParticles& Particles,
+void TransformVertsHelper(const FVec3& TriMeshScale, int32 TriIdx, const FTriangleMeshImplicitObject::ParticlesType& Particles,
 	const TArray<TVector<IdxType, 3>>& Elements, FVec3& OutA, FVec3& OutB, FVec3& OutC)
 {
 	OutA = Particles.X(Elements[TriIdx][0]) * TriMeshScale;
@@ -76,7 +76,9 @@ void TransformSweepOutputsHelper(FVec3 TriMeshScale, const FVec3& HitNormal, con
 template <typename IdxType>
 struct FTriangleMeshRaycastVisitor
 {
-	FTriangleMeshRaycastVisitor(const FVec3& InStart, const FVec3& InDir, const FReal InThickness, const FParticles& InParticles, const TArray<TVector<IdxType, 3>>& InElements, bool bInCullsBackFaceRaycast)
+	using ParticlesType = FTriangleMeshImplicitObject::ParticlesType;
+
+	FTriangleMeshRaycastVisitor(const FVec3& InStart, const FVec3& InDir, const FReal InThickness, const ParticlesType& InParticles, const TArray<TVector<IdxType, 3>>& InElements, bool bInCullsBackFaceRaycast)
 	: Particles(InParticles)
 	, Elements(InElements)
 	, StartPoint(InStart)
@@ -251,7 +253,7 @@ struct FTriangleMeshRaycastVisitor
 		return true;
 	}
 
-	const FParticles& Particles;
+	const ParticlesType& Particles;
 	const TArray<TVector<IdxType, 3>>& Elements;
 	const FVec3& StartPoint;
 	const FVec3& Dir;
@@ -912,9 +914,9 @@ template <typename IdxType>
 TUniquePtr<FTriangleMeshImplicitObject> FTriangleMeshImplicitObject::CopySlowImpl(const TArray<TVector<IdxType, 3>>& InElements) const
 {
 	using namespace Chaos;
-
-	TArray<Chaos::FVec3> XArray = MParticles.AllX();
-	FParticles ParticlesCopy(MoveTemp(XArray));
+	
+	TArray<ParticleVecType> XArray = MParticles.AllX();
+	ParticlesType ParticlesCopy(MoveTemp(XArray));
 	TArray<TVector<IdxType, 3>> ElementsCopy(InElements);
 	TArray<uint16> MaterialIndicesCopy = MaterialIndices;
 	TUniquePtr<TArray<int32>> ExternalFaceIndexMapCopy = nullptr;
@@ -983,13 +985,13 @@ FVec3 FTriangleMeshImplicitObject::GetFaceNormal(const int32 FaceIdx) const
 	{
 		auto LambdaHelper = [&](const auto& Elements)
 		{
-			const FVec3& A = MParticles.X(Elements[FaceIdx][0]);
-			const FVec3& B = MParticles.X(Elements[FaceIdx][1]);
-			const FVec3& C = MParticles.X(Elements[FaceIdx][2]);
+			const ParticleVecType& A = MParticles.X(Elements[FaceIdx][0]);
+			const ParticleVecType& B = MParticles.X(Elements[FaceIdx][1]);
+			const ParticleVecType& C = MParticles.X(Elements[FaceIdx][2]);
 
-			const FVec3 AB = B - A;
-			const FVec3 AC = C - A;
-			FVec3 Normal = FVec3::CrossProduct(AB, AC);
+			const ParticleVecType AB = B - A;
+			const ParticleVecType AC = C - A;
+			ParticleVecType Normal = ParticleVecType::CrossProduct(AB, AC);
 			
 			if(Normal.SafeNormalize() < SMALL_NUMBER)
 			{
@@ -998,7 +1000,7 @@ FVec3 FTriangleMeshImplicitObject::GetFaceNormal(const int32 FaceIdx) const
 				return FVec3(0, 0, 1);
 			}
 
-			return Normal;
+			return FVec3(Normal);
 		};
 		
 		if (MElements.RequiresLargeIndices())
@@ -1025,7 +1027,7 @@ uint16 FTriangleMeshImplicitObject::GetMaterialIndex(uint32 HintIndex) const
 	return 0;
 }
 
-const FParticles& FTriangleMeshImplicitObject::Particles() const
+const FTriangleMeshImplicitObject::ParticlesType& FTriangleMeshImplicitObject::Particles() const
 {
 	return MParticles;
 }
diff --git a/Engine/Source/Runtime/Experimental/Chaos/Public/Chaos/Particles.h b/Engine/Source/Runtime/Experimental/Chaos/Public/Chaos/Particles.h
--- a/Engine/Source/Runtime/Experimental/Chaos/Public/Chaos/Particles.h
+++ b/Engine/Source/Runtime/Experimental/Chaos/Public/Chaos/Particles.h
@@ -191,8 +191,7 @@ namespace Chaos
 #endif
 		}
 
-		template<typename operator_T, int operator_d>
-		friend FArchive& operator<<(FArchive& Ar, TParticles<operator_T, operator_d>& InParticles)
+		inline friend FArchive& operator<<(FArchive& Ar, TParticles<T, d>& InParticles)
 		{
 			InParticles.Serialize(Ar);
 			return Ar;
diff --git a/Engine/Source/Runtime/Experimental/Chaos/Public/Chaos/TriangleMeshImplicitObject.h b/Engine/Source/Runtime/Experimental/Chaos/Public/Chaos/TriangleMeshImplicitObject.h
--- a/Engine/Source/Runtime/Experimental/Chaos/Public/Chaos/TriangleMeshImplicitObject.h
+++ b/Engine/Source/Runtime/Experimental/Chaos/Public/Chaos/TriangleMeshImplicitObject.h
@@ -105,8 +105,11 @@ namespace Chaos
 	public:
 		using FImplicitObject::GetTypeName;
 
+		using ParticlesType = TParticles<FRealSingle, 3>;
+		using ParticleVecType = TVec3<FRealSingle>;
+
 		template <typename IdxType>
-		FTriangleMeshImplicitObject(FParticles&& Particles, TArray<TVec3<IdxType>>&& Elements, TArray<uint16>&& InMaterialIndices, TUniquePtr<TArray<int32>>&& InExternalFaceIndexMap = nullptr, TUniquePtr<TArray<int32>>&& InExternalVertexIndexMap = nullptr, const bool bInCullsBackFaceRaycast = false)
+		FTriangleMeshImplicitObject(ParticlesType&& Particles, TArray<TVec3<IdxType>>&& Elements, TArray<uint16>&& InMaterialIndices, TUniquePtr<TArray<int32>>&& InExternalFaceIndexMap = nullptr, TUniquePtr<TArray<int32>>&& InExternalVertexIndexMap = nullptr, const bool bInCullsBackFaceRaycast = false)
 		: FImplicitObject(EImplicitObject::HasBoundingBox | EImplicitObject::DisableCollisions, ImplicitObjectType::TriangleMesh)
 		, MParticles(MoveTemp(Particles))
 		, MElements(MoveTemp(Elements))
@@ -286,7 +289,7 @@ namespace Chaos
 
 		virtual uint16 GetMaterialIndex(uint32 HintIndex) const override;
 
-		const FParticles& Particles() const;
+		const ParticlesType& Particles() const;
 		const FTrimeshIndexBuffer& Elements() const;
 
 		void UpdateVertices(const TArray<FVector>& Positions);
@@ -296,7 +299,7 @@ namespace Chaos
 	private:
 		void RebuildBV();
 
-		FParticles MParticles;
+		ParticlesType MParticles;
 		FTrimeshIndexBuffer MElements;
 		FAABB3 MLocalBoundingBox;
 		TArray<uint16> MaterialIndices;
@@ -308,7 +311,7 @@ namespace Chaos
 
 		// Initialising constructor privately declared for use in CopySlow to copy the underlying BVH
 		template <typename IdxType>
-		FTriangleMeshImplicitObject(TParticles<FReal, 3>&& Particles, TArray<TVec3<IdxType>>&& Elements, TArray<uint16>&& InMaterialIndices, const BVHType& InBvhToCopy, TUniquePtr<TArray<int32>>&& InExternalFaceIndexMap = nullptr, TUniquePtr<TArray<int32>>&& InExternalVertexIndexMap = nullptr, const bool bInCullsBackFaceRaycast = false)
+		FTriangleMeshImplicitObject(ParticlesType&& Particles, TArray<TVec3<IdxType>>&& Elements, TArray<uint16>&& InMaterialIndices, const BVHType& InBvhToCopy, TUniquePtr<TArray<int32>>&& InExternalFaceIndexMap = nullptr, TUniquePtr<TArray<int32>>&& InExternalVertexIndexMap = nullptr, const bool bInCullsBackFaceRaycast = false)
 			: FImplicitObject(EImplicitObject::HasBoundingBox | EImplicitObject::DisableCollisions, ImplicitObjectType::TriangleMesh)
 			, MParticles(MoveTemp(Particles))
 			, MElements(MoveTemp(Elements))
diff --git a/Engine/Source/Runtime/Experimental/ChaosCore/Public/Chaos/Vector.h b/Engine/Source/Runtime/Experimental/ChaosCore/Public/Chaos/Vector.h
--- a/Engine/Source/Runtime/Experimental/ChaosCore/Public/Chaos/Vector.h
+++ b/Engine/Source/Runtime/Experimental/ChaosCore/Public/Chaos/Vector.h
@@ -386,26 +386,26 @@ namespace Chaos
 	};
 
 	template<>
-	class TVector<FReal, 3> : public UE::Math::TVector<FReal>
+	class TVector<FRealSingle, 3> : public UE::Math::TVector<FRealSingle>
 	{
 	public:
-		using FElement = FReal;
-		using UE::Math::TVector<FReal>::X;
-		using UE::Math::TVector<FReal>::Y;
-		using UE::Math::TVector<FReal>::Z;
+		using FElement = FRealSingle;
+		using UE::Math::TVector<FRealSingle>::X;
+		using UE::Math::TVector<FRealSingle>::Y;
+		using UE::Math::TVector<FRealSingle>::Z;
 
 		TVector()
-		    : UE::Math::TVector<FReal>() {}
-		explicit TVector(const FReal x)
-		    : UE::Math::TVector<FReal>(x, x, x) {}
-		TVector(const FReal x, const FReal y, const FReal z)
-		    : UE::Math::TVector<FReal>(x, y, z) {}
+		    : UE::Math::TVector<FRealSingle>() {}
+		explicit TVector(const FRealSingle x)
+		    : UE::Math::TVector<FRealSingle>(x, x, x) {}
+		TVector(const FRealSingle x, const FRealSingle y, const FRealSingle z)
+		    : UE::Math::TVector<FRealSingle>(x, y, z) {}
 		TVector(const UE::Math::TVector<FRealSingle>& vec)
-		    : UE::Math::TVector<FReal>((UE::Math::TVector<FReal>)vec) {}
-		TVector(const UE::Math::TVector<FRealDouble>& vec)					// LWC_TODO: Precision loss. Make explicit for FReal = FRealSingle?
-			: UE::Math::TVector<FReal>((UE::Math::TVector<FReal>)vec) {}
+		    : UE::Math::TVector<FRealSingle>((UE::Math::TVector<FRealSingle>)vec) {}
+		TVector(const UE::Math::TVector<FRealDouble>& vec)					// LWC_TODO: Precision loss. Make explicit for FRealSingle = FRealSingle?
+			: UE::Math::TVector<FRealSingle>((UE::Math::TVector<FRealSingle>)vec) {}
 		TVector(const FVector4& vec)
-		    : UE::Math::TVector<FReal>(vec.X, vec.Y, vec.Z) {}
+		    : UE::Math::TVector<FRealSingle>(vec.X, vec.Y, vec.Z) {}
 		TVector(std::istream& Stream)
 		{
 			Stream.read(reinterpret_cast<char*>(&X), sizeof(X));
@@ -419,86 +419,86 @@ namespace Chaos
 			Stream.write(reinterpret_cast<const char*>(&Y), sizeof(Y));
 			Stream.write(reinterpret_cast<const char*>(&Z), sizeof(Z));
 		}
-		static inline TVector<FReal, 3> Lerp(const TVector<FReal, 3>& V1, const TVector<FReal, 3>& V2, const FReal F) { return FMath::Lerp<UE::Math::TVector<FReal>, FReal>(V1, V2, F); }
-		static inline TVector<FReal, 3> CrossProduct(const TVector<FReal, 3>& V1, const TVector<FReal, 3>& V2) { return UE::Math::TVector<FReal>::CrossProduct(V1, V2); }
-		static inline FReal DotProduct(const TVector<FReal, 3>& V1, const TVector<FReal, 3>& V2) { return UE::Math::TVector<FReal>::DotProduct(V1, V2); }
-		bool operator<=(const TVector<FReal, 3>& V) const
+		static inline TVector<FRealSingle, 3> Lerp(const TVector<FRealSingle, 3>& V1, const TVector<FRealSingle, 3>& V2, const FRealSingle F) { return FMath::Lerp<UE::Math::TVector<FRealSingle>, FRealSingle>(V1, V2, F); }
+		static inline TVector<FRealSingle, 3> CrossProduct(const TVector<FRealSingle, 3>& V1, const TVector<FRealSingle, 3>& V2) { return UE::Math::TVector<FRealSingle>::CrossProduct(V1, V2); }
+		static inline FRealSingle DotProduct(const TVector<FRealSingle, 3>& V1, const TVector<FRealSingle, 3>& V2) { return UE::Math::TVector<FRealSingle>::DotProduct(V1, V2); }
+		bool operator<=(const TVector<FRealSingle, 3>& V) const
 		{
 			return X <= V.X && Y <= V.Y && Z <= V.Z;
 		}
-		bool operator>=(const TVector<FReal, 3>& V) const
+		bool operator>=(const TVector<FRealSingle, 3>& V) const
 		{
 			return X >= V.X && Y >= V.Y && Z >= V.Z;
 		}
-		TVector<FReal, 3> operator-() const
+		TVector<FRealSingle, 3> operator-() const
 		{
-			return TVector<FReal, 3>(-X, -Y, -Z);
+			return TVector<FRealSingle, 3>(-X, -Y, -Z);
 		}
-		TVector<FReal, 3> operator+(const FReal Other) const
+		TVector<FRealSingle, 3> operator+(const FRealSingle Other) const
 		{
-			return TVector<FReal, 3>(X + Other, Y + Other, Z + Other);
+			return TVector<FRealSingle, 3>(X + Other, Y + Other, Z + Other);
 		}
-		TVector<FReal, 3> operator-(const FReal Other) const
+		TVector<FRealSingle, 3> operator-(const FRealSingle Other) const
 		{
-			return TVector<FReal, 3>(X - Other, Y - Other, Z - Other);
+			return TVector<FRealSingle, 3>(X - Other, Y - Other, Z - Other);
 		}
-		TVector<FReal, 3> operator*(const FReal Other) const
+		TVector<FRealSingle, 3> operator*(const FRealSingle Other) const
 		{
-			return TVector<FReal, 3>(X * Other, Y * Other, Z * Other);
+			return TVector<FRealSingle, 3>(X * Other, Y * Other, Z * Other);
 		}
-		TVector<FReal, 3> operator/(const FReal Other) const
+		TVector<FRealSingle, 3> operator/(const FRealSingle Other) const
 		{
-			return TVector<FReal, 3>(X / Other, Y / Other, Z / Other);
+			return TVector<FRealSingle, 3>(X / Other, Y / Other, Z / Other);
 		}
-		friend TVector<FReal, 3> operator/(const FReal S, const TVector<FReal, 3>& V)
+		friend TVector<FRealSingle, 3> operator/(const FRealSingle S, const TVector<FRealSingle, 3>& V)
 		{
-			return TVector<FReal, 3>(S / V.X, S / V.Y, S / V.Z);
+			return TVector<FRealSingle, 3>(S / V.X, S / V.Y, S / V.Z);
 		}
-		TVector<FReal, 3> operator+(const TVector<FReal, 3>& Other) const
+		TVector<FRealSingle, 3> operator+(const TVector<FRealSingle, 3>& Other) const
 		{
-			return TVector<FReal, 3>(X + Other[0], Y + Other[1], Z + Other[2]);
+			return TVector<FRealSingle, 3>(X + Other[0], Y + Other[1], Z + Other[2]);
 		}
-		TVector<FReal, 3> operator-(const TVector<FReal, 3>& Other) const
+		TVector<FRealSingle, 3> operator-(const TVector<FRealSingle, 3>& Other) const
 		{
-			return TVector<FReal, 3>(X - Other[0], Y - Other[1], Z - Other[2]);
+			return TVector<FRealSingle, 3>(X - Other[0], Y - Other[1], Z - Other[2]);
 		}
-		TVector<FReal, 3> operator*(const TVector<FReal, 3>& Other) const
+		TVector<FRealSingle, 3> operator*(const TVector<FRealSingle, 3>& Other) const
 		{
-			return TVector<FReal, 3>(X * Other[0], Y * Other[1], Z * Other[2]);
+			return TVector<FRealSingle, 3>(X * Other[0], Y * Other[1], Z * Other[2]);
 		}
-		TVector<FReal, 3> operator/(const TVector<FReal, 3>& Other) const
+		TVector<FRealSingle, 3> operator/(const TVector<FRealSingle, 3>& Other) const
 		{
-			return TVector<FReal, 3>(X / Other[0], Y / Other[1], Z / Other[2]);
+			return TVector<FRealSingle, 3>(X / Other[0], Y / Other[1], Z / Other[2]);
 		}
 		template<class T2>
-		TVector<FReal, 3> operator+(const TVector<T2, 3>& Other) const
+		TVector<FRealSingle, 3> operator+(const TVector<T2, 3>& Other) const
 		{
-			return TVector<FReal, 3>(X + static_cast<FReal>(Other[0]), Y + static_cast<FReal>(Other[1]), Z + static_cast<FReal>(Other[2]));
+			return TVector<FRealSingle, 3>(X + static_cast<FRealSingle>(Other[0]), Y + static_cast<FRealSingle>(Other[1]), Z + static_cast<FRealSingle>(Other[2]));
 		}
 		template<class T2>
-		TVector<FReal, 3> operator-(const TVector<T2, 3>& Other) const
+		TVector<FRealSingle, 3> operator-(const TVector<T2, 3>& Other) const
 		{
-			return TVector<FReal, 3>(X - static_cast<FReal>(Other[0]), Y - static_cast<FReal>(Other[1]), Z - static_cast<FReal>(Other[2]));
+			return TVector<FRealSingle, 3>(X - static_cast<FRealSingle>(Other[0]), Y - static_cast<FRealSingle>(Other[1]), Z - static_cast<FRealSingle>(Other[2]));
 		}
 		template<class T2>
-		TVector<FReal, 3> operator*(const TVector<T2, 3>& Other) const
+		TVector<FRealSingle, 3> operator*(const TVector<T2, 3>& Other) const
 		{
-			return TVector<FReal, 3>(X * static_cast<FReal>(Other[0]), Y * static_cast<FReal>(Other[1]), Z * static_cast<FReal>(Other[2]));
+			return TVector<FRealSingle, 3>(X * static_cast<FRealSingle>(Other[0]), Y * static_cast<FRealSingle>(Other[1]), Z * static_cast<FRealSingle>(Other[2]));
 		}
 		template<class T2>
-		TVector<FReal, 3> operator/(const TVector<T2, 3>& Other) const
+		TVector<FRealSingle, 3> operator/(const TVector<T2, 3>& Other) const
 		{
-			return TVector<FReal, 3>(X / static_cast<FReal>(Other[0]), Y / static_cast<FReal>(Other[1]), Z / static_cast<FReal>(Other[2]));
+			return TVector<FRealSingle, 3>(X / static_cast<FRealSingle>(Other[0]), Y / static_cast<FRealSingle>(Other[1]), Z / static_cast<FRealSingle>(Other[2]));
 		}
-		FReal Product() const
+		FRealSingle Product() const
 		{
 			return X * Y * Z;
 		}
-		FReal Max() const
+		FRealSingle Max() const
 		{
 			return (X > Y && X > Z) ? X : (Y > Z ? Y : Z);
 		}
-		FReal Min() const
+		FRealSingle Min() const
 		{
 			return (X < Y && X < Z) ? X : (Y < Z ? Y : Z);
 		}
@@ -506,24 +506,24 @@ namespace Chaos
 		{
 			return (X > Y && X > Z) ? 0 : (Y > Z ? 1 : 2);
 		}
-		FReal Mid() const
+		FRealSingle Mid() const
 		{
-			FReal XmY = X - Y;
-			FReal YmZ = Y - Z;
-			FReal XmZ = X - Z;
+			FRealSingle XmY = X - Y;
+			FRealSingle YmZ = Y - Z;
+			FRealSingle XmZ = X - Z;
 			return (XmY * YmZ > -1 ? Y : XmY * XmZ < 1 ? X : Z);
 		}
-		TVector<FReal, 3> ComponentwiseMin(const TVector<FReal, 3>& Other) const { return {FMath::Min(X,Other.X), FMath::Min(Y,Other.Y), FMath::Min(Z,Other.Z)}; }
-		TVector<FReal, 3> ComponentwiseMax(const TVector<FReal, 3>& Other) const { return {FMath::Max(X,Other.X), FMath::Max(Y,Other.Y), FMath::Max(Z,Other.Z)}; }
-		static TVector<FReal, 3> Max(const TVector<FReal, 3>& V1, const TVector<FReal, 3>& V2)
+		TVector<FRealSingle, 3> ComponentwiseMin(const TVector<FRealSingle, 3>& Other) const { return {FMath::Min(X,Other.X), FMath::Min(Y,Other.Y), FMath::Min(Z,Other.Z)}; }
+		TVector<FRealSingle, 3> ComponentwiseMax(const TVector<FRealSingle, 3>& Other) const { return {FMath::Max(X,Other.X), FMath::Max(Y,Other.Y), FMath::Max(Z,Other.Z)}; }
+		static TVector<FRealSingle, 3> Max(const TVector<FRealSingle, 3>& V1, const TVector<FRealSingle, 3>& V2)
 		{
-			return TVector<FReal, 3>(V1.X > V2.X ? V1.X : V2.X, V1.Y > V2.Y ? V1.Y : V2.Y, V1.Z > V2.Z ? V1.Z : V2.Z);
+			return TVector<FRealSingle, 3>(V1.X > V2.X ? V1.X : V2.X, V1.Y > V2.Y ? V1.Y : V2.Y, V1.Z > V2.Z ? V1.Z : V2.Z);
 		}
-		static TVector<FReal, 3> AxisVector(const int32 Axis)
-		{ return Axis == 0 ? TVector<FReal, 3>(1.f, 0.f, 0.f) : (Axis == 1 ? TVector<FReal, 3>(0.f, 1.f, 0.f) : TVector<FReal, 3>(0.f, 0.f, 1.f)); }
-		static Pair<FReal, int32> MaxAndAxis(const TVector<FReal, 3>& V1, const TVector<FReal, 3>& V2)
+		static TVector<FRealSingle, 3> AxisVector(const int32 Axis)
+		{ return Axis == 0 ? TVector<FRealSingle, 3>(1.f, 0.f, 0.f) : (Axis == 1 ? TVector<FRealSingle, 3>(0.f, 1.f, 0.f) : TVector<FRealSingle, 3>(0.f, 0.f, 1.f)); }
+		static Pair<FRealSingle, int32> MaxAndAxis(const TVector<FRealSingle, 3>& V1, const TVector<FRealSingle, 3>& V2)
 		{
-			const TVector<FReal, 3> max = Max(V1, V2);
+			const TVector<FRealSingle, 3> max = Max(V1, V2);
 			if (max.X > max.Y)
 			{
 				if (max.X > max.Z)
@@ -539,9 +539,9 @@ namespace Chaos
 					return MakePair(max.Z, 2);
 			}
 		}
-		FReal SafeNormalize(FReal Epsilon = 1e-4f)
+		FRealSingle SafeNormalize(FRealSingle Epsilon = 1e-4f)
 		{
-			FReal Size = SizeSquared();
+			FRealSingle Size = SizeSquared();
 			if (Size < Epsilon)
 			{
 				*this = AxisVector(0);
@@ -551,40 +551,244 @@ namespace Chaos
 			*this = (*this) / Size;
 			return Size;
 		}
-		TVector<FReal, 3> GetOrthogonalVector() const
+		TVector<FRealSingle, 3> GetOrthogonalVector() const
 		{
-			TVector<FReal, 3> AbsVector(FMath::Abs(X), FMath::Abs(Y), FMath::Abs(Z));
+			TVector<FRealSingle, 3> AbsVector(FMath::Abs(X), FMath::Abs(Y), FMath::Abs(Z));
 			if ((AbsVector.X <= AbsVector.Y) && (AbsVector.X <= AbsVector.Z))
 			{
 				// X is the smallest component
-				return TVector<FReal, 3>(0, Z, -Y);
+				return TVector<FRealSingle, 3>(0, Z, -Y);
 			}
 			if ((AbsVector.Z <= AbsVector.X) && (AbsVector.Z <= AbsVector.Y))
 			{
 				// Z is the smallest component
-				return TVector<FReal, 3>(Y, -X, 0);
+				return TVector<FRealSingle, 3>(Y, -X, 0);
 			}
 			// Y is the smallest component
-			return TVector<FReal, 3>(-Z, 0, X);
+			return TVector<FRealSingle, 3>(-Z, 0, X);
 		}
-		static FReal AngleBetween(const TVector<FReal, 3>& V1, const TVector<FReal, 3>& V2)
+		static FRealSingle AngleBetween(const TVector<FRealSingle, 3>& V1, const TVector<FRealSingle, 3>& V2)
 		{
-			FReal s = CrossProduct(V1, V2).Size();
-			FReal c = DotProduct(V1, V2);
+			FRealSingle s = CrossProduct(V1, V2).Size();
+			FRealSingle c = DotProduct(V1, V2);
 			return atan2(s, c);
 		}
 		/** Calculate the velocity to move from P0 to P1 in time Dt. Exists just for symmetry with TRotation::CalculateAngularVelocity! */
-		static TVector<FReal, 3> CalculateVelocity(const TVector<FReal, 3>& P0, const TVector<FReal, 3>& P1, const FReal Dt)
+		static TVector<FRealSingle, 3> CalculateVelocity(const TVector<FRealSingle, 3>& P0, const TVector<FRealSingle, 3>& P1, const FRealSingle Dt)
 		{
 			return (P1 - P0) / Dt;
 		}
 
-		static bool IsNearlyEqual(const TVector<FReal, 3>& A, const TVector<FReal, 3>& B, const FReal Epsilon)
+		static bool IsNearlyEqual(const TVector<FRealSingle, 3>& A, const TVector<FRealSingle, 3>& B, const FRealSingle Epsilon)
 		{
 			return (B - A).IsNearlyZero(Epsilon);
 		}
 	};
 
+
+	template<>
+	class TVector<FRealDouble, 3> : public UE::Math::TVector<FRealDouble>
+	{
+	public:
+		using FElement = FRealDouble;
+		using UE::Math::TVector<FRealDouble>::X;
+		using UE::Math::TVector<FRealDouble>::Y;
+		using UE::Math::TVector<FRealDouble>::Z;
+
+		TVector()
+			: UE::Math::TVector<FRealDouble>() {}
+		explicit TVector(const FRealDouble x)
+			: UE::Math::TVector<FRealDouble>(x, x, x) {}
+		TVector(const FRealDouble x, const FRealDouble y, const FRealDouble z)
+			: UE::Math::TVector<FRealDouble>(x, y, z) {}
+		TVector(const UE::Math::TVector<FRealSingle>& vec)
+			: UE::Math::TVector<FRealDouble>((UE::Math::TVector<FRealDouble>)vec) {}
+		TVector(const UE::Math::TVector<FRealDouble>& vec)					// LWC_TODO: Precision loss. Make explicit for FRealDouble = FRealSingle?
+			: UE::Math::TVector<FRealDouble>((UE::Math::TVector<FRealDouble>)vec) {}
+		TVector(const FVector4& vec)
+			: UE::Math::TVector<FRealDouble>(vec.X, vec.Y, vec.Z) {}
+		TVector(std::istream& Stream)
+		{
+			Stream.read(reinterpret_cast<char*>(&X), sizeof(X));
+			Stream.read(reinterpret_cast<char*>(&Y), sizeof(Y));
+			Stream.read(reinterpret_cast<char*>(&Z), sizeof(Z));
+		}
+		//operator UE::Math::TVector<FRealDouble>() const { return UE::Math::TVector<FRealDouble>((FRealDouble)X, (FRealDouble)Y, (FRealDouble)Z); }
+		void Write(std::ostream& Stream) const
+		{
+			Stream.write(reinterpret_cast<const char*>(&X), sizeof(X));
+			Stream.write(reinterpret_cast<const char*>(&Y), sizeof(Y));
+			Stream.write(reinterpret_cast<const char*>(&Z), sizeof(Z));
+		}
+		static inline TVector<FRealDouble, 3> Lerp(const TVector<FRealDouble, 3>& V1, const TVector<FRealDouble, 3>& V2, const FRealDouble F) { return FMath::Lerp<UE::Math::TVector<FRealDouble>, FRealDouble>(V1, V2, F); }
+		static inline TVector<FRealDouble, 3> CrossProduct(const TVector<FRealDouble, 3>& V1, const TVector<FRealDouble, 3>& V2) { return UE::Math::TVector<FRealDouble>::CrossProduct(V1, V2); }
+		static inline FRealDouble DotProduct(const TVector<FRealDouble, 3>& V1, const TVector<FRealDouble, 3>& V2) { return UE::Math::TVector<FRealDouble>::DotProduct(V1, V2); }
+		bool operator<=(const TVector<FRealDouble, 3>& V) const
+		{
+			return X <= V.X && Y <= V.Y && Z <= V.Z;
+		}
+		bool operator>=(const TVector<FRealDouble, 3>& V) const
+		{
+			return X >= V.X && Y >= V.Y && Z >= V.Z;
+		}
+		TVector<FRealDouble, 3> operator-() const
+		{
+			return TVector<FRealDouble, 3>(-X, -Y, -Z);
+		}
+		TVector<FRealDouble, 3> operator+(const FRealDouble Other) const
+		{
+			return TVector<FRealDouble, 3>(X + Other, Y + Other, Z + Other);
+		}
+		TVector<FRealDouble, 3> operator-(const FRealDouble Other) const
+		{
+			return TVector<FRealDouble, 3>(X - Other, Y - Other, Z - Other);
+		}
+		TVector<FRealDouble, 3> operator*(const FRealDouble Other) const
+		{
+			return TVector<FRealDouble, 3>(X * Other, Y * Other, Z * Other);
+		}
+		TVector<FRealDouble, 3> operator/(const FRealDouble Other) const
+		{
+			return TVector<FRealDouble, 3>(X / Other, Y / Other, Z / Other);
+		}
+		friend TVector<FRealDouble, 3> operator/(const FRealDouble S, const TVector<FRealDouble, 3>& V)
+		{
+			return TVector<FRealDouble, 3>(S / V.X, S / V.Y, S / V.Z);
+		}
+		TVector<FRealDouble, 3> operator+(const TVector<FRealDouble, 3>& Other) const
+		{
+			return TVector<FRealDouble, 3>(X + Other[0], Y + Other[1], Z + Other[2]);
+		}
+		TVector<FRealDouble, 3> operator-(const TVector<FRealDouble, 3>& Other) const
+		{
+			return TVector<FRealDouble, 3>(X - Other[0], Y - Other[1], Z - Other[2]);
+		}
+		TVector<FRealDouble, 3> operator*(const TVector<FRealDouble, 3>& Other) const
+		{
+			return TVector<FRealDouble, 3>(X * Other[0], Y * Other[1], Z * Other[2]);
+		}
+		TVector<FRealDouble, 3> operator/(const TVector<FRealDouble, 3>& Other) const
+		{
+			return TVector<FRealDouble, 3>(X / Other[0], Y / Other[1], Z / Other[2]);
+		}
+		template<class T2>
+		TVector<FRealDouble, 3> operator+(const TVector<T2, 3>& Other) const
+		{
+			return TVector<FRealDouble, 3>(X + static_cast<FRealDouble>(Other[0]), Y + static_cast<FRealDouble>(Other[1]), Z + static_cast<FRealDouble>(Other[2]));
+		}
+		template<class T2>
+		TVector<FRealDouble, 3> operator-(const TVector<T2, 3>& Other) const
+		{
+			return TVector<FRealDouble, 3>(X - static_cast<FRealDouble>(Other[0]), Y - static_cast<FRealDouble>(Other[1]), Z - static_cast<FRealDouble>(Other[2]));
+		}
+		template<class T2>
+		TVector<FRealDouble, 3> operator*(const TVector<T2, 3>& Other) const
+		{
+			return TVector<FRealDouble, 3>(X * static_cast<FRealDouble>(Other[0]), Y * static_cast<FRealDouble>(Other[1]), Z * static_cast<FRealDouble>(Other[2]));
+		}
+		template<class T2>
+		TVector<FRealDouble, 3> operator/(const TVector<T2, 3>& Other) const
+		{
+			return TVector<FRealDouble, 3>(X / static_cast<FRealDouble>(Other[0]), Y / static_cast<FRealDouble>(Other[1]), Z / static_cast<FRealDouble>(Other[2]));
+		}
+		FRealDouble Product() const
+		{
+			return X * Y * Z;
+		}
+		FRealDouble Max() const
+		{
+			return (X > Y && X > Z) ? X : (Y > Z ? Y : Z);
+		}
+		FRealDouble Min() const
+		{
+			return (X < Y&& X < Z) ? X : (Y < Z ? Y : Z);
+		}
+		int32 MaxAxis() const
+		{
+			return (X > Y && X > Z) ? 0 : (Y > Z ? 1 : 2);
+		}
+		FRealDouble Mid() const
+		{
+			FRealDouble XmY = X - Y;
+			FRealDouble YmZ = Y - Z;
+			FRealDouble XmZ = X - Z;
+			return (XmY * YmZ > -1 ? Y : XmY * XmZ < 1 ? X : Z);
+		}
+		TVector<FRealDouble, 3> ComponentwiseMin(const TVector<FRealDouble, 3>& Other) const { return { FMath::Min(X,Other.X), FMath::Min(Y,Other.Y), FMath::Min(Z,Other.Z) }; }
+		TVector<FRealDouble, 3> ComponentwiseMax(const TVector<FRealDouble, 3>& Other) const { return { FMath::Max(X,Other.X), FMath::Max(Y,Other.Y), FMath::Max(Z,Other.Z) }; }
+		static TVector<FRealDouble, 3> Max(const TVector<FRealDouble, 3>& V1, const TVector<FRealDouble, 3>& V2)
+		{
+			return TVector<FRealDouble, 3>(V1.X > V2.X ? V1.X : V2.X, V1.Y > V2.Y ? V1.Y : V2.Y, V1.Z > V2.Z ? V1.Z : V2.Z);
+		}
+		static TVector<FRealDouble, 3> AxisVector(const int32 Axis)
+		{
+			return Axis == 0 ? TVector<FRealDouble, 3>(1.f, 0.f, 0.f) : (Axis == 1 ? TVector<FRealDouble, 3>(0.f, 1.f, 0.f) : TVector<FRealDouble, 3>(0.f, 0.f, 1.f));
+		}
+		static Pair<FRealDouble, int32> MaxAndAxis(const TVector<FRealDouble, 3>& V1, const TVector<FRealDouble, 3>& V2)
+		{
+			const TVector<FRealDouble, 3> max = Max(V1, V2);
+			if (max.X > max.Y)
+			{
+				if (max.X > max.Z)
+					return MakePair(max.X, 0);
+				else
+					return MakePair(max.Z, 2);
+			}
+			else
+			{
+				if (max.Y > max.Z)
+					return MakePair(max.Y, 1);
+				else
+					return MakePair(max.Z, 2);
+			}
+		}
+		FRealDouble SafeNormalize(FRealDouble Epsilon = 1e-4f)
+		{
+			FRealDouble Size = SizeSquared();
+			if (Size < Epsilon)
+			{
+				*this = AxisVector(0);
+				return 0.f;
+			}
+			Size = sqrt(Size);
+			*this = (*this) / Size;
+			return Size;
+		}
+		TVector<FRealDouble, 3> GetOrthogonalVector() const
+		{
+			TVector<FRealDouble, 3> AbsVector(FMath::Abs(X), FMath::Abs(Y), FMath::Abs(Z));
+			if ((AbsVector.X <= AbsVector.Y) && (AbsVector.X <= AbsVector.Z))
+			{
+				// X is the smallest component
+				return TVector<FRealDouble, 3>(0, Z, -Y);
+			}
+			if ((AbsVector.Z <= AbsVector.X) && (AbsVector.Z <= AbsVector.Y))
+			{
+				// Z is the smallest component
+				return TVector<FRealDouble, 3>(Y, -X, 0);
+			}
+			// Y is the smallest component
+			return TVector<FRealDouble, 3>(-Z, 0, X);
+		}
+		static FRealDouble AngleBetween(const TVector<FRealDouble, 3>& V1, const TVector<FRealDouble, 3>& V2)
+		{
+			FRealDouble s = CrossProduct(V1, V2).Size();
+			FRealDouble c = DotProduct(V1, V2);
+			return atan2(s, c);
+		}
+		/** Calculate the velocity to move from P0 to P1 in time Dt. Exists just for symmetry with TRotation::CalculateAngularVelocity! */
+		static TVector<FRealDouble, 3> CalculateVelocity(const TVector<FRealDouble, 3>& P0, const TVector<FRealDouble, 3>& P1, const FRealDouble Dt)
+		{
+			return (P1 - P0) / Dt;
+		}
+
+		static bool IsNearlyEqual(const TVector<FRealDouble, 3>& A, const TVector<FRealDouble, 3>& B, const FRealDouble Epsilon)
+		{
+			return (B - A).IsNearlyZero(Epsilon);
+		}
+	};
+
+
 	template<>
 	class TVector<FReal, 2> : public FVector2D
 	{
diff --git a/Engine/Source/Runtime/NavigationSystem/Private/NavMesh/RecastNavMeshGenerator.cpp b/Engine/Source/Runtime/NavigationSystem/Private/NavMesh/RecastNavMeshGenerator.cpp
--- a/Engine/Source/Runtime/NavigationSystem/Private/NavMesh/RecastNavMeshGenerator.cpp
+++ b/Engine/Source/Runtime/NavigationSystem/Private/NavMesh/RecastNavMeshGenerator.cpp
@@ -661,7 +661,7 @@ void ExportChaosTriMesh(const Chaos::FTriangleMeshImplicitObject* const TriMesh,
 	auto LambdaHelper = [&](const auto& Triangles)
 	{
 		int32 NumTris = Triangles.Num();
-		const FParticles& Vertices = TriMesh->Particles();
+		const Chaos::FTriangleMeshImplicitObject::ParticlesType& Vertices = TriMesh->Particles();
 	
 		VertexBuffer.Reserve(VertexBuffer.Num() + NumTris * 9);
 		IndexBuffer.Reserve(IndexBuffer.Num() + NumTris * 3);
