Fix for LWC build error
#rnx

diff --git a/Engine/Source/Runtime/Experimental/Chaos/Private/Chaos/CollisionResolution.cpp b/Engine/Source/Runtime/Experimental/Chaos/Private/Chaos/CollisionResolution.cpp
--- a/Engine/Source/Runtime/Experimental/Chaos/Private/Chaos/CollisionResolution.cpp
+++ b/Engine/Source/Runtime/Experimental/Chaos/Private/Chaos/CollisionResolution.cpp
@@ -193,12 +193,12 @@ namespace Chaos
 					Dir /= Length;
 
 					// Do not perform CCD if the vector is not close to unit length to prevent getting caught in a large or infinite loop when raycasting.
-					if (FMath::IsNearlyEqual(Dir.SizeSquared(), 1.f, KINDA_SMALL_NUMBER))
+					if (FMath::IsNearlyEqual((float)Dir.SizeSquared(), 1.f, KINDA_SMALL_NUMBER))
 					{
-					return true;
+						return true;
+					}
 				}
 			}
-			}
 
 			return false;
 		}
