Chaos Event Manager
-- defer handler registration until event handling has finished to avoid event container access deadlock.

diff --git a/Engine/Source/Runtime/Experimental/Chaos/Private/EventManager.cpp b/Engine/Source/Runtime/Experimental/Chaos/Private/EventManager.cpp
--- a/Engine/Source/Runtime/Experimental/Chaos/Private/EventManager.cpp
+++ b/Engine/Source/Runtime/Experimental/Chaos/Private/EventManager.cpp
@@ -75,6 +75,11 @@ namespace Chaos
 
 	void FEventManager::DispatchEvents()
 	{
+		FScopeLock ScopeLock(&AccessDeferredHandlersLock);
+
+		check(!bCurrentlyDispatchingEvents);
+		bCurrentlyDispatchingEvents = true;
+		
 		if (BufferMode == EMultiBufferMode::Double)
 		{
 			ResourceLock.ReadLock();
@@ -94,7 +99,22 @@ namespace Chaos
 		{
 			ResourceLock.ReadUnlock();
 		}
+		
+		ensure(bCurrentlyDispatchingEvents);
+		bCurrentlyDispatchingEvents = false;
 
+		// If we deferred any handler registration in RegisterHandler, complete it now. 
+		if (DeferredHandlers.Num() > 0)
+		{
+			ContainerLock.WriteLock();
+			for (TPair<FEventID, IEventHandler*>& DeferredHandler : DeferredHandlers)
+			{
+				checkf(DeferredHandler.Key < EventContainers.Num(), TEXT("Registering event Handler for an event ID that does not exist"));
+				EventContainers[DeferredHandler.Key]->RegisterHandler(DeferredHandler.Value);
+			}
+			DeferredHandlers.Reset();
+			ContainerLock.WriteUnlock();
+		}
 	}
 
 	void FEventManager::InternalRegisterInjector(const FEventID& EventID, const FEventContainerBasePtr& Container)
diff --git a/Engine/Source/Runtime/Experimental/Chaos/Public/EventManager.h b/Engine/Source/Runtime/Experimental/Chaos/Public/EventManager.h
--- a/Engine/Source/Runtime/Experimental/Chaos/Public/EventManager.h
+++ b/Engine/Source/Runtime/Experimental/Chaos/Public/EventManager.h
@@ -250,7 +250,10 @@ namespace Chaos
 
 	public:
 
-		FEventManager(const Chaos::EMultiBufferMode& BufferModeIn) : BufferMode(BufferModeIn) {}
+		FEventManager(const Chaos::EMultiBufferMode& BufferModeIn) 
+			: BufferMode(BufferModeIn)
+			, bCurrentlyDispatchingEvents(false)
+			{}
 
 		~FEventManager()
 		{
@@ -307,11 +310,22 @@ namespace Chaos
 		template<typename PayloadType, typename HandlerType>
 		void RegisterHandler(const EEventType& EventType, HandlerType* Handler, typename TRawEventHandler<PayloadType, HandlerType>::FHandlerFunction HandlerFunction)
 		{
+			FScopeLock ScopeLock(&AccessDeferredHandlersLock);
+
 			const FEventID EventID = FEventID(EventType);
-			ContainerLock.WriteLock();
-			checkf(EventID < EventContainers.Num(), TEXT("Registering event Handler for an event ID that does not exist"));
-			EventContainers[EventID]->RegisterHandler(new TRawEventHandler<PayloadType, HandlerType>(Handler, HandlerFunction));
-			ContainerLock.WriteUnlock();
+			
+			// If we are currently dispatching events, defer handler registration until completion of dispatch to avoid deadlock
+			if (bCurrentlyDispatchingEvents)
+			{
+				DeferredHandlers.Add(TPair<FEventID, IEventHandler*>(EventID, new TRawEventHandler<PayloadType, HandlerType>(Handler, HandlerFunction)));
+			}
+			else
+			{
+				ContainerLock.WriteLock();
+				checkf(EventID < EventContainers.Num(), TEXT("Registering event Handler for an event ID that does not exist"));
+				EventContainers[EventID]->RegisterHandler(new TRawEventHandler<PayloadType, HandlerType>(Handler, HandlerFunction));
+				ContainerLock.WriteUnlock();
+			}	
 		}
 
 		/**
@@ -359,6 +373,11 @@ namespace Chaos
 		TArray<FEventContainerBasePtr> EventContainers;	// Array of event types
 		FRWLock ResourceLock;
 		FRWLock ContainerLock;
+		FCriticalSection AccessDeferredHandlersLock;
+
+		/** Defer handler registration if we are currently dispatching events. */
+		bool bCurrentlyDispatchingEvents;
+		TArray<TPair<FEventID, IEventHandler*>> DeferredHandlers;
 
 	};
 
