GeometryCollection
-- Removed unused data from GeometryCollectionProxyData
-- Changed managed arrays in GeometryCollectionProxyData to TArray and deferred initialization until necessary.
-- Added scoped cycle counters for finer profiling in physics buffering

diff --git a/Engine/Source/Runtime/Experimental/Chaos/Private/ChaosStats.cpp b/Engine/Source/Runtime/Experimental/Chaos/Private/ChaosStats.cpp
--- a/Engine/Source/Runtime/Experimental/Chaos/Private/ChaosStats.cpp
+++ b/Engine/Source/Runtime/Experimental/Chaos/Private/ChaosStats.cpp
@@ -37,6 +37,9 @@ DEFINE_STAT(STAT_FlipResults);
 DEFINE_STAT(STAT_ProcessDeferredCreatePhysicsState);
 DEFINE_STAT(STAT_SqUpdateMaterials);
 DEFINE_STAT(STAT_CacheResultGeomCollection);
+DEFINE_STAT(STAT_UpdateGeometryCollectionViews);
+DEFINE_STAT(STAT_BufferPhysicsResultsParticleLoop);
+DEFINE_STAT(STAT_CaptureSolverData);
 DEFINE_STAT(STAT_CacheResultStaticMesh);
 DEFINE_STAT(STAT_CaptureDisabledState);
 DEFINE_STAT(STAT_CalcGlobalGCMatrices);
diff --git a/Engine/Source/Runtime/Experimental/Chaos/Private/PBDRigidActiveParticlesBuffer.cpp b/Engine/Source/Runtime/Experimental/Chaos/Private/PBDRigidActiveParticlesBuffer.cpp
--- a/Engine/Source/Runtime/Experimental/Chaos/Private/PBDRigidActiveParticlesBuffer.cpp
+++ b/Engine/Source/Runtime/Experimental/Chaos/Private/PBDRigidActiveParticlesBuffer.cpp
@@ -15,6 +15,8 @@ namespace Chaos
 
 	void FPBDRigidDirtyParticlesBuffer::CaptureSolverData(FPBDRigidsSolver* Solver)
 	{
+		SCOPE_CYCLE_COUNTER(STAT_CaptureSolverData);
+
 		WriteLock();
 		BufferPhysicsResults(Solver);
 		FlipDataOut();
diff --git a/Engine/Source/Runtime/Experimental/Chaos/Private/PhysicsProxy/GeometryCollectionPhysicsProxy.cpp b/Engine/Source/Runtime/Experimental/Chaos/Private/PhysicsProxy/GeometryCollectionPhysicsProxy.cpp
--- a/Engine/Source/Runtime/Experimental/Chaos/Private/PhysicsProxy/GeometryCollectionPhysicsProxy.cpp
+++ b/Engine/Source/Runtime/Experimental/Chaos/Private/PhysicsProxy/GeometryCollectionPhysicsProxy.cpp
@@ -65,24 +65,18 @@ DEFINE_LOG_CATEGORY_STATIC(UGCC_LOG, Error, All);
 //==============================================================================
 
 FGeometryCollectionResults::FGeometryCollectionResults()
-	: BaseIndex(0)
-	, NumParticlesAdded(0)
-	, IsObjectDynamic(false)
+	: IsObjectDynamic(false)
 	, IsObjectLoading(false)
-	, WorldBounds(ForceInit)
 {}
 
 void FGeometryCollectionResults::Reset()
 {
 	SolverDt = 0.0f;
-	BaseIndex = 0;
-	NumParticlesAdded = 0;
 	DisabledStates.SetNum(0);
 	GlobalTransforms.SetNum(0);
 	ParticleToWorldTransforms.SetNum(0);
 	IsObjectDynamic = false;
 	IsObjectLoading = false;
-	WorldBounds = FBoxSphereBounds(ForceInit);
 }
 
 //==============================================================================
@@ -1504,13 +1498,8 @@ void FGeometryCollectionPhysicsProxy::BufferPhysicsResults(Chaos::FPBDRigidsSolv
 	if (TargetResults.NumTransformGroup() != NumTransformGroupElements)
 	{
 		TargetResults.InitArrays(PhysicsThreadCollection);
-
-		// Base particle index to calculate index from a global particle index on the game thread
-		TargetResults.BaseIndex = BaseParticleIndex;
-		TargetResults.NumParticlesAdded = NumParticles;
 	}
 
-	
 	const FTransform& ActorToWorld = Parameters.WorldTransform;
 	const TManagedArray<int32>& Parent = PhysicsThreadCollection.Parent;
 	const TManagedArray<TSet<int32>>& Children = PhysicsThreadCollection.Children;
@@ -1518,14 +1507,12 @@ void FGeometryCollectionPhysicsProxy::BufferPhysicsResults(Chaos::FPBDRigidsSolv
 	if(NumTransformGroupElements > 0)
 	{ 
 		SCOPE_CYCLE_COUNTER(STAT_CalcParticleToWorld);
-
-		// initialize Target Results
-		TargetResults.Transforms.Init(PhysicsThreadCollection.Transform);
-		TargetResults.Children.Init(PhysicsThreadCollection.Children);
-		TargetResults.Parent.Init(PhysicsThreadCollection.Parent);
-
+		
 		for (int32 TransformGroupIndex = 0; TransformGroupIndex < NumTransformGroupElements; ++TransformGroupIndex)
 		{
+			TargetResults.Transforms[TransformGroupIndex] = PhysicsThreadCollection.Transform[TransformGroupIndex];
+			TargetResults.Parent[TransformGroupIndex] = PhysicsThreadCollection.Parent[TransformGroupIndex];
+
 			TargetResults.DisabledStates[TransformGroupIndex] = true;
 			Chaos::TPBDRigidClusteredParticleHandle<Chaos::FReal, 3>* Handle = SolverParticleHandles[TransformGroupIndex];
 			if (!Handle)
@@ -1667,24 +1654,15 @@ void FGeometryCollectionPhysicsProxy::BufferPhysicsResults(Chaos::FPBDRigidsSolv
 			PhysicsThreadCollection.Active[TransformGroupIndex] = !TargetResults.DisabledStates[TransformGroupIndex];
 		}    // end for
 	}        // STAT_CalcParticleToWorld scope
-
+	
 	// If object is dynamic, compute global matrices	
 	if (IsObjectDynamic || TargetResults.GlobalTransforms.Num() == 0)
 	{
 		SCOPE_CYCLE_COUNTER(STAT_CalcGlobalGCMatrices);
 		check(TargetResults.Transforms.Num() == TargetResults.Parent.Num());
 		GeometryCollectionAlgo::GlobalMatrices(TargetResults.Transforms, TargetResults.Parent, TargetResults.GlobalTransforms);
-
-		// Compute world bounds.  This is a loose bounds based on the circumscribed box 
-		// of a bounding sphere for the geometry.		
-		SCOPE_CYCLE_COUNTER(STAT_CalcGlobalGCBounds);
-		FBox BoundingBox(ForceInit);
-		for (int i = 0; i < ValidGeometryBoundingBoxes.Num(); ++i)
-		{
-			BoundingBox += ValidGeometryBoundingBoxes[i].TransformBy(FTransform(TargetResults.GlobalTransforms[ValidGeometryTransformIndices[i]]) * ActorToWorld);
-		}
-		TargetResults.WorldBounds = FBoxSphereBounds(BoundingBox);
 	}
+	
 
 	// Advertise to game thread
 	TargetResults.IsObjectDynamic = IsObjectDynamic;
diff --git a/Engine/Source/Runtime/Experimental/Chaos/Public/Chaos/PBDRigidsSOAs.h b/Engine/Source/Runtime/Experimental/Chaos/Public/Chaos/PBDRigidsSOAs.h
--- a/Engine/Source/Runtime/Experimental/Chaos/Public/Chaos/PBDRigidsSOAs.h
+++ b/Engine/Source/Runtime/Experimental/Chaos/Public/Chaos/PBDRigidsSOAs.h
@@ -725,6 +725,8 @@ public:
 	 */
 	void UpdateGeometryCollectionViews(const bool ForceUpdateViews=false)
 	{
+		SCOPE_CYCLE_COUNTER(STAT_UpdateGeometryCollectionViews);
+
 		if (!GeometryCollectionParticles)
 		{
 			if (ForceUpdateViews)
diff --git a/Engine/Source/Runtime/Experimental/Chaos/Public/ChaosStats.h b/Engine/Source/Runtime/Experimental/Chaos/Public/ChaosStats.h
--- a/Engine/Source/Runtime/Experimental/Chaos/Public/ChaosStats.h
+++ b/Engine/Source/Runtime/Experimental/Chaos/Public/ChaosStats.h
@@ -52,6 +52,9 @@ DECLARE_CYCLE_STAT_EXTERN(TEXT("Flip Results"), STAT_FlipResults, STATGROUP_Chao
 DECLARE_CYCLE_STAT_EXTERN(TEXT("ProcessDeferredCreatePhysicsState"), STAT_ProcessDeferredCreatePhysicsState, STATGROUP_Chaos, CHAOS_API);
 DECLARE_CYCLE_STAT_EXTERN(TEXT("SQ - Update Materials"), STAT_SqUpdateMaterials, STATGROUP_Chaos, CHAOS_API);
 DECLARE_CYCLE_STAT_EXTERN(TEXT("[BufferPhysicsResults] - Geometry Collection"), STAT_CacheResultGeomCollection, STATGROUP_ChaosWide, CHAOS_API);
+DECLARE_CYCLE_STAT_EXTERN(TEXT("Update GC Views"), STAT_UpdateGeometryCollectionViews, STATGROUP_Chaos, CHAOS_API);
+DECLARE_CYCLE_STAT_EXTERN(TEXT("Particle Loop"), STAT_BufferPhysicsResultsParticleLoop, STATGROUP_Chaos, CHAOS_API);
+DECLARE_CYCLE_STAT_EXTERN(TEXT("Capture Solver Data"), STAT_CaptureSolverData, STATGROUP_Chaos, CHAOS_API);
 DECLARE_CYCLE_STAT_EXTERN(TEXT("[BufferPhysicsResults] - StaticMesh"), STAT_CacheResultStaticMesh, STATGROUP_Chaos, CHAOS_API);
 DECLARE_CYCLE_STAT_EXTERN(TEXT("Capture Disabled State"), STAT_CaptureDisabledState, STATGROUP_Chaos, CHAOS_API);
 DECLARE_CYCLE_STAT_EXTERN(TEXT("Calc Global Matrices"), STAT_CalcGlobalGCMatrices, STATGROUP_Chaos, CHAOS_API);
diff --git a/Engine/Source/Runtime/Experimental/Chaos/Public/GeometryCollectionProxyData.h b/Engine/Source/Runtime/Experimental/Chaos/Public/GeometryCollectionProxyData.h
--- a/Engine/Source/Runtime/Experimental/Chaos/Public/GeometryCollectionProxyData.h
+++ b/Engine/Source/Runtime/Experimental/Chaos/Public/GeometryCollectionProxyData.h
@@ -129,42 +129,25 @@ public:
 
 	void InitArrays(const FGeometryDynamicCollection& Other)
 	{
-		// Managed arrays
-		Transforms.Init(Other.Transform);
-		DynamicState.Init(Other.DynamicState);
-		Parent.Init(Other.Parent);
-		Children.Init(Other.Children);
-		SimulationType.Init(Other.SimulationType);
-
-		// Arrays
 		const int32 NumTransforms = Other.NumElements(FGeometryCollection::TransformGroup);
 		DisabledStates.SetNumUninitialized(NumTransforms);
 		GlobalTransforms.SetNumUninitialized(NumTransforms);
 		ParticleToWorldTransforms.SetNumUninitialized(NumTransforms);
+
+		Transforms.SetNumUninitialized(NumTransforms);
+		Parent.SetNumUninitialized(NumTransforms);
+		DynamicState.SetNumUninitialized(NumTransforms);
 	}
 
 	Chaos::FReal SolverDt;
-	int32 BaseIndex;
-	int32 NumParticlesAdded;
 	TArray<bool> DisabledStates;
 	TArray<FMatrix> GlobalTransforms;
 	TArray<FTransform> ParticleToWorldTransforms;
 
-	TManagedArray<int32> TransformIndex;
-
-	TManagedArray<FTransform> Transforms;
-	TManagedArray<int32> BoneMap;
-	TManagedArray<int32> Parent;
-	TManagedArray<TSet<int32>> Children;
-	TManagedArray<int32> SimulationType;
-	TManagedArray<int32> DynamicState;
-	TManagedArray<float> Mass;
-	TManagedArray<FVector3f> InertiaTensor;
-
-	TManagedArray<int32> ClusterId;
-
+	TArray<FTransform> Transforms;
+	TArray<int32> Parent;
+	TArray<int32> DynamicState;
+	
 	bool IsObjectDynamic;
 	bool IsObjectLoading;
-
-	FBoxSphereBounds WorldBounds;
 };
