Add FPhysicsThreadContextScope to FPhysicsSolverBase::DestroySolver so we do not hit ensure in dev builds when calling physics API on game thread during solver shutdown.

diff --git a/Engine/Source/Runtime/Experimental/Chaos/Private/Chaos/Framework/PhysicsSolverBase.cpp b/Engine/Source/Runtime/Experimental/Chaos/Private/Chaos/Framework/PhysicsSolverBase.cpp
--- a/Engine/Source/Runtime/Experimental/Chaos/Private/Chaos/Framework/PhysicsSolverBase.cpp
+++ b/Engine/Source/Runtime/Experimental/Chaos/Private/Chaos/Framework/PhysicsSolverBase.cpp
@@ -216,7 +216,14 @@ namespace Chaos
 		// Advance in single threaded because we cannot block on an async task here if in multi threaded mode. see above comments.
 		InSolver.SetThreadingMode_External(EThreadingModeTemp::SingleThread);
 		InSolver.MarkShuttingDown();
-		InSolver.AdvanceAndDispatch_External(0);	//flush any pending commands are executed (for example unregister object)
+		{
+#if PHYSICS_THREAD_CONTEXT
+			// Must use physics thread context, as we are about to execute physics thread API from the game thread in single threaded mode.
+			FPhysicsThreadContextScope PhysicsThreadContext(/*InParentIsPhysicsThreadContext=*/true);
+#endif
+
+			InSolver.AdvanceAndDispatch_External(0);	//flush any pending commands are executed (for example unregister object)
+		}
 
 		// verify callbacks have been processed and we're not leaking.
 		// TODO: why is this still firing in 14.30? (Seems we're still leaking)
