Change ensure to a log warning

diff --git a/Engine/Source/Runtime/Experimental/ChaosVehicles/ChaosVehiclesCore/Private/SuspensionUtility.cpp b/Engine/Source/Runtime/Experimental/ChaosVehicles/ChaosVehiclesCore/Private/SuspensionUtility.cpp
--- a/Engine/Source/Runtime/Experimental/ChaosVehicles/ChaosVehiclesCore/Private/SuspensionUtility.cpp
+++ b/Engine/Source/Runtime/Experimental/ChaosVehicles/ChaosVehiclesCore/Private/SuspensionUtility.cpp
@@ -190,8 +190,9 @@ bool FSuspensionUtility::ComputeSprungMasses(const TArray<FVector>& MassSpringPo
 		- (XDotY * Count);
 
 	// Make sure the matrix is invertible!
-	if (!ensureMsgf(DetLL > SMALL_NUMBER || DetLL < -SMALL_NUMBER, TEXT("Spring configuration is invalid! Please make sure no two springs are at the same location.")))
+	if (!(DetLL > SMALL_NUMBER || DetLL < -SMALL_NUMBER))
 	{
+		UE_LOG(LogVehicleUtility, Warning, TEXT("Spring configuration is invalid! Please make sure no two springs are at the same location."));
 		return false;
 	}
 
