GeometryCollection
-- Cleaned up warning that was spamming logs unnecesarily.

diff --git a/Engine/Source/Runtime/Experimental/Chaos/Private/GeometryCollection/GeometryCollection.cpp b/Engine/Source/Runtime/Experimental/Chaos/Private/GeometryCollection/GeometryCollection.cpp
--- a/Engine/Source/Runtime/Experimental/Chaos/Private/GeometryCollection/GeometryCollection.cpp
+++ b/Engine/Source/Runtime/Experimental/Chaos/Private/GeometryCollection/GeometryCollection.cpp
@@ -1007,7 +1007,6 @@ void FGeometryCollection::Serialize(Chaos::FChaosArchive& Ar)
 		// Version 6 introduced the Exemplar Index array
 		if (Version < 6)
 		{
-			UE_LOG(FGeometryCollectionLogging, Warning, TEXT("GeometryCollection has uninitialized Exemplar Indices. Setting all indices to NONE."));
 			ExemplarIndex.Fill(INDEX_NONE);
 
 			// Structure is conditioned, now considered up to date.
