GeometryCollectionComponent
-- Add DisableThresholdand SleepingThreshold options to the GeometryCollectionComponents ApplyPhysicsField(...) method

diff --git a/Engine/Source/Runtime/Experimental/Chaos/Private/PhysicsProxy/GeometryCollectionPhysicsProxy.cpp b/Engine/Source/Runtime/Experimental/Chaos/Private/PhysicsProxy/GeometryCollectionPhysicsProxy.cpp
--- a/Engine/Source/Runtime/Experimental/Chaos/Private/PhysicsProxy/GeometryCollectionPhysicsProxy.cpp
+++ b/Engine/Source/Runtime/Experimental/Chaos/Private/PhysicsProxy/GeometryCollectionPhysicsProxy.cpp
@@ -2775,6 +2775,16 @@ void FGeometryCollectionPhysicsProxy::FieldParameterUpdateCallback(Chaos::FPBDRi
 								FieldContext, PositionTarget, TargetedParticles, FinalResults);
 						}
 					}
+					else if (FieldCommand.RootNode->Type() == FFieldNodeBase::EFieldType::EField_Float)
+					{
+						TArray<float>& FinalResults = ExecutionDatas.ScalarResults[(uint8)EFieldCommandResultType::FinalResult];
+						ResetResultsArray<float>(ExecutionDatas.SamplePositions.Num(), FinalResults, 0.0f);
+
+						TFieldArrayView<float> ResultsView(FinalResults, 0, FinalResults.Num());
+						
+						Chaos::FieldScalarParameterUpdate(RigidSolver, FieldCommand, ParticleHandles,
+							FieldContext, PositionTarget, TargetedParticles, FinalResults);
+					}
 				}
 				CommandsToRemove.Add(CommandIndex);
 			}
diff --git a/Engine/Source/Runtime/Experimental/Chaos/Public/GeometryCollection/GeometryCollectionSimulationTypes.h b/Engine/Source/Runtime/Experimental/Chaos/Public/GeometryCollection/GeometryCollectionSimulationTypes.h
--- a/Engine/Source/Runtime/Experimental/Chaos/Public/GeometryCollection/GeometryCollectionSimulationTypes.h
+++ b/Engine/Source/Runtime/Experimental/Chaos/Public/GeometryCollection/GeometryCollectionSimulationTypes.h
@@ -52,6 +52,8 @@ enum class EGeometryCollectionPhysicsTypeEnum : uint8
 	Chaos_CollisionGroup           UMETA(DisplayName = "Collision Group", ToolTip = "Set the particles collision group."),
 	Chaos_LinearForce              UMETA(DisplayName = "Linear Force", ToolTip = "Add a vector field to the particles linear force."),
 	Chaos_AngularTorque            UMETA(DisplayName = "Angular Torque", ToolTip = "Add a vector field to the particles angular torque."),
+	Chaos_DisableThreshold         UMETA(DisplayName = "Disable Threshold", ToolTip = "Disable the particles if their linear and angular velocity are less than the threshold."),
+	Chaos_SleepingThreshold        UMETA(DisplayName = "Sleeping Threshold", ToolTip = "Set particles in sleeping mode if their linear and angular velocity are less than the threshold."),
 	//
 	Chaos_Max						UMETA(Hidden)
 };
@@ -60,7 +62,7 @@ inline CHAOS_API EFieldPhysicsType GetGeometryCollectionPhysicsType(const EGeome
 {
 	static const TArray<EFieldPhysicsType> PhysicsTypes = { EFieldPhysicsType::Field_AngularVelociy , EFieldPhysicsType::Field_DynamicState, EFieldPhysicsType::Field_LinearVelocity, 
 		EFieldPhysicsType::Field_InitialAngularVelocity, EFieldPhysicsType::Field_InitialLinearVelocity, EFieldPhysicsType::Field_CollisionGroup, 
-		EFieldPhysicsType::Field_LinearForce, EFieldPhysicsType::Field_AngularTorque, EFieldPhysicsType::Field_PhysicsType_Max };
+		EFieldPhysicsType::Field_LinearForce, EFieldPhysicsType::Field_AngularTorque, EFieldPhysicsType::Field_DisableThreshold, EFieldPhysicsType::Field_SleepingThreshold, EFieldPhysicsType::Field_PhysicsType_Max };
 
 	return PhysicsTypes[(uint8)GeoCollectionType];
 }
diff --git a/Engine/Source/Runtime/Experimental/GeometryCollectionEngine/Private/GeometryCollection/GeometryCollectionActor.cpp b/Engine/Source/Runtime/Experimental/GeometryCollectionEngine/Private/GeometryCollection/GeometryCollectionActor.cpp
--- a/Engine/Source/Runtime/Experimental/GeometryCollectionEngine/Private/GeometryCollection/GeometryCollectionActor.cpp
+++ b/Engine/Source/Runtime/Experimental/GeometryCollectionEngine/Private/GeometryCollection/GeometryCollectionActor.cpp
@@ -49,6 +49,9 @@ AGeometryCollectionActor::AGeometryCollectionActor(const FObjectInitializer& Obj
 void AGeometryCollectionActor::Tick(float DeltaTime) 
 {
 	UE_LOG(AGeometryCollectionActorLogging, Verbose, TEXT("AGeometryCollectionActor::Tick()"));
+
+	Super::Tick(DeltaTime);
+
 	if (GeometryCollectionComponent)
 	{
 		GeometryCollectionComponent->SetRenderStateDirty();
