Chaos : GeometryCollection
Fix BoxCalcBounds with small box array
- preflight against Base 17099190

diff --git a/Engine/Source/Runtime/Experimental/GeometryCollectionEngine/Private/GeometryCollection/GeometryCollectionComponent.ispc b/Engine/Source/Runtime/Experimental/GeometryCollectionEngine/Private/GeometryCollection/GeometryCollectionComponent.ispc
--- a/Engine/Source/Runtime/Experimental/GeometryCollectionEngine/Private/GeometryCollection/GeometryCollectionComponent.ispc
+++ b/Engine/Source/Runtime/Experimental/GeometryCollectionEngine/Private/GeometryCollection/GeometryCollectionComponent.ispc
@@ -109,7 +109,8 @@ export void BoxCalcBounds(  const uniform int TransformToGeometryIndex[],
 		#pragma ignore warning(perf)
 		const FBox B = BoundingBoxes[BoxIdx];
 
-		if (TransformToGeometryIndex[BoxIdx] != INDEX_NONE && B.IsValid)
+		#pragma ignore warning(perf)
+		if (TransformToGeometryIndex[TransformIndex] != INDEX_NONE && B.IsValid)
 		{
 			#pragma ignore warning(perf)
 			const FMatrix M = GlobalMatrices[TransformIndex];
@@ -122,11 +123,14 @@ export void BoxCalcBounds(  const uniform int TransformToGeometryIndex[],
 	#pragma ignore warning(perf)
 	ReducedBox[programIndex] = BoundingBoxSum;
 
-	foreach_active(j)
+	if (programIndex < NumBoxes)
 	{
-		unmasked
+		foreach_active(j)
 		{
-			BoundingBox = BoundingBox + ReducedBox[j];
+			unmasked
+			{
+				BoundingBox = BoundingBox + ReducedBox[j];
+			}
 		}
 	}
 }
\ No newline at end of file
