Chaos Physics Solver
-- created solver notification on solver destruction
-- ChaosCacheManager subscribes to this notification and removes dangling references to the destroyed solver.

diff --git a/Engine/Plugins/Experimental/ChaosCaching/Source/ChaosCaching/Private/Chaos/CacheManagerActor.cpp b/Engine/Plugins/Experimental/ChaosCaching/Source/ChaosCaching/Private/Chaos/CacheManagerActor.cpp
--- a/Engine/Plugins/Experimental/ChaosCaching/Source/ChaosCaching/Private/Chaos/CacheManagerActor.cpp
+++ b/Engine/Plugins/Experimental/ChaosCaching/Source/ChaosCaching/Private/Chaos/CacheManagerActor.cpp
@@ -304,6 +304,7 @@ void AChaosCacheManager::BeginPlay()
 				SolverData->PreSolveHandle  = Solver->AddPreAdvanceCallback(FSolverPreAdvance::FDelegate::CreateUObject(this, &AChaosCacheManager::HandlePreSolve, Solver));
 				SolverData->PreBufferHandle = Solver->AddPreBufferCallback(FSolverPreAdvance::FDelegate::CreateUObject(this, &AChaosCacheManager::HandlePreBuffer, Solver));
 				SolverData->PostSolveHandle = Solver->AddPostAdvanceCallback(FSolverPostAdvance::FDelegate::CreateUObject(this, &AChaosCacheManager::HandlePostSolve, Solver));
+				SolverData->TeardownHandle  = Solver->AddTeardownCallback(FSolverTeardown::FDelegate::CreateUObject(this, &AChaosCacheManager::HandleTeardown, Solver));
 			}
 
 			if (CacheMode == ECacheMode::Play || CacheMode == ECacheMode::None || !CanRecord())
@@ -518,6 +519,12 @@ void AChaosCacheManager::HandlePostSolve(Chaos::FReal InDt, Chaos::FPhysicsSolve
 	});
 }
 
+void AChaosCacheManager::HandleTeardown(Chaos::FPhysicsSolver* InSolver)
+{
+	// Solver has been deleted, we remove all reference to it to avoid dereferencing a dangling pointer.
+	PerSolverData.Remove(InSolver);
+}
+
 void AChaosCacheManager::OnStartFrameChanged(Chaos::FReal InTime)
 {
 	if (bIsSimulating)
diff --git a/Engine/Plugins/Experimental/ChaosCaching/Source/ChaosCaching/Public/Chaos/CacheManagerActor.h b/Engine/Plugins/Experimental/ChaosCaching/Source/ChaosCaching/Public/Chaos/CacheManagerActor.h
--- a/Engine/Plugins/Experimental/ChaosCaching/Source/ChaosCaching/Public/Chaos/CacheManagerActor.h
+++ b/Engine/Plugins/Experimental/ChaosCaching/Source/ChaosCaching/Public/Chaos/CacheManagerActor.h
@@ -89,6 +89,7 @@ struct FPerSolverData
 	FDelegateHandle PreSolveHandle;
 	FDelegateHandle PreBufferHandle;
 	FDelegateHandle PostSolveHandle;
+	FDelegateHandle TeardownHandle;
 
 	/** List of the tick records for each playback index, tracks where the last tick was */
 	TArray<FPlaybackTickRecord> PlaybackTickRecords;
@@ -193,6 +194,9 @@ protected:
 	/** Handles physics thread post-solve (record data for components under record) */
 	void HandlePostSolve(Chaos::FReal InDt, Chaos::FPhysicsSolver* InSolver);
 
+	/** Handles solver teardown due to solver destruction / stream-out */
+	void HandleTeardown(Chaos::FPhysicsSolver* InSolver);
+
 	/** Evaluates and sets state for all observed components at the specified time. */
 	void OnStartFrameChanged(Chaos::FReal InT);
 
diff --git a/Engine/Source/Runtime/Experimental/Chaos/Private/Chaos/Framework/PhysicsSolverBase.cpp b/Engine/Source/Runtime/Experimental/Chaos/Private/Chaos/Framework/PhysicsSolverBase.cpp
--- a/Engine/Source/Runtime/Experimental/Chaos/Private/Chaos/Framework/PhysicsSolverBase.cpp
+++ b/Engine/Source/Runtime/Experimental/Chaos/Private/Chaos/Framework/PhysicsSolverBase.cpp
@@ -48,6 +48,16 @@ namespace Chaos
 		return EventPostSolve.Remove(InHandle);
 	}
 
+	FDelegateHandle FPhysicsSolverBase::AddTeardownCallback(FSolverTeardown::FDelegate InDelegate)
+	{
+		return EventTeardown.Add(InDelegate);
+	}
+
+	bool FPhysicsSolverBase::RemoveTeardownCallback(FDelegateHandle InHandle)
+	{
+		return EventTeardown.Remove(InHandle);
+	}
+
 	FAutoConsoleTaskPriority CPrio_FPhysicsTickTask(
 		TEXT("TaskGraph.TaskPriorities.PhysicsTickTask"),
 		TEXT("Task and thread priotiry for Chaos physics tick"),
diff --git a/Engine/Source/Runtime/Experimental/Chaos/Private/PBDRigidsSolver.cpp b/Engine/Source/Runtime/Experimental/Chaos/Private/PBDRigidsSolver.cpp
--- a/Engine/Source/Runtime/Experimental/Chaos/Private/PBDRigidsSolver.cpp
+++ b/Engine/Source/Runtime/Experimental/Chaos/Private/PBDRigidsSolver.cpp
@@ -1528,7 +1528,10 @@ namespace Chaos
 		SetUseContactGraph(InConfig.bGenerateContactGraph);
 	}
 
-	FPBDRigidsSolver::~FPBDRigidsSolver() = default;
+	FPBDRigidsSolver::~FPBDRigidsSolver()
+	{
+		EventTeardown.Broadcast();
+	}
 
 	void Chaos::FPBDRigidsSolver::FieldParameterUpdateCallback(
 		Chaos::FPBDPositionConstraints& PositionTarget,
diff --git a/Engine/Source/Runtime/Experimental/Chaos/Public/Chaos/Framework/PhysicsSolverBase.h b/Engine/Source/Runtime/Experimental/Chaos/Public/Chaos/Framework/PhysicsSolverBase.h
--- a/Engine/Source/Runtime/Experimental/Chaos/Public/Chaos/Framework/PhysicsSolverBase.h
+++ b/Engine/Source/Runtime/Experimental/Chaos/Public/Chaos/Framework/PhysicsSolverBase.h
@@ -17,6 +17,7 @@ class FChaosSolversModule;
 DECLARE_MULTICAST_DELEGATE_OneParam(FSolverPreAdvance, Chaos::FReal);
 DECLARE_MULTICAST_DELEGATE_OneParam(FSolverPreBuffer, Chaos::FReal);
 DECLARE_MULTICAST_DELEGATE_OneParam(FSolverPostAdvance, Chaos::FReal);
+DECLARE_MULTICAST_DELEGATE(FSolverTeardown);
 
 namespace Chaos
 {
@@ -504,6 +505,10 @@ namespace Chaos
 		FDelegateHandle AddPostAdvanceCallback(FSolverPostAdvance::FDelegate InDelegate);
 		bool            RemovePostAdvanceCallback(FDelegateHandle InHandle);
 
+		/** Teardown happens as the solver is destroyed or streamed out */
+		FDelegateHandle AddTeardownCallback(FSolverTeardown::FDelegate InDelegate);
+		bool            RemoveTeardownCallback(FDelegateHandle InHandle);
+
 		/** Get the lock used for external data manipulation. A better API would be to use scoped locks so that getting a write lock is non-const */
 		//NOTE: this is a const operation so that you can acquire a read lock on a const solver. The assumption is that non-const write operations are already marked non-const
 		FPhysicsSceneGuard& GetExternalDataLock_External() const { return *ExternalDataLock_External; }
@@ -513,6 +518,7 @@ namespace Chaos
 		FSolverPreAdvance EventPreSolve;
 		FSolverPreBuffer EventPreBuffer;
 		FSolverPostAdvance EventPostSolve;
+		FSolverTeardown EventTeardown;
 
 		void TrackGTParticle_External(FGeometryParticle& Particle);
 		void ClearGTParticle_External(FGeometryParticle& Particle);
