Move evolution logic into SetDynamics so that PushData and RewindTo go through the same logic

diff --git a/Engine/Source/Programs/HeadlessChaos/Private/HeadlessChaosTestRewind.cpp b/Engine/Source/Programs/HeadlessChaos/Private/HeadlessChaosTestRewind.cpp
--- a/Engine/Source/Programs/HeadlessChaos/Private/HeadlessChaosTestRewind.cpp
+++ b/Engine/Source/Programs/HeadlessChaos/Private/HeadlessChaosTestRewind.cpp
@@ -1233,6 +1233,71 @@ namespace ChaosTest {
 			});
 	}
 
+	GTEST_TEST(AllTraits, RewindTest_RewindBeforeSleep)
+	{
+		// Rewind to before an object was put to sleep and see that the Active view is valid
+		TRewindHelper::TestDynamicSphere([](auto* Solver, FReal SimDt, int32 Optimization, auto Proxy, auto Sphere)
+			{
+				int32 ExpectedSleepFrame = INDEX_NONE;
+				const int32 ResimStartFrame = 0;
+				FRewindCallbackTestHelper* Helper = RegisterCallbackHelper(Solver, 11, ResimStartFrame);
+				Helper->ProcessInputsFunc = [&ExpectedSleepFrame, ResimStartFrame, Proxy, RewindData = Solver->GetRewindData(), Solver](const int32 PhysicsStep, bool bIsResimming)
+				{
+					if(PhysicsStep == 5)
+					{
+						Proxy->GetPhysicsThreadAPI()->SetV(FVec3(0, 0, 0));
+						Proxy->GetPhysicsThreadAPI()->SetObjectState(EObjectStateType::Sleeping);
+					}
+
+					//before sleep the active view contains the particle, after we put it to sleep the active view is empty
+					EXPECT_EQ(Solver->GetParticles().GetActiveParticlesView().Num(), PhysicsStep < 5 ? 1 : 0);
+				};
+
+				auto& Particle = Proxy->GetGameThreadAPI();
+				Particle.SetGravityEnabled(false);
+				Particle.SetV(FVec3(0, 0, 1));
+
+				for (int Step = 0; Step <= 32; ++Step)
+				{
+					TickSolverHelper(Solver);
+				}
+
+			});
+	}
+
+	GTEST_TEST(AllTraits, RewindTest_RewindBeforeAwake)
+	{
+		// Rewind to before an object was made awake and confirm that the active view is valid
+		// This test rewinds to step 3 because it ensures there's no PushData which may fixup the view. Need to make sure it's fixed regardless of push data
+		TRewindHelper::TestDynamicSphere([](auto* Solver, FReal SimDt, int32 Optimization, auto Proxy, auto Sphere)
+			{
+				int32 ExpectedSleepFrame = INDEX_NONE;
+				const int32 ResimStartFrame = 3;
+				FRewindCallbackTestHelper* Helper = RegisterCallbackHelper(Solver, 13, ResimStartFrame);
+				Helper->ProcessInputsFunc = [&ExpectedSleepFrame, ResimStartFrame, Proxy, RewindData = Solver->GetRewindData(), Solver](const int32 PhysicsStep, bool bIsResimming)
+				{
+					if (PhysicsStep == 12)
+					{
+						Proxy->GetPhysicsThreadAPI()->SetV(FVec3(0, 0, 1));
+						Proxy->GetPhysicsThreadAPI()->SetObjectState(EObjectStateType::Dynamic);
+					}
+
+					//before wake up the active view was empty, after we wake it up the view contains the particle
+					EXPECT_EQ(Solver->GetParticles().GetActiveParticlesView().Num(), PhysicsStep < 12 ? 0 : 1);
+				};
+
+				auto& Particle = Proxy->GetGameThreadAPI();
+				Particle.SetGravityEnabled(false);
+				Particle.SetObjectState(EObjectStateType::Sleeping);
+
+				for (int Step = 0; Step <= 32; ++Step)
+				{
+					TickSolverHelper(Solver);
+				}
+
+			});
+	}
+
 	GTEST_TEST(AllTraits, RewindTest_RecordForcesInSimCallback)
 	{
 		//Makes sure that we record the forces applied during sim callback
diff --git a/Engine/Source/Runtime/Experimental/Chaos/Private/Chaos/ParticleHandle.cpp b/Engine/Source/Runtime/Experimental/Chaos/Private/Chaos/ParticleHandle.cpp
--- a/Engine/Source/Runtime/Experimental/Chaos/Private/Chaos/ParticleHandle.cpp
+++ b/Engine/Source/Runtime/Experimental/Chaos/Private/Chaos/ParticleHandle.cpp
@@ -258,4 +258,35 @@ namespace Chaos
 		return Idx;
 	}
 
+	template <>
+	void TPBDRigidParticleHandleImp<FReal, 3, true>::SetDynamicMisc(const FParticleDynamicMisc& DynamicMisc, FPBDRigidsEvolutionBase& Evolution)
+	{
+		SetLinearEtherDrag(DynamicMisc.LinearEtherDrag());
+		SetAngularEtherDrag(DynamicMisc.AngularEtherDrag());
+		SetMaxLinearSpeedSq(DynamicMisc.MaxLinearSpeedSq());
+		SetMaxAngularSpeedSq(DynamicMisc.MaxAngularSpeedSq());
+		SetCollisionGroup(DynamicMisc.CollisionGroup());
+		SetGravityEnabled(DynamicMisc.GravityEnabled());
+		SetCCDEnabled(DynamicMisc.CCDEnabled());
+		SetDisabled(DynamicMisc.Disabled());
+		SetResimType(DynamicMisc.ResimType());
+		SetOneWayInteraction(DynamicMisc.OneWayInteraction());
+		AddCollisionConstraintFlag((Chaos::ECollisionConstraintFlags)DynamicMisc.CollisionConstraintFlag());
+
+		Evolution.SetParticleObjectState(this, DynamicMisc.ObjectState());
+		Evolution.SetParticleSleepType(this, DynamicMisc.SleepType());
+
+		if (Disabled() != DynamicMisc.Disabled())
+		{
+			if (DynamicMisc.Disabled())
+			{
+				Evolution.DisableParticle(this);
+			}
+			else
+			{
+				Evolution.EnableParticle(this, nullptr);
+			}
+		}
+	}
+
 }
diff --git a/Engine/Source/Runtime/Experimental/Chaos/Private/PhysicsProxy/SingleParticlePhysicsProxy.cpp b/Engine/Source/Runtime/Experimental/Chaos/Private/PhysicsProxy/SingleParticlePhysicsProxy.cpp
--- a/Engine/Source/Runtime/Experimental/Chaos/Private/PhysicsProxy/SingleParticlePhysicsProxy.cpp
+++ b/Engine/Source/Runtime/Experimental/Chaos/Private/PhysicsProxy/SingleParticlePhysicsProxy.cpp
@@ -122,30 +122,7 @@ void PushToPhysicsStateImp(const Chaos::FDirtyPropertiesManager& Manager, Chaos:
 
 			if(auto NewData = ParticleData.FindDynamicMisc(Manager,DataIdx))
 			{
-				Evolution.SetParticleObjectState(RigidHandle, NewData->ObjectState());
-				Evolution.SetParticleSleepType(RigidHandle, NewData->SleepType());
-
-				if (RigidHandle->Disabled() != NewData->Disabled())
-				{
-					if (NewData->Disabled())
-					{
-						Evolution.DisableParticle(Handle);
-					}
-					else
-					{
-						Evolution.EnableParticle(Handle, nullptr);
-					}
-				}
-
-				RigidHandle->SetDynamicMisc(*NewData);
-
-				if(NewData->ObjectState() != EObjectStateType::Dynamic)
-				{
-					//this is needed because changing object state on external thread means we want to snap position to where the particle was at that time (on the external thread)
-					//for that to work we need to ensure the snap results (which we just got) are passed properly into the results manager
-					Evolution.GetParticles().MarkTransientDirtyParticle(RigidHandle);
-				}
-				
+				RigidHandle->SetDynamicMisc(*NewData, Evolution);				
 			}
 		}
 
diff --git a/Engine/Source/Runtime/Experimental/Chaos/Private/RewindData.cpp b/Engine/Source/Runtime/Experimental/Chaos/Private/RewindData.cpp
--- a/Engine/Source/Runtime/Experimental/Chaos/Private/RewindData.cpp
+++ b/Engine/Source/Runtime/Experimental/Chaos/Private/RewindData.cpp
@@ -173,12 +173,6 @@ bool FRewindData::RewindToFrame(int32 Frame)
 	{
 		FGeometryParticleHandle* PTParticle = DirtyParticleInfo.GetPTParticle();
 
-		if(DirtyParticleInfo.InitializedOnStep > Frame)
-		{
-			//hasn't initialized yet, so disable
-			Solver->GetEvolution()->DisableParticle(DirtyParticleInfo.GetPTParticle());
-		}
-
 		//rewind is about to start, all particles should be in sync at this point
 		ensure(PTParticle->SyncState() == ESyncState::InSync);
 
@@ -193,7 +187,7 @@ bool FRewindData::RewindToFrame(int32 Frame)
 			RewindNoSave(PTParticle->CastToKinematicParticle(), History.Velocities, [](auto Particle, const auto& Data) {Particle->SetVelocities(Data); });
 			RewindNoSave(PTParticle->CastToKinematicParticle(), History.KinematicTarget, [](auto Particle, const auto& Data) {Particle->SetKinematicTarget(Data); });
 			RewindNoSave(PTParticle->CastToRigidParticle(), History.Dynamics, [](auto Particle, const auto& Data) {Particle->SetDynamics(Data); });
-			RewindNoSave(PTParticle->CastToRigidParticle(), History.DynamicsMisc, [](auto Particle, const auto& Data) {Particle->SetDynamicMisc(Data); });
+			RewindNoSave(PTParticle->CastToRigidParticle(), History.DynamicsMisc, [Evolution = Solver->GetEvolution()](auto Particle, const auto& Data) {Particle->SetDynamicMisc(Data, *Evolution); });
 			RewindNoSave(PTParticle->CastToRigidParticle(), History.MassProps, [](auto Particle, const auto& Data) {Particle->SetMassProps(Data); });
 		}
 		else
@@ -205,7 +199,7 @@ bool FRewindData::RewindToFrame(int32 Frame)
 			bAnyChange |= RewindAndSave(PTParticle->CastToKinematicParticle(), History.Velocities, [](auto Particle, const auto& Data) {Particle->SetVelocities(Data); });
 			bAnyChange |= RewindAndSave(PTParticle->CastToKinematicParticle(), History.KinematicTarget, [](auto Particle, const auto& Data) {Particle->SetKinematicTarget(Data); });
 			bAnyChange |= RewindAndSave(PTParticle->CastToRigidParticle(), History.Dynamics, [](auto Particle, const auto& Data) {Particle->SetDynamics(Data); });
-			bAnyChange |= RewindAndSave(PTParticle->CastToRigidParticle(), History.DynamicsMisc, [](auto Particle, const auto& Data) {Particle->SetDynamicMisc(Data); });
+			bAnyChange |= RewindAndSave(PTParticle->CastToRigidParticle(), History.DynamicsMisc, [Evolution = Solver->GetEvolution()](auto Particle, const auto& Data) {Particle->SetDynamicMisc(Data, *Evolution); });
 			bAnyChange |= RewindAndSave(PTParticle->CastToRigidParticle(), History.MassProps, [](auto Particle, const auto& Data) {Particle->SetMassProps(Data); });
 
 			if (bAnyChange)
@@ -217,6 +211,14 @@ bool FRewindData::RewindToFrame(int32 Frame)
 
 				DirtyParticleInfo.DirtyDynamics = INDEX_NONE;	//make sure to undo this as we want to record it again during resim
 		}
+
+		if (DirtyParticleInfo.InitializedOnStep > Frame)
+		{
+			//hasn't initialized yet, so disable
+			//must do this after rewind because SetDynamicsMisc will re-enable
+			//(the disable is a temp way to ignore objects not spawned yet, they weren't really disabled which is why it gets re-enabled)
+			Solver->GetEvolution()->DisableParticle(DirtyParticleInfo.GetPTParticle());
+		}
 	}
 
 	}
diff --git a/Engine/Source/Runtime/Experimental/Chaos/Public/Chaos/ParticleHandle.h b/Engine/Source/Runtime/Experimental/Chaos/Public/Chaos/ParticleHandle.h
--- a/Engine/Source/Runtime/Experimental/Chaos/Public/Chaos/ParticleHandle.h
+++ b/Engine/Source/Runtime/Experimental/Chaos/Public/Chaos/ParticleHandle.h
@@ -19,6 +19,7 @@ class IPhysicsProxyBase;
 namespace Chaos
 {
 	class FConstraintHandle;
+	class FPBDRigidsEvolutionBase;
 
 struct FGeometryParticleParameters
 {
@@ -814,21 +815,7 @@ public:
 		SetInvM(Props.InvM());
 	}
 
-	void SetDynamicMisc(const FParticleDynamicMisc& DynamicMisc)
-	{
-		SetLinearEtherDrag(DynamicMisc.LinearEtherDrag());
-		SetAngularEtherDrag(DynamicMisc.AngularEtherDrag());
-		SetMaxLinearSpeedSq(DynamicMisc.MaxLinearSpeedSq());
-		SetMaxAngularSpeedSq(DynamicMisc.MaxAngularSpeedSq());
-		SetCollisionGroup(DynamicMisc.CollisionGroup());
-		SetGravityEnabled(DynamicMisc.GravityEnabled());
-		SetCCDEnabled(DynamicMisc.CCDEnabled());
-		SetDisabled(DynamicMisc.Disabled()); 
-		SetResimType(DynamicMisc.ResimType());
-		SetOneWayInteraction(DynamicMisc.OneWayInteraction());
-		AddCollisionConstraintFlag( (Chaos::ECollisionConstraintFlags)DynamicMisc.CollisionConstraintFlag() );
-		SetSleepType(DynamicMisc.SleepType());
-	}
+	void SetDynamicMisc(const FParticleDynamicMisc& DynamicMisc, FPBDRigidsEvolutionBase& Evolution);
 
 	void ResetSmoothedVelocities()
 	{
