Additional data needed to generate a resim subview for a more efficient resim. Removed debug only code for deleting particles.
Removed unnecessary disable during particle destruction

diff --git a/Engine/Source/Runtime/Experimental/Chaos/Private/PBDRigidsSolver.cpp b/Engine/Source/Runtime/Experimental/Chaos/Private/PBDRigidsSolver.cpp
--- a/Engine/Source/Runtime/Experimental/Chaos/Private/PBDRigidsSolver.cpp
+++ b/Engine/Source/Runtime/Experimental/Chaos/Private/PBDRigidsSolver.cpp
@@ -208,10 +208,6 @@ FAutoConsoleVariableRef CVarChaosSolverJointNumIterationsAtMaxSolverStiffness(TE
 int32 ChaosVisualDebuggerEnable = 1;
 FAutoConsoleVariableRef CVarChaosVisualDebuggerEnable(TEXT("p.Chaos.VisualDebuggerEnable"), ChaosVisualDebuggerEnable, TEXT("Enable/Disable pushing/saving data to the visual debugger"));
 
-bool bDisableParticleDeletion = false;
-FAutoConsoleVariableRef CVarDisableParticleDeletion(TEXT("p.Chaos.Solver.DisableParticleDeletion"), bDisableParticleDeletion, TEXT("Only ever use when debugging the particle clean up, this introduces a massive memory leak."));
-
-
 namespace Chaos
 {
 
@@ -232,11 +228,6 @@ namespace Chaos
 
 		void DoWork()
 		{
-			// WARNING
-			// Only ever use DisableParticleDeletion when debugging the particle clean up. 
-			// This introduces a massive memory leak.
-			MSolver->GetParticles().SetDisableParticleDeletion(bDisableParticleDeletion);
-
 			LLM_SCOPE(ELLMTag::Chaos);
 			UE_LOG(LogPBDRigidsSolver, Verbose, TEXT("AdvanceOneTimeStepTask::DoWork()"));
 			MSolver->StartingSceneSimulation();
diff --git a/Engine/Source/Runtime/Experimental/Chaos/Public/Chaos/PBDRigidsSOAs.h b/Engine/Source/Runtime/Experimental/Chaos/Public/Chaos/PBDRigidsSOAs.h
--- a/Engine/Source/Runtime/Experimental/Chaos/Public/Chaos/PBDRigidsSOAs.h
+++ b/Engine/Source/Runtime/Experimental/Chaos/Public/Chaos/PBDRigidsSOAs.h
@@ -181,7 +181,6 @@ class FPBDRigidsSOAs
 {
 public:
 	FPBDRigidsSOAs()
-		: bDisableParticleDeletion(false)
 	{
 #if CHAOS_DETERMINISTIC
 		BiggestParticleID = 0;
@@ -304,12 +303,6 @@ public:
 		}
 	}
 
-	// WARNING
-	// Only ever use DisableParticleDeletion when debugging the particle clean up. 
-	// This introduces a massive memory leak.
-	bool bDisableParticleDeletion;
-	void SetDisableParticleDeletion(bool bIn) { bDisableParticleDeletion = bIn; }
-
 	void DestroyParticle(FGeometryParticleHandle* Particle)
 	{
 		auto PBDRigid = Particle->CastToRigidParticle();
@@ -328,10 +321,6 @@ public:
 					DynamicClusteredMapArray.Remove(PBDRigidClustered);
 				}
 			}
-			else
-			{
-				Particle->MoveToSOA(*DynamicDisabledParticles);
-			}
 
 			// Check for sleep events referencing this particle
 			// TODO think about this case more
@@ -346,13 +335,7 @@ public:
 			GetDynamicParticles().GetSleepDataLock().WriteUnlock();
 		}
 
-		// WARNING
-		// Only ever use DisableParticleDeletion when debugging the particle clean up. 
-		// This introduces a massive memory leak.
-		if (!bDisableParticleDeletion)
-		{
-			ParticleHandles.DestroyHandleSwap(Particle);
-		}
+		ParticleHandles.DestroyHandleSwap(Particle);
 
 		UpdateViews();
 	}
@@ -587,15 +570,30 @@ public:
 			{
 			case EObjectStateType::Kinematic:
 				Particle->MoveToSOA(*DynamicKinematicParticles);
+				if(bResimulating)
+				{
+					ResimDynamicParticles.Remove(Particle);
+					ResimDynamicKinematicParticles.Insert(Particle);
+				}
 				break;
 
 			case EObjectStateType::Dynamic:
 				Particle->MoveToSOA(*DynamicParticles);
+				if (bResimulating)
+				{
+					ResimDynamicKinematicParticles.Remove(Particle);
+					ResimDynamicParticles.Insert(Particle);
+				}
 				break;
 
 			default:
 				// TODO: Special SOAs for sleeping and static particles?
 				Particle->MoveToSOA(*DynamicParticles);
+				if (bResimulating)
+				{
+					ResimDynamicKinematicParticles.Remove(Particle);
+					ResimDynamicParticles.Insert(Particle);
+				}
 				break;
 			}
 		}
@@ -873,6 +871,11 @@ private:
 	void AddToActiveArray(const TArray<FPBDRigidParticleHandle*>& Particles)
 	{
 		ActiveParticlesMapArray.Insert(Particles);
+
+		if(bResimulating)
+		{
+			ResimActiveParticlesMapArray.Insert(Particles);
+		}
 		
 		//dirty contains Active so make sure no duplicates
 		for(FPBDRigidParticleHandle* Particle : Particles)
@@ -885,6 +888,11 @@ private:
 	{
 		ActiveParticlesMapArray.Insert(Particle);
 
+		if (bResimulating)
+		{
+			ResimActiveParticlesMapArray.Insert(Particle);
+		}
+
 		//dirty contains Active so make sure no duplicates
 		TransientDirtyMapArray.Remove(Particle);
 	}
@@ -893,6 +901,11 @@ private:
 	{
 		ActiveParticlesMapArray.Remove(Particle);
 
+		if (bResimulating)
+		{
+			ResimActiveParticlesMapArray.Remove(Particle);
+		}
+
 		if(bStillDirty)
 		{
 			//no longer active, but still dirty
@@ -911,7 +924,15 @@ private:
 	{
 		//build various views. Group SOA types together for better branch prediction
 		{
-			TArray<TSOAView<FGeometryParticles>> TmpArray = 
+			auto TmpArray = bResimulating
+			? TArray<TSOAView<FGeometryParticles>>
+			{
+				{&ResimStaticParticles},
+				{&ResimKinematicParticles},
+				{&ResimDynamicParticles.GetArray() },
+				{&ResimDynamicKinematicParticles.GetArray() }
+			}
+			: TArray<TSOAView<FGeometryParticles>>
 			{ 
 				StaticParticles.Get(), 
 				KinematicParticles.Get(), 
@@ -928,17 +949,27 @@ private:
 			NonDisabledView = MakeParticleView(MoveTemp(TmpArray));
 		}
 		{
-			TArray<TSOAView<FPBDRigidParticles>> TmpArray = 
+			auto TmpArray = bResimulating
+			? TArray<TSOAView<FPBDRigidParticles>>
 			{ 
-				DynamicParticles.Get(), 
-				{&DynamicClusteredMapArray.GetArray()},
-				{&SleepingGeometryCollectionArray},
-				{&DynamicGeometryCollectionArray}
+				{&ResimDynamicParticles.GetArray() }
+			}
+			: TArray<TSOAView<FPBDRigidParticles>>
+			{
+				DynamicParticles.Get(),
+				{&DynamicClusteredMapArray.GetArray() },
+				{&SleepingGeometryCollectionArray },
+				{&DynamicGeometryCollectionArray }
 			};
 			NonDisabledDynamicView = MakeParticleView(MoveTemp(TmpArray));
 		}
 		{
-			TArray<TSOAView<FPBDRigidParticles>> TmpArray = 
+			auto TmpArray = bResimulating
+			? TArray<TSOAView<FPBDRigidParticles>>
+			{
+				{&ResimActiveParticlesMapArray.GetArray()}
+			}
+			: TArray<TSOAView<FPBDRigidParticles>>
 			{ 
 				{&ActiveParticlesMapArray.GetArray()},
 			//	{&DynamicClusteredMapArray.GetArray()},  Cluster particles appear in the ActiveParticlesArray
@@ -949,7 +980,13 @@ private:
 			ActiveParticlesView = MakeParticleView(MoveTemp(TmpArray));
 		}
 		{
-			TArray<TSOAView<FPBDRigidParticles>> TmpArray =
+			auto TmpArray = bResimulating
+			? TArray<TSOAView<FPBDRigidParticles>>
+			{
+				{&ResimActiveParticlesMapArray.GetArray()},
+				{&TransientDirtyMapArray.GetArray() }
+			}
+			: TArray<TSOAView<FPBDRigidParticles>>
 			{
 				{&ActiveParticlesMapArray.GetArray()},
 				//	{&DynamicClusteredMapArray.GetArray()},  Cluster particles appear in the ActiveParticlesArray
@@ -976,7 +1013,13 @@ private:
 			AllParticlesView = MakeParticleView(MoveTemp(TmpArray));
 		}
 		{
-			TArray<TSOAView<FKinematicGeometryParticles>> TmpArray = 
+			auto TmpArray = bResimulating
+			? TArray<TSOAView<FKinematicGeometryParticles>>
+			{
+				{&ResimKinematicParticles},
+				{&ResimDynamicKinematicParticles.GetArray()}
+			}
+			: TArray<TSOAView<FKinematicGeometryParticles>>
 			{ 
 				KinematicParticles.Get(),
 				DynamicKinematicParticles.Get(),
@@ -986,7 +1029,12 @@ private:
 			ActiveKinematicParticlesView = MakeParticleView(MoveTemp(TmpArray));
 		}
 		{
-			TArray<TSOAView<FGeometryParticles>> TmpArray =
+			auto TmpArray = bResimulating
+			? TArray<TSOAView<FGeometryParticles>>
+			{
+				{&ResimStaticParticles}
+			}
+			: TArray<TSOAView<FGeometryParticles>>
 			{
 				StaticParticles.Get(),
 				{&StaticClusteredMapArray.GetArray()}
@@ -994,7 +1042,9 @@ private:
 			ActiveStaticParticlesView = MakeParticleView(MoveTemp(TmpArray));
 		}
 		{
-			TArray<TSOAView<TPBDGeometryCollectionParticles<FReal, 3>>> TmpArray = 
+			auto TmpArray = bResimulating
+			? TArray<TSOAView<TPBDGeometryCollectionParticles<FReal, 3>>> {}	//no geometry collection during resim
+			: TArray<TSOAView<TPBDGeometryCollectionParticles<FReal, 3>>>
 			{ 
 				{&StaticGeometryCollectionArray},
 				{&KinematicGeometryCollectionArray},
@@ -1004,7 +1054,9 @@ private:
 		}
 
 		{
-			TArray<TSOAView<FPBDRigidClusteredParticles>> TmpArray =
+			auto TmpArray = bResimulating
+			? TArray<TSOAView<FPBDRigidClusteredParticles>> {}	//no clusters during resim
+			: TArray<TSOAView<FPBDRigidClusteredParticles>>
 			{
 				{&StaticClusteredMapArray.GetArray()},
 				{&KinematicClusteredMapArray.GetArray()},
@@ -1080,6 +1132,14 @@ private:
 	TParticleMapArray<FPBDRigidParticleHandle> ActiveParticlesMapArray;
 	TParticleMapArray<FPBDRigidParticleHandle> TransientDirtyMapArray;
 
+	// Structures for maintaining a subset view during a resim (TParticleMapArray used when we need to dynamically add/remove during resim)
+	TParticleMapArray<FPBDRigidParticleHandle> ResimActiveParticlesMapArray;
+	TParticleMapArray<FPBDRigidParticleHandle> ResimDynamicParticles;
+	TParticleMapArray<FPBDRigidParticleHandle> ResimDynamicKinematicParticles;
+	TArray<FGeometryParticleHandle*> ResimStaticParticles;
+	TArray<FKinematicGeometryParticleHandle*> ResimKinematicParticles;
+	bool bResimulating = false;
+
 	// NonDisabled clustered particle arrays
 	TParticleMapArray<FPBDRigidClusteredParticleHandle> StaticClusteredMapArray;
 	TParticleMapArray<FPBDRigidClusteredParticleHandle> KinematicClusteredMapArray;
