Fix ensure IsInPhysicsThreadContext() in EventDefaults.cpp in test/shipping builds.

diff --git a/Engine/Source/Runtime/Experimental/Chaos/Private/Chaos/Framework/PhysicsSolverBase.cpp b/Engine/Source/Runtime/Experimental/Chaos/Private/Chaos/Framework/PhysicsSolverBase.cpp
--- a/Engine/Source/Runtime/Experimental/Chaos/Private/Chaos/Framework/PhysicsSolverBase.cpp
+++ b/Engine/Source/Runtime/Experimental/Chaos/Private/Chaos/Framework/PhysicsSolverBase.cpp
@@ -227,11 +227,6 @@ namespace Chaos
 		InSolver.SetThreadingMode_External(EThreadingModeTemp::SingleThread);
 		InSolver.MarkShuttingDown();
 		{
-#if PHYSICS_THREAD_CONTEXT
-			// Must use physics thread context, as we are about to execute physics thread API from the game thread in single threaded mode.
-			FPhysicsThreadContextScope PhysicsThreadContext(/*InParentIsPhysicsThreadContext=*/true);
-#endif
-
 			InSolver.AdvanceAndDispatch_External(0);	//flush any pending commands are executed (for example unregister object)
 		}
 
diff --git a/Engine/Source/Runtime/Experimental/Chaos/Private/EventDefaults.cpp b/Engine/Source/Runtime/Experimental/Chaos/Private/EventDefaults.cpp
--- a/Engine/Source/Runtime/Experimental/Chaos/Private/EventDefaults.cpp
+++ b/Engine/Source/Runtime/Experimental/Chaos/Private/EventDefaults.cpp
@@ -31,7 +31,10 @@ namespace Chaos
 		(const Chaos::FPBDRigidsSolver* Solver, FCollisionEventData& CollisionEventData)
 		{
 			check(Solver);
+
+#if PHYSICS_THREAD_CONTEXT
 			ensure(IsInPhysicsThreadContext());
+#endif
 
 			SCOPE_CYCLE_COUNTER(STAT_GatherCollisionEvent);
 
