Minor geometry collection-related code cleanup

#rnx

diff --git a/Engine/Source/Runtime/Experimental/Chaos/Private/GeometryCollection/GeometryCollectionClusteringUtility.cpp b/Engine/Source/Runtime/Experimental/Chaos/Private/GeometryCollection/GeometryCollectionClusteringUtility.cpp
--- a/Engine/Source/Runtime/Experimental/Chaos/Private/GeometryCollection/GeometryCollectionClusteringUtility.cpp
+++ b/Engine/Source/Runtime/Experimental/Chaos/Private/GeometryCollection/GeometryCollectionClusteringUtility.cpp
@@ -419,6 +419,7 @@ bool FGeometryCollectionClusteringUtility::ContainsMultipleRootBones(FGeometryCo
 void FGeometryCollectionClusteringUtility::GetRootBones(const FGeometryCollection* GeometryCollection, TArray<int32>& RootBonesOut)
 {
 	check(GeometryCollection);
+	checkSlow(RootBonesOut.Num() == 0);
 	const TManagedArray<int32>& Parents = GeometryCollection->Parent;
 
 	// never assume the root bone is always index 0 in the particle group
@@ -426,7 +427,7 @@ void FGeometryCollectionClusteringUtility::GetRootBones(const FGeometryCollectio
 	{
 		if (Parents[i] == FGeometryCollection::Invalid)
 		{
-			RootBonesOut.AddUnique(i);
+			RootBonesOut.Add(i);
 		}
 	}
 }
diff --git a/Engine/Source/Runtime/Experimental/GeometryCollectionEngine/Private/GeometryCollection/GeometryCollectionComponent.cpp b/Engine/Source/Runtime/Experimental/GeometryCollectionEngine/Private/GeometryCollection/GeometryCollectionComponent.cpp
--- a/Engine/Source/Runtime/Experimental/GeometryCollectionEngine/Private/GeometryCollection/GeometryCollectionComponent.cpp
+++ b/Engine/Source/Runtime/Experimental/GeometryCollectionEngine/Private/GeometryCollection/GeometryCollectionComponent.cpp
@@ -980,10 +980,7 @@ void UGeometryCollectionComponent::PostEditChangeChainProperty(FPropertyChangedC
 
 	if (PropertyChangedEvent.Property && PropertyChangedEvent.Property->GetFName() == GET_MEMBER_NAME_CHECKED(UGeometryCollectionComponent, bShowBoneColors))
 	{
-		FScopedColorEdit EditBoneColor = EditBoneSelection();
-		EditBoneColor.SetShowBoneColors(bShowBoneColors);
-		MarkRenderStateDirty();
-		MarkRenderDynamicDataDirty();
+		FScopedColorEdit EditBoneColor(this, true /*bForceUpdate*/); // the property has already changed; this will trigger the color update + render state updates
 	}
 }
 #endif
