GeometryCollectionComponent
-- Fixed bNotifyBreaks and bNotifyCollisions

diff --git a/Engine/Source/Runtime/Experimental/GeometryCollectionEngine/Private/GeometryCollection/GeometryCollectionComponent.cpp b/Engine/Source/Runtime/Experimental/GeometryCollectionEngine/Private/GeometryCollection/GeometryCollectionComponent.cpp
--- a/Engine/Source/Runtime/Experimental/GeometryCollectionEngine/Private/GeometryCollection/GeometryCollectionComponent.cpp
+++ b/Engine/Source/Runtime/Experimental/GeometryCollectionEngine/Private/GeometryCollection/GeometryCollectionComponent.cpp
@@ -1004,20 +1004,29 @@ void UGeometryCollectionComponent::RegisterForEvents()
 {
 	if (BodyInstance.bNotifyRigidBodyCollision || bNotifyBreaks || bNotifyCollisions)
 	{
-		if (bNotifyCollisions || BodyInstance.bNotifyRigidBodyCollision)
+		Chaos::FPhysicsSolver* Solver = GetWorld()->GetPhysicsScene()->GetSolver();
+		if (Solver)
 		{
-			EventDispatcher->RegisterForCollisionEvents(this, this);
-#if INCLUDE_CHAOS
-			GetWorld()->GetPhysicsScene()->GetSolver()->SetGenerateCollisionData(true);
-#endif
-		}
+			if (bNotifyCollisions || BodyInstance.bNotifyRigidBodyCollision)
+			{
+				EventDispatcher->RegisterForCollisionEvents(this, this);
 
-		if (bNotifyBreaks)
-		{
-			EventDispatcher->RegisterForBreakEvents(this, &DispatchGeometryCollectionBreakEvent);
-#if INCLUDE_CHAOS
-			GetWorld()->GetPhysicsScene()->GetSolver()->SetGenerateBreakingData(true);
-#endif
+				Solver->EnqueueCommandImmediate([Solver]()
+					{
+						Solver->SetGenerateCollisionData(true);
+					});
+			}
+
+			if (bNotifyBreaks)
+			{
+				EventDispatcher->RegisterForBreakEvents(this, &DispatchGeometryCollectionBreakEvent);
+
+				Solver->EnqueueCommandImmediate([Solver]()
+					{
+						Solver->SetGenerateBreakingData(true);
+					});
+
+			}
 		}
 	}
 }
