fix undo/redo for bone selection on fracture

#rnx

diff --git a/Engine/Source/Runtime/Experimental/GeometryCollectionEngine/Private/GeometryCollection/GeometryCollectionComponent.cpp b/Engine/Source/Runtime/Experimental/GeometryCollectionEngine/Private/GeometryCollection/GeometryCollectionComponent.cpp
--- a/Engine/Source/Runtime/Experimental/GeometryCollectionEngine/Private/GeometryCollection/GeometryCollectionComponent.cpp
+++ b/Engine/Source/Runtime/Experimental/GeometryCollectionEngine/Private/GeometryCollection/GeometryCollectionComponent.cpp
@@ -1619,11 +1619,6 @@ void UGeometryCollectionComponent::OnRegister()
 	//UE_LOG(UGCC_LOG, Log, TEXT("GeometryCollectionComponent[%p]::OnRegister()[%p]"), this,RestCollection );
 	ResetDynamicCollection();
 
-#if WITH_EDITOR
-	FScopedColorEdit ColorEdit(this);
-	ColorEdit.ResetBoneSelection();
-	ColorEdit.ResetHighlightedBones();
-#endif
 
 #endif // WITH_CHAOS
 
