Chaos - fix joint settings validation when disabling slerp drive

diff --git a/Engine/Source/Runtime/Experimental/Chaos/Private/Chaos/PBDJointConstraints.cpp b/Engine/Source/Runtime/Experimental/Chaos/Private/Chaos/PBDJointConstraints.cpp
--- a/Engine/Source/Runtime/Experimental/Chaos/Private/Chaos/PBDJointConstraints.cpp
+++ b/Engine/Source/Runtime/Experimental/Chaos/Private/Chaos/PBDJointConstraints.cpp
@@ -211,8 +211,8 @@ namespace Chaos
 		if (bAngularSLerpPositionDriveEnabled || bAngularSLerpVelocityDriveEnabled)
 		{
 			if ((AngularMotionTypes[(int32)EJointAngularConstraintIndex::Twist] == EJointMotionType::Locked)
-				|| (AngularMotionTypes[(int32)EJointAngularConstraintIndex::Twist] == EJointMotionType::Locked)
-				|| (AngularMotionTypes[(int32)EJointAngularConstraintIndex::Twist] == EJointMotionType::Locked))
+				|| (AngularMotionTypes[(int32)EJointAngularConstraintIndex::Swing1] == EJointMotionType::Locked)
+				|| (AngularMotionTypes[(int32)EJointAngularConstraintIndex::Swing2] == EJointMotionType::Locked))
 			{
 				bAngularSLerpPositionDriveEnabled = false;
 				bAngularSLerpVelocityDriveEnabled = false;
