Chaos : Constraints
- User driven constraint ordering.

diff --git a/Engine/Source/Runtime/Experimental/Chaos/Private/Chaos/PBDRigidsEvolutionGBF.cpp b/Engine/Source/Runtime/Experimental/Chaos/Private/Chaos/PBDRigidsEvolutionGBF.cpp
--- a/Engine/Source/Runtime/Experimental/Chaos/Private/Chaos/PBDRigidsEvolutionGBF.cpp
+++ b/Engine/Source/Runtime/Experimental/Chaos/Private/Chaos/PBDRigidsEvolutionGBF.cpp
@@ -60,6 +60,9 @@ FAutoConsoleVariableRef CVarSmoothedPositionLerpRate(TEXT("p.Chaos.SmoothedPosit
 int DisableParticleUpdateVelocityParallelFor = 0;
 FAutoConsoleVariableRef CVarDisableParticleUpdateVelocityParallelFor(TEXT("p.DisableParticleUpdateVelocityParallelFor"), DisableParticleUpdateVelocityParallelFor, TEXT("Disable Particle Update Velocity ParallelFor and run the update on a single thread"));
 
+int32 ChaosSolverCollisionPriority = 0;
+FAutoConsoleVariableRef CVarChaosSolverCollisionPriority(TEXT("p.Chaos.Solver.Collision.Priority"), ChaosSolverCollisionPriority, TEXT("Set constraint priority. Larger values are evaluated later [def:0]"));
+
 DECLARE_CYCLE_STAT(TEXT("FPBDRigidsEvolutionGBF::AdvanceOneTimeStep"), STAT_Evolution_AdvanceOneTimeStep, STATGROUP_Chaos);
 DECLARE_CYCLE_STAT(TEXT("FPBDRigidsEvolutionGBF::UnclusterUnions"), STAT_Evolution_UnclusterUnions, STATGROUP_Chaos);
 DECLARE_CYCLE_STAT(TEXT("FPBDRigidsEvolutionGBF::Integrate"), STAT_Evolution_Integrate, STATGROUP_Chaos);
@@ -436,7 +439,7 @@ FPBDRigidsEvolutionGBF::FPBDRigidsEvolutionGBF(FPBDRigidsSOAs& InParticles,THand
 	: Base(InParticles, SolverPhysicsMaterials, DefaultNumIterations, DefaultNumPushOutIterations, InIsSingleThreaded)
 	, Clustering(*this, Particles.GetClusteredParticles())
 	, CollisionConstraints(InParticles, Collided, PhysicsMaterials, PerParticlePhysicsMaterials, DefaultNumCollisionPairIterations, DefaultNumCollisionPushOutPairIterations, DefaultRestitutionThreshold)
-	, CollisionRule(CollisionConstraints)
+	, CollisionRule(CollisionConstraints, ChaosSolverCollisionPriority)
 	, BroadPhase(InParticles, DefaultCollisionCullDistance, BoundsThicknessVelocityMultiplier, DefaultCollisionCullDistance)
 	, NarrowPhase()
 	, CollisionDetector(BroadPhase, NarrowPhase, CollisionConstraints)
diff --git a/Engine/Source/Runtime/Experimental/Chaos/Private/PBDRigidsSolver.cpp b/Engine/Source/Runtime/Experimental/Chaos/Private/PBDRigidsSolver.cpp
--- a/Engine/Source/Runtime/Experimental/Chaos/Private/PBDRigidsSolver.cpp
+++ b/Engine/Source/Runtime/Experimental/Chaos/Private/PBDRigidsSolver.cpp
@@ -195,6 +195,11 @@ FAutoConsoleVariableRef CVarChaosSolverCollisionDeferNarrowPhase(TEXT("p.Chaos.S
 int32 ChaosSolverCollisionUseManifolds = 1;
 FAutoConsoleVariableRef CVarChaosSolverCollisionUseManifolds(TEXT("p.Chaos.Solver.Collision.UseManifolds"), ChaosSolverCollisionUseManifolds, TEXT("Enable/Disable use of manifolds in collision."));
 
+int32 ChaosSolverJointPriority = 0;
+FAutoConsoleVariableRef CVarChaosSolverJointPriority(TEXT("p.Chaos.Solver.Joint.Priority"), ChaosSolverJointPriority, TEXT("Set constraint priority. Larger values are evaluated later [def:0]"));
+
+int32 ChaosSolverSuspensionPriority = 0;
+FAutoConsoleVariableRef CVarChaosSolverSuspensionPriority(TEXT("p.Chaos.Solver.Suspension.Priority"), ChaosSolverSuspensionPriority, TEXT("Set constraint priority. Larger values are evaluated later [def:0]"));
 
 // Joint cvars
 float ChaosSolverJointMinSolverStiffness = 1.0f;
@@ -401,8 +406,8 @@ namespace Chaos
 		, MSolverEventFilters(new FSolverEventFilters())
 		, MDirtyParticlesBuffer(new FDirtyParticlesBuffer(BufferingModeIn, BufferingModeIn == Chaos::EMultiBufferMode::Single))
 		, MCurrentLock(new FCriticalSection())
-		, JointConstraintRule(JointConstraints)
-		, SuspensionConstraintRule(SuspensionConstraints)
+		, JointConstraintRule(JointConstraints, ChaosSolverJointPriority)
+		, SuspensionConstraintRule(SuspensionConstraints, ChaosSolverSuspensionPriority)
 		, PerSolverField(nullptr)
 	{
 		UE_LOG(LogPBDRigidsSolver, Verbose, TEXT("PBDRigidsSolver::PBDRigidsSolver()"));
