Merge from Release-Engine-Staging of rewind specific changes
17071475 - 17202908

diff --git a/Engine/Source/Runtime/Experimental/Chaos/Private/Chaos/ChaosMarshallingManager.cpp b/Engine/Source/Runtime/Experimental/Chaos/Private/Chaos/ChaosMarshallingManager.cpp
--- a/Engine/Source/Runtime/Experimental/Chaos/Private/Chaos/ChaosMarshallingManager.cpp
+++ b/Engine/Source/Runtime/Experimental/Chaos/Private/Chaos/ChaosMarshallingManager.cpp
@@ -184,6 +184,11 @@ void FChaosMarshallingManager::FreeDataToHistory_Internal(FPushPhysicsData* Push
 
 void FChaosMarshallingManager::SetHistoryLength_Internal(int32 InHistoryLength)
 {
+	if (!ensure(InHistoryLength >= 0))
+	{
+		InHistoryLength = 0;
+	}
+	
 	HistoryLength = InHistoryLength;
 	//make sure late entries are pruned
 	if(HistoryQueue_Internal.Num() > HistoryLength)
diff --git a/Engine/Source/Runtime/Experimental/Chaos/Private/Chaos/ParticleHandle.cpp b/Engine/Source/Runtime/Experimental/Chaos/Private/Chaos/ParticleHandle.cpp
--- a/Engine/Source/Runtime/Experimental/Chaos/Private/Chaos/ParticleHandle.cpp
+++ b/Engine/Source/Runtime/Experimental/Chaos/Private/Chaos/ParticleHandle.cpp
@@ -269,7 +269,6 @@ namespace Chaos
 		SetGravityEnabled(DynamicMisc.GravityEnabled());
 		SetCCDEnabled(DynamicMisc.CCDEnabled());
 		SetDisabled(DynamicMisc.Disabled());
-		SetResimType(DynamicMisc.ResimType());
 		SetOneWayInteraction(DynamicMisc.OneWayInteraction());
 		AddCollisionConstraintFlag((Chaos::ECollisionConstraintFlags)DynamicMisc.CollisionConstraintFlag());
 
diff --git a/Engine/Source/Runtime/Experimental/Chaos/Private/RewindData.cpp b/Engine/Source/Runtime/Experimental/Chaos/Private/RewindData.cpp
--- a/Engine/Source/Runtime/Experimental/Chaos/Private/RewindData.cpp
+++ b/Engine/Source/Runtime/Experimental/Chaos/Private/RewindData.cpp
@@ -99,9 +99,9 @@ bool FGeometryParticleStateBase::IsInSync(const FGeometryParticleHandle& Handle,
 		if(!bSkipDynamics)
 		{
 			if (!Dynamics.IsInSync(*Rigid, FrameAndPhase, Pool))
-		{
-			return false;
-		}
+			{
+				return false;
+			}
 		}
 
 		if(!DynamicsMisc.IsInSync(*Rigid, FrameAndPhase, Pool))
@@ -150,7 +150,7 @@ bool FRewindData::RewindToFrame(int32 Frame)
 	auto RewindNoSave = [RewindFrameAndPhase, this](auto Particle, auto& Property, const auto& RewindFunc)
 	{
 		if (auto Val = Property.Read(RewindFrameAndPhase, PropertiesPool))
-	{
+		{
 			RewindFunc(Particle, *Val);
 		}
 	};
@@ -210,7 +210,10 @@ bool FRewindData::RewindToFrame(int32 Frame)
 				Solver->GetEvolution()->GetParticles().MarkTransientDirtyParticle(DirtyParticleInfo.GetPTParticle());
 
 				DirtyParticleInfo.DirtyDynamics = INDEX_NONE;	//make sure to undo this as we want to record it again during resim
-		}
+
+				//for now just mark anything that changed as enabled during resim. TODO: use bubble
+				DirtyParticleInfo.GetPTParticle()->SetEnabledDuringResim(true);
+			}
 
 		if (DirtyParticleInfo.InitializedOnStep > Frame)
 		{
@@ -219,7 +222,7 @@ bool FRewindData::RewindToFrame(int32 Frame)
 			//(the disable is a temp way to ignore objects not spawned yet, they weren't really disabled which is why it gets re-enabled)
 			Solver->GetEvolution()->DisableParticle(DirtyParticleInfo.GetPTParticle());
 		}
-	}
+		}
 
 	}
 
@@ -275,13 +278,14 @@ void FRewindData::FinishFrame()
 
 			auto& Handle = *Info.GetPTParticle();
 			if (auto Rigid = Handle.CastToRigidParticle())
-				{
+			{
 				if(Rigid->ResimType() == EResimType::FullResim)
-					{
+				{
 					if (IsFinalResim())
 					{
 						//Last resim so mark as in sync
 						Handle.SetSyncState(ESyncState::InSync);
+						Handle.SetEnabledDuringResim(false);
 
 						//Anything saved on upcoming frame (was done during rewind) can be removed since we are now at head
 						Info.ClearPhaseAndFuture(FutureFrame);
@@ -291,12 +295,12 @@ void FRewindData::FinishFrame()
 						//solver doesn't affect dynamics, so no reason to test if they desynced from original sim
 						//question: should we skip all other properties? dynamics is a commonly changed one but might be worth skipping everything solver skips
 						DesyncIfNecessary</*bSkipDynamics=*/true>(Info, FutureFrame);
+					}
 				}
 			}
 		}
-			}
-		}
-		
+	}
+	
 
 	++CurFrame;
 	LatestFrame = FMath::Max(LatestFrame,CurFrame);
@@ -305,6 +309,12 @@ void FRewindData::FinishFrame()
 int32 SkipDesyncTest = 0;
 FAutoConsoleVariableRef CVarSkipDesyncTest(TEXT("p.SkipDesyncTest"), SkipDesyncTest, TEXT("Skips hard desync test, this means all particles will assume to be clean except spawning at different times. This is useful for a perf lower bound, not actually correct"));
 
+void FRewindData::DesyncParticle(FDirtyParticleInfo& Info, const FFrameAndPhase FrameAndPhase)
+{
+	Info.ClearPhaseAndFuture(FrameAndPhase);
+	Info.GetPTParticle()->SetSyncState(ESyncState::HardDesync);
+}
+
 template <bool bSkipDynamics>
 void FRewindData::DesyncIfNecessary(FDirtyParticleInfo& Info, const FFrameAndPhase FrameAndPhase)
 {
@@ -316,8 +326,7 @@ void FRewindData::DesyncIfNecessary(FDirtyParticleInfo& Info, const FFrameAndPha
 		if (!SkipDesyncTest)
 		{
 			//first time desyncing so need to clear history from this point into the future
-			Info.ClearPhaseAndFuture(FrameAndPhase);
-			Info.GetPTParticle()->SetSyncState(ESyncState::HardDesync);
+			DesyncParticle(Info, FrameAndPhase);
 		}
 	}
 }
@@ -347,14 +356,18 @@ void FRewindData::AdvanceFrameImp(IResimCacheBase* ResimCache)
 		}
 		else
 		{
+
+			auto Handle = Info.GetPTParticle();
+			Info.bResimAsSlave = Handle->ResimType() == EResimType::ResimAsSlave;
+
 			if (IsResim() && !Info.bResimAsSlave)
 			{
 				DesyncIfNecessary(Info, FrameAndPhase);
-				}
+			}
 
-			auto Handle = Info.GetPTParticle();
 			if(IsResim() && Handle->SyncState() != ESyncState::InSync && !SkipDesyncTest)
-				{
+			{
+				Handle->SetEnabledDuringResim(true);	//for now just mark anything out of sync as resim enabled. TODO: use bubble
 				DesyncedParticles.Add(Handle);
 			}
 
@@ -362,7 +375,7 @@ void FRewindData::AdvanceFrameImp(IResimCacheBase* ResimCache)
 			{
 				//we only need to check the cast because right now there's no property system on PT, so any time a sim callback touches a particle we just mark it as dirty dynamics
 				if (auto Rigid = Handle->CastToRigidParticle())
-			{
+				{
 					//sim callback is finished so record the dynamics before solve starts
 					FGeometryParticleStateBase& Latest = Info.AddFrame(CurFrame);
 					Latest.Dynamics.WriteAccessMonotonic(FrameAndPhase, PropertiesPool).CopyFrom(*Rigid);
@@ -407,20 +420,16 @@ void FRewindData::PushGTDirtyData(const FDirtyPropertiesManager& SrcManager,cons
 			return;
 		}
 
-
-		FDirtyParticleInfo& Info = FindOrAddParticle(*PTParticle, Proxy->IsInitialized() ? INDEX_NONE : CurFrame);
-		FGeometryParticleStateBase& Latest = Info.AddFrame(CurFrame);
-
-		if(PTParticle->CastToRigidParticle() == nullptr)	//non rigid is always resimming as slave (TODO: we may want to move kinematics from callback)
+		if (bResim && Proxy->GetInitializedStep() == CurFrame)
 		{
-			Info.bResimAsSlave = true;
+			//Particle is reinitialized, since it's out of sync it must be at a different time
+			//So make sure it's considered during resim
+			//TODO: should check if in bubble
+			PTParticle->SetEnabledDuringResim(true);
 		}
 
-		if(auto NewData = Dirty.ParticleData.FindDynamicMisc(SrcManager, SrcDataIdx))
-		{
-			//question: does modifying this at runtime cause issues? For example a kinematic starting to simulate?
-			Info.bResimAsSlave = NewData->ResimType() == EResimType::ResimAsSlave;
-		}
+		FDirtyParticleInfo& Info = FindOrAddParticle(*PTParticle, Proxy->IsInitialized() ? INDEX_NONE : CurFrame);
+		FGeometryParticleStateBase& Latest = Info.AddFrame(CurFrame);
 
 		//At this point all phases should be clean
 		ensure(Latest.IsClean(FFrameAndPhase{ CurFrame, FFrameAndPhase::PrePushData }));
@@ -438,7 +447,7 @@ void FRewindData::PushGTDirtyData(const FDirtyPropertiesManager& SrcManager,cons
 				{
 					auto& Data = Property.WriteAccessMonotonic(FFrameAndPhase{ CurFrame, FFrameAndPhase::PrePushData }, PropertiesPool);
 					Data.CopyFrom(Particle);
-		}
+				}
 			};
 
 			DirtyPropHelper(Latest.ParticlePositionRotation, EParticleFlags::XR, *PTParticle);
@@ -474,12 +483,12 @@ void FRewindData::SpawnProxyIfNeeded(FSingleParticlePhysicsProxy& Proxy)
 	if(Proxy.GetInitializedStep() > CurFrame)
 	{
 		FGeometryParticleHandle* Handle = Proxy.GetHandle_LowLevel();
-		FindOrAddParticle(*Handle, CurFrame);
+		FDirtyParticleInfo& Info = FindOrAddParticle(*Handle, CurFrame);
 
 		Solver->GetEvolution()->EnableParticle(Handle, nullptr);
 		if(Proxy.GetInitializedStep() != CurFrame)
 		{
-			Handle->SetSyncState(ESyncState::HardDesync);	//spawn frame changed so desync
+			DesyncParticle(Info, FFrameAndPhase{ Proxy.GetInitializedStep(), FFrameAndPhase::PrePushData });	//Spawned earlier so mark as desynced from that first frame
 			Proxy.SetInitialized(CurFrame);
 		}
 	}
@@ -503,36 +512,36 @@ void FRewindData::MarkDirtyFromPT(FGeometryParticleHandle& Handle)
 	const FFrameAndPhase FrameAndPhase{ CurFrame, FFrameAndPhase::PostPushData };
 
 	if(bRecordingHistory || Latest.ParticlePositionRotation.IsClean(FrameAndPhase))
-{
+	{
 		if (auto Data = Latest.ParticlePositionRotation.WriteAccessNonDecreasing(FFrameAndPhase{ CurFrame, FFrameAndPhase::PostPushData }, PropertiesPool))
 		{
 			Data->CopyFrom(Handle);
 		}
-}
+	}
 
 
 	if (auto Kinematic = Handle.CastToKinematicParticle())
 	{
 		if (bRecordingHistory || Latest.Velocities.IsClean(FrameAndPhase))
-			{
+		{
 			if (auto Data = Latest.Velocities.WriteAccessNonDecreasing(FrameAndPhase, PropertiesPool))
-{
+			{
 				Data->CopyFrom(*Kinematic);
+			}
 		}
-	}
 
 		if (auto Rigid = Kinematic->CastToRigidParticle())
-	{
-			if (bRecordingHistory || Latest.DynamicsMisc.IsClean(FrameAndPhase))
 		{
-				if (auto Data = Latest.DynamicsMisc.WriteAccessNonDecreasing(FrameAndPhase, PropertiesPool))
+			if (bRecordingHistory || Latest.DynamicsMisc.IsClean(FrameAndPhase))
 			{
+				if (auto Data = Latest.DynamicsMisc.WriteAccessNonDecreasing(FrameAndPhase, PropertiesPool))
+				{
 					Data->CopyFrom(*Rigid);
 				}
 			}
-			}
 		}
 	}
+}
 
 template <bool bResim>
 void FRewindData::PushPTDirtyData(TPBDRigidParticleHandle<FReal,3>& Handle,const int32 SrcDataIdx)
@@ -546,16 +555,16 @@ void FRewindData::PushPTDirtyData(TPBDRigidParticleHandle<FReal,3>& Handle,const
 	ensure(!bRecordingHistory || Latest.IsCleanExcludingDynamics(FrameAndPhase));	//PostCallbacks should be clean before we write sim results
 
 	if(bRecordingHistory || Latest.ParticlePositionRotation.IsClean(FrameAndPhase))
-		{
+	{
 		Latest.ParticlePositionRotation.WriteAccessMonotonic(FrameAndPhase, PropertiesPool).CopyFrom(Handle);
-			}
+	}
 			
 	if(bRecordingHistory || Latest.Velocities.IsClean(FrameAndPhase))
-			{
+	{
 		FParticleVelocities& PreVelocities = Latest.Velocities.WriteAccessMonotonic(FrameAndPhase, PropertiesPool);
 		PreVelocities.SetV(Handle.PreV());
 		PreVelocities.SetW(Handle.PreW());
-		}
+	}
 
 	if(bRecordingHistory || Latest.DynamicsMisc.IsClean(FrameAndPhase))
 	{
@@ -626,4 +635,4 @@ template void FRewindData::PushGTDirtyData<false>(const FDirtyPropertiesManager&
 template void FRewindData::PushPTDirtyData<true>(TPBDRigidParticleHandle<FReal,3>& Rigid,const int32 SrcDataIdx);
 template void FRewindData::PushPTDirtyData<false>(TPBDRigidParticleHandle<FReal,3>& Rigid,const int32 SrcDataIdx);
 
-}
\ No newline at end of file
+}
diff --git a/Engine/Source/Runtime/Experimental/Chaos/Public/Chaos/Collision/SpatialAccelerationBroadPhase.h b/Engine/Source/Runtime/Experimental/Chaos/Public/Chaos/Collision/SpatialAccelerationBroadPhase.h
--- a/Engine/Source/Runtime/Experimental/Chaos/Public/Chaos/Collision/SpatialAccelerationBroadPhase.h
+++ b/Engine/Source/Runtime/Experimental/Chaos/Public/Chaos/Collision/SpatialAccelerationBroadPhase.h
@@ -201,6 +201,9 @@ namespace Chaos
 
 			TArray<FAccelerationStructureHandle> PotentialIntersections;
 
+			//If we are in a resim, this particle must be diverged and so it better be marked as enabled during resim
+			ensure(!bIsResimming || Particle1.EnabledDuringResim());
+
 			if (bIsResimming || (Particle1.ObjectState() == EObjectStateType::Dynamic || Particle1.ObjectState() == EObjectStateType::Sleeping))
 			{
 				const bool bBody1Bounded = HasBoundingBox(Particle1);
diff --git a/Engine/Source/Runtime/Experimental/Chaos/Public/Chaos/GeometryParticles.h b/Engine/Source/Runtime/Experimental/Chaos/Public/Chaos/GeometryParticles.h
--- a/Engine/Source/Runtime/Experimental/Chaos/Public/Chaos/GeometryParticles.h
+++ b/Engine/Source/Runtime/Experimental/Chaos/Public/Chaos/GeometryParticles.h
@@ -324,6 +324,14 @@ namespace Chaos
 		};
 		TSharedPtr<FData,ESPMode::ThreadSafe> SharedData;
 	};
+
+	enum class EResimType : uint8
+	{
+		FullResim,	//fully re-run simulation and keep results (any forces must be applied again)
+		//ResimWithPrevForces, //use previous forces and keep results (UNIMPLEMENTED)
+		ResimAsSlave //use previous forces and snap to previous results regardless of variation - used to push other objects away
+		//ResimAsKinematic //treat as kinematic (UNIMPLEMENTED)
+	};
 	
 	template<class T, int d, EGeometryParticlesSimType SimType>
 	class TGeometryParticlesImp : public TParticles<T, d>
@@ -356,6 +364,9 @@ namespace Chaos
 			TArrayCollection::AddArray(&MSyncState);
 			TArrayCollection::AddArray(&MWeakParticleHandle);
 			TArrayCollection::AddArray(&MParticleConstraints);
+			TArrayCollection::AddArray(&MResimType);
+			TArrayCollection::AddArray(&MEnabledDuringResim);
+
 
 #if CHAOS_DEBUG_NAME
 			TArrayCollection::AddArray(&MDebugName);
@@ -391,6 +402,9 @@ namespace Chaos
 			, MSyncState(MoveTemp(Other.MSyncState))
 			, MWeakParticleHandle(MoveTemp(Other.MWeakParticleHandle))
 			, MParticleConstraints(MoveTemp(Other.MParticleConstraints))
+			, MResimType(MoveTemp(Other.MResimType))
+			, MEnabledDuringResim(MoveTemp(Other.MEnabledDuringResim))
+
 
 #if CHAOS_DETERMINISTIC
 			, MParticleIDs(MoveTemp(Other.MParticleIDs))
@@ -427,6 +441,9 @@ namespace Chaos
 				TArrayCollection::AddArray(&MGeometryParticle);
 				TArrayCollection::AddArray(&MPhysicsProxy);
 			}
+
+			TArrayCollection::AddArray(&MResimType);
+			TArrayCollection::AddArray(&MEnabledDuringResim);
 		}
 
 		static constexpr bool IsRigidBodySim() { return SimType == EGeometryParticlesSimType::RigidBodySim; }
@@ -465,6 +482,9 @@ namespace Chaos
 				TArrayCollection::AddArray(&MGeometryParticle);
 				TArrayCollection::AddArray(&MPhysicsProxy);
 			}
+			
+			TArrayCollection::AddArray(&MResimType);
+			TArrayCollection::AddArray(&MEnabledDuringResim);
 		}
 
 		CHAOS_API virtual ~TGeometryParticlesImp()
@@ -639,6 +659,13 @@ namespace Chaos
 			MParticleConstraints[Index].RemoveSingleSwap(InConstraintHandle);
 			CHAOS_ENSURE(!MParticleConstraints[Index].Contains(InConstraintHandle));
 		}
+
+		FORCEINLINE EResimType ResimType(const int32 Index) const { return MResimType[Index]; }
+		FORCEINLINE EResimType& ResimType(const int32 Index) { return MResimType[Index]; }
+
+		FORCEINLINE bool EnabledDuringResim(const int32 Index) const { return MEnabledDuringResim[Index]; }
+		FORCEINLINE bool& EnabledDuringResim(const int32 Index) { return MEnabledDuringResim[Index]; }
+
 private:
 		friend THandleType;
 		CHAOS_API void ResetWeakParticleHandle(const int32 Index)
@@ -775,6 +802,8 @@ public:
 		TArrayCollectionArray<FSyncState> MSyncState;
 		TArrayCollectionArray<FWeakParticleHandle> MWeakParticleHandle;
 		TArrayCollectionArray<FConstraintHandleArray> MParticleConstraints;
+		TArrayCollectionArray<EResimType> MResimType;
+		TArrayCollectionArray<bool> MEnabledDuringResim;
 
 		void UpdateShapesArray(const int32 Index)
 		{
diff --git a/Engine/Source/Runtime/Experimental/Chaos/Public/Chaos/ParticleDirtyFlags.h b/Engine/Source/Runtime/Experimental/Chaos/Public/Chaos/ParticleDirtyFlags.h
--- a/Engine/Source/Runtime/Experimental/Chaos/Public/Chaos/ParticleDirtyFlags.h
+++ b/Engine/Source/Runtime/Experimental/Chaos/Public/Chaos/ParticleDirtyFlags.h
@@ -318,7 +318,6 @@ public:
 		SetObjectState(Other.ObjectState());
 		SetGravityEnabled(Other.GravityEnabled());
 		SetCollisionGroup(Other.CollisionGroup());
-		SetResimType(Other.ResimType());
 		SetSleepType(Other.SleepType());
 		SetOneWayInteraction(Other.OneWayInteraction());
 		SetCollisionConstraintFlag(Other.CollisionConstraintFlag());
@@ -336,7 +335,6 @@ public:
 			&& MaxAngularSpeedSq() == Other.MaxAngularSpeedSq()
 			&& GravityEnabled() == Other.GravityEnabled()
 			&& CollisionGroup() == Other.CollisionGroup()
-			&& ResimType() == Other.ResimType()
 			&& SleepType() == Other.SleepType()
 			&& OneWayInteraction() == Other.OneWayInteraction() 
 			&& CollisionConstraintFlag() == Other.CollisionConstraintFlag()
@@ -376,9 +374,6 @@ public:
 	int32 CollisionGroup() const { return MCollisionGroup; }
 	void SetCollisionGroup(int32 InGroup){ MCollisionGroup = InGroup; }
 
-	EResimType ResimType() const { return MResimType; }
-	void SetResimType(EResimType Type) { MResimType = Type; }
-
 	ESleepType SleepType() const { return MSleepType; }
 	void SetSleepType(ESleepType Type) { MSleepType = Type; }
 
@@ -434,6 +429,7 @@ public:
 		SetGeometry(Other.SharedGeometryLowLevel());
 		SetUniqueIdx(Other.UniqueIdx());
 		SetSpatialIdx(Other.SpatialIdx());
+		SetResimType(Other.ResimType());
 #if CHAOS_DEBUG_NAME
 		SetDebugName(Other.DebugName());
 #endif
@@ -444,7 +440,8 @@ public:
 	{
 		return Geometry() == Other.Geometry()
 			&& UniqueIdx() == Other.UniqueIdx()
-			&& SpatialIdx() == Other.SpatialIdx();
+			&& SpatialIdx() == Other.SpatialIdx()
+			&& ResimType() == Other.ResimType();
 	}
 
 	bool operator==(const FParticleNonFrequentData& Other) const
@@ -466,6 +463,13 @@ public:
 	FSpatialAccelerationIdx SpatialIdx() const { return MSpatialIdx; }
 	void SetSpatialIdx(FSpatialAccelerationIdx InIdx){ MSpatialIdx = InIdx; }
 
+	EResimType ResimType() const { return MResimType; }
+
+	void SetResimType(EResimType InType)
+	{
+		MResimType = InType;
+	}
+
 #if CHAOS_DEBUG_NAME
 	const TSharedPtr<FString, ESPMode::ThreadSafe>& DebugName() const { return MDebugName; }
 	void SetDebugName(const TSharedPtr<FString, ESPMode::ThreadSafe>& InName) { MDebugName = InName; }
@@ -474,7 +478,7 @@ private:
 	TSharedPtr<const FImplicitObject,ESPMode::ThreadSafe> MGeometry;
 	FUniqueIdx MUniqueIdx;
 	FSpatialAccelerationIdx MSpatialIdx;
-
+	EResimType MResimType;
 #if CHAOS_DEBUG_NAME
 	TSharedPtr<FString, ESPMode::ThreadSafe> MDebugName;
 #endif
diff --git a/Engine/Source/Runtime/Experimental/Chaos/Public/Chaos/ParticleHandle.h b/Engine/Source/Runtime/Experimental/Chaos/Public/Chaos/ParticleHandle.h
--- a/Engine/Source/Runtime/Experimental/Chaos/Public/Chaos/ParticleHandle.h
+++ b/Engine/Source/Runtime/Experimental/Chaos/Public/Chaos/ParticleHandle.h
@@ -64,6 +64,7 @@ void GeometryParticleDefaultConstruct(FConcrete& Concrete, const FGeometryPartic
 	Concrete.SetX(TVector<T, d>(0));
 	Concrete.SetR(TRotation<T, d>::Identity);
 	Concrete.SetSpatialIdx(FSpatialAccelerationIdx{ 0,0 });
+	Concrete.SetResimType(EResimType::FullResim);
 }
 
 template <typename T, int d, typename FConcrete>
@@ -92,7 +93,6 @@ void PBDRigidParticleDefaultConstruct(FConcrete& Concrete, const FPBDRigidPartic
 	Concrete.SetGravityEnabled(Params.bGravityEnabled);
 	Concrete.SetCCDEnabled(Params.bCCDEnabled);
 	Concrete.SetDisabled(Params.bDisabled);
-	Concrete.SetResimType(EResimType::FullResim);
 	Concrete.SetSleepType(ESleepType::MaterialSleep);
 }
 
@@ -345,6 +345,7 @@ protected:
 		//TODO: patch from SOA
 		GeometryParticleDefaultConstruct<T, d>(*this, Params);
 		SetHasBounds(false);
+		SetEnabledDuringResim(false);
 	}
 
 	template <typename TParticlesType, typename TParams>
@@ -413,6 +414,7 @@ public:
 		SetSharedGeometry(InData.SharedGeometryLowLevel());
 		SetUniqueIdx(InData.UniqueIdx());
 		SetSpatialIdx(InData.SpatialIdx());
+		SetResimType(InData.ResimType());
 
 #if CHAOS_DEBUG_NAME
 		SetDebugName(InData.DebugName());
@@ -504,6 +506,16 @@ public:
 		return AuxContainer[HandleIdx];
 	}
 
+	EResimType ResimType() const { return GeometryParticles->ResimType(ParticleIdx); }
+
+	void SetResimType(EResimType ResimType) { GeometryParticles->ResimType(ParticleIdx) = ResimType; }
+
+
+	bool EnabledDuringResim() const { return GeometryParticles->EnabledDuringResim(ParticleIdx); }
+	void SetEnabledDuringResim(bool bEnabledDuringResim) { GeometryParticles->EnabledDuringResim(ParticleIdx) = bEnabledDuringResim; }
+
+
+
 #if CHAOS_DETERMINISTIC
 	FParticleID ParticleID() const { return GeometryParticles->ParticleID(ParticleIdx); }
 	FParticleID& ParticleID() { return GeometryParticles->ParticleID(ParticleIdx); }
@@ -899,10 +911,6 @@ public:
 
 	void SetOneWayInteraction(bool bInOneWayInteraction) { PBDRigidParticles->OneWayInteraction(ParticleIdx) = bInOneWayInteraction; }
 
-	EResimType ResimType() const { return PBDRigidParticles->ResimType(ParticleIdx);}
-
-	void SetResimType(EResimType ResimType){ PBDRigidParticles->ResimType(ParticleIdx) = ResimType; }
-
 	ESleepType SleepType() const { return PBDRigidParticles->SleepType(ParticleIdx);}
 
 	void SetSleepType(ESleepType SleepType){ PBDRigidParticles->SetSleepType(ParticleIdx, SleepType); }
@@ -1172,6 +1180,8 @@ public:
 	bool Sleeping() const { return MHandle->Sleeping(); }
 	FString ToString() const { return MHandle->ToString(); }
 
+	bool EnabledDuringResim() const { return MHandle->EnabledDuringResim(); }
+
 	template <typename Container>
 	const auto& AuxilaryValue(const Container& AuxContainer) const { return MHandle->AuxilaryValue(AuxContainer); }
 	template <typename Container>
@@ -1824,6 +1834,17 @@ public:
 		MNonFrequentData.Modify(true,MDirtyFlags,Proxy,[Idx](auto& Data){ Data.SetSpatialIdx(Idx);});
 	}
 
+	void SetResimType(EResimType ResimType)
+	{
+		MNonFrequentData.Modify(true, MDirtyFlags, Proxy, [ResimType](auto& Data) { Data.SetResimType(ResimType); });
+	}
+
+	EResimType ResimType() const
+	{
+		return MNonFrequentData.Read().ResimType();
+	}
+
+
 	void SetNonFrequentData(const FParticleNonFrequentData& InData)
 	{
 		MNonFrequentData.Write(InData,true,MDirtyFlags,Proxy);
@@ -2197,16 +2218,6 @@ public:
 		this->MInitialized = InInitialized;
 	}
 
-	void SetResimType(EResimType ResimType)
-	{
-		MMiscData.Modify(true,MDirtyFlags,Proxy,[ResimType](auto& Data){ Data.SetResimType(ResimType);});
-	}
-
-	EResimType ResimType() const
-	{
-		return MMiscData.Read().ResimType();
-	}
-
 	const TVector<T, d>& F() const { return MDynamics.Read().F(); }
 	void AddForce(const TVector<T, d>& InF, bool bInvalidate = true)
 	{
diff --git a/Engine/Source/Runtime/Experimental/Chaos/Public/Chaos/RigidParticles.h b/Engine/Source/Runtime/Experimental/Chaos/Public/Chaos/RigidParticles.h
--- a/Engine/Source/Runtime/Experimental/Chaos/Public/Chaos/RigidParticles.h
+++ b/Engine/Source/Runtime/Experimental/Chaos/Public/Chaos/RigidParticles.h
@@ -16,14 +16,6 @@
 namespace Chaos
 {
 
-enum class EResimType : uint8
-{
-	FullResim,	//fully re-run simulation and keep results (any forces must be applied again)
-	//ResimWithPrevForces, //use previous forces and keep results (UNIMPLEMENTED)
-	ResimAsSlave //use previous forces and snap to previous results regardless of variation - used to push other objects away
-	//ResimAsKinematic //treat as kinematic (UNIMPLEMENTED)
-};
-
 enum class ESleepType : uint8
 {
 	MaterialSleep,	//physics material determines sleep threshold
@@ -101,7 +93,6 @@ public:
 		TArrayCollection::AddArray(&MToBeRemovedOnFracture);
 		TArrayCollection::AddArray(&MGravityEnabled);
 		TArrayCollection::AddArray(&MOneWayInteraction);
-		TArrayCollection::AddArray(&MResimType);
 		TArrayCollection::AddArray(&MSleepType);
 		TArrayCollection::AddArray(&bCCDEnabled);
 	}
@@ -128,7 +119,6 @@ public:
 		, MPreObjectState(MoveTemp(Other.MPreObjectState))
 		, MGravityEnabled(MoveTemp(Other.MGravityEnabled))
 		, MOneWayInteraction(MoveTemp(Other.MOneWayInteraction))
-		, MResimType(MoveTemp(Other.MResimType))
 		, MSleepType(MoveTemp(Other.MSleepType))
 		, bCCDEnabled(MoveTemp(Other.bCCDEnabled))
 	{
@@ -158,7 +148,6 @@ public:
 		TArrayCollection::AddArray(&MToBeRemovedOnFracture);
 		TArrayCollection::AddArray(&MGravityEnabled);
 		TArrayCollection::AddArray(&MOneWayInteraction);
-		TArrayCollection::AddArray(&MResimType);
 		TArrayCollection::AddArray(&MSleepType);
 		TArrayCollection::AddArray(&bCCDEnabled);
 	}
@@ -248,8 +237,6 @@ public:
 	FORCEINLINE const bool& OneWayInteraction(const int32 Index) const { return MOneWayInteraction[Index]; }
 	FORCEINLINE bool& OneWayInteraction(const int32 Index) { return MOneWayInteraction[Index]; }
 
-	FORCEINLINE EResimType ResimType(const int32 Index) const { return MResimType[Index]; }
-	FORCEINLINE EResimType& ResimType(const int32 Index) { return MResimType[Index]; }
 
 	FORCEINLINE const bool& CCDEnabled(const int32 Index) const { return bCCDEnabled[Index]; }
 	FORCEINLINE bool& CCDEnabled(const int32 Index) { return bCCDEnabled[Index]; }
@@ -370,7 +357,6 @@ private:
 	TArrayCollectionArray<EObjectStateType> MPreObjectState;
 	TArrayCollectionArray<bool> MGravityEnabled;
 	TArrayCollectionArray<bool> MOneWayInteraction;
-	TArrayCollectionArray<EResimType> MResimType;
 	TArrayCollectionArray<ESleepType> MSleepType;
 	TArrayCollectionArray<bool> bCCDEnabled;
 
diff --git a/Engine/Source/Runtime/Experimental/Chaos/Public/RewindData.h b/Engine/Source/Runtime/Experimental/Chaos/Public/RewindData.h
--- a/Engine/Source/Runtime/Experimental/Chaos/Public/RewindData.h
+++ b/Engine/Source/Runtime/Experimental/Chaos/Public/RewindData.h
@@ -1048,7 +1048,7 @@ private:
 		int32 LastDirtyFrame;	//Track how recently this was made dirty
 		int32 InitializedOnStep = INDEX_NONE;	//if not INDEX_NONE, it indicates we saw initialization during rewind history window
 		int32 DirtyDynamics = INDEX_NONE;
-		bool bResimAsSlave = false;	//Indicates the particle will always resim in the exact same way from game thread data
+		bool bResimAsSlave = true;	//Indicates the particle will always resim in the exact same way from game thread data
 
 		FDirtyParticleInfo(FDirtyPropertiesPool& InPropertiesPool, TGeometryParticleHandle<FReal,3>& InPTParticle, const FUniqueIdx UniqueIdx,const int32 CurFrame,const int32 NumFrames)
 		: History(NumFrames)
@@ -1132,6 +1132,7 @@ private:
 	FDirtyParticleInfo& FindOrAddParticle(TGeometryParticleHandle<FReal,3>& PTParticle, const int32 InitializedOnFrame = INDEX_NONE);
 	bool RewindToFrame(int32 Frame);
 
+	static void DesyncParticle(FDirtyParticleInfo& Info, const FFrameAndPhase FrameAndPhase);
 
 	TArrayAsMap<FUniqueIdx,int32> ParticleToAllDirtyIdx;
 	TCircularBuffer<FFrameManagerInfo> Managers;
diff --git a/Engine/Source/Runtime/PhysicsCore/Private/ChaosEngineInterface.cpp b/Engine/Source/Runtime/PhysicsCore/Private/ChaosEngineInterface.cpp
--- a/Engine/Source/Runtime/PhysicsCore/Private/ChaosEngineInterface.cpp
+++ b/Engine/Source/Runtime/PhysicsCore/Private/ChaosEngineInterface.cpp
@@ -1605,6 +1605,7 @@ void FChaosEngineInterface::CreateActor(const FActorCreationParams& InParams,FPh
 	if(InParams.bStatic)
 	{
 		Particle = FGeometryParticle::CreateParticle();
+		Particle->SetResimType(EResimType::ResimAsSlave);
 	}
 	else
 	{
