FractureTool
-- implemented HideUnselected mode

diff --git a/Engine/Plugins/Experimental/ChaosEditor/Source/FractureEditor/Private/FractureEditorMode.cpp b/Engine/Plugins/Experimental/ChaosEditor/Source/FractureEditor/Private/FractureEditorMode.cpp
--- a/Engine/Plugins/Experimental/ChaosEditor/Source/FractureEditor/Private/FractureEditorMode.cpp
+++ b/Engine/Plugins/Experimental/ChaosEditor/Source/FractureEditor/Private/FractureEditorMode.cpp
@@ -71,6 +71,7 @@ void UFractureEditorMode::Enter()
 	if (FFractureEditorModeToolkit* FractureToolkit = static_cast<FFractureEditorModeToolkit*>(Toolkit.Get()))
 	{
 		FractureToolkit->SetInitialPalette();
+		FractureToolkit->OnHideUnselectedChanged(FractureToolkit->GetHideUnselectedValue());
 	}
 	
 	OnActorSelectionChanged(SelectedObjects, false);
@@ -357,6 +358,16 @@ void UFractureEditorMode::OnActorSelectionChanged(const TArray<UObject*>& NewSel
 
 			FScopedColorEdit ShowBoneColorsEdit(ExistingSelection);
 			ShowBoneColorsEdit.SetEnableBoneSelection(false);
+
+			// If we have a Hide array on the collection, remove it.
+			if (const UGeometryCollection* RestCollection = ExistingSelection->GetRestCollection())
+			{
+				TSharedPtr<FGeometryCollection, ESPMode::ThreadSafe> GeometryCollection = RestCollection->GetGeometryCollection();
+				if (GeometryCollection->HasAttribute("Hide", FGeometryCollection::TransformGroup))
+				{
+					GeometryCollection->RemoveAttribute("Hide", FGeometryCollection::TransformGroup);
+				}
+			}
 			
 			ExistingSelection->MarkRenderStateDirty();
 		}
diff --git a/Engine/Plugins/Experimental/ChaosEditor/Source/FractureEditor/Private/FractureEditorModeToolkit.cpp b/Engine/Plugins/Experimental/ChaosEditor/Source/FractureEditor/Private/FractureEditorModeToolkit.cpp
--- a/Engine/Plugins/Experimental/ChaosEditor/Source/FractureEditor/Private/FractureEditorModeToolkit.cpp
+++ b/Engine/Plugins/Experimental/ChaosEditor/Source/FractureEditor/Private/FractureEditorModeToolkit.cpp
@@ -477,6 +477,10 @@ void FFractureEditorModeToolkit::OnObjectPostEditChange( UObject* Object, FPrope
 		{
 			OnLevelViewValueChanged();
 		}
+		if (PropertyChangedEvent.Property->GetFName() == GET_MEMBER_NAME_CHECKED(UFractureSettings, bHideUnselected))
+		{
+			OnHideUnselectedChanged(GetHideUnselectedValue());
+		}
 		else if (PropertyChangedEvent.Property->GetFName() == GET_MEMBER_NAME_CHECKED(UHistogramSettings, bSorted))
 		{
 			UHistogramSettings* HistogramSettings = GetMutableDefault<UHistogramSettings>();
@@ -660,6 +664,50 @@ void FFractureEditorModeToolkit::BindCommands()
 	}
 }
 
+void FFractureEditorModeToolkit::SetHideForUnselected(UGeometryCollectionComponent* GCComp)
+{
+	if (const UGeometryCollection* RestCollection = GCComp->GetRestCollection())
+	{
+		TSharedPtr<FGeometryCollection, ESPMode::ThreadSafe> GeometryCollection = RestCollection->GetGeometryCollection();
+
+		// If Hide managed array exists, set false for any selected bones, true for selected. 
+		// If a cluster is selected, set false for all children.
+		if (GeometryCollection->HasAttribute("Hide", FGeometryCollection::TransformGroup))
+		{
+			TManagedArray<bool>& Hide = GeometryCollection->GetAttribute<bool>("Hide", FGeometryCollection::TransformGroup);
+			const TManagedArray<TSet<int32>>& Children = GeometryCollection->GetAttribute<TSet<int32>>("Children", FGeometryCollection::TransformGroup);
+			
+			const TArray<int32>& SelectedBones = GCComp->GetSelectedBones();
+			if (SelectedBones.Num() > 0)
+			{ 
+				Hide.Fill(true);
+
+				for (int32 SelectedBone : SelectedBones)
+				{
+					Hide[SelectedBone] = false;
+					if (Children[SelectedBone].Num() > 0)
+					{ 
+						TArray<int32> BranchBones;
+						FGeometryCollectionClusteringUtility::RecursiveAddAllChildren(Children, SelectedBone, BranchBones);
+						for (int32 BranchBone : BranchBones)
+						{
+							Hide[BranchBone] = false;
+						}
+					}
+				}
+			}
+			else
+			{
+				// Don't hide anything if we've selected nothing
+				Hide.Fill(false);
+			}
+
+			GCComp->RefreshEmbeddedGeometry();
+
+		}
+	}
+}
+
 void FFractureEditorModeToolkit::OnToolPaletteChanged(FName PaletteName) 
 {
 	if (GetActiveTool() != nullptr)
@@ -700,6 +748,12 @@ int32 FFractureEditorModeToolkit::GetLevelViewValue() const
 	return FractureSettings->FractureLevel;
 }
 
+bool FFractureEditorModeToolkit::GetHideUnselectedValue() const
+{
+	UFractureSettings* FractureSettings = GetMutableDefault<UFractureSettings>();
+	return FractureSettings->bHideUnselected;
+}
+
 void FFractureEditorModeToolkit::OnSetExplodedViewValue(float NewValue)
 {
 	FScopedTransaction Transaction(LOCTEXT("SetExplodedViewValue", "Adjust Exploded View"));
@@ -822,6 +876,42 @@ void FFractureEditorModeToolkit::OnLevelViewValueChanged()
 	GCurrentLevelEditingViewportClient->Invalidate();
 }
 
+void FFractureEditorModeToolkit::OnHideUnselectedChanged(bool bHide)
+{
+	TSet<UGeometryCollectionComponent*> GeomCompSelection;
+	GetSelectedGeometryCollectionComponents(GeomCompSelection);
+
+	for (UGeometryCollectionComponent* Comp : GeomCompSelection)
+	{
+		if (const UGeometryCollection* RestCollection = Comp->GetRestCollection())
+		{ 
+			TSharedPtr<FGeometryCollection, ESPMode::ThreadSafe> GeometryCollection = RestCollection->GetGeometryCollection();
+			
+			if (bHide)
+			{
+				// If we are toggling on, add and configure the Hide array.
+				if (!GeometryCollection->HasAttribute("Hide", FGeometryCollection::TransformGroup))
+				{
+					GeometryCollection->AddAttribute<bool>("Hide", FGeometryCollection::TransformGroup);
+				}
+				SetHideForUnselected(Comp);
+			}
+			else
+			{
+				// If we are toggling off, remove the Hide array.
+				if (GeometryCollection->HasAttribute("Hide", FGeometryCollection::TransformGroup))
+				{
+					GeometryCollection->RemoveAttribute("Hide", FGeometryCollection::TransformGroup);
+				}
+			}
+		
+			// redraw
+			Comp->MarkRenderDynamicDataDirty();
+			Comp->MarkRenderStateDirty();
+		}
+	}
+}
+
 void FFractureEditorModeToolkit::ToggleShowBoneColors()
 {
 	TSet<UGeometryCollectionComponent*> GeomCompSelection;
@@ -1022,6 +1112,8 @@ void FFractureEditorModeToolkit::SetBoneSelection(UGeometryCollectionComponent*
 {
 	OutlinerView->SetBoneSelection(InRootComponent, InSelectedBones, bClearCurrentSelection);
 	HistogramView->SetBoneSelection(InRootComponent, InSelectedBones, bClearCurrentSelection);
+
+	SetHideForUnselected(InRootComponent);
 	
 	if (ActiveTool != nullptr)
 	{
@@ -1436,6 +1528,8 @@ void FFractureEditorModeToolkit::OnOutlinerBoneSelectionChanged(UGeometryCollect
 			ActiveTool->FractureContextChanged();
 		}
 
+		SetHideForUnselected(RootComponent);
+
 		RootComponent->MarkRenderStateDirty();
 		RootComponent->MarkRenderDynamicDataDirty();
 	}
@@ -1466,6 +1560,8 @@ void FFractureEditorModeToolkit::OnHistogramBoneSelectionChanged(UGeometryCollec
 			ActiveTool->FractureContextChanged();
 		}
 
+		SetHideForUnselected(RootComponent);
+
 		RootComponent->MarkRenderStateDirty();
 		RootComponent->MarkRenderDynamicDataDirty();
 	}
diff --git a/Engine/Plugins/Experimental/ChaosEditor/Source/FractureEditor/Private/FractureSettings.cpp b/Engine/Plugins/Experimental/ChaosEditor/Source/FractureEditor/Private/FractureSettings.cpp
--- a/Engine/Plugins/Experimental/ChaosEditor/Source/FractureEditor/Private/FractureSettings.cpp
+++ b/Engine/Plugins/Experimental/ChaosEditor/Source/FractureEditor/Private/FractureSettings.cpp
@@ -8,6 +8,7 @@ UFractureSettings::UFractureSettings(const FObjectInitializer& ObjInit)
 	: Super(ObjInit)
 	, ExplodeAmount(0.0)
 	, FractureLevel(-1)
+	, bHideUnselected(false)
 {}
 
 #endif
diff --git a/Engine/Plugins/Experimental/ChaosEditor/Source/FractureEditor/Public/FractureEditorModeToolkit.h b/Engine/Plugins/Experimental/ChaosEditor/Source/FractureEditor/Public/FractureEditorModeToolkit.h
--- a/Engine/Plugins/Experimental/ChaosEditor/Source/FractureEditor/Public/FractureEditorModeToolkit.h
+++ b/Engine/Plugins/Experimental/ChaosEditor/Source/FractureEditor/Public/FractureEditorModeToolkit.h
@@ -111,11 +111,13 @@ public:
 	// View Settings
 	float GetExplodedViewValue() const;
 	int32 GetLevelViewValue() const;
+	bool GetHideUnselectedValue() const;
 	void OnSetExplodedViewValue(float NewValue);
 	void OnSetLevelViewValue(int32 NewValue);
 
 	void OnExplodedViewValueChanged();
 	void OnLevelViewValueChanged();
+	void OnHideUnselectedChanged(bool bHide);
 
 	// Update any View Property Changes 
 	void OnObjectPostEditChange( UObject* Object, FPropertyChangedEvent& PropertyChangedEvent );
@@ -186,6 +188,8 @@ private:
 	void OnHistogramBoneSelectionChanged(UGeometryCollectionComponent* RootComponent, TArray<int32>& SelectedBones);
 	void BindCommands();
 
+	void SetHideForUnselected(UGeometryCollectionComponent* GCComp);
+
 private:
 	UFractureModalTool* ActiveTool;
 
diff --git a/Engine/Plugins/Experimental/ChaosEditor/Source/FractureEditor/Public/FractureSettings.h b/Engine/Plugins/Experimental/ChaosEditor/Source/FractureEditor/Public/FractureSettings.h
--- a/Engine/Plugins/Experimental/ChaosEditor/Source/FractureEditor/Public/FractureSettings.h
+++ b/Engine/Plugins/Experimental/ChaosEditor/Source/FractureEditor/Public/FractureSettings.h
@@ -23,4 +23,8 @@ public:
 	UPROPERTY(EditAnywhere, Category = ViewSettings, meta = (DisplayName = "Fracture Level", UIMin = "-1" ))
 	int32 FractureLevel;
 
+	/** When active, only show selected bones.  */
+	UPROPERTY(EditAnywhere, Category = ViewSettings, meta = (DisplayName = "Hide Unselected"))
+	bool bHideUnselected;
+
 };
\ No newline at end of file
diff --git a/Engine/Source/Runtime/Experimental/GeometryCollectionEngine/Private/GeometryCollection/GeometryCollectionComponent.cpp b/Engine/Source/Runtime/Experimental/GeometryCollectionEngine/Private/GeometryCollection/GeometryCollectionComponent.cpp
--- a/Engine/Source/Runtime/Experimental/GeometryCollectionEngine/Private/GeometryCollection/GeometryCollectionComponent.cpp
+++ b/Engine/Source/Runtime/Experimental/GeometryCollectionEngine/Private/GeometryCollection/GeometryCollectionComponent.cpp
@@ -816,6 +816,12 @@ void UGeometryCollectionComponent::RefreshEmbeddedGeometry()
 	const TManagedArray<int32>& ExemplarIndexArray = GetExemplarIndexArray();
 	const int32 TransformCount = GlobalMatrices.Num();
 	
+	const TManagedArray<bool>* HideArray = nullptr;
+	if (RestCollection->GetGeometryCollection()->HasAttribute("Hide", FGeometryCollection::TransformGroup))
+	{
+		HideArray = &RestCollection->GetGeometryCollection()->GetAttribute<bool>("Hide", FGeometryCollection::TransformGroup);
+	}
+	
 	const int32 ExemplarCount = EmbeddedGeometryComponents.Num();
 	for (int32 ExemplarIndex = 0; ExemplarIndex < ExemplarCount; ++ExemplarIndex)
 	{		
@@ -827,7 +833,10 @@ void UGeometryCollectionComponent::RefreshEmbeddedGeometry()
 		{
 			if (ExemplarIndexArray[Idx] == ExemplarIndex)
 			{
-				InstanceTransforms.Add(FTransform(GlobalMatrices[Idx]));
+				if (!HideArray || !(*HideArray)[Idx])
+				{ 
+					InstanceTransforms.Add(FTransform(GlobalMatrices[Idx]));
+				}
 			}
 		}
 
@@ -1334,6 +1343,11 @@ void UGeometryCollectionComponent::InitConstantData(FGeometryCollectionConstantD
 		const TManagedArray<TArray<FVector2D>>& UVs = Collection->UVs;
 		const TManagedArray<FLinearColor>& Color = Collection->Color;
 		const TManagedArray<FLinearColor>& BoneColors = Collection->BoneColor;
+		
+		const int32 NumGeom = Collection->NumElements(FGeometryCollection::GeometryGroup);
+		const TManagedArray<int32>& TransformIndex = Collection->TransformIndex;
+		const TManagedArray<int32>& FaceStart = Collection->FaceStart;
+		const TManagedArray<int32>& FaceCount = Collection->FaceCount;
 
 		ConstantData->Vertices = TArray<FVector3f>(Vertex.GetData(), Vertex.Num());
 		ConstantData->BoneMap = TArray<int32>(BoneMap.GetData(), BoneMap.Num());
@@ -1356,19 +1370,53 @@ void UGeometryCollectionComponent::InitConstantData(FGeometryCollectionConstantD
 		const TManagedArray<int32>& MaterialID = Collection->MaterialID;
 
 		const TManagedArray<bool>& Visible = GetVisibleArray();  // Use copy on write attribute. The rest collection visible array can be overriden for the convenience of debug drawing the collision volumes
+		
+#if WITH_EDITOR
+		// We will override visibility with the Hide array (if available).
+		TArray<bool> VisibleOverride;
+		VisibleOverride.Init(true, Visible.Num());
+		bool bUsingHideArray = false;
+
+		if (Collection->HasAttribute("Hide", FGeometryCollection::TransformGroup))
+		{
+			bUsingHideArray = true;
+
+			const TManagedArray<bool>& Hide = Collection->GetAttribute<bool>("Hide", FGeometryCollection::TransformGroup);
+			for (int32 GeomIdx = 0; GeomIdx < NumGeom; ++GeomIdx)
+			{
+				if (Hide[TransformIndex[GeomIdx]])
+				{
+					// (Temporarily) hide faces of this hidden geometry
+					for (int32 FaceIdxOffset = 0; FaceIdxOffset < FaceCount[GeomIdx]; ++FaceIdxOffset)
+					{
+						VisibleOverride[FaceStart[GeomIdx]+FaceIdxOffset] = false;
+					}
+				}
+			}
+		}
+#endif
+		
 		const TManagedArray<int32>& MaterialIndex = Collection->MaterialIndex;
 
 		const int32 NumFaceGroupEntries = Collection->NumElements(FGeometryCollection::FacesGroup);
 
 		for (int FaceIndex = 0; FaceIndex < NumFaceGroupEntries; ++FaceIndex)
 		{
+#if WITH_EDITOR
+			NumIndices += static_cast<int>(VisibleOverride[FaceIndex]);
+#else
 			NumIndices += static_cast<int>(Visible[FaceIndex]);
+#endif
 		}
 
 		ConstantData->Indices.AddUninitialized(NumIndices);
 		for (int IndexIdx = 0, cdx = 0; IndexIdx < NumFaceGroupEntries; ++IndexIdx)
 		{
+#if WITH_EDITOR
+			if (VisibleOverride[MaterialIndex[IndexIdx]])
+#else
 			if (Visible[MaterialIndex[IndexIdx]])
+#endif
 			{
 				ConstantData->Indices[cdx++] = Indices[MaterialIndex[IndexIdx]];
 			}
@@ -1384,7 +1432,11 @@ void UGeometryCollectionComponent::InitConstantData(FGeometryCollectionConstantD
 
 			for (int32 TriangleIndex = 0; TriangleIndex < Sections[SectionIndex].FirstIndex / 3; TriangleIndex++)
 			{
+#if WITH_EDITOR
+				if (!VisibleOverride[MaterialIndex[TriangleIndex]])
+#else
 				if (!Visible[MaterialIndex[TriangleIndex]])
+#endif
 				{
 					Section.FirstIndex -= 3;
 				}
@@ -1392,7 +1444,12 @@ void UGeometryCollectionComponent::InitConstantData(FGeometryCollectionConstantD
 
 			for (int32 TriangleIndex = 0; TriangleIndex < Sections[SectionIndex].NumTriangles; TriangleIndex++)
 			{
-				if (!Visible[MaterialIndex[Sections[SectionIndex].FirstIndex / 3 + TriangleIndex]])
+				int32 FaceIdx = MaterialIndex[Sections[SectionIndex].FirstIndex / 3 + TriangleIndex];
+#if WITH_EDITOR
+				if (!VisibleOverride[FaceIdx])
+#else
+				if (!Visible[FaceIdx])
+#endif
 				{
 					Section.NumTriangles--;
 				}
@@ -1407,9 +1464,7 @@ void UGeometryCollectionComponent::InitConstantData(FGeometryCollectionConstantD
 
 		// store the index buffer and render sections for the base unfractured mesh
 		const TManagedArray<int32>& TransformToGeometryIndex = Collection->TransformToGeometryIndex;
-		const TManagedArray<int32>& FaceStart = Collection->FaceStart;
-		const TManagedArray<int32>& FaceCount = Collection->FaceCount;
-
+		
 		const int32 NumFaces = Collection->NumElements(FGeometryCollection::FacesGroup);
 		TArray<FIntVector> BaseMeshIndices;
 		TArray<int32> BaseMeshOriginalFaceIndices;
@@ -1422,7 +1477,11 @@ void UGeometryCollectionComponent::InitConstantData(FGeometryCollectionConstantD
 		for (int FaceIndex = 0; FaceIndex < NumFaces; ++FaceIndex)
 		{
 			// only add visible external faces.  MaterialID that is even is an external material
+#if WITH_EDITOR
+			if (VisibleOverride[FaceIndex] && (MaterialID[FaceIndex] % 2 == 0 || bUsingHideArray))
+#else
 			if (Visible[FaceIndex] && MaterialID[FaceIndex] % 2 == 0)
+#endif
 			{
 				BaseMeshIndices.Add(Indices[FaceIndex]);
 				BaseMeshOriginalFaceIndices.Add(FaceIndex);
diff --git a/Engine/Source/Runtime/Experimental/GeometryCollectionEngine/Private/GeometryCollection/GeometryCollectionSceneProxy.cpp b/Engine/Source/Runtime/Experimental/GeometryCollectionEngine/Private/GeometryCollection/GeometryCollectionSceneProxy.cpp
--- a/Engine/Source/Runtime/Experimental/GeometryCollectionEngine/Private/GeometryCollection/GeometryCollectionSceneProxy.cpp
+++ b/Engine/Source/Runtime/Experimental/GeometryCollectionEngine/Private/GeometryCollection/GeometryCollectionSceneProxy.cpp
@@ -91,6 +91,7 @@ FGeometryCollectionSceneProxy::FGeometryCollectionSceneProxy(UGeometryCollection
 	, ConstantData(nullptr)
 	, bShowBoneColors(Component->GetShowBoneColors())
 	, bEnableBoneSelection(Component->GetEnableBoneSelection())
+	, bSuppressSelectionMaterial(Component->GetSuppressSelectionMaterial())
 	, BoneSelectionMaterialID(Component->GetBoneSelectedMaterialID())
 	, bUseFullPrecisionUVs(Component->GetRestCollection()->bUseFullPrecisionUVs)
 	, TransformVertexBuffersContainsOriginalMesh(false)
@@ -799,7 +800,7 @@ void FGeometryCollectionSceneProxy::GetDynamicMeshElements(const TArray<const FS
 			// bone selection is already contained in the rendered colors
 			// #note: This renders the geometry again but with the bone selection material.  Ideally we'd have one render pass and one
 			// material.
-			if (bShowBoneColors || bEnableBoneSelection)
+			if ((bShowBoneColors || bEnableBoneSelection) && !bSuppressSelectionMaterial)
 			{
 				FMaterialRenderProxy* MaterialRenderProxy = Materials[BoneSelectionMaterialID]->GetRenderProxy();
 
diff --git a/Engine/Source/Runtime/Experimental/GeometryCollectionEngine/Private/GeometryCollection/GeometryCollectionSceneProxy.h b/Engine/Source/Runtime/Experimental/GeometryCollectionEngine/Private/GeometryCollection/GeometryCollectionSceneProxy.h
--- a/Engine/Source/Runtime/Experimental/GeometryCollectionEngine/Private/GeometryCollection/GeometryCollectionSceneProxy.h
+++ b/Engine/Source/Runtime/Experimental/GeometryCollectionEngine/Private/GeometryCollection/GeometryCollectionSceneProxy.h
@@ -261,6 +261,7 @@ class FGeometryCollectionSceneProxy final : public FPrimitiveSceneProxy
 
 	bool bShowBoneColors;
 	bool bEnableBoneSelection;
+	bool bSuppressSelectionMaterial;
 	int BoneSelectionMaterialID;
 
 	bool bUseFullPrecisionUVs = false;
diff --git a/Engine/Source/Runtime/Experimental/GeometryCollectionEngine/Public/GeometryCollection/GeometryCollectionComponent.h b/Engine/Source/Runtime/Experimental/GeometryCollectionEngine/Public/GeometryCollection/GeometryCollectionComponent.h
--- a/Engine/Source/Runtime/Experimental/GeometryCollectionEngine/Public/GeometryCollection/GeometryCollectionComponent.h
+++ b/Engine/Source/Runtime/Experimental/GeometryCollectionEngine/Public/GeometryCollection/GeometryCollectionComponent.h
@@ -570,6 +570,8 @@ public:
 	bool GetShowBoneColors() const { return bShowBoneColors; }
 	bool GetEnableBoneSelection() const { return bEnableBoneSelection; }
 	
+	bool GetSuppressSelectionMaterial() const { return RestCollection->GetGeometryCollection()->HasAttribute("Hide", FGeometryCollection::TransformGroup); }
+	
 	FORCEINLINE const int GetBoneSelectedMaterialID() const { return RestCollection->GetBoneSelectedMaterialIndex(); }
 	
 #if WITH_EDITORONLY_DATA
