GeometryCollection
-- Selecting a cluster highlights all child bones

diff --git a/Engine/Plugins/Experimental/ChaosEditor/Source/FractureEditor/Private/FractureSelectionTools.cpp b/Engine/Plugins/Experimental/ChaosEditor/Source/FractureEditor/Private/FractureSelectionTools.cpp
--- a/Engine/Plugins/Experimental/ChaosEditor/Source/FractureEditor/Private/FractureSelectionTools.cpp
+++ b/Engine/Plugins/Experimental/ChaosEditor/Source/FractureEditor/Private/FractureSelectionTools.cpp
@@ -34,10 +34,15 @@ void FFractureSelectionTools::ToggleSelectedBones(UGeometryCollectionComponent*
 					TArray<int32> Highlighted;
 					FGeometryCollectionClusteringUtility::ContextBasedClusterSelection(GeometryCollection, EditBoneColor.GetViewLevel(), Selected, RevisedSelected, Highlighted);
 					EditBoneColor.SetSelectedBones(RevisedSelected);
-					EditBoneColor.SetHighlightedBones(Highlighted);
-
-					// @todo FractureEd:  Replace with delegate for custom fracture tree view
-					//SceneOutliner::FSceneOutlinerDelegates::Get().OnComponentSelectionChanged.Broadcast(GeometryCollectionComponent);
+					
+					const TArray<int32>& SelectedBones = EditBoneColor.GetSelectedBones();
+					TArray<int32> HighlightBones;
+					for (int32 SelectedBone : SelectedBones)
+					{
+						FGeometryCollectionClusteringUtility::RecursiveAddAllChildren(GeometryCollectionPtr->Children, SelectedBone, HighlightBones);
+					}
+					
+					EditBoneColor.SetHighlightedBones(HighlightBones);
 				}
 			}
 
diff --git a/Engine/Source/Runtime/Experimental/GeometryCollectionEngine/Private/GeometryCollection/GeometryCollectionComponent.cpp b/Engine/Source/Runtime/Experimental/GeometryCollectionEngine/Private/GeometryCollection/GeometryCollectionComponent.cpp
--- a/Engine/Source/Runtime/Experimental/GeometryCollectionEngine/Private/GeometryCollection/GeometryCollectionComponent.cpp
+++ b/Engine/Source/Runtime/Experimental/GeometryCollectionEngine/Private/GeometryCollection/GeometryCollectionComponent.cpp
@@ -410,7 +410,7 @@ FBoxSphereBounds UGeometryCollectionComponent::CalcBounds(const FTransform& Loca
 		const TManagedArray<FTransform>& Transforms = GetTransformArray();
 
 		const int32 NumBoxes = BoundingBoxes.Num();
-
+	
 		int32 NumElements = HackGeometryCollectionPtr->NumElements(FGeometryCollection::TransformGroup);
 		if (RestCollection->EnableNanite && HackGeometryCollectionPtr->HasAttribute("BoundingBox", FGeometryCollection::TransformGroup) && NumElements)
 		{
@@ -2408,7 +2408,12 @@ void FScopedColorEdit::SelectBones(GeometryCollection::ESelectionMode SelectionM
 		}
 
 		const TArray<int32>& SelectedBones = GetSelectedBones();
-		SetHighlightedBones(SelectedBones);
+		TArray<int32> HighlightBones;
+		for (int32 SelectedBone: SelectedBones)
+		{ 
+			FGeometryCollectionClusteringUtility::RecursiveAddAllChildren(GeometryCollectionPtr->Children, SelectedBone, HighlightBones);
+		}
+		SetHighlightedBones(HighlightBones);
 	}
 }
 
