Added LightWeightDisabled which makes it cheap to disable particles from participating in views

diff --git a/Engine/Source/Programs/HeadlessChaos/Private/HeadlessChaosTestParticleHandle.cpp b/Engine/Source/Programs/HeadlessChaos/Private/HeadlessChaosTestParticleHandle.cpp
--- a/Engine/Source/Programs/HeadlessChaos/Private/HeadlessChaosTestParticleHandle.cpp
+++ b/Engine/Source/Programs/HeadlessChaos/Private/HeadlessChaosTestParticleHandle.cpp
@@ -19,6 +19,22 @@ namespace ChaosTest {
 		Five->AddParticles(5);
 		auto Two = MakeUnique<FGeometryParticles>();
 		Two->AddParticles(2);
+
+		TArray<TUniquePtr<FGeometryParticleHandle>> HandleStorage;
+		auto CreateHandlesHelper = [&HandleStorage](auto& SOA)
+		{
+			//Create a handle with just the bare minimum that we need so that handle iterator uses real handles that are linked with SOA
+			for(int32 ParticleIdx = 0; (unsigned)ParticleIdx < SOA->Size(); ++ParticleIdx)
+			{
+				TUniquePtr<FGeometryParticleHandle> NewParticleHandle = FGeometryParticleHandle::CreateParticleHandle(MakeSerializable(SOA), ParticleIdx, HandleStorage.Num());
+				SOA->SetHandle(ParticleIdx, NewParticleHandle.Get());
+				HandleStorage.Add(MoveTemp(NewParticleHandle));
+			}
+		};
+
+		CreateHandlesHelper(Five);
+		CreateHandlesHelper(Two);
+
 		//empty soa in the start
 		{
 			TArray<FGeometryParticleHandle*> Handles;
@@ -30,13 +46,31 @@ namespace ChaosTest {
 			}
 			EXPECT_EQ(Handles.Num(), 7);
 
+			//disable first middle and last
+			Five->LightWeightDisabled(0) = true;
+			Five->LightWeightDisabled(2) = true;
+			Five->LightWeightDisabled(4) = true;
+			Two->LightWeightDisabled(1) = true;
+
+			int32 NonDisabledCount = 0;
+			for (auto& Particle : View)
+			{
+				++NonDisabledCount;
+			}
+			EXPECT_EQ(NonDisabledCount, 3);
+
 			THandleView<FGeometryParticles> HandleView = MakeHandleView(Handles);
 			int32 Count = 0;
 			for (auto& Handle : HandleView)
 			{
 				++Count;
 			}
-			EXPECT_EQ(Count, 7);
+			EXPECT_EQ(Count, 3);	//3 because 4 are disabled
+
+			Five->LightWeightDisabled(0) = false;
+			Five->LightWeightDisabled(2) = false;
+			Five->LightWeightDisabled(4) = false;
+			Two->LightWeightDisabled(1) = false;
 		}
 
 		//empty soa in the middle
@@ -70,13 +104,31 @@ namespace ChaosTest {
 			}
 			EXPECT_EQ(Handles.Num(), 7);
 
+			//disable first middle and last
+			Five->LightWeightDisabled(0) = true;
+			Five->LightWeightDisabled(2) = true;
+			Five->LightWeightDisabled(4) = true;
+			Two->LightWeightDisabled(1) = true;
+
+			int32 NonDisabledCount = 0;
+			for (auto& Particle : View)
+			{
+				++NonDisabledCount;
+			}
+			EXPECT_EQ(NonDisabledCount, 3);
+
 			THandleView<FGeometryParticles> HandleView = MakeHandleView(Handles);
 			int32 Count = 0;
 			for (auto& Handle : HandleView)
 			{
 				++Count;
 			}
-			EXPECT_EQ(Count, 7);
+			EXPECT_EQ(Count, 3);	//3 because 4 are disabled
+
+			Five->LightWeightDisabled(0) = false;
+			Five->LightWeightDisabled(2) = false;
+			Five->LightWeightDisabled(4) = false;
+			Two->LightWeightDisabled(1) = false;
 		}
 
 		//parallel for
@@ -131,6 +183,73 @@ namespace ChaosTest {
 					EXPECT_TRUE(Val);
 				}
 			}
+
+
+			//disable first middle and last
+			Five->LightWeightDisabled(0) = true;
+			Five->LightWeightDisabled(2) = true;
+			Five->LightWeightDisabled(4) = true;
+			Two->LightWeightDisabled(1) = true;
+
+			{
+				TArray<bool> AuxArray;
+				AuxArray.SetNumZeroed(View.Num());
+				bool DoubleWrite = false;
+				View.ParallelFor([&AuxArray, &DoubleWrite](const auto& Particle, int32 Idx)
+					{
+						if (AuxArray[Idx])
+						{
+							DoubleWrite = true;
+						}
+						AuxArray[Idx] = true;
+					});
+
+				EXPECT_FALSE(DoubleWrite);
+
+				for (int32 Idx = 0; Idx < AuxArray.Num(); ++Idx)
+				{
+					if(Idx == 0 || Idx == 2 || Idx == 4 || Idx == 6)
+					{
+						//didn't write because disabled
+						EXPECT_FALSE(AuxArray[Idx]);
+					}
+					else
+					{
+						EXPECT_TRUE(AuxArray[Idx]);
+					}
+					
+				}
+			}
+
+			{
+				TArray<bool> AuxArray;
+				AuxArray.SetNumZeroed(View.Num());
+				bool DoubleWrite = false;
+				HandleView.ParallelFor([&AuxArray, &DoubleWrite](const auto& Particle, int32 Idx)
+					{
+						if (AuxArray[Idx])
+						{
+							DoubleWrite = true;
+						}
+						AuxArray[Idx] = true;
+					});
+
+				EXPECT_FALSE(DoubleWrite);
+
+				for (int32 Idx = 0; Idx < AuxArray.Num(); ++Idx)
+				{
+					if (Idx == 0 || Idx == 2 || Idx == 4 || Idx == 6)
+					{
+						//didn't write because disabled
+						EXPECT_FALSE(AuxArray[Idx]);
+					}
+					else
+					{
+						EXPECT_TRUE(AuxArray[Idx]);
+					}
+
+				}
+			}
 		}
 	}
 
diff --git a/Engine/Source/Runtime/Experimental/Chaos/Public/Chaos/GeometryParticles.h b/Engine/Source/Runtime/Experimental/Chaos/Public/Chaos/GeometryParticles.h
--- a/Engine/Source/Runtime/Experimental/Chaos/Public/Chaos/GeometryParticles.h
+++ b/Engine/Source/Runtime/Experimental/Chaos/Public/Chaos/GeometryParticles.h
@@ -366,6 +366,7 @@ namespace Chaos
 			TArrayCollection::AddArray(&MParticleConstraints);
 			TArrayCollection::AddArray(&MResimType);
 			TArrayCollection::AddArray(&MEnabledDuringResim);
+			TArrayCollection::AddArray(&MLightWeightDisabled);
 
 
 #if CHAOS_DEBUG_NAME
@@ -404,6 +405,7 @@ namespace Chaos
 			, MParticleConstraints(MoveTemp(Other.MParticleConstraints))
 			, MResimType(MoveTemp(Other.MResimType))
 			, MEnabledDuringResim(MoveTemp(Other.MEnabledDuringResim))
+			, MLightWeightDisabled(MoveTemp(Other.MLightWeightDisabled))
 
 
 #if CHAOS_DETERMINISTIC
@@ -444,6 +446,7 @@ namespace Chaos
 
 			TArrayCollection::AddArray(&MResimType);
 			TArrayCollection::AddArray(&MEnabledDuringResim);
+			TArrayCollection::AddArray(&MLightWeightDisabled);
 		}
 
 		static constexpr bool IsRigidBodySim() { return SimType == EGeometryParticlesSimType::RigidBodySim; }
@@ -485,6 +488,7 @@ namespace Chaos
 			
 			TArrayCollection::AddArray(&MResimType);
 			TArrayCollection::AddArray(&MEnabledDuringResim);
+			TArrayCollection::AddArray(&MLightWeightDisabled);
 		}
 
 		CHAOS_API virtual ~TGeometryParticlesImp()
@@ -666,6 +670,9 @@ namespace Chaos
 		FORCEINLINE bool EnabledDuringResim(const int32 Index) const { return MEnabledDuringResim[Index]; }
 		FORCEINLINE bool& EnabledDuringResim(const int32 Index) { return MEnabledDuringResim[Index]; }
 
+		FORCEINLINE bool LightWeightDisabled(const int32 Index) const { return MLightWeightDisabled[Index]; }
+		FORCEINLINE bool& LightWeightDisabled(const int32 Index) { return MLightWeightDisabled[Index]; }
+
 private:
 		friend THandleType;
 		CHAOS_API void ResetWeakParticleHandle(const int32 Index)
@@ -804,6 +811,7 @@ public:
 		TArrayCollectionArray<FConstraintHandleArray> MParticleConstraints;
 		TArrayCollectionArray<EResimType> MResimType;
 		TArrayCollectionArray<bool> MEnabledDuringResim;
+		TArrayCollectionArray<bool> MLightWeightDisabled;
 
 		void UpdateShapesArray(const int32 Index)
 		{
diff --git a/Engine/Source/Runtime/Experimental/Chaos/Public/Chaos/PBDRigidsEvolution.h b/Engine/Source/Runtime/Experimental/Chaos/Public/Chaos/PBDRigidsEvolution.h
--- a/Engine/Source/Runtime/Experimental/Chaos/Public/Chaos/PBDRigidsEvolution.h
+++ b/Engine/Source/Runtime/Experimental/Chaos/Public/Chaos/PBDRigidsEvolution.h
@@ -196,6 +196,8 @@ struct FSpatialAccelerationCacheHandle
 		return Cache->Bounds(EntryIdx);
 	}
 
+	bool LightWeightDisabled() const { return false; }
+
 	union
 	{
 		FSpatialAccelerationCache* GeometryParticles;	//using same name as particles SOA for template reuse, should probably rethink this
diff --git a/Engine/Source/Runtime/Experimental/Chaos/Public/Chaos/ParticleDirtyFlags.h b/Engine/Source/Runtime/Experimental/Chaos/Public/Chaos/ParticleDirtyFlags.h
--- a/Engine/Source/Runtime/Experimental/Chaos/Public/Chaos/ParticleDirtyFlags.h
+++ b/Engine/Source/Runtime/Experimental/Chaos/Public/Chaos/ParticleDirtyFlags.h
@@ -430,6 +430,7 @@ public:
 		SetUniqueIdx(Other.UniqueIdx());
 		SetSpatialIdx(Other.SpatialIdx());
 		SetResimType(Other.ResimType());
+		SetEnabledDuringResim(Other.EnabledDuringResim());
 #if CHAOS_DEBUG_NAME
 		SetDebugName(Other.DebugName());
 #endif
@@ -441,7 +442,8 @@ public:
 		return Geometry() == Other.Geometry()
 			&& UniqueIdx() == Other.UniqueIdx()
 			&& SpatialIdx() == Other.SpatialIdx()
-			&& ResimType() == Other.ResimType();
+			&& ResimType() == Other.ResimType()
+			&& EnabledDuringResim() == Other.EnabledDuringResim();
 	}
 
 	bool operator==(const FParticleNonFrequentData& Other) const
@@ -470,6 +472,9 @@ public:
 		MResimType = InType;
 	}
 
+	bool EnabledDuringResim() const { return MEnabledDuringResim; }
+	void SetEnabledDuringResim(bool bEnabledDuringResim) { MEnabledDuringResim = bEnabledDuringResim; }
+
 #if CHAOS_DEBUG_NAME
 	const TSharedPtr<FString, ESPMode::ThreadSafe>& DebugName() const { return MDebugName; }
 	void SetDebugName(const TSharedPtr<FString, ESPMode::ThreadSafe>& InName) { MDebugName = InName; }
@@ -479,6 +484,7 @@ private:
 	FUniqueIdx MUniqueIdx;
 	FSpatialAccelerationIdx MSpatialIdx;
 	EResimType MResimType;
+	bool MEnabledDuringResim;
 #if CHAOS_DEBUG_NAME
 	TSharedPtr<FString, ESPMode::ThreadSafe> MDebugName;
 #endif
diff --git a/Engine/Source/Runtime/Experimental/Chaos/Public/Chaos/ParticleHandle.h b/Engine/Source/Runtime/Experimental/Chaos/Public/Chaos/ParticleHandle.h
--- a/Engine/Source/Runtime/Experimental/Chaos/Public/Chaos/ParticleHandle.h
+++ b/Engine/Source/Runtime/Experimental/Chaos/Public/Chaos/ParticleHandle.h
@@ -65,6 +65,7 @@ void GeometryParticleDefaultConstruct(FConcrete& Concrete, const FGeometryPartic
 	Concrete.SetR(TRotation<T, d>::Identity);
 	Concrete.SetSpatialIdx(FSpatialAccelerationIdx{ 0,0 });
 	Concrete.SetResimType(EResimType::FullResim);
+	Concrete.SetEnabledDuringResim(true);
 }
 
 template <typename T, int d, typename FConcrete>
@@ -345,7 +346,7 @@ protected:
 		//TODO: patch from SOA
 		GeometryParticleDefaultConstruct<T, d>(*this, Params);
 		SetHasBounds(false);
-		SetEnabledDuringResim(false);
+		SetLightWeightDisabled(false);
 	}
 
 	template <typename TParticlesType, typename TParams>
@@ -514,7 +515,8 @@ public:
 	bool EnabledDuringResim() const { return GeometryParticles->EnabledDuringResim(ParticleIdx); }
 	void SetEnabledDuringResim(bool bEnabledDuringResim) { GeometryParticles->EnabledDuringResim(ParticleIdx) = bEnabledDuringResim; }
 
-
+	bool LightWeightDisabled() const { return GeometryParticles->LightWeightDisabled(ParticleIdx); }
+	void SetLightWeightDisabled(bool bLightWeightDisabled) { GeometryParticles->LightWeightDisabled(ParticleIdx) = bLightWeightDisabled; }
 
 #if CHAOS_DETERMINISTIC
 	FParticleID ParticleID() const { return GeometryParticles->ParticleID(ParticleIdx); }
@@ -1844,6 +1846,16 @@ public:
 		return MNonFrequentData.Read().ResimType();
 	}
 
+	void SetEnabledDuringResim(bool bEnabledDuringResim)
+	{
+		MNonFrequentData.Modify(true, MDirtyFlags, Proxy, [bEnabledDuringResim](auto& Data) { Data.SetEnabledDuringResim(bEnabledDuringResim); });
+	}
+
+	bool EnabledDuringResim() const
+	{
+		return MNonFrequentData.Read().EnabledDuringResim();
+	}
+
 
 	void SetNonFrequentData(const FParticleNonFrequentData& InData)
 	{
diff --git a/Engine/Source/Runtime/Experimental/Chaos/Public/Chaos/ParticleIterator.h b/Engine/Source/Runtime/Experimental/Chaos/Public/Chaos/ParticleIterator.h
--- a/Engine/Source/Runtime/Experimental/Chaos/Public/Chaos/ParticleIterator.h
+++ b/Engine/Source/Runtime/Experimental/Chaos/Public/Chaos/ParticleIterator.h
@@ -103,47 +103,61 @@ class TConstHandleIterator
 {
 public:
 	using THandle = typename TSOA::THandleType;
-	TConstHandleIterator(const TArray<THandle*>& InHandles, int32 StartOffset = 0)
-		: Handles(InHandles)
+
+	TConstHandleIterator()
+		: Handles(nullptr)
 		, CurIdx(0)
 	{
-		Advance(StartOffset);
 	}
 
-	operator bool() const { return CurIdx < Handles.Num(); }
-	TConstHandleIterator<TSOA>& operator++()
+	TConstHandleIterator(const TArray<THandle*>& InHandles)
+		: Handles(&InHandles)
+		, CurIdx(0)
 	{
-		++CurIdx;
-		return *this;
+		while(CurrentIsLightWeightDisabled())
+		{
+			operator++();
+		}
 	}
 
-	TConstHandleIterator<TSOA>& Advance(int32 Offset)
+	operator bool() const { return Handles && CurIdx < Handles->Num(); }
+	TConstHandleIterator<TSOA>& operator++()
 	{
-		CurIdx = FMath::Min(CurIdx + Offset, Handles.Num());
+		bool bFindNextHandle = true;
+		while(bFindNextHandle)
+		{
+			++CurIdx;
+			bFindNextHandle = CurrentIsLightWeightDisabled();
+		}
 		return *this;
 	}
 
 	const THandle& operator*() const
 	{
-		return static_cast<const THandle&>(*Handles[CurIdx]);
+		return static_cast<const THandle&>(*(*Handles)[CurIdx]);
 	}
 
 	const THandle* operator->() const
 	{
-		return static_cast<const THandle*>(Handles[CurIdx]);
+		return static_cast<const THandle*>((*Handles)[CurIdx]);
 	}
 
 	int32 ComputeSize() const
 	{
-		return Handles.Num();
+		return Handles->Num();
 	}
 
 protected:
 	template <typename TSOA2>
 	friend class TConstHandleView;
 
-	const TArray<THandle*>& Handles;
+	const TArray<THandle*>* Handles;
 	int32 CurIdx;
+
+	bool CurrentIsLightWeightDisabled() const
+	{
+		return operator bool() && operator*().LightWeightDisabled();
+	}
 };
 
 template <typename TSOA>
@@ -155,19 +169,24 @@ public:
 	using Base::Handles;
 	using Base::CurIdx;
 
-	THandleIterator(const TArray<THandle*>& InHandles, int32 StartOffset = 0)
-		: Base(InHandles, StartOffset)
+	THandleIterator()
+		: Base()
+	{
+	}
+
+	THandleIterator(const TArray<THandle*>& InHandles)
+		: Base(InHandles)
 	{
 	}
 
 	THandle& operator*() const
 	{
-		return static_cast<THandle&>(*Handles[CurIdx]);
+		return static_cast<THandle&>(*(*Handles)[CurIdx]);
 	}
 
 	THandle* operator->() const
 	{
-		return static_cast<THandle*>(Handles[CurIdx]);
+		return static_cast<THandle*>((*Handles)[CurIdx]);
 	}
 
 	template <typename TSOA2>
@@ -208,7 +227,7 @@ public:
 	TConstHandleIterator<TSOA> Begin() const { return MakeConstHandleIterator(Handles); }
 	TConstHandleIterator<TSOA> begin() const { return Begin(); }
 
-	TConstHandleIterator<TSOA> End() const { return TConstHandleIterator<TSOA>(Handles, Num()); }
+	TConstHandleIterator<TSOA> End() const { return TConstHandleIterator<TSOA>(); }
 	TConstHandleIterator<TSOA> end() const { return End(); }
 
 	template <typename TParticleView, typename Lambda>
@@ -247,7 +266,7 @@ public:
 	THandleIterator<TSOA> Begin() const { return MakeHandleIterator(Handles); }
 	THandleIterator<TSOA> begin() const { return Begin(); }
 
-	THandleIterator<TSOA> End() const { return THandleIterator<TSOA>(Handles, Num()); }
+	THandleIterator<TSOA> End() const { return THandleIterator<TSOA>(); }
 	THandleIterator<TSOA> end() const { return End(); }
 
 	template <typename Lambda>
@@ -321,6 +340,10 @@ public:
 #endif
 	{
 		SeekNonEmptySOA();
+		while(CurrentIsLightWeightDisabled())
+		{
+			operator++();
+		}
 	}
 
 	TConstParticleIterator(const TConstParticleIterator& Rhs) = default;
@@ -333,32 +356,39 @@ public:
 	TConstParticleIterator<TSOA>& operator++()
 	{
 		RangedForValidation();
-		if (CurHandlesArray == nullptr)
+		bool bFindNextEnabledParticle = true;
+		while (bFindNextEnabledParticle)
 		{
-			//SOA is packed efficiently for iteration
-			++TransientHandle.ParticleIdx;
-			if (TransientHandle.ParticleIdx >= static_cast<int32>(CurSOASize))
+			if (CurHandlesArray == nullptr)
 			{
-				IncSOAIdx();
-			}
-		}
-		else
-		{
-			++CurHandleIdx;
-			if (CurHandleIdx < CurHandlesArray->Num())
-			{
-				// Reconstruct the TransientHandle so that it has a chance to update
-				// other data members, like the particle type in the case of geometry 
-				// particles.
-				THandle* HandlePtr = (*CurHandlesArray)[CurHandleIdx];
-				THandleBase Handle(HandlePtr->GeometryParticles, HandlePtr->ParticleIdx);
-				TransientHandle = static_cast<TTransientHandle&>(Handle);
+				//SOA is packed efficiently for iteration
+				++TransientHandle.ParticleIdx;
+				if (TransientHandle.ParticleIdx >= static_cast<int32>(CurSOASize))
+				{
+					IncSOAIdx();
+				}
 			}
 			else
 			{
-				IncSOAIdx();
+				++CurHandleIdx;
+				if (CurHandleIdx < CurHandlesArray->Num())
+				{
+					// Reconstruct the TransientHandle so that it has a chance to update
+					// other data members, like the particle type in the case of geometry 
+					// particles.
+					THandle* HandlePtr = (*CurHandlesArray)[CurHandleIdx];
+					THandleBase Handle(HandlePtr->GeometryParticles, HandlePtr->ParticleIdx);
+					TransientHandle = static_cast<TTransientHandle&>(Handle);
+				}
+				else
+				{
+					IncSOAIdx();
+				}
 			}
+
+			bFindNextEnabledParticle = CurrentIsLightWeightDisabled();
 		}
+
 		return *this;
 	}
 	
@@ -376,6 +406,11 @@ public:
 
 protected:
 
+	bool CurrentIsLightWeightDisabled() const
+	{
+		return operator bool() && (operator*()).LightWeightDisabled();
+	}
+
 	void MoveToEnd()
 	{
 		SOAIdx = 0;
@@ -611,7 +646,7 @@ TParticleView<TSOA> MakeParticleView(TSOA* SOA)
 // The non-parallel implementation of iteration should not deviate in behavior from
 // this parallel implementation.  They must be kept in sync.
 template <typename TParticleView, typename Lambda>
-void ParticleViewParallelForImp(const TParticleView& Particles, const Lambda& Func)
+void ParticleViewParallelForImp(const TParticleView& Particles, const Lambda& UserFunc)
 {
 	SCOPE_CYCLE_COUNTER(STAT_ParticleViewParallelForImp);
 
@@ -620,6 +655,14 @@ void ParticleViewParallelForImp(const TParticleView& Particles, const Lambda& Fu
 	using THandleBase = typename THandle::THandleBase;
 	using TTransientHandle = typename THandle::TTransientHandle;
 
+	auto Func = [&UserFunc](auto& Handle, const int32 Idx)
+	{
+		if (!Handle.LightWeightDisabled())
+		{
+			UserFunc(Handle, Idx);
+		}
+	};
+
 	// Loop over every SOA in this view, skipping empty ones
 	int32 ParticleIdxOff = 0;
 	for (int32 ViewIndex = 0; ViewIndex < Particles.SOAViews.Num(); ++ViewIndex)
@@ -672,7 +715,10 @@ void HandleViewParallelForImp(const THandleView& HandleView, const Lambda& Func)
 	const int32 HandleCount = HandleView.Handles.Num();
 	::ParallelFor(HandleCount, [&HandleView, &Func](const int32 Index)
 	{
-		Func(static_cast<THandle&>(*HandleView.Handles[Index]), Index);
+		if (!HandleView.Handles[Index]->LightWeightDisabled())
+		{
+			Func(static_cast<THandle&>(*HandleView.Handles[Index]), Index);
+		}
 	});
 }
 
