Clean up joint constraint to use fewer typedefs and be inside chaos namespace. Also inherit from base proxy to skip mostly empty API

diff --git a/Engine/Source/Runtime/Experimental/Chaos/Private/PhysicsProxy/JointConstraintProxy.cpp b/Engine/Source/Runtime/Experimental/Chaos/Private/PhysicsProxy/JointConstraintProxy.cpp
--- a/Engine/Source/Runtime/Experimental/Chaos/Private/PhysicsProxy/JointConstraintProxy.cpp
+++ b/Engine/Source/Runtime/Experimental/Chaos/Private/PhysicsProxy/JointConstraintProxy.cpp
@@ -15,20 +15,20 @@
 #include "PhysicsSolver.h"
 #include "Chaos/PullPhysicsDataImp.h"
 
+namespace Chaos
+{
 
-FJointConstraintPhysicsProxy::FJointConstraintPhysicsProxy(Chaos::FJointConstraint* InConstraint, FConstraintHandle* InHandle, UObject* InOwner)
-	: Base(InOwner)
+FJointConstraintPhysicsProxy::FJointConstraintPhysicsProxy(FJointConstraint* InConstraint, FPBDJointConstraintHandle* InHandle, UObject* InOwner)
+	: Base(EPhysicsProxyType::JointConstraintType, InOwner)
 	, Constraint(InConstraint) // This proxy assumes ownership of the Constraint, and will free it during DestroyOnPhysicsThread
 	, Handle(InHandle)
-	, bInitialized(false)
 {
 	check(Constraint!=nullptr);
 	Constraint->SetProxy(this);
 	JointSettingsBuffer = Constraint->GetJointSettings();
-
 }
 
-Chaos::TGeometryParticleHandle<Chaos::FReal, 3>*
+FGeometryParticleHandle*
 FJointConstraintPhysicsProxy::GetParticleHandleFromProxy(IPhysicsProxyBase* ProxyBase)
 {
 	if (ProxyBase)
@@ -42,7 +42,7 @@ FJointConstraintPhysicsProxy::GetParticleHandleFromProxy(IPhysicsProxyBase* Prox
 }
 
 /**/
-void FJointConstraintPhysicsProxy::BufferPhysicsResults(Chaos::FDirtyJointConstraintData& Buffer)
+void FJointConstraintPhysicsProxy::BufferPhysicsResults(FDirtyJointConstraintData& Buffer)
 {
 	Buffer.SetProxy(*this);
 	if (Constraint != nullptr && Constraint->IsValid() )
@@ -62,7 +62,7 @@ void FJointConstraintPhysicsProxy::BufferPhysicsResults(Chaos::FDirtyJointConstr
 }
 
 /**/
-bool FJointConstraintPhysicsProxy::PullFromPhysicsState(const Chaos::FDirtyJointConstraintData& Buffer, const int32 SolverSyncTimestamp)
+bool FJointConstraintPhysicsProxy::PullFromPhysicsState(const FDirtyJointConstraintData& Buffer, const int32 SolverSyncTimestamp)
 {
 	if (Constraint != nullptr && Constraint->IsValid())
 	{
@@ -79,7 +79,7 @@ bool FJointConstraintPhysicsProxy::PullFromPhysicsState(const Chaos::FDirtyJoint
 	return true;
 }
 
-void FJointConstraintPhysicsProxy::InitializeOnPhysicsThread(Chaos::FPBDRigidsSolver* InSolver)
+void FJointConstraintPhysicsProxy::InitializeOnPhysicsThread(FPBDRigidsSolver* InSolver)
 {
 	auto& Handles = InSolver->GetParticles().GetParticleHandles();
 	if (Handles.Size() && IsValid())
@@ -87,10 +87,10 @@ void FJointConstraintPhysicsProxy::InitializeOnPhysicsThread(Chaos::FPBDRigidsSo
 		auto& JointConstraints = InSolver->GetJointConstraints();
 		if (Constraint != nullptr)
 		{
-			Chaos::FConstraintBase::FProxyBasePair& BasePairs = Constraint->GetParticleProxies();
+			FConstraintBase::FProxyBasePair& BasePairs = Constraint->GetParticleProxies();
 
-			Chaos::TGeometryParticleHandle<Chaos::FReal, 3>* Handle0 = GetParticleHandleFromProxy(BasePairs[0]);
-			Chaos::TGeometryParticleHandle<Chaos::FReal, 3>* Handle1 = GetParticleHandleFromProxy(BasePairs[1]);
+			FGeometryParticleHandle* Handle0 = GetParticleHandleFromProxy(BasePairs[0]);
+			FGeometryParticleHandle* Handle1 = GetParticleHandleFromProxy(BasePairs[1]);
 			if (Handle0 && Handle1)
 			{
 				Handle = JointConstraints.AddConstraint({ Handle0,Handle1 }, Constraint->GetJointTransforms());
@@ -104,7 +104,7 @@ void FJointConstraintPhysicsProxy::InitializeOnPhysicsThread(Chaos::FPBDRigidsSo
 	}
 }
 
-void FJointConstraintPhysicsProxy::DestroyOnPhysicsThread(Chaos::FPBDRigidsSolver* InSolver)
+void FJointConstraintPhysicsProxy::DestroyOnPhysicsThread(FPBDRigidsSolver* InSolver)
 {
 	if (Handle && Handle->IsValid())
 	{
@@ -121,60 +121,60 @@ void FJointConstraintPhysicsProxy::DestroyOnPhysicsThread(Chaos::FPBDRigidsSolve
 }
 
 
-void FJointConstraintPhysicsProxy::PushStateOnGameThread(Chaos::FPBDRigidsSolver* InSolver)
+void FJointConstraintPhysicsProxy::PushStateOnGameThread(FPBDRigidsSolver* InSolver)
 {
 	if (Constraint != nullptr && Constraint->IsValid())
 	{
 		if (Constraint->IsDirty())
 		{
-			if (Constraint->IsDirty(Chaos::EJointConstraintFlags::JointTransforms))
+			if (Constraint->IsDirty(EJointConstraintFlags::JointTransforms))
 			{
 				JointSettingsBuffer.ConnectorTransforms = Constraint->GetJointTransforms();
-				DirtyFlagsBuffer.MarkDirty(Chaos::EJointConstraintFlags::JointTransforms);
+				DirtyFlagsBuffer.MarkDirty(EJointConstraintFlags::JointTransforms);
 			}
 
-			if (Constraint->IsDirty(Chaos::EJointConstraintFlags::CollisionEnabled))
+			if (Constraint->IsDirty(EJointConstraintFlags::CollisionEnabled))
 			{
 				JointSettingsBuffer.bCollisionEnabled = Constraint->GetCollisionEnabled();
-				DirtyFlagsBuffer.MarkDirty(Chaos::EJointConstraintFlags::CollisionEnabled);
+				DirtyFlagsBuffer.MarkDirty(EJointConstraintFlags::CollisionEnabled);
 			}
 
-			if (Constraint->IsDirty(Chaos::EJointConstraintFlags::Projection))
+			if (Constraint->IsDirty(EJointConstraintFlags::Projection))
 			{
 				JointSettingsBuffer.bProjectionEnabled = Constraint->GetProjectionEnabled();
 				JointSettingsBuffer.LinearProjection = Constraint->GetProjectionLinearAlpha();
 				JointSettingsBuffer.AngularProjection = Constraint->GetProjectionAngularAlpha();
-				DirtyFlagsBuffer.MarkDirty(Chaos::EJointConstraintFlags::Projection);
+				DirtyFlagsBuffer.MarkDirty(EJointConstraintFlags::Projection);
 			}
 
-			if (Constraint->IsDirty(Chaos::EJointConstraintFlags::ParentInvMassScale))
+			if (Constraint->IsDirty(EJointConstraintFlags::ParentInvMassScale))
 			{
 				JointSettingsBuffer.ParentInvMassScale = Constraint->GetParentInvMassScale();
-				DirtyFlagsBuffer.MarkDirty(Chaos::EJointConstraintFlags::ParentInvMassScale);
+				DirtyFlagsBuffer.MarkDirty(EJointConstraintFlags::ParentInvMassScale);
 			}
 
-			if (Constraint->IsDirty(Chaos::EJointConstraintFlags::LinearBreakForce))
+			if (Constraint->IsDirty(EJointConstraintFlags::LinearBreakForce))
 			{
 				JointSettingsBuffer.LinearBreakForce = Constraint->GetLinearBreakForce();
 				JointSettingsBuffer.LinearPlasticityType = Constraint->GetLinearPlasticityType();
 				JointSettingsBuffer.LinearPlasticityLimit = Constraint->GetLinearPlasticityLimit();
-				DirtyFlagsBuffer.MarkDirty(Chaos::EJointConstraintFlags::LinearBreakForce);
+				DirtyFlagsBuffer.MarkDirty(EJointConstraintFlags::LinearBreakForce);
 			}
 
-			if (Constraint->IsDirty(Chaos::EJointConstraintFlags::AngularBreakTorque))
+			if (Constraint->IsDirty(EJointConstraintFlags::AngularBreakTorque))
 			{
 				JointSettingsBuffer.AngularBreakTorque = Constraint->GetAngularBreakTorque();
 				JointSettingsBuffer.AngularPlasticityLimit = Constraint->GetAngularPlasticityLimit();
-				DirtyFlagsBuffer.MarkDirty(Chaos::EJointConstraintFlags::AngularBreakTorque);
+				DirtyFlagsBuffer.MarkDirty(EJointConstraintFlags::AngularBreakTorque);
 			}
 
-			if (Constraint->IsDirty(Chaos::EJointConstraintFlags::UserData))
+			if (Constraint->IsDirty(EJointConstraintFlags::UserData))
 			{
 				JointSettingsBuffer.UserData = Constraint->GetUserData();
-				DirtyFlagsBuffer.MarkDirty(Chaos::EJointConstraintFlags::UserData);
+				DirtyFlagsBuffer.MarkDirty(EJointConstraintFlags::UserData);
 			}
 
-			if (Constraint->IsDirty(Chaos::EJointConstraintFlags::LinearDrive))
+			if (Constraint->IsDirty(EJointConstraintFlags::LinearDrive))
 			{
 				JointSettingsBuffer.bLinearPositionDriveEnabled[0] = Constraint->GetLinearPositionDriveXEnabled();
 				JointSettingsBuffer.bLinearPositionDriveEnabled[1] = Constraint->GetLinearPositionDriveYEnabled();
@@ -190,11 +190,11 @@ void FJointConstraintPhysicsProxy::PushStateOnGameThread(Chaos::FPBDRigidsSolver
 				JointSettingsBuffer.LinearMotionTypes[2] = Constraint->GetLinearMotionTypesZ();
 				JointSettingsBuffer.LinearDriveStiffness = Constraint->GetLinearDriveStiffness();
 				JointSettingsBuffer.LinearDriveDamping = Constraint->GetLinearDriveDamping();
-				DirtyFlagsBuffer.MarkDirty(Chaos::EJointConstraintFlags::LinearDrive);
+				DirtyFlagsBuffer.MarkDirty(EJointConstraintFlags::LinearDrive);
 			}
 
 
-			if (Constraint->IsDirty(Chaos::EJointConstraintFlags::AngularDrive))
+			if (Constraint->IsDirty(EJointConstraintFlags::AngularDrive))
 			{
 				JointSettingsBuffer.bAngularSLerpPositionDriveEnabled = Constraint->GetAngularSLerpPositionDriveEnabled();
 				JointSettingsBuffer.bAngularTwistPositionDriveEnabled = Constraint->GetAngularTwistPositionDriveEnabled();
@@ -210,16 +210,16 @@ void FJointConstraintPhysicsProxy::PushStateOnGameThread(Chaos::FPBDRigidsSolver
 				JointSettingsBuffer.AngularMotionTypes[2] = Constraint->GetAngularMotionTypesZ();
 				JointSettingsBuffer.AngularDriveStiffness = Constraint->GetAngularDriveStiffness();
 				JointSettingsBuffer.AngularDriveDamping = Constraint->GetAngularDriveDamping();
-				DirtyFlagsBuffer.MarkDirty(Chaos::EJointConstraintFlags::AngularDrive);
+				DirtyFlagsBuffer.MarkDirty(EJointConstraintFlags::AngularDrive);
 			}
 
-			if (Constraint->IsDirty(Chaos::EJointConstraintFlags::Stiffness))
+			if (Constraint->IsDirty(EJointConstraintFlags::Stiffness))
 			{
 				JointSettingsBuffer.Stiffness = Constraint->GetStiffness();
-				DirtyFlagsBuffer.MarkDirty(Chaos::EJointConstraintFlags::Stiffness);
+				DirtyFlagsBuffer.MarkDirty(EJointConstraintFlags::Stiffness);
 			}
 
-			if (Constraint->IsDirty(Chaos::EJointConstraintFlags::Limits))
+			if (Constraint->IsDirty(EJointConstraintFlags::Limits))
 			{
 				JointSettingsBuffer.bSoftLinearLimitsEnabled = Constraint->GetSoftLinearLimitsEnabled();
 				JointSettingsBuffer.bSoftTwistLimitsEnabled = Constraint->GetSoftTwistLimitsEnabled();
@@ -241,7 +241,7 @@ void FJointConstraintPhysicsProxy::PushStateOnGameThread(Chaos::FPBDRigidsSolver
 				JointSettingsBuffer.TwistRestitution = Constraint->GetTwistRestitution();
 				JointSettingsBuffer.SwingRestitution = Constraint->GetSwingRestitution();
 
-				DirtyFlagsBuffer.MarkDirty(Chaos::EJointConstraintFlags::Limits);
+				DirtyFlagsBuffer.MarkDirty(EJointConstraintFlags::Limits);
 			}
 
 			Constraint->ClearDirtyFlags();
@@ -250,27 +250,26 @@ void FJointConstraintPhysicsProxy::PushStateOnGameThread(Chaos::FPBDRigidsSolver
 }
 
 
-void FJointConstraintPhysicsProxy::PushStateOnPhysicsThread(Chaos::FPBDRigidsSolver* InSolver)
+void FJointConstraintPhysicsProxy::PushStateOnPhysicsThread(FPBDRigidsSolver* InSolver)
 {
-	typedef typename Chaos::FPBDRigidsSolver::FPBDRigidsEvolution::FCollisionConstraints FCollisionConstraints;
 	if (Handle && Handle->IsValid())
 	{
 		if (DirtyFlagsBuffer.IsDirty())
 		{
-			FConstraintData ConstraintSettings = Handle->GetSettings();
+			FPBDJointSettings ConstraintSettings = Handle->GetSettings();
 
-			if (DirtyFlagsBuffer.IsDirty(Chaos::EJointConstraintFlags::JointTransforms))
+			if (DirtyFlagsBuffer.IsDirty(EJointConstraintFlags::JointTransforms))
 			{
 				ConstraintSettings.ConnectorTransforms = JointSettingsBuffer.ConnectorTransforms;
 			}
 
-			if (DirtyFlagsBuffer.IsDirty(Chaos::EJointConstraintFlags::CollisionEnabled))
+			if (DirtyFlagsBuffer.IsDirty(EJointConstraintFlags::CollisionEnabled))
 			{
 				if (!JointSettingsBuffer.bCollisionEnabled)
 				{
-					Chaos::FConstraintBase::FProxyBasePair& BasePairs = Constraint->GetParticleProxies();
-					Chaos::TGeometryParticleHandle<Chaos::FReal, 3>* Handle0 = GetParticleHandleFromProxy(BasePairs[0]);
-					Chaos::TGeometryParticleHandle<Chaos::FReal, 3>* Handle1 = GetParticleHandleFromProxy(BasePairs[1]);
+					FConstraintBase::FProxyBasePair& BasePairs = Constraint->GetParticleProxies();
+					FGeometryParticleHandle* Handle0 = GetParticleHandleFromProxy(BasePairs[0]);
+					FGeometryParticleHandle* Handle1 = GetParticleHandleFromProxy(BasePairs[1]);
 
 					// Three pieces of state to update on the physics thread. 
 					// .. Mask on the particle array
@@ -278,26 +277,26 @@ void FJointConstraintPhysicsProxy::PushStateOnPhysicsThread(Chaos::FPBDRigidsSol
 					// .. IgnoreCollisionsManager
 					if (Handle0 && Handle1)
 					{
-						Chaos::TPBDRigidParticleHandle<FReal, 3>* RigidHandle0 = Handle0->CastToRigidParticle();
-						Chaos::TPBDRigidParticleHandle<FReal, 3>* RigidHandle1 = Handle1->CastToRigidParticle();
+						FPBDRigidParticleHandle* RigidHandle0 = Handle0->CastToRigidParticle();
+						FPBDRigidParticleHandle* RigidHandle1 = Handle1->CastToRigidParticle();
 
 						// As long as one particle is a rigid we can add the ignore entry, one particle can be a static
 						if (RigidHandle0 || RigidHandle1)
 						{
-							const Chaos::FUniqueIdx ID0 = Handle0->UniqueIdx();
-							const Chaos::FUniqueIdx ID1 = Handle1->UniqueIdx();
-							Chaos::FIgnoreCollisionManager& IgnoreCollisionManager = InSolver->GetEvolution()->GetBroadPhase().GetIgnoreCollisionManager();
+							const FUniqueIdx ID0 = Handle0->UniqueIdx();
+							const FUniqueIdx ID1 = Handle1->UniqueIdx();
+							FIgnoreCollisionManager& IgnoreCollisionManager = InSolver->GetEvolution()->GetBroadPhase().GetIgnoreCollisionManager();
 
 							// For rigid/dynamic particles, add the broadphase flag and the IDs to check for disabled collisions
 							if(RigidHandle0)
 							{
-								RigidHandle0->AddCollisionConstraintFlag(Chaos::ECollisionConstraintFlags::CCF_BroadPhaseIgnoreCollisions);
+								RigidHandle0->AddCollisionConstraintFlag(ECollisionConstraintFlags::CCF_BroadPhaseIgnoreCollisions);
 								IgnoreCollisionManager.AddIgnoreCollisionsFor(ID0, ID1);
 							}
 
 							if(RigidHandle1)
 							{
-								RigidHandle1->AddCollisionConstraintFlag(Chaos::ECollisionConstraintFlags::CCF_BroadPhaseIgnoreCollisions);
+								RigidHandle1->AddCollisionConstraintFlag(ECollisionConstraintFlags::CCF_BroadPhaseIgnoreCollisions);
 								IgnoreCollisionManager.AddIgnoreCollisionsFor(ID1, ID0);
 							}
 
@@ -307,37 +306,37 @@ void FJointConstraintPhysicsProxy::PushStateOnPhysicsThread(Chaos::FPBDRigidsSol
 				}
 			}
 
-			if (DirtyFlagsBuffer.IsDirty(Chaos::EJointConstraintFlags::Projection))
+			if (DirtyFlagsBuffer.IsDirty(EJointConstraintFlags::Projection))
 			{
 				ConstraintSettings.bProjectionEnabled = JointSettingsBuffer.bProjectionEnabled;
 				ConstraintSettings.LinearProjection = JointSettingsBuffer.LinearProjection;
 				ConstraintSettings.AngularProjection = JointSettingsBuffer.AngularProjection;
 			}
 
-			if (DirtyFlagsBuffer.IsDirty(Chaos::EJointConstraintFlags::ParentInvMassScale))
+			if (DirtyFlagsBuffer.IsDirty(EJointConstraintFlags::ParentInvMassScale))
 			{
 				ConstraintSettings.ParentInvMassScale = JointSettingsBuffer.ParentInvMassScale;
 			}
 
-			if (DirtyFlagsBuffer.IsDirty(Chaos::EJointConstraintFlags::LinearBreakForce))
+			if (DirtyFlagsBuffer.IsDirty(EJointConstraintFlags::LinearBreakForce))
 			{
 				ConstraintSettings.LinearBreakForce = JointSettingsBuffer.LinearBreakForce;
 				ConstraintSettings.LinearPlasticityType = JointSettingsBuffer.LinearPlasticityType;
 				ConstraintSettings.LinearPlasticityLimit = FMath::Clamp((float)JointSettingsBuffer.LinearPlasticityLimit, 0.f, FLT_MAX);
 			}
 
-			if (DirtyFlagsBuffer.IsDirty(Chaos::EJointConstraintFlags::AngularBreakTorque))
+			if (DirtyFlagsBuffer.IsDirty(EJointConstraintFlags::AngularBreakTorque))
 			{
 				ConstraintSettings.AngularBreakTorque = JointSettingsBuffer.AngularBreakTorque;
 				ConstraintSettings.AngularPlasticityLimit = JointSettingsBuffer.AngularPlasticityLimit;
 			}
 
-			if (DirtyFlagsBuffer.IsDirty(Chaos::EJointConstraintFlags::UserData))
+			if (DirtyFlagsBuffer.IsDirty(EJointConstraintFlags::UserData))
 			{
 				ConstraintSettings.UserData = JointSettingsBuffer.UserData;
 			}
 
-			if (DirtyFlagsBuffer.IsDirty(Chaos::EJointConstraintFlags::LinearDrive))
+			if (DirtyFlagsBuffer.IsDirty(EJointConstraintFlags::LinearDrive))
 			{
 				ConstraintSettings.bLinearPositionDriveEnabled[0] = JointSettingsBuffer.bLinearPositionDriveEnabled[0];
 				ConstraintSettings.bLinearPositionDriveEnabled[1] = JointSettingsBuffer.bLinearPositionDriveEnabled[1];
@@ -355,7 +354,7 @@ void FJointConstraintPhysicsProxy::PushStateOnPhysicsThread(Chaos::FPBDRigidsSol
 				ConstraintSettings.LinearDriveDamping = JointSettingsBuffer.LinearDriveDamping;
 			}
 
-			if (DirtyFlagsBuffer.IsDirty(Chaos::EJointConstraintFlags::AngularDrive))
+			if (DirtyFlagsBuffer.IsDirty(EJointConstraintFlags::AngularDrive))
 			{
 				ConstraintSettings.bAngularSLerpPositionDriveEnabled = JointSettingsBuffer.bAngularSLerpPositionDriveEnabled;
 				ConstraintSettings.bAngularTwistPositionDriveEnabled = JointSettingsBuffer.bAngularTwistPositionDriveEnabled;
@@ -373,12 +372,12 @@ void FJointConstraintPhysicsProxy::PushStateOnPhysicsThread(Chaos::FPBDRigidsSol
 				ConstraintSettings.AngularDriveDamping = JointSettingsBuffer.AngularDriveDamping;
 			}
 
-			if (DirtyFlagsBuffer.IsDirty(Chaos::EJointConstraintFlags::Stiffness))
+			if (DirtyFlagsBuffer.IsDirty(EJointConstraintFlags::Stiffness))
 			{
 				ConstraintSettings.Stiffness = JointSettingsBuffer.Stiffness;
 			}
 
-			if (DirtyFlagsBuffer.IsDirty(Chaos::EJointConstraintFlags::Limits))
+			if (DirtyFlagsBuffer.IsDirty(EJointConstraintFlags::Limits))
 			{
 				ConstraintSettings.bSoftLinearLimitsEnabled = JointSettingsBuffer.bSoftLinearLimitsEnabled;
 				ConstraintSettings.bSoftTwistLimitsEnabled = JointSettingsBuffer.bSoftTwistLimitsEnabled;
@@ -407,3 +406,5 @@ void FJointConstraintPhysicsProxy::PushStateOnPhysicsThread(Chaos::FPBDRigidsSol
 		}
 	}
 }
+
+}
\ No newline at end of file
diff --git a/Engine/Source/Runtime/Experimental/Chaos/Private/PhysicsProxy/SingleParticlePhysicsProxy.cpp b/Engine/Source/Runtime/Experimental/Chaos/Private/PhysicsProxy/SingleParticlePhysicsProxy.cpp
--- a/Engine/Source/Runtime/Experimental/Chaos/Private/PhysicsProxy/SingleParticlePhysicsProxy.cpp
+++ b/Engine/Source/Runtime/Experimental/Chaos/Private/PhysicsProxy/SingleParticlePhysicsProxy.cpp
@@ -23,10 +23,9 @@ FAutoConsoleVariableRef CVar_SyncKinematicOnGameThread(TEXT("P.Chaos.SyncKinemat
 
 
 FSingleParticlePhysicsProxy::FSingleParticlePhysicsProxy(TUniquePtr<PARTICLE_TYPE>&& InParticle, FParticleHandle* InHandle, UObject* InOwner)
-	: IPhysicsProxyBase(EPhysicsProxyType::SingleParticleProxy)
+	: IPhysicsProxyBase(EPhysicsProxyType::SingleParticleProxy, InOwner)
 	, Particle(MoveTemp(InParticle))
 	, Handle(InHandle)
-	, Owner(InOwner)
 	, PullDataInterpIdx_External(INDEX_NONE)
 {
 	Particle->SetProxy(this);
diff --git a/Engine/Source/Runtime/Experimental/Chaos/Public/Chaos/Framework/PhysicsProxy.h b/Engine/Source/Runtime/Experimental/Chaos/Public/Chaos/Framework/PhysicsProxy.h
--- a/Engine/Source/Runtime/Experimental/Chaos/Public/Chaos/Framework/PhysicsProxy.h
+++ b/Engine/Source/Runtime/Experimental/Chaos/Public/Chaos/Framework/PhysicsProxy.h
@@ -44,14 +44,12 @@ public:
 	using FIntArray = Chaos::TArrayCollectionArray<int32>;
 
 	TPhysicsProxy()
-		: IPhysicsProxyBase(ConcreteType())
-		, Owner(nullptr)
+		: IPhysicsProxyBase(ConcreteType(), nullptr)
 	{
 	}
 
 	explicit TPhysicsProxy(UObject* InOwner)
-		: IPhysicsProxyBase(ConcreteType())
-		, Owner(InOwner)
+		: IPhysicsProxyBase(ConcreteType(), InOwner)
 	{
 	}
 
@@ -121,21 +119,7 @@ public:
 
 	bool IsDirty() { return static_cast<Concrete*>(this)->IsDirty(); }
 
-
-
-	/** Gets the owning external object for this solver object, never used internally */
-	virtual UObject* GetOwner() const override { return Owner; }
-
 	void* GetUserData() const { return nullptr; }
 
 	Chaos::FRigidTransform3 GetTransform() const { return Chaos::FRigidTransform3(); }
-
-
-private:
-
-	/** 
-	 * The owner for this solver object, essentially user-data managed by the caller 
-	 * @see GetOwner
-	 */
-	UObject* Owner;
 };
diff --git a/Engine/Source/Runtime/Experimental/Chaos/Public/Chaos/Framework/PhysicsProxyBase.h b/Engine/Source/Runtime/Experimental/Chaos/Public/Chaos/Framework/PhysicsProxyBase.h
--- a/Engine/Source/Runtime/Experimental/Chaos/Public/Chaos/Framework/PhysicsProxyBase.h
+++ b/Engine/Source/Runtime/Experimental/Chaos/Public/Chaos/Framework/PhysicsProxyBase.h
@@ -40,14 +40,15 @@ struct CHAOS_API FProxyTimestamp
 class CHAOS_API IPhysicsProxyBase
 {
 public:
-	IPhysicsProxyBase(EPhysicsProxyType InType)
+	IPhysicsProxyBase(EPhysicsProxyType InType, UObject* InOwner)
 		: Solver(nullptr)
+		, Owner(InOwner)
 		, DirtyIdx(INDEX_NONE)
 		, Type(InType)
 		, SyncTimestamp(new FProxyTimestamp)
 	{}
 
-	virtual UObject* GetOwner() const = 0;
+	UObject* GetOwner() const { return Owner; }
 
 	template< class SOLVER_TYPE>
 	SOLVER_TYPE* GetSolver() const { return static_cast<SOLVER_TYPE*>(Solver); }
@@ -79,6 +80,7 @@ protected:
 	
 	/** The solver that owns the solver object */
 	Chaos::FPhysicsSolverBase* Solver;
+	UObject* Owner;
 
 private:
 	int32 DirtyIdx;
diff --git a/Engine/Source/Runtime/Experimental/Chaos/Public/PhysicsProxy/JointConstraintProxy.h b/Engine/Source/Runtime/Experimental/Chaos/Public/PhysicsProxy/JointConstraintProxy.h
--- a/Engine/Source/Runtime/Experimental/Chaos/Public/PhysicsProxy/JointConstraintProxy.h
+++ b/Engine/Source/Runtime/Experimental/Chaos/Public/PhysicsProxy/JointConstraintProxy.h
@@ -14,108 +14,64 @@
 
 namespace Chaos
 {
-	class FJointConstraint;
+class FJointConstraint;
 
-	class FPBDRigidsEvolutionGBF;
+class FPBDRigidsEvolutionGBF;
 
-	struct FDirtyJointConstraintData;
-}
+struct FDirtyJointConstraintData;
 
-class FJointConstraintPhysicsProxy : public TPhysicsProxy<FJointConstraintPhysicsProxy,void>
+class FJointConstraintPhysicsProxy : public IPhysicsProxyBase
 {
-	typedef TPhysicsProxy<FJointConstraintPhysicsProxy, void> Base;
+	using Base = IPhysicsProxyBase;
 
 public:
-	using FReal = Chaos::FReal;
-	using FConstraintHandle = typename Chaos::FJointConstraint::FHandle;
-	using FConstraintData = typename Chaos::FJointConstraint::FData;
-	using FJointConstraints = Chaos::FPBDJointConstraints;
-	using FJointConstraintDirtyFlags = Chaos::FJointConstraintDirtyFlags;
-	using FParticlePair = Chaos::TVec2<Chaos::FGeometryParticle*>;
-	using FParticleHandlePair = Chaos::TVector<Chaos::TGeometryParticleHandle<FReal, 3>*, 2>;
-
 	FJointConstraintPhysicsProxy() = delete;
-	FJointConstraintPhysicsProxy(Chaos::FJointConstraint* InConstraint, FConstraintHandle* InHandle, UObject* InOwner = nullptr);
-
-	EPhysicsProxyType ConcreteType()
-	{
-		return EPhysicsProxyType::JointConstraintType;
-	}
-
-	
-	bool IsValid() { return Constraint != nullptr && Constraint->IsValid(); }
+	FJointConstraintPhysicsProxy(FJointConstraint* InConstraint, FPBDJointConstraintHandle* InHandle, UObject* InOwner = nullptr);
+		
+	bool IsValid() { return Constraint && Constraint->IsValid(); }
 
 	bool IsInitialized() const { return bInitialized; }
 	void SetInitialized() { bInitialized = true; }
 
-	static Chaos::TGeometryParticleHandle<Chaos::FReal, 3>* GetParticleHandleFromProxy(IPhysicsProxyBase* ProxyBase);
+	static FGeometryParticleHandle* GetParticleHandleFromProxy(IPhysicsProxyBase* ProxyBase);
 
 	//
 	//  Lifespan Management
 	//
 
-	void CHAOS_API InitializeOnPhysicsThread(Chaos::FPBDRigidsSolver* InSolver);
+	void CHAOS_API InitializeOnPhysicsThread(FPBDRigidsSolver* InSolver);
 
 	// Merge to perform a remote sync
-	void CHAOS_API PushStateOnGameThread(Chaos::FPBDRigidsSolver* InSolver);
+	void CHAOS_API PushStateOnGameThread(FPBDRigidsSolver* InSolver);
 
-	void CHAOS_API PushStateOnPhysicsThread(Chaos::FPBDRigidsSolver* InSolver);
+	void CHAOS_API PushStateOnPhysicsThread(FPBDRigidsSolver* InSolver);
 	// Merge to perform a remote sync - END
 
-	void CHAOS_API DestroyOnPhysicsThread(Chaos::FPBDRigidsSolver* InSolver);
-
-	void SyncBeforeDestroy() {}
-	void OnRemoveFromScene() {}
-
+	void CHAOS_API DestroyOnPhysicsThread(FPBDRigidsSolver* InSolver);
 
 	//
 	// Member Access
 	//
 
-	FConstraintHandle* GetHandle()
-	{
-		return Handle;
-	}
-
-	const FConstraintHandle* GetHandle() const
-	{
-		return Handle;
-	}
-
-	virtual void* GetHandleUnsafe() const override
-	{
-		return Handle;
-	}
+	FPBDJointConstraintHandle* GetHandle() { return Handle; }
+	const FPBDJointConstraintHandle* GetHandle() const { return Handle; }
 
-	void SetHandle(FConstraintHandle* InHandle)
-	{
-		Handle = InHandle;
-	}
+	virtual void* GetHandleUnsafe() const override { return Handle; }
 
-	Chaos::FJointConstraint* GetConstraint()
-	{
-		return Constraint;
-	}
+	void SetHandle(FPBDJointConstraintHandle* InHandle)	{ Handle = InHandle; }
 
-	const Chaos::FJointConstraint* GetConstraint() const
-	{
-		return Constraint;
-	}
+	FJointConstraint* GetConstraint(){ return Constraint; }
+	const FJointConstraint* GetConstraint() const { return Constraint; }
 
 	//
 	// Threading API
 	//
 	
-	void PushToPhysicsState(Chaos::FPBDRigidsEvolutionGBF& Evolution) {}
-
-	/**/
-	void ClearAccumulatedData() {}
-
 	/**/
-	void BufferPhysicsResults(Chaos::FDirtyJointConstraintData& Buffer);
+	void BufferPhysicsResults(FDirtyJointConstraintData& Buffer);
 
 	/**/
-	bool CHAOS_API PullFromPhysicsState(const Chaos::FDirtyJointConstraintData& Buffer, const int32 SolverSyncTimestamp);
+	bool CHAOS_API PullFromPhysicsState(const FDirtyJointConstraintData& Buffer, const int32 SolverSyncTimestamp);
 
 	/**/
 	bool IsDirty() { return Constraint->IsDirty(); }
@@ -123,23 +79,12 @@ public:
 private:
 
 	// Input Buffer
-	FConstraintData JointSettingsBuffer;
+	FPBDJointSettings JointSettingsBuffer;
 	FJointConstraintDirtyFlags DirtyFlagsBuffer;
 
-	Chaos::FJointConstraint* Constraint; 	// This proxy assumes ownership of the Constraint, and will free it during DestroyOnPhysicsThread
-	FConstraintHandle* Handle;
-	bool bInitialized;
-
-
-public:
-	void AddForceCallback(FParticlesType& InParticles, const FReal InDt, const int32 InIndex) {}
-	void DisableCollisionsCallback(TSet<TTuple<int32, int32>>& InPairs) {}
-	void BindParticleCallbackMapping(Chaos::TArrayCollectionArray<PhysicsProxyWrapper>& PhysicsProxyReverseMap, Chaos::TArrayCollectionArray<int32>& ParticleIDReverseMap) {}
-	void EndFrameCallback(const FReal InDt) {}
-	void ParameterUpdateCallback(FParticlesType& InParticles, const FReal InTime) {}
-	void CreateRigidBodyCallback(FParticlesType& InOutParticles) {}
-	bool IsSimulating() const { return true; }
-	void UpdateKinematicBodiesCallback(const FParticlesType& InParticles, const FReal InDt, const FReal InTime, FKinematicProxy& InKinematicProxy) {}
-	void StartFrameCallback(const FReal InDt, const FReal InTime) {}
-
+	FJointConstraint* Constraint; 	// This proxy assumes ownership of the Constraint, and will free it during DestroyOnPhysicsThread
+	FPBDJointConstraintHandle* Handle;
+	bool bInitialized = false;
 };
+
+}
\ No newline at end of file
diff --git a/Engine/Source/Runtime/Experimental/Chaos/Public/PhysicsProxy/SingleParticlePhysicsProxy.h b/Engine/Source/Runtime/Experimental/Chaos/Public/PhysicsProxy/SingleParticlePhysicsProxy.h
--- a/Engine/Source/Runtime/Experimental/Chaos/Public/PhysicsProxy/SingleParticlePhysicsProxy.h
+++ b/Engine/Source/Runtime/Experimental/Chaos/Public/PhysicsProxy/SingleParticlePhysicsProxy.h
@@ -170,9 +170,6 @@ public:
 	{
 		return static_cast<const Chaos::FPBDRigidParticle*>(GetParticle_LowLevel());
 	}
-
-	/** Gets the owning external object for this solver object, never used internally */
-	virtual UObject* GetOwner() const override { return Owner; }
 	
 private:
 	int32 InitializedOnStep = INDEX_NONE;
@@ -183,8 +180,6 @@ protected:
 
 private:
 
-	UObject* Owner;
-
 	//Used by interpolation code
 	int32 PullDataInterpIdx_External;
 
diff --git a/Engine/Source/Runtime/PhysicsCore/Private/ChaosEngineInterface.cpp b/Engine/Source/Runtime/PhysicsCore/Private/ChaosEngineInterface.cpp
--- a/Engine/Source/Runtime/PhysicsCore/Private/ChaosEngineInterface.cpp
+++ b/Engine/Source/Runtime/PhysicsCore/Private/ChaosEngineInterface.cpp
@@ -1075,17 +1075,18 @@ void FChaosEngineInterface::SetConstraintUserData(const FPhysicsConstraintHandle
 
 void FChaosEngineInterface::ReleaseConstraint(FPhysicsConstraintHandle& InConstraintRef)
 {
+	using namespace Chaos;
 	if (bEnableChaosJointConstraints)
 	{
 		LLM_SCOPE(ELLMTag::Chaos);
 		if (InConstraintRef.IsValid() && InConstraintRef.Constraint->IsType(Chaos::EConstraintType::JointConstraintType))
 		{
-			if (Chaos::FJointConstraint* Constraint = static_cast<Chaos::FJointConstraint*>(InConstraintRef.Constraint))
+			if (FJointConstraint* Constraint = static_cast<FJointConstraint*>(InConstraintRef.Constraint))
 			{
 				if (FJointConstraintPhysicsProxy* Proxy = Constraint->GetProxy<FJointConstraintPhysicsProxy>())
 				{
-					check(Proxy->GetSolver<Chaos::FPhysicsSolver>());
-					Chaos::FPhysicsSolver* Solver = Proxy->GetSolver<Chaos::FPhysicsSolver>();
+					check(Proxy->GetSolver<FPhysicsSolver>());
+					FPhysicsSolver* Solver = Proxy->GetSolver<FPhysicsSolver>();
 
 					Solver->UnregisterObject(Constraint);
 
@@ -1093,14 +1094,14 @@ void FChaosEngineInterface::ReleaseConstraint(FPhysicsConstraintHandle& InConstr
 				}
 			}
 		}
-		else if (InConstraintRef.IsValid() && InConstraintRef.Constraint->IsType(Chaos::EConstraintType::SuspensionConstraintType))
+		else if (InConstraintRef.IsValid() && InConstraintRef.Constraint->IsType(EConstraintType::SuspensionConstraintType))
 		{
-			if (Chaos::FSuspensionConstraint* Constraint = static_cast<Chaos::FSuspensionConstraint*>(InConstraintRef.Constraint))
+			if (Chaos::FSuspensionConstraint* Constraint = static_cast<FSuspensionConstraint*>(InConstraintRef.Constraint))
 			{
 				if (FSuspensionConstraintPhysicsProxy* Proxy = Constraint->GetProxy<FSuspensionConstraintPhysicsProxy>())
 				{
-					check(Proxy->GetSolver<Chaos::FPhysicsSolver>());
-					Chaos::FPhysicsSolver* Solver = Proxy->GetSolver<Chaos::FPhysicsSolver>();
+					check(Proxy->GetSolver<FPhysicsSolver>());
+					FPhysicsSolver* Solver = Proxy->GetSolver<FPhysicsSolver>();
 
 					Solver->UnregisterObject(Constraint);
 
