Simplify joint buffering code to use a single struct copy (copy was already happening anyway). This will make async support easier

diff --git a/Engine/Source/Runtime/Experimental/Chaos/Private/PhysicsProxy/JointConstraintProxy.cpp b/Engine/Source/Runtime/Experimental/Chaos/Private/PhysicsProxy/JointConstraintProxy.cpp
--- a/Engine/Source/Runtime/Experimental/Chaos/Private/PhysicsProxy/JointConstraintProxy.cpp
+++ b/Engine/Source/Runtime/Experimental/Chaos/Private/PhysicsProxy/JointConstraintProxy.cpp
@@ -123,127 +123,12 @@ void FJointConstraintPhysicsProxy::DestroyOnPhysicsThread(FPBDRigidsSolver* InSo
 
 void FJointConstraintPhysicsProxy::PushStateOnGameThread(FPBDRigidsSolver* InSolver)
 {
-	if (Constraint != nullptr && Constraint->IsValid())
+	if (Constraint && Constraint->IsValid())
 	{
 		if (Constraint->IsDirty())
 		{
-			if (Constraint->IsDirty(EJointConstraintFlags::JointTransforms))
-			{
-				JointSettingsBuffer.ConnectorTransforms = Constraint->GetJointTransforms();
-				DirtyFlagsBuffer.MarkDirty(EJointConstraintFlags::JointTransforms);
-			}
-
-			if (Constraint->IsDirty(EJointConstraintFlags::CollisionEnabled))
-			{
-				JointSettingsBuffer.bCollisionEnabled = Constraint->GetCollisionEnabled();
-				DirtyFlagsBuffer.MarkDirty(EJointConstraintFlags::CollisionEnabled);
-			}
-
-			if (Constraint->IsDirty(EJointConstraintFlags::Projection))
-			{
-				JointSettingsBuffer.bProjectionEnabled = Constraint->GetProjectionEnabled();
-				JointSettingsBuffer.LinearProjection = Constraint->GetProjectionLinearAlpha();
-				JointSettingsBuffer.AngularProjection = Constraint->GetProjectionAngularAlpha();
-				DirtyFlagsBuffer.MarkDirty(EJointConstraintFlags::Projection);
-			}
-
-			if (Constraint->IsDirty(EJointConstraintFlags::ParentInvMassScale))
-			{
-				JointSettingsBuffer.ParentInvMassScale = Constraint->GetParentInvMassScale();
-				DirtyFlagsBuffer.MarkDirty(EJointConstraintFlags::ParentInvMassScale);
-			}
-
-			if (Constraint->IsDirty(EJointConstraintFlags::LinearBreakForce))
-			{
-				JointSettingsBuffer.LinearBreakForce = Constraint->GetLinearBreakForce();
-				JointSettingsBuffer.LinearPlasticityType = Constraint->GetLinearPlasticityType();
-				JointSettingsBuffer.LinearPlasticityLimit = Constraint->GetLinearPlasticityLimit();
-				DirtyFlagsBuffer.MarkDirty(EJointConstraintFlags::LinearBreakForce);
-			}
-
-			if (Constraint->IsDirty(EJointConstraintFlags::AngularBreakTorque))
-			{
-				JointSettingsBuffer.AngularBreakTorque = Constraint->GetAngularBreakTorque();
-				JointSettingsBuffer.AngularPlasticityLimit = Constraint->GetAngularPlasticityLimit();
-				DirtyFlagsBuffer.MarkDirty(EJointConstraintFlags::AngularBreakTorque);
-			}
-
-			if (Constraint->IsDirty(EJointConstraintFlags::UserData))
-			{
-				JointSettingsBuffer.UserData = Constraint->GetUserData();
-				DirtyFlagsBuffer.MarkDirty(EJointConstraintFlags::UserData);
-			}
-
-			if (Constraint->IsDirty(EJointConstraintFlags::LinearDrive))
-			{
-				JointSettingsBuffer.bLinearPositionDriveEnabled[0] = Constraint->GetLinearPositionDriveXEnabled();
-				JointSettingsBuffer.bLinearPositionDriveEnabled[1] = Constraint->GetLinearPositionDriveYEnabled();
-				JointSettingsBuffer.bLinearPositionDriveEnabled[2] = Constraint->GetLinearPositionDriveZEnabled();
-				JointSettingsBuffer.LinearDrivePositionTarget = Constraint->GetLinearDrivePositionTarget();
-				JointSettingsBuffer.bLinearVelocityDriveEnabled[0] = Constraint->GetLinearVelocityDriveXEnabled();
-				JointSettingsBuffer.bLinearVelocityDriveEnabled[1] = Constraint->GetLinearVelocityDriveYEnabled();
-				JointSettingsBuffer.bLinearVelocityDriveEnabled[2] = Constraint->GetLinearVelocityDriveZEnabled();
-				JointSettingsBuffer.LinearDriveVelocityTarget = Constraint->GetLinearDriveVelocityTarget();
-				JointSettingsBuffer.LinearDriveForceMode = Constraint->GetLinearDriveForceMode();
-				JointSettingsBuffer.LinearMotionTypes[0] = Constraint->GetLinearMotionTypesX();
-				JointSettingsBuffer.LinearMotionTypes[1] = Constraint->GetLinearMotionTypesY();
-				JointSettingsBuffer.LinearMotionTypes[2] = Constraint->GetLinearMotionTypesZ();
-				JointSettingsBuffer.LinearDriveStiffness = Constraint->GetLinearDriveStiffness();
-				JointSettingsBuffer.LinearDriveDamping = Constraint->GetLinearDriveDamping();
-				DirtyFlagsBuffer.MarkDirty(EJointConstraintFlags::LinearDrive);
-			}
-
-
-			if (Constraint->IsDirty(EJointConstraintFlags::AngularDrive))
-			{
-				JointSettingsBuffer.bAngularSLerpPositionDriveEnabled = Constraint->GetAngularSLerpPositionDriveEnabled();
-				JointSettingsBuffer.bAngularTwistPositionDriveEnabled = Constraint->GetAngularTwistPositionDriveEnabled();
-				JointSettingsBuffer.bAngularSwingPositionDriveEnabled = Constraint->GetAngularSwingPositionDriveEnabled();
-				JointSettingsBuffer.AngularDrivePositionTarget = Constraint->GetAngularDrivePositionTarget();
-				JointSettingsBuffer.bAngularSLerpVelocityDriveEnabled = Constraint->GetAngularSLerpVelocityDriveEnabled();
-				JointSettingsBuffer.bAngularTwistVelocityDriveEnabled = Constraint->GetAngularTwistVelocityDriveEnabled();
-				JointSettingsBuffer.bAngularSwingVelocityDriveEnabled = Constraint->GetAngularSwingVelocityDriveEnabled();
-				JointSettingsBuffer.AngularDriveVelocityTarget = Constraint->GetAngularDriveVelocityTarget();
-				JointSettingsBuffer.AngularDriveForceMode = Constraint->GetAngularDriveForceMode();
-				JointSettingsBuffer.AngularMotionTypes[0] = Constraint->GetAngularMotionTypesX();
-				JointSettingsBuffer.AngularMotionTypes[1] = Constraint->GetAngularMotionTypesY();
-				JointSettingsBuffer.AngularMotionTypes[2] = Constraint->GetAngularMotionTypesZ();
-				JointSettingsBuffer.AngularDriveStiffness = Constraint->GetAngularDriveStiffness();
-				JointSettingsBuffer.AngularDriveDamping = Constraint->GetAngularDriveDamping();
-				DirtyFlagsBuffer.MarkDirty(EJointConstraintFlags::AngularDrive);
-			}
-
-			if (Constraint->IsDirty(EJointConstraintFlags::Stiffness))
-			{
-				JointSettingsBuffer.Stiffness = Constraint->GetStiffness();
-				DirtyFlagsBuffer.MarkDirty(EJointConstraintFlags::Stiffness);
-			}
-
-			if (Constraint->IsDirty(EJointConstraintFlags::Limits))
-			{
-				JointSettingsBuffer.bSoftLinearLimitsEnabled = Constraint->GetSoftLinearLimitsEnabled();
-				JointSettingsBuffer.bSoftTwistLimitsEnabled = Constraint->GetSoftTwistLimitsEnabled();
-				JointSettingsBuffer.bSoftSwingLimitsEnabled = Constraint->GetSoftSwingLimitsEnabled();
-				JointSettingsBuffer.LinearSoftForceMode = Constraint->GetLinearSoftForceMode();
-				JointSettingsBuffer.AngularSoftForceMode = Constraint->GetAngularSoftForceMode();
-				JointSettingsBuffer.SoftLinearStiffness = Constraint->GetSoftLinearStiffness();
-				JointSettingsBuffer.SoftLinearDamping = Constraint->GetSoftLinearDamping();
-				JointSettingsBuffer.SoftTwistStiffness = Constraint->GetSoftTwistStiffness();
-				JointSettingsBuffer.SoftTwistDamping = Constraint->GetSoftTwistDamping();
-				JointSettingsBuffer.SoftSwingStiffness = Constraint->GetSoftSwingStiffness();
-				JointSettingsBuffer.SoftSwingDamping = Constraint->GetSoftSwingDamping();
-				JointSettingsBuffer.LinearLimit = Constraint->GetLinearLimit();
-				JointSettingsBuffer.AngularLimits = Constraint->GetAngularLimits();
-				JointSettingsBuffer.LinearContactDistance = Constraint->GetLinearContactDistance();
-				JointSettingsBuffer.TwistContactDistance = Constraint->GetTwistContactDistance();
-				JointSettingsBuffer.SwingContactDistance = Constraint->GetSwingContactDistance();
-				JointSettingsBuffer.LinearRestitution = Constraint->GetLinearRestitution();
-				JointSettingsBuffer.TwistRestitution = Constraint->GetTwistRestitution();
-				JointSettingsBuffer.SwingRestitution = Constraint->GetSwingRestitution();
-
-				DirtyFlagsBuffer.MarkDirty(EJointConstraintFlags::Limits);
-			}
-
+			JointSettingsBuffer = Constraint->GetJointSettings();
+			DirtyFlagsBuffer = Constraint->GetDirtyFlags();
 			Constraint->ClearDirtyFlags();
 		}
 	}
@@ -256,152 +141,63 @@ void FJointConstraintPhysicsProxy::PushStateOnPhysicsThread(FPBDRigidsSolver* In
 	{
 		if (DirtyFlagsBuffer.IsDirty())
 		{
-			FPBDJointSettings ConstraintSettings = Handle->GetSettings();
+			const FPBDJointSettings& CurrentConstraintSettings = Handle->GetSettings();
 
-			if (DirtyFlagsBuffer.IsDirty(EJointConstraintFlags::JointTransforms))
+			if (CurrentConstraintSettings.bCollisionEnabled != JointSettingsBuffer.bCollisionEnabled)
 			{
-				ConstraintSettings.ConnectorTransforms = JointSettingsBuffer.ConnectorTransforms;
-			}
+				FConstraintBase::FProxyBasePair& BasePairs = Constraint->GetParticleProxies();
+				FGeometryParticleHandle* Handle0 = GetParticleHandleFromProxy(BasePairs[0]);
+				FGeometryParticleHandle* Handle1 = GetParticleHandleFromProxy(BasePairs[1]);
 
-			if (DirtyFlagsBuffer.IsDirty(EJointConstraintFlags::CollisionEnabled))
-			{
-				if (!JointSettingsBuffer.bCollisionEnabled)
+				// Three pieces of state to update on the physics thread. 
+				// .. Mask on the particle array
+				// .. Constraint collisions enabled array
+				// .. IgnoreCollisionsManager
+				if (Handle0 && Handle1)
 				{
-					FConstraintBase::FProxyBasePair& BasePairs = Constraint->GetParticleProxies();
-					FGeometryParticleHandle* Handle0 = GetParticleHandleFromProxy(BasePairs[0]);
-					FGeometryParticleHandle* Handle1 = GetParticleHandleFromProxy(BasePairs[1]);
+					FPBDRigidParticleHandle* RigidHandle0 = Handle0->CastToRigidParticle();
+					FPBDRigidParticleHandle* RigidHandle1 = Handle1->CastToRigidParticle();
 
-					// Three pieces of state to update on the physics thread. 
-					// .. Mask on the particle array
-					// .. Constraint collisions enabled array
-					// .. IgnoreCollisionsManager
-					if (Handle0 && Handle1)
+					// As long as one particle is a rigid we can add the ignore entry, one particle can be a static
+					if (RigidHandle0 || RigidHandle1)
 					{
-						FPBDRigidParticleHandle* RigidHandle0 = Handle0->CastToRigidParticle();
-						FPBDRigidParticleHandle* RigidHandle1 = Handle1->CastToRigidParticle();
+						const FUniqueIdx ID0 = Handle0->UniqueIdx();
+						const FUniqueIdx ID1 = Handle1->UniqueIdx();
+						FIgnoreCollisionManager& IgnoreCollisionManager = InSolver->GetEvolution()->GetBroadPhase().GetIgnoreCollisionManager();
 
-						// As long as one particle is a rigid we can add the ignore entry, one particle can be a static
-						if (RigidHandle0 || RigidHandle1)
+						// For rigid/dynamic particles, add the broadphase flag and the IDs to check for disabled collisions
+						if(RigidHandle0)
 						{
-							const FUniqueIdx ID0 = Handle0->UniqueIdx();
-							const FUniqueIdx ID1 = Handle1->UniqueIdx();
-							FIgnoreCollisionManager& IgnoreCollisionManager = InSolver->GetEvolution()->GetBroadPhase().GetIgnoreCollisionManager();
-
-							// For rigid/dynamic particles, add the broadphase flag and the IDs to check for disabled collisions
-							if(RigidHandle0)
+							if(JointSettingsBuffer.bCollisionEnabled)
+							{
+								//can't remove collision constraint flag because we may still be ignoring collision for others
+								IgnoreCollisionManager.RemoveIgnoreCollisionsFor(ID0, ID1);
+							}
+							else
 							{
 								RigidHandle0->AddCollisionConstraintFlag(ECollisionConstraintFlags::CCF_BroadPhaseIgnoreCollisions);
 								IgnoreCollisionManager.AddIgnoreCollisionsFor(ID0, ID1);
 							}
+						}
 
-							if(RigidHandle1)
+						if(RigidHandle1)
+						{
+							if (JointSettingsBuffer.bCollisionEnabled)
+							{
+								//can't remove collision constraint flag because we may still be ignoring collision for others
+								IgnoreCollisionManager.RemoveIgnoreCollisionsFor(ID1, ID0);
+							}
+							else
 							{
 								RigidHandle1->AddCollisionConstraintFlag(ECollisionConstraintFlags::CCF_BroadPhaseIgnoreCollisions);
 								IgnoreCollisionManager.AddIgnoreCollisionsFor(ID1, ID0);
 							}
-
-							ConstraintSettings.bCollisionEnabled = JointSettingsBuffer.bCollisionEnabled;
 						}
 					}
 				}
 			}
 
-			if (DirtyFlagsBuffer.IsDirty(EJointConstraintFlags::Projection))
-			{
-				ConstraintSettings.bProjectionEnabled = JointSettingsBuffer.bProjectionEnabled;
-				ConstraintSettings.LinearProjection = JointSettingsBuffer.LinearProjection;
-				ConstraintSettings.AngularProjection = JointSettingsBuffer.AngularProjection;
-			}
-
-			if (DirtyFlagsBuffer.IsDirty(EJointConstraintFlags::ParentInvMassScale))
-			{
-				ConstraintSettings.ParentInvMassScale = JointSettingsBuffer.ParentInvMassScale;
-			}
-
-			if (DirtyFlagsBuffer.IsDirty(EJointConstraintFlags::LinearBreakForce))
-			{
-				ConstraintSettings.LinearBreakForce = JointSettingsBuffer.LinearBreakForce;
-				ConstraintSettings.LinearPlasticityType = JointSettingsBuffer.LinearPlasticityType;
-				ConstraintSettings.LinearPlasticityLimit = FMath::Clamp((float)JointSettingsBuffer.LinearPlasticityLimit, 0.f, FLT_MAX);
-			}
-
-			if (DirtyFlagsBuffer.IsDirty(EJointConstraintFlags::AngularBreakTorque))
-			{
-				ConstraintSettings.AngularBreakTorque = JointSettingsBuffer.AngularBreakTorque;
-				ConstraintSettings.AngularPlasticityLimit = JointSettingsBuffer.AngularPlasticityLimit;
-			}
-
-			if (DirtyFlagsBuffer.IsDirty(EJointConstraintFlags::UserData))
-			{
-				ConstraintSettings.UserData = JointSettingsBuffer.UserData;
-			}
-
-			if (DirtyFlagsBuffer.IsDirty(EJointConstraintFlags::LinearDrive))
-			{
-				ConstraintSettings.bLinearPositionDriveEnabled[0] = JointSettingsBuffer.bLinearPositionDriveEnabled[0];
-				ConstraintSettings.bLinearPositionDriveEnabled[1] = JointSettingsBuffer.bLinearPositionDriveEnabled[1];
-				ConstraintSettings.bLinearPositionDriveEnabled[2] = JointSettingsBuffer.bLinearPositionDriveEnabled[2];
-				ConstraintSettings.LinearDrivePositionTarget = JointSettingsBuffer.LinearDrivePositionTarget;
-				ConstraintSettings.bLinearVelocityDriveEnabled[0] = JointSettingsBuffer.bLinearVelocityDriveEnabled[0];
-				ConstraintSettings.bLinearVelocityDriveEnabled[1] = JointSettingsBuffer.bLinearVelocityDriveEnabled[1];
-				ConstraintSettings.bLinearVelocityDriveEnabled[2] = JointSettingsBuffer.bLinearVelocityDriveEnabled[2];
-				ConstraintSettings.LinearDriveVelocityTarget = JointSettingsBuffer.LinearDriveVelocityTarget;
-				ConstraintSettings.LinearDriveForceMode = JointSettingsBuffer.LinearDriveForceMode;
-				ConstraintSettings.LinearMotionTypes[0] = JointSettingsBuffer.LinearMotionTypes[0];
-				ConstraintSettings.LinearMotionTypes[1] = JointSettingsBuffer.LinearMotionTypes[1];
-				ConstraintSettings.LinearMotionTypes[2] = JointSettingsBuffer.LinearMotionTypes[2];
-				ConstraintSettings.LinearDriveStiffness = JointSettingsBuffer.LinearDriveStiffness;
-				ConstraintSettings.LinearDriveDamping = JointSettingsBuffer.LinearDriveDamping;
-			}
-
-			if (DirtyFlagsBuffer.IsDirty(EJointConstraintFlags::AngularDrive))
-			{
-				ConstraintSettings.bAngularSLerpPositionDriveEnabled = JointSettingsBuffer.bAngularSLerpPositionDriveEnabled;
-				ConstraintSettings.bAngularTwistPositionDriveEnabled = JointSettingsBuffer.bAngularTwistPositionDriveEnabled;
-				ConstraintSettings.bAngularSwingPositionDriveEnabled = JointSettingsBuffer.bAngularSwingPositionDriveEnabled;
-				ConstraintSettings.AngularDrivePositionTarget = JointSettingsBuffer.AngularDrivePositionTarget;
-				ConstraintSettings.bAngularSLerpVelocityDriveEnabled = JointSettingsBuffer.bAngularSLerpVelocityDriveEnabled;
-				ConstraintSettings.bAngularTwistVelocityDriveEnabled = JointSettingsBuffer.bAngularTwistVelocityDriveEnabled;
-				ConstraintSettings.bAngularSwingVelocityDriveEnabled = JointSettingsBuffer.bAngularSwingVelocityDriveEnabled;
-				ConstraintSettings.AngularDriveVelocityTarget = JointSettingsBuffer.AngularDriveVelocityTarget;
-				ConstraintSettings.AngularDriveForceMode = JointSettingsBuffer.AngularDriveForceMode;
-				ConstraintSettings.AngularMotionTypes[0] = JointSettingsBuffer.AngularMotionTypes[0];
-				ConstraintSettings.AngularMotionTypes[1] = JointSettingsBuffer.AngularMotionTypes[1];
-				ConstraintSettings.AngularMotionTypes[2] = JointSettingsBuffer.AngularMotionTypes[2];
-				ConstraintSettings.AngularDriveStiffness = JointSettingsBuffer.AngularDriveStiffness;
-				ConstraintSettings.AngularDriveDamping = JointSettingsBuffer.AngularDriveDamping;
-			}
-
-			if (DirtyFlagsBuffer.IsDirty(EJointConstraintFlags::Stiffness))
-			{
-				ConstraintSettings.Stiffness = JointSettingsBuffer.Stiffness;
-			}
-
-			if (DirtyFlagsBuffer.IsDirty(EJointConstraintFlags::Limits))
-			{
-				ConstraintSettings.bSoftLinearLimitsEnabled = JointSettingsBuffer.bSoftLinearLimitsEnabled;
-				ConstraintSettings.bSoftTwistLimitsEnabled = JointSettingsBuffer.bSoftTwistLimitsEnabled;
-				ConstraintSettings.bSoftSwingLimitsEnabled = JointSettingsBuffer.bSoftSwingLimitsEnabled;
-				ConstraintSettings.LinearSoftForceMode = JointSettingsBuffer.LinearSoftForceMode;
-				ConstraintSettings.AngularSoftForceMode = JointSettingsBuffer.AngularSoftForceMode;
-				ConstraintSettings.SoftLinearStiffness = JointSettingsBuffer.SoftLinearStiffness;
-				ConstraintSettings.SoftLinearDamping = JointSettingsBuffer.SoftLinearDamping;
-				ConstraintSettings.SoftTwistStiffness = JointSettingsBuffer.SoftTwistStiffness;
-				ConstraintSettings.SoftTwistDamping = JointSettingsBuffer.SoftTwistDamping;
-				ConstraintSettings.SoftSwingStiffness = JointSettingsBuffer.SoftSwingStiffness;
-				ConstraintSettings.SoftSwingDamping = JointSettingsBuffer.SoftSwingDamping;
-				ConstraintSettings.LinearLimit = JointSettingsBuffer.LinearLimit;
-				ConstraintSettings.AngularLimits = JointSettingsBuffer.AngularLimits;
-				ConstraintSettings.LinearContactDistance = JointSettingsBuffer.LinearContactDistance;
-				ConstraintSettings.TwistContactDistance = JointSettingsBuffer.TwistContactDistance;
-				ConstraintSettings.SwingContactDistance = JointSettingsBuffer.SwingContactDistance;
-				ConstraintSettings.LinearRestitution = JointSettingsBuffer.LinearRestitution;
-				ConstraintSettings.TwistRestitution = JointSettingsBuffer.TwistRestitution;
-				ConstraintSettings.SwingRestitution = JointSettingsBuffer.SwingRestitution;
-			}
-
-			Handle->SetSettings(ConstraintSettings);
-
+			Handle->SetSettings(JointSettingsBuffer);
 			DirtyFlagsBuffer.Clear();
 		}
 	}
diff --git a/Engine/Source/Runtime/Experimental/Chaos/Public/Chaos/PBDJointConstraintData.h b/Engine/Source/Runtime/Experimental/Chaos/Public/Chaos/PBDJointConstraintData.h
--- a/Engine/Source/Runtime/Experimental/Chaos/Public/Chaos/PBDJointConstraintData.h
+++ b/Engine/Source/Runtime/Experimental/Chaos/Public/Chaos/PBDJointConstraintData.h
@@ -47,6 +47,7 @@ namespace Chaos
 		bool IsDirty() const { return MDirtyFlags.IsDirty(); }
 		bool IsDirty(const EJointConstraintFlags CheckBits) const { return MDirtyFlags.IsDirty(CheckBits); }
 		void ClearDirtyFlags() { MDirtyFlags.Clear(); }
+		FJointConstraintDirtyFlags GetDirtyFlags() const { return MDirtyFlags; }
 
 		const FData& GetJointSettings()const { return JointSettings; }
 
