GeometryCollection
-- Removal mechanism invoked when geometry collection particles have slept for specified time. Particle is disabled, event is triggered, scene proxt geometry is progressivly scaled to a point.

diff --git a/Engine/Source/Runtime/Experimental/Chaos/Private/Chaos/PBDRigidsEvolution.cpp b/Engine/Source/Runtime/Experimental/Chaos/Private/Chaos/PBDRigidsEvolution.cpp
--- a/Engine/Source/Runtime/Experimental/Chaos/Private/Chaos/PBDRigidsEvolution.cpp
+++ b/Engine/Source/Runtime/Experimental/Chaos/Private/Chaos/PBDRigidsEvolution.cpp
@@ -825,4 +825,29 @@ namespace Chaos
 		DisableConstraints(InParticles);
 	}
 
+	void FPBDRigidsEvolutionBase::DisableParticleWithRemovalEvent(FGeometryParticleHandle* Particle)
+	{
+		// Record removal for event generation
+		const int32 NewIdx = MAllRemovals.Add(FRemovalData());
+		FRemovalData& Removal = MAllRemovals[NewIdx];
+		Removal.Proxy = Particle->PhysicsProxy();
+		Removal.Location = Particle->X();
+		
+		if (Chaos::FPBDRigidParticleHandle* RigidParticle = Particle->CastToRigidParticle())
+		{
+			Removal.Mass = RigidParticle->M();
+		}
+		else
+		{
+			Removal.Mass = 0.0;
+		}
+		
+		if (Particle->Geometry()->HasBoundingBox())
+		{
+			Removal.BoundingBox = Particle->Geometry()->BoundingBox();
+		}
+		
+		DisableParticle(Particle);
+	}
+
 }
diff --git a/Engine/Source/Runtime/Experimental/Chaos/Private/EventDefaults.cpp b/Engine/Source/Runtime/Experimental/Chaos/Private/EventDefaults.cpp
--- a/Engine/Source/Runtime/Experimental/Chaos/Private/EventDefaults.cpp
+++ b/Engine/Source/Runtime/Experimental/Chaos/Private/EventDefaults.cpp
@@ -23,6 +23,7 @@ namespace Chaos
 		RegisterBreakingEvent(EventManager);
 		RegisterTrailingEvent(EventManager);
 		RegisterSleepingEvent(EventManager);
+		RegisterRemovalEvent(EventManager);
 	}
 
 	void FEventDefaults::RegisterCollisionEvent(FEventManager& EventManager)
@@ -438,4 +439,43 @@ namespace Chaos
 
 		});
 	}
+
+	void FEventDefaults::RegisterRemovalEvent(FEventManager& EventManager)
+	{
+		EventManager.template RegisterEvent<FRemovalEventData>(EEventType::Removal, []
+		(const Chaos::FPBDRigidsSolver* Solver, FRemovalEventData& RemovalEventData)
+			{
+				check(Solver);
+				ensure(IsInPhysicsThreadContext());
+
+				
+				FRemovalDataArray& AllRemovalDataArray = RemovalEventData.RemovalData.AllRemovalArray;
+				TMap<IPhysicsProxyBase*, TArray<int32>>& AllRemovalIndicesByPhysicsProxy = RemovalEventData.PhysicsProxyToRemovalIndices.PhysicsProxyToIndicesMap;
+
+				if (RemovalEventData.RemovalData.TimeCreated != Solver->MTime)
+				{
+					AllRemovalDataArray.Reset();
+					AllRemovalIndicesByPhysicsProxy.Reset();
+					RemovalEventData.RemovalData.TimeCreated = Solver->MTime;
+				}
+
+				const TArray<FRemovalData>& AllRemovalsArray = Solver->GetEvolution()->GetAllRemovals();
+				
+				for (int32 Idx = 0; Idx < AllRemovalsArray.Num(); ++Idx)
+				{
+					FRemovalData RemovalData;
+					RemovalData.Location = AllRemovalsArray[Idx].Location;
+					RemovalData.Mass = AllRemovalsArray[Idx].Mass;
+					RemovalData.Proxy = AllRemovalsArray[Idx].Proxy;
+					RemovalData.BoundingBox = AllRemovalsArray[Idx].BoundingBox;
+
+					int32 NewIdx = AllRemovalDataArray.Add(FRemovalData());
+					FRemovalData& RemovalDataArrayItem = AllRemovalDataArray[NewIdx];
+					RemovalDataArrayItem = RemovalData;
+
+					AllRemovalIndicesByPhysicsProxy.FindOrAdd(RemovalData.Proxy).Add(FEventManager::EncodeCollisionIndex(NewIdx, false));
+				}
+			
+			});
+	}
 }
\ No newline at end of file
diff --git a/Engine/Source/Runtime/Experimental/Chaos/Private/GeometryCollection/GeometryCollectionAlgo.cpp b/Engine/Source/Runtime/Experimental/Chaos/Private/GeometryCollection/GeometryCollectionAlgo.cpp
--- a/Engine/Source/Runtime/Experimental/Chaos/Private/GeometryCollection/GeometryCollectionAlgo.cpp
+++ b/Engine/Source/Runtime/Experimental/Chaos/Private/GeometryCollection/GeometryCollectionAlgo.cpp
@@ -297,51 +297,83 @@ namespace GeometryCollectionAlgo
 		}
 	}
 
-	void GlobalMatricesRecursive(const int32 Index, const TManagedArray<int32>& Parents, const TManagedArray<FTransform>& Transform, TArray<bool>& IsTransformComputed, TArray<FMatrix>& OutGlobalTransforms)
+	void GlobalMatricesRecursive(const int32 Index, const TManagedArray<int32>& Parents, const TManagedArray<FTransform>& Transform, TArray<bool>& IsTransformComputed, const TManagedArray<float>* Scales, const TManagedArray<FTransform>* MassToLocals, TArray<FMatrix>& OutGlobalTransforms)
 	{
 		if (!IsTransformComputed[Index])
 		{
 			FMatrix Result = Transform[Index].ToMatrixWithScale();
 			if (Parents[Index] != FGeometryCollection::Invalid)
 			{
-				GlobalMatricesRecursive(Parents[Index], Parents, Transform, IsTransformComputed, OutGlobalTransforms);
+				GlobalMatricesRecursive(Parents[Index], Parents, Transform, IsTransformComputed, Scales, MassToLocals, OutGlobalTransforms);
 				Result *= OutGlobalTransforms[Parents[Index]];
 			}
 
-			OutGlobalTransforms[Index] = Result;
+			if (Scales && !FMath::IsNearlyEqual((*Scales)[Index], 1.0f))
+			{
+				check (MassToLocals);
+				OutGlobalTransforms[Index] = (*MassToLocals)[Index].ToInverseMatrixWithScale() * FTransform(FQuat::Identity, FVector(), FVector((*Scales)[Index])).ToMatrixWithScale() * (*MassToLocals)[Index].ToMatrixWithScale() * Result;
+			}
+			else
+			{ 
+				OutGlobalTransforms[Index] = Result;
+			}
+			
 			IsTransformComputed[Index] = true;
 		}
 	}
 
 	// #note: this verison returns an FTransform to support functionality for getting global matrices for an array of indices.
-	void GlobalMatricesRecursive(const int32 Index, const TManagedArray<int32>& Parents, const TManagedArray<FTransform>& Transform, TArray<bool>& IsTransformComputed, TArray<FTransform>& OutGlobalTransforms)
+	void GlobalMatricesRecursive(const int32 Index, const TManagedArray<int32>& Parents, const TManagedArray<FTransform>& Transform, TArray<bool>& IsTransformComputed, const TManagedArray<float>* Scales, const TManagedArray<FTransform>* MassToLocals, TArray<FTransform>& OutGlobalTransforms)
 	{
 		if (!IsTransformComputed[Index])
 		{
 			FTransform Result = Transform[Index];
 			if (Parents[Index] != FGeometryCollection::Invalid)
 			{
-				GlobalMatricesRecursive(Parents[Index], Parents, Transform, IsTransformComputed, OutGlobalTransforms);
+				GlobalMatricesRecursive(Parents[Index], Parents, Transform, IsTransformComputed, Scales, MassToLocals, OutGlobalTransforms);
 				Result *= OutGlobalTransforms[Parents[Index]];
 			}
 
-			OutGlobalTransforms[Index] = Result;
+			if (Scales && !FMath::IsNearlyEqual((*Scales)[Index], 1.0f))
+			{
+				if(MassToLocals)
+				{ 
+					OutGlobalTransforms[Index] = (*MassToLocals)[Index].Inverse() * FTransform(FQuat::Identity, FVector(), FVector((*Scales)[Index])) * (*MassToLocals)[Index] * Result;
+				}
+				else
+				{
+					OutGlobalTransforms[Index] = FTransform(FQuat::Identity, FVector(), FVector((*Scales)[Index])) * Result;
+				}
+			}
+			else
+			{
+				OutGlobalTransforms[Index] = Result;
+			}
+			
 			IsTransformComputed[Index] = true;
 		}
 	}
 
-	FTransform GlobalMatricesRecursiveForIndices(const int32 Index, const TManagedArray<int32>& Parents, const TManagedArray<FTransform>& Transform, TArray<bool>& IsTransformComputed, TArray<FTransform>& TransformCache)
+	FTransform GlobalMatricesRecursiveForIndices(const int32 Index, const TManagedArray<int32>& Parents, const TManagedArray<FTransform>& Transform, TArray<bool>& IsTransformComputed, const TManagedArray<float>* Scales, const TManagedArray<FTransform>* MassToLocals, TArray<FTransform>& TransformCache)
 	{
 		if (!IsTransformComputed[Index])
 		{
 			FTransform Result = Transform[Index];
 			if (Parents[Index] != FGeometryCollection::Invalid)
 			{
-				GlobalMatricesRecursive(Parents[Index], Parents, Transform, IsTransformComputed, TransformCache);
+				GlobalMatricesRecursive(Parents[Index], Parents, Transform, IsTransformComputed, Scales, MassToLocals, TransformCache);
 				Result *= TransformCache[Parents[Index]];
 			}
 
-			TransformCache[Index] = Result;
+			if (Scales && !FMath::IsNearlyEqual((*Scales)[Index], 1.0f))
+			{
+				TransformCache[Index] = (*MassToLocals)[Index].Inverse() * FTransform(FQuat::Identity, FVector(), FVector((*Scales)[Index])) * (*MassToLocals)[Index] * Result;
+			}
+			else
+			{
+				TransformCache[Index] = Result;
+			}
+
 			IsTransformComputed[Index] = true;
 		}
 
@@ -366,6 +398,22 @@ namespace GeometryCollectionAlgo
 		return Transform;
 	}
 
+	void GlobalMatrices(const TManagedArray<FTransform>& RelativeTransforms, const TManagedArray<int32>& Parents, const TArray<int32>& Indices, const TManagedArray<float>& Scales, const TManagedArray<FTransform>& MassToLocals, TArray<FTransform>& OutGlobalTransforms)
+	{
+		check(Scales.Num() == RelativeTransforms.Num());
+		TArray<bool> IsTransformComputed;
+		IsTransformComputed.AddDefaulted(RelativeTransforms.Num());
+
+		TArray<FTransform> TransformCache;
+		TransformCache.SetNumUninitialized(RelativeTransforms.Num(), false);
+
+		OutGlobalTransforms.SetNumUninitialized(Indices.Num(), false);
+		for (int Idx = 0; Idx < Indices.Num(); Idx++)
+		{
+			OutGlobalTransforms[Idx] = GlobalMatricesRecursiveForIndices(Indices[Idx], Parents, RelativeTransforms, IsTransformComputed, &Scales, &MassToLocals, TransformCache);
+		}
+	}
+
 	void GlobalMatrices(const TManagedArray<FTransform>& RelativeTransforms, const TManagedArray<int32>& Parents, const TArray<int32>& Indices, TArray<FTransform>& OutGlobalTransforms)
 	{
 		TArray<bool> IsTransformComputed;
@@ -377,7 +425,7 @@ namespace GeometryCollectionAlgo
 		OutGlobalTransforms.SetNumUninitialized(Indices.Num(), false);
 		for (int Idx = 0; Idx < Indices.Num(); Idx++)
 		{
-			OutGlobalTransforms[Idx] = GlobalMatricesRecursiveForIndices(Indices[Idx], Parents, RelativeTransforms, IsTransformComputed, TransformCache);
+			OutGlobalTransforms[Idx] = GlobalMatricesRecursiveForIndices(Indices[Idx], Parents, RelativeTransforms, IsTransformComputed, nullptr, nullptr, TransformCache);
 		}
 	}
 
@@ -395,6 +443,24 @@ namespace GeometryCollectionAlgo
 		}
 	}
 
+	template<typename MatrixType>
+	void GlobalMatrices(const TManagedArray<FTransform>& RelativeTransforms, const TManagedArray<int32>& Parents, const TManagedArray<float>& Scales, const TManagedArray<FTransform>& MassToLocals, TArray<MatrixType>& OutGlobalTransforms)
+	{
+		int32 NumTransforms = RelativeTransforms.Num();
+		check(Scales.Num() == NumTransforms);
+		check(MassToLocals.Num() == NumTransforms);
+
+		TArray<bool> IsTransformComputed;
+		IsTransformComputed.AddDefaulted(NumTransforms);
+
+		OutGlobalTransforms.SetNumUninitialized(NumTransforms, false);
+
+		for (int BoneIdx = 0; BoneIdx < NumTransforms; ++BoneIdx)
+		{
+			GlobalMatricesRecursive(BoneIdx, Parents, RelativeTransforms, IsTransformComputed, &Scales, &MassToLocals, OutGlobalTransforms);
+		}
+	}
+
 	template<typename MatrixType>
 	void GlobalMatrices(const TManagedArray<FTransform>& RelativeTransforms, const TManagedArray<int32>& Parents, TArray<MatrixType>& OutGlobalTransforms)
 	{
@@ -407,10 +473,12 @@ namespace GeometryCollectionAlgo
 
 		for (int BoneIdx = 0; BoneIdx < NumTransforms; ++BoneIdx)
 		{
-			GlobalMatricesRecursive(BoneIdx, Parents, RelativeTransforms, IsTransformComputed, OutGlobalTransforms);
+			GlobalMatricesRecursive(BoneIdx, Parents, RelativeTransforms, IsTransformComputed, nullptr, nullptr, OutGlobalTransforms);
 		}
 	}
+	template void CHAOS_API GlobalMatrices<FTransform>(const TManagedArray<FTransform>&, const TManagedArray<int32>&, const TManagedArray<float>&, const TManagedArray<FTransform>&, TArray<FTransform>&);
 	template void CHAOS_API GlobalMatrices<FTransform>(const TManagedArray<FTransform>&, const TManagedArray<int32>&, TArray<FTransform>&);
+	template void CHAOS_API GlobalMatrices<FMatrix>(const TManagedArray<FTransform>&, const TManagedArray<int32>&, const TManagedArray<float>&, const TManagedArray<FTransform>&, TArray<FMatrix>&);
 	template void CHAOS_API GlobalMatrices<FMatrix>(const TManagedArray<FTransform>&, const TManagedArray<int32>&, TArray<FMatrix>&);
 
 	void FloodForOverlappedPairs(int Level, int32 BoneIndex, TMap<int32, int32> &BoneToGroup, const TManagedArray<int32>& Levels, const TMap<int32, FBox>& BoundingBoxes, TSet<TTuple<int32, int32>>& OutOverlappedPairs)
diff --git a/Engine/Source/Runtime/Experimental/Chaos/Private/PBDRigidsSolver.cpp b/Engine/Source/Runtime/Experimental/Chaos/Private/PBDRigidsSolver.cpp
--- a/Engine/Source/Runtime/Experimental/Chaos/Private/PBDRigidsSolver.cpp
+++ b/Engine/Source/Runtime/Experimental/Chaos/Private/PBDRigidsSolver.cpp
@@ -242,7 +242,7 @@ namespace Chaos
 				MSolver->ApplyCallbacks_Internal(MSolver->GetSolverTime(), MDeltaTime);	//question: is SolverTime the right thing to pass in here?
 			}
 			MSolver->GetEvolution()->GetRigidClustering().ResetAllClusterBreakings();
-
+			
 			{
 				SCOPE_CYCLE_COUNTER(STAT_UpdateParams);
 				Chaos::FPBDPositionConstraints PositionTarget; // Dummy for now
@@ -346,6 +346,7 @@ namespace Chaos
 				{
 					SCOPE_CYCLE_COUNTER(STAT_FillProducerData);
 					MSolver->GetEventManager()->FillProducerData(MSolver);
+					MSolver->GetEvolution()->ResetAllRemovals();
 				}
 				{
 					SCOPE_CYCLE_COUNTER(STAT_FlipBuffersIfRequired);
diff --git a/Engine/Source/Runtime/Experimental/Chaos/Private/PhysicsProxy/GeometryCollectionPhysicsProxy.cpp b/Engine/Source/Runtime/Experimental/Chaos/Private/PhysicsProxy/GeometryCollectionPhysicsProxy.cpp
--- a/Engine/Source/Runtime/Experimental/Chaos/Private/PhysicsProxy/GeometryCollectionPhysicsProxy.cpp
+++ b/Engine/Source/Runtime/Experimental/Chaos/Private/PhysicsProxy/GeometryCollectionPhysicsProxy.cpp
@@ -1271,6 +1271,22 @@ void FGeometryCollectionPhysicsProxy::GetRelevantParticleHandles(
 #endif
 }
 
+void FGeometryCollectionPhysicsProxy::DisableParticles(TArray<int32>& TransformGroupIndices)
+{
+	check(IsInGameThread());
+
+	if (Chaos::FPhysicsSolver* RBDSolver = GetSolver<Chaos::FPhysicsSolver>())
+	{
+		RBDSolver->EnqueueCommandImmediate([this, RBDSolver, TransformGroupIndices]()
+			{
+				for (int32 TransformIdx : TransformGroupIndices)
+				{
+					RBDSolver->GetEvolution()->DisableParticleWithRemovalEvent(SolverParticleHandles[TransformIdx]);
+				}
+			});
+	}
+}
+
 int32 FGeometryCollectionPhysicsProxy::CalculateHierarchyLevel(const FGeometryDynamicCollection& GeometryCollection, int32 TransformIndex) const
 {
 	int32 Level = 0;
diff --git a/Engine/Source/Runtime/Experimental/Chaos/Private/PhysicsProxy/SolverEventFilters.cpp b/Engine/Source/Runtime/Experimental/Chaos/Private/PhysicsProxy/SolverEventFilters.cpp
--- a/Engine/Source/Runtime/Experimental/Chaos/Private/PhysicsProxy/SolverEventFilters.cpp
+++ b/Engine/Source/Runtime/Experimental/Chaos/Private/PhysicsProxy/SolverEventFilters.cpp
@@ -68,4 +68,21 @@ namespace Chaos
 		return true;
 	}
 
+	bool FSolverRemovalEventFilter::Pass(const FRemovalData& InData) const
+	{
+		if (Settings.MinMass > 0.0f && InData.Mass < Settings.MinMass)
+			return false;
+
+		if (Settings.MinVolume > 0)
+		{
+			FVec3 Extents = InData.BoundingBox.Extents();
+			FReal Volume = Extents[0] * Extents[1] * Extents[2];
+
+			if (Volume < Settings.MinVolume)
+				return false;
+		}
+
+		return true;
+	}
+
 } // namespace Chaos
diff --git a/Engine/Source/Runtime/Experimental/Chaos/Public/Chaos/ExternalCollisionData.h b/Engine/Source/Runtime/Experimental/Chaos/Public/Chaos/ExternalCollisionData.h
--- a/Engine/Source/Runtime/Experimental/Chaos/Public/Chaos/ExternalCollisionData.h
+++ b/Engine/Source/Runtime/Experimental/Chaos/Public/Chaos/ExternalCollisionData.h
@@ -353,6 +353,83 @@ namespace Chaos
 		bool Sleeping;	// if !Sleeping == Awake
 	};
 
+	/*
+	RemovalData passed from the physics solver to subsystems
+	*/
+	struct FRemovalData
+	{
+		FRemovalData()
+			: Location(FVec3((FReal)0.0))
+			, Mass((FReal)0.0)
+			, Proxy(nullptr)
+			, BoundingBox(FAABB3(FVec3((FReal)0.0), FVec3((FReal)0.0)))
+		{}
+
+		FRemovalData(FVec3 InLocation, FReal InMass, IPhysicsProxyBase* InProxy, Chaos::TAABB<FReal, 3>& InBoundingBox)
+			: Location(InLocation)
+			, Mass(InMass)
+			, Proxy(InProxy)
+			, BoundingBox(InBoundingBox)
+		{}
+
+		FVec3 Location;
+		FReal Mass;
+
+		// The pointer to the proxy should be used with caution on the Game Thread.
+		// Ideally we only ever use this as a table key when acquiring related structures.
+		// If we genuinely need to dereference the pointer for any reason, test if it is deleted (nullptr) or
+		// pending deletion: if a call to FPhysScene_Chaos::GetOwningComponent<UPrimitiveComponent>() returns
+		// nullptr, the proxy should not be used.
+		IPhysicsProxyBase* Proxy;
+
+		Chaos::FAABB3 BoundingBox;
+	};
+
+	/*
+	RemovalData used in subsystems
+	*/
+	struct FRemovalDataExt
+	{
+		FRemovalDataExt()
+			: Location(FVec3((FReal)0.0))
+			, Mass((FReal)0.0)
+			, BoundingboxVolume((FReal)-1.0)
+			, BoundingboxExtentMin((FReal)-1.0)
+			, BoundingboxExtentMax((FReal)-1.0)
+			, SurfaceType(-1)
+		{}
+
+		FRemovalDataExt(FVec3 InLocation
+			, FReal InMass
+			, FGeometryParticleHandle* InParticle
+			, FReal InBoundingboxVolume
+			, FReal InBoundingboxExtentMin
+			, FReal InBoundingboxExtentMax
+			, int32 InSurfaceType)
+			: Location(InLocation)
+			, Mass(InMass)
+			, BoundingboxVolume(InBoundingboxVolume)
+			, BoundingboxExtentMin(InBoundingboxExtentMin)
+			, BoundingboxExtentMax(InBoundingboxExtentMax)
+			, SurfaceType(InSurfaceType)
+		{}
+
+		FRemovalDataExt(const FRemovalData& InRemovalData)
+			: Location(InRemovalData.Location)
+			, Mass(InRemovalData.Mass)
+			, BoundingboxVolume((FReal)-1.0)
+			, BoundingboxExtentMin((FReal)-1.0)
+			, BoundingboxExtentMax((FReal)-1.0)
+			, SurfaceType(-1)
+		{}
+
+		FVec3 Location;
+		FReal Mass;
+		FReal BoundingboxVolume;
+		FReal BoundingboxExtentMin, BoundingboxExtentMax;
+		int32 SurfaceType;
+	};
+
 	template<class T, int d>
 	using TCollisionData UE_DEPRECATED(4.27, "Deprecated. this class is to be deleted, use FCollidingData instead") = FCollidingData;
 
diff --git a/Engine/Source/Runtime/Experimental/Chaos/Public/Chaos/PBDRigidsEvolution.h b/Engine/Source/Runtime/Experimental/Chaos/Public/Chaos/PBDRigidsEvolution.h
--- a/Engine/Source/Runtime/Experimental/Chaos/Public/Chaos/PBDRigidsEvolution.h
+++ b/Engine/Source/Runtime/Experimental/Chaos/Public/Chaos/PBDRigidsEvolution.h
@@ -797,6 +797,10 @@ public:
 		bCanStartAsyncTasks = bInCanStartAsyncTasks;
 	}
 
+	CHAOS_API void DisableParticleWithRemovalEvent(FGeometryParticleHandle* Particle);
+	const TArray<FRemovalData>& GetAllRemovals() { return MAllRemovals; }
+	void ResetAllRemovals() { MAllRemovals.Reset(); }
+
 protected:
 	int32 NumConstraints() const
 	{
@@ -928,6 +932,8 @@ protected:
 	TQueue<int32, EQueueMode::Mpsc> IslandsToWake;
 
 	TArray<FUniqueIdx> UniqueIndicesPendingRelease;
+
+	TArray<FRemovalData> MAllRemovals;
 public:
 	//The latest external timestamp we consumed inputs from, assigned to evolution when solver task executes, is used to stamp output data.
 	int32 LatestExternalTimestampConsumed_Internal;	
diff --git a/Engine/Source/Runtime/Experimental/Chaos/Public/Chaos/PBDRigidsEvolutionGBF.h b/Engine/Source/Runtime/Experimental/Chaos/Public/Chaos/PBDRigidsEvolutionGBF.h
--- a/Engine/Source/Runtime/Experimental/Chaos/Public/Chaos/PBDRigidsEvolutionGBF.h
+++ b/Engine/Source/Runtime/Experimental/Chaos/Public/Chaos/PBDRigidsEvolutionGBF.h
@@ -79,6 +79,7 @@ namespace Chaos
 		using Base::AddForceFunction;
 		using Base::AddConstraintRule;
 		using Base::ParticleUpdatePosition;
+		using Base::GetAllRemovals;
 
 		using FGravityForces = FPerParticleGravity;
 		using FCollisionConstraints = FPBDCollisionConstraints;
diff --git a/Engine/Source/Runtime/Experimental/Chaos/Public/EventDefaults.h b/Engine/Source/Runtime/Experimental/Chaos/Public/EventDefaults.h
--- a/Engine/Source/Runtime/Experimental/Chaos/Public/EventDefaults.h
+++ b/Engine/Source/Runtime/Experimental/Chaos/Public/EventDefaults.h
@@ -32,8 +32,15 @@ namespace Chaos
 		 */
 		static void RegisterTrailingEvent(FEventManager& EventManager);
 
-
+		/**
+		 * Register sleeping event gathering function & data type
+		 */
 		static void RegisterSleepingEvent(FEventManager& EventManager);
 
+		/**
+		 * Register removal event gathering function & data type
+		 */
+		static void RegisterRemovalEvent(FEventManager& EventManager);
+
 	};
 }
\ No newline at end of file
diff --git a/Engine/Source/Runtime/Experimental/Chaos/Public/EventManager.h b/Engine/Source/Runtime/Experimental/Chaos/Public/EventManager.h
--- a/Engine/Source/Runtime/Experimental/Chaos/Public/EventManager.h
+++ b/Engine/Source/Runtime/Experimental/Chaos/Public/EventManager.h
@@ -24,7 +24,8 @@ namespace Chaos
 		Collision = 0,
 		Breaking = 1,
 		Trailing = 2,
-		Sleeping = 3
+		Sleeping = 3,
+		Removal = 4
 	};
 
 	typedef int32 FEventID;
diff --git a/Engine/Source/Runtime/Experimental/Chaos/Public/EventsData.h b/Engine/Source/Runtime/Experimental/Chaos/Public/EventsData.h
--- a/Engine/Source/Runtime/Experimental/Chaos/Public/EventsData.h
+++ b/Engine/Source/Runtime/Experimental/Chaos/Public/EventsData.h
@@ -20,6 +20,7 @@ namespace Chaos
 	typedef TArray<FCollidingData> FCollisionDataArray;
 	typedef TArray<FBreakingData> FBreakingDataArray;
 	typedef TArray<FTrailingData> FTrailingDataArray;
+	typedef TArray<FRemovalData> FRemovalDataArray;
 	typedef TArray<FSleepingData> FSleepingDataArray;
 
 	/* Common */
@@ -117,6 +118,32 @@ namespace Chaos
 		FIndicesByPhysicsProxy PhysicsProxyToTrailingIndices;
 	};
 
+	/* Removal */
+
+	/*
+	 * All the removal events for one frame time stamped with the time for that frame
+	 */
+	struct FAllRemovalData : FTimeResource
+	{
+		FAllRemovalData() : AllRemovalArray(FRemovalDataArray()) {}
+
+		void Reset()
+		{
+			AllRemovalArray.Reset();
+		}
+
+		FRemovalDataArray AllRemovalArray;
+	};
+
+
+	struct FRemovalEventData
+	{
+		FRemovalEventData() {}
+
+		FAllRemovalData RemovalData;
+		FIndicesByPhysicsProxy PhysicsProxyToRemovalIndices;
+	};
+
 	struct FSleepingEventData
 	{
 		FSleepingEventData() {}
diff --git a/Engine/Source/Runtime/Experimental/Chaos/Public/GeometryCollection/GeometryCollectionAlgo.h b/Engine/Source/Runtime/Experimental/Chaos/Public/GeometryCollection/GeometryCollectionAlgo.h
--- a/Engine/Source/Runtime/Experimental/Chaos/Public/GeometryCollection/GeometryCollectionAlgo.h
+++ b/Engine/Source/Runtime/Experimental/Chaos/Public/GeometryCollection/GeometryCollectionAlgo.h
@@ -148,6 +148,9 @@ namespace GeometryCollectionAlgo
 	/*
 	*  Global Matrices of the collection, transforms will be resized to fit
 	*/
+	template<typename MatrixType>
+	void CHAOS_API GlobalMatrices(const TManagedArray<FTransform>& RelativeTransforms, const TManagedArray<int32>& Parents, const TManagedArray<float>& Scales, const TManagedArray<FTransform>& MassToLocals, TArray<MatrixType>& Transforms);
+
 	template<typename MatrixType>
 	void CHAOS_API GlobalMatrices(const TManagedArray<FTransform>& RelativeTransforms, const TManagedArray<int32>& Parents, TArray<MatrixType>& Transforms);
 
diff --git a/Engine/Source/Runtime/Experimental/Chaos/Public/GeometryCollection/GeometryCollectionSimulationCoreTypes.h b/Engine/Source/Runtime/Experimental/Chaos/Public/GeometryCollection/GeometryCollectionSimulationCoreTypes.h
--- a/Engine/Source/Runtime/Experimental/Chaos/Public/GeometryCollection/GeometryCollectionSimulationCoreTypes.h
+++ b/Engine/Source/Runtime/Experimental/Chaos/Public/GeometryCollection/GeometryCollectionSimulationCoreTypes.h
@@ -189,6 +189,7 @@ struct FSimulationParameters
 		, bGenerateBreakingData(false)
 		, bGenerateCollisionData(false)
 		, bGenerateTrailingData(false)
+		, bGenerateRemovalsData(false)
 		, RemoveOnFractureEnabled(false)
 		, SimulationFilterData()
 		, QueryFilterData()
@@ -222,6 +223,7 @@ struct FSimulationParameters
 		, bGenerateBreakingData(Other.bGenerateBreakingData)
 		, bGenerateCollisionData(Other.bGenerateCollisionData)
 		, bGenerateTrailingData(Other.bGenerateTrailingData)
+		, bGenerateRemovalsData(Other.bGenerateRemovalsData)
 		, Shared(Other.Shared)
 		, RemoveOnFractureEnabled(false)
 		, SimulationFilterData(Other.SimulationFilterData)
@@ -277,6 +279,7 @@ struct FSimulationParameters
 	bool bGenerateBreakingData;
 	bool bGenerateCollisionData;
 	bool bGenerateTrailingData;
+	bool bGenerateRemovalsData;
 
 	FSharedSimulationParameters Shared;
 
diff --git a/Engine/Source/Runtime/Experimental/Chaos/Public/PBDRigidsSolver.h b/Engine/Source/Runtime/Experimental/Chaos/Public/PBDRigidsSolver.h
--- a/Engine/Source/Runtime/Experimental/Chaos/Public/PBDRigidsSolver.h
+++ b/Engine/Source/Runtime/Experimental/Chaos/Public/PBDRigidsSolver.h
@@ -199,9 +199,11 @@ namespace Chaos
 			GetEvolution()->GetRigidClustering().SetGenerateClusterBreaking(bDoGenerate);
 		}
 		void SetGenerateTrailingData(bool bDoGenerate) { GetEventFilters()->SetGenerateTrailingEvents(bDoGenerate); }
+		void SetGenerateRemovalData(bool bDoGenerate) { GetEventFilters()->SetGenerateRemovalEvents(bDoGenerate); }
 		void SetCollisionFilterSettings(const FSolverCollisionFilterSettings& InCollisionFilterSettings) { GetEventFilters()->GetCollisionFilter()->UpdateFilterSettings(InCollisionFilterSettings); }
 		void SetBreakingFilterSettings(const FSolverBreakingFilterSettings& InBreakingFilterSettings) { GetEventFilters()->GetBreakingFilter()->UpdateFilterSettings(InBreakingFilterSettings); }
 		void SetTrailingFilterSettings(const FSolverTrailingFilterSettings& InTrailingFilterSettings) { GetEventFilters()->GetTrailingFilter()->UpdateFilterSettings(InTrailingFilterSettings); }
+		void SetRemovalFilterSettings(const FSolverRemovalFilterSettings& InRemovalFilterSettings) { GetEventFilters()->GetRemovalFilter()->UpdateFilterSettings(InRemovalFilterSettings); }
 
 		/**/
 		FJointConstraints& GetJointConstraints() { return JointConstraints; }
diff --git a/Engine/Source/Runtime/Experimental/Chaos/Public/PhysicsProxy/GeometryCollectionPhysicsProxy.h b/Engine/Source/Runtime/Experimental/Chaos/Public/PhysicsProxy/GeometryCollectionPhysicsProxy.h
--- a/Engine/Source/Runtime/Experimental/Chaos/Public/PhysicsProxy/GeometryCollectionPhysicsProxy.h
+++ b/Engine/Source/Runtime/Experimental/Chaos/Public/PhysicsProxy/GeometryCollectionPhysicsProxy.h
@@ -237,6 +237,8 @@ public:
 
 	bool GetIsObjectDynamic() const { return IsObjectDynamic; }
 
+	void DisableParticles(TArray<int32>& TransformGroupIndices);
+
 protected:
 	/**
 	 * Build a physics thread cluster parent particle.
diff --git a/Engine/Source/Runtime/Experimental/Chaos/Public/SolverEventFilters.h b/Engine/Source/Runtime/Experimental/Chaos/Public/SolverEventFilters.h
--- a/Engine/Source/Runtime/Experimental/Chaos/Public/SolverEventFilters.h
+++ b/Engine/Source/Runtime/Experimental/Chaos/Public/SolverEventFilters.h
@@ -97,6 +97,30 @@
 
 	};
 
+	USTRUCT(BlueprintType)
+	struct CHAOS_API FSolverRemovalFilterSettings
+	{
+		GENERATED_USTRUCT_BODY()
+
+		FSolverRemovalFilterSettings()
+			: FilterEnabled(false)
+			, MinMass(0.0f)
+			, MinVolume(0.0f)
+		{}
+
+		/** Filter is enabled. */
+		UPROPERTY(EditAnywhere, BlueprintReadWrite, Category = "ChaosPhysics|BreakingData Generation")
+		bool FilterEnabled;
+
+		/** The minimum mass threshold for the results (compared with min of particle 1 mass and particle 2 mass). */
+		UPROPERTY(EditAnywhere, BlueprintReadWrite, Category = "ChaosPhysics|BreakingData Generation", meta = (DisplayName = "Min Mass Threshold"))
+		float MinMass;
+
+		UPROPERTY(EditAnywhere, BlueprintReadWrite, Category = "ChaosPhysics|BreakingData Generation", meta = (DisplayName = "Min Volume Threshold"))
+		float MinVolume;
+
+	};
+
 
 namespace Chaos
 {
@@ -107,6 +131,8 @@ namespace Chaos
 
 	struct FBreakingData;
 
+	struct FRemovalData;
+
 	class CHAOS_API FSolverCollisionEventFilter
 	{
 	public:
@@ -146,6 +172,19 @@ namespace Chaos
 		FSolverBreakingFilterSettings Settings;
 	};
 
+	class CHAOS_API FSolverRemovalEventFilter
+	{
+	public:
+		FSolverRemovalEventFilter() {}
+		FSolverRemovalEventFilter(const FSolverRemovalFilterSettings& InSettings) : Settings(InSettings) {}
+
+		bool Pass(const Chaos::FRemovalData& InData) const;
+		bool Enabled() const { return Settings.FilterEnabled; }
+		void UpdateFilterSettings(const FSolverRemovalFilterSettings& InSettings) { Settings = InSettings; }
+
+		FSolverRemovalFilterSettings Settings;
+	};
+
 
 	/**
 	 * Container for the Solver Event Filters that have settings exposed through the Solver Actor
@@ -157,38 +196,46 @@ namespace Chaos
 			: CollisionFilter(new FSolverCollisionEventFilter())
 			, BreakingFilter(new FSolverBreakingEventFilter())
 			, TrailingFilter(new FSolverTrailingEventFilter())
+			, RemovalFilter(new FSolverRemovalEventFilter())
 			, CollisionEventsEnabled(false)
 			, BreakingEventsEnabled(false)
 			, TrailingEventsEnabled(false)
+			, RemovalEventsEnabled(false)
 		{}
 
 		void SetGenerateCollisionEvents(bool bDoGenerate) { CollisionEventsEnabled = bDoGenerate; }
 		void SetGenerateBreakingEvents(bool bDoGenerate) { BreakingEventsEnabled = bDoGenerate; }
 		void SetGenerateTrailingEvents(bool bDoGenerate) { TrailingEventsEnabled = bDoGenerate; }
+		void SetGenerateRemovalEvents(bool bDoGenerate) { RemovalEventsEnabled = bDoGenerate; }
 
 		/* Const access */
 		FSolverCollisionEventFilter* GetCollisionFilter() const { return CollisionFilter.Get(); }
 		FSolverBreakingEventFilter* GetBreakingFilter() const { return BreakingFilter.Get(); }
 		FSolverTrailingEventFilter* GetTrailingFilter() const { return TrailingFilter.Get(); }
+		FSolverRemovalEventFilter* GetRemovalFilter() const { return RemovalFilter.Get(); }
 
 		/* non-const access */
 		FSolverCollisionEventFilter* GetCollisionFilter() { return CollisionFilter.Get(); }
 		FSolverBreakingEventFilter* GetBreakingFilter() { return BreakingFilter.Get(); }
 		FSolverTrailingEventFilter* GetTrailingFilter() { return TrailingFilter.Get(); }
+		FSolverRemovalEventFilter* GetRemovalFilter() { return RemovalFilter.Get(); }
 
 		bool IsCollisionEventEnabled() const { return CollisionEventsEnabled; }
 		bool IsBreakingEventEnabled() const { return BreakingEventsEnabled; }
 		bool IsTrailingEventEnabled() const { return TrailingEventsEnabled; }
+		bool IsRemovalEventEnabled() const { return RemovalEventsEnabled; }
 
 	private:
 
 		TUniquePtr<FSolverCollisionEventFilter> CollisionFilter;
 		TUniquePtr<FSolverBreakingEventFilter> BreakingFilter;
 		TUniquePtr<FSolverTrailingEventFilter> TrailingFilter;
+		TUniquePtr<FSolverRemovalEventFilter> RemovalFilter;
 
 		bool CollisionEventsEnabled;
 		bool BreakingEventsEnabled;
 		bool TrailingEventsEnabled;
+		bool RemovalEventsEnabled;
 	};
 
 
diff --git a/Engine/Source/Runtime/Experimental/ChaosSolverEngine/Private/Chaos/ChaosGameplayEventDispatcher.cpp b/Engine/Source/Runtime/Experimental/ChaosSolverEngine/Private/Chaos/ChaosGameplayEventDispatcher.cpp
--- a/Engine/Source/Runtime/Experimental/ChaosSolverEngine/Private/Chaos/ChaosGameplayEventDispatcher.cpp
+++ b/Engine/Source/Runtime/Experimental/ChaosSolverEngine/Private/Chaos/ChaosGameplayEventDispatcher.cpp
@@ -26,6 +26,14 @@ FChaosBreakEvent::FChaosBreakEvent()
 
 }
 
+FChaosRemovalEvent::FChaosRemovalEvent()
+	: Component(nullptr)
+	, Location(FVector::ZeroVector)
+	, Mass(0.0f)
+{
+
+}
+
 void UChaosGameplayEventDispatcher::OnRegister()
 {
 	Super::OnRegister();
@@ -55,6 +63,21 @@ static void DispatchPendingBreakEvents(TArray<FChaosBreakEvent> const& Events, T
 	}
 }
 
+static void DispatchPendingRemovalEvents(TArray<FChaosRemovalEvent> const& Events, TMap<TObjectPtr<UPrimitiveComponent>, FRemovalEventCallbackWrapper> const& Registrations)
+{
+	for (FChaosRemovalEvent const& E : Events)
+	{
+		if (E.Component)
+		{
+			const FRemovalEventCallbackWrapper* const Callback = Registrations.Find(E.Component);
+			if (Callback)
+			{
+				Callback->RemovalEventCallback(E);
+			}
+		}
+	}
+}
+
 static void SetCollisionInfoFromComp(FRigidBodyCollisionInfo& Info, UPrimitiveComponent* Comp)
 {
 	if (Comp)
@@ -204,6 +227,23 @@ void UChaosGameplayEventDispatcher::UnRegisterForBreakEvents(UPrimitiveComponent
 	}
 }
 
+void UChaosGameplayEventDispatcher::RegisterForRemovalEvents(UPrimitiveComponent* Component, FOnRemovalEventCallback InFunc)
+{
+	if (Component)
+	{
+		FRemovalEventCallbackWrapper F = { InFunc };
+		RemovalEventRegistrations.Add(Component, F);
+	}
+}
+
+void UChaosGameplayEventDispatcher::UnRegisterForRemovalEvents(UPrimitiveComponent* Component)
+{
+	if (Component)
+	{
+		RemovalEventRegistrations.Remove(Component);
+	}
+}
+
 void UChaosGameplayEventDispatcher::DispatchPendingWakeNotifies()
 {
 	for (auto MapItr = PendingSleepNotifies.CreateIterator(); MapItr; ++MapItr)
@@ -229,6 +269,7 @@ void UChaosGameplayEventDispatcher::RegisterChaosEvents()
 			EventManager->RegisterHandler<Chaos::FCollisionEventData>(Chaos::EEventType::Collision, this, &UChaosGameplayEventDispatcher::HandleCollisionEvents);
 			EventManager->RegisterHandler<Chaos::FBreakingEventData>(Chaos::EEventType::Breaking, this, &UChaosGameplayEventDispatcher::HandleBreakingEvents);
 			EventManager->RegisterHandler<Chaos::FSleepingEventData>(Chaos::EEventType::Sleeping, this, &UChaosGameplayEventDispatcher::HandleSleepingEvents);
+			EventManager->RegisterHandler<Chaos::FRemovalEventData>(Chaos::EEventType::Removal, this, &UChaosGameplayEventDispatcher::HandleRemovalEvents);
 		}
 	}
 
@@ -248,6 +289,7 @@ void UChaosGameplayEventDispatcher::UnregisterChaosEvents()
 				EventManager->UnregisterHandler(Chaos::EEventType::Collision, this);
 				EventManager->UnregisterHandler(Chaos::EEventType::Breaking, this);
 				EventManager->UnregisterHandler(Chaos::EEventType::Sleeping, this);
+				EventManager->UnregisterHandler(Chaos::EEventType::Removal, this);
 			}
 		}
 	}
@@ -464,3 +506,42 @@ void UChaosGameplayEventDispatcher::AddPendingSleepingNotify(FBodyInstance* Body
 {
 	PendingSleepNotifies.FindOrAdd(BodyInstance) = SleepEventType;
 }
+
+void UChaosGameplayEventDispatcher::HandleRemovalEvents(const Chaos::FRemovalEventData& Event)
+{
+	// REMOVAL EVENTS
+
+	TArray<FChaosRemovalEvent> PendingRemovalEvents;
+
+	const float RemovalDataTimestamp = Event.RemovalData.TimeCreated;
+	if (RemovalDataTimestamp > LastRemovalDataTime)
+	{
+		LastRemovalDataTime = RemovalDataTimestamp;
+
+		Chaos::FRemovalDataArray const& RemovalData = Event.RemovalData.AllRemovalArray;
+
+		const int32 NumRemovals = RemovalData.Num();
+		if (NumRemovals > 0)
+		{
+			for (Chaos::FRemovalData const& RemovalDataItem : RemovalData)
+			{
+				if (RemovalDataItem.Proxy)
+				{
+					UPrimitiveComponent* const PrimComp = Cast<UPrimitiveComponent>(RemovalDataItem.Proxy->GetOwner());
+					if (PrimComp && RemovalEventRegistrations.Contains(PrimComp))
+					{
+						// queue them up so we can release the physics data before trigging BP events
+						FChaosRemovalEvent& RemovalEvent = PendingRemovalEvents.AddZeroed_GetRef();
+						RemovalEvent.Component = PrimComp;
+						RemovalEvent.Location = RemovalDataItem.Location;
+						RemovalEvent.Mass = RemovalDataItem.Mass;
+					}
+				}
+			}
+
+			DispatchPendingRemovalEvents(PendingRemovalEvents, RemovalEventRegistrations);
+		}
+	}
+
+}
+
diff --git a/Engine/Source/Runtime/Experimental/ChaosSolverEngine/Public/Chaos/ChaosGameplayEventDispatcher.h b/Engine/Source/Runtime/Experimental/ChaosSolverEngine/Public/Chaos/ChaosGameplayEventDispatcher.h
--- a/Engine/Source/Runtime/Experimental/ChaosSolverEngine/Public/Chaos/ChaosGameplayEventDispatcher.h
+++ b/Engine/Source/Runtime/Experimental/ChaosSolverEngine/Public/Chaos/ChaosGameplayEventDispatcher.h
@@ -14,6 +14,7 @@ namespace Chaos
 	struct FCollisionEventData;
 	struct FBreakingEventData;
 	struct FSleepingEventData;
+	struct FRemovalEventData;
 }
 
 USTRUCT(BlueprintType)
@@ -41,6 +42,25 @@ public:
 	float Mass;
 };
 
+USTRUCT(BlueprintType)
+struct CHAOSSOLVERENGINE_API FChaosRemovalEvent
+{
+	GENERATED_BODY()
+
+public:
+
+	FChaosRemovalEvent();
+
+	UPROPERTY(BlueprintReadOnly, Category = "Removal Event")
+	TObjectPtr<UPrimitiveComponent> Component = nullptr;
+
+	UPROPERTY(BlueprintReadOnly, Category = "Removal Event")
+	FVector Location;
+
+	UPROPERTY(BlueprintReadOnly, Category = "Removal Event")
+	float Mass;
+};
+
 
 typedef TFunction<void(const FChaosBreakEvent&)> FOnBreakEventCallback;
 
@@ -54,6 +74,18 @@ public:
 	FOnBreakEventCallback BreakEventCallback;
 };
 
+typedef TFunction<void(const FChaosRemovalEvent&)> FOnRemovalEventCallback;
+
+/** UStruct wrapper so we can store the TFunction in a TMap */
+USTRUCT()
+struct CHAOSSOLVERENGINE_API FRemovalEventCallbackWrapper
+{
+	GENERATED_BODY()
+
+public:
+	FOnRemovalEventCallback RemovalEventCallback;
+};
+
 /** UStruct wrapper so we can store the TSet in a TMap */
 USTRUCT()
 struct FChaosHandlerSet
@@ -133,6 +165,9 @@ public:
 	void RegisterForBreakEvents(UPrimitiveComponent* Component, FOnBreakEventCallback InFunc);
 	void UnRegisterForBreakEvents(UPrimitiveComponent* Component);
 
+	void RegisterForRemovalEvents(UPrimitiveComponent* Component, FOnRemovalEventCallback InFunc);
+	void UnRegisterForRemovalEvents(UPrimitiveComponent* Component);
+
 private:
 
  	UPROPERTY()
@@ -141,8 +176,12 @@ private:
 	UPROPERTY()
 	TMap<TObjectPtr<UPrimitiveComponent>, FBreakEventCallbackWrapper> BreakEventRegistrations;
 
+	UPROPERTY()
+	TMap<TObjectPtr<UPrimitiveComponent>, FRemovalEventCallbackWrapper> RemovalEventRegistrations;
+
 	float LastCollisionDataTime = -1.f;
 	float LastBreakingDataTime = -1.f;
+	float LastRemovalDataTime = -1.f;
 
 	void DispatchPendingCollisionNotifies();
 	void DispatchPendingWakeNotifies();
@@ -155,6 +194,7 @@ private:
 	void HandleBreakingEvents(const Chaos::FBreakingEventData& BreakingData);
 	void HandleSleepingEvents(const Chaos::FSleepingEventData& SleepingData);
 	void AddPendingSleepingNotify(FBodyInstance* BodyInstance, ESleepEvent SleepEventType);
+	void HandleRemovalEvents(const Chaos::FRemovalEventData& RemovalData);
 
 };
 
diff --git a/Engine/Source/Runtime/Experimental/GeometryCollectionEngine/Private/ChaosBlueprint.cpp b/Engine/Source/Runtime/Experimental/GeometryCollectionEngine/Private/ChaosBlueprint.cpp
--- a/Engine/Source/Runtime/Experimental/GeometryCollectionEngine/Private/ChaosBlueprint.cpp
+++ b/Engine/Source/Runtime/Experimental/GeometryCollectionEngine/Private/ChaosBlueprint.cpp
@@ -10,7 +10,7 @@
 #define DISPATCH_BLUEPRINTS_IMMEDIATE 1
 
 UChaosDestructionListener::UChaosDestructionListener(FObjectInitializer const& ObjectInitializer)
-	: Super(ObjectInitializer), LastCollisionDataTimeStamp(-1.f), LastBreakingDataTimeStamp(-1.f), LastTrailingDataTimeStamp(-1.f)
+	: Super(ObjectInitializer), LastCollisionDataTimeStamp(-1.f), LastBreakingDataTimeStamp(-1.f), LastTrailingDataTimeStamp(-1.f), LastRemovalDataTimeStamp(-1.f)
 {
 	bUseAttachParentBound = true;
 	bAutoActivate = true;	
@@ -25,6 +25,7 @@ UChaosDestructionListener::UChaosDestructionListener(FObjectInitializer const& O
 	SetCollisionFilter(MakeShareable(new FChaosCollisionEventFilter(&CollisionEventRequestSettings)));
 	SetBreakingFilter(MakeShareable(new FChaosBreakingEventFilter(&BreakingEventRequestSettings)));
 	SetTrailingFilter(MakeShareable(new FChaosTrailingEventFilter(&TrailingEventRequestSettings)));
+	SetRemovalFilter(MakeShareable(new FChaosRemovalEventFilter(&RemovalEventRequestSettings)));
 }
 
 void UChaosDestructionListener::ClearEvents()
@@ -125,7 +126,8 @@ bool UChaosDestructionListener::IsEventListening() const
 {
 	return  bIsCollisionEventListeningEnabled ||
 			bIsBreakingEventListeningEnabled  ||
-			bIsTrailingEventListeningEnabled;
+			bIsTrailingEventListeningEnabled  ||
+			bIsRemovalEventListeningEnabled;
 }
 
 void UChaosDestructionListener::UpdateTransformSettings()
@@ -135,7 +137,8 @@ void UChaosDestructionListener::UpdateTransformSettings()
 	{
 		bWantsOnUpdateTransform = CollisionEventRequestSettings.SortMethod == EChaosCollisionSortMethod::SortByNearestFirst ||
 								  BreakingEventRequestSettings.SortMethod == EChaosBreakingSortMethod::SortByNearestFirst ||
-							      TrailingEventRequestSettings.SortMethod == EChaosTrailingSortMethod::SortByNearestFirst;
+							      TrailingEventRequestSettings.SortMethod == EChaosTrailingSortMethod::SortByNearestFirst ||
+								  RemovalEventRequestSettings.SortMethod == EChaosRemovalSortMethod::SortByNearestFirst;
 	}
 	else
 	{
@@ -346,6 +349,12 @@ void UChaosDestructionListener::SetTrailingEventRequestSettings(const FChaosTrai
 	UpdateTransformSettings();
 }
 
+void UChaosDestructionListener::SetRemovalEventRequestSettings(const FChaosRemovalEventRequestSettings& InSettings)
+{
+	RemovalEventRequestSettings = InSettings;
+	UpdateTransformSettings();
+}
+
 void UChaosDestructionListener::SetCollisionEventEnabled(bool bIsEnabled)
 {
 	bIsCollisionEventListeningEnabled = bIsEnabled;
@@ -364,6 +373,12 @@ void UChaosDestructionListener::SetTrailingEventEnabled(bool bIsEnabled)
 	UpdateTransformSettings();
 }
 
+void UChaosDestructionListener::SetRemovalEventEnabled(bool bIsEnabled)
+{
+	bIsRemovalEventListeningEnabled = bIsEnabled;
+	UpdateTransformSettings();
+}
+
 void UChaosDestructionListener::SortCollisionEvents(TArray<FChaosCollisionEventData>& CollisionEvents, EChaosCollisionSortMethod SortMethod)
 {
 	if (ChaosCollisionFilter.IsValid())
@@ -388,6 +403,14 @@ void UChaosDestructionListener::SortTrailingEvents(TArray<FChaosTrailingEventDat
 	}
 }
 
+void UChaosDestructionListener::SortRemovalEvents(TArray<FChaosRemovalEventData>& RemovalEvents, EChaosRemovalSortMethod SortMethod)
+{
+	if (ChaosRemovalFilter.IsValid())
+	{
+		ChaosRemovalFilter->SortEvents(RemovalEvents, SortMethod, GetComponentTransform());
+	}
+}
+
 void UChaosDestructionListener::RegisterChaosEvents(FPhysScene* Scene)
 {
 #if WITH_CHAOS
@@ -396,6 +419,7 @@ void UChaosDestructionListener::RegisterChaosEvents(FPhysScene* Scene)
 	EventManager->RegisterHandler<Chaos::FCollisionEventData>(Chaos::EEventType::Collision, this, &UChaosDestructionListener::HandleCollisionEvents);
 	EventManager->RegisterHandler<Chaos::FBreakingEventData>(Chaos::EEventType::Breaking, this, &UChaosDestructionListener::HandleBreakingEvents);
 	EventManager->RegisterHandler<Chaos::FTrailingEventData>(Chaos::EEventType::Trailing, this, &UChaosDestructionListener::HandleTrailingEvents);
+	EventManager->RegisterHandler<Chaos::FRemovalEventData>(Chaos::EEventType::Removal, this, &UChaosDestructionListener::HandleRemovalEvents);
 #endif
 }
 
@@ -407,6 +431,7 @@ void UChaosDestructionListener::UnregisterChaosEvents(FPhysScene* Scene)
 	EventManager->UnregisterHandler(Chaos::EEventType::Collision, this);
 	EventManager->UnregisterHandler(Chaos::EEventType::Breaking, this);
 	EventManager->UnregisterHandler(Chaos::EEventType::Trailing, this);
+	EventManager->UnregisterHandler(Chaos::EEventType::Removal, this);
 #endif
 }
 
@@ -418,6 +443,7 @@ void UChaosDestructionListener::RegisterChaosEvents(TSharedPtr<FPhysScene_Chaos>
 	EventManager->RegisterHandler<Chaos::FCollisionEventData>(Chaos::EEventType::Collision, this, &UChaosDestructionListener::HandleCollisionEvents);
 	EventManager->RegisterHandler<Chaos::FBreakingEventData>(Chaos::EEventType::Breaking, this, &UChaosDestructionListener::HandleBreakingEvents);
 	EventManager->RegisterHandler<Chaos::FTrailingEventData>(Chaos::EEventType::Trailing, this, &UChaosDestructionListener::HandleTrailingEvents);
+	EventManager->RegisterHandler<Chaos::FRemovalEventData>(Chaos::EEventType::Removal, this, &UChaosDestructionListener::HandleRemovalEvents);
 #endif
 }
 
@@ -429,6 +455,7 @@ void UChaosDestructionListener::UnregisterChaosEvents(TSharedPtr<FPhysScene_Chao
 	EventManager->UnregisterHandler(Chaos::EEventType::Collision, this);
 	EventManager->UnregisterHandler(Chaos::EEventType::Breaking, this);
 	EventManager->UnregisterHandler(Chaos::EEventType::Trailing, this);
+	EventManager->UnregisterHandler(Chaos::EEventType::Removal, this);
 #endif
 }
 
@@ -499,3 +526,25 @@ void UChaosDestructionListener::HandleTrailingEvents(const Chaos::FTrailingEvent
 #endif
 	}
 }
+
+void UChaosDestructionListener::HandleRemovalEvents(const Chaos::FRemovalEventData& Event)
+{
+	if (bIsRemovalEventListeningEnabled)
+	{
+		int NumRemovals = Event.RemovalData.AllRemovalArray.Num();
+		RawRemovalDataArray.Append(Event.RemovalData.AllRemovalArray.GetData(), NumRemovals);
+
+#if DISPATCH_BLUEPRINTS_IMMEDIATE
+		if (ChaosRemovalFilter.IsValid())
+		{
+			ChaosRemovalFilter->FilterEvents(ChaosComponentTransform, RawRemovalDataArray);
+
+			if (ChaosRemovalFilter->GetNumEvents() > 0 && OnRemovalEvents.IsBound())
+			{
+				OnRemovalEvents.Broadcast(ChaosRemovalFilter->GetFilteredResults());
+			}
+		}
+		RawRemovalDataArray.Reset();
+#endif
+	}
+}
diff --git a/Engine/Source/Runtime/Experimental/GeometryCollectionEngine/Private/ChaosRemovalEventFilter.cpp b/Engine/Source/Runtime/Experimental/GeometryCollectionEngine/Private/ChaosRemovalEventFilter.cpp
new file mode 100644
--- /dev/null
+++ b/Engine/Source/Runtime/Experimental/GeometryCollectionEngine/Private/ChaosRemovalEventFilter.cpp
@@ -0,0 +1,92 @@
+// Copyright Epic Games, Inc. All Rights Reserved.
+
+#include "ChaosRemovalEventFilter.h"
+
+void FChaosRemovalEventFilter::FilterEvents(const FTransform& ChaosComponentTransform, const Chaos::FRemovalDataArray& RawRemovalDataArray)
+{
+	FilteredDataArray.Reset();
+
+	{
+		check(RemovalEventRequestSettings);
+
+		// First filter the data 
+
+		float MaxDistanceSquared = RemovalEventRequestSettings->MaxDistance * RemovalEventRequestSettings->MaxDistance;
+
+		for (const Chaos::FRemovalData& RemovalData : RawRemovalDataArray)
+		{
+			if (RemovalEventRequestSettings->MinMass > 0.0f && RemovalData.Mass < RemovalEventRequestSettings->MinMass)
+			{
+				continue;
+			}
+
+			if (MaxDistanceSquared > 0.0f)
+			{
+				float DistSquared = FVector::DistSquared(RemovalData.Location, ChaosComponentTransform.GetLocation());
+				if (DistSquared > MaxDistanceSquared)
+				{
+					continue;
+				}
+			}
+
+			FChaosRemovalEventData NewData;
+			NewData.Location = RemovalData.Location;
+			NewData.Mass = RemovalData.Mass;
+#if TODO_MAKE_A_BLUEPRINT_WAY_TO_ACCESS_PARTICLES
+			NewData.ParticleIndex = RemovalData.Particle;
+#endif
+
+			FilteredDataArray.Add(NewData);
+
+			if (RemovalEventRequestSettings->MaxNumberOfResults > 0 && FilteredDataArray.Num() == RemovalEventRequestSettings->MaxNumberOfResults)
+			{
+				break;
+			}
+		}
+
+		SortEvents(FilteredDataArray, RemovalEventRequestSettings->SortMethod, ChaosComponentTransform);
+	}
+}
+
+void FChaosRemovalEventFilter::SortEvents(TArray<FChaosRemovalEventData>& InOutRemovalEvents, EChaosRemovalSortMethod SortMethod, const FTransform& InTransform)
+{
+	struct FSortRemovalByMassMaxToMin
+	{
+		FORCEINLINE bool operator()(const FChaosRemovalEventData& Lhs, const FChaosRemovalEventData& Rhs) const
+		{
+			return Lhs.Mass > Rhs.Mass;
+		}
+	};
+
+	struct FSortRemovalByNearestFirst
+	{
+		FSortRemovalByNearestFirst(const FVector& InListenerLocation)
+			: ListenerLocation(InListenerLocation)
+		{}
+
+		FORCEINLINE bool operator()(const FChaosRemovalEventData& Lhs, const FChaosRemovalEventData& Rhs) const
+		{
+			float DistSquaredLhs = FVector::DistSquared(Lhs.Location, ListenerLocation);
+			float DistSquaredRhs = FVector::DistSquared(Rhs.Location, ListenerLocation);
+
+			return DistSquaredLhs < DistSquaredRhs;
+		}
+
+		FVector ListenerLocation;
+	};
+
+	// Apply the sort
+	switch (SortMethod)
+	{
+	case EChaosRemovalSortMethod::SortNone:
+		break;
+
+	case EChaosRemovalSortMethod::SortByHighestMass:
+		InOutRemovalEvents.Sort(FSortRemovalByMassMaxToMin());
+		break;
+
+	case EChaosRemovalSortMethod::SortByNearestFirst:
+		InOutRemovalEvents.Sort(FSortRemovalByNearestFirst(InTransform.GetLocation()));
+		break;
+	}
+}
diff --git a/Engine/Source/Runtime/Experimental/GeometryCollectionEngine/Private/GeometryCollection/GeometryCollectionComponent.cpp b/Engine/Source/Runtime/Experimental/GeometryCollectionEngine/Private/GeometryCollection/GeometryCollectionComponent.cpp
--- a/Engine/Source/Runtime/Experimental/GeometryCollectionEngine/Private/GeometryCollection/GeometryCollectionComponent.cpp
+++ b/Engine/Source/Runtime/Experimental/GeometryCollectionEngine/Private/GeometryCollection/GeometryCollectionComponent.cpp
@@ -12,6 +12,7 @@
 #include "GeometryCollection/GeometryCollectionSQAccelerator.h"
 #include "GeometryCollection/GeometryCollectionUtility.h"
 #include "GeometryCollection/GeometryCollectionClusteringUtility.h"
+#include "GeometryCollection/GeometryCollectionProximityUtility.h"
 #include "GeometryCollection/GeometryCollectionCache.h"
 #include "GeometryCollection/GeometryCollectionActor.h"
 #include "GeometryCollection/GeometryCollectionDebugDrawComponent.h"
@@ -234,6 +235,7 @@ UGeometryCollectionComponent::UGeometryCollectionComponent(const FObjectInitiali
 	, CachePlayback(false)
 	, bNotifyBreaks(false)
 	, bNotifyCollisions(false)
+	, bNotifyRemovals(false)
 	, bShowBoneColors(false)
 	, bEnableReplication(false)
 	, bEnableAbandonAfterLevel(false)
@@ -614,6 +616,15 @@ void UGeometryCollectionComponent::SetNotifyBreaks(bool bNewNotifyBreaks)
 	}
 }
 
+void UGeometryCollectionComponent::SetNotifyRemovals(bool bNewNotifyRemovals)
+{
+	if (bNotifyRemovals != bNewNotifyRemovals)
+	{
+		bNotifyRemovals = bNewNotifyRemovals;
+		UpdateRemovalEventRegistration();
+	}
+}
+
 FBodyInstance* UGeometryCollectionComponent::GetBodyInstance(FName BoneName /*= NAME_None*/, bool bGetWelded /*= true*/, int32 Index /*=INDEX_NONE*/) const
 {
 	return nullptr;// const_cast<FBodyInstance*>(&DummyBodyInstance);
@@ -653,6 +664,18 @@ void UGeometryCollectionComponent::DispatchBreakEvent(const FChaosBreakEvent& Ev
 	}
 }
 
+void UGeometryCollectionComponent::DispatchRemovalEvent(const FChaosRemovalEvent& Event)
+{
+	// native
+	NotifyRemoval(Event);
+
+	// bp
+	if (OnChaosRemovalEvent.IsBound())
+	{
+		OnChaosRemovalEvent.Broadcast(Event);
+	}
+}
+
 bool UGeometryCollectionComponent::DoCustomNavigableGeometryExport(FNavigableGeometryExport& GeomExport) const
 {
 	if(!RestCollection)
@@ -1002,6 +1025,14 @@ static void DispatchGeometryCollectionBreakEvent(const FChaosBreakEvent& Event)
 	}
 }
 
+static void DispatchGeometryCollectionRemovalEvent(const FChaosRemovalEvent& Event)
+{
+	if (UGeometryCollectionComponent* const GC = Cast<UGeometryCollectionComponent>(Event.Component))
+	{
+		GC->DispatchRemovalEvent(Event);
+	}
+}
+
 void UGeometryCollectionComponent::DispatchChaosPhysicsCollisionBlueprintEvents(const FChaosPhysicsCollisionInfo& CollisionInfo)
 {
 	ReceivePhysicsCollision(CollisionInfo);
@@ -1011,7 +1042,7 @@ void UGeometryCollectionComponent::DispatchChaosPhysicsCollisionBlueprintEvents(
 // call when first registering
 void UGeometryCollectionComponent::RegisterForEvents()
 {
-	if (BodyInstance.bNotifyRigidBodyCollision || bNotifyBreaks || bNotifyCollisions)
+	if (BodyInstance.bNotifyRigidBodyCollision || bNotifyBreaks || bNotifyCollisions || bNotifyRemovals)
 	{
 		Chaos::FPhysicsSolver* Solver = GetWorld()->GetPhysicsScene()->GetSolver();
 		if (Solver)
@@ -1036,6 +1067,17 @@ void UGeometryCollectionComponent::RegisterForEvents()
 					});
 
 			}
+
+			if (bNotifyRemovals)
+			{
+				EventDispatcher->RegisterForRemovalEvents(this, &DispatchGeometryCollectionRemovalEvent);
+
+				Solver->EnqueueCommandImmediate([Solver]()
+					{
+						Solver->SetGenerateRemovalData(true);
+					});
+
+			}
 		}
 	}
 }
@@ -1064,6 +1106,18 @@ void UGeometryCollectionComponent::UpdateBreakEventRegistration()
 	}
 }
 
+void UGeometryCollectionComponent::UpdateRemovalEventRegistration()
+{
+	if (bNotifyRemovals)
+	{
+		EventDispatcher->RegisterForRemovalEvents(this, &DispatchGeometryCollectionRemovalEvent);
+	}
+	else
+	{
+		EventDispatcher->UnRegisterForRemovalEvents(this);
+	}
+}
+
 void ActivateClusters(Chaos::FPBDRigidsEvolution::FRigidClustering& Clustering, Chaos::TPBDRigidClusteredParticleHandle<Chaos::FReal, 3>* Cluster)
 {
 	if(!Cluster)
@@ -1640,6 +1694,11 @@ void UGeometryCollectionComponent::TickComponent(float DeltaTime, enum ELevelTic
 		// In editor mode we have no DynamicCollection so this test is necessary
 		if(DynamicCollection) //, TEXT("No dynamic collection available for component %s during tick."), *GetName()))
 		{
+			if (RestCollection->bRemoveOnMaxSleep)
+			{ 
+				IncrementSleepTimer(DeltaTime);
+			}
+
 			if (RestCollection->HasVisibleGeometry() || DynamicCollection->IsDirty())
 			{
 				// #todo review: When we've made changes to ISMC, we need to move this function call to SetRenderDynamicData_Concurrent
@@ -1715,6 +1774,22 @@ void UGeometryCollectionComponent::ResetDynamicCollection()
 		GetChildrenArrayCopyOnWrite();
 		GetSimulationTypeArrayCopyOnWrite();
 		GetStatusFlagsArrayCopyOnWrite();
+
+		if (RestCollection->bRemoveOnMaxSleep)
+		{
+			if (!DynamicCollection->HasAttribute("SleepTimer", FGeometryCollection::TransformGroup))
+			{
+				TManagedArray<float>& SleepTimer = DynamicCollection->AddAttribute<float>("SleepTimer", FGeometryCollection::TransformGroup);
+				SleepTimer.Fill(0.f);
+			}
+
+			if (!DynamicCollection->HasAttribute("UniformScale", FGeometryCollection::TransformGroup))
+			{
+				TManagedArray<float>& UniformScale = DynamicCollection->AddAttribute<float>("UniformScale", FGeometryCollection::TransformGroup);
+				UniformScale.Fill(1.f);
+			}
+		}
+
 		SetRenderStateDirty();
 	}
 
@@ -1892,6 +1967,7 @@ void UGeometryCollectionComponent::RegisterAndInitializePhysicsProxy()
 		SimulationParameters.bGenerateBreakingData = bNotifyBreaks;
 		SimulationParameters.bGenerateCollisionData = bNotifyCollisions;
 		SimulationParameters.bGenerateTrailingData = bNotifyTrailing;
+		SimulationParameters.bGenerateRemovalsData = bNotifyRemovals;
 		SimulationParameters.RemoveOnFractureEnabled = SimulationParameters.Shared.RemoveOnFractureIndices.Num() > 0;
 		SimulationParameters.WorldTransform = GetComponentToWorld();
 		SimulationParameters.UserData = static_cast<void*>(&PhysicsUserData);
@@ -2280,6 +2356,7 @@ void FScopedColorEdit::SelectBones(GeometryCollection::ESelectionMode SelectionM
 			if (ensureMsgf(GeometryCollectionPtr->HasAttribute("Proximity", FGeometryCollection::GeometryGroup),
 				TEXT("Must build breaking group for neighbor based selection")))
 			{
+				FGeometryCollectionProximityUtility::UpdateProximity(GeometryCollectionPtr.Get());
 
 				const TManagedArray<int32>& TransformIndex = GeometryCollectionPtr->TransformIndex;
 				const TManagedArray<int32>& TransformToGeometryIndex = GeometryCollectionPtr->TransformToGeometryIndex;
@@ -2700,7 +2777,25 @@ void UGeometryCollectionComponent::CalculateGlobalMatrices()
 		else
 		{
 			// Have to fully rebuild
-			GeometryCollectionAlgo::GlobalMatrices(GetTransformArray(), GetParentArray(), GlobalMatrices);
+			if (DynamicCollection && RestCollection->bRemoveOnMaxSleep && DynamicCollection->HasAttribute("SleepTimer", FGeometryCollection::TransformGroup) && DynamicCollection->HasAttribute("UniformScale", FGeometryCollection::TransformGroup))
+			{
+				const TManagedArray<float>& SleepTimer = DynamicCollection->GetAttribute<float>("SleepTimer", FGeometryCollection::TransformGroup);
+				TManagedArray<float>& UniformScale = DynamicCollection->GetAttribute<float>("UniformScale", FGeometryCollection::TransformGroup);
+				
+				for (int32 Idx = 0; Idx < GetTransformArray().Num(); ++Idx)
+				{
+					if (SleepTimer[Idx] > RestCollection->MaximumSleepTime)
+					{
+						UniformScale[Idx] = 1.0 - FMath::Min(1.0, (SleepTimer[Idx] - RestCollection->MaximumSleepTime) / RestCollection->RemovalDuration);
+					}
+				}
+				
+				GeometryCollectionAlgo::GlobalMatrices(GetTransformArray(), GetParentArray(), UniformScale, DynamicCollection->MassToLocal, GlobalMatrices);
+			}
+			else
+			{ 
+				GeometryCollectionAlgo::GlobalMatrices(GetTransformArray(), GetParentArray(), GlobalMatrices);		
+			}
 		}
 	}
 	
@@ -2886,3 +2981,36 @@ void UGeometryCollectionComponent::InitializeEmbeddedGeometry()
 	}
 }
 
+void UGeometryCollectionComponent::IncrementSleepTimer(float DeltaTime)
+{
+	// If a particle is sleeping, increment its sleep timer, otherwise reset it.
+	if (DynamicCollection && PhysicsProxy && DynamicCollection->HasAttribute("SleepTimer", FGeometryCollection::TransformGroup))
+	{
+		TManagedArray<float>& SleepTimer = DynamicCollection->GetAttribute<float>("SleepTimer", FGeometryCollection::TransformGroup);
+		TArray<int32> ToDisable;
+		for (int32 TransformIdx = 0; TransformIdx < SleepTimer.Num(); ++TransformIdx)
+		{
+			bool PreviouslyAwake = SleepTimer[TransformIdx] < RestCollection->MaximumSleepTime;
+			if (SleepTimer[TransformIdx] < (RestCollection->MaximumSleepTime + RestCollection->RemovalDuration))
+			{
+				SleepTimer[TransformIdx] = (DynamicCollection->DynamicState[TransformIdx] == (int)EObjectStateTypeEnum::Chaos_Object_Sleeping) ? SleepTimer[TransformIdx] + DeltaTime : 0.0f;
+
+				if (SleepTimer[TransformIdx] > RestCollection->MaximumSleepTime)
+				{ 
+					DynamicCollection->MakeDirty();
+					if (PreviouslyAwake)
+					{
+						// Disable the particle if it has been asleep for the requisite time
+						ToDisable.Add(TransformIdx);
+					}
+				}
+			}
+		}
+
+		if (ToDisable.Num())
+		{
+			PhysicsProxy->DisableParticles(ToDisable);	
+		}
+	}
+}
+
diff --git a/Engine/Source/Runtime/Experimental/GeometryCollectionEngine/Private/GeometryCollection/GeometryCollectionObject.cpp b/Engine/Source/Runtime/Experimental/GeometryCollectionEngine/Private/GeometryCollection/GeometryCollectionObject.cpp
--- a/Engine/Source/Runtime/Experimental/GeometryCollectionEngine/Private/GeometryCollection/GeometryCollectionObject.cpp
+++ b/Engine/Source/Runtime/Experimental/GeometryCollectionEngine/Private/GeometryCollection/GeometryCollectionObject.cpp
@@ -79,6 +79,9 @@ UGeometryCollection::UGeometryCollection(const FObjectInitializer& ObjectInitial
 	, bMassAsDensity(false)
 	, Mass(1.0f)
 	, MinimumMassClamp(0.1f)
+	, bRemoveOnMaxSleep(false)
+	, MaximumSleepTime(10.0)
+	, RemovalDuration(5.0)
 	, EnableRemovePiecesOnFracture(false)
 	, GeometryCollection(new FGeometryCollection())
 {
diff --git a/Engine/Source/Runtime/Experimental/GeometryCollectionEngine/Public/ChaosBlueprint.h b/Engine/Source/Runtime/Experimental/GeometryCollectionEngine/Public/ChaosBlueprint.h
--- a/Engine/Source/Runtime/Experimental/GeometryCollectionEngine/Public/ChaosBlueprint.h
+++ b/Engine/Source/Runtime/Experimental/GeometryCollectionEngine/Public/ChaosBlueprint.h
@@ -12,6 +12,7 @@
 #include "ChaosCollisionEventFilter.h"
 #include "ChaosBreakingEventFilter.h"
 #include "ChaosTrailingEventFilter.h"
+#include "ChaosRemovalEventFilter.h"
 #include "ChaosBlueprint.generated.h"
 
 class AGeometryCollectionActor;
@@ -31,6 +32,11 @@ DECLARE_DYNAMIC_MULTICAST_DELEGATE_OneParam(FOnChaosBreakingEvents, const TArray
 */
 DECLARE_DYNAMIC_MULTICAST_DELEGATE_OneParam(FOnChaosTrailingEvents, const TArray<FChaosTrailingEventData>&, TrailingEvents);
 
+/**
+* Called when new destruction events are available for removals. Removal listening must be enabled to get callbacks on this delegate.
+*/
+DECLARE_DYNAMIC_MULTICAST_DELEGATE_OneParam(FOnChaosRemovalEvents, const TArray<FChaosRemovalEventData>&, RemovalEvents);
+
 
 /** Object allowing for retrieving Chaos Destruction data. */
 UCLASS(ClassGroup = (Chaos), hidecategories = (Object, ActorComponent, Physics, Rendering, Mobility, LOD), ShowCategories = Trigger, meta = (BlueprintSpawnableComponent))
@@ -63,6 +69,10 @@ public:
 	UPROPERTY(EditAnywhere, BlueprintReadOnly, Category = "Trailing Events")
 	uint8 bIsTrailingEventListeningEnabled : 1;
 
+	// Whether or not removal event listening is enabled
+	UPROPERTY(EditAnywhere, BlueprintReadOnly, Category = "Removal Events")
+	uint8 bIsRemovalEventListeningEnabled : 1;
+
 	// The settings to use for collision event listening
 	UPROPERTY(EditAnywhere, BlueprintReadOnly, Category = "Collision Events", meta = (EditCondition = "bIsCollisionEventListeningEnabled"))
 	FChaosCollisionEventRequestSettings CollisionEventRequestSettings;
@@ -75,6 +85,10 @@ public:
 	UPROPERTY(EditAnywhere, BlueprintReadOnly, Category = "Trailing Events", meta = (EditCondition = "bIsTrailingEventListeningEnabled"))
 	FChaosTrailingEventRequestSettings TrailingEventRequestSettings;
 
+	// The settings to use for removal event listening
+	UPROPERTY(EditAnywhere, BlueprintReadOnly, Category = "Removal Events", meta = (EditCondition = "bIsRemovalEventListeningEnabled"))
+	FChaosRemovalEventRequestSettings RemovalEventRequestSettings;
+
 	// Which chaos solver actors we're using. If empty, this listener will fallback to the "world" solver.
 	UPROPERTY(EditAnywhere, BlueprintReadOnly, Category = "Solvers")
 	TSet<TObjectPtr<AChaosSolverActor>> ChaosSolverActors;
@@ -111,6 +125,10 @@ public:
 	UFUNCTION(BlueprintCallable, Category = "Destruction Listener", meta = (WorldContext = "WorldContextObject"))
 	void SetTrailingEventRequestSettings(const FChaosTrailingEventRequestSettings& InSettings);
 
+	// Sets removal event request settings dynamically
+	UFUNCTION(BlueprintCallable, Category = "Destruction Listener", meta = (WorldContext = "WorldContextObject"))
+	void SetRemovalEventRequestSettings(const FChaosRemovalEventRequestSettings& InSettings);
+
 	// Enables or disables collision event listening
 	UFUNCTION(BlueprintCallable, Category = "Destruction Listener", meta = (WorldContext = "WorldContextObject"))
 	void SetCollisionEventEnabled(bool bIsEnabled);
@@ -119,10 +137,14 @@ public:
 	UFUNCTION(BlueprintCallable, Category = "Destruction Listener", meta = (WorldContext = "WorldContextObject"))
 	void SetBreakingEventEnabled(bool bIsEnabled);
 
-	// Enables or disables breaking event listening
+	// Enables or disables trailing event listening
 	UFUNCTION(BlueprintCallable, Category = "Destruction Listener", meta = (WorldContext = "WorldContextObject"))
 	void SetTrailingEventEnabled(bool bIsEnabled);
 
+	// Enables or disables removal event listening
+	UFUNCTION(BlueprintCallable, Category = "Destruction Listener", meta = (WorldContext = "WorldContextObject"))
+	void SetRemovalEventEnabled(bool bIsEnabled);
+
 	// Returns if the destruction listener is listening to any events
 	UFUNCTION(BlueprintCallable, Category = "Destruction Listener", meta = (WorldContext = "WorldContextObject"))
 	bool IsEventListening() const;
@@ -139,6 +161,10 @@ public:
 	UPROPERTY(BlueprintAssignable)
 	FOnChaosTrailingEvents OnTrailingEvents;
 
+	/** Called when new trailing events are available. */
+	UPROPERTY(BlueprintAssignable)
+	FOnChaosRemovalEvents OnRemovalEvents;
+
 	// Sorts collision events according to the given sort method	
 	UFUNCTION(BlueprintCallable, Category = "Destruction Listener", meta = (WorldContext = "WorldContextObject"))
 	void SortCollisionEvents(UPARAM(ref) TArray<FChaosCollisionEventData>& CollisionEvents, EChaosCollisionSortMethod SortMethod);
@@ -151,6 +177,10 @@ public:
 	UFUNCTION(BlueprintCallable, Category = "Destruction Listener", meta = (WorldContext = "WorldContextObject"))
 	void SortTrailingEvents(UPARAM(ref) TArray<FChaosTrailingEventData>& TrailingEvents, EChaosTrailingSortMethod SortMethod);
 
+	// Sorts removal events according to the given sort method	
+	UFUNCTION(BlueprintCallable, Category = "Destruction Listener", meta = (WorldContext = "WorldContextObject"))
+	void SortRemovalEvents(UPARAM(ref) TArray<FChaosRemovalEventData>& RemovalEvents, EChaosRemovalSortMethod SortMethod);
+
 private:
 	// Updates the scene component transform settings
 	void UpdateTransformSettings();
@@ -181,6 +211,7 @@ protected:
 	Chaos::FCollisionDataArray RawCollisionDataArray;
 	Chaos::FBreakingDataArray RawBreakingDataArray;
 	Chaos::FTrailingDataArray RawTrailingDataArray;
+	Chaos::FRemovalDataArray RawRemovalDataArray;
 
 	FTransform ChaosComponentTransform;
 
@@ -189,6 +220,7 @@ protected:
 	float LastCollisionDataTimeStamp;
 	float LastBreakingDataTimeStamp;
 	float LastTrailingDataTimeStamp;
+	float LastRemovalDataTimeStamp;
 
 	// The list of rigid body solvers, used to retrieve destruction events
 	TSet<Chaos::FPhysicsSolver*> Solvers;
@@ -199,6 +231,7 @@ protected:
 	TSharedPtr<FChaosCollisionEventFilter> ChaosCollisionFilter;
 	TSharedPtr<FChaosBreakingEventFilter> ChaosBreakingFilter;
 	TSharedPtr<FChaosTrailingEventFilter> ChaosTrailingFilter;
+	TSharedPtr<FChaosRemovalEventFilter> ChaosRemovalFilter;
 
 public:
 	void SetCollisionFilter(TSharedPtr<FChaosCollisionEventFilter> InCollisionFilter)
@@ -216,6 +249,11 @@ public:
 		ChaosTrailingFilter = InTrailingFilter;
 	}
 
+	void SetRemovalFilter(TSharedPtr<FChaosRemovalEventFilter> InRemovalFilter)
+	{
+		ChaosRemovalFilter = InRemovalFilter;
+	}
+
 	void RegisterChaosEvents(FPhysScene* Scene);
 	void UnregisterChaosEvents(FPhysScene* Scene);
 
@@ -227,4 +265,5 @@ public:
 	void HandleCollisionEvents(const Chaos::FCollisionEventData& CollisionData);
 	void HandleBreakingEvents(const Chaos::FBreakingEventData& BreakingData);
 	void HandleTrailingEvents(const Chaos::FTrailingEventData& TrailingData);
+	void HandleRemovalEvents(const Chaos::FRemovalEventData& RemovalData);
 };
diff --git a/Engine/Source/Runtime/Experimental/GeometryCollectionEngine/Public/ChaosRemovalEventFilter.h b/Engine/Source/Runtime/Experimental/GeometryCollectionEngine/Public/ChaosRemovalEventFilter.h
new file mode 100644
--- /dev/null
+++ b/Engine/Source/Runtime/Experimental/GeometryCollectionEngine/Public/ChaosRemovalEventFilter.h
@@ -0,0 +1,91 @@
+// Copyright Epic Games, Inc. All Rights Reserved.
+#pragma once
+
+#include "CoreMinimal.h"
+#include "UObject/ObjectMacros.h"
+#include "ChaosFilter.h"
+#include "EventsData.h"
+#include "ChaosRemovalEventFilter.generated.h"
+
+// A Removal event data structure. 
+USTRUCT(BlueprintType)
+struct FChaosRemovalEventData
+{
+	GENERATED_USTRUCT_BODY()
+
+	/** Current removal location. */
+	UPROPERTY(EditAnywhere, BlueprintReadWrite, Category = "Filters")
+	FVector Location;
+
+	/** The mass of the removal. */
+	UPROPERTY(EditAnywhere, BlueprintReadWrite, Category = "Filters")
+	float Mass;
+
+	/** The particle index of the removal. */
+	UPROPERTY(EditAnywhere, BlueprintReadWrite, Category = "Filters")
+	int32 ParticleIndex;
+
+	FChaosRemovalEventData()
+		: Location(FVector::ZeroVector)
+		, Mass(0.0f)
+		, ParticleIndex(INDEX_NONE)
+	{
+	}
+};
+
+// Enumeration defining how to sort removal results
+UENUM(BlueprintType)
+enum class EChaosRemovalSortMethod : uint8
+{
+	SortNone,
+	SortByHighestMass,
+	SortByNearestFirst,
+
+	Count UMETA(Hidden)
+};
+
+USTRUCT(BlueprintType)
+struct FChaosRemovalEventRequestSettings
+{
+	GENERATED_USTRUCT_BODY()
+
+	/** The maximum number of results to return. */
+	UPROPERTY(EditAnywhere, BlueprintReadWrite, Category = "Filters")
+	int32 MaxNumberOfResults;
+
+	/** The minimum mass treshold for the results. */
+	UPROPERTY(EditAnywhere, BlueprintReadWrite, Category = "Filters")
+	float MinMass;
+
+	/** The maximum distance threshold for the results (if location is set on destruction event listener). */
+	UPROPERTY(EditAnywhere, BlueprintReadWrite, Category = "Filters")
+	float MaxDistance;
+
+	/** The method used to sort the removal events. */
+	UPROPERTY(EditAnywhere, BlueprintReadWrite, Category = "Sort")
+	EChaosRemovalSortMethod SortMethod;
+
+	FChaosRemovalEventRequestSettings()
+		: MaxNumberOfResults(0)
+		, MinMass(0.0f)
+		, MaxDistance(0.0f)
+		, SortMethod(EChaosRemovalSortMethod::SortByHighestMass)
+	{
+	}
+};
+
+class GEOMETRYCOLLECTIONENGINE_API FChaosRemovalEventFilter
+	: public IChaosEventFilter<Chaos::FRemovalDataArray, TArray<FChaosRemovalEventData>, EChaosRemovalSortMethod>
+{
+public:
+	FChaosRemovalEventFilter(FChaosRemovalEventRequestSettings* FilterSettingsIn) : RemovalEventRequestSettings(FilterSettingsIn) {}
+
+	virtual void FilterEvents(const FTransform& ChaosComponentTransform, const Chaos::FRemovalDataArray& RawRemovalDataArray) override;
+
+	virtual void SortEvents(TArray<FChaosRemovalEventData>& InOutRemovalEvents, EChaosRemovalSortMethod SortMethod, const FTransform& InTransform) override;
+
+private:
+	FChaosRemovalEventFilter() : RemovalEventRequestSettings(nullptr) {}
+	const FChaosRemovalEventRequestSettings* RemovalEventRequestSettings;
+};
+
diff --git a/Engine/Source/Runtime/Experimental/GeometryCollectionEngine/Public/GeometryCollection/GeometryCollectionComponent.h b/Engine/Source/Runtime/Experimental/GeometryCollectionEngine/Public/GeometryCollection/GeometryCollectionComponent.h
--- a/Engine/Source/Runtime/Experimental/GeometryCollectionEngine/Public/GeometryCollection/GeometryCollectionComponent.h
+++ b/Engine/Source/Runtime/Experimental/GeometryCollectionEngine/Public/GeometryCollection/GeometryCollectionComponent.h
@@ -38,6 +38,8 @@ class UInstancedStaticMeshComponent;
 
 DECLARE_DYNAMIC_MULTICAST_DELEGATE_OneParam(FOnChaosBreakEvent, const FChaosBreakEvent&, BreakEvent);
 
+DECLARE_DYNAMIC_MULTICAST_DELEGATE_OneParam(FOnChaosRemovalEvent, const FChaosRemovalEvent&, RemovalEvent);
+
 namespace GeometryCollection
 {
 	enum class ESelectionMode : uint8
@@ -621,14 +623,26 @@ public:
 	UFUNCTION(BlueprintCallable, Category = "Physics")
 	void SetNotifyBreaks(bool bNewNotifyBreaks);
 
+	/** Changes whether or not this component will get future removal notifications. */
+	UFUNCTION(BlueprintCallable, Category = "Physics")
+	void SetNotifyRemovals(bool bNewNotifyRemovals);
+
 	/** Overrideable native notification */
 	virtual void NotifyBreak(const FChaosBreakEvent& Event) {};
 
+	/** Overrideable native notification */
+	virtual void NotifyRemoval(const FChaosRemovalEvent& Event) {};
+
 	UPROPERTY(BlueprintAssignable, Category = "Chaos")
 	FOnChaosBreakEvent OnChaosBreakEvent;
+
+	UPROPERTY(BlueprintAssignable, Category = "Chaos")
+	FOnChaosRemovalEvent OnChaosRemovalEvent;
 	
 	void DispatchBreakEvent(const FChaosBreakEvent& Event);
 
+	void DispatchRemovalEvent(const FChaosRemovalEvent& Event);
+
 	UPROPERTY(Transient, VisibleAnywhere, BlueprintReadWrite, Interp, Category = "Chaos")
 	float DesiredCacheTime;
 
@@ -680,6 +694,10 @@ public:
 	UPROPERTY(EditAnywhere, BlueprintReadOnly, Category = "ChaosPhysics|General")
 	bool bNotifyTrailing;
 
+	/** If true, this component will generate removal events that other systems may subscribe to. */
+	UPROPERTY(EditAnywhere, BlueprintReadOnly, Category = "ChaosPhysics|General")
+	bool bNotifyRemovals;
+
 protected:
 	/** Display Bone Colors instead of assigned materials */
 	UPROPERTY(EditAnywhere, BlueprintReadWrite, Category = "ChaosPhysics|General")
@@ -706,6 +724,7 @@ protected:
 	void RegisterForEvents();
 	void UpdateRBCollisionEventRegistration();
 	void UpdateBreakEventRegistration();
+	void UpdateRemovalEventRegistration();
 
 	/* Per-instance override to enable/disable replication for the geometry collection */
 	UPROPERTY(EditInstanceOnly, BlueprintReadOnly, Category=Network)
@@ -809,6 +828,8 @@ private:
 	bool IsEmbeddedGeometryValid() const;
 	void ClearEmbeddedGeometry();
 
+	void IncrementSleepTimer(float DeltaTime);
+
 	/** True if GeometryCollection transforms have changed from previous tick. */
 	bool bIsMoving;
 
diff --git a/Engine/Source/Runtime/Experimental/GeometryCollectionEngine/Public/GeometryCollection/GeometryCollectionObject.h b/Engine/Source/Runtime/Experimental/GeometryCollectionEngine/Public/GeometryCollection/GeometryCollectionObject.h
--- a/Engine/Source/Runtime/Experimental/GeometryCollectionEngine/Public/GeometryCollection/GeometryCollectionObject.h
+++ b/Engine/Source/Runtime/Experimental/GeometryCollectionEngine/Public/GeometryCollection/GeometryCollectionObject.h
@@ -535,6 +535,18 @@ public:
 	int32 MaximumCollisionParticles_DEPRECATED;
 #endif
 
+	/** Remove particle from simulation and dissolve rendered geometry once sleep threshold has been exceeded. */
+	UPROPERTY(EditAnywhere, BlueprintReadOnly, Category = Removal, meta = (DisplayName = "RemoveOnMaxSleep"))
+	bool bRemoveOnMaxSleep;
+	
+	/** How long may the particle sleep before initiating removal (in seconds). */
+	UPROPERTY(EditAnywhere, BlueprintReadOnly, Category = Removal, meta = (ClampMin = 0.001))
+	float MaximumSleepTime;
+
+	/** How long does the removal process take (in seconds). */
+	UPROPERTY(EditAnywhere, BlueprintReadOnly, Category = Removal, meta = (ClampMin = 0.0))
+	float RemovalDuration;
+
 	/*
 	* Size Specfic Data reflects the default geometry to bind to rigid bodies smaller
 	* than the max size volume. This can also be empty to reflect no collision geometry
