fix Explode Amount slider

Specifically, change UGeometryCollectionComponent::OnCreatePhysicsState's EnsureDataIsCooked() call to just CreateSimulationData(), so we just update the sim data without also re-loading the serialized data (which had some editor/UI specific data discarded)

#rnx

diff --git a/Engine/Source/Runtime/Experimental/GeometryCollectionEngine/Private/GeometryCollection/GeometryCollectionComponent.cpp b/Engine/Source/Runtime/Experimental/GeometryCollectionEngine/Private/GeometryCollection/GeometryCollectionComponent.cpp
--- a/Engine/Source/Runtime/Experimental/GeometryCollectionEngine/Private/GeometryCollection/GeometryCollectionComponent.cpp
+++ b/Engine/Source/Runtime/Experimental/GeometryCollectionEngine/Private/GeometryCollection/GeometryCollectionComponent.cpp
@@ -1838,7 +1838,7 @@ void UGeometryCollectionComponent::OnCreatePhysicsState()
 		{
 			//hack: find a better place for this
 			UGeometryCollection* RestCollectionMutable = const_cast<UGeometryCollection*>(ToRawPtr(RestCollection));
-			RestCollectionMutable->EnsureDataIsCooked();
+			RestCollectionMutable->CreateSimulationData();
 		}
 #endif
 		const bool bValidWorld = GetWorld() && GetWorld()->IsGameWorld();
