GeometryCollectionComponent
-- Corrected motion blur problem when cache scrubbing in StaticPose mode

diff --git a/Engine/Source/Runtime/Experimental/GeometryCollectionEngine/Private/GeometryCollection/GeometryCollectionComponent.cpp b/Engine/Source/Runtime/Experimental/GeometryCollectionEngine/Private/GeometryCollection/GeometryCollectionComponent.cpp
--- a/Engine/Source/Runtime/Experimental/GeometryCollectionEngine/Private/GeometryCollection/GeometryCollectionComponent.cpp
+++ b/Engine/Source/Runtime/Experimental/GeometryCollectionEngine/Private/GeometryCollection/GeometryCollectionComponent.cpp
@@ -891,18 +891,18 @@ void UGeometryCollectionComponent::RefreshEmbeddedGeometry()
 }
 
 void UGeometryCollectionComponent::SetRestState(TArray<FTransform>&& InRestTransforms)
-{
+{	
 	RestTransforms = InRestTransforms;
 	
-	CalculateGlobalMatrices();
-
 	if (DynamicCollection)
 	{
 		SetInitialTransforms(RestTransforms);
 	}
 
 	FGeometryCollectionDynamicData* DynamicData = GDynamicDataPool.Allocate();
-	DynamicData->SetAllTransforms(GlobalMatrices);
+	DynamicData->SetPrevTransforms(GlobalMatrices);
+	CalculateGlobalMatrices();
+	DynamicData->SetTransforms(GlobalMatrices);
 	DynamicData->IsDynamic = true;
 
 	if (SceneProxy)
