GeometryCollectionProximityUtility
-- New algorithm, faster, more robust

diff --git a/Engine/Plugins/Experimental/ChaosEditor/Source/FractureEditor/Private/FractureToolAutoCluster.cpp b/Engine/Plugins/Experimental/ChaosEditor/Source/FractureEditor/Private/FractureToolAutoCluster.cpp
--- a/Engine/Plugins/Experimental/ChaosEditor/Source/FractureEditor/Private/FractureToolAutoCluster.cpp
+++ b/Engine/Plugins/Experimental/ChaosEditor/Source/FractureEditor/Private/FractureToolAutoCluster.cpp
@@ -8,6 +8,7 @@
 
 #include "GeometryCollection/GeometryCollectionComponent.h"
 #include "GeometryCollection/GeometryCollectionClusteringUtility.h"
+#include "GeometryCollection/GeometryCollectionProximityUtility.h"
 #include "GeometryCollection/GeometryCollectionAlgo.h"
 #include "Math/NumericLimits.h"
 
@@ -80,7 +81,8 @@ void UFractureToolAutoCluster::Execute(TWeakPtr<FFractureEditorModeToolkit> InTo
 
 				if (AutoClusterSettings->bEnforceConnectivity)
 				{
-					GenerateProximityIfNecessary(GeometryCollection);
+					FGeometryCollectionProximityUtility ProximityUtility(GeometryCollection);
+					ProximityUtility.UpdateProximity();
 					VoronoiPartition.SplitDisconnectedPartitions(GeometryCollection);
 				}
 					
diff --git a/Engine/Plugins/Experimental/ChaosEditor/Source/FractureEditor/Private/FractureToolClusterMagnet.cpp b/Engine/Plugins/Experimental/ChaosEditor/Source/FractureEditor/Private/FractureToolClusterMagnet.cpp
--- a/Engine/Plugins/Experimental/ChaosEditor/Source/FractureEditor/Private/FractureToolClusterMagnet.cpp
+++ b/Engine/Plugins/Experimental/ChaosEditor/Source/FractureEditor/Private/FractureToolClusterMagnet.cpp
@@ -164,7 +164,8 @@ bool UFractureToolClusterMagnet::CheckPresenceOfNecessaryAttributes(const FGeome
 
 TMap<int32, TSet<int32>> UFractureToolClusterMagnet::InitializeConnectivity(const TSet<int32>& TopNodes, FGeometryCollectionPtr GeometryCollection, int32 OperatingLevel) const
 {
-	FGeometryCollectionProximityUtility::UpdateProximity(GeometryCollection.Get());
+	FGeometryCollectionProximityUtility ProximityUtility(GeometryCollection.Get());
+	ProximityUtility.UpdateProximity();
 
 	TMap<int32, TSet<int32>> ConnectivityMap;
 	for (int32 Index : TopNodes)
diff --git a/Engine/Plugins/Experimental/ChaosEditor/Source/FractureEditor/Private/FractureToolConvex.cpp b/Engine/Plugins/Experimental/ChaosEditor/Source/FractureEditor/Private/FractureToolConvex.cpp
--- a/Engine/Plugins/Experimental/ChaosEditor/Source/FractureEditor/Private/FractureToolConvex.cpp
+++ b/Engine/Plugins/Experimental/ChaosEditor/Source/FractureEditor/Private/FractureToolConvex.cpp
@@ -8,6 +8,7 @@
 #include "GeometryCollection/GeometryCollectionComponent.h"
 #include "GeometryCollection/GeometryCollectionAlgo.h"
 #include "GeometryCollection/GeometryCollectionConvexUtility.h"
+#include "GeometryCollection/GeometryCollectionProximityUtility.h"
 
 #define LOCTEXT_NAMESPACE "FractureToolConvex"
 
@@ -152,7 +153,8 @@ int32 UFractureToolConvex::ExecuteFracture(const FFractureToolContext& FractureC
 	if (FractureContext.GetGeometryCollection().IsValid())
 	{
 		FGeometryCollection& Collection = *FractureContext.GetGeometryCollection();
-		GenerateProximityIfNecessary(&Collection);
+		FGeometryCollectionProximityUtility ProximityUtility(&Collection);
+		ProximityUtility.UpdateProximity();
 		FGeometryCollectionConvexUtility::CreateNonOverlappingConvexHullData(&Collection, ConvexSettings->FractionAllowRemove);
 	}
 
diff --git a/Engine/Plugins/Experimental/ChaosEditor/Source/FractureEditor/Public/FractureTool.cpp b/Engine/Plugins/Experimental/ChaosEditor/Source/FractureEditor/Public/FractureTool.cpp
--- a/Engine/Plugins/Experimental/ChaosEditor/Source/FractureEditor/Public/FractureTool.cpp
+++ b/Engine/Plugins/Experimental/ChaosEditor/Source/FractureEditor/Public/FractureTool.cpp
@@ -179,14 +179,6 @@ void UFractureActionTool::ClearProximity(FGeometryCollection* GeometryCollection
 	}
 }
 
-void UFractureActionTool::GenerateProximityIfNecessary(FGeometryCollection* GeometryCollection)
-{
-	if (!GeometryCollection->HasAttribute("Proximity", FGeometryCollection::GeometryGroup))
-	{
-		FGeometryCollectionProximityUtility::UpdateProximity(GeometryCollection);
-	}
-}
-
 TArray<FFractureToolContext> UFractureActionTool::GetFractureToolContexts() const
 {
 	TArray<FFractureToolContext> Contexts;
diff --git a/Engine/Plugins/Experimental/ChaosEditor/Source/FractureEditor/Public/FractureTool.h b/Engine/Plugins/Experimental/ChaosEditor/Source/FractureEditor/Public/FractureTool.h
--- a/Engine/Plugins/Experimental/ChaosEditor/Source/FractureEditor/Public/FractureTool.h
+++ b/Engine/Plugins/Experimental/ChaosEditor/Source/FractureEditor/Public/FractureTool.h
@@ -73,7 +73,6 @@ protected:
 	static void Refresh(FFractureToolContext& Context, FFractureEditorModeToolkit* Toolkit, bool bClearSelection=false);
 	static void SetOutlinerComponents(TArray<FFractureToolContext>& InContexts, FFractureEditorModeToolkit* Toolkit);
 	static void ClearProximity(FGeometryCollection* GeometryCollection);
-	static void GenerateProximityIfNecessary(FGeometryCollection* GeometryCollection);
 
 
 protected:
diff --git a/Engine/Plugins/Experimental/GeometryCollectionPlugin/Source/GeometryCollectionEditor/Private/GeometryCollection/GeometryCollectionCommands.cpp b/Engine/Plugins/Experimental/GeometryCollectionPlugin/Source/GeometryCollectionEditor/Private/GeometryCollection/GeometryCollectionCommands.cpp
--- a/Engine/Plugins/Experimental/GeometryCollectionPlugin/Source/GeometryCollectionEditor/Private/GeometryCollection/GeometryCollectionCommands.cpp
+++ b/Engine/Plugins/Experimental/GeometryCollectionPlugin/Source/GeometryCollectionEditor/Private/GeometryCollection/GeometryCollectionCommands.cpp
@@ -696,7 +696,8 @@ void FGeometryCollectionCommands::BuildProximityDatabase(const TArray<FString>&
 
 				if (2 <= RestCollection->GetGeometryCollection().Get()->NumElements(FGeometryCollection::GeometryGroup))
 				{
-					FGeometryCollectionProximityUtility::UpdateProximity(RestCollection->GetGeometryCollection().Get());
+					FGeometryCollectionProximityUtility ProximityUtility(RestCollection->GetGeometryCollection().Get());
+					ProximityUtility.UpdateProximity();
 				}
 			}
 		}
diff --git a/Engine/Plugins/Experimental/PlanarCutPlugin/Source/PlanarCut/Private/PlanarCut.cpp b/Engine/Plugins/Experimental/PlanarCutPlugin/Source/PlanarCut/Private/PlanarCut.cpp
--- a/Engine/Plugins/Experimental/PlanarCutPlugin/Source/PlanarCut/Private/PlanarCut.cpp
+++ b/Engine/Plugins/Experimental/PlanarCutPlugin/Source/PlanarCut/Private/PlanarCut.cpp
@@ -891,7 +891,8 @@ int32 MergeBones(
 
 	FTransform CellsToWorld = FTransform::Identity;
 
-	FGeometryCollectionProximityUtility::UpdateProximity(&Collection);
+	FGeometryCollectionProximityUtility ProximityUtility(&Collection);
+	ProximityUtility.UpdateProximity();
 	TManagedArray<TSet<int32>>& Proximity = Collection.GetAttribute<TSet<int32>>("Proximity", FGeometryCollection::GeometryGroup);
 
 	// local array so we can populate it with all transforms if input was empty
diff --git a/Engine/Source/Runtime/Experimental/Chaos/Private/GeometryCollection/GeometryCollection.cpp b/Engine/Source/Runtime/Experimental/Chaos/Private/GeometryCollection/GeometryCollection.cpp
--- a/Engine/Source/Runtime/Experimental/Chaos/Private/GeometryCollection/GeometryCollection.cpp
+++ b/Engine/Source/Runtime/Experimental/Chaos/Private/GeometryCollection/GeometryCollection.cpp
@@ -1703,7 +1703,8 @@ FGeometryCollection* FGeometryCollection::NewGeometryCollection(const TArray<flo
 	// Build the Geometry Group
 	GeometryCollection::AddGeometryProperties(RestCollection);
 
-	FGeometryCollectionProximityUtility::UpdateProximity(RestCollection);
+	FGeometryCollectionProximityUtility ProximityUtility(RestCollection);
+	ProximityUtility.UpdateProximity();
 
 	// add a material section
 	TManagedArray<FGeometryCollectionSection>&  Sections = RestCollection->Sections;
diff --git a/Engine/Source/Runtime/Experimental/Chaos/Private/GeometryCollection/GeometryCollectionProximityUtility.cpp b/Engine/Source/Runtime/Experimental/Chaos/Private/GeometryCollection/GeometryCollectionProximityUtility.cpp
--- a/Engine/Source/Runtime/Experimental/Chaos/Private/GeometryCollection/GeometryCollectionProximityUtility.cpp
+++ b/Engine/Source/Runtime/Experimental/Chaos/Private/GeometryCollection/GeometryCollectionProximityUtility.cpp
@@ -5,605 +5,283 @@
 #include "Async/ParallelFor.h"
 #include "GeometryCollection/GeometryCollection.h"
 #include "GeometryCollection/GeometryCollectionAlgo.h"
-#include "Math/GenericOctreePublic.h"
-#include "Math/GenericOctree.h"
-#include "GeometryCollection/GeometryCollectionClusteringUtility.h"
-
-#include <chrono>
-
-#if WITH_EDITOR
-#include "Misc/ScopedSlowTask.h"
-#endif
 
 DEFINE_LOG_CATEGORY_STATIC(LogChaosProximity, Verbose, All);
 
-#define USE_OLD_METHOD 1
-
-struct FProximityTriangle
+FGeometryCollectionProximityUtility::FGeometryCollectionProximityUtility(FGeometryCollection* InCollection)
+	: Collection(InCollection)
 {
-	int32 ArrayIndex;
-	FVector Vertices[3];
-	FVector Normal;
-	FBoxCenterAndExtent BoxCenterAndExtent;
-	FBox Bounds;
-};
-
-/** Semantics for the simple mesh paint octree */
-struct FMeshProximityTriangleOctreeSemantics
-{
-	enum { MaxElementsPerLeaf = 16 };
-	enum { MinInclusiveElementsPerNode = 7 };
-	enum { MaxNodeDepth = 12 };
-
-	typedef TInlineAllocator<MaxElementsPerLeaf> ElementAllocator;
-
-	/**
-	* Get the bounding box of the provided octree element. In this case, the box
-	* is merely the point specified by the element.
-	*
-	* @param	Element	Octree element to get the bounding box for
-	*
-	* @return	Bounding box of the provided octree element
-	*/
-	FORCEINLINE static FBoxCenterAndExtent GetBoundingBox(const FProximityTriangle& Element)
-	{
-		return Element.BoxCenterAndExtent;
-	}
-
+	check(Collection);
+	
+	// We quantize surface normals into 20 uniform bins on a unit sphere surface, ie an icosahedron
+	BinNormals.SetNum(NumBins);
+	
+	BinNormals[0] = FVector3f(0.171535f, -0.793715f, 0.583717f);
+	BinNormals[1] = FVector3f(0.627078f, -0.778267f, 0.034524f);
+	BinNormals[2] = FVector3f(-0.491358f, -0.810104f, 0.319894f);
+	BinNormals[3] = FVector3f(-0.445554f, -0.804788f, -0.392214f);
+	BinNormals[4] = FVector3f(0.245658f, -0.785111f, -0.568669f);
+	BinNormals[5] = FVector3f(0.984880f, -0.161432f, 0.062144f);
+	BinNormals[6] = FVector3f(0.247864f, -0.186425f, 0.950708f);
+	BinNormals[7] = FVector3f(-0.824669f, -0.212942f, 0.523975f);
+	BinNormals[8] = FVector3f(-0.750546f, -0.204339f, -0.628411f);
+	BinNormals[9] = FVector3f(0.367791f, -0.172505f, -0.913787f);
+	BinNormals[10] = -BinNormals[0];
+	BinNormals[11] = -BinNormals[1];
+	BinNormals[12] = -BinNormals[2];
+	BinNormals[13] = -BinNormals[3];
+	BinNormals[14] = -BinNormals[4];
+	BinNormals[15] = -BinNormals[5];
+	BinNormals[16] = -BinNormals[6];
+	BinNormals[17] = -BinNormals[7];
+	BinNormals[18] = -BinNormals[8];
+	BinNormals[19] = -BinNormals[9];
+	
+	NumFaces = Collection->NumElements(FGeometryCollection::FacesGroup);
+	PrepFaceData();
+}
 
-	/**
-	* Determine if two octree elements are equal
-	*
-	* @param	A	First octree element to check
-	* @param	B	Second octree element to check
-	*
-	* @return	true if both octree elements are equal, false if they are not
+void FGeometryCollectionProximityUtility::UpdateProximity()
+{
+	/*
+	*	For each face, we generate the surface normal and put the normal into a set of bins that discretize the surface
+	*	of the unit sphere. We know that any face can only contact other faces with opposite surface normal, which 
+	*	means that we need only compare a face with the bin of faces on the antipode of the sphere. We determine contact first
+	*	by testing if the normal is opposite, then if the faces are coplanar, then if the 2D triangles overlap.
+	*	Once we have face pairs, we can extend this to geometry indices and form the proximity structure
+	*	which is a map from geometry to geometry, one-to-many.
 	*/
-	FORCEINLINE static bool AreElementsEqual(const FProximityTriangle& A, const FProximityTriangle& B)
-	{
-		return (A.ArrayIndex == B.ArrayIndex);
-	}
-
-	/** Ignored for this implementation */
-	FORCEINLINE static void SetElementId(const FProximityTriangle& Element, FOctreeElementId2 Id)
-	{
-
-	}
-};
-typedef TOctree2<FProximityTriangle, FMeshProximityTriangleOctreeSemantics> FProximityTriangleOctree;
 
+	BinFaces();
+	FindContactingFaces();
+	ExtendFaceProximityToGeometry();
+}
 
-bool FGeometryCollectionProximityUtility::IsPointInsideOfTriangle(const FVector& P, const FVector& Vertex0, const FVector& Vertex1, const FVector& Vertex2, float Threshold)
-{
-	FVector::FReal FaceArea  = 0.5f * FVector::CrossProduct((Vertex1 - Vertex0), (Vertex2 - Vertex0)).SizeSquared();
-	FVector::FReal Face1Area = 0.5f * FVector::CrossProduct((Vertex0 - P),(Vertex2 - P)).SizeSquared();
-	FVector::FReal Face2Area = 0.5f * FVector::CrossProduct((Vertex0 - P),(Vertex1 - P)).SizeSquared();
-	FVector::FReal Face3Area = 0.5f * FVector::CrossProduct((Vertex2 - P),(Vertex1 - P)).SizeSquared();
 
-	return (FMath::Abs(Face1Area + Face2Area + Face3Area - FaceArea) < Threshold * Threshold);
+void FGeometryCollectionProximityUtility::PrepFaceData()
+{	
+	GenerateFaceToGeometry();
+	TransformVertices();
+	GenerateSurfaceNormals();
 }
 
-void FGeometryCollectionProximityUtility::UpdateProximity(FGeometryCollection* GeometryCollection)
+void FGeometryCollectionProximityUtility::BinFaces()
 {
-	check(GeometryCollection);
-	if (!GeometryCollection->HasAttribute("Proximity", FGeometryCollection::GeometryGroup))
+	// We estimate that each bin will contain approximately 1/20th of the faces -- we double that to provide a decent buffer
+	Bins.SetNum(NumBins);
+	for (int32 BinIdx = 0; BinIdx < NumBins; ++BinIdx)
 	{
-		// Proximity attribute
-		const FManagedArrayCollection::FConstructionParameters GeometryDependency(FGeometryCollection::GeometryGroup);
-		GeometryCollection->AddAttribute<TSet<int32>>("Proximity", FGeometryCollection::GeometryGroup, GeometryDependency);
+		Bins[BinIdx].Reserve(2 * (NumFaces/NumBins));
 	}
 
-	if (!GeometryCollection->HasGroup(FGeometryCollection::BreakingGroup))
+	for (int32 FaceIdx = 0; FaceIdx < NumFaces; ++FaceIdx)
 	{
-		// Breaking Group
-		GeometryCollection->AddAttribute<int32>("BreakingFaceIndex", FGeometryCollection::BreakingGroup);
-		GeometryCollection->AddAttribute<int32>("BreakingSourceTransformIndex", FGeometryCollection::BreakingGroup);
-		GeometryCollection->AddAttribute<int32>("BreakingTargetTransformIndex", FGeometryCollection::BreakingGroup);
-		GeometryCollection->AddAttribute<FVector3f>("BreakingRegionCentroid", FGeometryCollection::BreakingGroup);
-		GeometryCollection->AddAttribute<FVector3f>("BreakingRegionNormal", FGeometryCollection::BreakingGroup);
-		GeometryCollection->AddAttribute<float>("BreakingRegionRadius", FGeometryCollection::BreakingGroup);
-	}
-
-	auto start = std::chrono::high_resolution_clock::now();
-
-	const TManagedArray<FVector3f>& VertexArray = GeometryCollection->Vertex;
-	const TManagedArray<int32>& BoneMapArray = GeometryCollection->BoneMap;
-	const TManagedArray<FIntVector>& IndicesArray = GeometryCollection->Indices;
-	const TManagedArray<int32>& TransformIndexArray = GeometryCollection->TransformIndex;
-
-	// Breaking Data
-	TManagedArray<TSet<int32>>& ProximityArray = GeometryCollection->GetAttribute<TSet<int32>>("Proximity", FGeometryCollection::GeometryGroup);
-	TManagedArray<int32>& BreakingFaceIndexArray = GeometryCollection->GetAttribute<int32>("BreakingFaceIndex", FGeometryCollection::BreakingGroup);
-	TManagedArray<int32>& BreakingSourceTransformIndexArray = GeometryCollection->GetAttribute<int32>("BreakingSourceTransformIndex", FGeometryCollection::BreakingGroup);
-	TManagedArray<int32>& BreakingTargetTransformIndexArray = GeometryCollection->GetAttribute<int32>("BreakingTargetTransformIndex", FGeometryCollection::BreakingGroup);
-	TManagedArray<FVector3f>& BreakingRegionCentroidArray = GeometryCollection->GetAttribute<FVector3f>("BreakingRegionCentroid", FGeometryCollection::BreakingGroup);
-	TManagedArray<FVector3f>& BreakingRegionNormalArray = GeometryCollection->GetAttribute<FVector3f>("BreakingRegionNormal", FGeometryCollection::BreakingGroup);
-	TManagedArray<float>& BreakingRegionRadiusArray = GeometryCollection->GetAttribute<float>("BreakingRegionRadius", FGeometryCollection::BreakingGroup);
-
-	float DistanceThreshold = 1e-2f;
-	float DistanceThresholdSquared = DistanceThreshold * DistanceThreshold;
-	TArray<FFaceTransformData> FaceTransformDataArray;
-	FaceTransformDataArray.Empty();
-
-
-	//
-	// Create a FaceTransformDataArray for fast <FaceIndex, TransformIndex. lookup
-	// It only contains faces for GEOMETRY && RIGID, ie. not embedded geometry and not clustered
-	//
-	int32 NumFaces = GeometryCollection->NumElements(FGeometryCollection::FacesGroup);
-	for (int32 IdxFace = 0; IdxFace < NumFaces; ++IdxFace)
-	{
-		int32 TransformIndex = BoneMapArray[IndicesArray[IdxFace][0]];
-
-		//		UE_LOG(LogChaosProximity, Log, TEXT("IdxFace = %d, TransformIndex = %d>"), IdxFace, TransformIndex);
-
-		if (GeometryCollection->IsGeometry(TransformIndex) && GeometryCollection->IsRigid(TransformIndex))
-		{
-			//			UE_LOG(LogChaosProximity, Log, TEXT("ADDING TO FACETRANSFORMDATAARRAY"));
-
-			FFaceTransformData FaceData{ IdxFace,
-				TransformIndex
-			};
-			FaceTransformDataArray.Add(FaceData);
+		if (SurfaceNormals[FaceIdx].IsUnit())
+		{ 
+			int32 BestFitBin = FindBestBin(SurfaceNormals[FaceIdx]);
+			Bins[BestFitBin].Add(FaceIdx);
 		}
-		//		else
-		//		{
-		//			UE_LOG(LogChaosProximity, Log, TEXT("NOT VALID"));
-		//		}
 	}
-	NumFaces = FaceTransformDataArray.Num();
-
-#if WITH_EDITOR
-	// Create progress indicator dialog
-	static const FText SlowTaskText = NSLOCTEXT("ProximityUpdate", "UpdatingProximityBreakingText", "Updating proximity & breaking data...");
-
-	const int32 UnitProgressOutOfLoop = FMath::Max(1, (NumFaces / 50));  // One progress frame is equivalent to a minimum of 2% of the loop progress
-	const int32 NumProgressOutOfLoop = 5 * UnitProgressOutOfLoop;  // 5 tasks out of loop
-	const int32 NumProgressInLoop = NumFaces;
+}
 
-	FScopedSlowTask SlowTask(float(NumProgressOutOfLoop + NumProgressInLoop), SlowTaskText);
-	SlowTask.MakeDialog();
+int32 FGeometryCollectionProximityUtility::FindBestBin(const FVector3f& SurfaceNormal) const
+{
+	// We select the bin with the highest alignment with the surface normal
+	float BestAlignment = -1.0;
+	int32 BestBin = INDEX_NONE;
 
-	// Declare progress shortcut lambdas
-	auto EnterProgressFrame = [&SlowTask, UnitProgressOutOfLoop]()
-	{
-		SlowTask.EnterProgressFrame(float(UnitProgressOutOfLoop));
-	};
-	int32 PrevLoopCounter = 0;
-	auto EnterProgressFrameParallelLoop = [&SlowTask, &PrevLoopCounter](int32 LoopCounter)
+	for (int32 BinIdx = 0; BinIdx < NumBins; ++BinIdx)
 	{
-		if (IsInGameThread())
+		float Alignment = FVector3f::DotProduct(SurfaceNormal, BinNormals[BinIdx]);
+		if (Alignment > BestAlignment)
 		{
-			SlowTask.EnterProgressFrame(float(LoopCounter - PrevLoopCounter));
-			PrevLoopCounter = LoopCounter;
+			BestAlignment = Alignment;
+			BestBin = BinIdx;
 		}
-	};
-#else
-	auto EnterProgressFrame = []() {};
-	auto EnterProgressFrameParallelLoop = [](int32 /*LoopCounter*/) {};
-#endif
-
-	// Build reverse map between TransformIdx and GeometryGroup index
-	// #todo We already have this? TransformToGeometryIndex
-	EnterProgressFrame();
-	TMap<int32, int32> GeometryGroupIndexMap;
-	int32 NumGeometries = GeometryCollection->NumElements(FGeometryCollection::GeometryGroup);
-	for (int32 Idx = 0; Idx < NumGeometries; ++Idx)
-	{
-		GeometryGroupIndexMap.Add(TransformIndexArray[Idx], Idx);
 	}
 
-	// Transform vertices into world space
-	EnterProgressFrame();
-	TArray<FTransform> GlobalTransformArray;
-	GeometryCollectionAlgo::GlobalMatrices(GeometryCollection->Transform, GeometryCollection->Parent, GlobalTransformArray);
-
-	TArray<FVector> VertexInWorldArray;
-	int32 NumVertices = GeometryCollection->NumElements(FGeometryCollection::VerticesGroup);
-	VertexInWorldArray.SetNum(NumVertices);
-
-	FBox WorldBounds;
-	WorldBounds.IsValid = false;
-
-	for (int32 IdxVertex = 0; IdxVertex < NumVertices; ++IdxVertex)
-	{
-		FTransform Transform = GlobalTransformArray[BoneMapArray[IdxVertex]];
-		FVector VertexInWorld = Transform.TransformPosition(VertexArray[IdxVertex]);
-
-		VertexInWorldArray[IdxVertex] = VertexInWorld;
-		WorldBounds += VertexInWorld;
-	}
-
-	// Make an Octree
-	EnterProgressFrame();
-	TUniquePtr<FProximityTriangleOctree> MeshTriOctree = MakeUnique<FProximityTriangleOctree>(WorldBounds.GetCenter(), WorldBounds.GetExtent().GetMax());
-
-	for (int32 ii = 0; ii < FaceTransformDataArray.Num(); ++ii)
-	{
-		auto& FaceTransformDataRef = FaceTransformDataArray[ii];
-		int32 TriIndex = FaceTransformDataRef.FaceIdx;
-		// Grab the vertex indices and points for this triangle
-		FProximityTriangle MeshTri;
-
-		MeshTri.Vertices[0] = VertexInWorldArray[IndicesArray[TriIndex][0]];
-		MeshTri.Vertices[1] = VertexInWorldArray[IndicesArray[TriIndex][1]];
-		MeshTri.Vertices[2] = VertexInWorldArray[IndicesArray[TriIndex][2]];
-		MeshTri.Normal = FVector::CrossProduct(MeshTri.Vertices[1] - MeshTri.Vertices[0], MeshTri.Vertices[2] - MeshTri.Vertices[0]).GetSafeNormal();
-		MeshTri.ArrayIndex = ii;
-
-		FBox &TriBox = MeshTri.Bounds;
-		TriBox.Min.X = FMath::Min3(MeshTri.Vertices[0].X, MeshTri.Vertices[1].X, MeshTri.Vertices[2].X);
-		TriBox.Min.Y = FMath::Min3(MeshTri.Vertices[0].Y, MeshTri.Vertices[1].Y, MeshTri.Vertices[2].Y);
-		TriBox.Min.Z = FMath::Min3(MeshTri.Vertices[0].Z, MeshTri.Vertices[1].Z, MeshTri.Vertices[2].Z);
-
-		TriBox.Max.X = FMath::Max3(MeshTri.Vertices[0].X, MeshTri.Vertices[1].X, MeshTri.Vertices[2].X);
-		TriBox.Max.Y = FMath::Max3(MeshTri.Vertices[0].Y, MeshTri.Vertices[1].Y, MeshTri.Vertices[2].Y);
-		TriBox.Max.Z = FMath::Max3(MeshTri.Vertices[0].Z, MeshTri.Vertices[1].Z, MeshTri.Vertices[2].Z);
-
-		FaceTransformDataRef.Bounds = TriBox;
-
-		MeshTri.BoxCenterAndExtent = FBoxCenterAndExtent(TriBox);
-		MeshTriOctree->AddElement(MeshTri);
-	}
-
-	FCriticalSection Mutex;
-	TSet<FOverlappingFacePair> OverlappingFacePairSet;
-
-	ParallelFor(FaceTransformDataArray.Num(), [&](int32 FaceTransformArrayIdx) {
-		EnterProgressFrameParallelLoop(FaceTransformArrayIdx);
-
-		const auto& FaceTransformDataRef = FaceTransformDataArray[FaceTransformArrayIdx];
-		TSet<FOverlappingFacePair> LocalOverlappingFacePairSet;
-		FVertexPair VertexPairArray[9];
-
-		int32 IdxFace = FaceTransformDataRef.FaceIdx;
-
-		VertexPairArray[0].Vertex1 = VertexInWorldArray[IndicesArray[IdxFace][0]];
-		VertexPairArray[1].Vertex1 = VertexInWorldArray[IndicesArray[IdxFace][0]];
-		VertexPairArray[2].Vertex1 = VertexInWorldArray[IndicesArray[IdxFace][0]];
-
-		VertexPairArray[3].Vertex1 = VertexInWorldArray[IndicesArray[IdxFace][1]];
-		VertexPairArray[4].Vertex1 = VertexInWorldArray[IndicesArray[IdxFace][1]];
-		VertexPairArray[5].Vertex1 = VertexInWorldArray[IndicesArray[IdxFace][1]];
-
-		VertexPairArray[6].Vertex1 = VertexInWorldArray[IndicesArray[IdxFace][2]];
-		VertexPairArray[7].Vertex1 = VertexInWorldArray[IndicesArray[IdxFace][2]];
-		VertexPairArray[8].Vertex1 = VertexInWorldArray[IndicesArray[IdxFace][2]];
-
-		const FBox& ThisFaceBounds = FaceTransformDataRef.Bounds;
-		TArray<FFaceTransformData> OtherFaceTransformDataArray;
-
-		// 	Query the Octree
-		OtherFaceTransformDataArray.Reserve(NumFaces);
+	return BestBin;
+}
 
-		MeshTriOctree->FindElementsWithBoundsTest(ThisFaceBounds, [&OtherFaceTransformDataArray, &FaceTransformDataArray](const FProximityTriangle& OctreePolygon)
-		{
-			OtherFaceTransformDataArray.Add(FaceTransformDataArray[OctreePolygon.ArrayIndex]);
+void FGeometryCollectionProximityUtility::FindContactingFaces()
+{
+	FaceContacts.SetNum(NumFaces);
+	
+	// The bin hemispheres are symmetrical so we only need to iterate half of them.
+	// We are comparing antipodal bins.
+	ParallelFor((NumBins / 2), [this](int32 BinIdx)
+		{ 
+			FindContactPairs(BinIdx, (NumBins/2) + BinIdx); 
 		});
+}
 
-		for (auto& OtherFaceTransformDataRef : OtherFaceTransformDataArray)
+void FGeometryCollectionProximityUtility::FindContactPairs(int32 Zenith, int32 Nadir)
+{
+	for (int32 ZenithIdx : Bins[Zenith])
+	{
+		for (int32 NadirIdx : Bins[Nadir])
 		{
-			int32 IdxOtherFace = OtherFaceTransformDataRef.FaceIdx;
-
-			if (FaceTransformDataRef.TransformIndex != OtherFaceTransformDataRef.TransformIndex)
+			// Are the faces parallel?
+			if (AreNormalsOpposite(SurfaceNormals[ZenithIdx], SurfaceNormals[NadirIdx]))
 			{
-				//
-				// Vertex coincidence test
-				//
-				bool VertexCoincidenceTestFoundOverlappingFaces = false;
+				// Are faces co-planar?
+				if (AreFacesCoPlanar(ZenithIdx, NadirIdx))
 				{
-					VertexPairArray[0].Vertex2 = VertexInWorldArray[IndicesArray[IdxOtherFace][0]];
-					VertexPairArray[1].Vertex2 = VertexInWorldArray[IndicesArray[IdxOtherFace][1]];
-					VertexPairArray[2].Vertex2 = VertexInWorldArray[IndicesArray[IdxOtherFace][2]];
-
-					VertexPairArray[3].Vertex2 = VertexInWorldArray[IndicesArray[IdxOtherFace][0]];
-					VertexPairArray[4].Vertex2 = VertexInWorldArray[IndicesArray[IdxOtherFace][1]];
-					VertexPairArray[5].Vertex2 = VertexInWorldArray[IndicesArray[IdxOtherFace][2]];
-
-					VertexPairArray[6].Vertex2 = VertexInWorldArray[IndicesArray[IdxOtherFace][0]];
-					VertexPairArray[7].Vertex2 = VertexInWorldArray[IndicesArray[IdxOtherFace][1]];
-					VertexPairArray[8].Vertex2 = VertexInWorldArray[IndicesArray[IdxOtherFace][2]];
-
-					int32 NumCoincideVertices = 0;
-					for (int32 Idx = 0, ni = 9; Idx < ni; ++Idx)
+					// Do triangles overlap?
+					if (DoFacesOverlap(ZenithIdx, NadirIdx))
 					{
-						if (VertexPairArray[Idx].DistanceSquared() < DistanceThresholdSquared)
-						{
-							NumCoincideVertices++;
-						}
-					}
-
-					if (NumCoincideVertices >= 3)
-					{
-						VertexCoincidenceTestFoundOverlappingFaces = true;
-
-						if (!LocalOverlappingFacePairSet.Contains(FOverlappingFacePair{ FMath::Min(IdxFace, IdxOtherFace), FMath::Max(IdxFace, IdxOtherFace) }))
-						{
-							LocalOverlappingFacePairSet.Add(FOverlappingFacePair{ FMath::Min(IdxFace, IdxOtherFace), FMath::Max(IdxFace, IdxOtherFace) });
-						}
+						// Add contacts to both faces
+						FaceContacts[ZenithIdx].Add(NadirIdx);
+						FaceContacts[NadirIdx].Add(ZenithIdx);
 					}
 				}
+			}
+		}
+	}
+}
 
-				//
-				// FaceN and OtherFaceN are parallel and points of Face are in OtherFace test
-				//
-				if (!VertexCoincidenceTestFoundOverlappingFaces)
-				{
-					FVector Edge1(VertexInWorldArray[IndicesArray[IdxFace][1]] - VertexInWorldArray[IndicesArray[IdxFace][0]]);
-					FVector Edge2(VertexInWorldArray[IndicesArray[IdxFace][2]] - VertexInWorldArray[IndicesArray[IdxFace][0]]);
-					FVector FaceN(Edge1 ^ Edge2);
+bool FGeometryCollectionProximityUtility::AreNormalsOpposite(const FVector3f& Normal0, const FVector3f& Normal1) const
+{
+	return FVector3f::DotProduct(Normal0, Normal1) < (-1.0f + KINDA_SMALL_NUMBER);
+}
 
-					FVector OtherEdge1(VertexInWorldArray[IndicesArray[IdxOtherFace][1]] - VertexInWorldArray[IndicesArray[IdxOtherFace][0]]);
-					FVector OtherEdge2(VertexInWorldArray[IndicesArray[IdxOtherFace][2]] - VertexInWorldArray[IndicesArray[IdxOtherFace][0]]);
-					FVector OtherFaceN(OtherEdge1 ^ OtherEdge2);
+bool FGeometryCollectionProximityUtility::AreFacesCoPlanar(int32 Idx0, int32 Idx1) const
+{
+	// Assumes that faces have already been determined to be parallel.
 
-					if (FVector::Parallel(FaceN, OtherFaceN, 1e-1f))
-					{
-						FVector FaceCenter((VertexInWorldArray[IndicesArray[IdxFace][0]] + VertexInWorldArray[IndicesArray[IdxFace][1]] + VertexInWorldArray[IndicesArray[IdxFace][2]]) / 3.f);
-						FVector OtherFaceCenter = (VertexInWorldArray[IndicesArray[IdxOtherFace][0]] + VertexInWorldArray[IndicesArray[IdxOtherFace][1]] + VertexInWorldArray[IndicesArray[IdxOtherFace][2]]) / 3.f;
-						FVector PointInFace1((VertexInWorldArray[IndicesArray[IdxFace][0]] + FaceCenter) / 2.f);
-						FVector PointInFace2((VertexInWorldArray[IndicesArray[IdxFace][1]] + FaceCenter) / 2.f);
-						FVector PointInFace3((VertexInWorldArray[IndicesArray[IdxFace][2]] + FaceCenter) / 2.f);
-
-						FVector PointInFaceA[3] = { VertexInWorldArray[IndicesArray[IdxFace][0]], VertexInWorldArray[IndicesArray[IdxFace][1]], VertexInWorldArray[IndicesArray[IdxFace][2]] };
-						FVector PointInFaceB[3] = { VertexInWorldArray[IndicesArray[IdxOtherFace][0]], VertexInWorldArray[IndicesArray[IdxOtherFace][1]], VertexInWorldArray[IndicesArray[IdxOtherFace][2]] };
-
-						int32 CoincidentVerts = 0;
-						for (int32 ii = 0; ii < 3; ++ii)
-						{
-							for (int32 kk = 0; kk < 3; ++kk)
-							{
-								if ((PointInFaceA[ii] - PointInFaceB[kk]).SizeSquared() < 1e-1)
-								{
-									++CoincidentVerts;
-								}
-							}
-						}
-
-						if (CoincidentVerts > 1)
-						{
-							if (!LocalOverlappingFacePairSet.Contains(FOverlappingFacePair{ FMath::Min(IdxFace, IdxOtherFace), FMath::Max(IdxFace, IdxOtherFace) }))
-							{
-								LocalOverlappingFacePairSet.Add(FOverlappingFacePair{ FMath::Min(IdxFace, IdxOtherFace), FMath::Max(IdxFace, IdxOtherFace) });
-							}
-						}
-
-						// Check if points in Face are in OtherFace
-						else if ((FaceCenter - OtherFaceCenter).SizeSquared() < 1e-1)
-						{
-							if (!LocalOverlappingFacePairSet.Contains(FOverlappingFacePair{ FMath::Min(IdxFace, IdxOtherFace), FMath::Max(IdxFace, IdxOtherFace) }))
-							{
-								LocalOverlappingFacePairSet.Add(FOverlappingFacePair{ FMath::Min(IdxFace, IdxOtherFace), FMath::Max(IdxFace, IdxOtherFace) });
-							}
-						}
-						else if (IsPointInsideOfTriangle(FaceCenter, VertexInWorldArray[IndicesArray[IdxOtherFace][0]], VertexInWorldArray[IndicesArray[IdxOtherFace][1]], VertexInWorldArray[IndicesArray[IdxOtherFace][2]], 1e-1f) ||
-							IsPointInsideOfTriangle(PointInFace1, VertexInWorldArray[IndicesArray[IdxOtherFace][0]], VertexInWorldArray[IndicesArray[IdxOtherFace][1]], VertexInWorldArray[IndicesArray[IdxOtherFace][2]], 1e-1f) ||
-							IsPointInsideOfTriangle(PointInFace2, VertexInWorldArray[IndicesArray[IdxOtherFace][0]], VertexInWorldArray[IndicesArray[IdxOtherFace][1]], VertexInWorldArray[IndicesArray[IdxOtherFace][2]], 1e-1f) ||
-							IsPointInsideOfTriangle(PointInFace3, VertexInWorldArray[IndicesArray[IdxOtherFace][0]], VertexInWorldArray[IndicesArray[IdxOtherFace][1]], VertexInWorldArray[IndicesArray[IdxOtherFace][2]], 1e-1f))
-						{
-							if (!LocalOverlappingFacePairSet.Contains(FOverlappingFacePair{ FMath::Min(IdxFace, IdxOtherFace), FMath::Max(IdxFace, IdxOtherFace) }))
-							{
-								LocalOverlappingFacePairSet.Add(FOverlappingFacePair{ FMath::Min(IdxFace, IdxOtherFace), FMath::Max(IdxFace, IdxOtherFace) });
-							}
-						}
-						else
-						{
-							PointInFace1 = (VertexInWorldArray[IndicesArray[IdxOtherFace][0]] + OtherFaceCenter) / 2.f;
-							PointInFace2 = (VertexInWorldArray[IndicesArray[IdxOtherFace][1]] + OtherFaceCenter) / 2.f;
-							PointInFace3 = (VertexInWorldArray[IndicesArray[IdxOtherFace][2]] + OtherFaceCenter) / 2.f;
-
-							// Check if points in OtherFace are in Face
-							if (IsPointInsideOfTriangle(OtherFaceCenter, VertexInWorldArray[IndicesArray[IdxFace][0]], VertexInWorldArray[IndicesArray[IdxFace][1]], VertexInWorldArray[IndicesArray[IdxFace][2]], 1e-1f) ||
-								IsPointInsideOfTriangle(PointInFace1, VertexInWorldArray[IndicesArray[IdxFace][0]], VertexInWorldArray[IndicesArray[IdxFace][1]], VertexInWorldArray[IndicesArray[IdxFace][2]], 1e-1f) ||
-								IsPointInsideOfTriangle(PointInFace2, VertexInWorldArray[IndicesArray[IdxFace][0]], VertexInWorldArray[IndicesArray[IdxFace][1]], VertexInWorldArray[IndicesArray[IdxFace][2]], 1e-1f) ||
-								IsPointInsideOfTriangle(PointInFace3, VertexInWorldArray[IndicesArray[IdxFace][0]], VertexInWorldArray[IndicesArray[IdxFace][1]], VertexInWorldArray[IndicesArray[IdxFace][2]], 1e-1f))
-							{
-								if (!LocalOverlappingFacePairSet.Contains(FOverlappingFacePair{ FMath::Min(IdxFace, IdxOtherFace), FMath::Max(IdxFace, IdxOtherFace) }))
-								{
-									LocalOverlappingFacePairSet.Add(FOverlappingFacePair{ FMath::Min(IdxFace, IdxOtherFace), FMath::Max(IdxFace, IdxOtherFace) });
-								}
-							}
-						}
-					}
-				}
-			}
-		}
-		Mutex.Lock();
-		OverlappingFacePairSet.Append(LocalOverlappingFacePairSet);
-		Mutex.Unlock();
-	});
+	const TManagedArray<FIntVector>& Indices = Collection->Indices;
 
-	if (!OverlappingFacePairSet.Num())
-	{
-		return;
-	}
+	FVector3f SamplePoint = TransformedVertices[Indices[Idx0].X];
+	FVector3f PlaneOrigin = TransformedVertices[Indices[Idx1].X];
+	FVector3f PlaneNormal = SurfaceNormals[Idx1];
 
-	// Populate Proximity, BreakingFaceIndex, BreakingSourceTransformIndex, BreakingTargetTransformIndex structures
-	EnterProgressFrame();
-	for (int32 IdxGeometry = 0; IdxGeometry < NumGeometries; ++IdxGeometry)
-	{
-		ProximityArray[IdxGeometry].Empty();
-	}
+	return FMath::Abs(FVector3f::DotProduct((SamplePoint - PlaneOrigin), PlaneNormal)) < KINDA_SMALL_NUMBER;
+}
 
-	TArray<int32> AllBreakingFaceIndexArray;
-	TArray<int32> AllBreakingSourceTransformIndexArray;
-	TArray<int32> AllBreakingTargetTransformIndexArray;
-
-	int32 NewArraySize = 2 * OverlappingFacePairSet.Num();
-	AllBreakingFaceIndexArray.SetNum(NewArraySize);
-	AllBreakingSourceTransformIndexArray.SetNum(NewArraySize);
-	AllBreakingTargetTransformIndexArray.SetNum(NewArraySize);
-
-	//
-	// Create the {BreakingSourceTransformIndex, BreakingTargetTransformIndex} <-> FaceIndex data in the
-	// AllBreakingFaceIndexArray, AllBreakingSourceTransformIndexArray, AllBreakingTargetTransformIndexArray arrays
-	// This contains every connected face pairs, a lot of data
-	//
-	int32 IdxBreak = 0;
-	for (auto& OverlappingFacePair : OverlappingFacePairSet)
-	{
-		int32 TransformIndex1 = BoneMapArray[IndicesArray[OverlappingFacePair.FaceIdx1][0]];
-		int32 TransformIndex2 = BoneMapArray[IndicesArray[OverlappingFacePair.FaceIdx2][0]];
+bool FGeometryCollectionProximityUtility::DoFacesOverlap(int32 Idx0, int32 Idx1) const
+{
+	// Assumes that the faces are coplanar
 
-		check(GeometryCollection->IsGeometry(TransformIndex1) && !GeometryCollection->IsClustered(TransformIndex1));
-		check(GeometryCollection->IsGeometry(TransformIndex2) && !GeometryCollection->IsClustered(TransformIndex2));
+	const TManagedArray<FIntVector>& Indices = Collection->Indices;
+	
+	// Project the first triangle into its normal plane
+	FVector3f Basis0 = TransformedVertices[Indices[Idx0].Y] - TransformedVertices[Indices[Idx0].X];
+	Basis0.Normalize();
+	FVector3f Basis1 = FVector3f::CrossProduct(SurfaceNormals[Idx0], Basis0);
+	Basis1.Normalize();
+	
+	TStaticArray<FVector2D,3> T0; 
+	// T0[0] is the origin of the system
+	T0[0] = FVector2D(0.f, 0.f);
+	T0[1] = FVector2D(FVector3f::DotProduct(TransformedVertices[Indices[Idx0].Y], Basis0), FVector3f::DotProduct(TransformedVertices[Indices[Idx0].Y], Basis1));
+	T0[2] = FVector2D(FVector3f::DotProduct(TransformedVertices[Indices[Idx0].Z], Basis0), FVector3f::DotProduct(TransformedVertices[Indices[Idx0].Z], Basis1));
+
+	// Project the second triangle into these coordinates. We reverse the winding order to flip the normal.
+	FVector3f Point0 = TransformedVertices[Indices[Idx1].Z] - TransformedVertices[Indices[Idx0].Z];
+	FVector3f Point1 = TransformedVertices[Indices[Idx1].Y] - TransformedVertices[Indices[Idx0].Y];
+	FVector3f Point2 = TransformedVertices[Indices[Idx1].X] - TransformedVertices[Indices[Idx0].X];
+	TStaticArray<FVector2D,3> T1;
+	T1[0] = FVector2D(FVector3f::DotProduct(Point0, Basis0), FVector3f::DotProduct(Point0, Basis1));
+	T1[1] = FVector2D(FVector3f::DotProduct(Point1, Basis0), FVector3f::DotProduct(Point1, Basis1));
+	T1[2] = FVector2D(FVector3f::DotProduct(Point2, Basis0), FVector3f::DotProduct(Point2, Basis1));
+
+	return TrianglesIntersect(T0, T1);
+}
 
-		if (!ProximityArray[GeometryGroupIndexMap[TransformIndex1]].Contains(GeometryGroupIndexMap[TransformIndex2]))
-		{
-			ProximityArray[GeometryGroupIndexMap[TransformIndex1]].Add(GeometryGroupIndexMap[TransformIndex2]);
-		}
+bool FGeometryCollectionProximityUtility::TrianglesIntersect(const TStaticArray<FVector2D,3>& T0, const TStaticArray<FVector2D,3>& T1)
+{
+	// Test if one of the triangles has a side with all of the other triangle's points on the outside.
+	float Normal0 = (T0[1].X - T0[0].X) * (T0[2].Y - T0[0].Y) - (T0[1].Y - T0[0].Y) * (T0[2].X - T0[0].X);
+	float Normal1 = (T1[1].X - T1[0].X) * (T1[2].Y - T1[0].Y) - (T1[1].Y - T1[0].Y) * (T1[2].X - T1[0].X);
+
+	return !(Cross(T1, T0[0], T0[1], Normal0) ||
+		Cross(T1, T0[1], T0[2], Normal0) ||
+		Cross(T1, T0[2], T0[0], Normal0) ||
+		Cross(T0, T1[0], T1[1], Normal1) ||
+		Cross(T0, T1[1], T1[2], Normal1) ||
+		Cross(T0, T1[1], T1[0], Normal1));
+}
 
-		AllBreakingFaceIndexArray[IdxBreak] = OverlappingFacePair.FaceIdx1;
-		AllBreakingSourceTransformIndexArray[IdxBreak] = TransformIndex1;
-		AllBreakingTargetTransformIndexArray[IdxBreak] = TransformIndex2;
-		IdxBreak++;
+bool FGeometryCollectionProximityUtility::Cross(const TStaticArray<FVector2D,3>& Points, const FVector2D& B, const FVector2D& C, float Normal)
+{
+	float CyBy = C.Y - B.Y;
+	float CxBx = C.X - B.X;
+	const FVector2D& Pa = Points[0];
+	const FVector2D& Pb = Points[1];
+	const FVector2D& Pc = Points[2];
+
+	return !(
+		(((Pa.X - B.X) * CyBy - (Pa.Y - B.Y) * CxBx) * Normal < 0) ||
+		(((Pb.X - B.X) * CyBy - (Pb.Y - B.Y) * CxBx) * Normal < 0) ||
+		(((Pc.X - B.X) * CyBy - (Pc.Y - B.Y) * CxBx) * Normal < 0));
+}
 
-		if (!ProximityArray[GeometryGroupIndexMap[TransformIndex2]].Contains(GeometryGroupIndexMap[TransformIndex1]))
+void FGeometryCollectionProximityUtility::GenerateSurfaceNormals()
+{
+	// Generate surface normal for each face
+	SurfaceNormals.SetNum(NumFaces);
+	ParallelFor(NumFaces, [this](int32 FaceIdx)
 		{
-			ProximityArray[GeometryGroupIndexMap[TransformIndex2]].Add(GeometryGroupIndexMap[TransformIndex1]);
-		}
+			const TManagedArray<FIntVector>& Indices = Collection->Indices;
 
-		AllBreakingFaceIndexArray[IdxBreak] = OverlappingFacePair.FaceIdx2;
-		AllBreakingSourceTransformIndexArray[IdxBreak] = TransformIndex2;
-		AllBreakingTargetTransformIndexArray[IdxBreak] = TransformIndex1;
-		IdxBreak++;
-	}
+			FVector3f Edge0 = (TransformedVertices[Indices[FaceIdx].X] - TransformedVertices[Indices[FaceIdx].Y]);
+			FVector3f Edge1 = (TransformedVertices[Indices[FaceIdx].Z] - TransformedVertices[Indices[FaceIdx].Y]);
+			SurfaceNormals[FaceIdx] = FVector3f::CrossProduct(Edge0, Edge1);
+			SurfaceNormals[FaceIdx].Normalize();
+		});
+}
 
-	//
-	// Store the data as a MultiMap<{BreakingSourceTransformIndex, BreakingTargetTransformIndex}, FaceIndex>
-	//
-	EnterProgressFrame();
-	TMultiMap<FOverlappingFacePairTransformIndex, int32> FaceByConnectedTransformsMap;
-	FaceByConnectedTransformsMap.Reserve(NumFaces);
-	if (AllBreakingFaceIndexArray.Num())
-	{
-		for (int32 Idx = 0, ni = AllBreakingFaceIndexArray.Num(); Idx < ni; ++Idx)
-		{
-			FaceByConnectedTransformsMap.Add(FOverlappingFacePairTransformIndex{ AllBreakingSourceTransformIndexArray[Idx], AllBreakingTargetTransformIndexArray[Idx] },
-				AllBreakingFaceIndexArray[Idx]);
-		}
-	}
+void FGeometryCollectionProximityUtility::GenerateFaceToGeometry()
+{
+	// Create a map from face back to geometry
+	const TManagedArray<int32>& FaceStart = Collection->FaceStart;
+	const TManagedArray<int32>& FaceCount = Collection->FaceCount;
 
-	//
-	// Get all the keys from the MultiMap
-	//
-	TArray<FOverlappingFacePairTransformIndex> FaceByConnectedTransformsMapKeys;
-	FaceByConnectedTransformsMap.GenerateKeyArray(FaceByConnectedTransformsMapKeys);
-
-	//
-	// Delete all the duplicates
-	//
-	TSet<FOverlappingFacePairTransformIndex> FaceByConnectedTransformsMapKeysSet;
-	for (int32 Idx = 0; Idx < FaceByConnectedTransformsMapKeys.Num(); ++Idx)
+	FaceToGeometry.Init(INDEX_NONE, NumFaces);
+	for (int32 GeometryIdx = 0; GeometryIdx < Collection->NumElements(FGeometryCollection::GeometryGroup); ++GeometryIdx)
 	{
-		if (!FaceByConnectedTransformsMapKeysSet.Contains(FaceByConnectedTransformsMapKeys[Idx]))
+		for (int32 FaceOffset = 0; FaceOffset < FaceCount[GeometryIdx]; ++FaceOffset)
 		{
-			FaceByConnectedTransformsMapKeysSet.Add(FaceByConnectedTransformsMapKeys[Idx]);
+			int32 FaceIdx = FaceStart[GeometryIdx] + FaceOffset;
+			FaceToGeometry[FaceIdx] = GeometryIdx;
 		}
 	}
+}
 
-	int LastIndex = GeometryCollection->AddElements(FaceByConnectedTransformsMapKeysSet.Num() - BreakingFaceIndexArray.Num(), FGeometryCollection::BreakingGroup);
-
-	//
-	// Get one Face for every {BreakingSourceTransformIndex, BreakingTargetTransformIndex} pair and store the data in
-	// BreakingFaceIndexArray, BreakingSourceTransformIndexArray, BreakingTargetTransformIndexArray
-	//
-	IdxBreak = 0;
-	for (auto& Elem : FaceByConnectedTransformsMapKeysSet)
-	{
-		TArray<int32> FaceIndexArray;
-		FaceByConnectedTransformsMap.MultiFind(Elem, FaceIndexArray);
+void FGeometryCollectionProximityUtility::TransformVertices()
+{
+	TransformedVertices.SetNum(Collection->NumElements(FGeometryCollection::VerticesGroup));
 
-		// Find the centroid of the region and save it into BreakingRegionCentroidArray
-		FVector3f Centroid = FVector3f(ForceInitToZero);
-		float TotalArea = 0.f;
-		for (int32 LocalIdxFace = 0; LocalIdxFace < FaceIndexArray.Num(); ++LocalIdxFace)
+	TArray<FTransform> GlobalTransformArray;
+	GeometryCollectionAlgo::GlobalMatrices(Collection->Transform, Collection->Parent, GlobalTransformArray);
+	
+	ParallelFor(Collection->NumElements(FGeometryCollection::VerticesGroup), [this, &GlobalTransformArray](int32 VertIdx)
 		{
-			const FVector3f& Vertex0 = VertexArray[IndicesArray[FaceIndexArray[LocalIdxFace]][0]];
-			const FVector3f& Vertex1 = VertexArray[IndicesArray[FaceIndexArray[LocalIdxFace]][1]];
-			const FVector3f& Vertex2 = VertexArray[IndicesArray[FaceIndexArray[LocalIdxFace]][2]];
-
-			FVector3f FaceCentroid((Vertex0 + Vertex1 + Vertex2) / 3.f);
-			float FaceArea = 0.5f * ((Vertex1 - Vertex0) ^ (Vertex2 - Vertex0)).Size();
-
-			Centroid = (TotalArea * Centroid + FaceArea * FaceCentroid) / (TotalArea + FaceArea);
-
-			TotalArea += FaceArea;
-		}
-		BreakingRegionCentroidArray[IdxBreak] = Centroid;
-
-		// Find the inner radius of the region and save it into BreakingRegionRadiusArray
-		float RadiusMin = FLT_MAX;
-		float RadiusMax = FLT_MIN;
+			const TManagedArray<int32>& BoneMap = Collection->BoneMap;
+			const TManagedArray<FVector3f>& Vertex = Collection->Vertex;
 
-		TArray<FVector3f> TestPoints;
-		for (int32 LocalIdxFace = 0; LocalIdxFace < FaceIndexArray.Num(); ++LocalIdxFace)
-		{
-			for (int32 Idx = 0; Idx < 3; Idx++)
-			{
-				TestPoints.Add(VertexArray[IndicesArray[FaceIndexArray[LocalIdxFace]][Idx]]);
-			}
-		}
+			const FTransform& GlobalTransform = GlobalTransformArray[BoneMap[VertIdx]];
+			TransformedVertices[VertIdx] = GlobalTransform.TransformPosition(Vertex[VertIdx]);
+		});
 
-		for (int32 IdxPoint = 0; IdxPoint < TestPoints.Num(); ++IdxPoint)
-		{
-			float Distance = (Centroid - TestPoints[IdxPoint]).Size();
-			if (Distance < RadiusMin)
-			{
-				RadiusMin = Distance;
-			}
-			if (Distance > RadiusMax)
-			{
-				RadiusMax = Distance;
-			}
-		}
-		BreakingRegionRadiusArray[IdxBreak] = RadiusMin;
-
-		// Normal
-		const FVector3f& VertexA = VertexArray[IndicesArray[FaceIndexArray[0]][0]];
-		const FVector3f& VertexB = VertexArray[IndicesArray[FaceIndexArray[0]][1]];
-		const FVector3f& VertexC = VertexArray[IndicesArray[FaceIndexArray[0]][2]];
-		BreakingRegionNormalArray[IdxBreak] = ((VertexA - VertexB) ^ (VertexC - VertexB)).GetSafeNormal();
-
-		// grab the first face from the region and save it into BreakingFaceIndexArray
-		BreakingFaceIndexArray[IdxBreak] = FaceIndexArray[0];
-		BreakingSourceTransformIndexArray[IdxBreak] = Elem.TransformIdx1;
-		BreakingTargetTransformIndexArray[IdxBreak] = Elem.TransformIdx2;
-		IdxBreak++;
-	}
-	auto finish = std::chrono::high_resolution_clock::now();
-	std::chrono::duration<double> elapsed = finish - start;
-	UE_LOG(LogChaosProximity, Log, TEXT("Elapsed Time = %fs>"), elapsed.count());
 }
 
-
-void CollectTopNodeConnections(FGeometryCollection* GeometryCollection, int32 Index, TSet<int32>& OutConnections, int32 OperatingLevel)
+void FGeometryCollectionProximityUtility::ExtendFaceProximityToGeometry()
 {
-	const TManagedArray<int32>& Level = GeometryCollection->GetAttribute<int32>("Level", FGeometryCollection::TransformGroup);
-	if (OperatingLevel < 0)
+	if (!Collection->HasAttribute("Proximity", FGeometryCollection::GeometryGroup))
 	{
-		OperatingLevel = Level[Index];
+		Collection->AddAttribute<TSet<int32>>("Proximity", FGeometryCollection::GeometryGroup);
 	}
-	
-	const TManagedArray<TSet<int32>>& Children = GeometryCollection->Children;
-	if (GeometryCollection->IsRigid(Index)) // leaf node
-	{
-		const TManagedArray<TSet<int32>>& Proximity = GeometryCollection->GetAttribute<TSet<int32>>("Proximity", FGeometryCollection::GeometryGroup);
-		const TManagedArray<int32>& GeometryToTransformIndex = GeometryCollection->TransformIndex;
-		const TManagedArray<int32>& TransformToGeometryIndex = GeometryCollection->TransformToGeometryIndex;
 
+	TManagedArray<TSet<int32>>& Proximity = Collection->GetAttribute<TSet<int32>>("Proximity", FGeometryCollection::GeometryGroup);
 
-		for (int32 Neighbor : Proximity[TransformToGeometryIndex[Index]])
-		{
-			int32 NeighborTransformIndex = GeometryToTransformIndex[Neighbor];
-			OutConnections.Add(FGeometryCollectionClusteringUtility::GetParentOfBoneAtSpecifiedLevel(GeometryCollection, NeighborTransformIndex, OperatingLevel));
-		}
-	}
-	else
+	for (int32 FaceIdx = 0; FaceIdx < NumFaces; ++FaceIdx)
 	{
-		for (int32 ChildIndex : Children[Index])
+		TSet<int32>& CurrentContactList = Proximity[FaceToGeometry[FaceIdx]];
+		for (int32 ContactFaceIdx : FaceContacts[FaceIdx])
 		{
-			CollectTopNodeConnections(GeometryCollection, ChildIndex, OutConnections, OperatingLevel);
+			CurrentContactList.Add(FaceToGeometry[ContactFaceIdx]);
 		}
 	}
 }
+
diff --git a/Engine/Source/Runtime/Experimental/Chaos/Public/GeometryCollection/GeometryCollectionProximityUtility.h b/Engine/Source/Runtime/Experimental/Chaos/Public/GeometryCollection/GeometryCollectionProximityUtility.h
--- a/Engine/Source/Runtime/Experimental/Chaos/Public/GeometryCollection/GeometryCollectionProximityUtility.h
+++ b/Engine/Source/Runtime/Experimental/Chaos/Public/GeometryCollection/GeometryCollectionProximityUtility.h
@@ -3,7 +3,7 @@
 #pragma once
 
 #include "CoreMinimal.h"
-//#include "GeometryCollection/ManagedArrayCollection.h"
+#include "Containers/StaticArray.h"
 
 class FGeometryCollection;
 
@@ -12,77 +12,39 @@ class FGeometryCollection;
 class CHAOS_API FGeometryCollectionProximityUtility
 {
 public:
+	FGeometryCollectionProximityUtility(FGeometryCollection* InCollection);
 
-	struct FFaceTransformData {
-		int32 FaceIdx;
-		int32 TransformIndex;
-		FBox Bounds;
-	};
+	void UpdateProximity();
 
-	struct FVertexPair
-	{
-		FVector Vertex1, Vertex2;	// World space vertex positions
-		FVector::FReal Distance() { return (Vertex1 - Vertex2).Size(); }
-		FVector::FReal DistanceSquared() { return (Vertex1 - Vertex2).SizeSquared(); }
-	};
-
-	struct FOverlappingFacePair
-	{
-		int32 FaceIdx1;
-		int32 FaceIdx2;
-
-		friend inline uint32 GetTypeHash(const FOverlappingFacePair& Other)
-		{
-			return HashCombine(GetTypeHash(Other.FaceIdx1), GetTypeHash(Other.FaceIdx2));
-		}
-
-		friend bool operator==(const FOverlappingFacePair& A, const FOverlappingFacePair& B)
-		{
-			return A.FaceIdx1 == B.FaceIdx1 && A.FaceIdx2 == B.FaceIdx2;
-		}
-	};
-
-	struct FOverlappingFacePairTransformIndex
-	{
-		int32 TransformIdx1;
-		int32 TransformIdx2;
-
-		friend inline uint32 GetTypeHash(const FOverlappingFacePairTransformIndex& Other)
-		{
-			return HashCombine(GetTypeHash(Other.TransformIdx1), GetTypeHash(Other.TransformIdx2));
-		}
-
-		friend bool operator==(const FOverlappingFacePairTransformIndex& A, const FOverlappingFacePairTransformIndex& B)
-		{
-			return A.TransformIdx1 == B.TransformIdx1 && A.TransformIdx2 == B.TransformIdx2;
-		}
-	};
-
-	struct FFaceEdge
-	{
-		int32 VertexIdx1;
-		int32 VertexIdx2;
-
-		friend inline uint32 GetTypeHash(const FFaceEdge& Other)
-		{
-			return HashCombine(GetTypeHash(Other.VertexIdx1), GetTypeHash(Other.VertexIdx2));
-		}
-
-		friend bool operator==(const FFaceEdge& A, const FFaceEdge& B)
-		{
-			return A.VertexIdx1 == B.VertexIdx1 && A.VertexIdx2 == B.VertexIdx2;
-		}
-	};
+private:
+	void PrepFaceData();
+	void BinFaces();
+	void FindContactingFaces();
+	void ExtendFaceProximityToGeometry();
+
+	void TransformVertices();
+	void GenerateSurfaceNormals();
+	void GenerateFaceToGeometry();
+	int32 FindBestBin(const FVector3f& SurfaceNormal) const;
+	void FindContactPairs(int32 Zenith, int32 Nadir);
+	bool AreNormalsOpposite(const FVector3f& Normal0, const FVector3f& Normal1) const;
+	bool AreFacesCoPlanar(int32 Idx0, int32 Idx1) const;
+	bool DoFacesOverlap(int32 Idx0, int32 Idx1) const;
+	static bool TrianglesIntersect(const TStaticArray<FVector2D,3>& T0, const TStaticArray<FVector2D,3>& T1);
+	static bool Cross(const TStaticArray<FVector2D,3>& Points, const FVector2D& B, const FVector2D& C, float Normal);
 
-	//
-	// Builds the connectivity data in the GeometryGroup (Proximity array)
-	// and builds all the data in the BreakingGroup
-	static void UpdateProximity(FGeometryCollection* GeometryCollection);
+private:
+	FGeometryCollection* Collection;
 
-	static bool IsPointInsideOfTriangle(const FVector& P, const FVector& Vertex0, const FVector& Vertex1, const FVector& Vertex2, float Threshold);
+	int32 NumFaces;
+	TArray<FVector3f> SurfaceNormals;
+	TArray<int32> FaceToGeometry;
+	TArray<FVector3f> TransformedVertices;
 
-	// For any transform index, get the connected indices at the operating level. If OperatingLevel is -1, the operating level will be determined by the level of the Indexed node.
-	static void LevelConnectivity(FGeometryCollection* GeometryCollection, int32 Index, TSet<int32>& OutConnections, int32 OperatingLevel = -1);
+	constexpr static int32 NumBins = 20;
+	TArray<TArray<int32>> Bins;
+	TArray<FVector3f> BinNormals;
 
-private:
+	TArray<TArray<int32>> FaceContacts;
 };
+
diff --git a/Engine/Source/Runtime/Experimental/GeometryCollectionEngine/Private/GeometryCollection/GeometryCollectionComponent.cpp b/Engine/Source/Runtime/Experimental/GeometryCollectionEngine/Private/GeometryCollection/GeometryCollectionComponent.cpp
--- a/Engine/Source/Runtime/Experimental/GeometryCollectionEngine/Private/GeometryCollection/GeometryCollectionComponent.cpp
+++ b/Engine/Source/Runtime/Experimental/GeometryCollectionEngine/Private/GeometryCollection/GeometryCollectionComponent.cpp
@@ -2353,35 +2353,32 @@ void FScopedColorEdit::SelectBones(GeometryCollection::ESelectionMode SelectionM
 
 		case GeometryCollection::ESelectionMode::Neighbors:
 		{
-			if (ensureMsgf(GeometryCollectionPtr->HasAttribute("Proximity", FGeometryCollection::GeometryGroup),
-				TEXT("Must build breaking group for neighbor based selection")))
-			{
-				FGeometryCollectionProximityUtility::UpdateProximity(GeometryCollectionPtr.Get());
+			FGeometryCollectionProximityUtility ProximityUtility(GeometryCollectionPtr.Get());
+			ProximityUtility.UpdateProximity();
 
-				const TManagedArray<int32>& TransformIndex = GeometryCollectionPtr->TransformIndex;
-				const TManagedArray<int32>& TransformToGeometryIndex = GeometryCollectionPtr->TransformToGeometryIndex;
-				const TManagedArray<TSet<int32>>& Proximity = GeometryCollectionPtr->GetAttribute<TSet<int32>>("Proximity", FGeometryCollection::GeometryGroup);
+			const TManagedArray<int32>& TransformIndex = GeometryCollectionPtr->TransformIndex;
+			const TManagedArray<int32>& TransformToGeometryIndex = GeometryCollectionPtr->TransformToGeometryIndex;
+			const TManagedArray<TSet<int32>>& Proximity = GeometryCollectionPtr->GetAttribute<TSet<int32>>("Proximity", FGeometryCollection::GeometryGroup);
 
-				const TArray<int32> SelectedBones = GetSelectedBones();
+			const TArray<int32> SelectedBones = GetSelectedBones();
 
-				TArray<int32> NewSelection;
-				for (int32 Bone : SelectedBones)
+			TArray<int32> NewSelection;
+			for (int32 Bone : SelectedBones)
+			{
+				NewSelection.AddUnique(Bone);
+				int32 GeometryIdx = TransformToGeometryIndex[Bone];
+				if (GeometryIdx != INDEX_NONE)
 				{
-					NewSelection.AddUnique(Bone);
-					int32 GeometryIdx = TransformToGeometryIndex[Bone];
-					if (GeometryIdx != INDEX_NONE)
+					const TSet<int32>& Neighbors = Proximity[GeometryIdx];
+					for (int32 NeighborGeometryIndex : Neighbors)
 					{
-						const TSet<int32>& Neighbors = Proximity[GeometryIdx];
-						for (int32 NeighborGeometryIndex : Neighbors)
-						{
-							NewSelection.AddUnique(TransformIndex[NeighborGeometryIndex]);
-						}
+						NewSelection.AddUnique(TransformIndex[NeighborGeometryIndex]);
 					}
 				}
-
-				ResetBoneSelection();
-				AppendSelectedBones(NewSelection);
 			}
+
+			ResetBoneSelection();
+			AppendSelectedBones(NewSelection);
 		}
 		break;
 
