Chaos : Constraints
- Moved Joint and Suspension constraints to the Evolution to match the Collision constraints. This allows access to both from within the evolution.

diff --git a/Engine/Source/Runtime/Experimental/Chaos/Private/Chaos/PBDRigidsEvolutionGBF.cpp b/Engine/Source/Runtime/Experimental/Chaos/Private/Chaos/PBDRigidsEvolutionGBF.cpp
--- a/Engine/Source/Runtime/Experimental/Chaos/Private/Chaos/PBDRigidsEvolutionGBF.cpp
+++ b/Engine/Source/Runtime/Experimental/Chaos/Private/Chaos/PBDRigidsEvolutionGBF.cpp
@@ -63,6 +63,12 @@ FAutoConsoleVariableRef CVarDisableParticleUpdateVelocityParallelFor(TEXT("p.Dis
 int32 ChaosSolverCollisionPriority = 0;
 FAutoConsoleVariableRef CVarChaosSolverCollisionPriority(TEXT("p.Chaos.Solver.Collision.Priority"), ChaosSolverCollisionPriority, TEXT("Set constraint priority. Larger values are evaluated later [def:0]"));
 
+int32 ChaosSolverJointPriority = 0;
+FAutoConsoleVariableRef CVarChaosSolverJointPriority(TEXT("p.Chaos.Solver.Joint.Priority"), ChaosSolverJointPriority, TEXT("Set constraint priority. Larger values are evaluated later [def:0]"));
+
+int32 ChaosSolverSuspensionPriority = 0;
+FAutoConsoleVariableRef CVarChaosSolverSuspensionPriority(TEXT("p.Chaos.Solver.Suspension.Priority"), ChaosSolverSuspensionPriority, TEXT("Set constraint priority. Larger values are evaluated later [def:0]"));
+
 DECLARE_CYCLE_STAT(TEXT("FPBDRigidsEvolutionGBF::AdvanceOneTimeStep"), STAT_Evolution_AdvanceOneTimeStep, STATGROUP_Chaos);
 DECLARE_CYCLE_STAT(TEXT("FPBDRigidsEvolutionGBF::UnclusterUnions"), STAT_Evolution_UnclusterUnions, STATGROUP_Chaos);
 DECLARE_CYCLE_STAT(TEXT("FPBDRigidsEvolutionGBF::Integrate"), STAT_Evolution_Integrate, STATGROUP_Chaos);
@@ -438,6 +444,8 @@ void FPBDRigidsEvolutionGBF::AdvanceOneTimeStepImpl(const FReal Dt,const FSubSte
 FPBDRigidsEvolutionGBF::FPBDRigidsEvolutionGBF(FPBDRigidsSOAs& InParticles,THandleArray<FChaosPhysicsMaterial>& SolverPhysicsMaterials, const TArray<ISimCallbackObject*>* InCollisionModifiers, bool InIsSingleThreaded)
 	: Base(InParticles, SolverPhysicsMaterials, DefaultNumIterations, DefaultNumPushOutIterations, InIsSingleThreaded)
 	, Clustering(*this, Particles.GetClusteredParticles())
+	, JointConstraintRule(JointConstraints, ChaosSolverJointPriority)
+	, SuspensionConstraintRule(SuspensionConstraints, ChaosSolverSuspensionPriority)
 	, CollisionConstraints(InParticles, Collided, PhysicsMaterials, PerParticlePhysicsMaterials, DefaultNumCollisionPairIterations, DefaultNumCollisionPushOutPairIterations, DefaultRestitutionThreshold)
 	, CollisionRule(CollisionConstraints, ChaosSolverCollisionPriority)
 	, BroadPhase(InParticles, DefaultCollisionCullDistance, BoundsThicknessVelocityMultiplier, DefaultCollisionCullDistance)
@@ -487,7 +495,9 @@ FPBDRigidsEvolutionGBF::FPBDRigidsEvolutionGBF(FPBDRigidsSOAs& InParticles,THand
 		GravityForces.Apply(HandleIn, Dt);
 	});
 
+	AddConstraintRule(&SuspensionConstraintRule);
 	AddConstraintRule(&CollisionRule);
+	AddConstraintRule(&JointConstraintRule);
 
 	SetInternalParticleInitilizationFunction([](const FGeometryParticleHandle*, FGeometryParticleHandle*) {});
 	NarrowPhase.GetContext().bFilteringEnabled = true;
diff --git a/Engine/Source/Runtime/Experimental/Chaos/Private/PBDRigidsSolver.cpp b/Engine/Source/Runtime/Experimental/Chaos/Private/PBDRigidsSolver.cpp
--- a/Engine/Source/Runtime/Experimental/Chaos/Private/PBDRigidsSolver.cpp
+++ b/Engine/Source/Runtime/Experimental/Chaos/Private/PBDRigidsSolver.cpp
@@ -195,11 +195,6 @@ FAutoConsoleVariableRef CVarChaosSolverCollisionDeferNarrowPhase(TEXT("p.Chaos.S
 int32 ChaosSolverCollisionUseManifolds = 1;
 FAutoConsoleVariableRef CVarChaosSolverCollisionUseManifolds(TEXT("p.Chaos.Solver.Collision.UseManifolds"), ChaosSolverCollisionUseManifolds, TEXT("Enable/Disable use of manifolds in collision."));
 
-int32 ChaosSolverJointPriority = 0;
-FAutoConsoleVariableRef CVarChaosSolverJointPriority(TEXT("p.Chaos.Solver.Joint.Priority"), ChaosSolverJointPriority, TEXT("Set constraint priority. Larger values are evaluated later [def:0]"));
-
-int32 ChaosSolverSuspensionPriority = 0;
-FAutoConsoleVariableRef CVarChaosSolverSuspensionPriority(TEXT("p.Chaos.Solver.Suspension.Priority"), ChaosSolverSuspensionPriority, TEXT("Set constraint priority. Larger values are evaluated later [def:0]"));
 
 // Joint cvars
 float ChaosSolverJointMinSolverStiffness = 1.0f;
@@ -407,14 +402,11 @@ namespace Chaos
 		, MSolverEventFilters(new FSolverEventFilters())
 		, MDirtyParticlesBuffer(new FDirtyParticlesBuffer(BufferingModeIn, BufferingModeIn == Chaos::EMultiBufferMode::Single))
 		, MCurrentLock(new FCriticalSection())
-		, JointConstraintRule(JointConstraints, ChaosSolverJointPriority)
-		, SuspensionConstraintRule(SuspensionConstraints, ChaosSolverSuspensionPriority)
+
 		, PerSolverField(nullptr)
 	{
 		UE_LOG(LogPBDRigidsSolver, Verbose, TEXT("PBDRigidsSolver::PBDRigidsSolver()"));
 		Reset();
-		MEvolution->AddConstraintRule(&JointConstraintRule);
-		MEvolution->AddConstraintRule(&SuspensionConstraintRule);
 
 		MEvolution->SetInternalParticleInitilizationFunction(
 			[this](const Chaos::FGeometryParticleHandle* OldParticle, Chaos::FGeometryParticleHandle* NewParticle) 
@@ -743,7 +735,7 @@ namespace Chaos
 		MEvolution->GetCollisionDetector().GetNarrowPhase().GetContext().bDeferUpdate = (ChaosSolverCollisionDeferNarrowPhase != 0);
 		MEvolution->GetCollisionDetector().GetNarrowPhase().GetContext().bAllowManifolds = (ChaosSolverCollisionUseManifolds != 0);
 		
-		FPBDJointSolverSettings JointsSettings = JointConstraints.GetSettings();
+		FPBDJointSolverSettings JointsSettings = MEvolution->GetJointConstraints().GetSettings();
 		JointsSettings.MinSolverStiffness = ChaosSolverJointMinSolverStiffness;
 		JointsSettings.MaxSolverStiffness = ChaosSolverJointMaxSolverStiffness;
 		JointsSettings.NumIterationsAtMaxSolverStiffness = ChaosSolverJointNumIterationsAtMaxSolverStiffness;
@@ -751,7 +743,7 @@ namespace Chaos
 		JointsSettings.AngleTolerance = ChaosSolverJointAngleTolerance;
 		JointsSettings.MinParentMassRatio = ChaosSolverJointMinParentMassRatio;
 		JointsSettings.MaxInertiaRatio = ChaosSolverJointMaxInertiaRatio;
-		JointConstraints.SetSettings(JointsSettings);
+		MEvolution->GetJointConstraints().SetSettings(JointsSettings);
 
 		// Apply CVAR overrides if set
 		{
@@ -1296,7 +1288,7 @@ namespace Chaos
 
 	int32 FPBDRigidsSolver::NumJointConstraints() const
 	{
-		return JointConstraints.NumConstraints();
+		return MEvolution->GetJointConstraints().NumConstraints();
 	}
 	
 	int32 FPBDRigidsSolver::NumCollisionConstraints() const
@@ -1385,7 +1377,7 @@ namespace Chaos
 		}
 		if (ChaosSolverDrawJoints == 1)
 		{
-			DebugDraw::DrawJointConstraints(FRigidTransform3(), JointConstraints, 1.0f, ChaosSolverDrawJointFeatures, &ChaosSolverDebugDebugDrawSettings);
+			DebugDraw::DrawJointConstraints(FRigidTransform3(), MEvolution->GetJointConstraints(), 1.0f, ChaosSolverDrawJointFeatures, &ChaosSolverDebugDebugDrawSettings);
 		}
 		if (ChaosSolverDebugDrawCollidingShapes == 1)
 		{
diff --git a/Engine/Source/Runtime/Experimental/Chaos/Public/Chaos/PBDRigidsEvolutionGBF.h b/Engine/Source/Runtime/Experimental/Chaos/Public/Chaos/PBDRigidsEvolutionGBF.h
--- a/Engine/Source/Runtime/Experimental/Chaos/Public/Chaos/PBDRigidsEvolutionGBF.h
+++ b/Engine/Source/Runtime/Experimental/Chaos/Public/Chaos/PBDRigidsEvolutionGBF.h
@@ -14,6 +14,7 @@
 #include "Chaos/PerParticleGravity.h"
 #include "Chaos/PerParticleInitForce.h"
 #include "Chaos/PerParticlePBDEulerStep.h"
+#include "Chaos/PBDSuspensionConstraints.h"
 
 namespace Chaos
 {
@@ -87,6 +88,10 @@ namespace Chaos
 		using FCollisionDetector = FSpatialAccelerationCollisionDetector;
 		using FExternalForces = FPerParticleExternalForces;
 		using FRigidClustering = TPBDRigidClustering<FPBDRigidsEvolutionGBF, FPBDCollisionConstraints>;
+		using FJointConstraintsRule = TPBDConstraintIslandRule<FPBDJointConstraints>;
+		using FSuspensionConstraintsRule = TPBDConstraintIslandRule<FPBDSuspensionConstraints>;
+		using FJointConstraints = FPBDJointConstraints;
+		using FJointConstraintRule = TPBDConstraintIslandRule<FJointConstraints>;
 
 		// Default iteration counts
 		static constexpr int32 DefaultNumIterations = 8;
@@ -162,6 +167,12 @@ namespace Chaos
 		FORCEINLINE const TPBDRigidClustering<FPBDRigidsEvolutionGBF, FPBDCollisionConstraints>& GetRigidClustering() const { return Clustering; }
 		FORCEINLINE TPBDRigidClustering<FPBDRigidsEvolutionGBF, FPBDCollisionConstraints>& GetRigidClustering() { return Clustering; }
 
+		FORCEINLINE FJointConstraints& GetJointConstraints() { return JointConstraints; }
+		FORCEINLINE const FJointConstraints& GetJointConstraints() const { return JointConstraints; }
+
+		FORCEINLINE FPBDSuspensionConstraints& GetSuspensionConstraints() { return SuspensionConstraints; }
+		FORCEINLINE const FPBDSuspensionConstraints& GetSuspensionConstraints() const { return SuspensionConstraints; }
+
 		CHAOS_API inline void EndFrame(FReal Dt)
 		{
 			Particles.GetNonDisabledDynamicView().ParallelFor([&](auto& Particle, int32 Index) {
@@ -258,6 +269,11 @@ namespace Chaos
 
 		TPBDRigidClustering<FPBDRigidsEvolutionGBF, FPBDCollisionConstraints> Clustering;
 
+		FPBDJointConstraints JointConstraints;
+		TPBDConstraintIslandRule<FPBDJointConstraints> JointConstraintRule;
+		FPBDSuspensionConstraints SuspensionConstraints;
+		TPBDConstraintIslandRule<FPBDSuspensionConstraints> SuspensionConstraintRule;
+
 		FGravityForces GravityForces;
 		FCollisionConstraints CollisionConstraints;
 		FCollisionConstraintRule CollisionRule;
diff --git a/Engine/Source/Runtime/Experimental/Chaos/Public/PBDRigidsSolver.h b/Engine/Source/Runtime/Experimental/Chaos/Public/PBDRigidsSolver.h
--- a/Engine/Source/Runtime/Experimental/Chaos/Public/PBDRigidsSolver.h
+++ b/Engine/Source/Runtime/Experimental/Chaos/Public/PBDRigidsSolver.h
@@ -206,11 +206,11 @@ namespace Chaos
 		void SetRemovalFilterSettings(const FSolverRemovalFilterSettings& InRemovalFilterSettings) { GetEventFilters()->GetRemovalFilter()->UpdateFilterSettings(InRemovalFilterSettings); }
 
 		/**/
-		FJointConstraints& GetJointConstraints() { return JointConstraints; }
-		const FJointConstraints& GetJointConstraints() const { return JointConstraints; }
+		FJointConstraints& GetJointConstraints() { return MEvolution->GetJointConstraints(); }
+		const FJointConstraints& GetJointConstraints() const { return MEvolution->GetJointConstraints(); }
 
-		FPBDSuspensionConstraints& GetSuspensionConstraints() { return SuspensionConstraints; }
-		const FPBDSuspensionConstraints& GetSuspensionConstraints() const { return SuspensionConstraints; }
+		FPBDSuspensionConstraints& GetSuspensionConstraints() { return MEvolution->GetSuspensionConstraints(); }
+		const FPBDSuspensionConstraints& GetSuspensionConstraints() const { return MEvolution->GetSuspensionConstraints(); }
 		void SetSuspensionTargetOnPhysicsThread(Chaos::FSuspensionConstraint* GTConstraint, const FVector& TargetPos, bool Enabled);
 
 		/**/
@@ -330,15 +330,6 @@ namespace Chaos
 		TArray< FGeometryCollectionPhysicsProxy* > GeometryCollectionPhysicsProxies_Internal; // PT
 		TArray< FJointConstraintPhysicsProxy* > JointConstraintPhysicsProxies_Internal; // PT
 
-		//
-		//  Constraints
-		//
-		FPBDJointConstraints JointConstraints;
-		TPBDConstraintIslandRule<FPBDJointConstraints> JointConstraintRule;
-
-		FPBDSuspensionConstraints SuspensionConstraints;
-		TPBDConstraintIslandRule<FPBDSuspensionConstraints> SuspensionConstraintRule;
-
 		TUniquePtr<FPerSolverFieldSystem> PerSolverField;
 
 
